re-add the check of bUpdateNavigationInTick in UGeometryCollectionComponent::UpdateNavigationDataIfNeeded

diff --git a/Engine/Source/Runtime/Experimental/GeometryCollectionEngine/Private/GeometryCollection/GeometryCollectionComponent.cpp b/Engine/Source/Runtime/Experimental/GeometryCollectionEngine/Private/GeometryCollection/GeometryCollectionComponent.cpp
--- a/Engine/Source/Runtime/Experimental/GeometryCollectionEngine/Private/GeometryCollection/GeometryCollectionComponent.cpp
+++ b/Engine/Source/Runtime/Experimental/GeometryCollectionEngine/Private/GeometryCollection/GeometryCollectionComponent.cpp
@@ -3355,7 +3355,7 @@ void UGeometryCollectionComponent::UpdateRenderSystemsIfNeeded(bool bDynamicColl
 
 void UGeometryCollectionComponent::UpdateNavigationDataIfNeeded(bool bDynamicCollectionDirty)
 {
-	if (bDynamicCollectionDirty)
+	if (bUpdateNavigationInTick && bDynamicCollectionDirty)
 	{
 		const UWorld* MyWorld = GetWorld();
 		if (MyWorld && MyWorld->IsGameWorld())
