remove calls to ResetDirtyIdx on the PT to prevent a possible race condition

diff --git a/Engine/Source/Runtime/Experimental/Chaos/Private/PBDRigidsSolver.cpp b/Engine/Source/Runtime/Experimental/Chaos/Private/PBDRigidsSolver.cpp
--- a/Engine/Source/Runtime/Experimental/Chaos/Private/PBDRigidsSolver.cpp
+++ b/Engine/Source/Runtime/Experimental/Chaos/Private/PBDRigidsSolver.cpp
@@ -803,7 +803,6 @@ namespace Chaos
 
 				GeometryCollectionPhysicsProxies_Internal.RemoveSingle(InProxy);
 				InProxy->SyncBeforeDestroy();
-				InProxy->ResetDirtyIdx();
 				PendingDestroyGeometryCollectionPhysicsProxy.Add(InProxy);
 			});
 	}
@@ -859,7 +858,6 @@ namespace Chaos
 				}
 				ClusterUnionPhysicsProxies_Internal.RemoveSingle(Proxy);
 				Proxy->SyncBeforeDestroy();
-				Proxy->ResetDirtyIdx();
 				PendingDestroyClusterUnionProxy.Add(Proxy);
 			}
 		);
@@ -1480,9 +1478,6 @@ namespace Chaos
 							GeometryCollectionPhysicsProxies_Internal.Add(Proxy);
 						}
 						Proxy->PushToPhysicsState();
-						// Currently no push needed for geometry collections and they handle the particle creation internally
-						// #TODO This skips the rewind data push so GC will not be rewindable until resolved.
-						Dirty.Proxy->ResetDirtyIdx();
 						break;
 					}
 					case EPhysicsProxyType::ClusterUnionProxy:
@@ -1494,7 +1489,6 @@ namespace Chaos
 							ClusterUnionPhysicsProxies_Internal.Add(Proxy);
 						}
 						Proxy->PushToPhysicsState(*Manager, DataIdx, Dirty);
-						Dirty.Proxy->ResetDirtyIdx();
 						break;
 					}
 					case EPhysicsProxyType::JointConstraintType:
@@ -1507,8 +1501,6 @@ namespace Chaos
 					default:
 					{
 						ensure(0 && TEXT("Unknown proxy type in physics solver."));
-						//Can't use, but we can still mark as "clean"
-						Dirty.Proxy->ResetDirtyIdx();
 					}
 				}
 			}
@@ -1539,7 +1531,6 @@ namespace Chaos
 					}
 				
 					JointProxy->PushStateOnPhysicsThread(this, *Manager, DataIdx, Dirty.PropertyData);
-					Dirty.Proxy->ResetDirtyIdx();
 					break;
 				}
 
@@ -1553,7 +1544,6 @@ namespace Chaos
 						SuspensionProxy->SetInitialized();
 					}
 					SuspensionProxy->PushStateOnPhysicsThread(this, *Manager, DataIdx, Dirty.PropertyData);
-					Dirty.Proxy->ResetDirtyIdx();
 					break;
 				}
 
@@ -1575,7 +1565,6 @@ namespace Chaos
 					//}
 
 					ConstraintProxy->PushStateOnPhysicsThread(this, *Manager, DataIdx, Dirty.PropertyData);
-					Dirty.Proxy->ResetDirtyIdx();
 					break;
 				}
 
