Geometry Collection : Move creation of attribute from InitializeComponent to ResetDynamicCollection to make sure  they are created even if the rest collection changes during the component lifetime

diff --git a/Engine/Source/Runtime/Experimental/GeometryCollectionEngine/Private/GeometryCollection/GeometryCollectionComponent.cpp b/Engine/Source/Runtime/Experimental/GeometryCollectionEngine/Private/GeometryCollection/GeometryCollectionComponent.cpp
--- a/Engine/Source/Runtime/Experimental/GeometryCollectionEngine/Private/GeometryCollection/GeometryCollectionComponent.cpp
+++ b/Engine/Source/Runtime/Experimental/GeometryCollectionEngine/Private/GeometryCollection/GeometryCollectionComponent.cpp
@@ -1388,28 +1388,6 @@ void UGeometryCollectionComponent::SetRestState(TArray<FTransform>&& InRestTrans
 	RefreshCustomRenderer();
 }
 
-void UGeometryCollectionComponent::InitializeComponent()
-{
-	Super::InitializeComponent();
-
-	if (DynamicCollection)
-	{
-		if (bStoreVelocities || bNotifyTrailing)
-		{
-			if (!DynamicCollection->FindAttributeTyped<FVector3f>("LinearVelocity", FTransformCollection::TransformGroup))
-			{
-				DynamicCollection->AddAttribute<FVector3f>("LinearVelocity", FTransformCollection::TransformGroup);
-			}
-
-			if (!DynamicCollection->FindAttributeTyped<FVector3f>("AngularVelocity", FTransformCollection::TransformGroup))
-			{
-				DynamicCollection->AddAttribute<FVector3f>("AngularVelocity", FTransformCollection::TransformGroup);
-			}
-		}
-		DynamicCollection->AddAttribute<uint8>("InternalClusterParentTypeArray", FTransformCollection::TransformGroup);
-	}
-}
-
 void UGeometryCollectionComponent::GetResourceSizeEx(FResourceSizeEx& CumulativeResourceSize)
 {
 	Super::GetResourceSizeEx(CumulativeResourceSize);
@@ -2928,6 +2906,13 @@ void UGeometryCollectionComponent::ResetDynamicCollection()
 		GetSimulationTypeArrayCopyOnWrite();
 		GetStatusFlagsArrayCopyOnWrite();
 
+		if (bStoreVelocities || bNotifyTrailing)
+		{
+			DynamicCollection->AddAttribute<FVector3f>("LinearVelocity", FTransformCollection::TransformGroup);
+			DynamicCollection->AddAttribute<FVector3f>("AngularVelocity", FTransformCollection::TransformGroup);
+		}
+		DynamicCollection->AddAttribute<uint8>("InternalClusterParentTypeArray", FTransformCollection::TransformGroup);
+
 		FGeometryCollectionDecayDynamicFacade DecayDynamicFacade(*DynamicCollection);
 		
 		// we are not testing for bAllowRemovalOnSleep, so that we can enable it at runtime if necessary
diff --git a/Engine/Source/Runtime/Experimental/GeometryCollectionEngine/Public/GeometryCollection/GeometryCollectionComponent.h b/Engine/Source/Runtime/Experimental/GeometryCollectionEngine/Public/GeometryCollection/GeometryCollectionComponent.h
--- a/Engine/Source/Runtime/Experimental/GeometryCollectionEngine/Public/GeometryCollection/GeometryCollectionComponent.h
+++ b/Engine/Source/Runtime/Experimental/GeometryCollectionEngine/Public/GeometryCollection/GeometryCollectionComponent.h
@@ -461,7 +461,6 @@ public:
 	GEOMETRYCOLLECTIONENGINE_API virtual void BeginPlay() override;
 	GEOMETRYCOLLECTIONENGINE_API virtual void EndPlay(const EEndPlayReason::Type ReasonEnd) override;
 	GEOMETRYCOLLECTIONENGINE_API virtual void GetLifetimeReplicatedProps(TArray<FLifetimeProperty>& OutLifetimeProps) const override;
-	GEOMETRYCOLLECTIONENGINE_API virtual void InitializeComponent() override;
 	GEOMETRYCOLLECTIONENGINE_API virtual void GetResourceSizeEx(FResourceSizeEx& CumulativeResourceSize) override;
 
 #if WITH_EDITOR
