ensure that the DecayFacade and DynamicStateFacade are valid in UGeometryCollectionComponent::CheckFullyDecayed

diff --git a/Engine/Source/Runtime/Experimental/GeometryCollectionEngine/Private/GeometryCollection/GeometryCollectionComponent.cpp b/Engine/Source/Runtime/Experimental/GeometryCollectionEngine/Private/GeometryCollection/GeometryCollectionComponent.cpp
--- a/Engine/Source/Runtime/Experimental/GeometryCollectionEngine/Private/GeometryCollection/GeometryCollectionComponent.cpp
+++ b/Engine/Source/Runtime/Experimental/GeometryCollectionEngine/Private/GeometryCollection/GeometryCollectionComponent.cpp
@@ -2737,7 +2737,7 @@ void UGeometryCollectionComponent::TickComponent(float DeltaTime, enum ELevelTic
 
 void UGeometryCollectionComponent::CheckFullyDecayed()
 {
-	if (bAlreadyFullyDecayed)
+	if (bAlreadyFullyDecayed || !OnFullyDecayedEvent.IsBound())
 	{
 		// Already fully decayed - don't bother doing extra work.
 		return;
@@ -2746,38 +2746,40 @@ void UGeometryCollectionComponent::CheckFullyDecayed()
 	if (DynamicCollection && PhysicsProxy)
 	{
 		bool bFullyDecayed = true;
-		FGeometryCollectionDecayDynamicFacade DecayFacade(*DynamicCollection);
 		FGeometryCollectionDynamicStateFacade DynamicStateFacade(*DynamicCollection);
 
-		const int32 NumTransforms = DecayFacade.GetDecayAttributeSize();
-		for (int32 TransformIdx = 0; TransformIdx < NumTransforms; ++TransformIdx)
+		if (DynamicStateFacade.IsValid())
 		{
-			// If we didn't create a particle for this transform, we shouldn't consider this particle.
-			if (!PhysicsProxy->GetExternalParticles()[TransformIdx])
+			const int32 NumTransforms = DynamicCollection->NumElements(FGeometryCollection::TransformGroup);
+			for (int32 TransformIdx = 0; TransformIdx < NumTransforms; ++TransformIdx)
 			{
-				continue;
-			}
+				// If we didn't create a particle for this transform, we shouldn't consider this particle.
+				if (!PhysicsProxy->GetExternalParticles()[TransformIdx])
+				{
+					continue;
+				}
 
-			// In an internal cluster, decay hasn't gotten to this particle yet (e.g. could be in a cluster union).
-			if (DynamicStateFacade.HasInternalClusterParent(TransformIdx))
-			{
-				bFullyDecayed = false;
-				break;
+				// In an internal cluster, decay hasn't gotten to this particle yet (e.g. could be in a cluster union).
+				if (DynamicStateFacade.HasInternalClusterParent(TransformIdx))
+				{
+					bFullyDecayed = false;
+					break;
+				}
+
+				// If the particle is active, it's definitely not decayed either.
+				if (DynamicStateFacade.IsActive(TransformIdx))
+				{
+					bFullyDecayed = false;
+					break;
+				}
 			}
 
-			// If the particle is active, it's definitely not decayed either.
-			if (DynamicStateFacade.IsActive(TransformIdx))
+			if (bFullyDecayed)
 			{
-				bFullyDecayed = false;
-				break;
+				bAlreadyFullyDecayed = true;
+				OnFullyDecayedEvent.Broadcast();
 			}
 		}
-
-		if (bFullyDecayed)
-		{
-			bAlreadyFullyDecayed = true;
-			OnFullyDecayedEvent.Broadcast();
-		}
 	}
 }
 
