Geometry collection : fixed ApplyBreakingAngularVelocity not calling the right proxy method

diff --git a/Engine/Source/Runtime/Experimental/GeometryCollectionEngine/Private/GeometryCollection/GeometryCollectionComponent.cpp b/Engine/Source/Runtime/Experimental/GeometryCollectionEngine/Private/GeometryCollection/GeometryCollectionComponent.cpp
--- a/Engine/Source/Runtime/Experimental/GeometryCollectionEngine/Private/GeometryCollection/GeometryCollectionComponent.cpp
+++ b/Engine/Source/Runtime/Experimental/GeometryCollectionEngine/Private/GeometryCollection/GeometryCollectionComponent.cpp
@@ -5258,7 +5258,7 @@ void UGeometryCollectionComponent::ApplyBreakingAngularVelocity(int32 ItemIndex,
 {
 	if (PhysicsProxy)
 	{
-		PhysicsProxy->ApplyBreakingLinearVelocity_External(FGeometryCollectionItemIndex::CreateFromExistingItemIndex(ItemIndex), AngularVelocity);
+		PhysicsProxy->ApplyBreakingAngularVelocity_External(FGeometryCollectionItemIndex::CreateFromExistingItemIndex(ItemIndex), AngularVelocity);
 	}
 }
 
