GeometryCollectionComponent: using root proxy to take into account root particle position when computing bounds

diff --git a/Engine/Source/Runtime/Experimental/GeometryCollectionEngine/Private/GeometryCollection/GeometryCollectionComponent.cpp b/Engine/Source/Runtime/Experimental/GeometryCollectionEngine/Private/GeometryCollection/GeometryCollectionComponent.cpp
--- a/Engine/Source/Runtime/Experimental/GeometryCollectionEngine/Private/GeometryCollection/GeometryCollectionComponent.cpp
+++ b/Engine/Source/Runtime/Experimental/GeometryCollectionEngine/Private/GeometryCollection/GeometryCollectionComponent.cpp
@@ -687,6 +687,9 @@ FBoxSphereBounds UGeometryCollectionComponent::CalcBounds(const FTransform& Loca
 {	
 	SCOPE_CYCLE_COUNTER(STAT_GCCUpdateBounds);
 
+	const int32 RootIndex = GetRootIndex();
+	const FTransform FinalTransform = (DynamicCollection ? DynamicCollection->Transform[RootIndex] : FTransform::Identity) * LocalToWorldIn;
+	
 	if (bChaos_GC_CacheComponentSpaceBounds)
 	{
 		bool NeedBoundsUpdate = false;
@@ -704,10 +707,10 @@ FBoxSphereBounds UGeometryCollectionComponent::CalcBounds(const FTransform& Loca
 			NeedBoundsUpdate = false;
 		}
 
-		return ComponentSpaceBounds.TransformBy(LocalToWorldIn);
+		return ComponentSpaceBounds.TransformBy(FinalTransform);
 	}
 
-	return FBoxSphereBounds(ComputeBounds(LocalToWorldIn));
+	return FBoxSphereBounds(ComputeBounds(FinalTransform));
 }
 
 int32 UGeometryCollectionComponent::GetNumElements(FName Group) const
