Geometry Collection : Fix issue where breaking velocity would not be set properly for geometry collection part of a cluster union

diff --git a/Engine/Source/Runtime/Experimental/Chaos/Private/PhysicsProxy/GeometryCollectionPhysicsProxy.cpp b/Engine/Source/Runtime/Experimental/Chaos/Private/PhysicsProxy/GeometryCollectionPhysicsProxy.cpp
--- a/Engine/Source/Runtime/Experimental/Chaos/Private/PhysicsProxy/GeometryCollectionPhysicsProxy.cpp
+++ b/Engine/Source/Runtime/Experimental/Chaos/Private/PhysicsProxy/GeometryCollectionPhysicsProxy.cpp
@@ -2409,7 +2409,8 @@ template <typename TAction>
 static void ApplyToBreakingChildren_Internal(Chaos::FRigidClustering& Clustering, Chaos::FPBDRigidClusteredParticleHandle& ClusteredHandle, TAction Action)
 {
 	const bool bIsCluster = (ClusteredHandle.ClusterIds().NumChildren > 0);
-	if (bIsCluster && !ClusteredHandle.Disabled())
+	const bool bHasClusterUnionParent = Clustering.GetClusterUnionManager().IsClusterUnionParticle(ClusteredHandle.Parent());
+	if (bIsCluster && (!ClusteredHandle.Disabled() || bHasClusterUnionParent))
 	{
 		Chaos::FRigidClustering::FClusterMap& ChildrenMap = Clustering.GetChildrenMap();
 		if (const TArray<Chaos::FPBDRigidParticleHandle*>* ChildrenHandles = ChildrenMap.Find(&ClusteredHandle))
