IAS: Let the HttpPendingRequests counter include all inflight requests which makes it easier to understand.
Fix incorrect csv reporting of HttpPendingCount and HttpInflightCount.

#rnx

- Submitted on the behalf of PJ.Kack

diff --git a/Engine/Source/Runtime/Experimental/IoStoreOnDemand/Private/Statistics.cpp b/Engine/Source/Runtime/Experimental/IoStoreOnDemand/Private/Statistics.cpp
--- a/Engine/Source/Runtime/Experimental/IoStoreOnDemand/Private/Statistics.cpp
+++ b/Engine/Source/Runtime/Experimental/IoStoreOnDemand/Private/Statistics.cpp
@@ -200,17 +200,17 @@ void FOnDemandIoBackendStats::OnHttpEnqueue()
 
 void FOnDemandIoBackendStats::OnHttpDequeue()
 {
-	GHttpPendingCount.Add(-1);
-	CSV_CUSTOM_STAT_DEFINED(HttpPendingCount, int32(GHttpPendingCount.Get()), ECsvCustomStatOp::Set);
-
 	GHttpInflightCount.Add(1);
-	CSV_CUSTOM_STAT_DEFINED(HttpPendingCount, int32(GHttpInflightCount.Get()), ECsvCustomStatOp::Set);
+	CSV_CUSTOM_STAT_DEFINED(HttpInflightCount, int32(GHttpInflightCount.Get()), ECsvCustomStatOp::Set);
 }
 
 void FOnDemandIoBackendStats::OnHttpGet(uint64 InSize)
 {
+	GHttpPendingCount.Add(-1);
+	CSV_CUSTOM_STAT_DEFINED(HttpPendingCount, int32(GHttpPendingCount.Get()), ECsvCustomStatOp::Set);
+
 	GHttpInflightCount.Add(-1);
-	CSV_CUSTOM_STAT_DEFINED(HttpPendingCount, int32(GHttpInflightCount.Get()), ECsvCustomStatOp::Set);
+	CSV_CUSTOM_STAT_DEFINED(HttpInflightCount, int32(GHttpInflightCount.Get()), ECsvCustomStatOp::Set);
 
 	GHttpGetCount.Add(1);
 	CSV_CUSTOM_STAT_DEFINED(HttpGetCount, int32(GHttpGetCount.Get()), ECsvCustomStatOp::Set);
@@ -221,8 +221,11 @@ void FOnDemandIoBackendStats::OnHttpGet(uint64 InSize)
 
 void FOnDemandIoBackendStats::OnHttpRetry()
 {
+	GHttpPendingCount.Add(-1);
+	CSV_CUSTOM_STAT_DEFINED(HttpPendingCount, int32(GHttpPendingCount.Get()), ECsvCustomStatOp::Set);
+
 	GHttpInflightCount.Add(-1);
-	CSV_CUSTOM_STAT_DEFINED(HttpPendingCount, int32(GHttpInflightCount.Get()), ECsvCustomStatOp::Set);
+	CSV_CUSTOM_STAT_DEFINED(HttpInflightCount, int32(GHttpInflightCount.Get()), ECsvCustomStatOp::Set);
 
 	GHttpRetryCount.Add(1);
 	CSV_CUSTOM_STAT_DEFINED(HttpRetryCount, int32(GHttpRetryCount.Get()), ECsvCustomStatOp::Set);
@@ -230,8 +233,11 @@ void FOnDemandIoBackendStats::OnHttpRetry()
 
 void FOnDemandIoBackendStats::OnHttpError()
 {
+	GHttpPendingCount.Add(-1);
+	CSV_CUSTOM_STAT_DEFINED(HttpPendingCount, int32(GHttpPendingCount.Get()), ECsvCustomStatOp::Set);
+
 	GHttpInflightCount.Add(-1);
-	CSV_CUSTOM_STAT_DEFINED(HttpPendingCount, int32(GHttpInflightCount.Get()), ECsvCustomStatOp::Set);
+	CSV_CUSTOM_STAT_DEFINED(HttpInflightCount, int32(GHttpInflightCount.Get()), ECsvCustomStatOp::Set);
 
 	GHttpErrorCount.Add(1);
 	CSV_CUSTOM_STAT_DEFINED(HttpErrorCount, int32(GHttpErrorCount.Get()), ECsvCustomStatOp::Set);
