IAS: Add logging of all iorequests to complement the logging of all http requests.
The iorequest duration is very close to the actual latency as seen from the game client code.

#rnx

- Submitted on the behalf of PJ.Kack

diff --git a/Engine/Source/Runtime/Core/Private/IO/IoDispatcher.cpp b/Engine/Source/Runtime/Core/Private/IO/IoDispatcher.cpp
--- a/Engine/Source/Runtime/Core/Private/IO/IoDispatcher.cpp
+++ b/Engine/Source/Runtime/Core/Private/IO/IoDispatcher.cpp
@@ -94,7 +94,7 @@ public:
 
 	void OnBatchIssued(FIoBatch& Batch)
 	{
-		double StartTime = FPlatformTime::Seconds();
+		const uint64 StartTime = FPlatformTime::Cycles64();
 		FIoRequestImpl* Request = Batch.HeadRequest;
 		while (Request)
 		{
@@ -123,7 +123,8 @@ public:
 		}
 		FRequestCategory* Category = ChunkTypeToCategoryMap[static_cast<int32>(Request.ChunkId.GetChunkType())];
 		++Category->TotaRequestsCount;
-		Category->TotalRequestsTime += FPlatformTime::Seconds() - Request.StartTime;
+		const double Duration = FPlatformTime::ToSeconds64(FPlatformTime::Cycles64() - Request.GetStartTime());
+		Category->TotalRequestsTime += Duration;
 #if COUNTERSTRACE_ENABLED
 		Category->TotalLoadedCounter.Add(Request.GetBuffer().DataSize());
 		Category->AverageDurationCounter.Set(Category->TotalRequestsTime / double(Category->TotaRequestsCount));
diff --git a/Engine/Source/Runtime/Core/Public/IO/IoDispatcherBackend.h b/Engine/Source/Runtime/Core/Public/IO/IoDispatcherBackend.h
--- a/Engine/Source/Runtime/Core/Public/IO/IoDispatcherBackend.h
+++ b/Engine/Source/Runtime/Core/Public/IO/IoDispatcherBackend.h
@@ -54,6 +54,15 @@ public:
 		Buffer.Emplace(InBuffer);
 	}
 
+	uint64 GetStartTime() const
+	{
+#if UE_IODISPATCHER_STATS_ENABLED
+		return StartTime;
+#else
+		return 0;
+#endif
+	}
+
 private:
 	friend class FIoDispatcherImpl;
 	friend class FIoRequest;
@@ -79,7 +88,7 @@ private:
 	struct IIoDispatcherBackend* Backend = nullptr;
 	FIoBatchImpl* Batch = nullptr;
 #if UE_IODISPATCHER_STATS_ENABLED
-	double StartTime = -1.0;
+	uint64 StartTime = 0;
 #endif
 	TOptional<FIoBuffer> Buffer;
 	FIoReadCallback Callback;
diff --git a/Engine/Source/Runtime/Experimental/IoStoreOnDemand/Private/OnDemandIoDispatcherBackend.cpp b/Engine/Source/Runtime/Experimental/IoStoreOnDemand/Private/OnDemandIoDispatcherBackend.cpp
--- a/Engine/Source/Runtime/Experimental/IoStoreOnDemand/Private/OnDemandIoDispatcherBackend.cpp
+++ b/Engine/Source/Runtime/Experimental/IoStoreOnDemand/Private/OnDemandIoDispatcherBackend.cpp
@@ -60,10 +60,10 @@ namespace UE::IO::Private
 {
 
 ///////////////////////////////////////////////////////////////////////////////
-static void LogHttpResult(const TCHAR* Url, uint32 StatusCode, uint32 Duration, uint32 Size, uint32 Offset, const char* Memo="")
+static void LogHttpResult(const TCHAR* Url, uint32 StatusCode, uint32 Duration, uint32 Size, uint32 Offset, const char* Memo="ok")
 {
 	Size >>= 10;
-	UE_LOG(LogIas, VeryVerbose, TEXT("http:%u %ums %uKiB %u %S %s"), StatusCode, Duration, Size, Offset, Memo, Url);
+	UE_LOG(LogIas, VeryVerbose, TEXT("http-%3u: %4ums %5uKiB [%7u] '%S' %s"), StatusCode, Duration, Size, Offset, Memo, Url);
 };
 
 using namespace UE::Tasks;
@@ -656,6 +656,11 @@ struct FChunkRequestParams
 		return FChunkRequestParams{GetChunkKey(ChunkInfo.Entry->Hash, ChunkRange), ChunkRange, ChunkInfo};
 	}
 
+	const FIoHash& GetUrlHash() const
+	{
+		return ChunkInfo.Entry->Hash;
+	}
+
 	void GetUrl(FAnsiStringBuilderBase& Url) const
 	{
 		const FString HashString = LexToString(ChunkInfo.Entry->Hash);
@@ -766,6 +771,36 @@ struct FChunkRequest
 	bool bCached;
 };
 
+///////////////////////////////////////////////////////////////////////////////
+static void LogIoResult(
+	const FIoChunkId& ChunkId,
+	const FIoHash& UrlHash,
+	const FIoHash& CacheKey,
+	uint32 Duration,
+	uint32 UncompressedSize,
+	uint32 UncompressedOffset,
+	uint32 CompressedOffset,
+	bool bCached)
+{
+	const TCHAR* Prefix = [bCached, UncompressedSize]() -> const TCHAR*
+	{
+		if (UncompressedSize == 0)
+		{
+			return bCached ? TEXT("io-cache-error") : TEXT("io-http-error ");
+		}
+		return bCached ? TEXT("io-cache") : TEXT("io-http ");
+	}();
+	UE_LOG(LogIas, VeryVerbose, TEXT("%s: %4ums %5uKiB [%7u] %s:%s:%s|%u"),
+		Prefix,
+		Duration,
+		UncompressedSize >> 10,
+		UncompressedOffset,
+		*LexToString(ChunkId),
+		*LexToString(CacheKey),
+		*LexToString(UrlHash),
+		CompressedOffset);
+};
+
 ///////////////////////////////////////////////////////////////////////////////
 class FOnDemandIoBackend final
 	: public FRunnable
@@ -954,6 +989,7 @@ void FOnDemandIoBackend::Shutdown()
 
 void FOnDemandIoBackend::CompleteRequest(FChunkRequest* ChunkRequest)
 {
+	TRACE_CPUPROFILER_EVENT_SCOPE(FOnDemandIoBackend::CompleteRequest);
 	check(ChunkRequest != nullptr);
 	const bool bCancelled = ChunkRequest->CancellationToken.IsCancelled();
 
@@ -987,14 +1023,28 @@ void FOnDemandIoBackend::CompleteRequest(FChunkRequest* ChunkRequest)
 			bDecoded = FIoChunkEncoding::Decode(DecodingParams, Chunk.GetView(), Request->GetBuffer().GetMutableView());
 		}
 		
-		if (!bDecoded)
+		const uint64 Duration = Request->GetStartTime() > 0 ?
+			(uint64)FPlatformTime::ToMilliseconds64(FPlatformTime::Cycles64() - Request->GetStartTime()) : 0;
+
+		if (bDecoded)
+		{
+			Stats.OnIoRequestComplete(Request->GetBuffer().GetSize());
+			LogIoResult(Request->ChunkId, ChunkRequest->Params.GetUrlHash(), ChunkRequest->Params.ChunkKey, Duration,
+				Request->GetBuffer().DataSize(), Request->Options.GetOffset(),
+				ChunkRequest->Params.ChunkRange.GetOffset(), ChunkRequest->bCached);
+				
+		}
+		else
 		{
-			Stats.OnIoRequestFail();
 			bCanCache = false;
 			Request->SetFailed();
+
+			Stats.OnIoRequestFail();
+			LogIoResult(Request->ChunkId, ChunkRequest->Params.GetUrlHash(), ChunkRequest->Params.ChunkKey, Duration,
+				0, Request->Options.GetOffset(),
+				ChunkRequest->Params.ChunkRange.GetOffset(), ChunkRequest->bCached);
 		}
 
-		Stats.OnIoRequestComplete(Request->GetBuffer().GetSize());
 		CompletedRequests.Enqueue(Request);
 		BackendContext->WakeUpDispatcherThreadDelegate.Execute();
 	}
@@ -1032,6 +1082,7 @@ bool FOnDemandIoBackend::Resolve(FIoRequestImpl* Request)
 		
 	auto CompleteOrEnqueueHttpRequest = [this, ChunkRequest]()
 	{
+		TRACE_CPUPROFILER_EVENT_SCOPE(FOnDemandIoBackend::CompleteOrEnqueueHttpRequest);
 		if (ChunkRequest->CacheTask.IsValid())
 		{
 			if (TIoStatusOr<FIoBuffer> Status = ChunkRequest->CacheTask.GetResult(); Status.IsOk())
