Swap events double buffer on the game thread.

diff --git a/Engine/Source/Runtime/Experimental/Chaos/Private/Chaos/Framework/PhysicsSolverBase.cpp b/Engine/Source/Runtime/Experimental/Chaos/Private/Chaos/Framework/PhysicsSolverBase.cpp
--- a/Engine/Source/Runtime/Experimental/Chaos/Private/Chaos/Framework/PhysicsSolverBase.cpp
+++ b/Engine/Source/Runtime/Experimental/Chaos/Private/Chaos/Framework/PhysicsSolverBase.cpp
@@ -210,7 +210,6 @@ namespace Chaos
 		, MMinDeltaTime(UE_SMALL_NUMBER)
 		, MMaxSubSteps(1)
 		, ExternalSteps(0)
-		, bCanSwitchEventBuffer(true)
 #if !UE_BUILD_SHIPPING
 		, bStealAdvanceTasksForTesting(false)
 #endif
diff --git a/Engine/Source/Runtime/Experimental/Chaos/Private/PBDRigidsSolver.cpp b/Engine/Source/Runtime/Experimental/Chaos/Private/PBDRigidsSolver.cpp
--- a/Engine/Source/Runtime/Experimental/Chaos/Private/PBDRigidsSolver.cpp
+++ b/Engine/Source/Runtime/Experimental/Chaos/Private/PBDRigidsSolver.cpp
@@ -513,14 +513,6 @@ namespace Chaos
 					MSolver->GetEventManager()->FillProducerData(MSolver, ResetData);
 					MSolver->GetEvolution()->ResetAllRemovals();
 				}
-
-				// flip on last sub-step of frame
-				if (MSubStepInfo.Step == MSubStepInfo.NumSteps - 1 && MSolver->CanSwitchEventBuffer())
-				{
-					SCOPE_CYCLE_COUNTER(STAT_FlipBuffersIfRequired);
-					MSolver->GetEventManager()->FlipBuffersIfRequired();
-					MSolver->BlockSwitchEventBuffer();
-				}
 			}
 			
 			{
diff --git a/Engine/Source/Runtime/Experimental/Chaos/Public/Chaos/Framework/PhysicsSolverBase.h b/Engine/Source/Runtime/Experimental/Chaos/Public/Chaos/Framework/PhysicsSolverBase.h
--- a/Engine/Source/Runtime/Experimental/Chaos/Public/Chaos/Framework/PhysicsSolverBase.h
+++ b/Engine/Source/Runtime/Experimental/Chaos/Public/Chaos/Framework/PhysicsSolverBase.h
@@ -636,6 +636,8 @@ namespace Chaos
 			}
 		}
 
+		virtual void FlipEventManagerBuffer() {}
+
 		/**/
 		void SetSolverTime(const FReal InTime) { MTime = InTime; }
 		const FReal GetSolverTime() const { return MTime; }
@@ -683,11 +685,6 @@ namespace Chaos
 			return NetworkPhysicsPredictionInterpLerp;
 		}
 
-
-		bool CanSwitchEventBuffer() const { return bCanSwitchEventBuffer; }
-		void AllowSwitchEventBuffer() { bCanSwitchEventBuffer = true; }
-		void BlockSwitchEventBuffer() { bCanSwitchEventBuffer = false; }
-
 	protected:
 		/** Mode that the results buffers should be set to (single, double, triple) */
 		EMultiBufferMode BufferMode;
@@ -793,7 +790,6 @@ namespace Chaos
 		float MMinDeltaTime;
 		int32 MMaxSubSteps;
 		int32 ExternalSteps;
-		bool bCanSwitchEventBuffer;
 		TArray<FGeometryParticle*> UniqueIdxToGTParticles;
 
 	public:
diff --git a/Engine/Source/Runtime/Experimental/Chaos/Public/PBDRigidsSolver.h b/Engine/Source/Runtime/Experimental/Chaos/Public/PBDRigidsSolver.h
--- a/Engine/Source/Runtime/Experimental/Chaos/Public/PBDRigidsSolver.h
+++ b/Engine/Source/Runtime/Experimental/Chaos/Public/PBDRigidsSolver.h
@@ -232,6 +232,7 @@ namespace Chaos
 		
 		/**/
 		FEventManager* GetEventManager() { return MEventManager.Get(); }
+		virtual void FlipEventManagerBuffer() { MEventManager->FlipBuffersIfRequired(); }
 
 		/**/
 		FSolverEventFilters* GetEventFilters() { return MSolverEventFilters.Get(); }
diff --git a/Engine/Source/Runtime/PhysicsCore/Private/ChaosScene.cpp b/Engine/Source/Runtime/PhysicsCore/Private/ChaosScene.cpp
--- a/Engine/Source/Runtime/PhysicsCore/Private/ChaosScene.cpp
+++ b/Engine/Source/Runtime/PhysicsCore/Private/ChaosScene.cpp
@@ -538,7 +538,7 @@ void FChaosScene::EndFrame()
 			{
 				SyncBodies(&Concrete);
 				Concrete.SyncEvents_GameThread();
-				Solver->AllowSwitchEventBuffer();
+				Solver->FlipEventManagerBuffer();
 				{
 					SCOPE_CYCLE_COUNTER(STAT_SqUpdateMaterials);
 					Concrete.SyncQueryMaterials_External();
