Remove hack that disable collision event

diff --git a/Engine/Source/Runtime/Experimental/ChaosSolverEngine/Private/Chaos/ChaosGameplayEventDispatcher.cpp b/Engine/Source/Runtime/Experimental/ChaosSolverEngine/Private/Chaos/ChaosGameplayEventDispatcher.cpp
--- a/Engine/Source/Runtime/Experimental/ChaosSolverEngine/Private/Chaos/ChaosGameplayEventDispatcher.cpp
+++ b/Engine/Source/Runtime/Experimental/ChaosSolverEngine/Private/Chaos/ChaosGameplayEventDispatcher.cpp
@@ -375,11 +375,6 @@ TArray<UObject*> UChaosGameplayEventDispatcher::GetInterestedProxyOwnersForCrumb
 	return Result;
 }
 
-int32 GReducedCollisionDispatch = 0;
-float GReducedCollisionDispatchAmount = 0.1f;
-FAutoConsoleVariableRef CVar_ReducedCollisionDispatch(TEXT("p.chaos.gpdispatch.collisions.reduce"), GReducedCollisionDispatch, TEXT("Enable reduced dispatch of collision events where only a subset of total collisions are actually dispatched to outside systems [0 - off, 1- on]"));
-FAutoConsoleVariableRef CVar_ReducedCollisionDispatchAmount(TEXT("p.chaos.gpdispatch.collisions.reduceamount"), GReducedCollisionDispatchAmount, TEXT("Percentage of collisions to actually dispatch [0 - 1]"));
-
 void UChaosGameplayEventDispatcher::HandleCollisionEvents(const Chaos::FCollisionEventData& Event)
 {
 	// todo(chaos) : this code is very similar to FPhysScene_Chaos::HandleCollisionEvents, we should propably consolidate if possible or share as much code as possible 
@@ -421,13 +416,7 @@ void UChaosGameplayEventDispatcher::HandleCollisionEvents(const Chaos::FCollisio
 								continue;
 							}
 
-							const int32 NumToDispatch = GReducedCollisionDispatch == 1 ?
-								FMath::Max((int32)((float)NumCollisionIndices * FMath::Clamp(GReducedCollisionDispatchAmount, 0.0f, 1.0f)), 1) :
-								NumCollisionIndices;
-
-							const int32 Interval = NumCollisionIndices / NumToDispatch;
-
-							for(int32 Index = 0; Index < NumCollisionIndices; Index += Interval)
+							for(int32 Index = 0; Index < NumCollisionIndices; ++Index)
 							{
 								int32 EncodedCollisionIdx = (*CollisionIndices)[Index];
 
