Iris
- Fixed ensure in ReplicationWriter trying to mark no longer replicated owner as dirty.

diff --git a/Engine/Source/Runtime/Experimental/Iris/Core/Private/Iris/ReplicationSystem/ReplicationWriter.cpp b/Engine/Source/Runtime/Experimental/Iris/Core/Private/Iris/ReplicationSystem/ReplicationWriter.cpp
--- a/Engine/Source/Runtime/Experimental/Iris/Core/Private/Iris/ReplicationSystem/ReplicationWriter.cpp
+++ b/Engine/Source/Runtime/Experimental/Iris/Core/Private/Iris/ReplicationSystem/ReplicationWriter.cpp
@@ -598,7 +598,7 @@ void FReplicationWriter::SetPendingDestroyOrSubObjectPendingDestroyState(uint32
 			// If owner is not pending destroy we mark the state of the SubObject to SubObjectPendingDestroy and mark owner as having dirty subobjects which will 
 			// destroy the subobject using the replicated state path of the owner.
 			FReplicationInfo& OwnerInfo = ReplicatedObjects[ObjectData.SubObjectRootIndex];
-			if (OwnerInfo.GetState() < EReplicatedObjectState::PendingDestroy)
+			if ((OwnerInfo.GetState() != EReplicatedObjectState::Invalid) && (OwnerInfo.GetState() < EReplicatedObjectState::PendingDestroy))
 			{
 				MarkObjectDirty(ObjectData.SubObjectRootIndex, "SetPendingDestroyOrSubObjectPendingDestroyState");
 				OwnerInfo.HasDirtySubObjects = 1U;
