Check number of manifold points for each per-particle contact modifier

diff --git a/Engine/Source/Runtime/Experimental/Chaos/Private/Chaos/ContactModification.cpp b/Engine/Source/Runtime/Experimental/Chaos/Private/Chaos/ContactModification.cpp
--- a/Engine/Source/Runtime/Experimental/Chaos/Private/Chaos/ContactModification.cpp
+++ b/Engine/Source/Runtime/Experimental/Chaos/Private/Chaos/ContactModification.cpp
@@ -534,10 +534,13 @@ namespace Chaos
 		Constraints.Reserve(NumCollisions);
 		ParticleCollisions.VisitCollisions([this](FPBDCollisionConstraint& Constraint)
 		{
-			Constraints.Add(&Constraint);
+			if (Constraint.GetManifoldPoints().Num() > 0)
+			{
+				Constraints.Add(&Constraint);
+			}
 			return ECollisionVisitorResult::Continue;
 		});
-		ensureMsgf(Constraints.Num() == NumCollisions, TEXT("Number of constraints visited exceeded the number reported by the FParticleCollisions object, and therefore exceeded the reserved constraint cache."));
+		ensureMsgf(Constraints.Num() <= NumCollisions, TEXT("Number of constraints visited exceeded the number reported by the FParticleCollisions object, and therefore exceeded the reserved constraint cache."));
 	}
 
 	FContactPairModifierParticleRange FCollisionContactModifier::GetContacts(FGeometryParticleHandle* Particle)
