Add check to per-particle contact modifier for disabled constraint, so that the iterator doesn't hit constraints that are already disabled.

diff --git a/Engine/Source/Runtime/Experimental/Chaos/Private/Chaos/ContactModification.cpp b/Engine/Source/Runtime/Experimental/Chaos/Private/Chaos/ContactModification.cpp
--- a/Engine/Source/Runtime/Experimental/Chaos/Private/Chaos/ContactModification.cpp
+++ b/Engine/Source/Runtime/Experimental/Chaos/Private/Chaos/ContactModification.cpp
@@ -534,7 +534,7 @@ namespace Chaos
 		Constraints.Reserve(NumCollisions);
 		ParticleCollisions.VisitCollisions([this](FPBDCollisionConstraint& Constraint)
 		{
-			if (Constraint.GetManifoldPoints().Num() > 0)
+			if (Constraint.GetManifoldPoints().Num() > 0 && !Constraint.GetDisabled())
 			{
 				Constraints.Add(&Constraint);
 			}
