Geometry collection : override Primitive component GetMass method for geometry collections
Added display of mass in geometry collection physics editor properties

diff --git a/Engine/Source/Runtime/Experimental/GeometryCollectionEngine/Private/GeometryCollection/GeometryCollectionComponent.cpp b/Engine/Source/Runtime/Experimental/GeometryCollectionEngine/Private/GeometryCollection/GeometryCollectionComponent.cpp
--- a/Engine/Source/Runtime/Experimental/GeometryCollectionEngine/Private/GeometryCollection/GeometryCollectionComponent.cpp
+++ b/Engine/Source/Runtime/Experimental/GeometryCollectionEngine/Private/GeometryCollection/GeometryCollectionComponent.cpp
@@ -5498,6 +5498,44 @@ void UGeometryCollectionComponent::GetMassAndExtents(int32 ItemIndex, float& Out
 	}
 }
 
+float UGeometryCollectionComponent::GetMass() const
+{
+	float OutMass{ 0 };
+
+	if (RestCollection && RestCollection->GetGeometryCollection())
+	{
+		const FGeometryCollection& Collection = *RestCollection->GetGeometryCollection();
+		if (const TManagedArray<float>* CollectionMass = Collection.FindAttribute<float>(TEXT("Mass"), FTransformCollection::TransformGroup))
+		{
+			const int32 RootIndex = RestCollection->GetRootIndex();
+			if (CollectionMass->IsValidIndex(RootIndex))
+			{
+				OutMass = (*CollectionMass)[RootIndex];
+			}
+		}
+	}
+	return OutMass;
+}
+
+float UGeometryCollectionComponent::CalculateMass(FName BoneName)
+{
+	float OutMass{ 0 };
+
+	if (RestCollection && RestCollection->GetGeometryCollection())
+	{
+		const FGeometryCollection& Collection = *RestCollection->GetGeometryCollection();
+		if (const TManagedArray<float>* CollectionMass = Collection.FindAttribute<float>(TEXT("Mass"), FTransformCollection::TransformGroup))
+		{
+			const int32 BoneIndex = (BoneName == NAME_None)? RestCollection->GetRootIndex(): Collection.BoneName.Find(BoneName.ToString());
+			if (CollectionMass->IsValidIndex(BoneIndex))
+			{
+				OutMass = (*CollectionMass)[BoneIndex];
+			}
+		}
+	}
+	return OutMass;
+}
+
 bool UGeometryCollectionComponent::CalculateInnerSphere(int32 TransformIndex, UE::Math::TSphere<double>& SphereOut) const
 {
 	// Approximates the inscribed sphere. Returns false if no such sphere exists, if for instance the index is to an embedded geometry. 
diff --git a/Engine/Source/Runtime/Experimental/GeometryCollectionEngine/Public/GeometryCollection/GeometryCollectionComponent.h b/Engine/Source/Runtime/Experimental/GeometryCollectionEngine/Public/GeometryCollection/GeometryCollectionComponent.h
--- a/Engine/Source/Runtime/Experimental/GeometryCollectionEngine/Public/GeometryCollection/GeometryCollectionComponent.h
+++ b/Engine/Source/Runtime/Experimental/GeometryCollectionEngine/Public/GeometryCollection/GeometryCollectionComponent.h
@@ -669,6 +669,12 @@ public:
 	UFUNCTION(BlueprintCallable, Category = "ChaosPhysics")
 	GEOMETRYCOLLECTIONENGINE_API void GetMassAndExtents(int32 ItemIndex, float& OutMass, FBox& OutExtents);
 
+	/** Returns the mass of this component in kg. */
+	GEOMETRYCOLLECTIONENGINE_API virtual float GetMass() const override;
+
+	/** Returns the calculated mass in kg. This is not 100% exactly the mass physx will calculate, but it is very close ( difference < 0.1kg ). */
+	GEOMETRYCOLLECTIONENGINE_API virtual float CalculateMass(FName BoneName = NAME_None) override;
+
 	/** RestCollection */
 	UFUNCTION(BlueprintCallable, Category = "ChaosPhysics")
 	GEOMETRYCOLLECTIONENGINE_API void SetRestCollection(const UGeometryCollection * RestCollectionIn, bool bApplyAssetDefaults = true);
