Chaos debug : changed internal cluster color to purple to avoid confusion with other visualizations options

diff --git a/Engine/Source/Runtime/Experimental/Chaos/Private/Chaos/ChaosDebugDraw.cpp b/Engine/Source/Runtime/Experimental/Chaos/Private/Chaos/ChaosDebugDraw.cpp
--- a/Engine/Source/Runtime/Experimental/Chaos/Private/Chaos/ChaosDebugDraw.cpp
+++ b/Engine/Source/Runtime/Experimental/Chaos/Private/Chaos/ChaosDebugDraw.cpp
@@ -661,7 +661,7 @@ namespace Chaos
 					{
 						if (ClusteredParticle->InternalCluster())
 						{
-							ShapeColor = FColor::Red;
+							ShapeColor = FColor::Purple;
 						}
 					}
 				}
