PhysicsSolver: Logging: Downgrade log message to verbose to make cookcommandlet logs easier to read.
#rnx

diff --git a/Engine/Source/Runtime/Experimental/Chaos/Private/Chaos/Framework/PhysicsSolverBase.cpp b/Engine/Source/Runtime/Experimental/Chaos/Private/Chaos/Framework/PhysicsSolverBase.cpp
--- a/Engine/Source/Runtime/Experimental/Chaos/Private/Chaos/Framework/PhysicsSolverBase.cpp
+++ b/Engine/Source/Runtime/Experimental/Chaos/Private/Chaos/Framework/PhysicsSolverBase.cpp
@@ -214,7 +214,7 @@ namespace Chaos
 		, bStealAdvanceTasksForTesting(false)
 #endif
 	{
-		UE_LOG(LogChaos, Log, TEXT("FPhysicsSolverBase::AsyncDt:%f"), IsUsingAsyncResults() ? AsyncDt : -1);
+		UE_LOG(LogChaos, Verbose, TEXT("FPhysicsSolverBase::AsyncDt:%f"), IsUsingAsyncResults() ? AsyncDt : -1);
 
 		//If user is running with -PhysicsRunsOnGT override the cvar (doing it here to avoid parsing every time task is scheduled)
 		if(FParse::Param(FCommandLine::Get(), TEXT("PhysicsRunsOnGT")))
@@ -229,14 +229,14 @@ namespace Chaos
 		if (AsyncDt != FixedDt)
 		{
 			AccumulatedTime = 0;
-			UE_LOG(LogChaos, Log, TEXT("FPhysicsSolverBase::AsyncDt:%f"), IsUsingAsyncResults() ? AsyncDt : -1);
+			UE_LOG(LogChaos, Verbose, TEXT("FPhysicsSolverBase::AsyncDt:%f"), IsUsingAsyncResults() ? AsyncDt : -1);
 		}
 	}
 
 	void FPhysicsSolverBase::DisableAsyncMode()
 	{
 		AsyncDt = -1;
-		UE_LOG(LogChaos, Log, TEXT("FPhysicsSolverBase::AsyncDt:%f"), AsyncDt);
+		UE_LOG(LogChaos, Verbose, TEXT("FPhysicsSolverBase::AsyncDt:%f"), AsyncDt);
 	}
 
 
