#UE Assert that FPhysicsRwLock gets a TLS slot. Not obtaining one can result in deadlocks

diff --git a/Engine/Source/Runtime/Experimental/Chaos/Public/Framework/Threading.h b/Engine/Source/Runtime/Experimental/Chaos/Public/Framework/Threading.h
--- a/Engine/Source/Runtime/Experimental/Chaos/Public/Framework/Threading.h
+++ b/Engine/Source/Runtime/Experimental/Chaos/Public/Framework/Threading.h
@@ -726,14 +726,13 @@ FORCEINLINE void EnsureIsInGameThreadContext()
 		FPhysicsRwLock()
 		{
 			TlsSlot = FPlatformTLS::AllocTlsSlot();
+			check(FPlatformTLS::IsValidTlsSlot(TlsSlot));
 		}
 
 		~FPhysicsRwLock()
 		{
-			if (FPlatformTLS::IsValidTlsSlot(TlsSlot))
-			{
-				FPlatformTLS::FreeTlsSlot(TlsSlot);
-			}
+			check(FPlatformTLS::IsValidTlsSlot(TlsSlot));
+			FPlatformTLS::FreeTlsSlot(TlsSlot);
 		}
 
 		void ReadLock()
