Set editor to use lock for physics scenes that avoids TLS overflows
#rnx

diff --git a/Engine/Source/Runtime/Experimental/Chaos/Public/Framework/Threading.h b/Engine/Source/Runtime/Experimental/Chaos/Public/Framework/Threading.h
--- a/Engine/Source/Runtime/Experimental/Chaos/Public/Framework/Threading.h
+++ b/Engine/Source/Runtime/Experimental/Chaos/Public/Framework/Threading.h
@@ -33,8 +33,14 @@
 #define CHAOS_SCENE_LOCK_SIMPLE_MUTEX 4           // Just a critical section (not an RWLock). Provided For profiling/debugging only. Not recommended
 
 /** Controls the scene lock type. See above. */ 
-#ifndef CHAOS_SCENE_LOCK_TYPE
-#define CHAOS_SCENE_LOCK_TYPE CHAOS_SCENE_LOCK_FRWLOCK
+#if WITH_EDITOR
+	#ifndef CHAOS_SCENE_LOCK_TYPE
+	#define CHAOS_SCENE_LOCK_TYPE CHAOS_SCENE_LOCK_RWFIFO_CRITICALSECTION
+	#endif
+#else
+	#ifndef CHAOS_SCENE_LOCK_TYPE
+	#define CHAOS_SCENE_LOCK_TYPE CHAOS_SCENE_LOCK_FRWLOCK
+	#endif
 #endif
 
 
