Chaos - remove unnecessary includes to reduce excess compilation when changing headers

diff --git a/Engine/Source/Runtime/Experimental/Chaos/Private/PhysicsProxy/SkeletalMeshPhysicsProxy.cpp b/Engine/Source/Runtime/Experimental/Chaos/Private/PhysicsProxy/SkeletalMeshPhysicsProxy.cpp
--- a/Engine/Source/Runtime/Experimental/Chaos/Private/PhysicsProxy/SkeletalMeshPhysicsProxy.cpp
+++ b/Engine/Source/Runtime/Experimental/Chaos/Private/PhysicsProxy/SkeletalMeshPhysicsProxy.cpp
@@ -5,19 +5,11 @@
 #include "ChaosStats.h"
 #include "PhysicsSolver.h"
 
-#ifndef ADD_TRAITS_TO_SKELTAL_MESH_PROXY
-#define ADD_TRAITS_TO_SKELTAL_MESH_PROXY 0
-#endif
+// @todo(chaos): remove this file
 
 FSkeletalMeshPhysicsProxy::FSkeletalMeshPhysicsProxy(UObject* InOwner, const FInitFunc& InInitFunc)
 	: Base(InOwner)
-	, JointConstraints()
-	, NextInputProducerBuffer(nullptr)
-	, CurrentOutputConsumerBuffer(nullptr)
-	, bInitializedState(false)
-	, InitFunc(InInitFunc)
 {
-	check(IsInGameThread())
 }
 
 FSkeletalMeshPhysicsProxy::~FSkeletalMeshPhysicsProxy()
@@ -26,80 +18,19 @@ FSkeletalMeshPhysicsProxy::~FSkeletalMeshPhysicsProxy()
 
 void FSkeletalMeshPhysicsProxy::Initialize()
 {
-	check(IsInGameThread());
-
-	// InitFunc builds a bone hierarchy from the UPhysicsAsset on a USkeletalMeshComponent.
-	// Transforms will be set from the FReferenceSkeleton reference bone poses.
-	InitFunc(Parameters);
-
-	// Set InitializedState to false so that rigid bodies will now be created.
-	Reset();
 }
 
 void FSkeletalMeshPhysicsProxy::Reset()
 {
-	bInitializedState = false;
 }
 
 bool FSkeletalMeshPhysicsProxy::IsSimulating() const
 {
-	return Parameters.bSimulating;
+	return false;
 }
 
 void FSkeletalMeshPhysicsProxy::UpdateKinematicBodiesCallback(const FParticlesType& Particles, const float Dt, const float Time, FKinematicProxy& Proxy)
 {
-	// @todo(ccaulfield): Handle changes in body state
-
-	if (bInitializedState && Parameters.bSimulating)
-	{
-		// The triple buffer will return null if there's no new data for the consumer thread.
-		const FSkeletalMeshPhysicsProxyInputs* InputBuffer = InputBuffers.ExchangeConsumerBuffer();
-		if (InputBuffer == nullptr)
-		{
-			return;
-		}
-
-#if TODO_REIMPLEMENT_KINEMATIC_PROXY
-
-		// @todo(ccaulfield): Optimize. Keep a list of kinematic particles?
-		const int32 MaxRigidBodies = RigidBodyIds.Num();
-		Proxy.Ids.Reset(MaxRigidBodies);
-		Proxy.Position.Reset(MaxRigidBodies);
-		Proxy.NextPosition.Reset(MaxRigidBodies);
-		Proxy.Rotation.Reset(MaxRigidBodies);
-		Proxy.NextRotation.Reset(MaxRigidBodies);
-
-		int32 RigidBodyIndex = 0;
-		TArray<TUniquePtr<FAnalyticImplicitGroup>>& Groups = Parameters.BoneHierarchy.GetAnalyticShapeGroups();
-		for (TUniquePtr<FAnalyticImplicitGroup>& Group : Groups)
-		{
-			if (!Group->NumStructures())
-			{
-				continue;
-			}
-
-			const int32 RigidBodyId = Group->GetRigidBodyId();
-			const int32 BoneIndex = Group->GetBoneIndex();
-			const int32 TransformIndex = Parameters.BoneHierarchy.GetTransformIndex(BoneIndex);
-			bool bIsKinematic = (Particles.ObjectState(RigidBodyId) == Chaos::EObjectStateType::Kinematic);
-			if (bIsKinematic)
-			{
-				const FTransform& Transform = InputBuffer->Transforms[TransformIndex];
-				const FVector& LinearVelocity = InputBuffer->LinearVelocities[TransformIndex];
-				const FVector& AngularVelocity = InputBuffer->AngularVelocities[TransformIndex];
-
-				Proxy.Ids.Add(RigidBodyId);
-				Proxy.Position.Add(Chaos::FVec3(Transform.GetTranslation()));
-				Proxy.NextPosition.Add(Proxy.Position[RigidBodyIndex] + Chaos::FVec3(LinearVelocity) * Dt);
-				Proxy.Rotation.Add(Chaos::TRotation<float, 3>(Transform.GetRotation().GetNormalized()));
-				Proxy.NextRotation.Add(Chaos::TRotation<float, 3>::IntegrateRotationWithAngularVelocity(Proxy.Rotation[RigidBodyIndex], AngularVelocity, Dt));
-				++RigidBodyIndex;
-			}
-		}
-
-#endif
-
-	}
 }
 
 void FSkeletalMeshPhysicsProxy::StartFrameCallback(const float InDt, const float InTime)
@@ -112,248 +43,6 @@ void FSkeletalMeshPhysicsProxy::EndFrameCallback(const float InDt)
 
 void FSkeletalMeshPhysicsProxy::CreateRigidBodyCallback(FParticlesType& Particles)
 {
-	if (!bInitializedState && Parameters.bSimulating)
-	{
-#if TODO_REIMPLEMENT_RIGID_CLUSTERING
-		// @todo(ccaulfield): should not have to do this cluster thing here...
-		Chaos::TArrayCollectionArray<int32>& ClusterGroupIndex = GetSolver()->GetRigidClustering().GetClusterGroupIndexArray();
-#endif
-#if 0
-		Chaos::TArrayCollectionArray<float>& StrainArray = GetSolver()->GetRigidClustering().GetStrainArray();
-#endif
-		Parameters.BoneHierarchy.PrepareAnimWorldSpaceTransforms();
-		RigidBodyIds.Reset();
-		RigidBodyIds.Reserve(Parameters.BoneHierarchy.GetBoneIndices().Num());
-		for (const int32 BoneIndex : Parameters.BoneHierarchy.GetBoneIndices())
-		{
-			FAnalyticImplicitGroup* Group = Parameters.BoneHierarchy.GetAnalyticShapeGroup(BoneIndex);
-			if (!Group || !Group->NumStructures())
-			{
-				continue;
-			}
-
-			const FTransform* WorldTransform = Parameters.BoneHierarchy.GetAnimWorldSpaceTransformsForBone(BoneIndex);
-			check(WorldTransform);
-
-			Chaos::FReal TotalMass = 0.;
-			const Chaos::FReal DensityKgCm3 = Parameters.Density / 1000.f;
-			Chaos::FMassProperties MassProperties = Group->BuildMassProperties(DensityKgCm3, TotalMass);
-
-			const bool DoCollisionGeom = Parameters.CollisionType == ECollisionTypeEnum::Chaos_Surface_Volumetric;
-			TArray<Chaos::FVec3>* SamplePoints = nullptr;
-			Chaos::FAABB3 SamplePointsBBox = Chaos::FAABB3::EmptyAABB();
-			if (DoCollisionGeom)
-			{
-				SamplePoints =
-					Group->BuildSamplePoints(
-						Parameters.ParticlesPerUnitArea,
-						Parameters.MinNumParticles,
-						Parameters.MaxNumParticles);
-				for (const Chaos::FVec3& Point : *SamplePoints)
-					SamplePointsBBox.GrowToInclude(Point);
-			}
-			//Chaos::FTriangleMesh TriMesh = Chaos::FTriangleMesh::GetConvexHullFromParticles(
-			//	TArrayView<const Chaos::TVector<T, 3>>(SamplePoints));
-
-			TUniquePtr<Chaos::FImplicitObject> ImplicitObject(Group->BuildSimImplicitObject());
-			check(ImplicitObject);
-
-			if (SamplePoints)
-			{
-				auto PointVolumeRegistrationCheck = 
-					[&](const TArray<Chaos::FVec3> &InSamplePoints, const Chaos::FImplicitObject &InImplicitObject, const float InTolerance) -> bool
-					{
-						for (int32 i = 0; i < InSamplePoints.Num(); i++)
-						{
-							const auto& Pt = InSamplePoints[i];
-							const auto Phi = InImplicitObject.SignedDistance(Pt);
-							if (FMath::Abs(Phi) > InTolerance)
-							{
-								return false;
-							}
-						}
-						return true; 
-					};
-				// Precision gets worse in opt builds, so only do the phi test in debug.
-				checkSlow(PointVolumeRegistrationCheck(*SamplePoints, *ImplicitObject, 0.01f));
-				auto PointBBoxCheck =
-					[&](const auto &InSamplePoints, const Chaos::FAABB3 &BBox) -> bool
-					{
-						for (int32 i = 0; i < InSamplePoints.Num(); i++)
-						{
-							const auto& Pt = InSamplePoints[i];
-							if (!BBox.Contains(Pt))
-							{
-								return false;
-							}
-						}
-						return true;
-					};
-				// BBox tests should still work though.
-				check(PointBBoxCheck(*SamplePoints, ImplicitObject->BoundingBox()));
-			}
-
-			const Chaos::FAABB3& BBox = ImplicitObject->BoundingBox();
-
-			const FVector Scale = WorldTransform->GetScale3D();
-			const FVector &CenterOfMass = MassProperties.CenterOfMass;
-
-			FBox Bounds(BBox.Min(), BBox.Max());
-			Bounds = Bounds.InverseTransformBy(FTransform(CenterOfMass));
-			Bounds.Min *= Scale;
-			Bounds.Max *= Scale;
-			checkSlow(Bounds.GetVolume() > FLT_EPSILON);
-
-			const int32 RigidBodyId = Particles.Size();
-			Particles.AddParticles(1);
-			RigidBodyIds.Add(RigidBodyId);
-			Group->SetRigidBodyId(RigidBodyId);
-#if 0
-			ClusterGroupIndex[RigidBodyId] = Parameters.ClusterGroupIndex;
-			StrainArray[RigidBodyId] = Parameters.DamageThreshold;
-#else
-#if TODO_REIMPLEMENT_RIGID_CLUSTERING
-			ClusterGroupIndex[RigidBodyId] = 0;
-#endif
-#endif
-
-#if TODO_REIMPLEMENT_SET_PHYSICS_MATERIAL
-			GetSolver()->SetPhysicsMaterial(RigidBodyId, Parameters.PhysicalMaterial);
-#endif
-			Particles.DisabledRef(RigidBodyId) = false;
-
-			//Particles.X(RigidBodyId) = WorldTransform->TransformPosition(CenterOfMass);
-			Particles.X(RigidBodyId) = WorldTransform->GetTranslation();
-			check(!FMath::IsNaN(Particles.X(RigidBodyId)[0]) && !FMath::IsNaN(Particles.X(RigidBodyId)[1]) && !FMath::IsNaN(Particles.X(RigidBodyId)[2]));
-			Particles.V(RigidBodyId) = Chaos::FVec3(0.f);
-			//Particles.R(RigidBodyId) = (WorldTransform->TransformRotation(MassProperties.RotationOfMass)).GetNormalized();
-			Particles.R(RigidBodyId) = WorldTransform->GetRotation().GetNormalized();
-			Particles.W(RigidBodyId) = Chaos::FVec3(0.f);
-			Particles.P(RigidBodyId) = Particles.X(RigidBodyId);
-			Particles.Q(RigidBodyId) = Particles.R(RigidBodyId);
-
-			Particles.M(RigidBodyId) = FMath::Clamp(TotalMass, (Chaos::FReal)Parameters.MinMass, (Chaos::FReal)Parameters.MaxMass);
-			Particles.InvM(RigidBodyId) = Particles.M(RigidBodyId) < UE_KINDA_SMALL_NUMBER ? UE_BIG_NUMBER/2 : 1.f / Particles.M(RigidBodyId);
-			Particles.I(RigidBodyId) = Chaos::TVec3<Chaos::FRealSingle>((Chaos::FRealSingle)MassProperties.InertiaTensor.M[0][0], (Chaos::FRealSingle)MassProperties.InertiaTensor.M[1][1], (Chaos::FRealSingle)MassProperties.InertiaTensor.M[2][2]);
-			Particles.InvI(RigidBodyId) = Chaos::TVec3<Chaos::FRealSingle>(1.f / (Chaos::FRealSingle)MassProperties.InertiaTensor.M[0][0], 1.f / (Chaos::FRealSingle)MassProperties.InertiaTensor.M[1][1], 1.f / (Chaos::FRealSingle)MassProperties.InertiaTensor.M[2][2]);
-
-			Particles.CollisionGroup(RigidBodyId) = Parameters.CollisionGroup;
-			
-			// Set implicit surface
-			const Chaos::EImplicitObjectType ShapeType = ImplicitObject->GetType();
-			if (ShapeType == Chaos::ImplicitObjectType::LevelSet)
-			{
-				FCollisionStructureManager::UpdateImplicitFlags(
-					ImplicitObject.Get(), ECollisionTypeEnum::Chaos_Surface_Volumetric);
-				Particles.SetDynamicGeometry(RigidBodyId, MoveTemp(ImplicitObject));
-			}
-			else
-			{
-				Particles.SetDynamicGeometry(RigidBodyId, MoveTemp(ImplicitObject));
-			}
-
-			if (Parameters.CollisionType == ECollisionTypeEnum::Chaos_Surface_Volumetric)
-			{
-				if (SamplePoints)
-				{
-					Particles.CollisionParticlesInitIfNeeded(RigidBodyId);
-
-					Particles.CollisionParticles(RigidBodyId)->Resize(0);
-					Particles.CollisionParticles(RigidBodyId)->AddParticles(SamplePoints->Num());
-					for (int Vi = 0; Vi < SamplePoints->Num(); Vi++)
-					{
-						Particles.CollisionParticles(RigidBodyId)->X(Vi) = (*SamplePoints)[Vi];
-					}
-					if (Particles.CollisionParticles(RigidBodyId)->Size())
-					{
-						Particles.CollisionParticles(RigidBodyId)->UpdateAccelerationStructures();
-					}
-				}
-			}
-
-			// Set the simulation state of the body.
-			// If set to "user defined" we inherit the owner's state
-			// @todo(ccaulfield): formalize the state inheritence - currently hijacking UserDefined...
-			Chaos::EObjectStateType RigidBodyState = Chaos::EObjectStateType::Kinematic;
-			EObjectStateTypeEnum SourceState = Group->GetRigidBodyState();
-			if (SourceState == EObjectStateTypeEnum::Chaos_Object_UserDefined)
-			{
-				SourceState = Parameters.ObjectType;
-			}
-			switch (SourceState)
-			{
-			case EObjectStateTypeEnum::Chaos_Object_Sleeping:
-				RigidBodyState = Chaos::EObjectStateType::Sleeping;
-				break;
-			case EObjectStateTypeEnum::Chaos_Object_Kinematic:
-				RigidBodyState = Chaos::EObjectStateType::Kinematic;
-				break;
-			case EObjectStateTypeEnum::Chaos_Object_Static:
-				RigidBodyState = Chaos::EObjectStateType::Static;
-				break;
-			default:
-				RigidBodyState = Chaos::EObjectStateType::Dynamic;
-				break;
-			}
-			Particles.SetObjectState(RigidBodyId, RigidBodyState);
-
-			// Create joints
-			if (RigidBodyState == Chaos::EObjectStateType::Dynamic)
-			{
-				const int32 ParentBoneIndex = Group->GetParentBoneIndex();
-				if (ParentBoneIndex != INDEX_NONE)
-				{
-					const FAnalyticImplicitGroup* ParentGroup = Parameters.BoneHierarchy.GetAnalyticShapeGroup(ParentBoneIndex);
-					if (ParentGroup && (ParentGroup->GetRigidBodyId() != INDEX_NONE))
-					{
-						const int32 ParentRigidBodyId = ParentGroup->GetRigidBodyId();
-						const FTransform& ParentWorldTransform = *Parameters.BoneHierarchy.GetAnimWorldSpaceTransformsForBone(ParentBoneIndex);
-
-						const Chaos::TVector<int32, 2> ConstrainedParticleIndices = { RigidBodyId, ParentRigidBodyId };
-						const Chaos::FVec3 JointPosition = ParentWorldTransform.GetTranslation();
-#if TODO_REWRITE_ALL_CONSTRAINT_ADDS_TO_USE_HANDLES
-						JointConstraints.AddConstraint(Particles, ConstrainedParticleIndices, JointPosition);
-#endif
-					}
-				}
-			}
-
-
-			UE_LOG(USkeletalMeshSimulationComponentLogging, Log,
-				TEXT("FSkeletalMeshPhysicsProxy::CreateRigidBodyCallback() - this: %p - "
-					"Created rigid body %d for bone %d.\n"
-					"    center of mass: (%g, %g, %g), Bone pos: (%g, %g, %g), RB position: (%g, %g, %g)\n"
-					"    implicit shape domain: [(%g %g %g), (%g %g %g)] - volume: %g.\n"
-					"    %d local space collision points, domain: [(%g %g %g), (%g %g %g)] - volume: %g."),
-
-				this, 
-				RigidBodyId, Group->GetBoneIndex(),
-
-				CenterOfMass[0], CenterOfMass[1], CenterOfMass[2], 
-				WorldTransform->GetTranslation()[0], WorldTransform->GetTranslation()[1], WorldTransform->GetTranslation()[2], 
-				Particles.X(RigidBodyId)[0], Particles.X(RigidBodyId)[2], Particles.X(RigidBodyId)[2], 
-
-				SamplePointsBBox.Min()[0], SamplePointsBBox.Min()[1], SamplePointsBBox.Min()[2],
-				SamplePointsBBox.Max()[0], SamplePointsBBox.Max()[1], SamplePointsBBox.Max()[2],
-				SamplePointsBBox.Extents().Product(),
-
-				SamplePoints?SamplePoints->Num():0,
-				BBox.Min()[0], BBox.Min()[1], BBox.Min()[2],
-				BBox.Max()[0], BBox.Max()[1], BBox.Max()[2],
-				BBox.Extents().Product());
-		}
-
-	ensure(false);	//GetSolver needs a trait, but we don't know it - this class changing anyway
-#if ADD_TRAITS_TO_SKELTAL_MESH_PROXY
-		// Add joints to the scene
-		if (JointConstraints.NumConstraints() > 0)
-		{
-			GetSolver()->GetEvolution()->AddConstraintRule(&JointConstraintsRule);
-		}
-
-		bInitializedState = true;
-#endif
-	}
 }
 
 void FSkeletalMeshPhysicsProxy::ParameterUpdateCallback(FParticlesType& InParticles, const float InTime)
@@ -367,17 +56,6 @@ void FSkeletalMeshPhysicsProxy::AddForceCallback(FParticlesType& InParticles, co
 
 void FSkeletalMeshPhysicsProxy::BindParticleCallbackMapping(Chaos::TArrayCollectionArray<PhysicsProxyWrapper>& PhysicsProxyReverseMap, Chaos::TArrayCollectionArray<int32>& ParticleIDReverseMap)
 {
-	if (bInitializedState)
-	{
-		for (int32 Index = 0; Index < RigidBodyIds.Num(); Index++)
-		{
-			if (RigidBodyIds[Index] != INDEX_NONE)
-			{
-				PhysicsProxyReverseMap[RigidBodyIds[Index]] = { this, EPhysicsProxyType::SkeletalMeshType };
-				ParticleIDReverseMap[RigidBodyIds[Index]] = Index;
-			}
-		}
-	}
 }
 
 void FSkeletalMeshPhysicsProxy::SyncBeforeDestroy() 
@@ -385,131 +63,21 @@ void FSkeletalMeshPhysicsProxy::SyncBeforeDestroy()
 
 void FSkeletalMeshPhysicsProxy::OnRemoveFromScene()
 {
-	ensure(false);	//GetSolver needs a trait, but we don't know it - this class changing anyway
-#if ADD_TRAITS_TO_SKELTAL_MESH_PROXY
-	// Disable the particle we added
-	Chaos::FPhysicsSolver* CurrSolver = GetSolver();
-	if(CurrSolver && RigidBodyIds.Num())
-	{
-#if TODO_REIMPLEMENT_GET_RIGID_PARTICLES
-		Chaos::FPhysicsSolver::FParticlesType &Particles = CurrSolver->GetRigidParticles();
-		// #BG TODO Special case here because right now we reset/realloc the evolution per geom component
-		// in endplay which clears this out. That needs to not happen and be based on world shutdown
-		if (Particles.Size() == 0)
-		{
-			return;
-		}
-
-		for (const int32 Id : RigidBodyIds)
-		{
-			if(static_cast<uint32>(Id) < Particles.Size())
-			{
-				Particles.DisabledRef(Id) = true;
-				CurrSolver->GetRigidClustering().GetTopLevelClusterParents().Remove(Id);
-			}
-		}
-#endif
-	}
-#endif
 }
 
 void FSkeletalMeshPhysicsProxy::BufferPhysicsResults()
 {
-	using namespace Chaos;
-	// PHYSICS THREAD
-
-	ensure(false);	//GetSolver needs a trait, but we don't know it - this class changing anyway
-#if ADD_TRAITS_TO_SKELTAL_MESH_PROXY
-	// @todo(ccaulfield)
-	if (GetSolver() != nullptr)
-	{
-		// Copy the world-space state of the simulated bones into the output buffer
-		int NumBones = Parameters.BoneHierarchy.GetBoneIndices().Num();
-		FSkeletalMeshPhysicsProxyOutputs& Outputs = OutputBuffers.GetPhysicsDataForWrite();
-		Outputs.Transforms.SetNumZeroed(NumBones);
-		Outputs.LinearVelocities.SetNumZeroed(NumBones);
-		Outputs.AngularVelocities.SetNumZeroed(NumBones);
-
-#if TODO_REIMPLEMENT_GET_RIGID_PARTICLES
-		FPhysicsSolver::FParticlesType& Particles = GetSolver()->GetRigidParticles();
-		for (const int32 BoneIndex : Parameters.BoneHierarchy.GetBoneIndices())
-		{
-			FAnalyticImplicitGroup* ShapeGroup = Parameters.BoneHierarchy.GetAnalyticShapeGroup(BoneIndex);
-			if (ShapeGroup != nullptr)
-			{
-				int32 RigidBodyId = ShapeGroup->GetRigidBodyId();
-				if (RigidBodyId >= 0)
-				{
-					const FTransform BodyTransform = FTransform(Particles.Q(RigidBodyId), Particles.P(RigidBodyId));
-					int32 TransformIndex = Parameters.BoneHierarchy.GetTransformIndex(BoneIndex);
-					Outputs.Transforms[TransformIndex] = BodyTransform;
-					Outputs.LinearVelocities[TransformIndex] = Particles.V(RigidBodyId);
-					Outputs.AngularVelocities[TransformIndex] = Particles.W(RigidBodyId);
-				}
-			}
-		}
-#endif
-	}
-#endif
 }
 
 void FSkeletalMeshPhysicsProxy::FlipBuffer()
 {
-	// PHYSICS THREAD
-
-	OutputBuffers.Flip();
 }
 
 bool FSkeletalMeshPhysicsProxy::PullFromPhysicsState(const int32 SolverSyncTimestamp)
 {
-	// GAME THREAD
-	// @todo(ccaulfield): this needs to update the Component-Space pose in BoneHierarchy and update non-simulated children
-
-	CurrentOutputConsumerBuffer = &OutputBuffers.GetGameDataForRead();
 	return true;
 }
 
 void FSkeletalMeshPhysicsProxy::CaptureInputs(const float Dt, const FInputFunc& InputFunc)
 {
-	// GAME THREAD
-	SCOPE_CYCLE_COUNTER(STAT_SkelMeshUpdateAnim);
-
-	// Get a new buffer to write the input data to
-	if (NextInputProducerBuffer == nullptr)
-	{
-		NextInputProducerBuffer = InputBuffers.ExchangeProducerBuffer();
-	}
-
-	// Fill in the input data and switch the buffers if we had new data
-	if (InputFunc(Dt, Parameters))
-	{
-		const FBoneHierarchy& BoneHierarchy = Parameters.BoneHierarchy;
-		FSkeletalMeshPhysicsProxyInputs& PhysicsInputs = *NextInputProducerBuffer;
-		PhysicsInputs.Transforms.SetNumZeroed(BoneHierarchy.GetBoneIndices().Num());
-		PhysicsInputs.LinearVelocities.SetNumZeroed(BoneHierarchy.GetBoneIndices().Num());
-		PhysicsInputs.AngularVelocities.SetNumZeroed(BoneHierarchy.GetBoneIndices().Num());
-		for (int32 BoneIndex : BoneHierarchy.GetBoneIndices())
-		{
-			const int32 TransformIndex = BoneHierarchy.GetTransformIndex(BoneIndex);
-			if (TransformIndex != INDEX_NONE)
-			{
-				const FTransform& BoneWorldTransform = *BoneHierarchy.GetAnimWorldSpaceTransformsForBone(BoneIndex);
-				const FTransform* PrevBoneWorldTransform = BoneHierarchy.GetPrevAnimWorldSpaceTransformForBone(BoneIndex);
-				if ((Dt > UE_SMALL_NUMBER) && PrevBoneWorldTransform)
-				{
-					PhysicsInputs.Transforms[TransformIndex] = *PrevBoneWorldTransform;
-					PhysicsInputs.LinearVelocities[TransformIndex] = (BoneWorldTransform.GetTranslation() - PrevBoneWorldTransform->GetTranslation()) / Dt;
-					PhysicsInputs.AngularVelocities[TransformIndex] = Chaos::FRotation3::CalculateAngularVelocity(PrevBoneWorldTransform->GetRotation(), BoneWorldTransform.GetRotation(), Dt);
-				}
-				else
-				{
-					PhysicsInputs.Transforms[TransformIndex] = BoneWorldTransform;
-					PhysicsInputs.LinearVelocities[TransformIndex] = FVector::ZeroVector;
-					PhysicsInputs.AngularVelocities[TransformIndex] = FVector::ZeroVector;
-				}
-			}
-		}
-
-		NextInputProducerBuffer = InputBuffers.ExchangeProducerBuffer();
-	}
 }
diff --git a/Engine/Source/Runtime/Experimental/Chaos/Private/PhysicsProxy/StaticMeshPhysicsObject.cpp b/Engine/Source/Runtime/Experimental/Chaos/Private/PhysicsProxy/StaticMeshPhysicsObject.cpp
--- a/Engine/Source/Runtime/Experimental/Chaos/Private/PhysicsProxy/StaticMeshPhysicsObject.cpp
+++ b/Engine/Source/Runtime/Experimental/Chaos/Private/PhysicsProxy/StaticMeshPhysicsObject.cpp
@@ -10,89 +10,32 @@
 #include "Chaos/Serializable.h"
 #include "ChaosStats.h"
 
-#ifndef ADD_TRAITS_TO_STATIC_MESH_PROXY
-#define ADD_TRAITS_TO_STATIC_MESH_PROXY 0
-#endif
+// @todo(chaos): remove this file
 
 FStaticMeshPhysicsProxy::FStaticMeshPhysicsProxy(UObject* InOwner, FCallbackInitFunc InInitFunc, FSyncDynamicFunc InSyncFunc)
 	: Base(InOwner)
-	, bInitializedState(false)
-	, RigidBodyId(INDEX_NONE)
-	, CenterOfMass(FVector::ZeroVector)
-	, Scale(FVector::ZeroVector)
-	, SimTransform(FTransform::Identity)
-	, InitialiseCallbackParamsFunc(InInitFunc)
-	, SyncDynamicTransformFunc(InSyncFunc)
 {
-	check(IsInGameThread());
-
-	Results.Get(0) = FTransform::Identity;
-	Results.Get(1) = FTransform::Identity;
 }
 
 void FStaticMeshPhysicsProxy::Initialize()
 {
-	check(IsInGameThread());
-
-	// Safe here - we're not in the solver yet if we're creating callbacks
-	Results.Get(0) = FTransform::Identity;
-	Results.Get(1) = FTransform::Identity;
-
-	InitialiseCallbackParamsFunc(Parameters);
-	Parameters.TargetTransform = &SimTransform;
-
-	Reset();
 }
 
 void FStaticMeshPhysicsProxy::Reset()
 {
-	bInitializedState = false;
 }
 
 void FStaticMeshPhysicsProxy::BufferKinematicUpdate(const FPhysicsProxyKinematicUpdate& InParamUpdate)
 {
-	BufferedKinematicUpdate = InParamUpdate;
-	bPendingKinematicUpdate = true;
 };
 
 bool FStaticMeshPhysicsProxy::IsSimulating() const
 {
-	return Parameters.bSimulating;
+	return false;
 }
 
 void FStaticMeshPhysicsProxy::UpdateKinematicBodiesCallback(const FParticlesType& Particles, const float Dt, const float Time, FKinematicProxy& Proxy)
 {
-	bool IsControlled = Parameters.ObjectType == EObjectStateTypeEnum::Chaos_Object_Kinematic;
-	if(IsControlled && Parameters.bSimulating)
-	{
-#if TODO_REIMPLEMENT_KINEMATIC_PROXY
-		bool bFirst = !Proxy.Ids.Num();
-		if(bFirst)
-		{
-			Proxy.Ids.Add(RigidBodyId);
-			Proxy.Position.SetNum(1);
-			Proxy.NextPosition.SetNum(1);
-			Proxy.Rotation.SetNum(1);
-			Proxy.NextRotation.SetNum(1);
-
-			const FTransform& Transform = Parameters.InitialTransform;
-			Proxy.Position[0] = Chaos::FVec3(Transform.GetTranslation());
-			Proxy.NextPosition[0] = Proxy.Position[0] + Chaos::FVec3(Parameters.InitialLinearVelocity) * Dt;
-			Proxy.Rotation[0] = Chaos::TRotation<float, 3>(Transform.GetRotation().GetNormalized());
-			Proxy.NextRotation[0] = Proxy.Rotation[0];
-		}
-
-		if (bPendingKinematicUpdate)
-		{
-			Proxy.Position[0] = Particles.X(RigidBodyId);
-			Proxy.NextPosition[0] = BufferedKinematicUpdate.NewTransform.GetTranslation();
-			Proxy.Rotation[0] = Particles.R(RigidBodyId);
-			Proxy.NextRotation[0] = BufferedKinematicUpdate.NewTransform.GetRotation().GetNormalized();
-
-			bPendingKinematicUpdate = false;
-		}
-#endif
-	}
 }
 
 void FStaticMeshPhysicsProxy::StartFrameCallback(const float InDt, const float InTime)
@@ -102,278 +45,41 @@ void FStaticMeshPhysicsProxy::StartFrameCallback(const float InDt, const float I
 
 void FStaticMeshPhysicsProxy::EndFrameCallback(const float InDt)
 {
-	bool IsControlled = Parameters.ObjectType == EObjectStateTypeEnum::Chaos_Object_Kinematic;
-	if(Parameters.bSimulating && !IsControlled && Parameters.TargetTransform)
-	{
-#if TODO_REIMPLEMENT_GET_RIGID_PARTICLES
-		const Chaos::TPBDRigidParticles<float, 3>& Particles = GetSolver()->GetRigidParticles();
-
-		Parameters.TargetTransform->SetTranslation((FVector)Particles.X(RigidBodyId));
-		Parameters.TargetTransform->SetRotation((FQuat)Particles.R(RigidBodyId));
-#endif
-	}
 }
 
 void FStaticMeshPhysicsProxy::BindParticleCallbackMapping(Chaos::TArrayCollectionArray<PhysicsProxyWrapper> & PhysicsProxyReverseMap, Chaos::TArrayCollectionArray<int32> & ParticleIDReverseMap)
 {
-	if (bInitializedState)
-	{
-		PhysicsProxyReverseMap[RigidBodyId] = { this, EPhysicsProxyType::StaticMeshType };
-		ParticleIDReverseMap[RigidBodyId] = 0;
-	}
 }
 
 void FStaticMeshPhysicsProxy::CreateRigidBodyCallback(FParticlesType& Particles)
 {
-	if(!bInitializedState && Parameters.bSimulating)
-	{
-		RigidBodyId = Particles.Size();
-		Particles.AddParticles(1);
-
-		FBox Bounds(ForceInitToZero);
-		if (Parameters.ShapeType == EImplicitTypeEnum::Chaos_Implicit_LevelSet)
-		{
-			// Build implicit surface
-			const ECollisionTypeEnum CollisionType = ECollisionTypeEnum::Chaos_Surface_Volumetric;
-			check(sizeof(Chaos::FVec3) == sizeof(FVector));
-
-			//TODO: this should all be done on the asset instead of in the callback. Duplicates all mesh vertices every time
-			Chaos::FParticles MeshParticles(MoveTemp(Parameters.MeshVertexPositions));
-
-			TUniquePtr<Chaos::FTriangleMesh> TriangleMesh(new Chaos::FTriangleMesh(MoveTemp(Parameters.TriIndices)));
-			Chaos::FErrorReporter ErrorReporter(Parameters.Name + " | RigidBodyId: " + FString::FromInt(RigidBodyId));;
-			Particles.SetDynamicGeometry(
-				RigidBodyId,
-				TUniquePtr<Chaos::FImplicitObject>(
-					FCollisionStructureManager::NewImplicit(
-						ErrorReporter,
-						MeshParticles, *TriangleMesh,
-						Bounds,
-						FVector::Distance(FVector(0.f, 0.f, 0.f), 
-						Bounds.GetExtent()) * 0.5f,
-						Parameters.MinRes, 
-						Parameters.MaxRes, 
-						0.f,
-						CollisionType, 
-						Parameters.ShapeType)));
-
-			if (!ensure(Parameters.MeshVertexPositions.Size()))
-			{
-				Parameters.MeshVertexPositions.AddParticles(1);
-				Parameters.MeshVertexPositions.X(0) = Chaos::FVec3(0, 0, 0);
-			}
-
-			for(uint32 i=0; i < Parameters.MeshVertexPositions.Size(); ++i)
-			{
-				Bounds += Parameters.MeshVertexPositions.X(i);
-			}
-		}
-		else if (Parameters.ShapeType == EImplicitTypeEnum::Chaos_Implicit_Sphere)
-		{
-			Chaos::TSphere<Chaos::FReal,3>* Sphere = new Chaos::TSphere<Chaos::FReal, 3>(Chaos::FVec3(0), Parameters.ShapeParams.SphereRadius);
-			const Chaos::FAABB3 BBox = Sphere->BoundingBox();
-			Bounds.Min = BBox.Min();
-			Bounds.Max = BBox.Max();
-			Particles.SetDynamicGeometry(RigidBodyId, TUniquePtr<Chaos::FImplicitObject>(Sphere));
-			if (!Parameters.MeshVertexPositions.Size())
-			{
-				float Radius = Parameters.ShapeParams.SphereRadius;
-
-				Parameters.MeshVertexPositions.AddParticles(6);
-				Parameters.MeshVertexPositions.X(0) = Chaos::FVec3(Radius, 0, 0);
-				Parameters.MeshVertexPositions.X(1) = Chaos::FVec3(-Radius, 0, 0);
-				Parameters.MeshVertexPositions.X(2) = Chaos::FVec3(0, Radius, Radius);
-				Parameters.MeshVertexPositions.X(3) = Chaos::FVec3(0, -Radius, Radius);
-				Parameters.MeshVertexPositions.X(4) = Chaos::FVec3(0, -Radius, -Radius);
-				Parameters.MeshVertexPositions.X(5) = Chaos::FVec3(0, Radius, -Radius);
-			}
-		}
-		else if (Parameters.ShapeType == EImplicitTypeEnum::Chaos_Implicit_Box)
-		{
-			Chaos::FVec3 HalfExtents = Parameters.ShapeParams.BoxExtents / static_cast<Chaos::FReal>(2);
-			Chaos::TBox<Chaos::FReal,3>* Box = new Chaos::TBox<Chaos::FReal, 3>(-HalfExtents, HalfExtents);
-			Bounds.Min = Box->Min();
-			Bounds.Max = Box->Max();
-			Particles.SetDynamicGeometry(RigidBodyId, TUniquePtr<Chaos::FImplicitObject>(Box));
-			if (!Parameters.MeshVertexPositions.Size())
-			{
-				Chaos::FVec3 x1(-HalfExtents);
-				Chaos::FVec3 x2(HalfExtents);
-
-				Parameters.MeshVertexPositions.AddParticles(8);
-				Parameters.MeshVertexPositions.X(0) = Chaos::FVec3(x1.X, x1.Y, x1.Z);
-				Parameters.MeshVertexPositions.X(1) = Chaos::FVec3(x1.X, x1.Y, x2.Z);
-				Parameters.MeshVertexPositions.X(2) = Chaos::FVec3(x1.X, x2.Y, x1.Z);
-				Parameters.MeshVertexPositions.X(3) = Chaos::FVec3(x2.X, x1.Y, x1.Z);
-				Parameters.MeshVertexPositions.X(4) = Chaos::FVec3(x2.X, x2.Y, x2.Z);
-				Parameters.MeshVertexPositions.X(5) = Chaos::FVec3(x2.X, x2.Y, x1.Z);
-				Parameters.MeshVertexPositions.X(6) = Chaos::FVec3(x2.X, x1.Y, x2.Z);
-				Parameters.MeshVertexPositions.X(7) = Chaos::FVec3(x1.X, x2.Y, x2.Z);
-			}
-		}
-		else if (Parameters.ShapeType == EImplicitTypeEnum::Chaos_Implicit_Capsule)
-		{
-			Chaos::FVec3 x1(0, -Parameters.ShapeParams.CapsuleHalfHeightAndRadius.X, 0);
-			Chaos::FVec3 x2(0, Parameters.ShapeParams.CapsuleHalfHeightAndRadius.X, 0);
-			Chaos::FCapsule* Capsule = new Chaos::FCapsule(x1, x2, Parameters.ShapeParams.CapsuleHalfHeightAndRadius.Y);
-			const Chaos::FAABB3 BBox = Capsule->BoundingBox();
-			Bounds.Min = BBox.Min();
-			Bounds.Max = BBox.Max();
-			Particles.SetDynamicGeometry(RigidBodyId, TUniquePtr<Chaos::FImplicitObject>(Capsule));
-			if (!Parameters.MeshVertexPositions.Size())
-			{
-				FVector2D::FReal HalfHeight = Parameters.ShapeParams.CapsuleHalfHeightAndRadius.X;
-				FVector2D::FReal Radius = Parameters.ShapeParams.CapsuleHalfHeightAndRadius.Y;
-
-				Parameters.MeshVertexPositions.AddParticles(14);
-				Parameters.MeshVertexPositions.X(0) = Chaos::FVec3(HalfHeight + Radius, 0, 0);
-				Parameters.MeshVertexPositions.X(1) = Chaos::FVec3(-HalfHeight - Radius, 0, 0);
-				Parameters.MeshVertexPositions.X(2) = Chaos::FVec3(HalfHeight, Radius, Radius);
-				Parameters.MeshVertexPositions.X(3) = Chaos::FVec3(HalfHeight, -Radius, Radius);
-				Parameters.MeshVertexPositions.X(4) = Chaos::FVec3(HalfHeight, -Radius, -Radius);
-				Parameters.MeshVertexPositions.X(5) = Chaos::FVec3(HalfHeight, Radius, -Radius);
-				Parameters.MeshVertexPositions.X(6) = Chaos::FVec3(0, Radius, Radius);
-				Parameters.MeshVertexPositions.X(7) = Chaos::FVec3(0, -Radius, Radius);
-				Parameters.MeshVertexPositions.X(8) = Chaos::FVec3(0, -Radius, -Radius);
-				Parameters.MeshVertexPositions.X(9) = Chaos::FVec3(0, Radius, -Radius);
-				Parameters.MeshVertexPositions.X(10) = Chaos::FVec3(-HalfHeight, Radius, Radius);
-				Parameters.MeshVertexPositions.X(11) = Chaos::FVec3(-HalfHeight, -Radius, Radius);
-				Parameters.MeshVertexPositions.X(12) = Chaos::FVec3(-HalfHeight, -Radius, -Radius);
-				Parameters.MeshVertexPositions.X(13) = Chaos::FVec3(-HalfHeight, Radius, -Radius);
-			}
-		}
-		else
-		{
-			Bounds.Min = Chaos::FVec3(-UE_KINDA_SMALL_NUMBER);
-			Bounds.Max = Chaos::FVec3(UE_KINDA_SMALL_NUMBER);
-			Particles.SetSharedGeometry(RigidBodyId, TSharedPtr<Chaos::FImplicitObject, ESPMode::ThreadSafe>());
-		}
-
-		FTransform WorldTransform = Parameters.InitialTransform;
-
-		Scale = WorldTransform.GetScale3D();
-		CenterOfMass = Bounds.GetCenter();
-		Bounds = Bounds.InverseTransformBy(FTransform(CenterOfMass));
-		Bounds.Min *= Scale;
-		Bounds.Max *= Scale;
-		checkSlow((Bounds.Max + Bounds.Min).Size() < UE_KINDA_SMALL_NUMBER);
-
-		Particles.InvM(RigidBodyId) = 0.f;
-		ensure(Parameters.Mass >= 0.f);
-		Particles.M(RigidBodyId) = Parameters.Mass;
-		if(Parameters.Mass < FLT_EPSILON)
-		{
-			Particles.InvM(RigidBodyId) = 1.f;
-		}
-		else
-		{
-			Particles.InvM(RigidBodyId) = 1.f / Parameters.Mass;
-		}
-
-		Particles.X(RigidBodyId) = WorldTransform.TransformPosition(CenterOfMass);
-		Particles.V(RigidBodyId) = Chaos::FVec3(Parameters.InitialLinearVelocity);
-		Particles.R(RigidBodyId) = WorldTransform.GetRotation().GetNormalized();
-		Particles.W(RigidBodyId) = Chaos::FVec3(Parameters.InitialAngularVelocity);
-		Particles.P(RigidBodyId) = Particles.X(RigidBodyId);
-		Particles.Q(RigidBodyId) = Particles.R(RigidBodyId);
-
-		const FVector SideSquared(Bounds.GetSize()[0] * Bounds.GetSize()[0], Bounds.GetSize()[1] * Bounds.GetSize()[1], Bounds.GetSize()[2] * Bounds.GetSize()[2]);
-		const FVector3f Inertia((Chaos::FRealSingle)(Parameters.Mass * (SideSquared.Y + SideSquared.Z) / 12.f), (Chaos::FRealSingle)(Parameters.Mass * (SideSquared.X + SideSquared.Z) / 12.f), (Chaos::FRealSingle)(Parameters.Mass * (SideSquared.X + SideSquared.Y) / 12.f));
-		Particles.I(RigidBodyId) = Inertia;
-		Particles.InvI(RigidBodyId) = Chaos::TVec3<Chaos::FRealSingle>(1.f / Inertia.X, 1.f / Inertia.Y, 1.f / Inertia.Z);
-
-		if(Parameters.ObjectType == EObjectStateTypeEnum::Chaos_Object_Sleeping)
-		{
-			Particles.SetObjectState(RigidBodyId, Chaos::EObjectStateType::Sleeping);
-			Particles.SetSleeping(RigidBodyId, true);
-		}
-		else if(Parameters.ObjectType != EObjectStateTypeEnum::Chaos_Object_Dynamic)
-		{
-			Particles.InvM(RigidBodyId) = 0.f;
-			Particles.InvI(RigidBodyId) = Chaos::TVec3<Chaos::FRealSingle>(0);
-			Particles.SetObjectState(RigidBodyId, Chaos::EObjectStateType::Kinematic);
-		}
-		else
-		{
-			Particles.SetObjectState(RigidBodyId, Chaos::EObjectStateType::Dynamic);
-		}
-
-		if (ensure(Parameters.MeshVertexPositions.Size()))
-		{
-			// Add collision vertices
-			check(Particles.CollisionParticles(RigidBodyId) == nullptr);
-			Particles.SetCollisionParticles(RigidBodyId, MoveTemp(Parameters.MeshVertexPositions));
-			if (Particles.CollisionParticles(RigidBodyId)->Size())
-			{
-				Particles.CollisionParticles(RigidBodyId)->UpdateAccelerationStructures();
-			}
-		}
-
-#if TODO_REIMPLEMENT_SET_PHYSICS_MATERIAL
-		GetSolver()->SetPhysicsMaterial(RigidBodyId, Parameters.PhysicalMaterial);
-#endif
-
-		bInitializedState = true;
-	}
 }
 
 void FStaticMeshPhysicsProxy::ParameterUpdateCallback(FParticlesType& InParticles, const float InTime)
 {
-
 }
 
 void FStaticMeshPhysicsProxy::DisableCollisionsCallback(TSet<TTuple<int32, int32>>& InPairs)
 {
-
 }
 
 void FStaticMeshPhysicsProxy::AddForceCallback(FParticlesType& InParticles, const float InDt, const int32 InIndex)
 {
-
 }
 
 void FStaticMeshPhysicsProxy::OnRemoveFromScene()
 {
-	ensure(false);	//GetSolver needs a trait, but we don't know it - this class changing anyway
-#if ADD_TRAITS_TO_STATIC_MESH_PROXY
-	// Disable the particle we added
-	Chaos::FPhysicsSolver* CurrSolver = GetSolver();
-
-	if(CurrSolver && RigidBodyId != INDEX_NONE)
-	{
-#if TODO_REIMPLEMENT_GET_RIGID_PARTICLES
-		// #BG TODO Special case here because right now we reset/realloc the evolution per geom component
-		// in endplay which clears this out. That needs to not happen and be based on world shutdown
-		if(CurrSolver->GetRigidParticles().Size() == 0)
-		{
-			return;
-		}
-
-		CurrSolver->GetEvolution()->DisableParticle(RigidBodyId);
-		CurrSolver->GetRigidClustering().GetTopLevelClusterParents().Remove(RigidBodyId);
-#endif
-	}
-#endif
 }
 
 void FStaticMeshPhysicsProxy::BufferPhysicsResults()
 {
-	SCOPE_CYCLE_COUNTER(STAT_CacheResultStaticMesh);
-	Results.GetPhysicsDataForWrite() = SimTransform;
 }
 
 void FStaticMeshPhysicsProxy::FlipBuffer()
 {
-	Results.Flip();
 }
 
 bool FStaticMeshPhysicsProxy::PullFromPhysicsState(const int32 SolverSyncTimestamp)
 {
-	if(Parameters.ObjectType == EObjectStateTypeEnum::Chaos_Object_Dynamic && Parameters.bSimulating && SyncDynamicTransformFunc)
-	{
-		// Send transform to update callable
-		SyncDynamicTransformFunc(Results.GetGameDataForRead());
-	}
-
 	return true;
 }
diff --git a/Engine/Source/Runtime/Experimental/Chaos/Public/Chaos/ExternalCollisionData.h b/Engine/Source/Runtime/Experimental/Chaos/Public/Chaos/ExternalCollisionData.h
--- a/Engine/Source/Runtime/Experimental/Chaos/Public/Chaos/ExternalCollisionData.h
+++ b/Engine/Source/Runtime/Experimental/Chaos/Public/Chaos/ExternalCollisionData.h
@@ -1,14 +1,17 @@
 // Copyright Epic Games, Inc. All Rights Reserved.
 #pragma once
-#include "Box.h"
-#include "Chaos/Framework/PhysicsProxyBase.h"
-#include "Chaos/PBDConstraintContainer.h"
+#include "Chaos/Box.h"
 #include "Chaos/ParticleHandleFwd.h"
-#include "Chaos/Vector.h"
 #include "Chaos/PhysicalMaterials.h"
+#include "Chaos/Vector.h"
+//
+// NOTE: This file is widely included in Engine code. 
+// Avoid including Chaos headers when possible.
+//
 
 
 class UPhysicalMaterial;
+class IPhysicsProxyBase;
 
 namespace Chaos
 {
diff --git a/Engine/Source/Runtime/Experimental/Chaos/Public/Chaos/Framework/PhysicsProxy.h b/Engine/Source/Runtime/Experimental/Chaos/Public/Chaos/Framework/PhysicsProxy.h
--- a/Engine/Source/Runtime/Experimental/Chaos/Public/Chaos/Framework/PhysicsProxy.h
+++ b/Engine/Source/Runtime/Experimental/Chaos/Public/Chaos/Framework/PhysicsProxy.h
@@ -6,14 +6,22 @@
 #include "Chaos/Declares.h"
 #include "Chaos/Framework/PhysicsSolverBase.h"
 #include "Chaos/Framework/PhysicsProxyBase.h"
-#include "Chaos/PBDCollisionConstraints.h"
-#include "Chaos/PBDRigidParticles.h"
 #include "UObject/GCObject.h"
+//
+// NOTE: This file is widely included in Engine code. 
+// Avoid including Chaos headers when possible.
+//
 
 struct FKinematicProxy;
 class FFieldSystemCommand;
 struct FBodyInstance;
 
+namespace Chaos
+{
+	template<typename T, int D> class TPBDRigidParticles;
+
+	using FPBDRigidParticles = TPBDRigidParticles<FReal, 3>;
+}
 
 /**
  * Base object interface for solver objects. Defines the expected API for objects
@@ -40,7 +48,6 @@ public:
 	using FParticleType = Concrete;
 
 	using FParticlesType = Chaos::FPBDRigidParticles;
-	using FCollisionConstraintsType = Chaos::FPBDCollisionConstraints;
 	using FIntArray = Chaos::TArrayCollectionArray<int32>;
 
 	TPhysicsProxy()
diff --git a/Engine/Source/Runtime/Experimental/Chaos/Public/PhysicsProxy/GeometryCollectionPhysicsProxy.h b/Engine/Source/Runtime/Experimental/Chaos/Public/PhysicsProxy/GeometryCollectionPhysicsProxy.h
--- a/Engine/Source/Runtime/Experimental/Chaos/Public/PhysicsProxy/GeometryCollectionPhysicsProxy.h
+++ b/Engine/Source/Runtime/Experimental/Chaos/Public/PhysicsProxy/GeometryCollectionPhysicsProxy.h
@@ -27,6 +27,7 @@ namespace Chaos
 
 	struct FDirtyGeometryCollectionData;
 
+	class FPBDCollisionConstraints;
 	class FPBDRigidsEvolutionBase;
 }
 
@@ -592,7 +593,7 @@ private:
 	TFunction<void(void)> ResetAnimationCacheCallback;
 	TFunction<void(const TArrayView<FTransform> &)> UpdateTransformsCallback;
 	TFunction<void(const int32 & CurrentFrame, const TManagedArray<int32> & RigidBodyID, const TManagedArray<int32>& Level, const TManagedArray<int32>& Parent, const TManagedArray<TSet<int32>>& Children, const TManagedArray<uint32>& SimulationType, const TManagedArray<uint32>& StatusFlags, const FParticlesType& Particles)> UpdateRestStateCallback;
-	TFunction<void(float SolverTime, const TManagedArray<int32> & RigidBodyID, const FParticlesType& Particles, const FCollisionConstraintsType& CollisionRule)> UpdateRecordedStateCallback;
+	TFunction<void(float SolverTime, const TManagedArray<int32> & RigidBodyID, const FParticlesType& Particles, const Chaos::FPBDCollisionConstraints& CollisionRule)> UpdateRecordedStateCallback;
 	TFunction<void(FRecordedTransformTrack& InTrack)> CommitRecordedStateCallback;
 
 	// Index of the first particles for this collection in the larger particle array
diff --git a/Engine/Source/Runtime/Experimental/Chaos/Public/PhysicsProxy/SkeletalMeshPhysicsProxy.h b/Engine/Source/Runtime/Experimental/Chaos/Public/PhysicsProxy/SkeletalMeshPhysicsProxy.h
--- a/Engine/Source/Runtime/Experimental/Chaos/Public/PhysicsProxy/SkeletalMeshPhysicsProxy.h
+++ b/Engine/Source/Runtime/Experimental/Chaos/Public/PhysicsProxy/SkeletalMeshPhysicsProxy.h
@@ -2,14 +2,19 @@
 #pragma once
 
 #include "Chaos/Framework/PhysicsProxy.h"
+#include "Chaos/Framework/BufferedData.h"
 #include "BoneHierarchy.h"
 #include "GeometryCollection/GeometryCollectionSimulationTypes.h"
 
-#include "Chaos/PBDJointConstraints.h"
 #include "Framework/TripleBufferedData.h"
 
 // @todo(chaos): remove this file
 
+namespace Chaos
+{
+	class FParticleData;
+}
+
 struct FSkeletalMeshPhysicsProxyParams
 {
 	FSkeletalMeshPhysicsProxyParams()
@@ -95,7 +100,7 @@ struct FSkeletalMeshPhysicsProxyInputs
 	TArray<FVector> AngularVelocities;
 };
 
-struct FSkeletalMeshPhysicsProxyOutputs : public Chaos::FParticleData
+struct FSkeletalMeshPhysicsProxyOutputs //: public Chaos::FParticleData
 {
 	TArray<FTransform> Transforms;
 	TArray<FVector> LinearVelocities;
@@ -162,7 +167,7 @@ public:
 
 	/** 
 	 */
-	const FSkeletalMeshPhysicsProxyOutputs* GetOutputs() const { return CurrentOutputConsumerBuffer; }
+	const FSkeletalMeshPhysicsProxyOutputs* GetOutputs() const { return nullptr; }
 
 	const FBoneHierarchy& GetBoneHierarchy() const { return Parameters.BoneHierarchy; }
 
@@ -170,14 +175,4 @@ private:
 	using FJointConstraints = Chaos::FPBDJointConstraints;
 
 	FSkeletalMeshPhysicsProxyParams Parameters;
-	TArray<int32> RigidBodyIds;
-	FJointConstraints JointConstraints;
-	// @todo(ccaulfield): sort out the IO buffer stuff
-	Chaos::TTripleBufferedData<FSkeletalMeshPhysicsProxyInputs> InputBuffers;
-	Chaos::TBufferedData<FSkeletalMeshPhysicsProxyOutputs> OutputBuffers;
-	FSkeletalMeshPhysicsProxyInputs* NextInputProducerBuffer;				// Buffer for the game to write to next
-	const FSkeletalMeshPhysicsProxyOutputs* CurrentOutputConsumerBuffer;	// Buffer for the game to read from next
-	bool bInitializedState;
-
-	FInitFunc InitFunc;
 };
diff --git a/Engine/Source/Runtime/Experimental/Chaos/Public/PhysicsProxy/StaticMeshPhysicsProxy.h b/Engine/Source/Runtime/Experimental/Chaos/Public/PhysicsProxy/StaticMeshPhysicsProxy.h
--- a/Engine/Source/Runtime/Experimental/Chaos/Public/PhysicsProxy/StaticMeshPhysicsProxy.h
+++ b/Engine/Source/Runtime/Experimental/Chaos/Public/PhysicsProxy/StaticMeshPhysicsProxy.h
@@ -6,13 +6,21 @@
 #include "GeometryCollection/GeometryCollectionSimulationTypes.h"
 #include "Chaos/Framework/BufferedData.h"
 
+// @todo(chaos): remove this file
+
+namespace Chaos
+{
+	class FParticleData;
+}
+
 struct FPhysicsProxyKinematicUpdate
 {
 	FTransform NewTransform;
 	FVector NewVelocity;
 };
 
-struct FStubSkeletalMeshData : public Chaos::FParticleData {
+struct FStubSkeletalMeshData //: public Chaos::FParticleData 
+{
 	void Reset() { };
 };
 
@@ -104,28 +112,4 @@ public:
 	/** ----------------------- */
 
 private:
-
-	Params Parameters;
-
-	bool bInitializedState;
-	int32 RigidBodyId;
-	FVector CenterOfMass;
-	FVector Scale;
-
-	// Transform that the callback object will write into during simulation.
-	// During sync this will be pushed back to the component
-	FTransform SimTransform;
-
-	// Double buffered result data
-	Chaos::TBufferedData<FTransform> Results;
-
-	/**
-	 *	External functions for setup and sync, called on the game thread during callback creation and syncing
-	 */
-	FCallbackInitFunc InitialiseCallbackParamsFunc;
-	FSyncDynamicFunc SyncDynamicTransformFunc;
-	//////////////////////////////////////////////////////////////////////////
-
-	bool bPendingKinematicUpdate;
-	FPhysicsProxyKinematicUpdate BufferedKinematicUpdate;
 };
