Chaos Cloth - Fixed constraint property index initializations.

#rnx

diff --git a/Engine/Source/Runtime/Experimental/Chaos/Private/Chaos/XPBDAnisotropicBendingConstraints.cpp b/Engine/Source/Runtime/Experimental/Chaos/Private/Chaos/XPBDAnisotropicBendingConstraints.cpp
--- a/Engine/Source/Runtime/Experimental/Chaos/Private/Chaos/XPBDAnisotropicBendingConstraints.cpp
+++ b/Engine/Source/Runtime/Experimental/Chaos/Private/Chaos/XPBDAnisotropicBendingConstraints.cpp
@@ -88,6 +88,17 @@ FXPBDAnisotropicBendingConstraints::FXPBDAnisotropicBendingConstraints(const FSo
 		ParticleOffset,
 		ParticleCount)
 	, WarpWeftBiasBaseMultipliers(GenerateWarpWeftBiasBaseMultipliers(FaceVertexPatternPositions, TriangleMesh))
+	, XPBDAnisoBendingStiffnessWarpIndex(PropertyCollection)
+	, XPBDAnisoBendingStiffnessWeftIndex(PropertyCollection)
+	, XPBDAnisoBendingStiffnessBiasIndex(PropertyCollection)
+	, XPBDAnisoBendingDampingIndex(PropertyCollection)
+	, XPBDAnisoBucklingRatioIndex(PropertyCollection)
+	, XPBDAnisoBucklingStiffnessWarpIndex(PropertyCollection)
+	, XPBDAnisoBucklingStiffnessWeftIndex(PropertyCollection)
+	, XPBDAnisoBucklingStiffnessBiasIndex(PropertyCollection)
+	, XPBDAnisoFlatnessRatioIndex(PropertyCollection)
+	, XPBDAnisoRestAngleIndex(PropertyCollection)
+	, XPBDAnisoRestAngleTypeIndex(PropertyCollection)
 {
 	Lambdas.Init((FSolverReal)0., Constraints.Num());
 	InitColor(InParticles);
@@ -169,6 +180,17 @@ FXPBDAnisotropicBendingConstraints::FXPBDAnisotropicBendingConstraints(const FSo
 		ParticleOffset,
 		ParticleCount)
 	, WarpWeftBiasBaseMultipliers(GenerateWarpWeftBiasBaseMultipliers(FaceVertexPatternPositions, TriangleMesh))
+	, XPBDAnisoBendingStiffnessWarpIndex(ForceInit)
+	, XPBDAnisoBendingStiffnessWeftIndex(ForceInit)
+	, XPBDAnisoBendingStiffnessBiasIndex(ForceInit)
+	, XPBDAnisoBendingDampingIndex(ForceInit)
+	, XPBDAnisoBucklingRatioIndex(ForceInit)
+	, XPBDAnisoBucklingStiffnessWarpIndex(ForceInit)
+	, XPBDAnisoBucklingStiffnessWeftIndex(ForceInit)
+	, XPBDAnisoBucklingStiffnessBiasIndex(ForceInit)
+	, XPBDAnisoFlatnessRatioIndex(ForceInit)
+	, XPBDAnisoRestAngleIndex(ForceInit)
+	, XPBDAnisoRestAngleTypeIndex(ForceInit)
 {
 	Lambdas.Init((FSolverReal)0., Constraints.Num());
 	InitColor(InParticles);
diff --git a/Engine/Source/Runtime/Experimental/Chaos/Private/Chaos/XPBDStretchBiasElementConstraints.cpp b/Engine/Source/Runtime/Experimental/Chaos/Private/Chaos/XPBDStretchBiasElementConstraints.cpp
--- a/Engine/Source/Runtime/Experimental/Chaos/Private/Chaos/XPBDStretchBiasElementConstraints.cpp
+++ b/Engine/Source/Runtime/Experimental/Chaos/Private/Chaos/XPBDStretchBiasElementConstraints.cpp
@@ -82,6 +82,13 @@ FXPBDStretchBiasElementConstraints::FXPBDStretchBiasElementConstraints(const FSo
 		TConstArrayView<TVec3<int32>>(Constraints),
 		ParticleOffset,
 		ParticleCount)
+	, XPBDAnisoStretchUse3dRestLengthsIndex(PropertyCollection)
+	, XPBDAnisoStretchStiffnessWarpIndex(PropertyCollection)
+	, XPBDAnisoStretchStiffnessWeftIndex(PropertyCollection)
+	, XPBDAnisoStretchStiffnessBiasIndex(PropertyCollection)
+	, XPBDAnisoStretchDampingIndex(PropertyCollection)
+	, XPBDAnisoStretchWarpScaleIndex(PropertyCollection)
+	, XPBDAnisoStretchWeftScaleIndex(PropertyCollection)
 {
 	Lambdas.Init(FSolverVec3(0.), Constraints.Num());
 	InitConstraintsAndRestData(InParticles, TriangleMesh, FaceVertexUVs, GetXPBDAnisoStretchUse3dRestLengths(PropertyCollection, bDefaultUse3dRestLengths), bTrimKinematicConstraints);
@@ -154,6 +161,13 @@ FXPBDStretchBiasElementConstraints::FXPBDStretchBiasElementConstraints(const FSo
 		TConstArrayView<TVec3<int32>>(Constraints),
 		ParticleOffset,
 		ParticleCount)
+	, XPBDAnisoStretchUse3dRestLengthsIndex(ForceInit)
+	, XPBDAnisoStretchStiffnessWarpIndex(ForceInit)
+	, XPBDAnisoStretchStiffnessWeftIndex(ForceInit)
+	, XPBDAnisoStretchStiffnessBiasIndex(ForceInit)
+	, XPBDAnisoStretchDampingIndex(ForceInit)
+	, XPBDAnisoStretchWarpScaleIndex(ForceInit)
+	, XPBDAnisoStretchWeftScaleIndex(ForceInit)
 {
 	Lambdas.Init(FSolverVec3(0.), Constraints.Num());
 	InitConstraintsAndRestData(InParticles, TriangleMesh, FaceVertexUVs, bUse3dRestLengths, bTrimKinematicConstraints);
diff --git a/Engine/Source/Runtime/Experimental/Chaos/Public/Chaos/CollectionPropertyFacade.h b/Engine/Source/Runtime/Experimental/Chaos/Public/Chaos/CollectionPropertyFacade.h
--- a/Engine/Source/Runtime/Experimental/Chaos/Public/Chaos/CollectionPropertyFacade.h
+++ b/Engine/Source/Runtime/Experimental/Chaos/Public/Chaos/CollectionPropertyFacade.h
@@ -446,33 +446,33 @@ PRAGMA_ENABLE_DEPRECATION_WARNINGS
 	} \
 	Type GetLow##PropertyName(const FCollectionPropertyConstFacade& PropertyCollection, const Type& Default) \
 	{ \
-		return PropertyCollection.GetLowValue<Type>(PropertyName##Name.ToString(), Default, &PropertyName##Index); \
+		return PropertyCollection.GetLowValue<Type>(PropertyName##Name.ToString(), Default); \
 	} \
 	Type GetHigh##PropertyName(const FCollectionPropertyConstFacade& PropertyCollection, const Type& Default) \
 	{ \
-		return PropertyCollection.GetHighValue<Type>(PropertyName##Name.ToString(), Default, &PropertyName##Index); \
+		return PropertyCollection.GetHighValue<Type>(PropertyName##Name.ToString(), Default); \
 	} \
 	TPair<Type, Type> GetWeighted##PropertyName(const FCollectionPropertyConstFacade& PropertyCollection, const Type& Default) \
 	{ \
-		return PropertyCollection.GetWeightedValue<Type>(PropertyName##Name.ToString(), Default, &PropertyName##Index); \
+		return PropertyCollection.GetWeightedValue<Type>(PropertyName##Name.ToString(), Default); \
 	} \
 	FVector2f GetWeightedFloat##PropertyName(const FCollectionPropertyConstFacade& PropertyCollection, const float& Default) \
 	{ \
-		return PropertyCollection.GetWeightedFloatValue(PropertyName##Name.ToString(), Default, &PropertyName##Index); \
+		return PropertyCollection.GetWeightedFloatValue(PropertyName##Name.ToString(), Default); \
 	} \
 	Type Get##PropertyName(const FCollectionPropertyConstFacade& PropertyCollection, const Type& Default) \
 	{ \
-		return PropertyCollection.GetValue<Type>(PropertyName##Name.ToString(), Default, &PropertyName##Index); \
+		return PropertyCollection.GetValue<Type>(PropertyName##Name.ToString(), Default); \
 	} \
 	FString Get##PropertyName##String(const FCollectionPropertyConstFacade& PropertyCollection, const FString& Default) \
 	{ \
-		return PropertyCollection.GetStringValue(PropertyName##Name.ToString(), Default, &PropertyName##Index); \
+		return PropertyCollection.GetStringValue(PropertyName##Name.ToString(), Default); \
 	} \
 	UE_DEPRECATED(5.3, "GetFlags is being phased out to promote correct dirtying operations.") \
 	uint8 Get##PropertyName##Flags(const FCollectionPropertyConstFacade& PropertyCollection, uint8 Default) \
 	{ \
 PRAGMA_DISABLE_DEPRECATION_WARNINGS \
-	return PropertyCollection.GetFlags(PropertyName##Name.ToString(), Default, &PropertyName##Index); \
+	return PropertyCollection.GetFlags(PropertyName##Name.ToString(), Default); \
 PRAGMA_ENABLE_DEPRECATION_WARNINGS \
 	} \
 	Type GetLow##PropertyName(const FCollectionPropertyConstFacade& PropertyCollection) const \
@@ -549,4 +549,12 @@ PRAGMA_ENABLE_DEPRECATION_WARNINGS \
 		return PropertyName##Index != INDEX_NONE && \
 			PropertyCollection.IsEnabled(PropertyName##Index) && PropertyCollection.IsAnimatable(PropertyName##Index) && PropertyCollection.IsDirty(PropertyName##Index); \
 	} \
-	int32 PropertyName##Index = INDEX_NONE;
+	struct F##PropertyName##Index \
+	{ \
+		int32 Index = INDEX_NONE; \
+		UE_DEPRECATED(5.3, PREPROCESSOR_TO_STRING(PropertyName##Index) " must be explicitly initialized. Add " PREPROCESSOR_TO_STRING(PropertyName##Index) "(PropertyCollection) or (ForceInit) to this constructor initialization list.") \
+		F##PropertyName##Index() {} \
+		explicit F##PropertyName##Index(EForceInit) : Index(INDEX_NONE) {} \
+		explicit F##PropertyName##Index(const FCollectionPropertyConstFacade& PropertyCollection) : Index(PropertyCollection.GetKeyIndex(PropertyName##Name.ToString())) {} \
+		operator int32() const { return Index; } \
+	} PropertyName##Index;
diff --git a/Engine/Source/Runtime/Experimental/Chaos/Public/Chaos/PBDAnimDriveConstraint.h b/Engine/Source/Runtime/Experimental/Chaos/Public/Chaos/PBDAnimDriveConstraint.h
--- a/Engine/Source/Runtime/Experimental/Chaos/Public/Chaos/PBDAnimDriveConstraint.h
+++ b/Engine/Source/Runtime/Experimental/Chaos/Public/Chaos/PBDAnimDriveConstraint.h
@@ -40,6 +40,8 @@ namespace Chaos::Softs
 			, UseDeprecatedApply(true)
 			, Stiffness(FSolverVec2::UnitVector, StiffnessMultipliers, InParticleCount)
 			, Damping(FSolverVec2::UnitVector, DampingMultipliers, InParticleCount)
+			, AnimDriveStiffnessIndex(ForceInit)
+			, AnimDriveDampingIndex(ForceInit)
 		{
 		}
 
@@ -65,6 +67,8 @@ namespace Chaos::Softs
 				FSolverVec2(GetWeightedFloatAnimDriveDamping(PropertyCollection, 1.f)),
 				WeightMaps.FindRef(GetAnimDriveDampingString(PropertyCollection, AnimDriveDampingName.ToString())),
 				InParticleCount)
+			, AnimDriveStiffnessIndex(PropertyCollection)
+			, AnimDriveDampingIndex(PropertyCollection)
 		{
 		}
 
@@ -92,6 +96,8 @@ namespace Chaos::Softs
 				FSolverVec2(GetWeightedFloatAnimDriveDamping(PropertyCollection, 1.f)),
 				DampingMultipliers,
 				InParticleCount)
+			, AnimDriveStiffnessIndex(PropertyCollection)
+			, AnimDriveDampingIndex(PropertyCollection)
 		{
 		}
 
@@ -112,6 +118,8 @@ namespace Chaos::Softs
 			, UseDeprecatedApply(false)
 			, Stiffness(FSolverVec2::UnitVector, StiffnessMultipliers, InParticleCount)
 			, Damping(FSolverVec2::UnitVector, DampingMultipliers, InParticleCount)
+			, AnimDriveStiffnessIndex(ForceInit)
+			, AnimDriveDampingIndex(ForceInit)
 		{
 		}
 
diff --git a/Engine/Source/Runtime/Experimental/Chaos/Public/Chaos/PBDAxialSpringConstraints.h b/Engine/Source/Runtime/Experimental/Chaos/Public/Chaos/PBDAxialSpringConstraints.h
--- a/Engine/Source/Runtime/Experimental/Chaos/Public/Chaos/PBDAxialSpringConstraints.h
+++ b/Engine/Source/Runtime/Experimental/Chaos/Public/Chaos/PBDAxialSpringConstraints.h
@@ -76,6 +76,7 @@ public:
 			WeightMaps.FindRef(GetAreaSpringStiffnessString(PropertyCollection, AreaSpringStiffnessName.ToString())),
 			FSolverVec2(GetWeightedFloatAreaSpringStiffness(PropertyCollection, 1.f)),
 			bTrimKinematicConstraints)
+		, AreaSpringStiffnessIndex(PropertyCollection)
 	{}
 
 	UE_DEPRECATED(5.3, "Use weight map constructor instead.")
@@ -95,6 +96,7 @@ public:
 			StiffnessMultipliers,
 			FSolverVec2(GetWeightedFloatAreaSpringStiffness(PropertyCollection, 1.f)),
 			bTrimKinematicConstraints)
+		, AreaSpringStiffnessIndex(PropertyCollection)
 	{}
 
 	virtual ~FPBDAreaSpringConstraints() override = default;
diff --git a/Engine/Source/Runtime/Experimental/Chaos/Public/Chaos/PBDBendingConstraints.h b/Engine/Source/Runtime/Experimental/Chaos/Public/Chaos/PBDBendingConstraints.h
--- a/Engine/Source/Runtime/Experimental/Chaos/Public/Chaos/PBDBendingConstraints.h
+++ b/Engine/Source/Runtime/Experimental/Chaos/Public/Chaos/PBDBendingConstraints.h
@@ -38,6 +38,12 @@ public:
 			GetRestAngleValueFromCollection(PropertyCollection),
 			(ERestAngleConstructionType)GetRestAngleType(PropertyCollection, (int32)ERestAngleConstructionType::Use3DRestAngles),
 			bTrimKinematicConstraints)
+		, BendingElementStiffnessIndex(PropertyCollection)
+		, BucklingRatioIndex(PropertyCollection)
+		, BucklingStiffnessIndex(PropertyCollection)
+		, FlatnessRatioIndex(PropertyCollection)
+		, RestAngleIndex(PropertyCollection)
+		, RestAngleTypeIndex(PropertyCollection)
 	{
 		InitColor(InParticles);
 	}
@@ -62,6 +68,12 @@ public:
 			(FSolverReal)GetBucklingRatio(PropertyCollection, 0.f),  // BucklingRatio is clamped in base class
 			FSolverVec2(GetWeightedFloatBucklingStiffness(PropertyCollection, 1.f)),
 			bTrimKinematicConstraints) 
+		, BendingElementStiffnessIndex(PropertyCollection)
+		, BucklingRatioIndex(PropertyCollection)
+		, BucklingStiffnessIndex(PropertyCollection)
+		, FlatnessRatioIndex(PropertyCollection)
+		, RestAngleIndex(PropertyCollection)
+		, RestAngleTypeIndex(PropertyCollection)
 	{
 		InitColor(InParticles);
 	}
@@ -86,7 +98,13 @@ public:
 			InStiffness,
 			InBucklingRatio,
 			InBucklingStiffness,
-			bTrimKinematicConstraints) 
+			bTrimKinematicConstraints)
+		, BendingElementStiffnessIndex(ForceInit)
+		, BucklingRatioIndex(ForceInit)
+		, BucklingStiffnessIndex(ForceInit)
+		, FlatnessRatioIndex(ForceInit)
+		, RestAngleIndex(ForceInit)
+		, RestAngleTypeIndex(ForceInit)
 	{
 		InitColor(InParticles);
 	}
diff --git a/Engine/Source/Runtime/Experimental/Chaos/Public/Chaos/PBDCollisionSpringConstraints.h b/Engine/Source/Runtime/Experimental/Chaos/Public/Chaos/PBDCollisionSpringConstraints.h
--- a/Engine/Source/Runtime/Experimental/Chaos/Public/Chaos/PBDCollisionSpringConstraints.h
+++ b/Engine/Source/Runtime/Experimental/Chaos/Public/Chaos/PBDCollisionSpringConstraints.h
@@ -39,6 +39,9 @@ public:
 			(FSolverReal)FMath::Max(GetSelfCollisionThickness(PropertyCollection, Base::BackCompatThickness), 0.f),
 			(FSolverReal)FMath::Clamp(GetSelfCollisionStiffness(PropertyCollection, Base::BackCompatStiffness), 0.f, 1.f),
 			FMath::Clamp((FSolverReal)GetSelfCollisionFriction(PropertyCollection, Base::BackCompatFrictionCoefficient), MinFrictionCoefficient, MaxFrictionCoefficient))
+		, SelfCollisionThicknessIndex(PropertyCollection)
+		, SelfCollisionStiffnessIndex(PropertyCollection)
+		, SelfCollisionFrictionIndex(PropertyCollection)
 	{}
 
 	FPBDCollisionSpringConstraints(
@@ -59,6 +62,9 @@ public:
 			InThickness,
 			InStiffness,
 			InFrictionCoefficient)
+		, SelfCollisionThicknessIndex(ForceInit)
+		, SelfCollisionStiffnessIndex(ForceInit)
+		, SelfCollisionFrictionIndex(ForceInit)
 	{}
 
 	virtual ~FPBDCollisionSpringConstraints() override {}
diff --git a/Engine/Source/Runtime/Experimental/Chaos/Public/Chaos/PBDLongRangeConstraints.h b/Engine/Source/Runtime/Experimental/Chaos/Public/Chaos/PBDLongRangeConstraints.h
--- a/Engine/Source/Runtime/Experimental/Chaos/Public/Chaos/PBDLongRangeConstraints.h
+++ b/Engine/Source/Runtime/Experimental/Chaos/Public/Chaos/PBDLongRangeConstraints.h
@@ -40,6 +40,8 @@ public:
 			FSolverVec2(GetWeightedFloatTetherScale(PropertyCollection, 1.f)),  // Scale clamping done in constructor
 			FPBDStiffness::DefaultPBDMaxStiffness,
 			MeshScale)
+		, TetherStiffnessIndex(PropertyCollection)
+		, TetherScaleIndex(PropertyCollection)
 	{}
 
 	UE_DEPRECATED(5.3, "Use weight map constructor instead.")
@@ -63,6 +65,8 @@ public:
 			FSolverVec2(GetWeightedFloatTetherScale(PropertyCollection, 1.f)),  // Scale clamping done in constructor
 			FPBDStiffness::DefaultPBDMaxStiffness,
 			MeshScale)
+		, TetherStiffnessIndex(PropertyCollection)
+		, TetherScaleIndex(PropertyCollection)
 	{}
 
 	FPBDLongRangeConstraints(
@@ -85,6 +89,8 @@ public:
 			InStiffness,
 			InScale,
 			MeshScale)
+		, TetherStiffnessIndex(ForceInit)
+		, TetherScaleIndex(ForceInit)
 	{}
 
 	virtual ~FPBDLongRangeConstraints() override {}
diff --git a/Engine/Source/Runtime/Experimental/Chaos/Public/Chaos/PBDSphericalConstraint.h b/Engine/Source/Runtime/Experimental/Chaos/Public/Chaos/PBDSphericalConstraint.h
--- a/Engine/Source/Runtime/Experimental/Chaos/Public/Chaos/PBDSphericalConstraint.h
+++ b/Engine/Source/Runtime/Experimental/Chaos/Public/Chaos/PBDSphericalConstraint.h
@@ -49,6 +49,7 @@ public:
 		, Scale(MeshScale)
 		, MaxDistanceBase((FSolverReal)GetLowMaxDistance(PropertyCollection, 0.f))
 		, MaxDistanceRange((FSolverReal)GetHighMaxDistance(PropertyCollection, 1.f) - MaxDistanceBase)
+		, MaxDistanceIndex(PropertyCollection)
 	{
 	}
 
@@ -66,6 +67,9 @@ public:
 		, ParticleOffset(InParticleOffset)
 		, ParticleCount(InParticleCount)
 		, Scale(MeshScale)
+		, MaxDistanceBase((FSolverReal)GetLowMaxDistance(PropertyCollection, 0.f))
+		, MaxDistanceRange((FSolverReal)GetHighMaxDistance(PropertyCollection, 1.f) - MaxDistanceBase)
+		, MaxDistanceIndex(PropertyCollection)
 	{
 	}
 
@@ -80,6 +84,7 @@ public:
 		, ParticleOffset(InParticleOffset)
 		, ParticleCount(InParticleCount)
 		, Scale((FSolverReal)1.)
+		, MaxDistanceIndex(ForceInit)
 	{
 	}
 
@@ -213,6 +218,9 @@ public:
 		, BackstopDistanceBase((FSolverReal)GetLowBackstopDistance(PropertyCollection, 0.f))
 		, BackstopDistanceRange((FSolverReal)GetHighBackstopDistance(PropertyCollection, 1.f) - BackstopDistanceBase)
 		, bUseLegacyBackstop(GetUseLegacyBackstop(PropertyCollection, false))  // Only set the legacy backstop in constructor
+		, BackstopDistanceIndex(PropertyCollection)
+		, BackstopRadiusIndex(PropertyCollection)
+		, UseLegacyBackstopIndex(PropertyCollection)
 	{
 	}
 
@@ -234,8 +242,14 @@ public:
 		, ParticleOffset(InParticleOffset)
 		, ParticleCount(InParticleCount)
 		, Scale(MeshScale)
-		, bEnabled(true)
+		, BackstopRadiusBase((FSolverReal)FMath::Max(GetLowBackstopRadius(PropertyCollection, 0.f), 0.f))
+		, BackstopRadiusRange((FSolverReal)FMath::Max(GetHighBackstopRadius(PropertyCollection, 1.f), 0.f) - BackstopRadiusBase)
+		, BackstopDistanceBase((FSolverReal)GetLowBackstopDistance(PropertyCollection, 0.f))
+		, BackstopDistanceRange((FSolverReal)GetHighBackstopDistance(PropertyCollection, 1.f) - BackstopDistanceBase)
 		, bUseLegacyBackstop(GetUseLegacyBackstop(PropertyCollection, false))  // Only set the legacy backstop in constructor
+		, BackstopDistanceIndex(PropertyCollection)
+		, BackstopRadiusIndex(PropertyCollection)
+		, UseLegacyBackstopIndex(PropertyCollection)
 	{
 	}
 
@@ -257,6 +271,9 @@ public:
 		, Scale((FSolverReal)1.)
 		, bEnabled(true)
 		, bUseLegacyBackstop(bInUseLegacyBackstop)
+		, BackstopDistanceIndex(ForceInit)
+		, BackstopRadiusIndex(ForceInit)
+		, UseLegacyBackstopIndex(ForceInit)
 	{
 	}
 	~FPBDSphericalBackstopConstraint() {}
diff --git a/Engine/Source/Runtime/Experimental/Chaos/Public/Chaos/PBDSpringConstraints.h b/Engine/Source/Runtime/Experimental/Chaos/Public/Chaos/PBDSpringConstraints.h
--- a/Engine/Source/Runtime/Experimental/Chaos/Public/Chaos/PBDSpringConstraints.h
+++ b/Engine/Source/Runtime/Experimental/Chaos/Public/Chaos/PBDSpringConstraints.h
@@ -77,6 +77,7 @@ public:
 			WeightMaps.FindRef(GetEdgeSpringStiffnessString(PropertyCollection, EdgeSpringStiffnessName.ToString())),
 			FSolverVec2(GetWeightedFloatEdgeSpringStiffness(PropertyCollection, 1.f)),
 			bTrimKinematicConstraints)
+		, EdgeSpringStiffnessIndex(PropertyCollection)
 	{}
 
 	UE_DEPRECATED(5.3, "Use weight map constructor instead.")
@@ -96,6 +97,7 @@ public:
 			StiffnessMultipliers,
 			FSolverVec2(GetWeightedFloatEdgeSpringStiffness(PropertyCollection, 1.f)),
 			bTrimKinematicConstraints)
+		, EdgeSpringStiffnessIndex(PropertyCollection)
 	{}
 
 	virtual ~FPBDEdgeSpringConstraints() override = default;
@@ -140,6 +142,7 @@ public:
 			WeightMaps.FindRef(GetBendingSpringStiffnessString(PropertyCollection, BendingSpringStiffnessName.ToString())),
 			FSolverVec2(GetWeightedFloatBendingSpringStiffness(PropertyCollection, 1.f)),
 			bTrimKinematicConstraints)
+		, BendingSpringStiffnessIndex(PropertyCollection)
 	{}
 
 	UE_DEPRECATED(5.3, "Use weight map constructor instead.")
@@ -159,6 +162,7 @@ public:
 			StiffnessMultipliers,
 			FSolverVec2(GetWeightedFloatBendingSpringStiffness(PropertyCollection, 1.f)),
 			bTrimKinematicConstraints)
+		, BendingSpringStiffnessIndex(PropertyCollection)
 	{}
 
 	virtual ~FPBDBendingSpringConstraints() override = default;
diff --git a/Engine/Source/Runtime/Experimental/Chaos/Public/Chaos/PBDTriangleMeshCollisions.h b/Engine/Source/Runtime/Experimental/Chaos/Public/Chaos/PBDTriangleMeshCollisions.h
--- a/Engine/Source/Runtime/Experimental/Chaos/Public/Chaos/PBDTriangleMeshCollisions.h
+++ b/Engine/Source/Runtime/Experimental/Chaos/Public/Chaos/PBDTriangleMeshCollisions.h
@@ -111,6 +111,9 @@ public:
 		, NumParticles(InNumParticles)
 		, bGlobalIntersectionAnalysis(GetUseSelfIntersections(PropertyCollection, false) || GetUseGlobalIntersectionAnalysis(PropertyCollection, false))
 		, bContourMinimization(GetUseSelfIntersections(PropertyCollection, false) || GetUseContourMinimization(PropertyCollection, false))
+		, UseSelfIntersectionsIndex(PropertyCollection)
+		, UseGlobalIntersectionAnalysisIndex(PropertyCollection)
+		, UseContourMinimizationIndex(PropertyCollection)
 	{}
 
 	FPBDTriangleMeshCollisions(
@@ -125,6 +128,9 @@ public:
 		, NumParticles(InNumParticles)
 		, bGlobalIntersectionAnalysis(bInGlobalIntersectionAnalysis)
 		, bContourMinimization(bInContourMinimization)
+		, UseSelfIntersectionsIndex(ForceInit)
+		, UseGlobalIntersectionAnalysisIndex(ForceInit)
+		, UseContourMinimizationIndex(ForceInit)
 	{}
 
 	virtual ~FPBDTriangleMeshCollisions() = default;
diff --git a/Engine/Source/Runtime/Experimental/Chaos/Public/Chaos/VelocityField.h b/Engine/Source/Runtime/Experimental/Chaos/Public/Chaos/VelocityField.h
--- a/Engine/Source/Runtime/Experimental/Chaos/Public/Chaos/VelocityField.h
+++ b/Engine/Source/Runtime/Experimental/Chaos/Public/Chaos/VelocityField.h
@@ -30,6 +30,10 @@ public:
 	explicit FVelocityAndPressureField(const FCollectionPropertyConstFacade& PropertyCollection)
 		: Offset(INDEX_NONE)
 		, NumParticles(0)
+		, DragIndex(PropertyCollection)
+		, LiftIndex(PropertyCollection)
+		, FluidDensityIndex(PropertyCollection)
+		, PressureIndex(PropertyCollection)
 	{
 		SetProperties(
 			FSolverVec2(GetWeightedFloatDrag(PropertyCollection, (FSolverReal)0.)),  // If these properties don't exist, set their values to 0, not to DefaultCoefficients!
@@ -42,6 +46,10 @@ public:
 	FVelocityAndPressureField()
 		: Offset(INDEX_NONE)
 		, NumParticles(0)
+		, DragIndex(ForceInit)
+		, LiftIndex(ForceInit)
+		, FluidDensityIndex(ForceInit)
+		, PressureIndex(ForceInit)
 	{
 		SetProperties(FSolverVec2(0.), FSolverVec2(0.), (FSolverReal)0., FSolverVec2(0.));
 	}
diff --git a/Engine/Source/Runtime/Experimental/Chaos/Public/Chaos/XPBDAxialSpringConstraints.h b/Engine/Source/Runtime/Experimental/Chaos/Public/Chaos/XPBDAxialSpringConstraints.h
--- a/Engine/Source/Runtime/Experimental/Chaos/Public/Chaos/XPBDAxialSpringConstraints.h
+++ b/Engine/Source/Runtime/Experimental/Chaos/Public/Chaos/XPBDAxialSpringConstraints.h
@@ -191,6 +191,7 @@ public:
 			WeightMaps.FindRef(GetXPBDAreaSpringStiffnessString(PropertyCollection, XPBDAreaSpringStiffnessName.ToString())),
 			FSolverVec2(GetWeightedFloatXPBDAreaSpringStiffness(PropertyCollection, MaxStiffness)),
 			bTrimKinematicConstraints)
+		, XPBDAreaSpringStiffnessIndex(PropertyCollection)
 	{}
 
 	UE_DEPRECATED(5.3, "Use weight map constructor instead.")
@@ -210,6 +211,7 @@ public:
 			StiffnessMultipliers,
 			FSolverVec2(GetWeightedFloatXPBDAreaSpringStiffness(PropertyCollection, MaxStiffness)),
 			bTrimKinematicConstraints)
+		, XPBDAreaSpringStiffnessIndex(PropertyCollection)
 	{}
 
 	virtual ~FXPBDAreaSpringConstraints() override = default;
diff --git a/Engine/Source/Runtime/Experimental/Chaos/Public/Chaos/XPBDBendingConstraints.h b/Engine/Source/Runtime/Experimental/Chaos/Public/Chaos/XPBDBendingConstraints.h
--- a/Engine/Source/Runtime/Experimental/Chaos/Public/Chaos/XPBDBendingConstraints.h
+++ b/Engine/Source/Runtime/Experimental/Chaos/Public/Chaos/XPBDBendingConstraints.h
@@ -58,6 +58,13 @@ public:
 			TConstArrayView<TVec2<int32>>(ConstraintSharedEdges),
 			ParticleOffset,
 			ParticleCount)
+		, XPBDBendingElementStiffnessIndex(PropertyCollection)
+		, XPBDBendingElementDampingIndex(PropertyCollection)
+		, XPBDBucklingRatioIndex(PropertyCollection)
+		, XPBDBucklingStiffnessIndex(PropertyCollection)
+		, XPBDFlatnessRatioIndex(PropertyCollection)
+		, XPBDRestAngleIndex(PropertyCollection)
+		, XPBDRestAngleTypeIndex(PropertyCollection)
 	{
 		Lambdas.Init((FSolverReal)0., Constraints.Num());
 		InitColor(InParticles);
@@ -91,6 +98,13 @@ public:
 			TConstArrayView<TVec2<int32>>(ConstraintSharedEdges),
 			ParticleOffset,
 			ParticleCount)
+		, XPBDBendingElementStiffnessIndex(PropertyCollection)
+		, XPBDBendingElementDampingIndex(PropertyCollection)
+		, XPBDBucklingRatioIndex(PropertyCollection)
+		, XPBDBucklingStiffnessIndex(PropertyCollection)
+		, XPBDFlatnessRatioIndex(PropertyCollection)
+		, XPBDRestAngleIndex(PropertyCollection)
+		, XPBDRestAngleTypeIndex(PropertyCollection)
 	{
 		Lambdas.Init((FSolverReal)0., Constraints.Num());
 		InitColor(InParticles);
@@ -120,6 +134,13 @@ public:
 			bTrimKinematicConstraints,
 			MaxStiffness)
 		, DampingRatio(FSolverVec2::ZeroVector)
+		, XPBDBendingElementStiffnessIndex(ForceInit)
+		, XPBDBendingElementDampingIndex(ForceInit)
+		, XPBDBucklingRatioIndex(ForceInit)
+		, XPBDBucklingStiffnessIndex(ForceInit)
+		, XPBDFlatnessRatioIndex(ForceInit)
+		, XPBDRestAngleIndex(ForceInit)
+		, XPBDRestAngleTypeIndex(ForceInit)
 	{
 		Lambdas.Init((FSolverReal)0., Constraints.Num());
 		InitColor(InParticles);
@@ -155,6 +176,13 @@ public:
 			TConstArrayView<TVec2<int32>>(ConstraintSharedEdges),
 			ParticleOffset,
 			ParticleCount)
+		, XPBDBendingElementStiffnessIndex(ForceInit)
+		, XPBDBendingElementDampingIndex(ForceInit)
+		, XPBDBucklingRatioIndex(ForceInit)
+		, XPBDBucklingStiffnessIndex(ForceInit)
+		, XPBDFlatnessRatioIndex(ForceInit)
+		, XPBDRestAngleIndex(ForceInit)
+		, XPBDRestAngleTypeIndex(ForceInit)
 	{
 		Lambdas.Init((FSolverReal)0., Constraints.Num());
 		InitColor(InParticles);
diff --git a/Engine/Source/Runtime/Experimental/Chaos/Public/Chaos/XPBDLongRangeConstraints.h b/Engine/Source/Runtime/Experimental/Chaos/Public/Chaos/XPBDLongRangeConstraints.h
--- a/Engine/Source/Runtime/Experimental/Chaos/Public/Chaos/XPBDLongRangeConstraints.h
+++ b/Engine/Source/Runtime/Experimental/Chaos/Public/Chaos/XPBDLongRangeConstraints.h
@@ -44,6 +44,8 @@ public:
 			FSolverVec2(GetWeightedFloatXPBDTetherScale(PropertyCollection, 1.f)),  // Scale clamping done in constructor
 			MaxStiffness,
 			MeshScale)
+		, XPBDTetherStiffnessIndex(PropertyCollection)
+		, XPBDTetherScaleIndex(PropertyCollection)
 	{
 		NumTethers = 0;
 		for (const TConstArrayView<FTether>& TetherBatch : Tethers)
@@ -72,6 +74,8 @@ public:
 			FSolverVec2(GetWeightedFloatXPBDTetherStiffness(PropertyCollection, MaxStiffness)),
 			FSolverVec2(GetWeightedFloatXPBDTetherScale(PropertyCollection, 1.f)),  // Scale clamping done in constructor
 			MaxStiffness)
+		, XPBDTetherStiffnessIndex(PropertyCollection)
+		, XPBDTetherScaleIndex(PropertyCollection)
 	{
 		NumTethers = 0;
 		for (const TConstArrayView<FTether>& TetherBatch : Tethers)
@@ -100,6 +104,8 @@ public:
 			InStiffness,
 			InScale,  // Scale clamping done in constructor
 			MaxStiffness)
+		, XPBDTetherStiffnessIndex(ForceInit)
+		, XPBDTetherScaleIndex(ForceInit)
 	{
 		NumTethers = 0;
 		for (const TConstArrayView<FTether>& TetherBatch : Tethers)
diff --git a/Engine/Source/Runtime/Experimental/Chaos/Public/Chaos/XPBDSpringConstraints.h b/Engine/Source/Runtime/Experimental/Chaos/Public/Chaos/XPBDSpringConstraints.h
--- a/Engine/Source/Runtime/Experimental/Chaos/Public/Chaos/XPBDSpringConstraints.h
+++ b/Engine/Source/Runtime/Experimental/Chaos/Public/Chaos/XPBDSpringConstraints.h
@@ -188,6 +188,8 @@ public:
 			FSolverVec2(GetWeightedFloatXPBDEdgeSpringStiffness(PropertyCollection, MaxStiffness)),
 			FSolverVec2(GetWeightedFloatXPBDEdgeSpringDamping(PropertyCollection, MinDampingRatio)),
 			bTrimKinematicConstraints)
+		, XPBDEdgeSpringStiffnessIndex(PropertyCollection)
+		, XPBDEdgeSpringDampingIndex(PropertyCollection)
 	{}
 
 	UE_DEPRECATED(5.3, "Use weight map constructor instead.")
@@ -210,6 +212,8 @@ public:
 			FSolverVec2(GetWeightedFloatXPBDEdgeSpringStiffness(PropertyCollection, MaxStiffness)),
 			FSolverVec2(GetWeightedFloatXPBDEdgeSpringDamping(PropertyCollection, MinDampingRatio)),
 			bTrimKinematicConstraints)
+		, XPBDEdgeSpringStiffnessIndex(PropertyCollection)
+		, XPBDEdgeSpringDampingIndex(PropertyCollection)
 	{}
 
 	virtual ~FXPBDEdgeSpringConstraints() override = default;
@@ -261,6 +265,8 @@ public:
 			FSolverVec2(GetWeightedFloatXPBDBendingSpringStiffness(PropertyCollection, MaxStiffness)),
 			FSolverVec2(GetWeightedFloatXPBDBendingSpringDamping(PropertyCollection, MinDampingRatio)),
 			bTrimKinematicConstraints)
+		, XPBDBendingSpringStiffnessIndex(PropertyCollection)
+		, XPBDBendingSpringDampingIndex(PropertyCollection)
 	{}
 
 	UE_DEPRECATED(5.3, "Use weight map constructor instead.")
@@ -283,6 +289,8 @@ public:
 			FSolverVec2(GetWeightedFloatXPBDBendingSpringStiffness(PropertyCollection, MaxStiffness)),
 			FSolverVec2(GetWeightedFloatXPBDBendingSpringDamping(PropertyCollection, MinDampingRatio)),
 			bTrimKinematicConstraints)
+		, XPBDBendingSpringStiffnessIndex(PropertyCollection)
+		, XPBDBendingSpringDampingIndex(PropertyCollection)
 	{}
 
 	virtual ~FXPBDBendingSpringConstraints() override = default;
