Chaos - Destruction Determinism Fixes

Previously we were only sorting newly created constraints as they are added to the graph, but Destruction uses the active constraints list directly and that was not sorted. This should also fix non-deteminate issues related to callbacks assuming the callback are themselves deterministic.

diff --git a/Engine/Source/Runtime/Experimental/Chaos/Private/Chaos/Collision/CollisionConstraintAllocator.cpp b/Engine/Source/Runtime/Experimental/Chaos/Private/Chaos/Collision/CollisionConstraintAllocator.cpp
--- a/Engine/Source/Runtime/Experimental/Chaos/Private/Chaos/Collision/CollisionConstraintAllocator.cpp
+++ b/Engine/Source/Runtime/Experimental/Chaos/Private/Chaos/Collision/CollisionConstraintAllocator.cpp
@@ -32,9 +32,42 @@ namespace Chaos
 			return Constraint;
 		}
 
-		void FCollisionContextAllocator::ProcessNewMidPhases(FCollisionConstraintAllocator* Allocator)
+		void FCollisionConstraintAllocator::ProcessNewMidPhases()
 		{
-			// NOTE: Called from the physics thread, and never from a physics task/parallel-for. No need for locks.
+			QUICK_SCOPE_CYCLE_COUNTER(STAT_Collisions_ProcessNewMidPhases);
+
+			FMemMark Mark(FMemStack::Get());
+			TArray<FParticlePairMidPhase*, TMemStackAllocator<>> NewMidPhases;
+
+			// Count the midphases so we can initialize the array
+			int32 NumNewMidPhases = 0;
+			for (TUniquePtr<FCollisionContextAllocator>& ContextAllocator : ContextAllocators)
+			{
+				NumNewMidPhases += ContextAllocator->NewMidPhases.Num();
+			}
+			NewMidPhases.Reserve(NumNewMidPhases);
+
+			// Collect the midphases
+			for (TUniquePtr<FCollisionContextAllocator>& ContextAllocator : ContextAllocators)
+			{
+				NewMidPhases.Append(ContextAllocator->NewMidPhases);
+				ContextAllocator->NewMidPhases.Reset();
+			}
+			check(NewMidPhases.Num() == NumNewMidPhases);
+
+			// For deterministic behaviour we need to sort the midphases so that, when they are added to 
+			// our and each particle's lists, they are in a repeatable order
+			// @todo(chaos): we could sort each context's array and then process them in order here instead
+			if (bIsDeteministic)
+			{
+				QUICK_SCOPE_CYCLE_COUNTER(STAT_Collisions_SortMidPhases);
+
+				NewMidPhases.Sort(
+					[](const FParticlePairMidPhase& L, const FParticlePairMidPhase& R)
+					{
+						return L.GetKey() < R.GetKey();
+					});
+			}
 
 			// Register the midphases with each of their particles and add the midphase to the central list
 			for (FParticlePairMidPhase* MidPhase : NewMidPhases)
@@ -50,23 +83,80 @@ namespace Chaos
 				Particle1->ParticleCollisions().AddMidPhase(Particle1, MidPhase);
 
 #if CHAOS_MIDPHASE_OBJECTPOOL_ENABLED 
-				Allocator->AddMidPhase(FParticlePairMidPhasePtr(MidPhase, FParticlePairMidPhaseDeleter(MidPhasePool)));
+				AddMidPhase(FParticlePairMidPhasePtr(MidPhase, FParticlePairMidPhaseDeleter(MidPhasePool)));
 #else
-				Allocator->AddMidPhase(FParticlePairMidPhasePtr(MidPhase));
+				AddMidPhase(FParticlePairMidPhasePtr(MidPhase));
 #endif
 			}
-			NewMidPhases.Reset();
 		}
 
-		void FCollisionContextAllocator::ProcessNewConstraints(FCollisionConstraintAllocator* Allocator)
+		void FCollisionConstraintAllocator::ProcessNewConstraints()
 		{
-			// NOTE: Called from the physics thread, and never from a physics task/parallel-for. No need for locks.
+			QUICK_SCOPE_CYCLE_COUNTER(STAT_Collisions_ProcessNewConstraints);
+
+			FMemMark Mark(FMemStack::Get());
+			TArray<FPBDCollisionConstraint*, TMemStackAllocator<>> NewConstraints;
+
+			// Count the constraints so we can initialize the array
+			int32 NumNewConstraints = 0;
+			for (TUniquePtr<FCollisionContextAllocator>& ContextAllocator : ContextAllocators)
+			{
+				NumNewConstraints += ContextAllocator->NewActiveConstraints.Num();
+			}
+			NewConstraints.Reserve(NumNewConstraints);
+
+			// Collect the constraints
+			int32 NewConstraintBegin = 0;
+			for (TUniquePtr<FCollisionContextAllocator>& ContextAllocator : ContextAllocators)
+			{
+				NewConstraints.Append(ContextAllocator->NewActiveConstraints);
+				ContextAllocator->NewActiveConstraints.Reset();
+			}
+			check(NewConstraints.Num() == NumNewConstraints);
+
+			// For deterministic behaviour we need to sort the constraints so that any systems that iterate over
+			// active constraints will do so in a repeatable way.
+			// @todo(chaos): we could sort each context's array and then process them in order here instead
+			// @todo(chaos): if the context held the sort keys as well we could avoid some cache misses?
+			if (bIsDeteministic)
+			{
+				QUICK_SCOPE_CYCLE_COUNTER(STAT_Collisions_SortConstraints);
 
-			for (FPBDCollisionConstraint* NewConstraint : NewActiveConstraints)
+				NewConstraints.Sort(
+					[](const FPBDCollisionConstraint& L, const FPBDCollisionConstraint& R)
+					{
+						return L.GetCollisionSortKey() < R.GetCollisionSortKey();
+					});
+			}
+
+			// Add the constraints to the active list (and the CCD list if appropriate)
+			for (FPBDCollisionConstraint* Constraint : NewConstraints)
 			{
-				Allocator->ActivateConstraintImp(NewConstraint);
+				FPBDCollisionConstraintContainerCookie& Cookie = Constraint->GetContainerCookie();
+
+				// Add the constraint to the active list and update its epoch
+				checkSlow(ActiveConstraints.Find(CollisionConstraint) == INDEX_NONE);
+				Cookie.ConstraintIndex = ActiveConstraints.Add(Constraint);
+
+				// If the constraint uses CCD, keep it in another list so we don't have to search the full list for CCD contacts
+				if (Constraint->GetCCDEnabled())
+				{
+					checkSlow(ActiveCCDConstraints.Find(Constraint) == INDEX_NONE);
+					Cookie.CCDConstraintIndex = ActiveCCDConstraints.Add(Constraint);
+				}
+
+				Cookie.LastUsedEpoch = CurrentEpoch;
 			}
-			NewActiveConstraints.Reset();
+		}
+
+
+		void FCollisionConstraintAllocator::EndDetectCollisions()
+		{
+			check(bInCollisionDetectionPhase);
+
+			bInCollisionDetectionPhase = false;
+
+			ProcessNewItems();
 		}
 
 		void FCollisionConstraintAllocator::PruneExpiredMidPhases()
@@ -151,5 +241,27 @@ namespace Chaos
 				});
 		}
 
+		// @todo(chaos): No longer used (see ProcessNewConstraints)
+		void FCollisionConstraintAllocator::SortActiveConstraints()
+		{
+			if (ActiveConstraints.Num())
+			{
+				// Sort constraints for determinism and/or improved memory-access ordering
+				ActiveConstraints.Sort(
+					[](const FPBDCollisionConstraintHandle& L, const FPBDCollisionConstraintHandle& R)
+					{
+						return L.GetContact().GetCollisionSortKey() < R.GetContact().GetCollisionSortKey();
+					});
+
+				// If we re-ordered the array, we need to update indices
+				for (int32 ConstraintIndex = 0; ConstraintIndex < ActiveConstraints.Num(); ++ConstraintIndex)
+				{
+					FPBDCollisionConstraintContainerCookie& Cookie = ActiveConstraints[ConstraintIndex]->GetContainerCookie();
+
+					Cookie.ConstraintIndex = ConstraintIndex;
+				}
+			}
+		}
+
 	}	// namespace Private
 }	// namespace Chaos
\ No newline at end of file
diff --git a/Engine/Source/Runtime/Experimental/Chaos/Private/Chaos/Collision/PBDCollisionConstraint.cpp b/Engine/Source/Runtime/Experimental/Chaos/Private/Chaos/Collision/PBDCollisionConstraint.cpp
--- a/Engine/Source/Runtime/Experimental/Chaos/Private/Chaos/Collision/PBDCollisionConstraint.cpp
+++ b/Engine/Source/Runtime/Experimental/Chaos/Private/Chaos/Collision/PBDCollisionConstraint.cpp
@@ -168,6 +168,7 @@ namespace Chaos
 		// We do not want to overwrite these properties because they are managed by the systems
 		// where the constraint is registsred (allocator, container, graph).
 		const FPBDCollisionConstraintContainerCookie Cookie = ContainerCookie;
+		const Private::FCollisionSortKey SortKey = CollisionSortKey;
 		Private::FPBDIslandConstraint* ConstraintGraphEdge = GetConstraintGraphEdge();
 
 		// Copy everything
@@ -175,6 +176,7 @@ namespace Chaos
 
 		// Restore the system state
 		ContainerCookie = Cookie;
+		CollisionSortKey = SortKey;
 		SetConstraintGraphEdge(ConstraintGraphEdge);
 	}
 
diff --git a/Engine/Source/Runtime/Experimental/Chaos/Private/Chaos/Collision/ParticlePairMidPhase.cpp b/Engine/Source/Runtime/Experimental/Chaos/Private/Chaos/Collision/ParticlePairMidPhase.cpp
--- a/Engine/Source/Runtime/Experimental/Chaos/Private/Chaos/Collision/ParticlePairMidPhase.cpp
+++ b/Engine/Source/Runtime/Experimental/Chaos/Private/Chaos/Collision/ParticlePairMidPhase.cpp
@@ -195,6 +195,58 @@ namespace Chaos
 		return Context.GetAllocator()->CreateConstraint(Particle0, Implicit0, Shape0, BVHParticles0, ShapeRelativeTransform0, Particle1, Implicit1, Shape1, BVHParticles1, ShapeRelativeTransform1, CullDistance, bUseManifold, ShapePairType);
 	}
 
+	// A unique key for a collision between two particles (key is only unique within the particle pair midphase)
+	class FParticlePairMidPhaseCollisionKey
+	{
+	public:
+
+		FParticlePairMidPhaseCollisionKey()
+			: Key(0)
+		{
+		}
+
+		FParticlePairMidPhaseCollisionKey(const int32 InShapeID0, const int32 InShapeID1)
+		{
+			Generate(InShapeID0, InShapeID1);
+		}
+
+		uint64 GetKey() const
+		{
+			return Key;
+		}
+
+		friend bool operator==(const FParticlePairMidPhaseCollisionKey& L, const FParticlePairMidPhaseCollisionKey& R)
+		{
+			return L.Key == R.Key;
+		}
+
+		friend bool operator!=(const FParticlePairMidPhaseCollisionKey& L, const FParticlePairMidPhaseCollisionKey& R)
+		{
+			return !(L == R);
+		}
+
+		friend bool operator<(const FParticlePairMidPhaseCollisionKey& L, const FParticlePairMidPhaseCollisionKey& R)
+		{
+			return L.Key < R.Key;
+		}
+
+	private:
+		void Generate(const int32 InShapeID0, const int32 InShapeID1)
+		{
+			ShapeID0 = InShapeID0;
+			ShapeID1 = InShapeID1;
+		}
+
+		union {
+			struct
+			{
+				int32 ShapeID0;
+				int32 ShapeID1;
+			};
+			uint64 Key;
+		};
+	};
+
 
 	////////////////////////////////////////////////////////////////////////////////////////
 	////////////////////////////////////////////////////////////////////////////////////////
@@ -207,6 +259,7 @@ namespace Chaos
 		const FPerShapeData* InShape0,
 		FGeometryParticleHandle* InParticle1,
 		const FPerShapeData* InShape1,
+		const Private::FCollisionSortKey& InCollisionSortKey,
 		const EContactShapesType InShapePairType,
 		FParticlePairMidPhase& InMidPhase)
 		: MidPhase(InMidPhase)
@@ -215,6 +268,7 @@ namespace Chaos
 		, Particle1(InParticle1)
 		, Shape0(InShape0)
 		, Shape1(InShape1)
+		, CollisionSortKey(InCollisionSortKey)
 		, SphereBoundsCheckSize(0)
 		, LastUsedEpoch(-1)
 		, ShapePairType(InShapePairType)
@@ -229,10 +283,6 @@ namespace Chaos
 			SphereBoundsCheckSize);
 	}
 
-	FSingleShapePairCollisionDetector::~FSingleShapePairCollisionDetector()
-	{
-	}
-
 	FSingleShapePairCollisionDetector::FSingleShapePairCollisionDetector(FSingleShapePairCollisionDetector&& R)
 		: MidPhase(R.MidPhase)
 		, Constraint(MoveTemp(R.Constraint))
@@ -240,12 +290,17 @@ namespace Chaos
 		, Particle1(R.Particle1)
 		, Shape0(R.Shape0)
 		, Shape1(R.Shape1)
+		, CollisionSortKey(R.CollisionSortKey)
 		, SphereBoundsCheckSize(R.SphereBoundsCheckSize)
 		, ShapePairType(R.ShapePairType)
 		, BoundsTestFlags(R.BoundsTestFlags)
 	{
 	}
 
+	FSingleShapePairCollisionDetector::~FSingleShapePairCollisionDetector()
+	{
+	}
+
 	bool FSingleShapePairCollisionDetector::DoBoundsOverlap(const FReal CullDistance, const int32 CurrentEpoch)
 	{
 		PHYSICS_CSV_SCOPED_EXPENSIVE(PhysicsVerbose, NarrowPhase_ShapeBounds);
@@ -342,6 +397,9 @@ namespace Chaos
 		Constraint->GetContainerCookie().MidPhase = &MidPhase;
 		Constraint->GetContainerCookie().bIsMultiShapePair = false;
 		Constraint->GetContainerCookie().CreationEpoch = CurrentEpoch;
+
+		Constraint->SetCollisionSortKey(CollisionSortKey);
+
 		LastUsedEpoch = -1;
 	}
 
@@ -607,7 +665,7 @@ namespace Chaos
 		, Particle0(nullptr)
 		, Particle1(nullptr)
 		, CullDistanceScale(1)
-		, Key()
+		, ParticlePairKey()
 		, LastUsedEpoch(INDEX_NONE)
 		, NumActiveConstraints(0)
 		, ParticleCollisionsIndex0(INDEX_NONE)
@@ -705,14 +763,14 @@ namespace Chaos
 	void FParticlePairMidPhase::Init(
 		FGeometryParticleHandle* InParticle0,
 		FGeometryParticleHandle* InParticle1,
-		const FCollisionParticlePairKey& InKey,
+		const Private::FCollisionParticlePairKey& InParticlePairKey,
 		const FCollisionContext& Context)
 	{
 		PHYSICS_CSV_SCOPED_EXPENSIVE(PhysicsVerbose, NarrowPhase_Filter);
 
 		Particle0 = InParticle0;
 		Particle1 = InParticle1;
-		Key = InKey;
+		ParticlePairKey = InParticlePairKey;
 
 		Flags.bIsActive = true;
 
@@ -920,13 +978,13 @@ namespace Chaos
 				for (int32 ShapeIndex1 = 0; ShapeIndex1 < Shapes1.Num(); ++ShapeIndex1)
 				{
 					const FPerShapeData* Shape1 = Shapes1[ShapeIndex1].Get();
-					TryAddShapePair(Shape0, Shape1);
+					TryAddShapePair(Shape0, ShapeIndex0, Shape1, ShapeIndex1);
 				}
 			}
 		}
 	}
 
-	void FShapePairParticlePairMidPhase::TryAddShapePair(const FPerShapeData* Shape0, const FPerShapeData* Shape1)
+	void FShapePairParticlePairMidPhase::TryAddShapePair(const FPerShapeData* Shape0, const int32 ShapeIndex0, const FPerShapeData* Shape1, const int32 ShapeIndex1)
 	{
 		const FImplicitObject* Implicit0 = Shape0->GetLeafGeometry();
 		const EImplicitObjectType ImplicitType0 = Private::GetImplicitCollisionType(Particle0, Implicit0);
@@ -944,11 +1002,13 @@ namespace Chaos
 			{
 				if (!bSwap)
 				{
-					ShapePairDetectors.Emplace(FSingleShapePairCollisionDetector(Particle0, Shape0, Particle1, Shape1, ShapePairType, *this));
+					const Private::FCollisionSortKey CollisionSortKey = Private::FCollisionSortKey(Particle0, ShapeIndex0, Particle1, ShapeIndex1);
+					ShapePairDetectors.Emplace(FSingleShapePairCollisionDetector(Particle0, Shape0, Particle1, Shape1, CollisionSortKey, ShapePairType, *this));
 				}
 				else
 				{
-					ShapePairDetectors.Emplace(FSingleShapePairCollisionDetector(Particle1, Shape1, Particle0, Shape0, ShapePairType, *this));
+					const Private::FCollisionSortKey CollisionSortKey = Private::FCollisionSortKey(Particle1, ShapeIndex1, Particle0, ShapeIndex0);
+					ShapePairDetectors.Emplace(FSingleShapePairCollisionDetector(Particle1, Shape1, Particle0, Shape0, CollisionSortKey, ShapePairType, *this));
 				}
 			}
 			else
@@ -965,7 +1025,7 @@ namespace Chaos
 		const FReal Dt,
 		const FCollisionContext& Context)
 	{
-		TRACE_COUNTER_INCREMENT(ChaosTraceCounter_MidPhase_NumShapePair);
+		//TRACE_COUNTER_INCREMENT(ChaosTraceCounter_MidPhase_NumShapePair);
 
 		int32 NumActive = 0;
 		if (Flags.bIsCCD)
@@ -1042,7 +1102,7 @@ namespace Chaos
 		const FReal Dt,
 		const FCollisionContext& Context)
 	{
-		TRACE_COUNTER_INCREMENT(ChaosTraceCounter_MidPhase_NumGeneric);
+		//TRACE_COUNTER_INCREMENT(ChaosTraceCounter_MidPhase_NumGeneric);
 
 		const FImplicitObject* Implicit0 = GetParticle0()->Geometry().Get();
 		const FImplicitObject* Implicit1 = GetParticle1()->Geometry().Get();
@@ -1183,7 +1243,7 @@ namespace Chaos
 		const FRigidTransform3 ParticleWorldTransformB = PB->GetTransformPQ();
 		const FRigidTransform3 ParticleTransformAToB = ParticleWorldTransformA.GetRelativeTransform(ParticleWorldTransformB);
 
-		// Detect collisons between Implicit Hierarchy of ParticleA and Implicit Hierarchy of ParticleB
+		// Detect collisions between Implicit Hierarchy of ParticleA and Implicit Hierarchy of ParticleB
 		// Given an ImplicitObject from ParticleA (which we know overlaps the bounds of some parts of ParticleB),
 		// run collision detection on ImplicitA against the implicit object hierarchy of ParticleB.
 		RootImplicitA->VisitLeafObjects(
@@ -1197,7 +1257,7 @@ namespace Chaos
 
 				const FShapeInstance* ShapeInstanceA = GetShapeInstance(ShapeInstancesA, RootObjectIndexA);
 
-				// Detect collisons between ImplicitA and Implicit Hierarchy of ParticleB
+				// Detect collisions between ImplicitA and Implicit Hierarchy of ParticleB
 				RootImplicitB->VisitOverlappingLeafObjects(ShapeBoundsAInB,
 					[this, ParticleA, ImplicitA, ShapeInstanceA, &ParticleWorldTransformA, &RelativeTransformA, LeafObjectIndexA,
 					ParticleB, &ParticleWorldTransformB, &ShapeInstancesB,
@@ -1206,7 +1266,7 @@ namespace Chaos
 					{
 						const FShapeInstance* ShapeInstanceB = GetShapeInstance(ShapeInstancesB, RootObjectIndexB);
 
-						// Detect collisons between ImplicitA and ImplicitB (both leaf implicits)
+						// Detect collisions between ImplicitA and ImplicitB (both leaf implicits)
 						GenerateCollisionsImplicitLeafImplicitLeaf(
 							ParticleA, ImplicitA, ShapeInstanceA, ParticleWorldTransformA, RelativeTransformA, LeafObjectIndexA,
 							ParticleB, ImplicitB, ShapeInstanceB, ParticleWorldTransformB, RelativeTransformB, LeafObjectIndexB,
@@ -1468,7 +1528,7 @@ namespace Chaos
 		}
 #endif
 
-		const FCollisionParticlePairConstraintKey CollisionKey = FCollisionParticlePairConstraintKey(InShape0, InImplicit0, InImplicitId0, InBVHParticles0, InShape1, InImplicit1, InImplicitId1, InBVHParticles1);
+		const FParticlePairMidPhaseCollisionKey CollisionKey = FParticlePairMidPhaseCollisionKey(InImplicitId0, InImplicitId1);
 		FPBDCollisionConstraint* Constraint = FindConstraint(CollisionKey);
 
 		// @todo(chaos): fix key uniqueness guarantee.  We need a truly unique key gen function
@@ -1489,7 +1549,12 @@ namespace Chaos
 		if (Constraint == nullptr)
 		{
 			// NOTE: Using InParticle0 and InParticle1 here because the order may be different to what we have stored
-			Constraint = CreateConstraint(InParticle0, InImplicit0, InShape0, InBVHParticles0, InShapeRelativeTransform0, InParticle1, InImplicit1, InShape1, InBVHParticles1, InShapeRelativeTransform1, CullDistance, ShapePairType, bUseManifold, CollisionKey, Context);
+			const Private::FCollisionSortKey CollisionSortKey = Private::FCollisionSortKey(InParticle0, InImplicitId0, InParticle1, InImplicitId1);
+			Constraint = CreateConstraint(
+				InParticle0, InImplicit0, InShape0, InBVHParticles0, InShapeRelativeTransform0, 
+				InParticle1, InImplicit1, InShape1, InBVHParticles1, InShapeRelativeTransform1, 
+				CollisionKey, CollisionSortKey,
+				CullDistance, ShapePairType, bUseManifold, Context);
 		}
 
 		NewConstraints.Add(Constraint);
@@ -1497,7 +1562,7 @@ namespace Chaos
 		return Constraint;
 	}
 
-	FPBDCollisionConstraint* FGenericParticlePairMidPhase::FindConstraint(const FCollisionParticlePairConstraintKey& CollisionKey)
+	FPBDCollisionConstraint* FGenericParticlePairMidPhase::FindConstraint(const FParticlePairMidPhaseCollisionKey& CollisionKey)
 	{
 		FPBDCollisionConstraintPtr* PConstraint = Constraints.Find(CollisionKey.GetKey());
 		if (PConstraint != nullptr)
@@ -1518,10 +1583,11 @@ namespace Chaos
 		const FPerShapeData* InShape1,
 		const FBVHParticles* BVHParticles1,
 		const FRigidTransform3& ShapeRelativeTransform1,
+		const FParticlePairMidPhaseCollisionKey& CollisionKey,
+		const Private::FCollisionSortKey& CollisionSortKey,
 		const FReal CullDistance,
 		const EContactShapesType ShapePairType,
 		const bool bUseManifold,
-		const FCollisionParticlePairConstraintKey& CollisionKey,
 		const FCollisionContext& Context)
 	{
 		PHYSICS_CSV_SCOPED_EXPENSIVE(PhysicsVerbose, NarrowPhase_CreateConstraint);
@@ -1536,6 +1602,8 @@ namespace Chaos
 		Constraint->GetContainerCookie().bIsMultiShapePair = true;
 		Constraint->GetContainerCookie().CreationEpoch = CurrentEpoch;
 
+		Constraint->SetCollisionSortKey(CollisionSortKey);
+
 		// Is this a CCD constraint?
 		Constraint->SetCCDEnabled(IsCCD());
 
@@ -1740,12 +1808,12 @@ namespace Chaos
 
 		// Find all the expired collisions
 		FMemMark Mark(FMemStack::Get());
-		TArray<uint32, TMemStackAllocator<alignof(uint32)>> Pruned;
+		TArray<uint64, TMemStackAllocator<alignof(uint32)>> Pruned;
 		Pruned.Reserve(Constraints.Num());
 
 		for (auto& KVP : Constraints)
 		{
-			const uint32 CollisionKey = KVP.Key;
+			const uint64 CollisionKey = KVP.Key;
 			FPBDCollisionConstraintPtr& Constraint = KVP.Value;
 
 			// NOTE: Constraints in sleeping islands should be kept alive. They will still be in the graph
@@ -1757,7 +1825,7 @@ namespace Chaos
 		}
 
 		// Destroy expired collisions
-		for (uint32 CollisionKey : Pruned)
+		for (uint64 CollisionKey : Pruned)
 		{
 			Constraints.Remove(CollisionKey);
 		}
diff --git a/Engine/Source/Runtime/Experimental/Chaos/Private/Chaos/PBDCollisionConstraints.cpp b/Engine/Source/Runtime/Experimental/Chaos/Private/Chaos/PBDCollisionConstraints.cpp
--- a/Engine/Source/Runtime/Experimental/Chaos/Private/Chaos/PBDCollisionConstraints.cpp
+++ b/Engine/Source/Runtime/Experimental/Chaos/Private/Chaos/PBDCollisionConstraints.cpp
@@ -91,6 +91,7 @@ namespace Chaos
 	DECLARE_CYCLE_STAT(TEXT("Collisions::Reset"), STAT_Collisions_Reset, STATGROUP_ChaosCollision);
 	DECLARE_CYCLE_STAT(TEXT("Collisions::BeginDetect"), STAT_Collisions_BeginDetect, STATGROUP_ChaosCollision);
 	DECLARE_CYCLE_STAT(TEXT("Collisions::EndDetect"), STAT_Collisions_EndDetect, STATGROUP_ChaosCollision);
+	DECLARE_CYCLE_STAT(TEXT("Collisions::Sort"), STAT_Collisions_Sort, STATGROUP_ChaosCollision);
 	DECLARE_CYCLE_STAT(TEXT("Collisions::DetectProbeCollisions"), STAT_Collisions_DetectProbeCollisions, STATGROUP_ChaosCollision);
 
 #if CHAOS_ABTEST_CONSTRAINTSOLVER_ENABLED
@@ -371,7 +372,7 @@ namespace Chaos
 #if CHAOS_ABTEST_CONSTRAINTSOLVER_ENABLED
 			if (CVars::bCollisionsEnableSolverABTest)
 			{
-				// Create an AB testing collison solver for simd testing
+				// Create an AB testing collision solver for simd testing
 				return MakeUnique<FABTestingCollisionContainerSolver>(
 					MakeUnique<FPBDCollisionContainerSolver>(*this, Priority),
 					MakeUnique<Private::FPBDCollisionContainerSolverSimd>(*this, Priority),
@@ -389,6 +390,12 @@ namespace Chaos
 		return nullptr;
 	}
 
+	void FPBDCollisionConstraints::SetIsDeterministic(const bool bInIsDeterministic)
+	{
+		bIsDeterministic = bInIsDeterministic;
+		ConstraintAllocator.SetIsDeterministic(bInIsDeterministic);
+	}
+
 	void FPBDCollisionConstraints::DisableHandles()
 	{
 		check(NumConstraints() == 0);
@@ -640,21 +647,27 @@ namespace Chaos
 			}
 		}
 
-		// Sort new constraints into a predictable order. This isn't strictly required unless we have
-		// deterministic mode enabled, but we do it always because we can get fairly different behaviour 
-		// from run to run because the collisions detection order is effectively random on multicore machines.
+		// Sort new constraints into a predictable order. This isn't strictly required if we don't need
+		// deterministic behaviour, but without sorting we can get fairly different behaviour from run 
+		// to run because the collisions detection order is effectively random on multicore machines.
 		//
 		// @todo(chaos): this is still not good enough for some types of determinism. Specifically if we 
 		// create two set of objects in a different order but with the same physical positions and other 
 		// state, they will behave differently which is undesirable. To fix this we need a sorting 
-		// mechanism that does not rely on properties like IDs. E.g., some kind of physical state hash?
-		TempCollisions.Sort(
-			[](const FPBDCollisionConstraintHandle& L, const FPBDCollisionConstraintHandle& R)
-			{
-				const uint64 LKey = L.GetContact().GetParticlePairKey().GetKey();
-				const uint64 RKey = R.GetContact().GetParticlePairKey().GetKey();
-				return LKey < RKey;
-			});
+		// mechanism that does not rely on properties like Particle IDs, but it would be expensive.
+		//
+		// NOTE: If bIsDeterministic is true, we have already sorted the active constraints list 
+		// (see EndDetectCollisions) so we don't need to do it again here
+		if (!bIsDeterministic)
+		{
+			SCOPE_CYCLE_COUNTER(STAT_Collisions_Sort);
+
+			TempCollisions.Sort(
+				[](const FPBDCollisionConstraintHandle& L, const FPBDCollisionConstraintHandle& R)
+				{
+					return L.GetContact().GetCollisionSortKey() < R.GetContact().GetCollisionSortKey();
+				});
+		}
 
 		// Add the new constraints to the graph
 		for (FPBDCollisionConstraintHandle* ConstraintHandle : TempCollisions)
diff --git a/Engine/Source/Runtime/Experimental/Chaos/Public/Chaos/Collision/CollisionConstraintAllocator.h b/Engine/Source/Runtime/Experimental/Chaos/Public/Chaos/Collision/CollisionConstraintAllocator.h
--- a/Engine/Source/Runtime/Experimental/Chaos/Public/Chaos/Collision/CollisionConstraintAllocator.h
+++ b/Engine/Source/Runtime/Experimental/Chaos/Public/Chaos/Collision/CollisionConstraintAllocator.h
@@ -149,26 +149,20 @@ namespace Chaos
 				CurrentEpoch = InEpoch;
 			}
 
-			// Process all the new midphases from the parallel collision detection
-			CHAOS_API void ProcessNewMidPhases(FCollisionConstraintAllocator* Allocator);
-
-			// Process all the activated collisions from the parallel collision detection
-			CHAOS_API void ProcessNewConstraints(FCollisionConstraintAllocator* Allocator);
-
 			// Find the midphase for the particle pair if it exists. Every particle holds a list of its midphases. We search "SearchParticle" which should be
 			// ideally be the one with fewer midphases on it.
 			FParticlePairMidPhase* FindMidPhaseImpl(FGeometryParticleHandle* Particle0, FGeometryParticleHandle* Particle1, FGeometryParticleHandle* SearchParticle)
 			{
 				check((SearchParticle == Particle0) || (SearchParticle == Particle1));
 
-				const FCollisionParticlePairKey Key = FCollisionParticlePairKey(Particle0, Particle1);
+				const Private::FCollisionParticlePairKey Key = Private::FCollisionParticlePairKey(Particle0, Particle1);
 				return SearchParticle->ParticleCollisions().FindMidPhase(Key.GetKey());
 			}
 
 			// Create and initialize a midphase for the particle pair. Adds it to the particles' lists of midphases.
 			FParticlePairMidPhase* CreateMidPhase(FGeometryParticleHandle* Particle0, FGeometryParticleHandle* Particle1, const FCollisionContext& Context)
 			{
-				const FCollisionParticlePairKey Key = FCollisionParticlePairKey(Particle0, Particle1);
+				const Private::FCollisionParticlePairKey Key = Private::FCollisionParticlePairKey(Particle0, Particle1);
 
 				// We temporarily hold new midphases as raw pointers and wrap them in a unique ptr later in ProcessNewMidPhases
 #if CHAOS_MIDPHASE_OBJECTPOOL_ENABLED 
@@ -243,6 +237,7 @@ namespace Chaos
 				, ActiveConstraints()
 				, ActiveCCDConstraints()
 				, CurrentEpoch(0)
+				, bIsDeteministic(false)
 				, bInCollisionDetectionPhase(false)
 			{
 			}
@@ -386,13 +381,7 @@ namespace Chaos
 			 * @brief Called after collision detection to clean up
 			 * Prunes unused contacts
 			*/
-			void EndDetectCollisions()
-			{
-				check(bInCollisionDetectionPhase);
-				bInCollisionDetectionPhase = false;
-
-				ProcessNewItems();
-			}
+			CHAOS_API void EndDetectCollisions();
 
 			/**
 			 * @brief Called each tick after the graph is updated to remove unused collisions
@@ -405,13 +394,7 @@ namespace Chaos
 			/**
 			 * Collect all the midphases created on the context allocators (probably on multiple threads) and register them
 			*/
-			void ProcessNewMidPhases()
-			{
-				for (TUniquePtr<FCollisionContextAllocator>& ContextAllocator : ContextAllocators)
-				{
-					ContextAllocator->ProcessNewMidPhases(this);
-				}
-			}
+			CHAOS_API void ProcessNewMidPhases();
 
 			/**
 			 * @brief If we add new constraints after collision detection, do what needs to be done to add them to the system
@@ -459,17 +442,7 @@ namespace Chaos
 			*/
 			void SortConstraintsHandles()
 			{
-				if(ActiveConstraints.Num())
-				{
-					// We need to sort constraints for solver stability
-					// We have to use StableSort so that constraints of the same pair stay in the same order
-					// Otherwise the order within each pair can change due to where they start out in the array
-					// @todo(chaos): we should label each contact (and shape) for things like warm starting GJK
-					// and so we could use that label as part of the key
-					// and then we could use regular Sort (which is faster)			
-					// @todo(chaos): this can be moved to the island and therefoe done in parallel
-					ActiveConstraints.StableSort(ContactConstraintSortPredicate);
-				}
+				SortActiveConstraints();
 			}
 
 			/**
@@ -526,6 +499,14 @@ namespace Chaos
 				}
 			}
 
+			/**
+			 * Enable/Disable determinism (extra sorting steps after collision detection)
+			*/
+			void SetIsDeterministic(const bool bInIsDeterministic)
+			{
+				bIsDeteministic = bInIsDeterministic;
+			}
+
 		private:
 			friend class FCollisionContextAllocator;
 
@@ -562,35 +543,12 @@ namespace Chaos
 			CHAOS_API void PruneExpiredMidPhases();
 
 			// Collect all the constraints activated in the collision tasks and register them (calls ActivateConstraintImp on each)
-			void ProcessNewConstraints()
-			{
-				for (TUniquePtr<FCollisionContextAllocator>& ContextAllocator : ContextAllocators)
-				{
-					ContextAllocator->ProcessNewConstraints(this);
-				}
-			}
-
-			// Register an activated constraint
-			void ActivateConstraintImp(FPBDCollisionConstraint* CollisionConstraint)
-			{
-				FPBDCollisionConstraintContainerCookie& Cookie = CollisionConstraint->GetContainerCookie();
-
-				// Add the constraint to the active list and update its epoch
-				checkSlow(ActiveConstraints.Find(CollisionConstraint) == INDEX_NONE);
-				Cookie.ConstraintIndex = ActiveConstraints.Add(CollisionConstraint);
-
-				// If the constraint uses CCD, keep it in another list so we don't have to search the full list for CCD contacts
-				if (CollisionConstraint->GetCCDEnabled())
-				{
-					checkSlow(ActiveCCDConstraints.Find(CollisionConstraint) == INDEX_NONE);
-					Cookie.CCDConstraintIndex = ActiveCCDConstraints.Add(CollisionConstraint);
-				}
-
-				Cookie.LastUsedEpoch = CurrentEpoch;
-			}
+			CHAOS_API void ProcessNewConstraints();
 
 			CHAOS_API void RemoveActiveConstraint(FPBDCollisionConstraint& Constraint);
 
+			CHAOS_API void SortActiveConstraints();
+
 			// The container that owns the allocator (only needed because new constraints need to know)
 			FPBDCollisionConstraints* CollisionContainer;
 
@@ -611,6 +569,9 @@ namespace Chaos
 			// older than the current Epoch at the end of the tick was not refreshed this tick.
 			int32 CurrentEpoch;
 
+			// Whether we running a deterministic sim
+			bool bIsDeteministic;
+
 			// For assertions
 			bool bInCollisionDetectionPhase;
 		};
diff --git a/Engine/Source/Runtime/Experimental/Chaos/Public/Chaos/Collision/CollisionKeys.h b/Engine/Source/Runtime/Experimental/Chaos/Public/Chaos/Collision/CollisionKeys.h
--- a/Engine/Source/Runtime/Experimental/Chaos/Public/Chaos/Collision/CollisionKeys.h
+++ b/Engine/Source/Runtime/Experimental/Chaos/Public/Chaos/Collision/CollisionKeys.h
@@ -47,95 +47,275 @@ namespace Chaos
 		return ShouldSwapParticleOrder(bIsDynamicOrSleeping0, bIsDynamicOrSleeping1, AreParticlesInPreferredOrder(Particle0, Particle1));
 	}
 
-	/**
-	 * A key which uniquely identifes a particle pair for use by the collision detection system.
-	 * This key will be the same if particles order is reversed.
-	 * @note This uses ParticleID and truncates it from 32 to 31 bits
-	 */
-	class FCollisionParticlePairKey
+	namespace Private
 	{
-	public:
-		using KeyType = uint64;
 
-		FCollisionParticlePairKey()
+		/**
+		 * A key which uniquely identifes a particle pair for use by the collision detection system.
+		 * This key can optionally be the same if particles order is reversed (see bSymmetric constructor parameter).
+		 * @note This uses ParticleID and truncates it to 31 bits (from 32)
+		 */
+		class FCollisionParticlePairKey
 		{
-			Key.Key64 = 0;
-		}
+		public:
+			using KeyType = uint64;
 
-		FCollisionParticlePairKey(const FGeometryParticleHandle* Particle0, const FGeometryParticleHandle* Particle1)
-		{
-			GenerateKey(Particle0, Particle1);
-		}
+			FCollisionParticlePairKey()
+			{
+				Key.Key64 = 0;
+			}
 
-		uint64 GetKey() const
-		{
-			return Key.Key64;
-		}
+			FCollisionParticlePairKey(const FGeometryParticleHandle* Particle0, const FGeometryParticleHandle* Particle1, const bool bSymmetric = true)
+			{
+				GenerateKey(Particle0, Particle1, bSymmetric);
+			}
 
-		friend bool operator<(const FCollisionParticlePairKey& L, const FCollisionParticlePairKey& R)
-		{
-			return L.Key.Key64 < R.Key.Key64;
-		}
+			uint64 GetKey() const
+			{
+				return Key.Key64;
+			}
 
-	private:
-		void GenerateKey(const FGeometryParticleHandle* Particle0, const FGeometryParticleHandle* Particle1)
-		{
-			Key.Key64 = 0;
-			if ((Particle0 != nullptr) && (Particle1 != nullptr))
+			uint32 GetHash() const
 			{
-				const bool bIsLocalID0 = Particle0->ParticleID().LocalID != INDEX_NONE;
-				const bool bIsLocalID1 = Particle1->ParticleID().LocalID != INDEX_NONE;
-				const uint32 ID0 = uint32((bIsLocalID0) ? Particle0->ParticleID().LocalID : Particle0->ParticleID().GlobalID);
-				const uint32 ID1 = uint32((bIsLocalID1) ? Particle1->ParticleID().LocalID : Particle1->ParticleID().GlobalID);
+				// NOTE: GetTypeHash(uint64 V) does not work well if the 64 bit int is actually 2x32 bit ints
+				// so we hash-combine the two 32 bit ints instead
+				return ::HashCombineFast(uint32(Key.Key64 & 0xFFFFFFFF), uint32((Key.Key64 >> 32) & 0xFFFFFFFF));
+			}
 
-				if (ID0 < ID1)
-				{
-					Key.Key32s[0].Key31 = ID0;
-					Key.Key32s[0].IsLocal = bIsLocalID0;
-					Key.Key32s[1].Key31 = ID1;
-					Key.Key32s[1].IsLocal = bIsLocalID1;
-				}
-				else
+			friend bool operator==(const FCollisionParticlePairKey& L, const FCollisionParticlePairKey& R)
+			{
+				return L.Key.Key64 == R.Key.Key64;
+			}
+
+			friend bool operator!=(const FCollisionParticlePairKey& L, const FCollisionParticlePairKey& R)
+			{
+				return L.Key.Key64 != R.Key.Key64;
+			}
+
+			friend bool operator<(const FCollisionParticlePairKey& L, const FCollisionParticlePairKey& R)
+			{
+				return L.Key.Key64 < R.Key.Key64;
+			}
+
+		private:
+			void GenerateKey(const FGeometryParticleHandle* Particle0, const FGeometryParticleHandle* Particle1, const bool bSymmetric)
+			{
+				Key.Key64 = 0;
+				if ((Particle0 != nullptr) && (Particle1 != nullptr))
 				{
-					Key.Key32s[0].Key31 = ID1;
-					Key.Key32s[0].IsLocal = bIsLocalID1;
-					Key.Key32s[1].Key31 = ID0;
-					Key.Key32s[1].IsLocal = bIsLocalID0;
+					const bool bIsLocalID0 = Particle0->ParticleID().LocalID != INDEX_NONE;
+					const bool bIsLocalID1 = Particle1->ParticleID().LocalID != INDEX_NONE;
+					const uint32 ID0 = uint32((bIsLocalID0) ? Particle0->ParticleID().LocalID : Particle0->ParticleID().GlobalID);
+					const uint32 ID1 = uint32((bIsLocalID1) ? Particle1->ParticleID().LocalID : Particle1->ParticleID().GlobalID);
+
+					// If we want Key(A,B) == Key(B,A) swap the order if particle IDs to guarantee this
+					if (!bSymmetric || (ID0 < ID1))
+					{
+						Key.Key32s[0].Key31 = ID0;
+						Key.Key32s[0].IsLocal = bIsLocalID0;
+						Key.Key32s[1].Key31 = ID1;
+						Key.Key32s[1].IsLocal = bIsLocalID1;
+					}
+					else
+					{
+						Key.Key32s[0].Key31 = ID1;
+						Key.Key32s[0].IsLocal = bIsLocalID1;
+						Key.Key32s[1].Key31 = ID0;
+						Key.Key32s[1].IsLocal = bIsLocalID0;
+					}
 				}
 			}
-		}
 
-		struct FParticleIDKey
+			struct FParticleIDKey
+			{
+				uint32 Key31 : 31;
+				uint32 IsLocal : 1;
+			};
+			union FIDKey
+			{
+				uint64 Key64;
+				FParticleIDKey Key32s[2];
+			};
+
+			// This class is sensitive to changes in FParticleID - try to catch that here...
+			static_assert(sizeof(FParticleID) == 8, "FParticleID size does not match FCollisionParticlePairKey (expected 64 bits)");
+			static_assert(sizeof(FParticleID::GlobalID) == 4, "FParticleID::GlobalID size does not match FCollisionParticlePairKey (expected 32 bits)");
+			static_assert(sizeof(FParticleID::LocalID) == 4, "FParticleID::LocalID size does not match FCollisionParticlePairKey (expected 32 bits)");
+			static_assert(sizeof(FParticleIDKey) == 4, "FCollisionParticlePairKey::FParticleIDKey size is not 32 bits");
+			static_assert(sizeof(FIDKey) == 8, "FCollisionParticlePairKey::FIDKey size is not 64 bits");
+
+			FIDKey Key;
+		};
+
+
+		/**
+		 * A key for two shapes within a particle pair. This key is will be unique for the shape pair in the particle pair, 
+		 * but not will be duplicated in different particle pairs.
+		 */
+		class FCollisionShapePairKey
 		{
-			uint32 Key31 : 31;
-			uint32 IsLocal : 1;
+		public:
+			FCollisionShapePairKey()
+				: Key(0)
+			{
+			}
+
+			FCollisionShapePairKey(const int32 InShapeID0, const int32 InShapeID1)
+				: ShapeID0(InShapeID0)
+				, ShapeID1(InShapeID1)
+			{
+			}
+
+			uint64 GetKey() const
+			{
+				return Key;
+			}
+
+			friend bool operator==(const FCollisionShapePairKey& L, const FCollisionShapePairKey& R)
+			{
+				return L.Key == R.Key;
+			}
+
+			friend bool operator!=(const FCollisionShapePairKey& L, const FCollisionShapePairKey& R)
+			{
+				return L.Key != R.Key;
+			}
+
+			friend bool operator<(const FCollisionShapePairKey& L, const FCollisionShapePairKey& R)
+			{
+				return L.Key < R.Key;
+			}
+
+		private:
+			union
+			{
+				struct
+				{
+					uint32 ShapeID0;
+					uint32 ShapeID1;
+				};
+				uint64 Key;
+			};
 		};
-		union FIDKey
+
+		/**
+		 * A key used to sort collisions. This to ensure consistent collision solver order when collisions are generated 
+		 * in a semi-random order (during multi-threaded collision detection).
+		 * NOTE: This version does not result in collisions being sorted by Particle, but it is smaller than FCollisionSortKeyNonHashed
+		 * 
+		 * @todo(chaos): Hashing the particle pair key does not work well at all. Do not use this for the time being (see FCollisionSortKeyNonHashed)
+		 * 
+		 */
+		class FCollisionSortKeyHashed
 		{
-			uint64 Key64;
-			FParticleIDKey Key32s[2];
+		public:
+			FCollisionSortKeyHashed()
+				: Key(0)
+			{
+			}
+
+			FCollisionSortKeyHashed(
+				const FGeometryParticleHandle* InParticle0, const int32 InShapeID0,
+				const FGeometryParticleHandle* InParticle1, const int32 InShapeID1)
+			{
+				check(false);	// Don't use unless we can make a better hash
+				FCollisionParticlePairKey ParticlePairKey = FCollisionParticlePairKey(InParticle0, InParticle1, false);
+				ParticlesKey = ParticlePairKey.GetHash();
+				ShapesKey = ::HashCombineFast(::GetTypeHash(InShapeID0), ::GetTypeHash(InShapeID1));
+			}
+
+			friend bool operator==(const FCollisionSortKeyHashed& L, const FCollisionSortKeyHashed& R)
+			{
+				return L.Key == R.Key;
+			}
+
+			friend bool operator!=(const FCollisionSortKeyHashed& L, const FCollisionSortKeyHashed& R)
+			{
+				return L.Key != R.Key;
+			}
+
+			friend bool operator<(const FCollisionSortKeyHashed& L, const FCollisionSortKeyHashed& R)
+			{
+				return L.Key < R.Key;
+			}
+
+		private:
+			union
+			{
+				struct
+				{
+					uint32 ParticlesKey;
+					uint32 ShapesKey;
+				};
+				uint64 Key;
+			};
 		};
 
-		// This class is sensitive to changes in FParticleID - try to catch that here...
-		static_assert(sizeof(FParticleID) == 8, "FParticleID size does not match FCollisionParticlePairKey (expected 64 bits)");
-		static_assert(sizeof(FParticleID::GlobalID) == 4, "FParticleID::GlobalID size does not match FCollisionParticlePairKey (expected 32 bits)");
-		static_assert(sizeof(FParticleID::LocalID) == 4, "FParticleID::LocalID size does not match FCollisionParticlePairKey (expected 32 bits)");
-		static_assert(sizeof(FParticleIDKey) == 4, "FCollisionParticlePairKey::FParticleIDKey size is not 32 bits");
-		static_assert(sizeof(FIDKey) == 8, "FCollisionParticlePairKey::FIDKey size is not 64 bits");
+		/**
+		 * A key used to sort collisions. This to ensure consistent collision solver order when collisions are generated
+		 * in a semi-random order (during multi-threaded collision detection).
+		 * NOTE: This version also results in collisions being sorted by Particle (as opposed to FCollisionSortKeyHashed)
+		 * but it is larger and the comparison operator is more complex, so sorting cost is higher
+		 */
+		class FCollisionSortKeyNonHashed
+		{
+		public:
+			FCollisionSortKeyNonHashed()
+				: ParticlePairKey()
+				, ShapePairKey()
+			{
+			}
 
-		FIDKey Key;
-	};
+			FCollisionSortKeyNonHashed(
+				const FGeometryParticleHandle* InParticle0, const int32 InShapeID0,
+				const FGeometryParticleHandle* InParticle1, const int32 InShapeID1)
+				: ParticlePairKey(InParticle0, InParticle1, false)
+				, ShapePairKey(InShapeID0, InShapeID1)
+			{
+			}
 
-	/**
-	 * A key which uniquely identifes a collision constraint within a particle pair
-	 * 
-	 * This key only needs to be uinque within the context of a particle pair. There is no
-	 * guarantee of global uniqueness. This key is only used by the FMultiShapePairCollisionDetector
-	 * class which is used for colliding shape pairs where each shape is actually a hierarchy
-	 * of shapes. 
-	 * 
-	 */
-	class FCollisionParticlePairConstraintKey
+			friend bool operator==(const FCollisionSortKeyNonHashed& L, const FCollisionSortKeyNonHashed& R)
+			{
+				return (L.ParticlePairKey == R.ParticlePairKey) && (L.ShapePairKey == R.ShapePairKey);
+			}
+
+			friend bool operator!=(const FCollisionSortKeyNonHashed& L, const FCollisionSortKeyNonHashed& R)
+			{
+				return (L.ParticlePairKey != R.ParticlePairKey) || (L.ShapePairKey != R.ShapePairKey);
+			}
+
+			friend bool operator<(const FCollisionSortKeyNonHashed& L, const FCollisionSortKeyNonHashed& R)
+			{
+				// @todo(chaos): make sure this doesn't add extra branches
+				const bool ParticlePairLess = (L.ParticlePairKey < R.ParticlePairKey);
+				const bool ShapePairLess = (L.ShapePairKey < R.ShapePairKey);
+				const bool ParticlePairNotEqual = (L.ParticlePairKey != R.ParticlePairKey);
+				return ParticlePairNotEqual ? ParticlePairLess : ShapePairLess;
+			}
+
+		private:
+			FCollisionParticlePairKey ParticlePairKey;
+			FCollisionShapePairKey ShapePairKey;
+		};
+
+		// A key used to sort constraints
+		// @todo(chaos): ideally we would use the smallest key possible here, but see comments in FCollisionSortKeyHashed
+		using FCollisionSortKey = FCollisionSortKeyNonHashed;
+	}
+
+
+	///////////////////////////////////////////////////////////////////////////////////////////////
+	///////////////////////////////////////////////////////////////////////////////////////////////
+	//
+	// DEPRECATED STUFF
+	//
+	///////////////////////////////////////////////////////////////////////////////////////////////
+	///////////////////////////////////////////////////////////////////////////////////////////////
+
+
+	using FCollisionParticlePairKey UE_DEPRECATED(5.4, "Not for external use") = Private::FCollisionParticlePairKey;
+
+	class UE_DEPRECATED(5.4, "No longer used") FCollisionParticlePairConstraintKey
 	{
 	public:
 		FCollisionParticlePairConstraintKey()
@@ -153,10 +333,10 @@ namespace Chaos
 		}
 
 		FCollisionParticlePairConstraintKey(
-			const FShapeInstance* Shape0, 
-			const FImplicitObject* Implicit0, 
-			const int32 ImplicitId0, 
-			const FBVHParticles* Simplicial0, 
+			const FShapeInstance* Shape0,
+			const FImplicitObject* Implicit0,
+			const int32 ImplicitId0,
+			const FBVHParticles* Simplicial0,
 			const FShapeInstance* Shape1,
 			const FImplicitObject* Implicit1,
 			const int32 ImplicitId1,
diff --git a/Engine/Source/Runtime/Experimental/Chaos/Public/Chaos/Collision/PBDCollisionConstraint.h b/Engine/Source/Runtime/Experimental/Chaos/Public/Chaos/Collision/PBDCollisionConstraint.h
--- a/Engine/Source/Runtime/Experimental/Chaos/Public/Chaos/Collision/PBDCollisionConstraint.h
+++ b/Engine/Source/Runtime/Experimental/Chaos/Public/Chaos/Collision/PBDCollisionConstraint.h
@@ -721,9 +721,22 @@ namespace Chaos
 		/**
 		 *	A key used to uniquely identify the constraint (it is based on the two particle IDs)
 		 */
-		FCollisionParticlePairKey GetParticlePairKey() const
+		Private::FCollisionParticlePairKey GetParticlePairKey() const
 		{
-			return FCollisionParticlePairKey(GetParticle0(), GetParticle1());
+			return Private::FCollisionParticlePairKey(GetParticle0(), GetParticle1());
+		}
+
+		/**
+		 * A key that uniquely identifies a collision constraint and provides a good sort order for the solver
+		 */
+		Private::FCollisionSortKey GetCollisionSortKey() const
+		{
+			return CollisionSortKey;
+		}
+
+		void SetCollisionSortKey(const Private::FCollisionSortKey& InCollisionSortKey)
+		{
+			CollisionSortKey = InCollisionSortKey;
 		}
 
 		Private::FImplicitBoundsTestFlags GetBoundsTestFlags() const
@@ -837,6 +850,7 @@ namespace Chaos
 
 	private:
 		FPBDCollisionConstraintContainerCookie ContainerCookie;
+		Private::FCollisionSortKey CollisionSortKey;
 		EContactShapesType ShapesType;
 
 		// The shape transforms at the current particle transforms
diff --git a/Engine/Source/Runtime/Experimental/Chaos/Public/Chaos/Collision/ParticlePairMidPhase.h b/Engine/Source/Runtime/Experimental/Chaos/Public/Chaos/Collision/ParticlePairMidPhase.h
--- a/Engine/Source/Runtime/Experimental/Chaos/Public/Chaos/Collision/ParticlePairMidPhase.h
+++ b/Engine/Source/Runtime/Experimental/Chaos/Public/Chaos/Collision/ParticlePairMidPhase.h
@@ -21,6 +21,8 @@ namespace Chaos
 		class FCollisionConstraintAllocator;
 	}
 
+	class FParticlePairMidPhaseCollisionKey;
+
 	/**
 	 * The type of the particle pair midphase.
 	 */
@@ -51,6 +53,7 @@ namespace Chaos
 			const FPerShapeData* InShape0,
 			FGeometryParticleHandle* InParticle1,
 			const FPerShapeData* InShape1,
+			const Private::FCollisionSortKey& InCollisionSortKey,
 			const EContactShapesType InShapePairType, 
 			FParticlePairMidPhase& MidPhase);
 		CHAOS_API FSingleShapePairCollisionDetector(FSingleShapePairCollisionDetector&& R);
@@ -141,6 +144,7 @@ namespace Chaos
 		FGeometryParticleHandle* Particle1;
 		const FPerShapeData* Shape0;
 		const FPerShapeData* Shape1;
+		Private::FCollisionSortKey CollisionSortKey;
 		FRealSingle SphereBoundsCheckSize;
 		int32 LastUsedEpoch;
 		EContactShapesType ShapePairType;
@@ -186,14 +190,14 @@ namespace Chaos
 		CHAOS_API void Init(
 			FGeometryParticleHandle* InParticle0,
 			FGeometryParticleHandle* InParticle1,
-			const FCollisionParticlePairKey& InKey,
+			const Private::FCollisionParticlePairKey& InKey,
 			const FCollisionContext& Context);
 
 		inline FGeometryParticleHandle* GetParticle0() { return Particle0; }
 
 		inline FGeometryParticleHandle* GetParticle1() { return Particle1; }
 
-		inline const FCollisionParticlePairKey& GetKey() const { return Key; }
+		inline const Private::FCollisionParticlePairKey& GetKey() const { return ParticlePairKey; }
 
 		inline bool IsValid() const { return (Particle0 != nullptr) && (Particle1 != nullptr); }
 
@@ -394,7 +398,7 @@ namespace Chaos
 		// A number based on the size of the dynamic objects used to scale cull distance
 		FRealSingle CullDistanceScale;					// 4 bytes
 
-		FCollisionParticlePairKey Key;					// 8 bytes
+		Private::FCollisionParticlePairKey ParticlePairKey;		// 8 bytes
 
 		int32 LastUsedEpoch;							// 4 bytes
 		int32 NumActiveConstraints;						// 4 bytes
@@ -435,8 +439,10 @@ namespace Chaos
 
 	private:
 		CHAOS_API void TryAddShapePair(
-			const FPerShapeData* Shape0,
-			const FPerShapeData* Shape1);
+			const FPerShapeData* Shape0, 
+			const int32 ShapeIndex0, 
+			const FPerShapeData* Shape1, 
+			const int32 ShapeIndex1);
 
 		TArray<FSingleShapePairCollisionDetector, TInlineAllocator<1>> ShapePairDetectors;	// 88 bytes
 	};
@@ -514,12 +520,13 @@ namespace Chaos
 			FGeometryParticleHandle* InParticle1, const FImplicitObject* InImplicit1, const int32 InImplicitId1, const FShapeInstance* InShape1, const FBVHParticles* InBVHParticles1, const FRigidTransform3& InShapeRelativeTransform1,
 			const FReal CullDistance, const EContactShapesType ShapePairType, const bool bUseManifold, const bool bEnableSweep, const FCollisionContext& Context);
 
-		CHAOS_API FPBDCollisionConstraint* FindConstraint(const FCollisionParticlePairConstraintKey& Key);
+		CHAOS_API FPBDCollisionConstraint* FindConstraint(const FParticlePairMidPhaseCollisionKey& CollisionKey);
 
 		CHAOS_API FPBDCollisionConstraint* CreateConstraint(
 			FGeometryParticleHandle* Particle0, const FImplicitObject* Implicit0, const FPerShapeData* Shape0, const FBVHParticles* BVHParticles0, const FRigidTransform3& ShapeRelativeTransform0,
 			FGeometryParticleHandle* Particle1, const FImplicitObject* Implicit1, const FPerShapeData* Shape1, const FBVHParticles* BVHParticles1, const FRigidTransform3& ShapeRelativeTransform1,
-			const FReal CullDistance, const EContactShapesType ShapePairType, const bool bInUseManifold, const FCollisionParticlePairConstraintKey& Key, const FCollisionContext& Context);
+			const FParticlePairMidPhaseCollisionKey& CollisionKey, const Private::FCollisionSortKey& CollisionSortKey, 
+			const FReal CullDistance, const EContactShapesType ShapePairType, const bool bInUseManifold, const FCollisionContext& Context);
 
 		CHAOS_API int32 ProcessNewConstraints(const FReal CullDistance, const FReal Dt, const FCollisionContext& Context);
 
@@ -529,7 +536,7 @@ namespace Chaos
 
 		CHAOS_API bool UpdateCollisionCCD( FPBDCollisionConstraint* Constraint, const FReal CullDistance, const FReal Dt, const FCollisionContext& Context);
 
-		TMap<uint32, FPBDCollisionConstraintPtr> Constraints;
+		TMap<uint64, FPBDCollisionConstraintPtr> Constraints;
 		TArray<FPBDCollisionConstraint*> NewConstraints;
 	};
 }
diff --git a/Engine/Source/Runtime/Experimental/Chaos/Public/Chaos/PBDCollisionConstraints.h b/Engine/Source/Runtime/Experimental/Chaos/Public/Chaos/PBDCollisionConstraints.h
--- a/Engine/Source/Runtime/Experimental/Chaos/Public/Chaos/PBDCollisionConstraints.h
+++ b/Engine/Source/Runtime/Experimental/Chaos/Public/Chaos/PBDCollisionConstraints.h
@@ -83,10 +83,7 @@ public:
 	 * @brief Enable or disable determinism.
 	 * Support for determinism requires that we sort active constraints each tick, so there is additional cost.
 	*/
-	void SetIsDeterministic(const bool bInIsDeterministic)
-	{
-		bIsDeterministic = bInIsDeterministic;
-	}
+	void SetIsDeterministic(const bool bInIsDeterministic);
 
 	/**
 	 *  Clears the list of active constraints.
