Iris
* Improve error message when NetRPC is discarded

#rnx

diff --git a/Engine/Source/Runtime/Experimental/Iris/Core/Private/Iris/ReplicationSystem/NetBlob/NetRPC.cpp b/Engine/Source/Runtime/Experimental/Iris/Core/Private/Iris/ReplicationSystem/NetBlob/NetRPC.cpp
--- a/Engine/Source/Runtime/Experimental/Iris/Core/Private/Iris/ReplicationSystem/NetBlob/NetRPC.cpp
+++ b/Engine/Source/Runtime/Experimental/Iris/Core/Private/Iris/ReplicationSystem/NetBlob/NetRPC.cpp
@@ -28,8 +28,14 @@ namespace UE::Net::Private
 {
 
 static bool NetRPC_GetFunctionLocator(const UReplicationSystem* ReplicationSystem, const FNetObjectReference& ObjectReference, const UFunction* Function, FNetRPC::FFunctionLocator& OutFunctionLocator, const FReplicationStateMemberFunctionDescriptor*& OutFunctionDescriptor);
-static UObject* NetRPC_GetObject(FNetSerializationContext& Context, const FNetObjectReference& ObjectReference, const FNetObjectReference& SubObjectReference);							                
-static bool NetRPC_GetFunctionAndObject(FNetSerializationContext& Context, const FNetObjectReference& ObjectReference, const FNetObjectReference& SubObjectReferece, FNetRPC::FFunctionLocator& FunctionLocator, const FReplicationStateMemberFunctionDescriptor*& OutFunctionDescriptor, TWeakObjectPtr<UObject>& OutObject);							                
+
+static UObject* NetRPC_GetObject(FNetSerializationContext& Context, const FNetObjectReference& ObjectReference);
+static UObject* NetRPC_GetObject(FNetSerializationContext& Context, const FNetObjectReference& RootObjectReference, const FNetObjectReference& SubObjectReference);							                
+
+static bool NetRPC_GetFunctionAndObject(FNetSerializationContext& Context, const FNetObjectReference& RootObjectReference, const FNetObjectReference& SubObjectReferece, const FNetRPC::FFunctionLocator& FunctionLocator, const FReplicationStateMemberFunctionDescriptor*& OutFunctionDescriptor, TWeakObjectPtr<UObject>& OutObject);
+
+static FString NetRPC_GetDebugFunctionName(const UReplicationSystem* ReplicationSystem, const FNetObjectReference& ObjectReference, const FNetRPC::FFunctionLocator& FunctionLocator);
+static FString NetRPC_GetDebugObjectRefName(FNetSerializationContext& Context, const FNetObjectReference& ObjectReference, UObject* ObjectPtr = nullptr);
 
 static const FName NetError_InvalidNetObjectReference("Invalid NetObjectRefererence");
 static const FName NetError_UnknownFunction("Unknown RPC");
@@ -418,6 +424,8 @@ void FNetRPC::CallFunction(FNetSerializationContext& Context)
 	CSV_SCOPED_TIMING_STAT_EXCLUSIVE(HandleRPC);
 #endif
 
+	const UReplicationSystem* ReplicationSystem = Context.GetInternalContext()->ReplicationSystem;
+
 	// Check whether we are ok with calling the function
 	UObject* Object = ObjectPtr.Get();
 
@@ -428,11 +436,21 @@ void FNetRPC::CallFunction(FNetSerializationContext& Context)
 
 	if (Object == nullptr || Function == nullptr)
 	{
-		UE_LOG(LogIrisRpc, Error, TEXT("Rejected RPC function due missing object or function."));
+		if (!TargetObjectReference.IsValid())
+		{
+			UE_LOG(LogIrisRpc, Error, TEXT("Rejected RPC %s due to missing object or function for object: %s."),
+				*(Function?Function->GetName():NetRPC_GetDebugFunctionName(ReplicationSystem, NetObjectReference, FunctionLocator)),
+				*NetRPC_GetDebugObjectRefName(Context, NetObjectReference, Object));
+		}
+		else
+		{
+			UE_LOG(LogIrisRpc, Error, TEXT("Rejected RPC %s due to missing object or function for subobject: %s of rootobject: %s"),
+				*(Function ? Function->GetName() : NetRPC_GetDebugFunctionName(ReplicationSystem, NetObjectReference, FunctionLocator)),
+				*NetRPC_GetDebugObjectRefName(Context, TargetObjectReference, Object), *NetRPC_GetDebugObjectRefName(Context, NetObjectReference));
+		}
 		return;
 	}
 
-	const UReplicationSystem* ReplicationSystem = Context.GetInternalContext()->ReplicationSystem;
 	const bool bIsServer = ReplicationSystem->IsServer();
 	if (bIsServer)
 	{
@@ -537,86 +555,122 @@ bool FNetRPC::IsServerAllowedToExecuteRPC(FNetSerializationContext& Context) con
 	return bTargetIsOwnedByConnection;
 }
 
-	static bool NetRPC_GetFunctionLocator(const UReplicationSystem* ReplicationSystem, const FNetObjectReference& ObjectReference, const UFunction* Function, FNetRPC::FFunctionLocator& OutFunctionLocator, const FReplicationStateMemberFunctionDescriptor*& OutFunctionDescriptor)
+static bool NetRPC_GetFunctionLocator(const UReplicationSystem* ReplicationSystem, const FNetObjectReference& ObjectReference, const UFunction* Function, FNetRPC::FFunctionLocator& OutFunctionLocator, const FReplicationStateMemberFunctionDescriptor*& OutFunctionDescriptor)
+{
+	const UObjectReplicationBridge* Bridge = Cast<UObjectReplicationBridge>(ReplicationSystem->GetReplicationBridge());
+	if (!ensure(Bridge != nullptr))
 	{
-		const UObjectReplicationBridge* Bridge = Cast<UObjectReplicationBridge>(ReplicationSystem->GetReplicationBridge());
-		if (!ensure(Bridge != nullptr))
-		{
-			return false;
-		}
+		return false;
+	}
 
-		const FReplicationProtocol* Protocol = ReplicationSystem->GetReplicationProtocol(ObjectReference.GetRefHandle());
-		if (!ensure(Protocol != nullptr))
-		{
-			return false;
-		}
+	const FReplicationProtocol* Protocol = ReplicationSystem->GetReplicationProtocol(ObjectReference.GetRefHandle());
+	if (!ensure(Protocol != nullptr))
+	{
+		return false;
+	}
 
-		for (const FReplicationStateDescriptor*& Descriptor : MakeArrayView(Protocol->ReplicationStateDescriptors, Protocol->ReplicationStateCount))
+	for (const FReplicationStateDescriptor*& Descriptor : MakeArrayView(Protocol->ReplicationStateDescriptors, Protocol->ReplicationStateCount))
+	{
+		for (const FReplicationStateMemberFunctionDescriptor& FunctionDescriptor : MakeArrayView(Descriptor->MemberFunctionDescriptors, Descriptor->FunctionCount))
 		{
-			for (const FReplicationStateMemberFunctionDescriptor& FunctionDescriptor : MakeArrayView(Descriptor->MemberFunctionDescriptors, Descriptor->FunctionCount))
+			if (FunctionDescriptor.Function == Function)
 			{
-				if (FunctionDescriptor.Function == Function)
-				{
-					OutFunctionLocator.DescriptorIndex = static_cast<uint16>(&Descriptor - Protocol->ReplicationStateDescriptors);
-					OutFunctionLocator.FunctionIndex = static_cast<uint16>(&FunctionDescriptor - Descriptor->MemberFunctionDescriptors);
-					OutFunctionDescriptor = &FunctionDescriptor;
-					return true;
-				}
+				OutFunctionLocator.DescriptorIndex = static_cast<uint16>(&Descriptor - Protocol->ReplicationStateDescriptors);
+				OutFunctionLocator.FunctionIndex = static_cast<uint16>(&FunctionDescriptor - Descriptor->MemberFunctionDescriptors);
+				OutFunctionDescriptor = &FunctionDescriptor;
+				return true;
 			}
 		}
+	}
+
+	return false;
+}
 
+static UObject* NetRPC_GetObject(FNetSerializationContext& Context, const FNetObjectReference& RootObjectReference, const FNetObjectReference& SubObjectReference)
+{
+	return NetRPC_GetObject(Context, SubObjectReference.IsValid() ? SubObjectReference : RootObjectReference);
+}
+
+static UObject* NetRPC_GetObject(FNetSerializationContext& Context, const FNetObjectReference& ObjectReference)
+{
+	FInternalNetSerializationContext* InternalContext = Context.GetInternalContext();
+	return InternalContext->ObjectReferenceCache->ResolveObjectReference(ObjectReference, InternalContext->ResolveContext);
+}
+
+static bool NetRPC_GetFunctionAndObject(FNetSerializationContext& Context, const FNetObjectReference& ObjectReference, const FNetObjectReference& SubObjectReference, const FNetRPC::FFunctionLocator& FunctionLocator, const FReplicationStateMemberFunctionDescriptor*& OutFunctionDescriptor, TWeakObjectPtr<UObject>& OutObject)
+{
+	const UReplicationSystem* ReplicationSystem = Context.GetInternalContext()->ReplicationSystem;
+	const UObjectReplicationBridge* Bridge = Cast<UObjectReplicationBridge>(ReplicationSystem->GetReplicationBridge());
+	if (!ensure(Bridge != nullptr))
+	{
+		Context.SetError(NetError_InvalidNetObjectReference);
 		return false;
 	}
 
-	static UObject* NetRPC_GetObject(FNetSerializationContext& Context, const FNetObjectReference& ObjectReference, const FNetObjectReference& SubObjectReference)
+	UObject* RefObject = NetRPC_GetObject(Context, ObjectReference, SubObjectReference);
+	if (!RefObject)
 	{
-		FInternalNetSerializationContext* InternalContext = Context.GetInternalContext();
-		return InternalContext->ObjectReferenceCache->ResolveObjectReference(SubObjectReference.IsValid() ? SubObjectReference : ObjectReference, InternalContext->ResolveContext);
+		// Ignore this RPC and continue processing the rest of the data
+        return false;
 	}
 
-	static bool NetRPC_GetFunctionAndObject(FNetSerializationContext& Context, const FNetObjectReference& ObjectReference, const FNetObjectReference& SubObjectReference, FNetRPC::FFunctionLocator& FunctionLocator, const FReplicationStateMemberFunctionDescriptor*& OutFunctionDescriptor, TWeakObjectPtr<UObject>& OutObject)
+	const FReplicationProtocol* Protocol = ReplicationSystem->GetReplicationProtocol(ObjectReference.GetRefHandle());
+	if (!ensureMsgf(Protocol != nullptr, TEXT("ReplicationProtocol doesn't exist for %s (Connection %u). Ignoring RPC (%u|%u)"), 
+		*GetNameSafe(RefObject), Context.GetLocalConnectionId(), FunctionLocator.DescriptorIndex, FunctionLocator.FunctionIndex))
 	{
-		const UReplicationSystem* ReplicationSystem = Context.GetInternalContext()->ReplicationSystem;
-		const UObjectReplicationBridge* Bridge = Cast<UObjectReplicationBridge>(ReplicationSystem->GetReplicationBridge());
-		if (!ensure(Bridge != nullptr))
-		{
-			Context.SetError(NetError_InvalidNetObjectReference);
-			return false;
-		}
+		Context.SetError(NetError_InvalidNetObjectReference);
+		return false;
+	}
 
-		UObject* RefObject = NetRPC_GetObject(Context, ObjectReference, SubObjectReference);
-		if (!RefObject)
-		{
-			// Ignore this RPC and continue processing the rest of the data
-            return false;
-		}
+	if (!ensure(FunctionLocator.DescriptorIndex < Protocol->ReplicationStateCount))
+	{
+		Context.SetError(NetError_UnknownFunction);
+		return false;
+	}
 
-		const FReplicationProtocol* Protocol = ReplicationSystem->GetReplicationProtocol(ObjectReference.GetRefHandle());
-		if (!ensureMsgf(Protocol != nullptr, TEXT("ReplicationProtocol doesn't exist for %s (Connection %u). Ignoring RPC (%u|%u)"), 
-			*GetNameSafe(RefObject), Context.GetLocalConnectionId(), FunctionLocator.DescriptorIndex, FunctionLocator.FunctionIndex))
-		{
-			Context.SetError(NetError_InvalidNetObjectReference);
-			return false;
-		}
+	const FReplicationStateDescriptor* Descriptor = Protocol->ReplicationStateDescriptors[FunctionLocator.DescriptorIndex];
+	if (!ensure(FunctionLocator.FunctionIndex < Descriptor->FunctionCount))
+	{
+		Context.SetError(NetError_UnknownFunction);
+		return false;
+	}
 
-		if (!ensure(FunctionLocator.DescriptorIndex < Protocol->ReplicationStateCount))
-		{
-			Context.SetError(NetError_UnknownFunction);
-			return false;
-		}
+	OutFunctionDescriptor = &Descriptor->MemberFunctionDescriptors[FunctionLocator.FunctionIndex];
 
-		const FReplicationStateDescriptor* Descriptor = Protocol->ReplicationStateDescriptors[FunctionLocator.DescriptorIndex];
-		if (!ensure(FunctionLocator.FunctionIndex < Descriptor->FunctionCount))
-		{
-			Context.SetError(NetError_UnknownFunction);
-			return false;
-		}
+	OutObject = RefObject;
 
-		OutFunctionDescriptor = &Descriptor->MemberFunctionDescriptors[FunctionLocator.FunctionIndex];
+	return true;
+}
 
-		OutObject = RefObject;
+static FString NetRPC_GetDebugFunctionName(const UReplicationSystem* ReplicationSystem, const FNetObjectReference& ObjectReference, const FNetRPC::FFunctionLocator& FunctionLocator)
+{
+	const FReplicationProtocol* Protocol = ReplicationSystem->GetReplicationProtocol(ObjectReference.GetRefHandle());
+	if (!Protocol)
+	{
+		return TEXT("ProtocolNotFound");
+	}
 
-		return true;
+	if (FunctionLocator.DescriptorIndex >= Protocol->ReplicationStateCount)
+	{
+		return TEXT("InvalidDescriptorIndex");
 	}
 
+	const FReplicationStateDescriptor* Descriptor = Protocol->ReplicationStateDescriptors[FunctionLocator.DescriptorIndex];
+	if (FunctionLocator.FunctionIndex >= Descriptor->FunctionCount)
+	{
+		return TEXT("InvalidFunctionIndex");
+	}
+
+	return Descriptor->MemberFunctionDescriptors[FunctionLocator.FunctionIndex].Function->GetName();
+}
+
+static FString NetRPC_GetDebugObjectRefName(FNetSerializationContext& Context, const FNetObjectReference& ObjectReference, UObject* ObjectPtr)
+{
+	if (!ObjectPtr)
+	{
+		ObjectPtr = NetRPC_GetObject(Context, ObjectReference);
+	}
+	return FString::Printf(TEXT("%s (%s)"), *GetNameSafe(ObjectPtr), *ObjectReference.ToString());
+}
+
 }
