Fix debug build compile error

diff --git a/Engine/Source/Runtime/Experimental/Chaos/Private/Chaos/Collision/CollisionConstraintAllocator.cpp b/Engine/Source/Runtime/Experimental/Chaos/Private/Chaos/Collision/CollisionConstraintAllocator.cpp
--- a/Engine/Source/Runtime/Experimental/Chaos/Private/Chaos/Collision/CollisionConstraintAllocator.cpp
+++ b/Engine/Source/Runtime/Experimental/Chaos/Private/Chaos/Collision/CollisionConstraintAllocator.cpp
@@ -135,7 +135,7 @@ namespace Chaos
 				FPBDCollisionConstraintContainerCookie& Cookie = Constraint->GetContainerCookie();
 
 				// Add the constraint to the active list and update its epoch
-				checkSlow(ActiveConstraints.Find(CollisionConstraint) == INDEX_NONE);
+				checkSlow(ActiveConstraints.Find(Constraint) == INDEX_NONE);
 				Cookie.ConstraintIndex = ActiveConstraints.Add(Constraint);
 
 				// If the constraint uses CCD, keep it in another list so we don't have to search the full list for CCD contacts
