Iris
* Now storing list of polled objects in the current frame
* Only polled objects will get copied and clean now.
* Fixes crash when bEnableForceNetUpdate is true caused by copying data from a dirty object that went from non-relevant to relevant before we polled it again.

#rnx

diff --git a/Engine/Source/Runtime/Experimental/Iris/Core/Private/Iris/ReplicationSystem/NetRefHandleManager.cpp b/Engine/Source/Runtime/Experimental/Iris/Core/Private/Iris/ReplicationSystem/NetRefHandleManager.cpp
--- a/Engine/Source/Runtime/Experimental/Iris/Core/Private/Iris/ReplicationSystem/NetRefHandleManager.cpp
+++ b/Engine/Source/Runtime/Experimental/Iris/Core/Private/Iris/ReplicationSystem/NetRefHandleManager.cpp
@@ -24,7 +24,7 @@ FNetRefHandleManager::FNetRefHandleManager(FReplicationProtocolManager& InReplic
 , ScopableInternalIndices(MaxActiveObjectCount)
 , PrevFrameScopableInternalIndices(MaxActiveObjectCount)
 , RelevantObjectsInternalIndices(MaxActiveObjectCount)
-, CleanedObjectsInternalIndices(MaxActiveObjectCount)
+, PolledObjectsInternalIndices(MaxActiveObjectCount)
 , AssignedInternalIndices(MaxActiveObjectCount)
 , SubObjectInternalIndices(MaxActiveObjectCount)
 , DependentObjectInternalIndices(MaxActiveObjectCount)
diff --git a/Engine/Source/Runtime/Experimental/Iris/Core/Private/Iris/ReplicationSystem/NetRefHandleManager.h b/Engine/Source/Runtime/Experimental/Iris/Core/Private/Iris/ReplicationSystem/NetRefHandleManager.h
--- a/Engine/Source/Runtime/Experimental/Iris/Core/Private/Iris/ReplicationSystem/NetRefHandleManager.h
+++ b/Engine/Source/Runtime/Experimental/Iris/Core/Private/Iris/ReplicationSystem/NetRefHandleManager.h
@@ -194,8 +194,8 @@ public:
 	/** List of objects that are always relevant or currently relevant to at least one connection. */
 	FNetBitArrayView GetRelevantObjectsInternalIndices() const { return MakeNetBitArrayView(RelevantObjectsInternalIndices); }
 
-	/** List of objects on whom we copied their dirty state this frame */
-	FNetBitArrayView GetCleanedObjectsInternalIndices() const { return MakeNetBitArrayView(CleanedObjectsInternalIndices); }
+	/** List of objects that we polled this frame */
+	FNetBitArrayView GetPolledObjectsInternalIndices() const { return MakeNetBitArrayView(PolledObjectsInternalIndices); }
 
 	// Get bitarray for all internal indices that currently are assigned
 	const FNetBitArray& GetAssignedInternalIndices() const { return AssignedInternalIndices; }
@@ -316,8 +316,8 @@ private:
 	/** This contains the ScopableInternalIndices list minus filtered objects that are not relevant to any connection this frame. */
 	FNetBitArray RelevantObjectsInternalIndices;
 
-	/** List of objects that were dirty but whom we copied their dirty data this frame */
-	FNetBitArray CleanedObjectsInternalIndices;
+	/** List of objects that we polled this frame */
+	FNetBitArray PolledObjectsInternalIndices;
 
 	// Bitset containing all internal indices that are assigned
 	FNetBitArray AssignedInternalIndices;
diff --git a/Engine/Source/Runtime/Experimental/Iris/Core/Private/Iris/ReplicationSystem/ObjectReplicationBridge.cpp b/Engine/Source/Runtime/Experimental/Iris/Core/Private/Iris/ReplicationSystem/ObjectReplicationBridge.cpp
--- a/Engine/Source/Runtime/Experimental/Iris/Core/Private/Iris/ReplicationSystem/ObjectReplicationBridge.cpp
+++ b/Engine/Source/Runtime/Experimental/Iris/Core/Private/Iris/ReplicationSystem/ObjectReplicationBridge.cpp
@@ -739,8 +739,8 @@ void UObjectReplicationBridge::PreUpdateAndPoll()
 	const FNetBitArrayView WantToBeDormantObjects = MakeNetBitArrayView(LocalNetRefHandleManager.GetWantToBeDormantInternalIndices());
 
 	// Filter the set of objects considered for pre-update and polling
-	FNetBitArray ObjectsConsideredForPolling(LocalNetRefHandleManager.GetMaxActiveObjectCount());
-	FNetBitArrayView ObjectsConsideredForPollingView = MakeNetBitArrayView(ObjectsConsideredForPolling);
+	FNetBitArrayView ObjectsConsideredForPolling = LocalNetRefHandleManager.GetPolledObjectsInternalIndices();
+	ObjectsConsideredForPolling.Reset();
 
 	// We always want to consider objects marked as dirty and their subobjects
 	const FNetBitArrayView AccumulatedDirtyObjects = ReplicationSystemInternal->GetDirtyNetObjectTracker().GetAccumulatedDirtyNetObjects();
@@ -750,16 +750,16 @@ void UObjectReplicationBridge::PreUpdateAndPoll()
 		if (bEnableForceNetUpdate)
 		{
 			const FNetBitArrayView ForceNetUpdateObjects = ReplicationSystemInternal->GetDirtyNetObjectTracker().GetForceNetUpdateObjects();
-			PollFrequencyLimiter->Update(RelevantObjects, ForceNetUpdateObjects, ObjectsConsideredForPollingView);
+			PollFrequencyLimiter->Update(RelevantObjects, ForceNetUpdateObjects, ObjectsConsideredForPolling);
 		}
 		else
 		{
-			PollFrequencyLimiter->Update(RelevantObjects, AccumulatedDirtyObjects, ObjectsConsideredForPollingView);
+			PollFrequencyLimiter->Update(RelevantObjects, AccumulatedDirtyObjects, ObjectsConsideredForPolling);
 		}
 	}
 	else
 	{
-		ObjectsConsideredForPollingView.Copy(RelevantObjects);
+		ObjectsConsideredForPolling.Copy(RelevantObjects);
 	}
 
 	// Mask off objects pending dormancy as we do not want to poll/pre-update them unless they are marked for flush or are dirty
@@ -768,7 +768,7 @@ void UObjectReplicationBridge::PreUpdateAndPoll()
 		IRIS_PROFILER_SCOPE(PreUpdateAndPoll_Dormancy);
 
 		// Mask off objects pending dormancy that are not dirty
-		ObjectsConsideredForPollingView.CombineMultiple(FNetBitArrayView::AndNotOp, WantToBeDormantObjects, FNetBitArrayView::AndNotOp, AccumulatedDirtyObjects);
+		ObjectsConsideredForPolling.CombineMultiple(FNetBitArrayView::AndNotOp, WantToBeDormantObjects, FNetBitArrayView::AndNotOp, AccumulatedDirtyObjects);
 
 		FNetBitArrayView ForceNetUpdateObjects = ReplicationSystemInternal->GetDirtyNetObjectTracker().GetForceNetUpdateObjects();
 
@@ -848,7 +848,8 @@ void UObjectReplicationBridge::PreUpdateAndPoll()
 		{
 			IRIS_PROFILER_SCOPE(PreUpdateAndPoll_PatchDependentObjects);
 
-			FNetBitArray TempObjectsConsideredForPolling(ObjectsConsideredForPolling);
+			FNetBitArray TempObjectsConsideredForPolling;
+			TempObjectsConsideredForPolling.InitAndCopy(ObjectsConsideredForPolling);
 			FNetBitArray::ForAllSetBits(TempObjectsConsideredForPolling, NetRefHandleManager->GetObjectsWithDependentObjectsInternalIndices(), FNetBitArray::AndOp,
 				[&LocalNetRefHandleManager, &ObjectsConsideredForPolling](FInternalNetRefIndex ObjectIndex) 
 				{
@@ -869,7 +870,7 @@ void UObjectReplicationBridge::PreUpdateAndPoll()
 		PollerInitParams.ReplicationSystemInternal = GetReplicationSystem()->GetReplicationSystemInternal();
 
 		FObjectPoller Poller(PollerInitParams);
-		Poller.PollObjects(ObjectsConsideredForPollingView);
+		Poller.PollObjects(ObjectsConsideredForPolling);
 	
 		FObjectPoller::FPreUpdateAndPollStats Stats = Poller.GetPollStats();
 
diff --git a/Engine/Source/Runtime/Experimental/Iris/Core/Private/Iris/ReplicationSystem/ReplicationSystem.cpp b/Engine/Source/Runtime/Experimental/Iris/Core/Private/Iris/ReplicationSystem/ReplicationSystem.cpp
--- a/Engine/Source/Runtime/Experimental/Iris/Core/Private/Iris/ReplicationSystem/ReplicationSystem.cpp
+++ b/Engine/Source/Runtime/Experimental/Iris/Core/Private/Iris/ReplicationSystem/ReplicationSystem.cpp
@@ -350,21 +350,16 @@ public:
 
 		SerializationContext.SetInternalContext(&InternalContext);
 
-		// Build the list of dirty objects whom we will copy their data and "clean"
-		FNetBitArrayView CleanedObjects = NetRefHandleManager.GetCleanedObjectsInternalIndices();
-		const FNetBitArrayView AccumulatedDirtyNetObjects = ReplicationSystemInternal.GetDirtyNetObjectTracker().GetAccumulatedDirtyNetObjects();
-		const FNetBitArrayView RelevantObjects = NetRefHandleManager.GetRelevantObjectsInternalIndices();
-
-		CleanedObjects.Set(AccumulatedDirtyNetObjects, FNetBitArrayView::AndOp, RelevantObjects);
-
-
+		// Copy the state data of objects that got polled this frame.
+		const FNetBitArrayView ObjectsToCopy = NetRefHandleManager.GetPolledObjectsInternalIndices();
+		
 		auto CopyFunction = [&ChangeMaskWriter, &Cache, &NetRefHandleManager, &CopiedObjectCount, &SerializationContext](uint32 DirtyIndex)
 		{
 			CopiedObjectCount += FReplicationInstanceOperationsInternal::CopyObjectStateData(ChangeMaskWriter, Cache, NetRefHandleManager, SerializationContext, DirtyIndex);
 		};
 
 		// Only copy both dirty and relevant objects
-		CleanedObjects.ForAllSetBits(CopyFunction);
+		ObjectsToCopy.ForAllSetBits(CopyFunction);
 
 		const uint32 ReplicationSystemId = ReplicationSystem->GetId();
 		UE_NET_TRACE_FRAME_STATSCOUNTER(ReplicationSystemId, ReplicationSystem.CopiedObjectCount, CopiedObjectCount, ENetTraceVerbosity::Trace);
@@ -375,16 +370,18 @@ public:
 		IRIS_PROFILER_SCOPE(FReplicationSystem_ResetObjectStateDirtiness);
 
 		FNetRefHandleManager& NetRefHandleManager = ReplicationSystemInternal.GetNetRefHandleManager();
-		const FNetBitArrayView CleanedObjects = NetRefHandleManager.GetCleanedObjectsInternalIndices();
+
+		// Clean the objects that got polled and copied this frame
+		const FNetBitArrayView ObjectsToClean = NetRefHandleManager.GetPolledObjectsInternalIndices();
 
 		// Reset object dirtyness
-		CleanedObjects.ForAllSetBits([&NetRefHandleManager](uint32 DirtyIndex)
+		ObjectsToClean.ForAllSetBits([&NetRefHandleManager](uint32 DirtyIndex)
 		{
 			FReplicationInstanceOperationsInternal::ResetObjectStateDirtiness(NetRefHandleManager, DirtyIndex);
 		});
 
 		// Reset cleaned objects in the tracker
-		ReplicationSystemInternal.GetDirtyNetObjectTracker().ClearDirtyNetObjects(CleanedObjects);
+		ReplicationSystemInternal.GetDirtyNetObjectTracker().ClearDirtyNetObjects(ObjectsToClean);
 	}
 
 	void ProcessNetObjectAttachmentSendQueue(FNetBlobManager::EProcessMode ProcessMode)
