[ChaosVD] Replacing FThreadSafeBool with std::atomic<bool> in FChaosVisualDebuggerTrace

FThreadSafeBool is deprecated and it was resulting on a perf impact when it was read from multiple threads.This needed to be addressed before enabling compilation of the CVD tool by default in non-shipping builds.

diff --git a/Engine/Source/Runtime/Experimental/Chaos/Private/ChaosVisualDebugger/ChaosVisualDebuggerTrace.cpp b/Engine/Source/Runtime/Experimental/Chaos/Private/ChaosVisualDebugger/ChaosVisualDebuggerTrace.cpp
--- a/Engine/Source/Runtime/Experimental/Chaos/Private/ChaosVisualDebugger/ChaosVisualDebuggerTrace.cpp
+++ b/Engine/Source/Runtime/Experimental/Chaos/Private/ChaosVisualDebugger/ChaosVisualDebuggerTrace.cpp
@@ -101,7 +101,7 @@ FDelegateHandle FChaosVisualDebuggerTrace::RecordingFullCaptureRequestedHandle =
 FRWLock FChaosVisualDebuggerTrace::DeltaRecordingStatesLock = FRWLock();
 TSet<int32> FChaosVisualDebuggerTrace::SolverIDsForDeltaRecording = TSet<int32>();
 TSet<int32> FChaosVisualDebuggerTrace::RequestedFullCaptureSolverIDs = TSet<int32>();
-FThreadSafeBool FChaosVisualDebuggerTrace::bIsTracing = false;
+std::atomic<bool> FChaosVisualDebuggerTrace::bIsTracing = false;
 
 void FChaosVisualDebuggerTrace::TraceParticle(const Chaos::FGeometryParticleHandle* ParticleHandle)
 {
diff --git a/Engine/Source/Runtime/Experimental/Chaos/Public/ChaosVisualDebugger/ChaosVisualDebuggerTrace.h b/Engine/Source/Runtime/Experimental/Chaos/Public/ChaosVisualDebugger/ChaosVisualDebuggerTrace.h
--- a/Engine/Source/Runtime/Experimental/Chaos/Public/ChaosVisualDebugger/ChaosVisualDebuggerTrace.h
+++ b/Engine/Source/Runtime/Experimental/Chaos/Public/ChaosVisualDebugger/ChaosVisualDebuggerTrace.h
@@ -474,7 +474,7 @@ private:
 	static TSet<int32> SolverIDsForDeltaRecording;
 	static TSet<int32> RequestedFullCaptureSolverIDs;
 
-	static FThreadSafeBool bIsTracing;
+	static std::atomic<bool> bIsTracing;
 
 	static FRWLock DeltaRecordingStatesLock;
 
