Geometry Collection : fix physics thread collection active attribute not being properly initialized ( and taking one update cycle to do so )
This was causing geometry collections with damage threshold set to 0 to not reflect the broken state of the collection because the root element would be not detected as chaged ( Active attribute was not yet set to true and because the root element would break before the first update cycle WasActive and Active had the same values causing the state of this element to be skipped
This wasn't a problem before because we used to force update everything. Recent optimizations are the reason why this bug was showing

diff --git a/Engine/Source/Runtime/Experimental/Chaos/Private/PhysicsProxy/GeometryCollectionPhysicsProxy.cpp b/Engine/Source/Runtime/Experimental/Chaos/Private/PhysicsProxy/GeometryCollectionPhysicsProxy.cpp
--- a/Engine/Source/Runtime/Experimental/Chaos/Private/PhysicsProxy/GeometryCollectionPhysicsProxy.cpp
+++ b/Engine/Source/Runtime/Experimental/Chaos/Private/PhysicsProxy/GeometryCollectionPhysicsProxy.cpp
@@ -755,6 +755,7 @@ void FGeometryCollectionPhysicsProxy::Initialize(Chaos::FPBDRigidsEvolutionBase
 		{
 			bool bIsRoot = (GameThreadCollection.Parent[TransformIndex] == INDEX_NONE);
 			GameThreadCollection.Active[TransformIndex] = bIsRoot;
+			PhysicsThreadCollection.Active[TransformIndex] = bIsRoot;
 
 			if (FParticle* P = GTParticles[TransformIndex].Get())
 			{
