Chaos - Add options to VisitCollisions methods to restrict visitor to active awake collisions only, with options to also visit sleeping or inactive.

diff --git a/Engine/Source/Runtime/Experimental/Chaos/Private/Chaos/Collision/PBDCollisionConstraint.cpp b/Engine/Source/Runtime/Experimental/Chaos/Private/Chaos/Collision/PBDCollisionConstraint.cpp
--- a/Engine/Source/Runtime/Experimental/Chaos/Private/Chaos/Collision/PBDCollisionConstraint.cpp
+++ b/Engine/Source/Runtime/Experimental/Chaos/Private/Chaos/Collision/PBDCollisionConstraint.cpp
@@ -423,7 +423,8 @@ namespace Chaos
 
 	void FPBDCollisionConstraint::Activate()
 	{
-		SetDisabled(false);
+		Flags.bIsActivated = true;
+		Flags.bDisabled = false;
 
 		AssignSavedManifoldPoints();
 
diff --git a/Engine/Source/Runtime/Experimental/Chaos/Private/Chaos/Collision/PBDCollisionContainerSolver.cpp b/Engine/Source/Runtime/Experimental/Chaos/Private/Chaos/Collision/PBDCollisionContainerSolver.cpp
--- a/Engine/Source/Runtime/Experimental/Chaos/Private/Chaos/Collision/PBDCollisionContainerSolver.cpp
+++ b/Engine/Source/Runtime/Experimental/Chaos/Private/Chaos/Collision/PBDCollisionContainerSolver.cpp
@@ -273,8 +273,6 @@ namespace Chaos
 				StaticFrictionRatio,
 				Dt);
 		}
-
-		Constraint->EndTick();
 	}
 
 	//////////////////////////////////////////////////////////////////////////////////////////////////
diff --git a/Engine/Source/Runtime/Experimental/Chaos/Private/Chaos/Collision/PBDCollisionContainerSolverJacobi.cpp b/Engine/Source/Runtime/Experimental/Chaos/Private/Chaos/Collision/PBDCollisionContainerSolverJacobi.cpp
--- a/Engine/Source/Runtime/Experimental/Chaos/Private/Chaos/Collision/PBDCollisionContainerSolverJacobi.cpp
+++ b/Engine/Source/Runtime/Experimental/Chaos/Private/Chaos/Collision/PBDCollisionContainerSolverJacobi.cpp
@@ -248,8 +248,6 @@ namespace Chaos
 					StaticFrictionRatio,
 					Dt);
 			}
-
-			Constraint->EndTick();
 		}
 
 		//////////////////////////////////////////////////////////////////////////////////////////////////
diff --git a/Engine/Source/Runtime/Experimental/Chaos/Private/Chaos/Collision/PBDCollisionContainerSolverSimd.cpp b/Engine/Source/Runtime/Experimental/Chaos/Private/Chaos/Collision/PBDCollisionContainerSolverSimd.cpp
--- a/Engine/Source/Runtime/Experimental/Chaos/Private/Chaos/Collision/PBDCollisionContainerSolverSimd.cpp
+++ b/Engine/Source/Runtime/Experimental/Chaos/Private/Chaos/Collision/PBDCollisionContainerSolverSimd.cpp
@@ -292,8 +292,6 @@ namespace Chaos
 						StaticFrictionRatio,
 						Dt);
 				}
-
-				Constraint->EndTick();
 			}
 		}
 
diff --git a/Engine/Source/Runtime/Experimental/Chaos/Private/Chaos/Collision/ParticlePairMidPhase.cpp b/Engine/Source/Runtime/Experimental/Chaos/Private/Chaos/Collision/ParticlePairMidPhase.cpp
--- a/Engine/Source/Runtime/Experimental/Chaos/Private/Chaos/Collision/ParticlePairMidPhase.cpp
+++ b/Engine/Source/Runtime/Experimental/Chaos/Private/Chaos/Collision/ParticlePairMidPhase.cpp
@@ -944,7 +944,7 @@ namespace Chaos
 				return ECollisionVisitorResult::Stop;
 			}
 			return ECollisionVisitorResult::Continue;
-		}, false);
+		}, ECollisionVisitorFlags::VisitAll);
 		return bInGraph;
 	}
 
diff --git a/Engine/Source/Runtime/Experimental/Chaos/Private/Chaos/ContactModification.cpp b/Engine/Source/Runtime/Experimental/Chaos/Private/Chaos/ContactModification.cpp
--- a/Engine/Source/Runtime/Experimental/Chaos/Private/Chaos/ContactModification.cpp
+++ b/Engine/Source/Runtime/Experimental/Chaos/Private/Chaos/ContactModification.cpp
@@ -539,7 +539,7 @@ namespace Chaos
 				Constraints.Add(&Constraint);
 			}
 			return ECollisionVisitorResult::Continue;
-		});
+		}, ECollisionVisitorFlags::VisitActiveAwake | ECollisionVisitorFlags::VisitDisabled);
 		ensureMsgf(Constraints.Num() <= NumCollisions, TEXT("Number of constraints visited exceeded the number reported by the FParticleCollisions object, and therefore exceeded the reserved constraint cache."));
 	}
 
diff --git a/Engine/Source/Runtime/Experimental/Chaos/Private/Chaos/PBDCollisionConstraints.cpp b/Engine/Source/Runtime/Experimental/Chaos/Private/Chaos/PBDCollisionConstraints.cpp
--- a/Engine/Source/Runtime/Experimental/Chaos/Private/Chaos/PBDCollisionConstraints.cpp
+++ b/Engine/Source/Runtime/Experimental/Chaos/Private/Chaos/PBDCollisionConstraints.cpp
@@ -527,6 +527,17 @@ namespace Chaos
 		PruneEdgeCollisions();
 	}
 
+	void FPBDCollisionConstraints::EndTick()
+	{
+		for (FPBDCollisionConstraint* Constraint : GetConstraints())
+		{
+			if (Constraint != nullptr)
+			{
+				Constraint->EndTick();
+			}
+		}
+	}
+
 	void FPBDCollisionConstraints::DetectProbeCollisions(FReal Dt)
 	{
 		SCOPE_CYCLE_COUNTER(STAT_Collisions_DetectProbeCollisions);
diff --git a/Engine/Source/Runtime/Experimental/Chaos/Private/Chaos/PBDRigidsEvolutionGBF.cpp b/Engine/Source/Runtime/Experimental/Chaos/Private/Chaos/PBDRigidsEvolutionGBF.cpp
--- a/Engine/Source/Runtime/Experimental/Chaos/Private/Chaos/PBDRigidsEvolutionGBF.cpp
+++ b/Engine/Source/Runtime/Experimental/Chaos/Private/Chaos/PBDRigidsEvolutionGBF.cpp
@@ -149,6 +149,8 @@ namespace Chaos
 		DECLARE_CYCLE_STAT(TEXT("FPBDRigidsEvolutionGBF::SaveParticlePostSolve"), STAT_Evolution_SavePostSolve, STATGROUP_Chaos);
 		DECLARE_CYCLE_STAT(TEXT("FPBDRigidsEvolutionGBF::DeactivateSleep"), STAT_Evolution_DeactivateSleep, STATGROUP_Chaos);
 		DECLARE_CYCLE_STAT(TEXT("FPBDRigidsEvolutionGBF::InertiaConditioning"), STAT_Evolution_InertiaConditioning, STATGROUP_Chaos);
+		DECLARE_CYCLE_STAT(TEXT("FPBDRigidsEvolutionGBF::CollisionEndTick"), STAT_Evolution_CollisionEndTick, STATGROUP_Chaos);
+
 
 		int32 SerializeEvolution = 0;
 		FAutoConsoleVariableRef CVarSerializeEvolution(TEXT("p.SerializeEvolution"), SerializeEvolution, TEXT(""));
@@ -622,8 +624,18 @@ void FPBDRigidsEvolutionGBF::AdvanceOneTimeStepImpl(const FReal Dt, const FSubSt
 	if (CVars::DoFinalProbeNarrowPhase)
 	{
 		// Run contact updates on probe constraints
+		// NOTE: This happens after particles have been moved to their new locations
+		// so that we get contacts that are correct for end-of-frame positions
 		GetCollisionConstraints().DetectProbeCollisions(Dt);
 	}
+	
+	{
+		SCOPE_CYCLE_COUNTER(STAT_Evolution_CollisionEndTick);
+		CSV_SCOPED_TIMING_STAT(PhysicsVerbose, StepSolver_CollisionEndTick);
+
+		// Reset collision state ready for the next tick
+		CollisionConstraints.EndTick();
+	}
 
 	{
 		CVD_SCOPE_TRACE_SOLVER_STEP(TEXT("Evolution End"));
@@ -1084,7 +1096,8 @@ void FPBDRigidsEvolutionGBF::ParticleMaterialChanged(FGeometryParticleHandle* Pa
 		// Reset the material - the material properties will get collected later (collision activation)
 		Collision.ClearMaterialProperties();
 		return ECollisionVisitorResult::Continue;
-	});
+
+	}, ECollisionVisitorFlags::VisitAll);
 
 	// The graph caches some sleep thresholds etc 
 	GetIslandManager().UpdateParticleMaterial(Particle);
diff --git a/Engine/Source/Runtime/Experimental/Chaos/Public/Chaos/Collision/CollisionConstraintAllocator.h b/Engine/Source/Runtime/Experimental/Chaos/Public/Chaos/Collision/CollisionConstraintAllocator.h
--- a/Engine/Source/Runtime/Experimental/Chaos/Public/Chaos/Collision/CollisionConstraintAllocator.h
+++ b/Engine/Source/Runtime/Experimental/Chaos/Public/Chaos/Collision/CollisionConstraintAllocator.h
@@ -587,16 +587,51 @@ namespace Chaos
 	///////////////////////////////////////////////////////////////////////////////////////////////
 	///////////////////////////////////////////////////////////////////////////////////////////////
 
+	// Check the visitor flags and return true if we want to visit the constraint
+	inline bool CollisionVisitorShouldVisit(const FPBDCollisionConstraint* Constraint, const ECollisionVisitorFlags VisitFlags)
+	{
+		if (Constraint == nullptr)
+		{
+			return false;
+		}
+
+		if (!(VisitFlags & ECollisionVisitorFlags::VisitInactive))
+		{
+			if (!Constraint->IsActivated())
+			{
+				return false;
+			}
+		}
+
+		if (!(VisitFlags & ECollisionVisitorFlags::VisitDisabled))
+		{
+			if (!Constraint->IsEnabled())
+			{
+				return false;
+			}
+		}
+
+		if (!(VisitFlags & ECollisionVisitorFlags::VisitSleeping))
+		{
+			if (Constraint->IsSleeping())
+			{
+				return false;
+			}
+		}
+
+		return true;
+	}
 
 	template<typename TLambda>
-	inline ECollisionVisitorResult FParticlePairMidPhase::VisitCollisions(const TLambda& Visitor, const bool bOnlyActive)
+	inline ECollisionVisitorResult FParticlePairMidPhase::VisitCollisions(const TLambda& Visitor, const ECollisionVisitorFlags VisitFlags)
 	{
 		if (GetMidPhaseType() == EParticlePairMidPhaseType::ShapePair)
 		{
 			FShapePairParticlePairMidPhase* This = static_cast<FShapePairParticlePairMidPhase*>(this);
+
 			for (FSingleShapePairCollisionDetector& ShapePair : This->ShapePairDetectors)
 			{
-				if ((ShapePair.GetConstraint() != nullptr) && (!bOnlyActive || !ShapePair.GetConstraint()->GetDisabled()))
+				if (CollisionVisitorShouldVisit(ShapePair.GetConstraint(), VisitFlags))
 				{
 					if (Visitor(*ShapePair.GetConstraint()) == ECollisionVisitorResult::Stop)
 					{
@@ -612,7 +647,8 @@ namespace Chaos
 			for (auto& KVP : This->Constraints)
 			{
 				FPBDCollisionConstraintPtr& Constraint = KVP.Value;
-				if (!bOnlyActive || Constraint->IsEnabled())
+
+				if (CollisionVisitorShouldVisit(Constraint.Get(), VisitFlags))
 				{
 					if (Visitor(*Constraint) == ECollisionVisitorResult::Stop)
 					{
@@ -627,14 +663,14 @@ namespace Chaos
 
 
 	template<typename TLambda>
-	inline ECollisionVisitorResult FParticlePairMidPhase::VisitConstCollisions(const TLambda& Visitor, const bool bOnlyActive) const
+	inline ECollisionVisitorResult FParticlePairMidPhase::VisitConstCollisions(const TLambda& Visitor, const ECollisionVisitorFlags VisitFlags) const
 	{
 		if (GetMidPhaseType() == EParticlePairMidPhaseType::ShapePair)
 		{
 			const FShapePairParticlePairMidPhase* This = static_cast<const FShapePairParticlePairMidPhase*>(this);
 			for (const FSingleShapePairCollisionDetector& ShapePair : This->ShapePairDetectors)
 			{
-				if ((ShapePair.GetConstraint() != nullptr) && (!bOnlyActive || !ShapePair.GetConstraint()->GetDisabled()))
+				if (CollisionVisitorShouldVisit(ShapePair.GetConstraint(), VisitFlags))
 				{
 					if (Visitor(*ShapePair.GetConstraint()) == ECollisionVisitorResult::Stop)
 					{
@@ -650,7 +686,8 @@ namespace Chaos
 			for (const auto& KVP : This->Constraints)
 			{
 				const FPBDCollisionConstraintPtr& Constraint = KVP.Value;
-				if (!bOnlyActive || Constraint->IsEnabled())
+				
+				if (CollisionVisitorShouldVisit(Constraint.Get(), VisitFlags))
 				{
 					if (Visitor(*Constraint) == ECollisionVisitorResult::Stop)
 					{
@@ -691,11 +728,11 @@ namespace Chaos
 	}
 
 	template<typename TLambda>
-	inline ECollisionVisitorResult FParticleCollisions::VisitCollisions(const TLambda& Visitor)
+	inline ECollisionVisitorResult FParticleCollisions::VisitCollisions(const TLambda& Visitor, const ECollisionVisitorFlags VisitFlags)
 	{
-		return VisitMidPhases([&Visitor](FParticlePairMidPhase& MidPhase)
+		return VisitMidPhases([&Visitor, VisitFlags](FParticlePairMidPhase& MidPhase)
 			{
-				if (MidPhase.VisitCollisions(Visitor) == ECollisionVisitorResult::Stop)
+				if (MidPhase.VisitCollisions(Visitor, VisitFlags) == ECollisionVisitorResult::Stop)
 				{
 					return ECollisionVisitorResult::Stop;
 				}
@@ -704,11 +741,11 @@ namespace Chaos
 	}
 
 	template<typename TLambda>
-	inline ECollisionVisitorResult FParticleCollisions::VisitConstCollisions(const TLambda& Visitor) const
+	inline ECollisionVisitorResult FParticleCollisions::VisitConstCollisions(const TLambda& Visitor, const ECollisionVisitorFlags VisitFlags) const
 	{
-		return VisitConstMidPhases([&Visitor](const FParticlePairMidPhase& MidPhase)
+		return VisitConstMidPhases([&Visitor, VisitFlags](const FParticlePairMidPhase& MidPhase)
 			{
-				if (MidPhase.VisitConstCollisions(Visitor) == ECollisionVisitorResult::Stop)
+				if (MidPhase.VisitConstCollisions(Visitor, VisitFlags) == ECollisionVisitorResult::Stop)
 				{
 					return ECollisionVisitorResult::Stop;
 				}
diff --git a/Engine/Source/Runtime/Experimental/Chaos/Public/Chaos/Collision/CollisionVisitor.h b/Engine/Source/Runtime/Experimental/Chaos/Public/Chaos/Collision/CollisionVisitor.h
--- a/Engine/Source/Runtime/Experimental/Chaos/Public/Chaos/Collision/CollisionVisitor.h
+++ b/Engine/Source/Runtime/Experimental/Chaos/Public/Chaos/Collision/CollisionVisitor.h
@@ -16,4 +16,19 @@ namespace Chaos
 		Continue
 	};
 
+	/*
+	* @brief Flags to control which constraints will be visited in the various constraint visitor functions
+	*/
+	enum class ECollisionVisitorFlags : uint8
+	{
+		VisitActiveAwake	= 0,		// Collisions that we detected this tick and have not been disabled or put to sleep
+		VisitSleeping		= 1 << 1,	// Collisions on body pairs that are sleeping
+		VisitDisabled		= 1 << 2,	// Collisions that have been disabled by the user in a collision callback
+		VisitInactive		= 1 << 3,	// Collisions that were not detected this frame but may still be kept around for potential reuse (only really useful for debugging etc)
+
+		VisitDefault		= VisitActiveAwake | VisitSleeping | VisitDisabled,
+		VisitAll			= VisitActiveAwake | VisitSleeping | VisitDisabled | VisitInactive
+	};
+	ENUM_CLASS_FLAGS(ECollisionVisitorFlags);
+
 }
\ No newline at end of file
diff --git a/Engine/Source/Runtime/Experimental/Chaos/Public/Chaos/Collision/PBDCollisionConstraint.h b/Engine/Source/Runtime/Experimental/Chaos/Public/Chaos/Collision/PBDCollisionConstraint.h
--- a/Engine/Source/Runtime/Experimental/Chaos/Public/Chaos/Collision/PBDCollisionConstraint.h
+++ b/Engine/Source/Runtime/Experimental/Chaos/Public/Chaos/Collision/PBDCollisionConstraint.h
@@ -349,7 +349,17 @@ namespace Chaos
 		void ResetPhi(FReal InPhi) { ClosestManifoldPointIndex = INDEX_NONE; }
 		FReal GetPhi() const { return (ClosestManifoldPointIndex != INDEX_NONE) ? ManifoldPoints[ClosestManifoldPointIndex].ContactPoint.Phi : TNumericLimits<FReal>::Max(); }
 
+		// Was this constraint activated this frame? It will be activated if the shapes are within CullDistance of each other.
+		// NOTE: All Active constraints are in the ActiveConstraints list on the CollisionConstraintAllocator. A constraint
+		// remains "activated" in this respect even if disabled by the user (via SetDisabled()) in a callback.
+		// This flag is only really useful if you are iterating over constraints that have been kept in memory
+		// for optimization reasons, but are not currently in use. Normally you would only need to consider GetDisabled()
+		bool IsActivated() const { return Flags.bIsActivated; }
+
+		// Allow the user to disable this constraint. @see GetActive().
 		void SetDisabled(bool bInDisabled) { Flags.bDisabled = bInDisabled; }
+
+		// Whether this constraint was disabled by the user (e.g., via a collision callback)
 		bool GetDisabled() const { return Flags.bDisabled; }
 
 		void SetIsProbe(bool bInProbe) { Flags.bIsProbe = bInProbe; }
@@ -636,6 +646,9 @@ namespace Chaos
 		 */
 		void EndTick()
 		{
+			// Reset the activated flag so that we don't need to reset it for constraints that are kept
+			// in memory but not re-activated next frame.
+			Flags.bIsActivated = false;
 		}
 
 		/**
@@ -820,6 +833,7 @@ namespace Chaos
 			FFlags() : Bits(0) {}
 			struct
 			{
+				uint16 bIsActivated : 1;				// Was this constraint activated this tick and therefore in the current tick's active list (note: it may subsequently be disabled)
 				uint16 bDisabled : 1;					// Is this contact disabled (by the user or because cull distance is exceeded)
 				uint16 bUseManifold : 1;				// Should we use contact manifolds or single points (faster but poor behaviour)
 				uint16 bUseIncrementalManifold : 1;		// Do we need to run incremental collision detection (only LavelSets now)
diff --git a/Engine/Source/Runtime/Experimental/Chaos/Public/Chaos/Collision/ParticleCollisions.h b/Engine/Source/Runtime/Experimental/Chaos/Public/Chaos/Collision/ParticleCollisions.h
--- a/Engine/Source/Runtime/Experimental/Chaos/Public/Chaos/Collision/ParticleCollisions.h
+++ b/Engine/Source/Runtime/Experimental/Chaos/Public/Chaos/Collision/ParticleCollisions.h
@@ -112,7 +112,7 @@ namespace Chaos
 		 * @note do not delete constraint from the lambda. You may disable them though.
 		*/
 		template<typename TLambda>
-		inline ECollisionVisitorResult VisitCollisions(const TLambda& Visitor);
+		inline ECollisionVisitorResult VisitCollisions(const TLambda& Visitor, const ECollisionVisitorFlags VisitFlags = ECollisionVisitorFlags::VisitDefault);
 
 		/**
 		 * @brief Visit all the collisions on this particle
@@ -121,7 +121,7 @@ namespace Chaos
 		 * @note do not delete constraint from the lambda. You may disable them though.
 		*/
 		template<typename TLambda>
-		inline ECollisionVisitorResult VisitConstCollisions(const TLambda& Visitor) const;
+		inline ECollisionVisitorResult VisitConstCollisions(const TLambda& Visitor, const ECollisionVisitorFlags VisitFlags = ECollisionVisitorFlags::VisitDefault) const;
 
 	private:
 		FContainerType MidPhases;
diff --git a/Engine/Source/Runtime/Experimental/Chaos/Public/Chaos/Collision/ParticlePairMidPhase.h b/Engine/Source/Runtime/Experimental/Chaos/Public/Chaos/Collision/ParticlePairMidPhase.h
--- a/Engine/Source/Runtime/Experimental/Chaos/Public/Chaos/Collision/ParticlePairMidPhase.h
+++ b/Engine/Source/Runtime/Experimental/Chaos/Public/Chaos/Collision/ParticlePairMidPhase.h
@@ -269,7 +269,7 @@ namespace Chaos
 		 * @param Visitor functor with signature ECollisionVisitorResult(FPBDCollisionConstaint& Constraint)
 		*/
 		template<typename TLambda>
-		ECollisionVisitorResult VisitCollisions(const TLambda& Visitor, const bool bOnlyActive = true);
+		ECollisionVisitorResult VisitCollisions(const TLambda& Visitor, const ECollisionVisitorFlags VisitFlags = ECollisionVisitorFlags::VisitDefault);
 
 
 		/**
@@ -277,7 +277,7 @@ namespace Chaos
 		 * @param Visitor functor with signature ECollisionVisitorResult(const FPBDCollisionConstaint& Constraint)
 		 */
 		template<typename TLambda>
-		ECollisionVisitorResult VisitConstCollisions(const TLambda& Visitor, const bool bOnlyActive = true) const;
+		ECollisionVisitorResult VisitConstCollisions(const TLambda& Visitor, const ECollisionVisitorFlags VisitFlags = ECollisionVisitorFlags::VisitDefault) const;
 
 		/**
 		 * @brief Cookie for use by FParticleCollisions
diff --git a/Engine/Source/Runtime/Experimental/Chaos/Public/Chaos/PBDCollisionConstraints.h b/Engine/Source/Runtime/Experimental/Chaos/Public/Chaos/PBDCollisionConstraints.h
--- a/Engine/Source/Runtime/Experimental/Chaos/Public/Chaos/PBDCollisionConstraints.h
+++ b/Engine/Source/Runtime/Experimental/Chaos/Public/Chaos/PBDCollisionConstraints.h
@@ -146,6 +146,10 @@ public:
 	*/
 	void DisableConstraints(const TSet<FGeometryParticleHandle*>& ParticleHandle) {}
 
+	/**
+	 * End of tick cleanup
+	 */
+	CHAOS_API void EndTick();
 
 	//
 	// FConstraintContainer Implementation
