Chaos - fix default collision visitor options following previous change

diff --git a/Engine/Source/Runtime/Experimental/Chaos/Public/Chaos/Collision/CollisionVisitor.h b/Engine/Source/Runtime/Experimental/Chaos/Public/Chaos/Collision/CollisionVisitor.h
--- a/Engine/Source/Runtime/Experimental/Chaos/Public/Chaos/Collision/CollisionVisitor.h
+++ b/Engine/Source/Runtime/Experimental/Chaos/Public/Chaos/Collision/CollisionVisitor.h
@@ -26,7 +26,7 @@ namespace Chaos
 		VisitDisabled		= 1 << 2,	// Collisions that have been disabled by the user in a collision callback
 		VisitInactive		= 1 << 3,	// Collisions that were not detected this frame but may still be kept around for potential reuse (only really useful for debugging etc)
 
-		VisitDefault		= VisitActiveAwake | VisitSleeping | VisitDisabled,
+		VisitDefault		= VisitActiveAwake | VisitSleeping,
 		VisitAll			= VisitActiveAwake | VisitSleeping | VisitDisabled | VisitInactive
 	};
 	ENUM_CLASS_FLAGS(ECollisionVisitorFlags);
