Adding Cluster Union Modular Vehicles

diff --git a/Engine/Source/Runtime/Engine/Classes/PhysicsEngine/ClusterUnionComponent.h b/Engine/Source/Runtime/Engine/Classes/PhysicsEngine/ClusterUnionComponent.h
--- a/Engine/Source/Runtime/Engine/Classes/PhysicsEngine/ClusterUnionComponent.h
+++ b/Engine/Source/Runtime/Engine/Classes/PhysicsEngine/ClusterUnionComponent.h
@@ -176,6 +176,7 @@ public:
 	ENGINE_API void ForceRebuildGTParticleGeometry();
 
 	friend class UClusterUnionReplicatedProxyComponent;
+	friend class UModularVehicleBaseComponent;
 protected:
 
 	// This should only be called on the client when replication happens.
@@ -184,6 +185,9 @@ protected:
 
 	ENGINE_API TArray<int32> GetAddedBoneIdsForComponent(UPrimitiveComponent* Component) const;
 
+	Chaos::FClusterUnionPhysicsProxy* GetPhysicsProxy() const { return PhysicsProxy; }
+	Chaos::FClusterUnionPhysicsProxy* GetPhysicsProxy() { return PhysicsProxy; }
+
 private:
 	// These are the statically clustered components. These should
 	// be specified in the editor and never change.
@@ -281,7 +285,7 @@ public:
 	ENGINE_API virtual bool OverlapComponentWithResult(const FVector& Pos, const FQuat& Rot, const FPhysicsGeometry& Geometry, ECollisionChannel TraceChannel, const struct FCollisionQueryParams& Params, const struct FCollisionResponseParams& ResponseParams, const struct FCollisionObjectQueryParams& ObjectParams, TArray<FOverlapResult>& OutOverlap) const override;
 	ENGINE_API virtual bool ComponentOverlapComponentWithResultImpl(const class UPrimitiveComponent* const PrimComp, const FVector& Pos, const FQuat& Rot, const FCollisionQueryParams& Params, TArray<FOverlapResult>& OutOverlap) const override;
 	virtual bool ShouldDispatchWakeEvents(FName BoneName) const override { return true; }
-	virtual bool DoCustomNavigableGeometryExport(FNavigableGeometryExport& GeomExport) const override;
+	ENGINE_API virtual bool DoCustomNavigableGeometryExport(FNavigableGeometryExport& GeomExport) const override;
 	//~ End UPrimitiveComponent Interface
 
 	//~ Begin USceneComponent Interface
diff --git a/Engine/Source/Runtime/Experimental/Chaos/Public/RewindData.h b/Engine/Source/Runtime/Experimental/Chaos/Public/RewindData.h
--- a/Engine/Source/Runtime/Experimental/Chaos/Public/RewindData.h
+++ b/Engine/Source/Runtime/Experimental/Chaos/Public/RewindData.h
@@ -213,6 +213,8 @@ public :
 		return NumDatas;
 	}
 
+	TArray<DatasType>& GetDatasArray() { return DatasArray; }
+
 protected : 
 
 	/** Check if the history is on the local/remote client*/
diff --git a/Engine/Source/Runtime/Experimental/ChaosVehicles/ChaosVehiclesCore/Private/SimModule/AerofoilModule.cpp b/Engine/Source/Runtime/Experimental/ChaosVehicles/ChaosVehiclesCore/Private/SimModule/AerofoilModule.cpp
--- a/Engine/Source/Runtime/Experimental/ChaosVehicles/ChaosVehiclesCore/Private/SimModule/AerofoilModule.cpp
+++ b/Engine/Source/Runtime/Experimental/ChaosVehicles/ChaosVehiclesCore/Private/SimModule/AerofoilModule.cpp
@@ -5,7 +5,7 @@
 #include "VehicleUtility.h"
 
 #if VEHICLE_DEBUGGING_ENABLED
-PRAGMA_DISABLE_OPTIMIZATION
+UE_DISABLE_OPTIMIZATION
 #endif
 
 namespace Chaos
@@ -141,5 +141,5 @@ namespace Chaos
 } // namespace Chaos
 
 #if VEHICLE_DEBUGGING_ENABLED
-PRAGMA_ENABLE_OPTIMIZATION
+UE_ENABLE_OPTIMIZATION
 #endif
diff --git a/Engine/Source/Runtime/Experimental/ChaosVehicles/ChaosVehiclesCore/Private/SimModule/AxleModule.cpp b/Engine/Source/Runtime/Experimental/ChaosVehicles/ChaosVehiclesCore/Private/SimModule/AxleModule.cpp
--- a/Engine/Source/Runtime/Experimental/ChaosVehicles/ChaosVehiclesCore/Private/SimModule/AxleModule.cpp
+++ b/Engine/Source/Runtime/Experimental/ChaosVehicles/ChaosVehiclesCore/Private/SimModule/AxleModule.cpp
@@ -5,7 +5,7 @@
 #include "VehicleUtility.h"
 
 #if VEHICLE_DEBUGGING_ENABLED
-PRAGMA_DISABLE_OPTIMIZATION
+UE_DISABLE_OPTIMIZATION
 #endif
 
 namespace Chaos
@@ -28,5 +28,5 @@ namespace Chaos
 } // namespace Chaos
 
 #if VEHICLE_DEBUGGING_ENABLED
-PRAGMA_ENABLE_OPTIMIZATION
+UE_ENABLE_OPTIMIZATION
 #endif
diff --git a/Engine/Source/Runtime/Experimental/ChaosVehicles/ChaosVehiclesCore/Private/SimModule/ChassisModule.cpp b/Engine/Source/Runtime/Experimental/ChaosVehicles/ChaosVehiclesCore/Private/SimModule/ChassisModule.cpp
--- a/Engine/Source/Runtime/Experimental/ChaosVehicles/ChaosVehiclesCore/Private/SimModule/ChassisModule.cpp
+++ b/Engine/Source/Runtime/Experimental/ChaosVehicles/ChaosVehiclesCore/Private/SimModule/ChassisModule.cpp
@@ -15,7 +15,6 @@ namespace Chaos
 
 	void FChassisSimModule::Simulate(float DeltaTime, const FAllInputs& Inputs, FSimModuleTree& VehicleModuleSystem)
 	{
-		AddLocalForce(FVector(20000000,0,0));
 
 	}
 
diff --git a/Engine/Source/Runtime/Experimental/ChaosVehicles/ChaosVehiclesCore/Private/SimModule/ClutchModule.cpp b/Engine/Source/Runtime/Experimental/ChaosVehicles/ChaosVehiclesCore/Private/SimModule/ClutchModule.cpp
--- a/Engine/Source/Runtime/Experimental/ChaosVehicles/ChaosVehiclesCore/Private/SimModule/ClutchModule.cpp
+++ b/Engine/Source/Runtime/Experimental/ChaosVehicles/ChaosVehiclesCore/Private/SimModule/ClutchModule.cpp
@@ -5,7 +5,7 @@
 #include "VehicleUtility.h"
 
 #if VEHICLE_DEBUGGING_ENABLED
-PRAGMA_DISABLE_OPTIMIZATION
+UE_DISABLE_OPTIMIZATION
 #endif
 
 namespace Chaos
@@ -60,5 +60,5 @@ namespace Chaos
 } // namespace Chaos
 
 #if VEHICLE_DEBUGGING_ENABLED
-PRAGMA_ENABLE_OPTIMIZATION
+UE_ENABLE_OPTIMIZATION
 #endif
diff --git a/Engine/Source/Runtime/Experimental/ChaosVehicles/ChaosVehiclesCore/Private/SimModule/DeferredForcesModular.cpp b/Engine/Source/Runtime/Experimental/ChaosVehicles/ChaosVehiclesCore/Private/SimModule/DeferredForcesModular.cpp
--- a/Engine/Source/Runtime/Experimental/ChaosVehicles/ChaosVehiclesCore/Private/SimModule/DeferredForcesModular.cpp
+++ b/Engine/Source/Runtime/Experimental/ChaosVehicles/ChaosVehiclesCore/Private/SimModule/DeferredForcesModular.cpp
@@ -10,6 +10,7 @@
 #include "Chaos/DebugDrawQueue.h"
 #endif
 
+UE_DISABLE_OPTIMIZATION
 
 FCoreModularVehicleDebugParams GCoreModularVehicleDebugParams;
 
@@ -17,10 +18,11 @@ FCoreModularVehicleDebugParams GCoreModularVehicleDebugParams;
 FAutoConsoleVariableRef CVarChaosModularVehiclesShowMass(TEXT("p.ModularVehicle.ShowMass"), GCoreModularVehicleDebugParams.ShowMass, TEXT("Show mass and inertia above objects."));
 FAutoConsoleVariableRef CVarChaosModularVehiclesShowForces(TEXT("p.ModularVehicle.ShowForces"), GCoreModularVehicleDebugParams.ShowForces, TEXT("Show forces on modular vehicles."));
 FAutoConsoleVariableRef CVarChaosModularVehiclesDrawForceScaling(TEXT("p.ModularVehicle.DrawForceScaling"), GCoreModularVehicleDebugParams.DrawForceScaling, TEXT("Set scaling of modular vehicle debug draw forces."));
+FAutoConsoleVariableRef CVarChaosModularVehiclesDisableForces(TEXT("p.ModularVehicle.DisableForces"), GCoreModularVehicleDebugParams.DisableForces, TEXT("Disable all forces."));
 #endif
 
-// TODO: passing all these parameters in is horrible, tidy this up now we know what needs to be done
-Chaos::FPBDRigidClusteredParticleHandle* FDeferredForcesModular::GetParticle(TArray<Chaos::FPBDRigidClusteredParticleHandle*>& Particles
+// #TODO: passing all these parameters in is horrible, tidy this up now we know what needs to be done
+Chaos::FPBDRigidParticleHandle* FDeferredForcesModular::GetParticle(TArray<Chaos::FPBDRigidClusteredParticleHandle*>& Particles
 		, TArray<Chaos::FPBDRigidClusteredParticleHandle*>& ClusterParticles
 		, int TransformIndex
 		, const TManagedArray<FTransform>& Transforms
@@ -29,7 +31,7 @@ Chaos::FPBDRigidClusteredParticleHandle* FDeferredForcesModular::GetParticle(TAr
 		, FTransform& TransformOut)
 {
 
-	Chaos::FPBDRigidClusteredParticleHandle* SimParticle = Particles[TransformIndex];
+	Chaos::FPBDRigidParticleHandle* SimParticle = Particles[TransformIndex];
 	Chaos::FPBDRigidClusteredParticleHandle* SimClusterParticle = ClusterParticles[TransformIndex];
 
 	if (SimParticle && !SimParticle->Disabled())
@@ -48,37 +50,124 @@ Chaos::FPBDRigidClusteredParticleHandle* FDeferredForcesModular::GetParticle(TAr
 	return nullptr;
 }
 
+Chaos::FPBDRigidParticleHandle* FDeferredForcesModular::GetParticle(const FTransform& OffsetTransform
+	, TArray<Chaos::FPBDRigidParticleHandle*>& Particles
+	, TArray<Chaos::FPBDRigidClusteredParticleHandle*>& ClusterParticles
+	, int TransformIndex
+	, FTransform& TransformOut)
+{
+	using namespace Chaos;
+
+	ensure(ClusterParticles.Num() == 1);
+	const int SingleChassisIndex = 0;
+	TransformOut = FTransform::Identity;
+
+	const FRigidTransform3 ClusterWorldTM(ClusterParticles[SingleChassisIndex]->X(), ClusterParticles[SingleChassisIndex]->R());
+	//const FRigidTransform3 ClusterWorldTM(ClusterParticles[SingleChassisIndex]->GetTransformXRCom());
+	FRigidTransform3 Frame = FRigidTransform3::Identity;
+
+	//UE_LOG(LogInit, Warning, TEXT("TransformIndex %d  vs  Particles.Num() %d  vs ClusterParticles.Num() %d"), TransformIndex, Particles.Num(), ClusterParticles.Num());
+
+	if (TransformIndex < Particles.Num())
+	{
+		FPBDRigidParticleHandle* Child = Particles[TransformIndex];
+		if (FPBDRigidClusteredParticleHandle* ClusterChild = Child->CastToClustered(); ClusterChild && ClusterChild->IsChildToParentLocked())
+		{
+			Frame = ClusterChild->ChildToParent();
+		}
+		else
+		{
+			//FTransform ChildT = Child->GetTransformXRCom();
+			//const FRigidTransform3 ChildWorldTM(ChildT.GetLocation(), ChildT.GetRotation());
+			const FRigidTransform3 ChildWorldTM(Child->X(), Child->R());
+			Frame = ChildWorldTM.GetRelativeTransform(ClusterWorldTM);
+
+			//UE_LOG(LogChaos, Warning, TEXT("ApplyForce Data TIdx %d %s"), TransformIndex, *Frame.GetTranslation().ToString());
+		}
+	}
+
+	// local transform
+	FVector Offset = OffsetTransform.InverseTransformVector(Frame.GetTranslation());
+	TransformOut.SetTranslation(Offset);
+	return ClusterParticles[SingleChassisIndex];
+}
+
 void FDeferredForcesModular::Apply(TArray<Chaos::FPBDRigidClusteredParticleHandle*>& Particles
 		, TArray<Chaos::FPBDRigidClusteredParticleHandle*>& ClusterParticles
 		, const TManagedArray<FTransform>& Transforms
-	, const TManagedArray<FTransform>& CollectionMassToLocal
-	, const TManagedArray<int32>& Parent)
+		, const TManagedArray<FTransform>& CollectionMassToLocal
+		, const TManagedArray<int32>& Parent)
 {
-	FTransform RelativeTransform;
-	for (const FApplyForceData& Data : ApplyForceDatas)
+	if (!GCoreModularVehicleDebugParams.DisableForces)
 	{
-		Chaos::FPBDRigidClusteredParticleHandle* RigidHandle = GetParticle(Particles, ClusterParticles, Data.TransformIndex, Transforms, CollectionMassToLocal, Parent, RelativeTransform);
-		if (RigidHandle)
+		FTransform RelativeTransform;
+		for (const FApplyForceData& Data : ApplyForceDatas)
 		{
-			AddForce(RigidHandle, Data, RelativeTransform);
+			Chaos::FPBDRigidParticleHandle* RigidHandle = GetParticle(Particles, ClusterParticles, Data.TransformIndex, Transforms, CollectionMassToLocal, Parent, RelativeTransform);
+			if (RigidHandle)
+			{
+				AddForce(RigidHandle, Data, RelativeTransform);
+			}
 		}
-	}
 
-	for (const FApplyForceAtPositionData& Data : ApplyForceAtPositionDatas)
-	{
-		Chaos::FPBDRigidClusteredParticleHandle* RigidHandle = GetParticle(Particles, ClusterParticles, Data.TransformIndex, Transforms, CollectionMassToLocal, Parent, RelativeTransform);
-		if (RigidHandle)
+		for (const FApplyForceAtPositionData& Data : ApplyForceAtPositionDatas)
+		{
+			Chaos::FPBDRigidParticleHandle* RigidHandle = GetParticle(Particles, ClusterParticles, Data.TransformIndex, Transforms, CollectionMassToLocal, Parent, RelativeTransform);
+			if (RigidHandle)
+			{
+				AddForceAtPosition(RigidHandle, Data, RelativeTransform);
+			}
+		}
+
+		for (const FAddTorqueInRadiansData& Data : ApplyTorqueDatas)
 		{
-			AddForceAtPosition(RigidHandle, Data, RelativeTransform);
+			Chaos::FPBDRigidParticleHandle* RigidHandle = GetParticle(Particles, ClusterParticles, Data.TransformIndex, Transforms, CollectionMassToLocal, Parent, RelativeTransform);
+			if (RigidHandle)
+			{
+				AddTorque(RigidHandle, Data, RelativeTransform);
+			}
 		}
 	}
 
-	for (const FAddTorqueInRadiansData& Data : ApplyTorqueDatas)
+	ApplyForceDatas.Empty();
+	ApplyForceAtPositionDatas.Empty();
+	ApplyTorqueDatas.Empty();
+}
+
+
+
+void FDeferredForcesModular::Apply(const FTransform& OffsetTransform
+	, TArray<Chaos::FPBDRigidParticleHandle*>& Particles
+	, TArray<Chaos::FPBDRigidClusteredParticleHandle*>& ClusterParticles)
+{
+	if (!GCoreModularVehicleDebugParams.DisableForces)
 	{
-		Chaos::FPBDRigidClusteredParticleHandle* RigidHandle = GetParticle(Particles, ClusterParticles, Data.TransformIndex, Transforms, CollectionMassToLocal, Parent, RelativeTransform);
-		if (RigidHandle)
+		FTransform RelativeTransform;
+		for (const FApplyForceData& Data : ApplyForceDatas)
 		{
-			AddTorque(RigidHandle, Data, RelativeTransform);
+			Chaos::FPBDRigidParticleHandle* RigidHandle = GetParticle(OffsetTransform, Particles, ClusterParticles, Data.TransformIndex, RelativeTransform);
+			if (RigidHandle)
+			{
+				AddForce(RigidHandle, Data, RelativeTransform);
+			}
+		}
+
+		for (const FApplyForceAtPositionData& Data : ApplyForceAtPositionDatas)
+		{
+			Chaos::FPBDRigidParticleHandle* RigidHandle = GetParticle(OffsetTransform, Particles, ClusterParticles, Data.TransformIndex, RelativeTransform);
+			if (RigidHandle)
+			{
+				AddForceAtPosition(RigidHandle, Data, RelativeTransform);
+			}
+		}
+
+		for (const FAddTorqueInRadiansData& Data : ApplyTorqueDatas)
+		{
+			Chaos::FPBDRigidParticleHandle* RigidHandle = GetParticle(OffsetTransform, Particles, ClusterParticles, Data.TransformIndex, RelativeTransform);
+			if (RigidHandle)
+			{
+				AddTorque(RigidHandle, Data, RelativeTransform);
+			}
 		}
 	}
 
@@ -88,39 +177,41 @@ void FDeferredForcesModular::Apply(TArray<Chaos::FPBDRigidClusteredParticleHandl
 }
 
 
-void AddForce_Implementation(Chaos::FPBDRigidClusteredParticleHandle* RigidHandle, const FVector& LocalForce, const FTransform& OffsetTransform, bool bLevelSlope, const FColor& DebugColor)
+
+void AddForce_Implementation(Chaos::FPBDRigidParticleHandle* RigidHandle, const FTransform& ParticleOffsetTransform, const FVector& LocalForce, const FTransform& OffsetTransform, bool bLevelSlope, const FColor& DebugColor)
 {
 	const Chaos::FVec3 WorldCOM = Chaos::FParticleUtilitiesGT::GetCoMWorldPosition(RigidHandle);
 
 	// local to world
 	const FTransform WorldTM(RigidHandle->R(), RigidHandle->X());
-	Chaos::FVec3 Position = WorldTM.TransformPosition(OffsetTransform.GetLocation());
-	Chaos::FVec3 Force = WorldTM.TransformVector(OffsetTransform.TransformVector(LocalForce));
+	Chaos::FVec3 Position = WorldTM.TransformPosition(ParticleOffsetTransform.TransformVector(OffsetTransform.GetLocation()));
+	Chaos::FVec3 Force = WorldTM.TransformVector(ParticleOffsetTransform.TransformVector(OffsetTransform.TransformVector(LocalForce)));
+
 
 	if (bLevelSlope && WorldTM.GetUnitAxis(EAxis::Z).Z > GCoreModularVehicleDebugParams.LevelSlopeThreshold)
 	{
-		Force.X = 0.0f;
-		Force.Y = 0.0f;
+		static float ValueToUse = 0.5f;
+		Force.X = ValueToUse;
+		Force.Y = ValueToUse;
 	}
 
 	Chaos::FVec3 Arm = Position - WorldCOM;
 
-	const Chaos::FVec3 WorldTorque = Chaos::FVec3::CrossProduct(Arm, Force);
+	Chaos::FVec3 WorldTorque = Chaos::FVec3::CrossProduct(Arm, Force);
 	auto* RP = RigidHandle->CastToRigidParticle();
 	if (RP)
 	{
-
 		RP->AddForce(Force);
 		RP->AddTorque(WorldTorque);
 
 		// #TODO: remove HACK - doing this because vehicles are falling asleep and not waking up when forces are applied
-		RP->SetSleepType(Chaos::ESleepType::NeverSleep);
+	//	RP->SetSleepType(Chaos::ESleepType::NeverSleep);
 
 #if CHAOS_DEBUG_DRAW
 		if (GCoreModularVehicleDebugParams.ShowMass)
 		{
 			FString OutputText = FString::Format(TEXT("{0} {1}"), { RP->M(), RP->I().ToString() });
-			Chaos::FDebugDrawQueue::GetInstance().DrawDebugString(RigidHandle->X() + FVector(0,0,200), OutputText, nullptr, FColor::White, -1, true, 1.0f);
+			Chaos::FDebugDrawQueue::GetInstance().DrawDebugString(RigidHandle->X() + FVector(0,0,200), OutputText, nullptr, FColor::White, -1, false, -1.0f);
 		}
 #endif
 	}
@@ -128,7 +219,7 @@ void AddForce_Implementation(Chaos::FPBDRigidClusteredParticleHandle* RigidHandl
 #if CHAOS_DEBUG_DRAW
 	if (GCoreModularVehicleDebugParams.ShowForces)
 	{
-		// TODO: perhaps render different force types/sources in different colors, i.e. aerofoil, suspension, friction, etc.
+		// #TODO: perhaps render different force types/sources in different colors, i.e. aerofoil, suspension, friction, etc.
 		Chaos::FDebugDrawQueue::GetInstance().DrawDebugLine(Position, Position + Force * GCoreModularVehicleDebugParams.DrawForceScaling, DebugColor, false, -1.f, 0, 5.f);
 		Chaos::FDebugDrawQueue::GetInstance().DrawDebugSphere(Position, 5, 8, FColor::White, false, -1.f, 0, 5.f);
 	}
@@ -136,11 +227,11 @@ void AddForce_Implementation(Chaos::FPBDRigidClusteredParticleHandle* RigidHandl
 
 }
 
-void FDeferredForcesModular::AddForce(Chaos::FPBDRigidClusteredParticleHandle* RigidHandle, const FApplyForceData& DataIn, const FTransform& OffsetTransform)
+void FDeferredForcesModular::AddForce(Chaos::FPBDRigidParticleHandle* RigidHandle, const FApplyForceData& DataIn, const FTransform& OffsetTransform)
 {
 	if (ensure(RigidHandle))
 	{
-		AddForce_Implementation(RigidHandle, DataIn.Force, OffsetTransform, (DataIn.Flags & EForceFlags::LevelSlope) == EForceFlags::LevelSlope, DataIn.DebugColor);
+		AddForce_Implementation(RigidHandle, ParticleOffsetTransform, DataIn.Force, OffsetTransform, (DataIn.Flags & EForceFlags::LevelSlope) == EForceFlags::LevelSlope, DataIn.DebugColor);
 	}
 
 }
@@ -168,17 +259,17 @@ void FDeferredForcesModular::AddForce(Chaos::FPBDRigidClusteredParticleHandle* R
 //	}
 //}
 
-void FDeferredForcesModular::AddForceAtPosition(Chaos::FPBDRigidClusteredParticleHandle* RigidHandle, const FApplyForceAtPositionData& DataIn, const FTransform& OffsetTransform)
+void FDeferredForcesModular::AddForceAtPosition(Chaos::FPBDRigidParticleHandle* RigidHandle, const FApplyForceAtPositionData& DataIn, const FTransform& OffsetTransform)
 {
 	if (ensure(RigidHandle))
 	{
 		FTransform ShiftedTransform = OffsetTransform;
 		ShiftedTransform.SetLocation(ShiftedTransform.GetLocation() + DataIn.Position);
-		AddForce_Implementation(RigidHandle, DataIn.Force, ShiftedTransform, (DataIn.Flags & EForceFlags::LevelSlope) == EForceFlags::LevelSlope, DataIn.DebugColor);
+		AddForce_Implementation(RigidHandle, ParticleOffsetTransform, DataIn.Force, ShiftedTransform, (DataIn.Flags & EForceFlags::LevelSlope) == EForceFlags::LevelSlope, DataIn.DebugColor);
 	}
 }
 
-void FDeferredForcesModular::AddTorque(Chaos::FPBDRigidClusteredParticleHandle* RigidHandle, const FAddTorqueInRadiansData& DataIn, const FTransform& OffsetTransform)
+void FDeferredForcesModular::AddTorque(Chaos::FPBDRigidParticleHandle* RigidHandle, const FAddTorqueInRadiansData& DataIn, const FTransform& OffsetTransform)
 {	
 	check(false); //#TODO - incorporate OffsetTransform as this is currently wrong
 	//if (ensure(RigidHandle))
@@ -198,3 +289,4 @@ void FDeferredForcesModular::AddTorque(Chaos::FPBDRigidClusteredParticleHandle*
 	//}
 }
 
+UE_ENABLE_OPTIMIZATION
diff --git a/Engine/Source/Runtime/Experimental/ChaosVehicles/ChaosVehiclesCore/Private/SimModule/EngineModule.cpp b/Engine/Source/Runtime/Experimental/ChaosVehicles/ChaosVehiclesCore/Private/SimModule/EngineModule.cpp
--- a/Engine/Source/Runtime/Experimental/ChaosVehicles/ChaosVehiclesCore/Private/SimModule/EngineModule.cpp
+++ b/Engine/Source/Runtime/Experimental/ChaosVehicles/ChaosVehiclesCore/Private/SimModule/EngineModule.cpp
@@ -5,7 +5,7 @@
 #include "VehicleUtility.h"
 
 #if VEHICLE_DEBUGGING_ENABLED
-PRAGMA_DISABLE_OPTIMIZATION
+UE_DISABLE_OPTIMIZATION
 #endif
 
 namespace Chaos
@@ -74,5 +74,5 @@ namespace Chaos
 } // namespace Chaos
 
 #if VEHICLE_DEBUGGING_ENABLED
-PRAGMA_ENABLE_OPTIMIZATION
+UE_ENABLE_OPTIMIZATION
 #endif
diff --git a/Engine/Source/Runtime/Experimental/ChaosVehicles/ChaosVehiclesCore/Private/SimModule/MotorModule.cpp b/Engine/Source/Runtime/Experimental/ChaosVehicles/ChaosVehiclesCore/Private/SimModule/MotorModule.cpp
--- a/Engine/Source/Runtime/Experimental/ChaosVehicles/ChaosVehiclesCore/Private/SimModule/MotorModule.cpp
+++ b/Engine/Source/Runtime/Experimental/ChaosVehicles/ChaosVehiclesCore/Private/SimModule/MotorModule.cpp
@@ -5,7 +5,7 @@
 #include "VehicleUtility.h"
 
 #if VEHICLE_DEBUGGING_ENABLED
-PRAGMA_DISABLE_OPTIMIZATION
+UE_DISABLE_OPTIMIZATION
 #endif
 
 
@@ -49,5 +49,5 @@ namespace Chaos
 } // namespace Chaos
 
 #if VEHICLE_DEBUGGING_ENABLED
-PRAGMA_ENABLE_OPTIMIZATION
+UE_ENABLE_OPTIMIZATION
 #endif
diff --git a/Engine/Source/Runtime/Experimental/ChaosVehicles/ChaosVehiclesCore/Private/SimModule/SimModuleTree.cpp b/Engine/Source/Runtime/Experimental/ChaosVehicles/ChaosVehiclesCore/Private/SimModule/SimModuleTree.cpp
--- a/Engine/Source/Runtime/Experimental/ChaosVehicles/ChaosVehiclesCore/Private/SimModule/SimModuleTree.cpp
+++ b/Engine/Source/Runtime/Experimental/ChaosVehicles/ChaosVehiclesCore/Private/SimModule/SimModuleTree.cpp
@@ -3,10 +3,12 @@
 #include "SimModule/SimModuleTree.h"
 #include "Chaos/ParticleHandleFwd.h"
 #include "PhysicsProxy/GeometryCollectionPhysicsProxy.h"
+#include "PhysicsProxy/ClusterUnionPhysicsProxy.h"
+#include "Chaos/DebugDrawQueue.h"
 #include "VehicleUtility.h"
 
 #if VEHICLE_DEBUGGING_ENABLED
-PRAGMA_DISABLE_OPTIMIZATION
+UE_DISABLE_OPTIMIZATION
 #endif
 
 namespace Chaos
@@ -14,11 +16,13 @@ namespace Chaos
 
 int FSimModuleTree::AddRoot(ISimulationModuleBase* SimModule)
 {
+	check(!bIsSimulating);
 	return AddNodeBelow(-1, SimModule);
 }
 
 void FSimModuleTree::Reparent(int AtIndex, int ParentIndex)
 {
+	check(!bIsSimulating);
 	check(AtIndex < SimulationModuleTree.Num());
 	check(ParentIndex < SimulationModuleTree.Num());
 
@@ -32,12 +36,17 @@ void FSimModuleTree::Reparent(int AtIndex, int ParentIndex)
 		SimulationModuleTree[AtIndex].Parent = ParentIndex;
 		SimulationModuleTree[ParentIndex].Children.Add(AtIndex);
 
-		SimulationModuleTree[OrginalParent].Children.Remove(AtIndex);
+		// if had a parent and wasn't a root
+		if (OrginalParent != -1)
+		{
+			SimulationModuleTree[OrginalParent].Children.Remove(AtIndex);
+		}
 	}
 }
 
 int FSimModuleTree::AddNodeBelow(int AtIndex, ISimulationModuleBase* SimModule)
 {
+	check(!bIsSimulating);
 	int NewIndex = GetNextIndex();
 	FSimModuleNode& Node = SimulationModuleTree[NewIndex];
 	SimModule->SetTreeIndex(NewIndex);
@@ -49,20 +58,118 @@ int FSimModuleTree::AddNodeBelow(int AtIndex, ISimulationModuleBase* SimModule)
 	}
 	else
 	{
-		Node.Parent = FSimModuleNode::INVALID_INDEX;
+		Node.Parent = FSimModuleNode::INVALID_IDX;
 	}
 
 	return NewIndex;
 }
 
+void FSimModuleTree::AppendTreeUpdates(const FSimTreeUpdates& TreeUpdates)
+{
+	ensure(!IsSimulating());
+
+	int TreeIndex = -1;
+	TMap<int, int> SimTreeMapping;
+
+	//if (GetNumNodes() == 0) //Always add a null node, so there is a parent when the root chassis is removed
+	//{
+	//	// add a single chassis root component
+	//	Chaos::FChassisSettings Settings;
+	//	Chaos::ISimulationModuleBase* Chassis = new Chaos::FChassisSimModule(Settings);
+	//	ParentIndex = AddRoot(Chassis);
+	//	Chassis->SetTransformIndex(-1);
+	//}
+
+	int LocalIndex = 0;
+	for (const FPendingModuleAdds& TreeUpdate : TreeUpdates.GetNewModules())
+	{
+		int AddIndex = -1;
+		if (int* AddIndexPtr = SimTreeMapping.Find(TreeUpdate.ParentIndex))
+		{
+			AddIndex = *AddIndexPtr;
+		}
+
+		ensure(!IsSimulating());
+
+		TreeIndex = AddNodeBelow(AddIndex, TreeUpdate.NewSimModule);
+		SimTreeMapping.Add(LocalIndex, TreeIndex);
+		LocalIndex++;
+	}
+
+	TArray<int> ComponentIndices;
+	for (const FPendingModuleDeletions& TreeUpdate : TreeUpdates.GetDeletedModules())
+	{
+		for (int Index = 0; Index < SimulationModuleTree.Num(); Index++)
+		{
+			if (Chaos::ISimulationModuleBase* SimModule = GetNode(Index).SimModule)
+			{ 
+				if (SimModule->GetGuid() == TreeUpdate.Guid)
+				{
+					ComponentIndices.AddUnique(SimModule->GetTransformIndex());
+					DeleteNode(Index);
+					break;
+				}
+			}
+		}				
+	}
+
+	//// when the particle is removed the references will be out by one unless we fix them up
+	//// this may only apply to ClusterUnion and not GeometryCollection
+	//for (int ComponentIndex : ComponentIndices)
+	//{
+	//	for (int I = 0; I < GetNumNodes(); I++)
+	//	{
+	//		if (Chaos::ISimulationModuleBase* SimModule = GetNode(I).SimModule)
+	//		{
+	//			if (SimModule->GetTransformIndex() > ComponentIndex)
+	//			{
+	//				SimModule->SetTransformIndex(SimModule->GetTransformIndex() - 1);
+	//			}
+
+	//		}
+	//	}
+	//}
+
+	for (int ComponentIndex : ComponentIndices)
+	{
+		// find the largest transform index
+		int LargestIndex = -1;
+		for (int I = 0; I < GetNumNodes(); I++)
+		{
+			if (Chaos::ISimulationModuleBase* SimModule = GetNode(I).SimModule)
+			{
+				if (SimModule->GetTransformIndex() > LargestIndex)
+				{
+					LargestIndex = SimModule->GetTransformIndex();
+				}
+			}
+		}
+
+		// swap with the one that has just been deleted 
+		// - there can be more than one SimModule referencing the same component index
+		for (int I = 0; I < GetNumNodes(); I++)
+		{
+			if (Chaos::ISimulationModuleBase* SimModule = GetNode(I).SimModule)
+			{
+				if (SimModule->GetTransformIndex() == LargestIndex)
+				{
+					SimModule->SetTransformIndex(ComponentIndex);
+				}
+			}
+		}
+	}
+
+}
+
+
 int FSimModuleTree::GetNextIndex()
 {
-	int NewIndex = FSimModuleNode::INVALID_INDEX;
+	int NewIndex = FSimModuleNode::INVALID_IDX;
 	if (FreeList.IsEmpty())
 	{
 		NewIndex = SimulationModuleTree.Num();
 		SimulationModuleTree.AddZeroed(1);
-		SimulationModuleTree[NewIndex].Parent = FSimModuleNode::INVALID_INDEX;
+		SimulationModuleTree[NewIndex].Parent = FSimModuleNode::INVALID_IDX;
 		SimulationModuleTree[NewIndex].SimModule = nullptr;
 	}
 	else
@@ -75,7 +182,8 @@ int FSimModuleTree::GetNextIndex()
 
 int FSimModuleTree::InsertNodeAbove(int AtIndex, ISimulationModuleBase* SimModule)
 {
-	int NewIndex = FSimModuleNode::INVALID_INDEX;
+	check(!bIsSimulating);
+	int NewIndex = FSimModuleNode::INVALID_IDX;
 
 	if (ensure(AtIndex < SimulationModuleTree.Num()))
 	{
@@ -101,6 +209,7 @@ int FSimModuleTree::InsertNodeAbove(int AtIndex, ISimulationModuleBase* SimModul
 
 void FSimModuleTree::DeleteNode(int AtIndex)
 {
+	check(!bIsSimulating);
 	// multiple children might become equal parents?
 	
 	int ParentIndex = SimulationModuleTree[AtIndex].Parent;
@@ -121,7 +230,7 @@ void FSimModuleTree::DeleteNode(int AtIndex)
 		SimulationModuleTree[ChildIndex].Parent = ParentIndex;
 	}	
 
-	SimulationModuleTree[AtIndex].Parent = FSimModuleNode::INVALID_INDEX;
+	SimulationModuleTree[AtIndex].Parent = FSimModuleNode::INVALID_IDX;
 	SimulationModuleTree[AtIndex].Children.Empty();
 	delete SimulationModuleTree[AtIndex].SimModule;
 	SimulationModuleTree[AtIndex].SimModule = nullptr;
@@ -130,6 +239,24 @@ void FSimModuleTree::DeleteNode(int AtIndex)
 
 }
 
+
+void FSimModuleTree::Simulate(float DeltaTime, FAllInputs& Inputs, FClusterUnionPhysicsProxy* PhysicsProxy)
+{
+	if (PhysicsProxy)
+	{
+		UpdateModuleVelocites(PhysicsProxy);
+	}
+
+	TArray<int> RootNodes;
+	GetRootNodes(RootNodes);
+
+	for (int RootIndex : RootNodes)
+	{
+		SimulateNode(DeltaTime, Inputs, RootIndex);
+	}
+}
+
+
 void FSimModuleTree::Simulate(float DeltaTime, FAllInputs& Inputs, FGeometryCollectionPhysicsProxy* PhysicsProxy)
 {
 	if (PhysicsProxy)
@@ -144,6 +271,7 @@ void FSimModuleTree::Simulate(float DeltaTime, FAllInputs& Inputs, FGeometryColl
 	{
 		SimulateNode(DeltaTime, Inputs, RootIndex);
 	}
+
 }
 
 void FSimModuleTree::SimulateNode(float DeltaTime, FAllInputs& Inputs, int NodeIndex)
@@ -175,7 +303,7 @@ void FSimModuleTree::DeleteNodesBelow(int AtIndex)
 		delete SimulationModuleTree[AtIndex].SimModule;
 		SimulationModuleTree[AtIndex].SimModule = nullptr;
 		SimulationModuleTree[AtIndex].Children.Empty();
-		SimulationModuleTree[AtIndex].Parent = FSimModuleNode ::INVALID_INDEX;
+		SimulationModuleTree[AtIndex].Parent = FSimModuleNode ::INVALID_IDX;
 
 		FreeList.Push(AtIndex);
 
@@ -190,7 +318,7 @@ void FSimModuleTree::GetRootNodes(TArray<int>& RootNodesOut)
 	// never assume the root bone is always index 0
 	for (int i = 0; i < SimulationModuleTree.Num(); i++)
 	{
-		if (SimulationModuleTree[i].SimModule != nullptr && SimulationModuleTree[i].Parent == FSimModuleNode::INVALID_INDEX)
+		if (SimulationModuleTree[i].SimModule != nullptr && SimulationModuleTree[i].Parent == FSimModuleNode::INVALID_IDX)
 		{
 			RootNodesOut.Add(i);
 		}
@@ -244,10 +372,103 @@ void FSimModuleTree::UpdateModuleVelocites(FGeometryCollectionPhysicsProxy* Phys
 
 }
 
+void FSimModuleTree::UpdateModuleVelocites(FClusterUnionPhysicsProxy* PhysicsProxy)
+{
+	check(PhysicsProxy);
+
+	// capture the velocities at the start of each sim iteration
+	for (int i = 0; i < SimulationModuleTree.Num(); i++)
+	{
+		if (ISimulationModuleBase* Module = SimulationModuleTree[i].SimModule)
+		{
+			if (const Chaos::FClusterUnionPhysicsProxy::FInternalParticle* ParentParticle = PhysicsProxy->GetParticle_Internal())
+			{
+				const FTransform BodyTransform(ParentParticle->R(), ParentParticle->X());
+
+				if (Module->IsBehaviourType(eSimModuleTypeFlags::Velocity))
+				{
+					const Chaos::FClusterUnionPhysicsProxy::FInternalParticle* Particle = nullptr;
+					if (Module->IsClustered())
+					{
+						Particle = ParentParticle;
+					}
+					else
+					{
+		//				Particle = Particles[Module->GetTransformIndex()];
+					}
+
+					if (Particle)
+					{
+						const FTransform& OffsetTransform = DeferredForces.GetOffsetTransform();
+						FVector LocalPos = Module->GetParentRelativeTransform().GetLocation();
+						FVector WorldLocation = BodyTransform.TransformPosition(OffsetTransform.TransformVector(LocalPos));
+						const Chaos::FVec3 Arm = WorldLocation - Particle->X();
+
+						//Chaos::FDebugDrawQueue::GetInstance().DrawDebugLine(Particle->X(), Particle->X() + Arm, FColor::Yellow, false, -1.f, 0, 2.f);
+
+						FVector WorldVelocity = Particle->V() - Chaos::FVec3::CrossProduct(Arm, Particle->W());
+						FVector LocalVelocity = OffsetTransform.InverseTransformVector(BodyTransform.InverseTransformVector(WorldVelocity));
+						Module->SetLocalVelocity(LocalVelocity);
+					}
+
+				}
+			}
+		}
+	}
+
+}
+
+void FSimModuleTree::GenerateReplicationStructure(Chaos::FModuleNetDataArray& NetData)
+{
+	const TArray<FSimModuleNode>& Tree = SimulationModuleTree;
+	NetData.Reserve(Tree.Num());
+	for (int Index = 0; Index < Tree.Num(); Index++)
+	{
+		TSharedPtr<FModuleNetData>&& Data = Tree[Index].SimModule->GenerateNetData(Index);
+		// not all modules will have net replication data - nullptr is a valid response
+		if (Data)
+		{
+			NetData.Emplace(Data);
+		}
+	}
+}
+
+void FSimModuleTree::SetNetState(Chaos::FModuleNetDataArray& ModuleDatas)
+{
+	if (ModuleDatas.IsEmpty())
+	{
+		GenerateReplicationStructure(ModuleDatas);
+	}
+
+	for (TSharedPtr<FModuleNetData>& DataElement : ModuleDatas)
+	{
+		if (!SimulationModuleTree.IsEmpty() && SimulationModuleTree[DataElement->SimArrayIndex].SimModule)
+		{
+			DataElement->FillNetState(SimulationModuleTree[DataElement->SimArrayIndex].SimModule);
+		}
+	}
+}
+
+void FSimModuleTree::SetSimState(const Chaos::FModuleNetDataArray& ModuleDatas)
+{
+	for (const TSharedPtr<FModuleNetData>& DataElement : ModuleDatas)
+	{
+		if (!SimulationModuleTree.IsEmpty() && SimulationModuleTree[DataElement->SimArrayIndex].SimModule)
+		{
+			DataElement->FillSimState(SimulationModuleTree[DataElement->SimArrayIndex].SimModule);
+		}
+	}
+}
+
+
+void FSimModuleTree::InterpolateState(const float LerpFactor, Chaos::FModuleNetDataArray& LerpDatas, const Chaos::FModuleNetDataArray& MinDatas, const Chaos::FModuleNetDataArray& MaxDatas)
+{
+
+}
 
 } // namespace Chaos
 
 
 #if VEHICLE_DEBUGGING_ENABLED
-PRAGMA_ENABLE_OPTIMIZATION
+UE_ENABLE_OPTIMIZATION
 #endif
diff --git a/Engine/Source/Runtime/Experimental/ChaosVehicles/ChaosVehiclesCore/Private/SimModule/SimulationModuleBase.cpp b/Engine/Source/Runtime/Experimental/ChaosVehicles/ChaosVehiclesCore/Private/SimModule/SimulationModuleBase.cpp
--- a/Engine/Source/Runtime/Experimental/ChaosVehicles/ChaosVehiclesCore/Private/SimModule/SimulationModuleBase.cpp
+++ b/Engine/Source/Runtime/Experimental/ChaosVehicles/ChaosVehiclesCore/Private/SimModule/SimulationModuleBase.cpp
@@ -23,6 +23,10 @@ void ISimulationModuleBase::AddLocalForce(const FVector& Force, bool bAllowSubst
 	AppliedForce = Force;
 	if (SimModuleTree)
 	{
+		//FString DebugString;
+		//GetDebugString(DebugString);
+		//UE_LOG(LogInit, Warning, TEXT("AddLocalForce To %s"), *DebugString);
+
 		SimModuleTree->AccessDeferredForces().Add(FDeferredForcesModular::FApplyForceData(TransformIndex, Force, bAllowSubstepping, bIsLocalForce, bLevelSlope, DebugColorIn));
 	}
 }
@@ -47,11 +51,23 @@ ISimulationModuleBase* ISimulationModuleBase::GetFirstChild()
 
 bool ISimulationModuleBase::GetDebugString(FString& StringOut) const 
 {
-	StringOut += FString::Format(TEXT("{0}: Enabled {1}, InCluster {2}, TFormIdx {3}, ")
-		, { GetDebugName(), IsEnabled(), IsClustered(), GetTransformIndex() });
+	StringOut += FString::Format(TEXT("{0}: TreeIndex {1}, Enabled {2}, InCluster {3}, TFormIdx {4}, ")
+		, { GetDebugName(), GetTreeIndex(), IsEnabled(), IsClustered(), GetTransformIndex() });
 
 	return true; 
 }
 
+const FTransform& ISimulationModuleBase::GetParentRelativeTransform() const
+{
+	if (bClustered)
+	{
+		return GetClusteredTransform();
+	}
+	else
+	{
+		return GetIntactTransform();
+	}
+}
+
 
 } //namespace Chaos
\ No newline at end of file
diff --git a/Engine/Source/Runtime/Experimental/ChaosVehicles/ChaosVehiclesCore/Private/SimModule/SuspensionModule.cpp b/Engine/Source/Runtime/Experimental/ChaosVehicles/ChaosVehiclesCore/Private/SimModule/SuspensionModule.cpp
--- a/Engine/Source/Runtime/Experimental/ChaosVehicles/ChaosVehiclesCore/Private/SimModule/SuspensionModule.cpp
+++ b/Engine/Source/Runtime/Experimental/ChaosVehicles/ChaosVehiclesCore/Private/SimModule/SuspensionModule.cpp
@@ -6,7 +6,7 @@
 #include "VehicleUtility.h"
 
 #if VEHICLE_DEBUGGING_ENABLED
-PRAGMA_DISABLE_OPTIMIZATION
+UE_DISABLE_OPTIMIZATION
 #endif
 
 namespace Chaos
@@ -16,7 +16,7 @@ namespace Chaos
 		: TSimModuleSettings<FSuspensionSettings>(Settings)
 		, SpringDisplacement(0.f)
 		, LastDisplacement(0.f)
-		, WheelSimTreeIndex(INVALID_INDEX)
+		, WheelSimTreeIndex(INVALID_IDX)
 	{
 		AccessSetup().MaxLength = FMath::Abs(Settings.MaxRaise + Settings.MaxDrop);
 	}
@@ -36,8 +36,9 @@ namespace Chaos
 	void FSuspensionSimModule::GetWorldRaycastLocation(const FTransform& BodyTransform, float WheelRadius, FSpringTrace& OutTrace)
 	{
 		FVector LocalDirection = Setup().SuspensionAxis;
-		FVector WorldLocation = BodyTransform.TransformPosition(GetParentRelativeTransform().GetLocation()/*Setup().LocalOffset*/);
-		FVector WorldDirection = BodyTransform.TransformVector(LocalDirection);
+		FVector Local = GetParentRelativeTransform().GetLocation(); // change to just a vector and GetLocalLocation
+		FVector WorldLocation = BodyTransform.TransformPosition(GetRelativeOffsetTransform().TransformVector(Local/*Setup().LocalOffset*/));
+		FVector WorldDirection = BodyTransform.TransformVector(GetRelativeOffsetTransform().TransformVector(LocalDirection));
 
 		OutTrace.Start = WorldLocation - WorldDirection * (Setup().MaxRaise);
 		OutTrace.End = WorldLocation + WorldDirection * (Setup().MaxDrop + WheelRadius);
@@ -79,10 +80,43 @@ namespace Chaos
 		}
 	}
 
+	void FSuspensionSimModuleDatas::FillSimState(ISimulationModuleBase* SimModule)
+	{
+		if (FSuspensionSimModule* Sim = static_cast<FSuspensionSimModule*>(SimModule))
+		{
+			Sim->SpringDisplacement = SpringDisplacement;
+			Sim->LastDisplacement = LastDisplacement;
+		}
+	}
+
+	void FSuspensionSimModuleDatas::FillNetState(const ISimulationModuleBase* SimModule)
+	{
+		if (const FSuspensionSimModule* Sim = static_cast<const FSuspensionSimModule*>(SimModule))
+		{
+			SpringDisplacement = Sim->SpringDisplacement;
+			LastDisplacement = Sim->LastDisplacement;
+		}
+	}
+
+	void FSuspensionSimModuleDatas::Lerp(const float LerpFactor, const FModuleNetData& Min, const FModuleNetData& Max)
+	{
+		const FSuspensionSimModuleDatas& MinData = static_cast<const FSuspensionSimModuleDatas&>(Min);
+		const FSuspensionSimModuleDatas& MaxData = static_cast<const FSuspensionSimModuleDatas&>(Max);
+
+		SpringDisplacement = FMath::Lerp(MinData.SpringDisplacement, MaxData.SpringDisplacement, LerpFactor);
+		LastDisplacement = FMath::Lerp(MinData.LastDisplacement, MaxData.LastDisplacement, LerpFactor);
+	}
 
+#if !(UE_BUILD_SHIPPING || UE_BUILD_TEST)
+	FString FSuspensionSimModuleDatas::ToString() const
+	{
+		return FString::Printf(TEXT("Module:%s SpringDisplacement:%f LastDisplacement:%f"),
+			*DebugString, SpringDisplacement, LastDisplacement);
+	}
+#endif
 
 } // namespace Chaos
 
 #if VEHICLE_DEBUGGING_ENABLED
-PRAGMA_ENABLE_OPTIMIZATION
+UE_ENABLE_OPTIMIZATION
 #endif
\ No newline at end of file
diff --git a/Engine/Source/Runtime/Experimental/ChaosVehicles/ChaosVehiclesCore/Private/SimModule/ThrusterModule.cpp b/Engine/Source/Runtime/Experimental/ChaosVehicles/ChaosVehiclesCore/Private/SimModule/ThrusterModule.cpp
--- a/Engine/Source/Runtime/Experimental/ChaosVehicles/ChaosVehiclesCore/Private/SimModule/ThrusterModule.cpp
+++ b/Engine/Source/Runtime/Experimental/ChaosVehicles/ChaosVehiclesCore/Private/SimModule/ThrusterModule.cpp
@@ -5,7 +5,7 @@
 #include "VehicleUtility.h"
 
 #if VEHICLE_DEBUGGING_ENABLED
-PRAGMA_DISABLE_OPTIMIZATION
+UE_DISABLE_OPTIMIZATION
 #endif
 
 namespace Chaos
@@ -28,5 +28,5 @@ namespace Chaos
 } // namespace Chaos
 
 #if VEHICLE_DEBUGGING_ENABLED
-PRAGMA_ENABLE_OPTIMIZATION
+UE_ENABLE_OPTIMIZATION
 #endif
diff --git a/Engine/Source/Runtime/Experimental/ChaosVehicles/ChaosVehiclesCore/Private/SimModule/TorqueSimModule.cpp b/Engine/Source/Runtime/Experimental/ChaosVehicles/ChaosVehiclesCore/Private/SimModule/TorqueSimModule.cpp
--- a/Engine/Source/Runtime/Experimental/ChaosVehicles/ChaosVehiclesCore/Private/SimModule/TorqueSimModule.cpp
+++ b/Engine/Source/Runtime/Experimental/ChaosVehicles/ChaosVehiclesCore/Private/SimModule/TorqueSimModule.cpp
@@ -6,7 +6,7 @@
 #include "VehicleUtility.h"
 
 #if VEHICLE_DEBUGGING_ENABLED
-PRAGMA_DISABLE_OPTIMIZATION
+UE_DISABLE_OPTIMIZATION
 #endif
 
 namespace Chaos
@@ -95,8 +95,47 @@ void FTorqueSimModule::IntegrateAngularVelocity(float DeltaTime, float Inertia,
 		, *GetDebugName(), DriveTorque, BrakingTorque, LoadTorque, AngularVelocity, GetRPM(), AngularPosition, Inertia);
 }
 
+void FTorqueSimModuleDatas::FillSimState(ISimulationModuleBase* SimModule)
+{
+	if (FTorqueSimModule* Sim = static_cast<FTorqueSimModule*>(SimModule))
+	{
+		Sim->AngularVelocity = AngularVelocity;
+		Sim->AngularPosition = AngularPosition;
+	}
+}
+
+void FTorqueSimModuleDatas::FillNetState(const ISimulationModuleBase* SimModule)
+{
+	if (const FTorqueSimModule* Sim = static_cast<const FTorqueSimModule*>(SimModule))
+	{
+		AngularVelocity = Sim->AngularVelocity;
+		AngularPosition = Sim->AngularPosition;
+	}
+}
+
+void FTorqueSimModuleDatas::Lerp(const float LerpFactor, const FModuleNetData& Min, const FModuleNetData& Max)
+{
+	const FTorqueSimModuleDatas& MinData = static_cast<const FTorqueSimModuleDatas&>(Min);
+	const FTorqueSimModuleDatas& MaxData = static_cast<const FTorqueSimModuleDatas&>(Max);
+
+	AngularVelocity = FMath::Lerp(MinData.AngularVelocity, MaxData.AngularVelocity, LerpFactor);
+	AngularPosition = FMath::Lerp(MinData.AngularPosition, MaxData.AngularPosition, LerpFactor);
+	DriveTorque = FMath::Lerp(MinData.DriveTorque, MaxData.DriveTorque, LerpFactor);
+	LoadTorque = FMath::Lerp(MinData.LoadTorque, MaxData.LoadTorque, LerpFactor);
+	BrakingTorque = FMath::Lerp(MinData.BrakingTorque, MaxData.BrakingTorque, LerpFactor);
+}
+
+
+#if !(UE_BUILD_SHIPPING || UE_BUILD_TEST)
+FString FTorqueSimModuleDatas::ToString() const
+{
+	return FString::Printf(TEXT("Module:%s AngularVelocity:%f AngularPosition:%f"),
+		*DebugString, AngularVelocity, AngularPosition);
+}
+#endif
+
 } // namespace Chaos
 
 #if VEHICLE_DEBUGGING_ENABLED
-PRAGMA_ENABLE_OPTIMIZATION
+UE_ENABLE_OPTIMIZATION
 #endif
diff --git a/Engine/Source/Runtime/Experimental/ChaosVehicles/ChaosVehiclesCore/Private/SimModule/TransmissionModule.cpp b/Engine/Source/Runtime/Experimental/ChaosVehicles/ChaosVehiclesCore/Private/SimModule/TransmissionModule.cpp
--- a/Engine/Source/Runtime/Experimental/ChaosVehicles/ChaosVehiclesCore/Private/SimModule/TransmissionModule.cpp
+++ b/Engine/Source/Runtime/Experimental/ChaosVehicles/ChaosVehiclesCore/Private/SimModule/TransmissionModule.cpp
@@ -5,7 +5,7 @@
 #include "VehicleUtility.h"
 
 #if VEHICLE_DEBUGGING_ENABLED
-PRAGMA_DISABLE_OPTIMIZATION
+UE_DISABLE_OPTIMIZATION
 #endif
 
 namespace Chaos
@@ -133,9 +133,47 @@ namespace Chaos
 		return true;
 	}
 
+	void FTransmissionSimModuleDatas::FillSimState(ISimulationModuleBase* SimModule)
+	{
+		if(FTransmissionSimModule* Sim = static_cast<FTransmissionSimModule*>(SimModule))
+		{
+			Sim->CurrentGear = CurrentGear;
+			Sim->TargetGear = TargetGear;
+			Sim->CurrentGearChangeTime = CurrentGearChangeTime;
+		}
+	}
+
+	void FTransmissionSimModuleDatas::FillNetState(const ISimulationModuleBase* SimModule)
+	{
+		if (const FTransmissionSimModule* Sim = static_cast<const FTransmissionSimModule*>(SimModule))
+		{
+			CurrentGear = Sim->CurrentGear;
+			TargetGear = Sim->TargetGear;
+			CurrentGearChangeTime = Sim->CurrentGearChangeTime;
+		}
+	}
+
+	void FTransmissionSimModuleDatas::Lerp(const float LerpFactor, const FModuleNetData& Min, const FModuleNetData& Max)
+	{
+		const FTransmissionSimModuleDatas& MinData = static_cast<const FTransmissionSimModuleDatas&>(Min);
+		const FTransmissionSimModuleDatas& MaxData = static_cast<const FTransmissionSimModuleDatas&>(Max);
+
+		CurrentGear = LerpFactor < 0.5 ? MinData.CurrentGear : MaxData.CurrentGear;
+		TargetGear = LerpFactor < 0.5 ? MinData.TargetGear : MaxData.TargetGear;
+		CurrentGearChangeTime = FMath::Lerp(MinData.CurrentGearChangeTime, MaxData.CurrentGearChangeTime, LerpFactor);
+	}
+
+#if !(UE_BUILD_SHIPPING || UE_BUILD_TEST)
+	FString FTransmissionSimModuleDatas::ToString() const
+	{
+		return FString::Printf(TEXT("Module:%s CurrentGear:%d TargetGear:%d CurrentGearChangeTime:%f"),
+			*DebugString, CurrentGear, TargetGear, CurrentGearChangeTime);
+	}
+#endif
+
 
 } // namespace Chaos
 
 #if VEHICLE_DEBUGGING_ENABLED
-PRAGMA_ENABLE_OPTIMIZATION
+UE_ENABLE_OPTIMIZATION
 #endif
diff --git a/Engine/Source/Runtime/Experimental/ChaosVehicles/ChaosVehiclesCore/Private/SimModule/WheelModule.cpp b/Engine/Source/Runtime/Experimental/ChaosVehicles/ChaosVehiclesCore/Private/SimModule/WheelModule.cpp
--- a/Engine/Source/Runtime/Experimental/ChaosVehicles/ChaosVehiclesCore/Private/SimModule/WheelModule.cpp
+++ b/Engine/Source/Runtime/Experimental/ChaosVehicles/ChaosVehiclesCore/Private/SimModule/WheelModule.cpp
@@ -5,7 +5,7 @@
 #include "VehicleUtility.h"
 
 #if VEHICLE_DEBUGGING_ENABLED
-PRAGMA_DISABLE_OPTIMIZATION
+UE_DISABLE_OPTIMIZATION
 #endif
 
 namespace Chaos
@@ -16,7 +16,7 @@ namespace Chaos
 		, BrakeTorque(0.0f)
 		, ForceIntoSurface(0.0f)
 		, SurfaceFriction(1.0f)
-		, SuspensionSimTreeIndex(INVALID_INDEX)
+		, SuspensionSimTreeIndex(INVALID_IDX)
 		, ForceFromFriction(FVector::ZeroVector)
 		, MassPerWheel(500.0f*0.25f)
 		, SteerAngleDegrees(0.0f)
@@ -27,8 +27,8 @@ namespace Chaos
 	void FWheelSimModule::Simulate(float DeltaTime, const FAllInputs& Inputs, FSimModuleTree& VehicleModuleSystem)
 	{	
 		float Re = Setup().Radius;
-		float K = 0.4f;
-		float TorqueScaling = 0.00005f;
+		float K = 0.4f; // #TODO: 0.4 * number of wheels - fix the amount of ForceRequiredToStop wrt number of active contact points
+		float TorqueScaling = 1.0f;
 		float TractionControlAndABSScaling = 0.98f;	// how close to perfection is the system working
 
 		float HandbrakeTorque = Setup().HandbrakeEnabled ? Inputs.ControlInputs.Handbrake * Setup().HandbrakeTorque : 0.0f;
@@ -37,7 +37,7 @@ namespace Chaos
 		LoadTorque = 0.0f;
 
 		// TODO: think about doing this properly, stops vehicles rolling around on their own too much
-		// kindof an auto handbrake
+		// i.e. an auto handbrake feature
 		if (Inputs.ControlInputs.Brake < SMALL_NUMBER && Inputs.ControlInputs.Throttle < SMALL_NUMBER && ModuleLocalVelocity.X < 10.0f)
 		{
 			BrakeTorque = Setup().HandbrakeTorque;
@@ -48,7 +48,8 @@ namespace Chaos
 		if (bTouchingGround)
 		{
 			FRotator SteeringRotator(0.f, SteerAngleDegrees, 0.f);
-			FVector LocalWheelVelocity = SteeringRotator.UnrotateVector(ModuleLocalVelocity);
+			FVector Vel = SteeringRotator.UnrotateVector(ModuleLocalVelocity);
+			FVector LocalWheelVelocity = (Setup().Axis == EWheelAxis::X) ? FVector(Vel.X, Vel.Y, Vel.Z) : FVector(Vel.Y, Vel.X, Vel.Z); // Potential Axis Swap
 
 			float GroundAngularVelocity = LocalWheelVelocity.X / Re;
 			float Delta = GroundAngularVelocity - AngularVelocity;
@@ -58,7 +59,7 @@ namespace Chaos
 			float SlipAngle = FVehicleUtility::CalculateSlipAngle(LocalWheelVelocity.Y, LocalWheelVelocity.X);
 
 			float AppliedLinearDriveForce = DriveTorque / Re;
-			float AppliedLinearBrakeForce = BrakeTorque / Re;
+			float AppliedLinearBrakeForce = FMath::Abs(BrakeTorque) / Re;
 
 			// Longitudinal multiplier now affecting both brake and steering equally
 			float AvailableGrip = ForceIntoSurface * SurfaceFriction * Setup().FrictionMultiplier;
@@ -113,7 +114,7 @@ namespace Chaos
 					FinalLongitudinalForce = AppliedLinearDriveForce;
 				}
 
-				float ForceRequiredToBringToStop = -(MassPerWheel * K * LocalWheelVelocity.Y) / DeltaTime;
+				float ForceRequiredToBringToStop = (MassPerWheel * K * LocalWheelVelocity.Y) / DeltaTime;
 
 				// use slip angle to generate a sideways force
 				if (Setup().LateralSlipGraph.IsEmpty())
@@ -136,13 +137,43 @@ namespace Chaos
 					FinalLateralForce = -FinalLateralForce;
 				}
 
+				float LengthSquared = FinalLongitudinalForce * FinalLongitudinalForce + FinalLateralForce * FinalLateralForce;
+				bool bClipping = false;
+				if (LengthSquared > 0.05f)
+				{
+					float Length = FMath::Sqrt(LengthSquared);
+
+					float Clip = (AvailableGrip) / Length;
+					if (Clip < 1.0f)
+					{
+						if (Braking /*&& !bEngineBraking*/)
+						{
+							WheelLocked = true;
+						}
+
+						bClipping = true;
+						FinalLongitudinalForce *= Clip;
+						FinalLateralForce *= Clip;
+					}
+				}
+
 			}
 
 			ForceFromFriction = FVector::ZeroVector;
-			ForceFromFriction.X = FinalLongitudinalForce;
-			ForceFromFriction.Y = FinalLateralForce;
+			// Potential Axis Swap
+			if (Setup().Axis == EWheelAxis::X)
+			{
+				ForceFromFriction.X = FinalLongitudinalForce;
+				ForceFromFriction.Y = FinalLateralForce;
+			}
+			else
+			{
+				check(Setup().Axis == EWheelAxis::Y);
+				ForceFromFriction.Y = FinalLongitudinalForce;
+				ForceFromFriction.X = FinalLateralForce;
+			}
 
-			AddLocalForce(SteeringRotator.RotateVector(ForceFromFriction));
+ 			AddLocalForce(SteeringRotator.RotateVector(ForceFromFriction));
 			TransmitTorque(VehicleModuleSystem, DriveTorque, BrakeTorque);
 
 			DriveTorque -= AvailableGrip;
@@ -150,16 +181,17 @@ namespace Chaos
 			{
 				DriveTorque = 0.0f;
 			}
-			else
-			{ 
-				DriveTorque *= TorqueScaling;
-			}
+
+			DriveTorque *= TorqueScaling;
+
 			BrakingTorque -= AvailableGrip;
 			if (BrakingTorque < 0.0f)
 			{
 				BrakingTorque = 0.0f;
 			}
 
+			BrakingTorque *= TorqueScaling;
+
 			LoadTorque = TorqueFromGroundInteraction;
 		}
 
@@ -180,5 +212,5 @@ namespace Chaos
 
 
 #if VEHICLE_DEBUGGING_ENABLED
-PRAGMA_ENABLE_OPTIMIZATION
+UE_ENABLE_OPTIMIZATION
 #endif
diff --git a/Engine/Source/Runtime/Experimental/ChaosVehicles/ChaosVehiclesCore/Public/SimModule/AerofoilModule.h b/Engine/Source/Runtime/Experimental/ChaosVehicles/ChaosVehiclesCore/Public/SimModule/AerofoilModule.h
--- a/Engine/Source/Runtime/Experimental/ChaosVehicles/ChaosVehiclesCore/Public/SimModule/AerofoilModule.h
+++ b/Engine/Source/Runtime/Experimental/ChaosVehicles/ChaosVehiclesCore/Public/SimModule/AerofoilModule.h
@@ -67,6 +67,8 @@ namespace Chaos
 
 		virtual ~FAerofoilSimModule() {}
 
+		virtual TSharedPtr<FModuleNetData> GenerateNetData(int NodeArrayIndex) const { return nullptr; }
+
 		virtual eSimType GetSimType() const { return eSimType::Aerofoil; }
 
 		virtual const FString GetDebugName() const { return TEXT("Aerofoil"); }
diff --git a/Engine/Source/Runtime/Experimental/ChaosVehicles/ChaosVehiclesCore/Public/SimModule/ChassisModule.h b/Engine/Source/Runtime/Experimental/ChaosVehicles/ChaosVehiclesCore/Public/SimModule/ChassisModule.h
--- a/Engine/Source/Runtime/Experimental/ChaosVehicles/ChaosVehiclesCore/Public/SimModule/ChassisModule.h
+++ b/Engine/Source/Runtime/Experimental/ChaosVehicles/ChaosVehiclesCore/Public/SimModule/ChassisModule.h
@@ -31,6 +31,8 @@ namespace Chaos
 	public:
 		FChassisSimModule(const FChassisSettings& Settings);
 
+		virtual TSharedPtr<FModuleNetData> GenerateNetData(int NodeArrayIndex) const { return nullptr; }
+
 		virtual eSimType GetSimType() const { return eSimType::Chassis; }
 
 		virtual const FString GetDebugName() const { return TEXT("Chassis"); }
diff --git a/Engine/Source/Runtime/Experimental/ChaosVehicles/ChaosVehiclesCore/Public/SimModule/DeferredForcesModular.h b/Engine/Source/Runtime/Experimental/ChaosVehicles/ChaosVehiclesCore/Public/SimModule/DeferredForcesModular.h
--- a/Engine/Source/Runtime/Experimental/ChaosVehicles/ChaosVehiclesCore/Public/SimModule/DeferredForcesModular.h
+++ b/Engine/Source/Runtime/Experimental/ChaosVehicles/ChaosVehiclesCore/Public/SimModule/DeferredForcesModular.h
@@ -135,7 +135,7 @@ public:
 		ApplyImpulseAtPositionDatas.Add(ApplyImpulseAtPositionDataIn);
 	}	
 
-	Chaos::FPBDRigidClusteredParticleHandle* GetParticle(TArray<Chaos::FPBDRigidClusteredParticleHandle*>& Particles
+	Chaos::FPBDRigidParticleHandle* GetParticle(TArray<Chaos::FPBDRigidClusteredParticleHandle*>& Particles
 			, TArray<Chaos::FPBDRigidClusteredParticleHandle*>& ClusterParticles
 			, int TransformIndex
 			, const TManagedArray<FTransform>& Transforms
@@ -143,21 +143,35 @@ public:
 			, const TManagedArray<int32>& Parent
 			, FTransform& TransformOut);
 
+	Chaos::FPBDRigidParticleHandle* GetParticle(const FTransform& OffsetTransform
+		, TArray<Chaos::FPBDRigidParticleHandle*>& Particles
+		, TArray<Chaos::FPBDRigidClusteredParticleHandle*>& ClusterParticles
+		, int TransformIndex
+		, FTransform& TransformOut);
+
 	void Apply(TArray<Chaos::FPBDRigidClusteredParticleHandle*>& Particles
 		, TArray<Chaos::FPBDRigidClusteredParticleHandle*>& ClusterParticles
 		, const TManagedArray<FTransform>& Transforms
 		, const TManagedArray<FTransform>& CollectionMassToLocal
 		, const TManagedArray<int32>& Parent);
+	
+	void Apply(const FTransform& OffsetTransform
+		, TArray<Chaos::FPBDRigidParticleHandle*>& Particles
+		, TArray<Chaos::FPBDRigidClusteredParticleHandle*>& ClusterParticles);
 
+	void SetOffsetTransform(const FTransform& TransformIn) { ParticleOffsetTransform = TransformIn; }
+	const FTransform& GetOffsetTransform() const { return ParticleOffsetTransform; }
 private:
 
-	void AddForceAtPosition(Chaos::FPBDRigidClusteredParticleHandle* RigidHandle, const FApplyForceAtPositionData& DataIn, const FTransform& OffsetTransform);
-	void AddTorque(Chaos::FPBDRigidClusteredParticleHandle* RigidHandle, const FAddTorqueInRadiansData& DataIn, const FTransform& OffsetTransform);
-	void AddForce(Chaos::FPBDRigidClusteredParticleHandle* RigidHandle, const FApplyForceData& DataIn, const FTransform& OffsetTransform);
+	void AddForceAtPosition(Chaos::FPBDRigidParticleHandle* RigidHandle, const FApplyForceAtPositionData& DataIn, const FTransform& OffsetTransform);
+	void AddTorque(Chaos::FPBDRigidParticleHandle* RigidHandle, const FAddTorqueInRadiansData& DataIn, const FTransform& OffsetTransform);
+	void AddForce(Chaos::FPBDRigidParticleHandle* RigidHandle, const FApplyForceData& DataIn, const FTransform& OffsetTransform);
 
 	TArray<FApplyForceData> ApplyForceDatas;
 	TArray<FApplyForceAtPositionData> ApplyForceAtPositionDatas;
 	TArray<FAddTorqueInRadiansData> ApplyTorqueDatas;
 	TArray<FAddImpulseData> ApplyImpulseDatas;
 	TArray<FAddImpulseAtPositionData> ApplyImpulseAtPositionDatas;
+
+	FTransform ParticleOffsetTransform; // Odd rotation coming through from CU physics bodies
 };
diff --git a/Engine/Source/Runtime/Experimental/ChaosVehicles/ChaosVehiclesCore/Public/SimModule/SimModuleTree.h b/Engine/Source/Runtime/Experimental/ChaosVehicles/ChaosVehiclesCore/Public/SimModule/SimModuleTree.h
--- a/Engine/Source/Runtime/Experimental/ChaosVehicles/ChaosVehiclesCore/Public/SimModule/SimModuleTree.h
+++ b/Engine/Source/Runtime/Experimental/ChaosVehicles/ChaosVehiclesCore/Public/SimModule/SimModuleTree.h
@@ -6,28 +6,85 @@
 #include "SimModule/SimulationModuleBase.h"
 
 class FGeometryCollectionPhysicsProxy;
-
 namespace Chaos
 {
 	class ISimulationModuleBase;
-	struct FAllInputs;  // #TODO - needs template or bass class
+	class FClusterUnionPhysicsProxy;
+	struct FAllInputs;
+
+
+	struct FPendingModuleAdds
+	{
+		FPendingModuleAdds(int ParentIndexIn, ISimulationModuleBase* NewSimModuleIn)
+			: ParentIndex(ParentIndexIn), NewSimModule(NewSimModuleIn) {}
+
+		int ParentIndex;
+		ISimulationModuleBase* NewSimModule;
+	};
+
+	struct FPendingModuleDeletions
+	{
+		FPendingModuleDeletions(int GuidIn) : Guid(GuidIn) {}
+		int Guid;
+	};
+
+	// Each update tree has it's own local tree hierarchy, this will be translated into the actual tree hierarchy.
+	// The root in here will always translate to the root in the main sim tree
+	class FSimTreeUpdates
+	{
+	public:
+		int AddRoot(ISimulationModuleBase* NewSimModuleIn)
+		{
+			NewModules.Add(FPendingModuleAdds(-1, NewSimModuleIn));
+			return NewModules.Num()-1;
+		}
+
+		int AddNodeBelow(int ParentIndex, ISimulationModuleBase* NewSimModuleIn)
+		{
+			NewModules.Add(FPendingModuleAdds(ParentIndex, NewSimModuleIn));
+			return NewModules.Num() - 1;
+		}
+
+		void RemoveNode(int Guid)
+		{
+			DeletedModules.Add(FPendingModuleDeletions(Guid));
+		}
+
+		void ClearUpdates()
+		{
+			NewModules.Empty();
+			DeletedModules.Empty();
+		}
+
+		const TArray<FPendingModuleAdds>& GetNewModules() const { return NewModules; }
+		const TArray<FPendingModuleDeletions>& GetDeletedModules() const { return DeletedModules; }
+
+	private:
+		TArray<FPendingModuleAdds> NewModules;
+		TArray<FPendingModuleDeletions> DeletedModules;
+	};
+
 
 	class CHAOSVEHICLESCORE_API FSimModuleTree
 	{
+		friend class FModularVehicleBuilder;
+
 	public:
 		struct FSimModuleNode
 		{
 			FSimModuleNode()
 				: SimModule(nullptr)
-				, Parent(INVALID_INDEX)
+				, Parent(INVALID_IDX)
 			{
 			}
 
-			ISimulationModuleBase* SimModule;	// we own the pointer, or using sim module container?
+			bool IsValid() const { return (SimModule != nullptr); }
+
+			ISimulationModuleBase* SimModule;
 			int Parent;
 			TSet<int> Children;
 
-			const static int INVALID_INDEX = -1;
+			const static int INVALID_IDX = -1;
 
 		};
 
@@ -58,6 +115,7 @@ namespace Chaos
 		int AddRoot(ISimulationModuleBase* SimModule);
 
 		void Reparent(int Index, int ParentIndex);
+		FSimModuleNode& GetNode(int Index) { return SimulationModuleTree[Index]; }
 
 		int AddNodeBelow(int AtIndex, ISimulationModuleBase* SimModule);
 
@@ -65,7 +123,10 @@ namespace Chaos
 
 		void DeleteNode(int AtIndex);
 
+		void AppendTreeUpdates(const FSimTreeUpdates& TreeUpdates);
+
 		void Simulate(float DeltaTime, FAllInputs& Inputs, FGeometryCollectionPhysicsProxy* PhysicsProxy);
+		void Simulate(float DeltaTime, FAllInputs& Inputs, FClusterUnionPhysicsProxy* PhysicsProxy);
 
 		FDeferredForcesModular& AccessDeferredForces() { return DeferredForces; }
 		const FDeferredForcesModular& GetDeferredForces() const { return DeferredForces; }
@@ -76,6 +137,43 @@ namespace Chaos
 			return AllInputs.ControlInputs;
 		}
 
+		FSimModuleNode* LocateNodeByType(Chaos::eSimType InType)
+		{
+			for (FSimModuleNode& Node : SimulationModuleTree)
+			{
+				if (Node.SimModule && Node.SimModule->GetSimType() == InType)
+				{
+					return &Node;
+				}
+			}
+
+			return nullptr;
+		}
+
+		int GetLargestComponentIndex()
+		{
+			int LargestIndex = -1;
+			for (int I = 0; I < GetNumNodes(); I++)
+			{
+				if (Chaos::ISimulationModuleBase* SimModule = GetNode(I).SimModule)
+				{
+					if (SimModule->GetTransformIndex() > LargestIndex)
+					{
+						LargestIndex = SimModule->GetTransformIndex();
+					}
+				}
+			}
+			return LargestIndex;
+		}
+
+		void SetIsSimulating(bool IsSimulatingIn) { bIsSimulating = IsSimulatingIn; }
+		bool IsSimulating() { return bIsSimulating; }
+
+		void GenerateReplicationStructure(Chaos::FModuleNetDataArray& NetData);
+		void SetNetState(Chaos::FModuleNetDataArray& ModuleDatas);
+		void SetSimState(const Chaos::FModuleNetDataArray& ModuleDatas);
+		void InterpolateState(const float LerpFactor, Chaos::FModuleNetDataArray& LerpDatas, const Chaos::FModuleNetDataArray& MinDatas, const Chaos::FModuleNetDataArray& MaxDatas);
+
 	protected:
 		void SimulateNode(float DeltaTime, FAllInputs& Inputs, int NodeIdx);
 
@@ -84,8 +182,7 @@ namespace Chaos
 		int GetNextIndex();
 
 		void UpdateModuleVelocites(FGeometryCollectionPhysicsProxy* PhysicsProxy);
-
-	protected:
+		void UpdateModuleVelocites(FClusterUnionPhysicsProxy* PhysicsProxy);
 
 		TArray<FSimModuleNode> SimulationModuleTree;
 		TArray<int> FreeList;
@@ -94,6 +191,7 @@ namespace Chaos
 
 		Chaos::FAllInputs AllInputs;
 
+		bool bIsSimulating = false; //#TODO : temp while testing - remove
 	};
 
 
diff --git a/Engine/Source/Runtime/Experimental/ChaosVehicles/ChaosVehiclesCore/Public/SimModule/SimulationModuleBase.h b/Engine/Source/Runtime/Experimental/ChaosVehicles/ChaosVehiclesCore/Public/SimModule/SimulationModuleBase.h
--- a/Engine/Source/Runtime/Experimental/ChaosVehicles/ChaosVehiclesCore/Public/SimModule/SimulationModuleBase.h
+++ b/Engine/Source/Runtime/Experimental/ChaosVehicles/ChaosVehiclesCore/Public/SimModule/SimulationModuleBase.h
@@ -12,13 +12,15 @@ struct CHAOSVEHICLESCORE_API FCoreModularVehicleDebugParams
 {
 	bool ShowMass = false;
 	bool ShowForces = false;
-	float DrawForceScaling = 0.0005f;
+	float DrawForceScaling = 0.002f;
 	float LevelSlopeThreshold = 0.96f; // ~16 degrees
+	bool DisableForces = false;
 };
 
 namespace Chaos
 {
 	class FSimModuleTree;
+	struct FModuleNetData;
 
 	struct CHAOSVEHICLESCORE_API FControlInputs
 	{
@@ -131,7 +133,7 @@ namespace Chaos
 		Rudder,			// controls aircraft yaw
 		Elevator,		// controls aircraft pitch
 		Propeller,		// generates thrust when connected to a motor/engine
-
+		TorqueSim
 	};
 
 	/**
@@ -140,19 +142,23 @@ namespace Chaos
 	class CHAOSVEHICLESCORE_API ISimulationModuleBase
 	{
 	public:
-		const static int INVALID_INDEX = -1;
+		const static int INVALID_IDX = -1;
 
 		ISimulationModuleBase()
 			: SimModuleTree(nullptr)
-			, SimTreeIndex(INVALID_INDEX)
+			, SimTreeIndex(INVALID_IDX)
 			, StateFlags(Enabled)
-			, TransformIndex(INVALID_INDEX)
+			, TransformIndex(INVALID_IDX)
 			, ModuleLocalVelocity(FVector::ZeroVector)
 			, bClustered(true)
 			, AppliedForce(FVector::ZeroVector)
+			, Guid(INDEX_NONE)
 		{}
 		virtual ~ISimulationModuleBase() {}
 
+		const int GetGuid() { return Guid; }
+		void SetGuid(int GuidIn) { Guid = GuidIn; }
+
 		/**
 		* Get the friendly name for this module, primarily for logging & debugging module tree
 		 */
@@ -236,27 +242,26 @@ namespace Chaos
 		void SetClusteredTransform(const FTransform& TransformIn) { ClusteredCOMRelativeTransform  = TransformIn; }
 		const FTransform& GetClusteredTransform() const { return ClusteredCOMRelativeTransform; }
 
+		void SetInitialParticleTransform(const FTransform& TransformIn) { InitialParticleTransform = TransformIn; }
+		const FTransform& GetInitialParticleTransform() const { return InitialParticleTransform; }
+
+		void SetComponentTransform(const FTransform& TransformIn) { ComponentTransform = TransformIn; }
+		const FTransform& GetComponentTransform() const { return ComponentTransform; }
+
+		void SetRelativeOffsetTransform(const FTransform& TransformIn) { RelativeOffsetTransform = TransformIn; }
+		const FTransform& GetRelativeOffsetTransform() const { return RelativeOffsetTransform; }
+
 		/**
 		 * Set the COM relative transform of module when it is broken off, so relative to itself
 		 */
-		void SetIntactTransform(const FTransform& TransformIn) { IntactCOMRelativeTransform = TransformIn; }
+		void SetIntactTransform(const FTransform& TransformIn) { IntactCOMRelativeTransform = TransformIn; IsInitialized = true; }
 		const FTransform& GetIntactTransform() const { return IntactCOMRelativeTransform; }
-
+		bool IsInitialized = false;
 		/**
 		 * The modules transform relative to the simulating body will depend on whether the GC is intact (get the transform relative to intact cluster)
 		 * or fractured (transform relative to fractured part)
 		 */
-		const FTransform& GetParentRelativeTransform() const
-		{
-			if (bClustered)
-			{
-				return GetClusteredTransform();
-			}
-			else
-			{
-				return GetIntactTransform();
-			}
-		}
+		const FTransform& GetParentRelativeTransform() const;
 
 		/**
 		 * Update the module with its current velocity
@@ -270,6 +275,9 @@ namespace Chaos
 		// for headless chaos testing
 		const FVector& GetAppliedForce() { return AppliedForce; }
 
+		// this is the replication datas
+		virtual TSharedPtr<FModuleNetData> GenerateNetData(int NodeArrayIndex) const = 0;
+
 	protected:
 
 		FSimModuleTree* SimModuleTree;	// A pointer back to the simulation tree where we are stored
@@ -277,6 +285,10 @@ namespace Chaos
 		eSimModuleState StateFlags;	// TODO: make this more like flags
 		int TransformIndex; // Index of this Sim Module's node in Geometry Collection Transform array
 
+		FTransform InitialParticleTransform;
+		FTransform RelativeOffsetTransform;
+		FTransform ComponentTransform;
+
 		FTransform ClusteredCOMRelativeTransform;
 		FTransform IntactCOMRelativeTransform;
 		FVector ModuleLocalVelocity;
@@ -284,7 +296,40 @@ namespace Chaos
 
 		// for headless chaos testing
 		FVector AppliedForce;
+		int Guid; // needed a way of associating internal module with game thread.
+
+	};
+
+	/**
+	* Interface base class for all module network serialization
+	*/
+	struct CHAOSVEHICLESCORE_API FModuleNetData
+	{
+		FModuleNetData(int InSimArrayIndex, const FString& InDebugString = FString())
+			: SimArrayIndex(InSimArrayIndex)
+#if !(UE_BUILD_SHIPPING || UE_BUILD_TEST)
+			, DebugString(InDebugString)
+#endif
+		{}
+
+		virtual ~FModuleNetData() {}
 
+		virtual eSimType GetType() = 0;
+		virtual void Serialize(FArchive& Ar) = 0;
+		virtual void FillNetState(const ISimulationModuleBase* SimModule) = 0;
+		virtual void FillSimState(ISimulationModuleBase* SimModule) = 0;
+		virtual void Lerp(const float LerpFactor, const FModuleNetData& Max, const FModuleNetData& MaxValue) = 0;
+
+		int SimArrayIndex = -1;
+
+#if !(UE_BUILD_SHIPPING || UE_BUILD_TEST)
+		virtual FString ToString() const = 0;
+		FString DebugString;
+#endif
 	};
 
+	using FModuleNetDataArray = TArray<TSharedPtr<FModuleNetData>>;
+
 } // namespace Chaos
+
+
diff --git a/Engine/Source/Runtime/Experimental/ChaosVehicles/ChaosVehiclesCore/Public/SimModule/SuspensionModule.h b/Engine/Source/Runtime/Experimental/ChaosVehicles/ChaosVehiclesCore/Public/SimModule/SuspensionModule.h
--- a/Engine/Source/Runtime/Experimental/ChaosVehicles/ChaosVehiclesCore/Public/SimModule/SuspensionModule.h
+++ b/Engine/Source/Runtime/Experimental/ChaosVehicles/ChaosVehiclesCore/Public/SimModule/SuspensionModule.h
@@ -9,6 +9,38 @@ namespace Chaos
 	struct FAllInputs;
 	class FSimModuleTree;
 
+	struct CHAOSVEHICLESCORE_API FSuspensionSimModuleDatas : public FModuleNetData
+	{
+#if !(UE_BUILD_SHIPPING || UE_BUILD_TEST)
+		FSuspensionSimModuleDatas(int NodeArrayIndex, const FString& InDebugString) : FModuleNetData(NodeArrayIndex, InDebugString) {}
+#else
+		FSuspensionSimModuleDatas(int NodeArrayIndex) : FModuleNetData(NodeArrayIndex) {}
+#endif
+
+		virtual eSimType GetType() override { return eSimType::Suspension; }
+
+		virtual void FillSimState(ISimulationModuleBase* SimModule) override;
+
+		virtual void FillNetState(const ISimulationModuleBase* SimModule) override;
+
+		virtual void Serialize(FArchive& Ar) override
+		{
+			Ar << SpringDisplacement;
+			Ar << LastDisplacement;
+		}
+
+		virtual void Lerp(const float LerpFactor, const FModuleNetData& Min, const FModuleNetData& Max) override;
+
+#if !(UE_BUILD_SHIPPING || UE_BUILD_TEST)
+		virtual FString ToString() const override;
+#endif
+
+		float SpringDisplacement = 0.0f;
+		float LastDisplacement = 0.0f;
+	};
+
+
+
 	struct CHAOSVEHICLESCORE_API FSuspensionSettings
 	{
 		FSuspensionSettings()
@@ -64,11 +96,22 @@ namespace Chaos
 
 	class CHAOSVEHICLESCORE_API FSuspensionSimModule : public ISimulationModuleBase, public TSimModuleSettings<FSuspensionSettings>
 	{
+		friend FSuspensionSimModuleDatas;
+
 	public:
 
 		FSuspensionSimModule(const FSuspensionSettings& Settings);
 
 		virtual bool IsBehaviourType(eSimModuleTypeFlags InType) const override { return (InType & Raycast); }
+		virtual TSharedPtr<FModuleNetData> GenerateNetData(int SimArrayIndex) const
+		{
+			return MakeShared<FSuspensionSimModuleDatas>(
+				SimArrayIndex
+#if !(UE_BUILD_SHIPPING || UE_BUILD_TEST)
+				, GetDebugName()
+#endif			
+			);
+		}
 
 		virtual eSimType GetSimType() const { return eSimType::Suspension; }
 
@@ -80,9 +123,6 @@ namespace Chaos
 
 		virtual void Simulate(float DeltaTime, const FAllInputs& Inputs, FSimModuleTree& VehicleModuleSystem) override;
 
-		void SetLocation(const FVector& LocationIn) { Location = LocationIn; }
-		const FVector GetLocation() const { return /*Setup().RestOffset + */Location; }
-
 		const FVector& GetRestLocation() const { return Setup().RestOffset; }
 
 		void SetWheelSimTreeIndex(int WheelTreeIndexIn) { WheelSimTreeIndex = WheelTreeIndexIn; }
@@ -93,10 +133,9 @@ namespace Chaos
 		float SpringDisplacement;
 		float LastDisplacement;
 		int WheelSimTreeIndex;
-
-		FVector Location;
 	};
 
 
 } // namespace Chaos
 
+
diff --git a/Engine/Source/Runtime/Experimental/ChaosVehicles/ChaosVehiclesCore/Public/SimModule/ThrusterModule.h b/Engine/Source/Runtime/Experimental/ChaosVehicles/ChaosVehiclesCore/Public/SimModule/ThrusterModule.h
--- a/Engine/Source/Runtime/Experimental/ChaosVehicles/ChaosVehiclesCore/Public/SimModule/ThrusterModule.h
+++ b/Engine/Source/Runtime/Experimental/ChaosVehicles/ChaosVehiclesCore/Public/SimModule/ThrusterModule.h
@@ -37,6 +37,8 @@ namespace Chaos
 	public:
 		FThrusterSimModule(const FThrusterSettings& Settings);
 
+		virtual TSharedPtr<FModuleNetData> GenerateNetData(int NodeArrayIndex) const { return nullptr; }
+
 		virtual eSimType GetSimType() const { return eSimType::Thruster; }
 
 		virtual const FString GetDebugName() const { return TEXT("Thruster"); }
diff --git a/Engine/Source/Runtime/Experimental/ChaosVehicles/ChaosVehiclesCore/Public/SimModule/TorqueSimModule.h b/Engine/Source/Runtime/Experimental/ChaosVehicles/ChaosVehiclesCore/Public/SimModule/TorqueSimModule.h
--- a/Engine/Source/Runtime/Experimental/ChaosVehicles/ChaosVehiclesCore/Public/SimModule/TorqueSimModule.h
+++ b/Engine/Source/Runtime/Experimental/ChaosVehicles/ChaosVehiclesCore/Public/SimModule/TorqueSimModule.h
@@ -7,10 +7,48 @@
 
 namespace Chaos
 {
+	class FTorqueSimModule;
+
+	struct CHAOSVEHICLESCORE_API FTorqueSimModuleDatas : public FModuleNetData
+	{
+#if !(UE_BUILD_SHIPPING || UE_BUILD_TEST)
+		FTorqueSimModuleDatas(int NodeArrayIndex, const FString& InDebugString) : FModuleNetData(NodeArrayIndex, InDebugString) {}
+#else
+		FTorqueSimModuleDatas(int NodeArrayIndex) : FModuleNetData(NodeArrayIndex) {}
+#endif
+		virtual eSimType GetType() override { return eSimType::TorqueSim; }
+		
+		virtual void FillSimState(ISimulationModuleBase* SimModule) override;
+
+		virtual void FillNetState(const ISimulationModuleBase* SimModule) override;
+
+		virtual void Serialize(FArchive& Ar) override
+		{
+			Ar << AngularVelocity;
+			Ar << AngularPosition;
+			Ar << DriveTorque;
+			Ar << LoadTorque;
+			Ar << BrakingTorque;
+		}
+
+		virtual void Lerp(const float LerpFactor, const FModuleNetData& Min, const FModuleNetData& Max) override;
+
+#if !(UE_BUILD_SHIPPING || UE_BUILD_TEST)
+		virtual FString ToString() const override;
+#endif
+
+		float AngularVelocity = 0.0f;
+		float AngularPosition = 0.0f;
+		float DriveTorque = 0.0f;
+		float LoadTorque = 0.0f;
+		float BrakingTorque = 0.0f;
+	};
 
 
 class CHAOSVEHICLESCORE_API FTorqueSimModule : public ISimulationModuleBase
 {
+	friend FTorqueSimModuleDatas;
+
 public:
 	FTorqueSimModule()
 		: DriveTorque(0.f)
@@ -18,7 +56,6 @@ public:
 		, BrakingTorque(0.f)
 		, AngularVelocity(0.f)
 		, AngularPosition(0.0f)
-		, CombinedInertia(0.0f)
 	{
 	}
 
@@ -26,6 +63,15 @@ public:
 	 * Is Module of a specific type - used for casting
 	 */
 	virtual bool IsBehaviourType(eSimModuleTypeFlags InType) const { return (InType & TorqueBased); }
+	virtual TSharedPtr<FModuleNetData> GenerateNetData(int SimArrayIndex) const
+	{
+		return MakeShared<FTorqueSimModuleDatas>(
+			SimArrayIndex
+#if !(UE_BUILD_SHIPPING || UE_BUILD_TEST)
+			, GetDebugName()
+#endif			
+		);
+	}
 
 	void SetDriveTorque(float TorqueIn) { DriveTorque = TorqueIn; }
 	float GetDriveTorque() const { return DriveTorque; }
@@ -77,7 +123,6 @@ protected:
 	float BrakingTorque;
 	float AngularVelocity;
 	float AngularPosition;
-	float CombinedInertia; // Note: ?? Do we want this
 };
 
 
diff --git a/Engine/Source/Runtime/Experimental/ChaosVehicles/ChaosVehiclesCore/Public/SimModule/TransmissionModule.h b/Engine/Source/Runtime/Experimental/ChaosVehicles/ChaosVehiclesCore/Public/SimModule/TransmissionModule.h
--- a/Engine/Source/Runtime/Experimental/ChaosVehicles/ChaosVehiclesCore/Public/SimModule/TransmissionModule.h
+++ b/Engine/Source/Runtime/Experimental/ChaosVehicles/ChaosVehiclesCore/Public/SimModule/TransmissionModule.h
@@ -10,6 +10,40 @@ namespace Chaos
 	struct FAllInputs;
 	class FSimModuleTree;
 
+
+	struct CHAOSVEHICLESCORE_API FTransmissionSimModuleDatas : public FModuleNetData
+	{
+#if !(UE_BUILD_SHIPPING || UE_BUILD_TEST)
+		FTransmissionSimModuleDatas(int NodeArrayIndex, const FString& InDebugString) : FModuleNetData(NodeArrayIndex, InDebugString) {}
+#else
+		FTransmissionSimModuleDatas(int NodeArrayIndex) : FModuleNetData(NodeArrayIndex) {}
+#endif
+
+		virtual eSimType GetType() override { return eSimType::Transmission; }
+
+		virtual void FillSimState(ISimulationModuleBase* SimModule) override;
+
+		virtual void FillNetState(const ISimulationModuleBase* SimModule) override;
+
+		virtual void Serialize(FArchive& Ar) override
+		{
+			Ar << CurrentGear;
+			Ar << TargetGear;
+			Ar << CurrentGearChangeTime;
+		}
+
+		virtual void Lerp(const float LerpFactor, const FModuleNetData& Min, const FModuleNetData& Max) override;
+
+#if !(UE_BUILD_SHIPPING || UE_BUILD_TEST)
+		virtual FString ToString() const override;
+#endif
+
+		int32 CurrentGear = 0;
+		int32 TargetGear = 0;
+		float CurrentGearChangeTime = 0.0f;
+	};
+
+
 	struct CHAOSVEHICLESCORE_API FTransmissionSettings
 	{
 		enum ETransType : uint8
@@ -52,6 +86,8 @@ namespace Chaos
 
 	class CHAOSVEHICLESCORE_API FTransmissionSimModule : public FTorqueSimModule, public TSimModuleSettings<FTransmissionSettings>
 	{
+		friend FTransmissionSimModuleDatas;
+
 	public:
 
 		FTransmissionSimModule(const FTransmissionSettings& Settings)
@@ -63,6 +99,16 @@ namespace Chaos
 		{
 		}
 
+		virtual TSharedPtr<FModuleNetData> GenerateNetData(int SimArrayIndex) const
+		{
+			return MakeShared<FTransmissionSimModuleDatas>(
+				SimArrayIndex
+#if !(UE_BUILD_SHIPPING || UE_BUILD_TEST)
+				, GetDebugName()
+#endif			
+			);
+		}
+
 		virtual eSimType GetSimType() const { return eSimType::Transmission; }
 
 		virtual const FString GetDebugName() const { return TEXT("Transmission"); }
diff --git a/Engine/Source/Runtime/Experimental/ChaosVehicles/ChaosVehiclesCore/Public/SimModule/WheelModule.h b/Engine/Source/Runtime/Experimental/ChaosVehicles/ChaosVehiclesCore/Public/SimModule/WheelModule.h
--- a/Engine/Source/Runtime/Experimental/ChaosVehicles/ChaosVehiclesCore/Public/SimModule/WheelModule.h
+++ b/Engine/Source/Runtime/Experimental/ChaosVehicles/ChaosVehiclesCore/Public/SimModule/WheelModule.h
@@ -9,6 +9,12 @@ namespace Chaos
 	struct FAllInputs;
 	class FSimModuleTree;
 
+	enum CHAOSVEHICLESCORE_API EWheelAxis
+	{
+		X,	// X forward
+		Y	// Y forward
+	};
+
 	struct CHAOSVEHICLESCORE_API FWheelSettings
 	{
 		FWheelSettings()
@@ -29,6 +35,7 @@ namespace Chaos
 			, MaxBrakeTorque(4000)
 			, HandbrakeTorque(3000)
 			, MaxRotationVel(100.0f)
+			, Axis(EWheelAxis::X)
 
 		{
 
@@ -54,6 +61,7 @@ namespace Chaos
 		float HandbrakeTorque;
 
 		float MaxRotationVel;
+		EWheelAxis Axis;
 
 
 	};
