Fix CIS

diff --git a/Engine/Source/Runtime/Experimental/ChaosVehicles/ChaosVehiclesCore/Private/SimModule/DeferredForcesModular.cpp b/Engine/Source/Runtime/Experimental/ChaosVehicles/ChaosVehiclesCore/Private/SimModule/DeferredForcesModular.cpp
--- a/Engine/Source/Runtime/Experimental/ChaosVehicles/ChaosVehiclesCore/Private/SimModule/DeferredForcesModular.cpp
+++ b/Engine/Source/Runtime/Experimental/ChaosVehicles/ChaosVehiclesCore/Private/SimModule/DeferredForcesModular.cpp
@@ -10,8 +10,6 @@
 #include "Chaos/DebugDrawQueue.h"
 #endif
 
-UE_DISABLE_OPTIMIZATION
-
 FCoreModularVehicleDebugParams GCoreModularVehicleDebugParams;
 
 #if CHAOS_DEBUG_DRAW
@@ -288,5 +286,3 @@ void FDeferredForcesModular::AddTorque(Chaos::FPBDRigidParticleHandle* RigidHand
 	//	}
 	//}
 }
-
-UE_ENABLE_OPTIMIZATION
diff --git a/Engine/Source/Runtime/Experimental/ChaosVehicles/ChaosVehiclesCore/Private/SimModule/TransmissionModule.cpp b/Engine/Source/Runtime/Experimental/ChaosVehicles/ChaosVehiclesCore/Private/SimModule/TransmissionModule.cpp
--- a/Engine/Source/Runtime/Experimental/ChaosVehicles/ChaosVehiclesCore/Private/SimModule/TransmissionModule.cpp
+++ b/Engine/Source/Runtime/Experimental/ChaosVehicles/ChaosVehiclesCore/Private/SimModule/TransmissionModule.cpp
@@ -13,7 +13,7 @@ namespace Chaos
 
 	void FTransmissionSimModule::Simulate(float DeltaTime, const FAllInputs& Inputs, FSimModuleTree& VehicleModuleSystem)
 	{
-		if (Setup().TransmissionType == FTransmissionSettings::ETransType::Automatic)
+		if (Setup().TransmissionType == FTransmissionSettings::ETransType::AutomaticType)
 		{
 			// not currently changing gear, also don't want to change up because the wheels are spinning up due to having no load
 			if (!IsCurrentlyChangingGear() && AllowedToChangeGear)
diff --git a/Engine/Source/Runtime/Experimental/ChaosVehicles/ChaosVehiclesCore/Public/SimModule/TransmissionModule.h b/Engine/Source/Runtime/Experimental/ChaosVehicles/ChaosVehiclesCore/Public/SimModule/TransmissionModule.h
--- a/Engine/Source/Runtime/Experimental/ChaosVehicles/ChaosVehiclesCore/Public/SimModule/TransmissionModule.h
+++ b/Engine/Source/Runtime/Experimental/ChaosVehicles/ChaosVehiclesCore/Public/SimModule/TransmissionModule.h
@@ -48,8 +48,8 @@ namespace Chaos
 	{
 		enum ETransType : uint8
 		{
-			Manual,
-			Automatic
+			ManualType,
+			AutomaticType
 		};
 
 		FTransmissionSettings()
@@ -58,7 +58,7 @@ namespace Chaos
 			, ChangeDownRPM(2500)
 			, GearChangeTime(0.5f)
 			, TransmissionEfficiency(1.f)
-			, TransmissionType(ETransType::Automatic)
+			, TransmissionType(ETransType::AutomaticType)
 			, AutoReverse(true)
 		{
 			ForwardRatios.Add(2.85f);
