Add visitation tracking mechanism for contact modifiers

diff --git a/Engine/Source/Runtime/Experimental/Chaos/Private/Chaos/ContactModification.cpp b/Engine/Source/Runtime/Experimental/Chaos/Private/Chaos/ContactModification.cpp
--- a/Engine/Source/Runtime/Experimental/Chaos/Private/Chaos/ContactModification.cpp
+++ b/Engine/Source/Runtime/Experimental/Chaos/Private/Chaos/ContactModification.cpp
@@ -485,6 +485,11 @@ namespace Chaos
 		Constraint->DisableManifoldPoint(ContactPointIdx);
 	}
 
+	const FPBDCollisionConstraintContainerCookie& FContactPairModifier::GetConstraintContainerCookie() const
+	{
+		return Constraint->GetContainerCookie();
+	}
+
 	void FContactPairModifierIterator::SeekValidContact()
 	{
 		// Not valid to call from end.
@@ -530,8 +535,6 @@ namespace Chaos
 		, Particle(InParticle)
 	{
 		FParticleCollisions& ParticleCollisions = Particle->ParticleCollisions();
-		const int32 NumCollisions = ParticleCollisions.Num();
-		Constraints.Reserve(NumCollisions);
 		ParticleCollisions.VisitCollisions([this](FPBDCollisionConstraint& Constraint)
 		{
 			if (Constraint.GetManifoldPoints().Num() > 0 && !Constraint.GetDisabled())
@@ -540,7 +543,23 @@ namespace Chaos
 			}
 			return ECollisionVisitorResult::Continue;
 		}, ECollisionVisitorFlags::VisitActiveAwake | ECollisionVisitorFlags::VisitDisabled);
-		ensureMsgf(Constraints.Num() <= NumCollisions, TEXT("Number of constraints visited exceeded the number reported by the FParticleCollisions object, and therefore exceeded the reserved constraint cache."));
+	}
+
+	bool FVisitedContactPairsTracker::Visit(const FContactPairModifier& ContactPair)
+	{
+		// Get a reference object to the relevant bit
+		const int32 ConstraintIndex = ContactPair.GetConstraintContainerCookie().ConstraintIndex;
+		FBitReference BitRef = VisitedContacts[ConstraintIndex];
+
+		// Mark constraint visited and return true
+		if (BitRef == false)
+		{
+			BitRef = true;
+			return true;
+		}
+
+		// Return false - this constraint has already been marked visited
+		return false;
 	}
 
 	FContactPairModifierParticleRange FCollisionContactModifier::GetContacts(FGeometryParticleHandle* Particle)
@@ -548,6 +567,11 @@ namespace Chaos
 		return FContactPairModifierParticleRange(this, Particle);
 	}
 
+	FVisitedContactPairsTracker FCollisionContactModifier::MakeVisitedContactPairsTracker() const
+	{
+		return FVisitedContactPairsTracker(Constraints.Num());
+	}
+
 	TArrayView<FPBDCollisionConstraint* const>& FCollisionContactModifier::GetConstraints()
 	{
 		return Constraints;
diff --git a/Engine/Source/Runtime/Experimental/Chaos/Public/Chaos/Collision/PBDCollisionConstraint.h b/Engine/Source/Runtime/Experimental/Chaos/Public/Chaos/Collision/PBDCollisionConstraint.h
--- a/Engine/Source/Runtime/Experimental/Chaos/Public/Chaos/Collision/PBDCollisionConstraint.h
+++ b/Engine/Source/Runtime/Experimental/Chaos/Public/Chaos/Collision/PBDCollisionConstraint.h
@@ -194,6 +194,7 @@ namespace Chaos
 	{
 		friend class Private::FCollisionConstraintAllocator;
 		friend class Private::FCollisionContextAllocator;
+		friend class FContactPairModifier;
 		friend class FGenericParticlePairMidPhase;
 		friend class FParticlePairMidPhase;
 		friend class FPBDCollisionConstraints;
diff --git a/Engine/Source/Runtime/Experimental/Chaos/Public/Chaos/ContactModification.h b/Engine/Source/Runtime/Experimental/Chaos/Public/Chaos/ContactModification.h
--- a/Engine/Source/Runtime/Experimental/Chaos/Public/Chaos/ContactModification.h
+++ b/Engine/Source/Runtime/Experimental/Chaos/Public/Chaos/ContactModification.h
@@ -4,6 +4,16 @@
 #include "Chaos/Core.h"
 #include "ParticleHandleFwd.h"
 
+namespace
+{
+	// When a FContactPairModifierParticleRange is created to iterate over the contacts
+	// of a particular particle, an array of relevant constraints is generated.
+	//
+	// So long as the number of constraints for this particle does not exceed this constant
+	// value, no heap allocations will occur.
+	constexpr int32 ParticleContactsStackSize = 16;
+}
+
 namespace Chaos
 {
 	class FImplicitObject;
@@ -12,6 +22,7 @@ namespace Chaos
 	class FCollisionContactModifier;
 	class FContactPairModifier;
 	class FPerShapeData;
+	class FPBDCollisionConstraintContainerCookie;
 
 	class FContactPairModifier
 	{
@@ -236,6 +247,12 @@ namespace Chaos
 		CHAOS_API void SetContactPointDisabled(int32 ContactPointIdx) const;
 
 	private:
+
+		/*
+		* Get direct const access to the index of this constraint container cookie
+		*/
+		CHAOS_API const FPBDCollisionConstraintContainerCookie& GetConstraintContainerCookie() const;
+
 		/**
 		 * @brief Update cached shape transforms in the constraint after modifying particle positions
 		*/
@@ -243,6 +260,8 @@ namespace Chaos
 
 		FPBDCollisionConstraint* Constraint;
 		FCollisionContactModifier* Modifier;
+
+		friend class FVisitedContactPairsTracker;
 	};
 
 	class FContactPairModifierIterator
@@ -371,7 +390,7 @@ namespace Chaos
 		}
 
 	private:
-		FContactPairModifierParticleRangeIterator(FCollisionContactModifier* InModifier, TArray<FPBDCollisionConstraint*>& InConstraints, int32 InIndex)
+		FContactPairModifierParticleRangeIterator(FCollisionContactModifier* InModifier, TArray<FPBDCollisionConstraint*, TInlineAllocator<ParticleContactsStackSize>>& InConstraints, int32 InIndex)
 			: Modifier(InModifier)
 			, Constraints(InConstraints)
 			, Index(InIndex)
@@ -385,7 +404,7 @@ namespace Chaos
 		FCollisionContactModifier* Modifier;
 		FContactPairModifier PairModifier;
 
-		TArray<FPBDCollisionConstraint*>& Constraints;
+		TArray<FPBDCollisionConstraint*, TInlineAllocator<ParticleContactsStackSize>>& Constraints;
 		int32 Index;
 
 		// Befriend the range object so that it can create iterators
@@ -411,12 +430,33 @@ namespace Chaos
 		FContactPairModifierParticleRange(FCollisionContactModifier* InModifier, FGeometryParticleHandle* InParticle);
 		FCollisionContactModifier* Modifier;
 		FGeometryParticleHandle* Particle;
-		TArray<FPBDCollisionConstraint*> Constraints;
+		TArray<FPBDCollisionConstraint*, TInlineAllocator<ParticleContactsStackSize>> Constraints;
 
 		// Befriend the modifier so that it can create ranges
 		friend class FCollisionContactModifier;
 	};
-	
+
+	/*
+	* Utility for tracking visited contact pair modifiers so that clients can avoid processing duplicates
+	*/
+	class FVisitedContactPairsTracker
+	{
+	public:
+
+		// If a contact pair was not already visited, mark it visited and return true. Otherwise return false.
+		CHAOS_API bool Visit(const FContactPairModifier& ContactPair);
+
+	private:
+		FVisitedContactPairsTracker(const int32 NumConstraints)
+			: VisitedContacts(false, NumConstraints)
+		{ }
+
+		TBitArray<> VisitedContacts;
+
+		// Befriend the modifier so that it can create these
+		friend class FCollisionContactModifier;
+	};
+
 	/*
 	*  Provides interface for iterating over modifiable contact pairs
 	*/
@@ -441,6 +481,8 @@ namespace Chaos
 
 		CHAOS_API FContactPairModifierParticleRange GetContacts(FGeometryParticleHandle* Particle);
 
+		CHAOS_API FVisitedContactPairsTracker MakeVisitedContactPairsTracker() const;
+
 	private:
 		CHAOS_API TArrayView<FPBDCollisionConstraint* const>& GetConstraints();
 		CHAOS_API void DisableConstraint(FPBDCollisionConstraint& Constraint);
