fix old style geometry collection events (break/crumbling/removal events) to still fire even if we set the notification boolean after the initial registration

diff --git a/Engine/Source/Runtime/Experimental/GeometryCollectionEngine/Private/GeometryCollection/GeometryCollectionComponent.cpp b/Engine/Source/Runtime/Experimental/GeometryCollectionEngine/Private/GeometryCollection/GeometryCollectionComponent.cpp
--- a/Engine/Source/Runtime/Experimental/GeometryCollectionEngine/Private/GeometryCollection/GeometryCollectionComponent.cpp
+++ b/Engine/Source/Runtime/Experimental/GeometryCollectionEngine/Private/GeometryCollection/GeometryCollectionComponent.cpp
@@ -803,7 +803,7 @@ void UGeometryCollectionComponent::SetNotifyBreaks(bool bNewNotifyBreaks)
 			PhysicsProxy->SetNotifyBreakings_External(bNewNotifyBreaks);
 		}
 		bNotifyBreaks = bNewNotifyBreaks;
-		UpdateBreakEventRegistration();
+		RegisterForEvents();
 	}
 }
 
@@ -816,7 +816,7 @@ void UGeometryCollectionComponent::SetNotifyRemovals(bool bNewNotifyRemovals)
 			PhysicsProxy->SetNotifyRemovals_External(bNewNotifyRemovals);
 		}
 		bNotifyRemovals = bNewNotifyRemovals;
-		UpdateRemovalEventRegistration();
+		RegisterForEvents();
 	}
 }
 
@@ -831,7 +831,7 @@ void UGeometryCollectionComponent::SetNotifyCrumblings(bool bNewNotifyCrumblings
 		}
 		bNotifyCrumblings = bNewNotifyCrumblings;
 		bCrumblingEventIncludesChildren = bNewCrumblingEventIncludesChildren;
-		UpdateCrumblingEventRegistration();
+		RegisterForEvents();
 	}
 }
 
@@ -1509,7 +1509,6 @@ void UGeometryCollectionComponent::RegisterForEvents()
 			// re-create the physics proxy (e.g. in the case of the construction script in PIE). In that case, UChaosGameplayEventDispatcher::OnRegister
 			// will run *before* the new physics proxy is created.
 			EventDispatcher->UnregisterChaosEvents();
-			EventDispatcher->RegisterChaosEvents();
 
 			if (BodyInstance.bNotifyRigidBodyCollision || bNotifyBreaks || bNotifyCollisions || bNotifyRemovals || bNotifyCrumblings)
 			{
@@ -1553,6 +1552,10 @@ void UGeometryCollectionComponent::RegisterForEvents()
 						});
 				}
 			}
+
+			// Needs to be after the RegisterForXYZEvents calls on the event dispatcher or else GetInterestedProxyOwnersForXYZEvents will not
+			// properly allow our registration to go through.
+			EventDispatcher->RegisterChaosEvents();
 		}
 		if (bNotifyGlobalBreaks || bNotifyGlobalCrumblings)
 		{
@@ -1615,35 +1618,6 @@ void UGeometryCollectionComponent::UpdateGlobalCollisionEventRegistration()
 		}
 	}
 }
-void UGeometryCollectionComponent::UpdateBreakEventRegistration()
-{
-	if (EventDispatcher)
-	{
-		if (bNotifyBreaks)
-		{
-			EventDispatcher->RegisterForBreakEvents(this, &DispatchGeometryCollectionBreakEvent);
-		}
-		else
-		{
-			EventDispatcher->UnRegisterForBreakEvents(this);
-		}
-	}
-}
-
-void UGeometryCollectionComponent::UpdateRemovalEventRegistration()
-{
-	if (EventDispatcher)
-	{
-		if (bNotifyRemovals)
-		{
-			EventDispatcher->RegisterForRemovalEvents(this, &DispatchGeometryCollectionRemovalEvent);
-		}
-		else
-		{
-			EventDispatcher->UnRegisterForRemovalEvents(this);
-		}
-	}
-}
 
 void UGeometryCollectionComponent::UpdateGlobalRemovalEventRegistration()
 {
@@ -1660,21 +1634,6 @@ void UGeometryCollectionComponent::UpdateGlobalRemovalEventRegistration()
 	}
 }
 
-void UGeometryCollectionComponent::UpdateCrumblingEventRegistration()
-{
-	if (EventDispatcher)
-	{
-		if (bNotifyCrumblings)
-		{
-			EventDispatcher->RegisterForCrumblingEvents(this, &DispatchGeometryCollectionCrumblingEvent);
-		}
-		else
-		{
-			EventDispatcher->UnRegisterForCrumblingEvents(this);
-		}
-	}
-}
-
 void ActivateClusters(Chaos::FRigidClustering& Clustering, Chaos::TPBDRigidClusteredParticleHandle<Chaos::FReal, 3>* Cluster)
 {
 	if(!Cluster)
diff --git a/Engine/Source/Runtime/Experimental/GeometryCollectionEngine/Public/GeometryCollection/GeometryCollectionComponent.h b/Engine/Source/Runtime/Experimental/GeometryCollectionEngine/Public/GeometryCollection/GeometryCollectionComponent.h
--- a/Engine/Source/Runtime/Experimental/GeometryCollectionEngine/Public/GeometryCollection/GeometryCollectionComponent.h
+++ b/Engine/Source/Runtime/Experimental/GeometryCollectionEngine/Public/GeometryCollection/GeometryCollectionComponent.h
@@ -1227,10 +1227,7 @@ protected:
 	GEOMETRYCOLLECTIONENGINE_API void RegisterForEvents();
 	GEOMETRYCOLLECTIONENGINE_API void UpdateRBCollisionEventRegistration();
 	GEOMETRYCOLLECTIONENGINE_API void UpdateGlobalCollisionEventRegistration();
-	GEOMETRYCOLLECTIONENGINE_API void UpdateBreakEventRegistration();
-	GEOMETRYCOLLECTIONENGINE_API void UpdateRemovalEventRegistration();
 	GEOMETRYCOLLECTIONENGINE_API void UpdateGlobalRemovalEventRegistration();
-	GEOMETRYCOLLECTIONENGINE_API void UpdateCrumblingEventRegistration();
 	
 	/* Per-instance override to enable/disable replication for the geometry collection */
 	UPROPERTY(EditAnywhere, BlueprintReadWrite, Category=Network)
