strain modifications should happen on all the particles we want to process and not just the top level strained particles

diff --git a/Engine/Source/Runtime/Experimental/Chaos/Private/Chaos/PBDRigidClustering.cpp b/Engine/Source/Runtime/Experimental/Chaos/Private/Chaos/PBDRigidClustering.cpp
--- a/Engine/Source/Runtime/Experimental/Chaos/Private/Chaos/PBDRigidClustering.cpp
+++ b/Engine/Source/Runtime/Experimental/Chaos/Private/Chaos/PBDRigidClustering.cpp
@@ -1343,14 +1343,6 @@ namespace Chaos
 				ResetCollisionImpulseArray();
 			}
 
-			//
-			// Modify internal strains
-			//
-			if (StrainModifiers)
-			{
-				ApplyStrainModifiers();
-			}
-
 			//  Monitor the MStrain array for 0 or less values.
 			//  That will trigger a break too.
 			//
@@ -1444,6 +1436,16 @@ namespace Chaos
 				}
 			}
 
+			//
+			// Modify internal strains. This needs to happen after the previous block so we can make sure
+			// we're applying the strain modifiers on ParticlesToProcess instead of TopLevelClusterParentsStrained
+			// since the latter may miss things related to cluster unions.
+			//
+			if (StrainModifiers)
+			{
+				ApplyStrainModifiers(ParticlesToProcess);
+			}
+
 			ClusterUnionManager.HandleDeferredClusterUnionUpdateProperties();
 
 			// Breaking can populate this again with relevant children - so we clear before running the breaking model
@@ -1984,13 +1986,13 @@ namespace Chaos
 		ClusteredParticle->ClusterGroupIndex() = 0;
 	}
 
-	void FRigidClustering::ApplyStrainModifiers()
+	void FRigidClustering::ApplyStrainModifiers(const TArray<FPBDRigidClusteredParticleHandle*>& StrainedParticles)
 	{
 		if (StrainModifiers)
 		{
 			for (ISimCallbackObject* Modifier : *StrainModifiers)
 			{
-				FStrainModifierAccessor Accessor(*this);
+				FStrainModifierAccessor Accessor(*this, &StrainedParticles);
 				Modifier->StrainModification_Internal(Accessor);
 			}
 		}
diff --git a/Engine/Source/Runtime/Experimental/Chaos/Private/Chaos/StrainModification.cpp b/Engine/Source/Runtime/Experimental/Chaos/Private/Chaos/StrainModification.cpp
--- a/Engine/Source/Runtime/Experimental/Chaos/Private/Chaos/StrainModification.cpp
+++ b/Engine/Source/Runtime/Experimental/Chaos/Private/Chaos/StrainModification.cpp
@@ -138,58 +138,79 @@ bool Chaos::FStrainedProxyIterator::operator==(const FStrainedProxyIterator& Oth
 		Index == Other.Index;
 }
 
-Chaos::FStrainedProxyRange::FStrainedProxyRange(Chaos::FRigidClustering& InRigidClustering, const bool bRootLevelOnly)
+Chaos::FStrainedProxyRange::FStrainedProxyRange(Chaos::FRigidClustering& InRigidClustering, const bool bRootLevelOnly, const TArray<FPBDRigidClusteredParticleHandle*>* InStrainedParticles)
 	: RigidClustering(InRigidClustering)
+	, StrainedParticles(InStrainedParticles)
 {
 	const TSet<Chaos::FPBDRigidClusteredParticleHandle*>& StrainedParents = RigidClustering.GetTopLevelClusterParentsStrained();
 
-	Proxies.Reserve(StrainedParents.Num());
-	for (Chaos::FPBDRigidClusteredParticleHandle* Cluster : StrainedParents)
+	auto ForEveryParticle = [this, &StrainedParents]<typename TLambda>(TLambda&& Func)
 	{
-		// Make sure the cluster's physics proxy exists and is the right type
-		IPhysicsProxyBase* ProxyBase = Cluster->PhysicsProxy();
-		if (ProxyBase == nullptr)
+		if (StrainedParticles)
 		{
-			continue;
+			for (Chaos::FPBDRigidClusteredParticleHandle* Cluster : *StrainedParticles)
+			{
+				Func(Cluster);
+			}
 		}
-
-		if (ProxyBase->GetType() != EPhysicsProxyType::GeometryCollectionType)
+		else
 		{
-			continue;
+			for (Chaos::FPBDRigidClusteredParticleHandle* Cluster : StrainedParents)
+			{
+				Func(Cluster);
+			}
 		}
+	};
 
-		FGeometryCollectionPhysicsProxy* Proxy = static_cast<FGeometryCollectionPhysicsProxy*>(ProxyBase);
-
-		// Make sure the rest collection has a root index
-		if (bRootLevelOnly)
+	Proxies.Reserve(StrainedParticles ? StrainedParticles->Num() : StrainedParents.Num());
+	ForEveryParticle(
+		[this, bRootLevelOnly](Chaos::FPBDRigidClusteredParticleHandle* Cluster)
 		{
-			FSimulationParameters& Parameters = Proxy->GetSimParameters();
-			const int32 RootIndex = Parameters.InitialRootIndex;
-			TArray<Chaos::FPBDRigidClusteredParticleHandle*>& ParticleHandles = Proxy->GetSolverParticleHandles();
-			if (ParticleHandles.IsValidIndex(RootIndex))
+			// Make sure the cluster's physics proxy exists and is the right type
+			IPhysicsProxyBase* ProxyBase = Cluster->PhysicsProxy();
+			if (ProxyBase == nullptr)
+			{
+				return;
+			}
+
+			if (ProxyBase->GetType() != EPhysicsProxyType::GeometryCollectionType)
+			{
+				return;
+			}
+
+			FGeometryCollectionPhysicsProxy* Proxy = static_cast<FGeometryCollectionPhysicsProxy*>(ProxyBase);
+
+			// Make sure the rest collection has a root index
+			if (bRootLevelOnly)
 			{
-				if (Cluster != ParticleHandles[RootIndex])
+				FSimulationParameters& Parameters = Proxy->GetSimParameters();
+				const int32 RootIndex = Parameters.InitialRootIndex;
+				TArray<Chaos::FPBDRigidClusteredParticleHandle*>& ParticleHandles = Proxy->GetSolverParticleHandles();
+				if (ParticleHandles.IsValidIndex(RootIndex))
 				{
-					continue;
+					if (Cluster != ParticleHandles[RootIndex])
+					{
+						return;
+					}
 				}
 			}
-		}
 
-		// Only need to use AddUnique if we're not checking for root, since at most
-		// one cluster will have the rest collection's root index.
-		if (bRootLevelOnly)
-		{
-			Proxies.Add(Proxy);
-		}
-		else
-		{
-			Proxies.AddUnique(Proxy);
+			// Only need to use AddUnique if we're not checking for root, since at most
+			// one cluster will have the rest collection's root index.
+			if (bRootLevelOnly)
+			{
+				Proxies.Add(Proxy);
+			}
+			else
+			{
+				Proxies.AddUnique(Proxy);
+			}
 		}
-	}
+	);
 }
 
 Chaos::FStrainedProxyRange Chaos::FStrainModifierAccessor::GetStrainedProxies(const bool bRootLevelOnly)
 {
-	return FStrainedProxyRange(RigidClustering, bRootLevelOnly);
+	return FStrainedProxyRange(RigidClustering, bRootLevelOnly, StrainedParticles);
 }
 
diff --git a/Engine/Source/Runtime/Experimental/Chaos/Public/Chaos/PBDRigidClustering.h b/Engine/Source/Runtime/Experimental/Chaos/Public/Chaos/PBDRigidClustering.h
--- a/Engine/Source/Runtime/Experimental/Chaos/Public/Chaos/PBDRigidClustering.h
+++ b/Engine/Source/Runtime/Experimental/Chaos/Public/Chaos/PBDRigidClustering.h
@@ -419,7 +419,7 @@ public:
 	CHAOS_API void ComputeStrainFromCollision(const FPBDCollisionConstraints& CollisionRule, const FReal Dt);
 	CHAOS_API void ResetCollisionImpulseArray();
 	CHAOS_API void DisableCluster(FPBDRigidClusteredParticleHandle* ClusteredParticle);
-	CHAOS_API void ApplyStrainModifiers();
+	CHAOS_API void ApplyStrainModifiers(const TArray<FPBDRigidClusteredParticleHandle*>& StrainedParticles);
 
 	/*
 	* Connectivity
diff --git a/Engine/Source/Runtime/Experimental/Chaos/Public/Chaos/StrainModification.h b/Engine/Source/Runtime/Experimental/Chaos/Public/Chaos/StrainModification.h
--- a/Engine/Source/Runtime/Experimental/Chaos/Public/Chaos/StrainModification.h
+++ b/Engine/Source/Runtime/Experimental/Chaos/Public/Chaos/StrainModification.h
@@ -115,7 +115,7 @@ namespace Chaos
 	class FStrainedProxyRange
 	{
 	public:
-		CHAOS_API FStrainedProxyRange(Chaos::FRigidClustering& InRigidClustering, bool bRootLevelOnly);
+		CHAOS_API FStrainedProxyRange(Chaos::FRigidClustering& InRigidClustering, bool bRootLevelOnly, const TArray<FPBDRigidClusteredParticleHandle*>* InStrainedParticles);
 
 		FStrainedProxyRange(const FStrainedProxyRange& Other)
 			: RigidClustering(Other.RigidClustering)
@@ -135,6 +135,7 @@ namespace Chaos
 	private:
 		FRigidClustering& RigidClustering;
 		TArray<FGeometryCollectionPhysicsProxy*> Proxies;
+		const TArray<FPBDRigidClusteredParticleHandle*>* StrainedParticles;
 	};
 
 	// FStrainModifierAccessor
@@ -144,7 +145,10 @@ namespace Chaos
 	{
 	public:
 
-		FStrainModifierAccessor(FRigidClustering& InRigidClustering) : RigidClustering(InRigidClustering) { }
+		FStrainModifierAccessor(FRigidClustering& InRigidClustering, const TArray<FPBDRigidClusteredParticleHandle*>* InStrainedParticles = nullptr)
+			: RigidClustering(InRigidClustering)
+			, StrainedParticles(InStrainedParticles)
+		{}
 
 		// Get an iterable range of unique geometry collection proxies which
 		// correspond to all strained clusters. Optionally, only include proxies
@@ -154,5 +158,6 @@ namespace Chaos
 	private:
 
 		FRigidClustering& RigidClustering;
+		const TArray<FPBDRigidClusteredParticleHandle*>* StrainedParticles;
 	};
 }
