Chaos - remove bIsActivated flag from collisions and use bDisabled instead
- fix collision debugdraw

diff --git a/Engine/Source/Runtime/Experimental/Chaos/Private/Chaos/Collision/PBDCollisionConstraint.cpp b/Engine/Source/Runtime/Experimental/Chaos/Private/Chaos/Collision/PBDCollisionConstraint.cpp
--- a/Engine/Source/Runtime/Experimental/Chaos/Private/Chaos/Collision/PBDCollisionConstraint.cpp
+++ b/Engine/Source/Runtime/Experimental/Chaos/Private/Chaos/Collision/PBDCollisionConstraint.cpp
@@ -289,6 +289,9 @@ namespace Chaos
 			GetParticle1(), GetImplicit1(), GetShape1(),
 			DistanceCheckSize);
 
+		// The constraint will be enabled after the narrow phase if required
+		Flags.bDisabled = true;
+
 		// Contact modification may change bIsProbe. The original value is stored in BoundsTestFlags.
 		Flags.bIsProbe = BoundsTestFlags.bIsProbe;
 
diff --git a/Engine/Source/Runtime/Experimental/Chaos/Private/Chaos/Collision/ParticlePairMidPhase.cpp b/Engine/Source/Runtime/Experimental/Chaos/Private/Chaos/Collision/ParticlePairMidPhase.cpp
--- a/Engine/Source/Runtime/Experimental/Chaos/Private/Chaos/Collision/ParticlePairMidPhase.cpp
+++ b/Engine/Source/Runtime/Experimental/Chaos/Private/Chaos/Collision/ParticlePairMidPhase.cpp
@@ -483,9 +483,6 @@ namespace Chaos
 					return 1;
 				}
 			}
-
-			// If we get here, we did not activate the constraint and it should be disabled for this tick
-			Constraint->SetDisabled(true);
 		}
 
 		return 0;
@@ -1659,10 +1656,6 @@ namespace Chaos
 				Context.GetAllocator()->ActivateConstraint(Constraint);
 				++NumActive;
 			}
-			else
-			{
-				Constraint->SetDisabled(true);
-			}
 		}
 
 		NewConstraints.Reset();
diff --git a/Engine/Source/Runtime/Experimental/Chaos/Private/Chaos/PBDCollisionConstraints.cpp b/Engine/Source/Runtime/Experimental/Chaos/Private/Chaos/PBDCollisionConstraints.cpp
--- a/Engine/Source/Runtime/Experimental/Chaos/Private/Chaos/PBDCollisionConstraints.cpp
+++ b/Engine/Source/Runtime/Experimental/Chaos/Private/Chaos/PBDCollisionConstraints.cpp
@@ -512,6 +512,17 @@ namespace Chaos
 	{
 		SCOPE_CYCLE_COUNTER(STAT_Collisions_BeginDetect);
 
+		// Disable existing constraints so that if they are not re-activated this tick
+		// they do not have state indicating that that are still active.
+		// @todo(chaos): ideally we would do this only for constraints that do not get reused this tick in EndDetectCollisions
+		for (FPBDCollisionConstraint* Constraint : GetConstraints())
+		{
+			if (Constraint != nullptr)
+			{
+				Constraint->SetDisabled(true);
+			}
+		}
+
 		ConstraintAllocator.BeginDetectCollisions();
 	}
 
@@ -527,17 +538,6 @@ namespace Chaos
 		PruneEdgeCollisions();
 	}
 
-	void FPBDCollisionConstraints::EndTick()
-	{
-		for (FPBDCollisionConstraint* Constraint : GetConstraints())
-		{
-			if (Constraint != nullptr)
-			{
-				Constraint->EndTick();
-			}
-		}
-	}
-
 	void FPBDCollisionConstraints::DetectProbeCollisions(FReal Dt)
 	{
 		SCOPE_CYCLE_COUNTER(STAT_Collisions_DetectProbeCollisions);
diff --git a/Engine/Source/Runtime/Experimental/Chaos/Private/Chaos/PBDRigidsEvolutionGBF.cpp b/Engine/Source/Runtime/Experimental/Chaos/Private/Chaos/PBDRigidsEvolutionGBF.cpp
--- a/Engine/Source/Runtime/Experimental/Chaos/Private/Chaos/PBDRigidsEvolutionGBF.cpp
+++ b/Engine/Source/Runtime/Experimental/Chaos/Private/Chaos/PBDRigidsEvolutionGBF.cpp
@@ -149,7 +149,6 @@ namespace Chaos
 		DECLARE_CYCLE_STAT(TEXT("FPBDRigidsEvolutionGBF::SaveParticlePostSolve"), STAT_Evolution_SavePostSolve, STATGROUP_Chaos);
 		DECLARE_CYCLE_STAT(TEXT("FPBDRigidsEvolutionGBF::DeactivateSleep"), STAT_Evolution_DeactivateSleep, STATGROUP_Chaos);
 		DECLARE_CYCLE_STAT(TEXT("FPBDRigidsEvolutionGBF::InertiaConditioning"), STAT_Evolution_InertiaConditioning, STATGROUP_Chaos);
-		DECLARE_CYCLE_STAT(TEXT("FPBDRigidsEvolutionGBF::CollisionEndTick"), STAT_Evolution_CollisionEndTick, STATGROUP_Chaos);
 
 
 		int32 SerializeEvolution = 0;
@@ -629,14 +628,6 @@ void FPBDRigidsEvolutionGBF::AdvanceOneTimeStepImpl(const FReal Dt, const FSubSt
 		GetCollisionConstraints().DetectProbeCollisions(Dt);
 	}
 	
-	{
-		SCOPE_CYCLE_COUNTER(STAT_Evolution_CollisionEndTick);
-		CSV_SCOPED_TIMING_STAT(PhysicsVerbose, StepSolver_CollisionEndTick);
-
-		// Reset collision state ready for the next tick
-		CollisionConstraints.EndTick();
-	}
-
 	{
 		CVD_SCOPE_TRACE_SOLVER_STEP(TEXT("Evolution End"));
 
diff --git a/Engine/Source/Runtime/Experimental/Chaos/Public/Chaos/Collision/CollisionConstraintAllocator.h b/Engine/Source/Runtime/Experimental/Chaos/Public/Chaos/Collision/CollisionConstraintAllocator.h
--- a/Engine/Source/Runtime/Experimental/Chaos/Public/Chaos/Collision/CollisionConstraintAllocator.h
+++ b/Engine/Source/Runtime/Experimental/Chaos/Public/Chaos/Collision/CollisionConstraintAllocator.h
@@ -456,11 +456,11 @@ namespace Chaos
 			 * Visitor signature: ECollisionVisitorResult(const FPBDCollisionConstraint&)
 			*/
 			template<typename TLambda>
-			void VisitConstCollisions(const TLambda& Visitor) const
+			void VisitConstCollisions(const TLambda& Visitor, const ECollisionVisitorFlags& VisitFlags = ECollisionVisitorFlags::VisitDefault) const
 			{
 				for (const FParticlePairMidPhasePtr& MidPhase : ParticlePairMidPhases)
 				{
-					if (MidPhase->VisitConstCollisions(Visitor) == ECollisionVisitorResult::Stop)
+					if (MidPhase->VisitConstCollisions(Visitor, VisitFlags) == ECollisionVisitorResult::Stop)
 					{
 						return;
 					}
@@ -595,14 +595,6 @@ namespace Chaos
 			return false;
 		}
 
-		if (!(VisitFlags & ECollisionVisitorFlags::VisitInactive))
-		{
-			if (!Constraint->IsActivated())
-			{
-				return false;
-			}
-		}
-
 		if (!(VisitFlags & ECollisionVisitorFlags::VisitDisabled))
 		{
 			if (!Constraint->IsEnabled())
diff --git a/Engine/Source/Runtime/Experimental/Chaos/Public/Chaos/Collision/CollisionVisitor.h b/Engine/Source/Runtime/Experimental/Chaos/Public/Chaos/Collision/CollisionVisitor.h
--- a/Engine/Source/Runtime/Experimental/Chaos/Public/Chaos/Collision/CollisionVisitor.h
+++ b/Engine/Source/Runtime/Experimental/Chaos/Public/Chaos/Collision/CollisionVisitor.h
@@ -23,11 +23,10 @@ namespace Chaos
 	{
 		VisitActiveAwake	= 0,		// Collisions that we detected this tick and have not been disabled or put to sleep
 		VisitSleeping		= 1 << 1,	// Collisions on body pairs that are sleeping
-		VisitDisabled		= 1 << 2,	// Collisions that have been disabled by the user in a collision callback
-		VisitInactive		= 1 << 3,	// Collisions that were not detected this frame but may still be kept around for potential reuse (only really useful for debugging etc)
+		VisitDisabled		= 1 << 2,	// Collisions that have been disabled by the user, or were not enabled this tick
 
 		VisitDefault		= VisitActiveAwake | VisitSleeping,
-		VisitAll			= VisitActiveAwake | VisitSleeping | VisitDisabled | VisitInactive
+		VisitAll			= VisitActiveAwake | VisitSleeping | VisitDisabled
 	};
 	ENUM_CLASS_FLAGS(ECollisionVisitorFlags);
 
diff --git a/Engine/Source/Runtime/Experimental/Chaos/Public/Chaos/Collision/PBDCollisionConstraint.h b/Engine/Source/Runtime/Experimental/Chaos/Public/Chaos/Collision/PBDCollisionConstraint.h
--- a/Engine/Source/Runtime/Experimental/Chaos/Public/Chaos/Collision/PBDCollisionConstraint.h
+++ b/Engine/Source/Runtime/Experimental/Chaos/Public/Chaos/Collision/PBDCollisionConstraint.h
@@ -642,16 +642,6 @@ namespace Chaos
 			SavedManifoldPoints.Reset();
 		}
 
-		/**
-		 * Called after the simulation to reset any state that need to be reset before the next tick
-		 */
-		void EndTick()
-		{
-			// Reset the activated flag so that we don't need to reset it for constraints that are kept
-			// in memory but not re-activated next frame.
-			Flags.bIsActivated = false;
-		}
-
 		/**
 		 * @brief Time of impact from CCD sweep test if CCD is activate.Otherwise undefined.
 		*/
diff --git a/Engine/Source/Runtime/Experimental/Chaos/Public/Chaos/PBDCollisionConstraints.h b/Engine/Source/Runtime/Experimental/Chaos/Public/Chaos/PBDCollisionConstraints.h
--- a/Engine/Source/Runtime/Experimental/Chaos/Public/Chaos/PBDCollisionConstraints.h
+++ b/Engine/Source/Runtime/Experimental/Chaos/Public/Chaos/PBDCollisionConstraints.h
@@ -146,11 +146,6 @@ public:
 	*/
 	void DisableConstraints(const TSet<FGeometryParticleHandle*>& ParticleHandle) {}
 
-	/**
-	 * End of tick cleanup
-	 */
-	CHAOS_API void EndTick();
-
 	//
 	// FConstraintContainer Implementation
 	//
