geometry collections need to reset their dirty index after calling PushStateOnGameThread or else dirtying that proxy will never cause a GT->PT sync ever again

diff --git a/Engine/Source/Runtime/Experimental/Chaos/Private/PBDRigidsSolver.cpp b/Engine/Source/Runtime/Experimental/Chaos/Private/PBDRigidsSolver.cpp
--- a/Engine/Source/Runtime/Experimental/Chaos/Private/PBDRigidsSolver.cpp
+++ b/Engine/Source/Runtime/Experimental/Chaos/Private/PBDRigidsSolver.cpp
@@ -1275,6 +1275,7 @@ namespace Chaos
 			{
 				auto Proxy = static_cast<FGeometryCollectionPhysicsProxy*>(Dirty.Proxy);
 				Proxy->PushStateOnGameThread(this);
+				Proxy->ResetDirtyIdx();
 				break;
 			}
 			case EPhysicsProxyType::ClusterUnionProxy:
