Vist function to MidPhaseModifier to allow access to all midphase modifiers. Maybe something we shouldn't use often, but a useful tool when developing midphase modifiers.

diff --git a/Engine/Source/Runtime/Experimental/Chaos/Private/Chaos/MidPhaseModification.cpp b/Engine/Source/Runtime/Experimental/Chaos/Private/Chaos/MidPhaseModification.cpp
--- a/Engine/Source/Runtime/Experimental/Chaos/Private/Chaos/MidPhaseModification.cpp
+++ b/Engine/Source/Runtime/Experimental/Chaos/Private/Chaos/MidPhaseModification.cpp
@@ -1,6 +1,7 @@
 // Copyright Epic Games, Inc. All Rights Reserved.
 
 #include "Chaos/MidPhaseModification.h"
+#include "Chaos/Collision/CollisionConstraintAllocator.h"
 #include "Chaos/Collision/PBDCollisionConstraint.h"
 #include "Chaos/Collision/ParticlePairMidPhase.h"
 #include "Chaos/Particle/ParticleUtilities.h"
@@ -20,8 +21,8 @@ namespace Chaos
 	}
 
 	void FMidPhaseModifier::GetParticles(
-		const FGeometryParticleHandle** Particle0,
-		const FGeometryParticleHandle** Particle1) const
+		FGeometryParticleHandle** Particle0,
+		FGeometryParticleHandle** Particle1) const
 	{
 		*Particle0 = MidPhase->GetParticle0();
 		*Particle1 = MidPhase->GetParticle1();
@@ -87,4 +88,16 @@ namespace Chaos
 
 		return FMidPhaseModifier();
 	}
+
+	void FMidPhaseModifierAccessor::VisitMidPhases(const TFunction<void(FMidPhaseModifier&)>& Visitor)
+	{
+		// This visitor is a thin function wrapper so that users can't directly access
+		// the midphase itself, but instead can access a modifier.
+		ConstraintAllocator.VisitMidPhases([this, &Visitor](FParticlePairMidPhase& MidPhase)
+		{
+			FMidPhaseModifier MidPhaseModifier(&MidPhase, this);
+			Visitor(MidPhaseModifier);
+			return ECollisionVisitorResult::Continue;
+		});
+	}
 }
diff --git a/Engine/Source/Runtime/Experimental/Chaos/Private/Chaos/PBDCollisionConstraints.cpp b/Engine/Source/Runtime/Experimental/Chaos/Private/Chaos/PBDCollisionConstraints.cpp
--- a/Engine/Source/Runtime/Experimental/Chaos/Private/Chaos/PBDCollisionConstraints.cpp
+++ b/Engine/Source/Runtime/Experimental/Chaos/Private/Chaos/PBDCollisionConstraints.cpp
@@ -571,7 +571,7 @@ namespace Chaos
 
 	void FPBDCollisionConstraints::ApplyMidPhaseModifier(const TArray<ISimCallbackObject*>& MidPhaseModifiers, FReal Dt)
 	{
-		FMidPhaseModifierAccessor ModifierAccessor;
+		FMidPhaseModifierAccessor ModifierAccessor(GetConstraintAllocator());
 		for(ISimCallbackObject* ModifierCallback : MidPhaseModifiers)
 		{
 			ModifierCallback->MidPhaseModification_Internal(ModifierAccessor);
diff --git a/Engine/Source/Runtime/Experimental/Chaos/Public/Chaos/MidPhaseModification.h b/Engine/Source/Runtime/Experimental/Chaos/Public/Chaos/MidPhaseModification.h
--- a/Engine/Source/Runtime/Experimental/Chaos/Public/Chaos/MidPhaseModification.h
+++ b/Engine/Source/Runtime/Experimental/Chaos/Public/Chaos/MidPhaseModification.h
@@ -7,7 +7,10 @@
 
 namespace Chaos
 {
-	class FCollisionConstraintAllocator;	// Container for solver's constraints and midphase pairs
+	namespace Private
+	{
+		class FCollisionConstraintAllocator;	// Container for solver's constraints and midphase pairs
+	}
 	class FParticlePairMidPhase;			// The underlying midphase type
 	class FMidPhaseModifier;			// Modifier for a particular midphase
 	class FMidPhaseModifierAccessor;		// Access class for all midphase pair modifiers
@@ -32,7 +35,7 @@ namespace Chaos
 
 		// Since a modifier can be invalid (ie, null midphase and/or accessor), make it
 		// castable to bool so that users can check validity.
-		bool IsValid() const
+		CHAOS_API bool IsValid() const
 		{
 			return
 				MidPhase != nullptr &&
@@ -53,7 +56,7 @@ namespace Chaos
 		// Accessor functions
 		//
 
-		CHAOS_API void GetParticles(const FGeometryParticleHandle** Particle0, const FGeometryParticleHandle** Particle1) const;
+		CHAOS_API void GetParticles(FGeometryParticleHandle** Particle0, FGeometryParticleHandle** Particle1) const;
 
 		CHAOS_API const FGeometryParticleHandle* GetOtherParticle(const FGeometryParticleHandle* InParticle) const;
 
@@ -67,7 +70,7 @@ namespace Chaos
 	/*
 	 * Class for iterating over midphases involving a specific particle
 	 */
-	class FMidPhaseModifierParticleIterator
+	class CHAOS_API FMidPhaseModifierParticleIterator
 	{
 	public:
 		FMidPhaseModifierParticleIterator(
@@ -121,18 +124,26 @@ namespace Chaos
 		friend class FMidPhaseModifierParticleRange;
 	};
 
+	/*
+	* A range-iterable object for accessing all mid-phases which include a particular particle
+	*/
 	class FMidPhaseModifierParticleRange
 	{
 	public:
+		CHAOS_API FMidPhaseModifierParticleIterator begin() const;
+		CHAOS_API FMidPhaseModifierParticleIterator end() const;
+
+	private:
 		FMidPhaseModifierParticleRange(
 			FMidPhaseModifierAccessor* InAccessor, FGeometryParticleHandle* InParticle)
 			: Accessor(InAccessor)
 			, Particle(InParticle) { }
-		CHAOS_API FMidPhaseModifierParticleIterator begin() const;
-		CHAOS_API FMidPhaseModifierParticleIterator end() const;
-	private:
+
 		FMidPhaseModifierAccessor* Accessor;
 		FGeometryParticleHandle* Particle;
+
+		// Befriend the accessor so that it alone can instantiate this range object
+		friend class FMidPhaseModifierAccessor;
 	};
 
 	/*
@@ -141,11 +152,22 @@ namespace Chaos
 	class FMidPhaseModifierAccessor
 	{
 	public:
+		FMidPhaseModifierAccessor(Private::FCollisionConstraintAllocator& InConstraintAllocator)
+			: ConstraintAllocator(InConstraintAllocator)
+		{ }
+
 		// Get an object which allows for range iteration over the list of
 		// midphases for a particle
 		CHAOS_API FMidPhaseModifierParticleRange GetMidPhases(FGeometryParticleHandle* Particle);
 
 		// Get a midphase modifier for a particular object pair
 		CHAOS_API FMidPhaseModifier GetMidPhase(FGeometryParticleHandle* Particle0, FGeometryParticleHandle* Particle1);
+
+		// Provide visitor access to all midphase modifiers
+		CHAOS_API void VisitMidPhases(const TFunction<void(FMidPhaseModifier&)>& Visitor);
+
+	private:
+
+		Private::FCollisionConstraintAllocator& ConstraintAllocator;
 	};
 }
