Chaos - perf add counters for broadphase tracking

diff --git a/Engine/Source/Runtime/Experimental/Chaos/Private/PBDRigidsSolver.cpp b/Engine/Source/Runtime/Experimental/Chaos/Private/PBDRigidsSolver.cpp
--- a/Engine/Source/Runtime/Experimental/Chaos/Private/PBDRigidsSolver.cpp
+++ b/Engine/Source/Runtime/Experimental/Chaos/Private/PBDRigidsSolver.cpp
@@ -59,6 +59,8 @@ DECLARE_DWORD_COUNTER_STAT(TEXT("NumRestoredManifoldPoints"), STAT_ChaosCounter_
 DECLARE_DWORD_COUNTER_STAT(TEXT("NumUpdatedManifoldPoints"), STAT_ChaosCounter_NumUpdatedManifoldPoints, STATGROUP_ChaosCounters);
 DECLARE_DWORD_COUNTER_STAT(TEXT("NumJoints"), STAT_ChaosCounter_NumJoints, STATGROUP_ChaosCounters);
 DECLARE_DWORD_COUNTER_STAT(TEXT("NumCharacterGroundConstraints"), STAT_ChaosCounter_NumCharacterGroundConstraints, STATGROUP_ChaosCounters);
+DECLARE_DWORD_COUNTER_STAT(TEXT("NumBroadPhasePairs"), STAT_ChaosCounter_NumBroadPhasePairs, STATGROUP_ChaosCounters);
+DECLARE_DWORD_COUNTER_STAT(TEXT("NumMidPhases"), STAT_ChaosCounter_NumMidPhases, STATGROUP_ChaosCounters);
 
 TRACE_DECLARE_INT_COUNTER(ChaosTraceCounter_NumDisabledBodies, TEXT("Chaos/Solver/Bodies/NumDisabled"));
 TRACE_DECLARE_INT_COUNTER(ChaosTraceCounter_NumBodies, TEXT("Chaos/Solver/Bodies/Num"));
@@ -82,6 +84,8 @@ TRACE_DECLARE_INT_COUNTER(ChaosTraceCounter_NumRestoredManifoldPoints, TEXT("Cha
 TRACE_DECLARE_INT_COUNTER(ChaosTraceCounter_NumUpdatedManifoldPoints, TEXT("Chaos/Solver/Collisions/NumUpdatedManifoldPoints"));
 TRACE_DECLARE_INT_COUNTER(ChaosTraceCounter_NumJoints, TEXT("Chaos/Solver/Joints/NumConstraints"));
 TRACE_DECLARE_INT_COUNTER(ChaosTraceCounter_NumCharacterGroundConstraints, TEXT("Chaos/Solver/Character/NumConstraints"));
+TRACE_DECLARE_INT_COUNTER(ChaosTraceCounter_NumBroadPhasePairs, TEXT("Chaos/Solver/Collisions/NumBroadPhasePairs"));
+TRACE_DECLARE_INT_COUNTER(ChaosTraceCounter_NumMidPhases, TEXT("Chaos/Solver/Collisions/NumMidPhases"));
 
 TRACE_DECLARE_INT_COUNTER(ChaosTraceCounter_MidPhase_NumShapePair, TEXT("Chaos/Solver/MidPhase/NumShapePair"));
 TRACE_DECLARE_INT_COUNTER(ChaosTraceCounter_MidPhase_NumGeneric, TEXT("Chaos/Solver/MidPhase/NumGeneric"));
@@ -2030,6 +2034,10 @@ TRACE_COUNTER_SET(ChaosTraceCounter_##Name, Value)
 		UpdateExpensiveStatCounters();
 #endif
 
+		// Collision detection info
+		CHAOS_COUNTER_STAT(NumBroadPhasePairs, GetEvolution()->GetBroadPhase().GetNumBroadPhasePairs());
+		CHAOS_COUNTER_STAT(NumMidPhases, GetEvolution()->GetBroadPhase().GetNumMidPhases());
+
 		// Iterations
 		CHAOS_COUNTER_STAT(NumPositionIterations, GetEvolution()->GetNumPositionIterations());
 		CHAOS_COUNTER_STAT(NumVelocityIterations, GetEvolution()->GetNumVelocityIterations());
@@ -2099,21 +2107,22 @@ TRACE_COUNTER_SET(ChaosTraceCounter_##Name, Value)
 			const FConstGenericParticleHandle P = Particle.Handle();
 			if (Particle.Geometry().Get() != nullptr)
 			{
-				if (const FImplicitObjectUnion* Union = Particle.Geometry()->GetObject<FImplicitObjectUnion>())
+				int32 NumShapes = 1;
+				if (const FImplicitObjectUnion* Union = Particle.Geometry()->AsA<FImplicitObjectUnion>())
 				{
-					const int32 NumShapes = Union->GetNumLeafObjects();
-					if (P->IsDynamic())
-					{
-						NumDynamicShapes += NumShapes;
-					}
-					else if (P->IsKinematic())
-					{
-						NumKinematicShapes += NumShapes;
-					}
-					else
-					{
-						NumStaticShapes += NumShapes;
-					}
+					NumShapes = Union->GetNumLeafObjects();
+				}
+				if (P->IsDynamic())
+				{
+					NumDynamicShapes += NumShapes;
+				}
+				else if (P->IsKinematic())
+				{
+					NumKinematicShapes += NumShapes;
+				}
+				else
+				{
+					NumStaticShapes += NumShapes;
 				}
 			}
 		}
diff --git a/Engine/Source/Runtime/Experimental/Chaos/Public/Chaos/Collision/SpatialAccelerationBroadPhase.h b/Engine/Source/Runtime/Experimental/Chaos/Public/Chaos/Collision/SpatialAccelerationBroadPhase.h
--- a/Engine/Source/Runtime/Experimental/Chaos/Public/Chaos/Collision/SpatialAccelerationBroadPhase.h
+++ b/Engine/Source/Runtime/Experimental/Chaos/Public/Chaos/Collision/SpatialAccelerationBroadPhase.h
@@ -309,6 +309,10 @@ namespace Chaos
 
 			bNeedsResim = ResimCache && ResimCache->IsResimming();
 
+			// Reset stats
+			NumBroadPhasePairs = 0;
+			NumMidPhases = 0;
+
 			{
 				SCOPE_CYCLE_COUNTER(STAT_Collisions_SpatialBroadPhase);
 
@@ -356,6 +360,13 @@ namespace Chaos
 			// Run some error checks in non-shipping builds
 			// NOTE: This must come after MidPhase assignment for now because that's where the filter is applied
 			CheckOverlapResults();
+
+			// Update stats
+			for (int32 ContextIndex = 0; ContextIndex < NumActiveBroadphaseContexts; ++ContextIndex)
+			{
+				NumBroadPhasePairs += BroadphaseContexts[ContextIndex].Overlaps.Num();
+				NumMidPhases += BroadphaseContexts[ContextIndex].MidPhases.Num();
+			}
 		}
 
 		/**
@@ -487,6 +498,17 @@ namespace Chaos
 
 		FIgnoreCollisionManager& GetIgnoreCollisionManager() { return IgnoreCollisionManager; }
 
+		// Stats
+		int32 GetNumBroadPhasePairs() const
+		{
+			return NumBroadPhasePairs;
+		}
+
+		int32 GetNumMidPhases() const
+		{
+			return NumMidPhases;
+		}
+
 	private:
 
 		// Generate the set of particles that overlap the specified particle and are allowed to collide with it
@@ -596,16 +618,16 @@ namespace Chaos
 			
 			// Prefetch initial set of MidPhases
 			const int32 PrefetchLookahead = 4;
-			const int32 NumMidPhases = BroadphaseContext.MidPhases.Num();
-			for (int32 Index = 0; Index < NumMidPhases && Index < PrefetchLookahead; Index++)
+			const int32 NumContextMidPhases = BroadphaseContext.MidPhases.Num();
+			for (int32 Index = 0; Index < NumContextMidPhases && Index < PrefetchLookahead; Index++)
 			{
 				BroadphaseContext.MidPhases[Index]->CachePrefetch();
 			}
 
-			for (int32 Index = 0; Index < NumMidPhases; Index++)
+			for (int32 Index = 0; Index < NumContextMidPhases; Index++)
 			{
 				// Prefetch next MidPhase
-				if (Index + PrefetchLookahead < NumMidPhases)
+				if (Index + PrefetchLookahead < NumContextMidPhases)
 				{
 					BroadphaseContext.MidPhases[Index + PrefetchLookahead]->CachePrefetch();
 				}
@@ -632,7 +654,7 @@ namespace Chaos
 				{
 					if (Overlap.bCollisionsEnabled)
 					{
-						const FCollisionParticlePairKey PairKey = FCollisionParticlePairKey(Overlap.Particles[0], Overlap.Particles[1]);
+						const Private::FCollisionParticlePairKey PairKey = Private::FCollisionParticlePairKey(Overlap.Particles[0], Overlap.Particles[1]);
 						const bool bIsInSet = ParticlePairKeys.Contains(PairKey.GetKey());
 						if (ensure(!bIsInSet))
 						{
@@ -655,6 +677,9 @@ namespace Chaos
 		FIgnoreCollisionManager IgnoreCollisionManager;
 		int32 NumActiveBroadphaseContexts;
 		bool bNeedsResim;
+
+		int32 NumBroadPhasePairs;
+		int32 NumMidPhases;
 	};
 
 }
