Add GetVertex and GetEdge functions to TAABB. These each take an index. GetVertex returns a Vec3 corresponding to a vertex on the box. GetEdge returns a struct containing two vertex indices.

diff --git a/Engine/Source/Runtime/Experimental/Chaos/Public/Chaos/AABB.h b/Engine/Source/Runtime/Experimental/Chaos/Public/Chaos/AABB.h
--- a/Engine/Source/Runtime/Experimental/Chaos/Public/Chaos/AABB.h
+++ b/Engine/Source/Runtime/Experimental/Chaos/Public/Chaos/AABB.h
@@ -19,6 +19,12 @@ namespace Chaos
 		}
 	};
 
+	struct CHAOS_API FAABBEdge
+	{
+		int8 VertexIndex0;
+		int8 VertexIndex1;
+	};
+
 	template<class T, int d>
 	class TAABB
 	{
@@ -466,6 +472,69 @@ namespace Chaos
 		FORCEINLINE TVector<T, d> GetCenterOfMass() const { return GetCenter(); }
 		FORCEINLINE TVector<T, d> Extents() const { return MMax - MMin; }
 
+		/**
+		*	Get a vector one of the eight corners of the box.
+		*	
+		*	Each of the first three bits in an index are used to pick
+		*	an axis value from either the min or max vector of the AABB.
+		*	0 = min, 1 = max.
+		*	
+		*	This algorithm produces the following index scheme, where
+		*	the vertex at 0 is the "min" vertex and 7 is the "max":
+		*	
+		*	   6---------7
+		*	  /|        /|
+		*	 / |       / |
+		*	4---------5  |
+		*	|  |      |  |
+		*	|  2------|--3
+		*	| /       | /
+		*	|/        |/
+		*	0---------1
+		*/
+		FORCEINLINE TVector<T, d> GetVertex(const int32 Index) const
+		{
+			check(0 <= Index && Index < 8);
+			return TVector<T, d>(
+				(Index & (1 << 0)) == 0 ? MMin.X : MMax.X,
+				(Index & (1 << 1)) == 0 ? MMin.Y : MMax.Y,
+				(Index & (1 << 2)) == 0 ? MMin.Z : MMax.Z);
+		}
+
+		/**
+		*	Get an array of two indices into the vertex list for one of the
+		*	twelve edges of the box. Edges are ordered by increasing vertex
+		*	indices.
+		*	
+		*	This algorithm produces the following index scheme, where the
+		*	bottom left vertex is the "min" vertex and the top right is "max":
+		*	
+		*	        *-----11------*
+		*	       /|            /|
+		*	      9 |          10 |
+		*	     /  6          /  7
+		*	    *------8------*   |
+		*	    |   |         |   |
+		*	    |   *------5--|---*
+		*	    2  /          4  /
+		*	    | 1           | 3
+		*	    |/            |/
+		*	    *------0------*
+		*/
+		FORCEINLINE FAABBEdge GetEdge(const int32 Index) const
+		{
+			// See "GetVertex(int32)"
+			check(0 <= Index && Index < 12);
+			static constexpr FAABBEdge Edges[]
+			{
+				{ 0, 1 }, { 0, 2 }, { 0, 4 },
+				{ 1, 3 }, { 1, 5 }, { 2, 3 },
+				{ 2, 6 }, { 3, 7 }, { 4, 5 },
+				{ 4, 6 }, { 5, 7 }, { 6, 7 }
+			};
+			return Edges[Index];
+		}
+
 		FORCEINLINE int LargestAxis() const
 		{
 			const auto Extents = this->Extents();
