Chaos - fix collision disabled flag

diff --git a/Engine/Source/Runtime/Experimental/Chaos/Private/Chaos/Collision/CollisionConstraintAllocator.cpp b/Engine/Source/Runtime/Experimental/Chaos/Private/Chaos/Collision/CollisionConstraintAllocator.cpp
--- a/Engine/Source/Runtime/Experimental/Chaos/Private/Chaos/Collision/CollisionConstraintAllocator.cpp
+++ b/Engine/Source/Runtime/Experimental/Chaos/Private/Chaos/Collision/CollisionConstraintAllocator.cpp
@@ -159,6 +159,20 @@ namespace Chaos
 			ProcessNewItems();
 		}
 
+		void FCollisionConstraintAllocator::ResetActiveConstraints()
+		{
+			// Disable existing constraints so that if they are not re-activated this tick
+			// they do not have state indicating that that are still active.
+			// @todo(chaos): ideally we would do this only for constraints that do not get reused this tick in EndDetectCollisions
+			for (FPBDCollisionConstraint* Constraint : ActiveConstraints)
+			{
+				if (Constraint != nullptr)
+				{
+					Constraint->SetDisabled(true);
+				}
+			}
+		}
+
 		void FCollisionConstraintAllocator::PruneExpiredMidPhases()
 		{
 			// NOTE: Called from the physics thread, and never from a physics task/parallel-for. No need for locks.
diff --git a/Engine/Source/Runtime/Experimental/Chaos/Private/Chaos/Collision/ParticlePairMidPhase.cpp b/Engine/Source/Runtime/Experimental/Chaos/Private/Chaos/Collision/ParticlePairMidPhase.cpp
--- a/Engine/Source/Runtime/Experimental/Chaos/Private/Chaos/Collision/ParticlePairMidPhase.cpp
+++ b/Engine/Source/Runtime/Experimental/Chaos/Private/Chaos/Collision/ParticlePairMidPhase.cpp
@@ -507,6 +507,7 @@ namespace Chaos
 			// Flag this contact as requiring CCD
 			Constraint->SetCCDEnabled(true);
 		}
+		check(!Constraint->IsEnabled());
 
 		// Do we want to enable the CCD sweep? If not, we fall back to the standard collision detection for this tick
 		Constraint->SetCCDSweepEnabled(bEnableCCDSweep);
@@ -595,6 +596,7 @@ namespace Chaos
 		{
 			CreateConstraint(CullDistance, Context);
 		}
+		check(!Constraint->IsEnabled());
 
 		if (Constraint.IsValid())
 		{
@@ -639,6 +641,7 @@ namespace Chaos
 			Constraint->GetContainerCookie().bIsMultiShapePair = false;
 			Constraint->GetContainerCookie().CreationEpoch = CurrentEpoch;
 		}
+		check(!Constraint->IsEnabled());
 
 		// Copy the constraint data over the existing one (ensure we do not replace data required by the graph and the allocator/container)
 		Constraint->RestoreFrom(SourceConstraint);
@@ -1553,6 +1556,7 @@ namespace Chaos
 				CollisionKey, CollisionSortKey,
 				CullDistance, ShapePairType, bUseManifold, Context);
 		}
+		check(!Constraint->IsEnabled());
 
 		NewConstraints.Add(Constraint);
 
diff --git a/Engine/Source/Runtime/Experimental/Chaos/Private/Chaos/ContactModification.cpp b/Engine/Source/Runtime/Experimental/Chaos/Private/Chaos/ContactModification.cpp
--- a/Engine/Source/Runtime/Experimental/Chaos/Private/Chaos/ContactModification.cpp
+++ b/Engine/Source/Runtime/Experimental/Chaos/Private/Chaos/ContactModification.cpp
@@ -542,13 +542,21 @@ namespace Chaos
 				Constraints.Add(&Constraint);
 			}
 			return ECollisionVisitorResult::Continue;
-		}, ECollisionVisitorFlags::VisitActiveAwake | ECollisionVisitorFlags::VisitDisabled);
+		}, ECollisionVisitorFlags::VisitActiveAwake);
 	}
 
 	bool FVisitedContactPairsTracker::Visit(const FContactPairModifier& ContactPair)
 	{
-		// Get a reference object to the relevant bit
+		// We use the index in the ActiveConstraints list as an ID for this constraint. This requires
+		// that we are only visiting the active constraints (and not sleeping or non-activated constraints)
+		// in the VisitCollisions call in FContactPairModifierParticleRange
 		const int32 ConstraintIndex = ContactPair.GetConstraintContainerCookie().ConstraintIndex;
+		
+		// If we hit this, the constraint is not in the active list and the index is invalid
+		// Likely an error in the VisitCollisions call above
+		check(Constraints[ConstraintIndex] == ContactPair.Constraint);
+
+		// Get a reference object to the relevant bit
 		FBitReference BitRef = VisitedContacts[ConstraintIndex];
 
 		// Mark constraint visited and return true
@@ -569,7 +577,7 @@ namespace Chaos
 
 	FVisitedContactPairsTracker FCollisionContactModifier::MakeVisitedContactPairsTracker() const
 	{
-		return FVisitedContactPairsTracker(Constraints.Num());
+		return FVisitedContactPairsTracker(Constraints);
 	}
 
 	TArrayView<FPBDCollisionConstraint* const>& FCollisionContactModifier::GetConstraints()
diff --git a/Engine/Source/Runtime/Experimental/Chaos/Private/Chaos/PBDCollisionConstraints.cpp b/Engine/Source/Runtime/Experimental/Chaos/Private/Chaos/PBDCollisionConstraints.cpp
--- a/Engine/Source/Runtime/Experimental/Chaos/Private/Chaos/PBDCollisionConstraints.cpp
+++ b/Engine/Source/Runtime/Experimental/Chaos/Private/Chaos/PBDCollisionConstraints.cpp
@@ -512,17 +512,6 @@ namespace Chaos
 	{
 		SCOPE_CYCLE_COUNTER(STAT_Collisions_BeginDetect);
 
-		// Disable existing constraints so that if they are not re-activated this tick
-		// they do not have state indicating that that are still active.
-		// @todo(chaos): ideally we would do this only for constraints that do not get reused this tick in EndDetectCollisions
-		for (FPBDCollisionConstraint* Constraint : GetConstraints())
-		{
-			if (Constraint != nullptr)
-			{
-				Constraint->SetDisabled(true);
-			}
-		}
-
 		ConstraintAllocator.BeginDetectCollisions();
 	}
 
diff --git a/Engine/Source/Runtime/Experimental/Chaos/Public/Chaos/Collision/CollisionConstraintAllocator.h b/Engine/Source/Runtime/Experimental/Chaos/Public/Chaos/Collision/CollisionConstraintAllocator.h
--- a/Engine/Source/Runtime/Experimental/Chaos/Public/Chaos/Collision/CollisionConstraintAllocator.h
+++ b/Engine/Source/Runtime/Experimental/Chaos/Public/Chaos/Collision/CollisionConstraintAllocator.h
@@ -346,6 +346,8 @@ namespace Chaos
 			*/
 			void BeginFrame()
 			{
+				ResetActiveConstraints();
+
 				ActiveConstraints.Reset();
 				ActiveCCDConstraints.Reset();
 			}
@@ -372,6 +374,8 @@ namespace Chaos
 					ContextAllocator->BeginDetectCollisions(CurrentEpoch);
 				}
 
+				ResetActiveConstraints();
+
 				// Clear the collision list for this tick - we are about to rebuild them
 				ActiveConstraints.Reset();
 				ActiveCCDConstraints.Reset();
@@ -549,6 +553,8 @@ namespace Chaos
 
 			CHAOS_API void SortActiveConstraints();
 
+			CHAOS_API void ResetActiveConstraints();
+
 			// The container that owns the allocator (only needed because new constraints need to know)
 			FPBDCollisionConstraints* CollisionContainer;
 
diff --git a/Engine/Source/Runtime/Experimental/Chaos/Public/Chaos/ContactModification.h b/Engine/Source/Runtime/Experimental/Chaos/Public/Chaos/ContactModification.h
--- a/Engine/Source/Runtime/Experimental/Chaos/Public/Chaos/ContactModification.h
+++ b/Engine/Source/Runtime/Experimental/Chaos/Public/Chaos/ContactModification.h
@@ -447,10 +447,12 @@ namespace Chaos
 		CHAOS_API bool Visit(const FContactPairModifier& ContactPair);
 
 	private:
-		FVisitedContactPairsTracker(const int32 NumConstraints)
-			: VisitedContacts(false, NumConstraints)
+		FVisitedContactPairsTracker(const TArrayView<FPBDCollisionConstraint* const>& InConstraints)
+			: Constraints(InConstraints)
+			, VisitedContacts(false, Constraints.Num())
 		{ }
 
+		TArrayView<FPBDCollisionConstraint* const> Constraints;
 		TBitArray<> VisitedContacts;
 
 		// Befriend the modifier so that it can create these
@@ -468,7 +470,7 @@ namespace Chaos
 		friend FPBDCollisionConstraints; // Calls UpdateConstraintManifolds after callback.
 
 
-		FCollisionContactModifier(TArrayView<FPBDCollisionConstraint* const>& InConstraints, FReal InDt)
+		FCollisionContactModifier(const TArrayView<FPBDCollisionConstraint* const>& InConstraints, FReal InDt)
 			: Constraints(InConstraints)
 			, Dt(InDt)
 		{}
@@ -497,7 +499,7 @@ namespace Chaos
 		// Update manifolds of modified constraints.
 		CHAOS_API void UpdateConstraintManifolds();
 
-		TArrayView<FPBDCollisionConstraint* const>& Constraints;
+		TArrayView<FPBDCollisionConstraint* const> Constraints;
 
 		// Constraints that should update manifold from contact points.
 		TSet<FPBDCollisionConstraint*> NeedsManifoldUpdate;
