geometry collection physics proxy should sync the particle transform into the dynamic collection

diff --git a/Engine/Source/Runtime/Experimental/Chaos/Private/PhysicsProxy/GeometryCollectionPhysicsProxy.cpp b/Engine/Source/Runtime/Experimental/Chaos/Private/PhysicsProxy/GeometryCollectionPhysicsProxy.cpp
--- a/Engine/Source/Runtime/Experimental/Chaos/Private/PhysicsProxy/GeometryCollectionPhysicsProxy.cpp
+++ b/Engine/Source/Runtime/Experimental/Chaos/Private/PhysicsProxy/GeometryCollectionPhysicsProxy.cpp
@@ -3185,6 +3185,7 @@ void FGeometryCollectionPhysicsProxy::BufferPhysicsResults_Internal(Chaos::FPBDR
 				Results.SetPositions(EntryIndex, PositionData);
 				Results.SetVelocities(EntryIndex, VelocityData);
 				Results.SetTransform(EntryIndex, ParentSpaceTransform);
+				PhysicsThreadCollection.Transform[TransformGroupIndex] = ParentSpaceTransform;
 				IsObjectDynamic = true;
 			}
 
