add cvar (p.Chaos.OnlyUseInterclusterEdgesAttachedToMainParticles, default false) to only allow particles directly connected to a cluster union via an intercluster edge to be added/kept in a cluster union

diff --git a/Engine/Source/Runtime/Experimental/Chaos/Private/Chaos/ClusterUnionManager.cpp b/Engine/Source/Runtime/Experimental/Chaos/Private/Chaos/ClusterUnionManager.cpp
--- a/Engine/Source/Runtime/Experimental/Chaos/Private/Chaos/ClusterUnionManager.cpp
+++ b/Engine/Source/Runtime/Experimental/Chaos/Private/Chaos/ClusterUnionManager.cpp
@@ -926,6 +926,36 @@ namespace Chaos
 		RequestDeferredClusterPropertiesUpdate(ClusterIndex, EUpdateClusterUnionPropertiesFlags::IncrementalGenerateConnectionGraph);
 	}
 
+	bool FClusterUnionManager::IsDirectlyConnectedToMainParticleInClusterUnion(const FClusterUnion& ClusterUnion, FPBDRigidParticleHandle* Particle) const
+	{
+		bool bIsConnectedToMainParticle = false;
+		if (FPBDRigidClusteredParticleHandle* ClusterParticle = Particle->CastToClustered())
+		{
+			for (const FConnectivityEdge& Edge : ClusterParticle->ConnectivityEdges())
+			{
+				if (IsInterclusterEdge(*ClusterParticle, Edge))
+				{
+					if (const FClusterUnionParticleProperties* SiblingProps = ClusterUnion.ChildProperties.Find(Edge.Sibling))
+					{
+						if (!SiblingProps->bIsAuxiliaryParticle)
+						{
+							bIsConnectedToMainParticle = true;
+							break;
+						}
+					}
+					else
+					{
+						// No props = main particles.
+						bIsConnectedToMainParticle = true;
+						break;
+					}
+				}
+			}
+		}
+
+		return bIsConnectedToMainParticle;
+	}
+
 	void FClusterUnionManager::RequestDeferredClusterPropertiesUpdate(FClusterUnionIndex ClusterIndex, EUpdateClusterUnionPropertiesFlags Flags)
 	{
 		EUpdateClusterUnionPropertiesFlags& ExistingFlags = DeferredClusterUnionsForUpdateProperties.FindOrAdd(ClusterIndex);
diff --git a/Engine/Source/Runtime/Experimental/Chaos/Private/Chaos/PBDRigidClustering.cpp b/Engine/Source/Runtime/Experimental/Chaos/Private/Chaos/PBDRigidClustering.cpp
--- a/Engine/Source/Runtime/Experimental/Chaos/Private/Chaos/PBDRigidClustering.cpp
+++ b/Engine/Source/Runtime/Experimental/Chaos/Private/Chaos/PBDRigidClustering.cpp
@@ -39,6 +39,9 @@ namespace Chaos
 	bool bCheckForInterclusterEdgesOnRelease = true;
 	FAutoConsoleVariableRef CVarCheckForInterclusterEdgesOnRelease(TEXT("p.Chaos.CheckForInterclusterEdgesOnRelease"), bCheckForInterclusterEdgesOnRelease, TEXT("Whether to check for intercluster edges when removing a child from its parent cluster so that we can add the particle back into a cluster union."));
 
+	bool bOnlyUseInterclusterEdgesAttachedToMainParticles = false;
+	FAutoConsoleVariableRef CVarOnlyUseInterclusterEdgesAttachedToMainParticles(TEXT("p.Chaos.OnlyUseInterclusterEdgesAttachedToMainParticles"), bOnlyUseInterclusterEdgesAttachedToMainParticles, TEXT("If true, an intercluster edge must be directly attached to a main particle for the particle to remain a part of the cluster union."));
+
 	int32 ComputeClusterCollisionStrains = 1;
 	FAutoConsoleVariableRef CVarComputeClusterCollisionStrains(TEXT("p.ComputeClusterCollisionStrains"), ComputeClusterCollisionStrains, TEXT("Whether to use collision constraints when processing clustering."));
 
@@ -1168,14 +1171,47 @@ namespace Chaos
 
 					if (AttachedClusterUnion)
 					{
-						// Need to manually add cluster union properties since the API for adding a pending operation doesn't have an option for that.
+						// We may not actually want to attach the *entire* island into the cluster union. 
+						// We may only want to attach a subset into the cluster union. In that case, we just release
+						// each of the remaining particles separately (TODO: maybe figure out a way to create internal clusters for them).
+						TArray<FPBDRigidParticleHandle*> ParticlesForClusterUnion;
+						TArray<FPBDRigidParticleHandle*> ParticlesToRelease;
+
+						if (bOnlyUseInterclusterEdgesAttachedToMainParticles)
+						{
+							ParticlesForClusterUnion.Reserve(Island.Num());
+							ParticlesToRelease.Reserve(Island.Num());
+						}
+
 						for (FPBDRigidParticleHandle* ChildParticle : Island)
 						{
+							if (bOnlyUseInterclusterEdgesAttachedToMainParticles)
+							{
+								if (ClusterUnionManager.IsDirectlyConnectedToMainParticleInClusterUnion(*AttachedClusterUnion, ChildParticle))
+								{
+									ParticlesForClusterUnion.Add(ChildParticle);
+								}
+								else
+								{
+									ParticlesToRelease.Add(ChildParticle);
+								}
+							}
+
+							// Need to manually add cluster union properties since the API for adding a pending operation doesn't have an option for that.
 							FClusterUnionParticleProperties Properties;
 							Properties.bIsAuxiliaryParticle = true;
 							AttachedClusterUnion->ChildProperties.Add(ChildParticle, Properties);
 						}
-						ClusterUnionManager.AddPendingClusterIndexOperation(AttachedClusterUnion->InternalIndex, EClusterUnionOperation::Add, Island);
+						ClusterUnionManager.AddPendingClusterIndexOperation(AttachedClusterUnion->InternalIndex, EClusterUnionOperation::Add, bOnlyUseInterclusterEdgesAttachedToMainParticles ? ParticlesForClusterUnion : Island);
+
+						if (!ParticlesToRelease.IsEmpty())
+						{
+							for (FPBDRigidParticleHandle* ChildParticle : ParticlesToRelease)
+							{
+								RemoveChildFromParent(ChildParticle, ClusteredParticle);
+							}
+							ActivatedChildren.Append(ParticlesToRelease);
+						}
 						IslandIndicesToRemove.Add(IslandIndex);
 					}
 					else if (Island.Num() == 1 && !AttachedClusterUnion) //need to break single pieces first
@@ -1206,19 +1242,34 @@ namespace Chaos
 				for (const FParticleIsland& Island : Islands)
 				{
 					bool bHasMainParticle = false;
+					
+					TArray<FPBDRigidParticleHandle*> ParticlesToRemove;
 					for (FPBDRigidParticleHandle* ChildParticle : Island)
 					{
-						if (const FClusterUnionParticleProperties* Props = ParentClusterUnion->ChildProperties.Find(ChildParticle))
+						if (!bHasMainParticle)
 						{
-							if (!Props->bIsAuxiliaryParticle)
+							if (const FClusterUnionParticleProperties* Props = ParentClusterUnion->ChildProperties.Find(ChildParticle))
+							{
+								if (!Props->bIsAuxiliaryParticle)
+								{
+									bHasMainParticle = true;
+								}
+							}
+							else
 							{
 								bHasMainParticle = true;
-								break;
+							}
+						}
+						
+						if (bOnlyUseInterclusterEdgesAttachedToMainParticles)
+						{
+							if (!ClusterUnionManager.IsDirectlyConnectedToMainParticleInClusterUnion(*ParentClusterUnion, ChildParticle))
+							{
+								ParticlesToRemove.Add(ChildParticle);
 							}
 						}
 						else
 						{
-							bHasMainParticle = true;
 							break;
 						}
 					}
@@ -1228,6 +1279,11 @@ namespace Chaos
 						ClusterUnionManager.AddPendingClusterIndexOperation(ParentClusterUnion->InternalIndex, EClusterUnionOperation::Remove, Island);
 						ActivatedChildren.Append(Island);
 					}
+					else if (!ParticlesToRemove.IsEmpty())
+					{
+						ClusterUnionManager.AddPendingClusterIndexOperation(ParentClusterUnion->InternalIndex, EClusterUnionOperation::Remove, ParticlesToRemove);
+						ActivatedChildren.Append(ParticlesToRemove);
+					}
 				}
 			}
 		}
diff --git a/Engine/Source/Runtime/Experimental/Chaos/Public/Chaos/ClusterUnionManager.h b/Engine/Source/Runtime/Experimental/Chaos/Public/Chaos/ClusterUnionManager.h
--- a/Engine/Source/Runtime/Experimental/Chaos/Public/Chaos/ClusterUnionManager.h
+++ b/Engine/Source/Runtime/Experimental/Chaos/Public/Chaos/ClusterUnionManager.h
@@ -181,6 +181,8 @@ namespace Chaos
 		// Returns all cluster unions. Really meant only to be used for debugging.
 		const TMap<FClusterUnionIndex, FClusterUnion>& GetAllClusterUnions() const { return ClusterUnions; }
 
+		bool IsDirectlyConnectedToMainParticleInClusterUnion(const FClusterUnion& ClusterUnion, FPBDRigidParticleHandle* Particle) const;
+
 		// Put in a deferred request to update the properties of the specified cluster union. Flags can limit the scope of the update.
 		CHAOS_API void RequestDeferredClusterPropertiesUpdate(FClusterUnionIndex ClusterIndex, EUpdateClusterUnionPropertiesFlags Flags);
 	private:
