IAS: Fix inefficiency in OnDemandBAckend thread loop by always trying to maximize the number of concurrent requests.

#rnx

- Submitted on the behalf of PJ.Kack

diff --git a/Engine/Source/Runtime/Experimental/IoStoreOnDemand/Private/OnDemandIoDispatcherBackend.cpp b/Engine/Source/Runtime/Experimental/IoStoreOnDemand/Private/OnDemandIoDispatcherBackend.cpp
--- a/Engine/Source/Runtime/Experimental/IoStoreOnDemand/Private/OnDemandIoDispatcherBackend.cpp
+++ b/Engine/Source/Runtime/Experimental/IoStoreOnDemand/Private/OnDemandIoDispatcherBackend.cpp
@@ -304,10 +304,14 @@ class FHttpClient
 {
 public:
 	FHttpClient(const FString& ServiceUrl, int32 MaxConnectionCount = 8);
+	~FHttpClient() = default;
+
 	const FString& ServiceUrl() const { return SvcsUrl; }
 	int32 MaxConnectionCount() const { return MaxConnections;}
 	void Get(FAnsiStringView Url, FIoReadCallback&& Callback);
 	void Get(FAnsiStringView Url, const FIoOffsetAndLength& Range, FIoReadCallback&& Callback);
+
+	/** @return True if the client has pending work otherwise false. */
 	bool Tick();
 
 private:
@@ -1316,74 +1320,85 @@ FIoStatus FOnDemandIoBackend::AddToc(const FOnDemandEndpoint& Endpoint)
 uint32 FOnDemandIoBackend::Run()
 {
 	const int32 MaxConcurrentRequests = HttpClient->MaxConnectionCount();
+	FChunkRequest* NextChunkRequest = nullptr;
+	int32 NumConcurrentRequests = 0;
 
 	while (!bStopRequested)
 	{
-		for (;;)
+		NextChunkRequest = HttpRequests.Dequeue();
+		while (NextChunkRequest)
 		{
-			FChunkRequest* NextChunkRequest = HttpRequests.Dequeue();
-			if (NextChunkRequest == nullptr)
-			{
-				break;
-			}
-
-			int32 NumConcurrentRequests = 0;
 			while (NextChunkRequest)
 			{
-				FChunkRequest* ChunkRequest = NextChunkRequest;
-				NextChunkRequest = ChunkRequest->NextRequest;
-				ChunkRequest->NextRequest = nullptr;
+				{
+					TRACE_CPUPROFILER_EVENT_SCOPE(FOnDemandIoBackend::IssueHttpGet);
+					FChunkRequest* ChunkRequest = NextChunkRequest;
+					NextChunkRequest = ChunkRequest->NextRequest;
+					ChunkRequest->NextRequest = nullptr;
 
-				Stats.OnHttpDequeue();
+					Stats.OnHttpDequeue();
 
-				TAnsiStringBuilder<256> Url;
-				ChunkRequest->Params.GetUrl(Url);
+					TAnsiStringBuilder<256> Url;
+					ChunkRequest->Params.GetUrl(Url);
 
-				NumConcurrentRequests++;
-				HttpClient->Get(Url.ToView(), ChunkRequest->Params.ChunkRange,
-					[this, ChunkRequest, &NumConcurrentRequests](TIoStatusOr<FIoBuffer> Status)
-					{
-						NumConcurrentRequests--;
-
-						if (Status.Status().GetErrorCode() == EIoErrorCode::ReadError)
+					NumConcurrentRequests++;
+					HttpClient->Get(Url.ToView(), ChunkRequest->Params.ChunkRange,
+						[this, ChunkRequest, &NumConcurrentRequests](TIoStatusOr<FIoBuffer> Status)
 						{
-							if (++ChunkRequest->HttpRetryCount <= GIoDispatcherMaxHttpRetryCount)
+							NumConcurrentRequests--;
+
+							if (Status.Status().GetErrorCode() == EIoErrorCode::ReadError)
 							{
-								Stats.OnHttpRetry();
-								Stats.OnHttpEnqueue();
-								return HttpRequests.Enqueue(ChunkRequest); 
+								if (++ChunkRequest->HttpRetryCount <= GIoDispatcherMaxHttpRetryCount)
+								{
+									Stats.OnHttpRetry();
+									Stats.OnHttpEnqueue();
+									return HttpRequests.Enqueue(ChunkRequest); 
+								}
 							}
-						}
 
-						if (Status.IsOk())
-						{
-							ChunkRequest->Chunk = Status.ConsumeValueOrDie();
-							Stats.OnHttpGet(ChunkRequest->Chunk.DataSize());
-						}
-						else
-						{
-							TAnsiStringBuilder<256> Url;
-							ChunkRequest->Params.GetUrl(Url);
-							LogHttpResult(StringCast<TCHAR>(*Url).Get(), -1, 0, 0, ChunkRequest->Params.ChunkRange.GetOffset(), "HTTP FAILED");
-							Stats.OnHttpError();
-						}
+							if (Status.IsOk())
+							{
+								ChunkRequest->Chunk = Status.ConsumeValueOrDie();
+								Stats.OnHttpGet(ChunkRequest->Chunk.DataSize());
+							}
+							else
+							{
+								TAnsiStringBuilder<256> Url;
+								ChunkRequest->Params.GetUrl(Url);
+								LogHttpResult(StringCast<TCHAR>(*Url).Get(), -1, 0, 0, ChunkRequest->Params.ChunkRange.GetOffset(), "HTTP FAILED");
+								Stats.OnHttpError();
+							}
 
-						Launch(UE_SOURCE_LOCATION, [this, ChunkRequest]()
-						{
-							CompleteRequest(ChunkRequest);
+							Launch(UE_SOURCE_LOCATION, [this, ChunkRequest]()
+							{
+								CompleteRequest(ChunkRequest);
+							});
 						});
-					});
+				}
 
-				while (NumConcurrentRequests >= MaxConcurrentRequests)
+				if (NumConcurrentRequests >= MaxConcurrentRequests)
 				{
-					TRACE_CPUPROFILER_EVENT_SCOPE(FOnDemandIoBackend::TickHttpEventLoop);
-					HttpClient->Tick();
+					TRACE_CPUPROFILER_EVENT_SCOPE(FOnDemandIoBackend::TickHttpSaturated);
+					while (NumConcurrentRequests >= MaxConcurrentRequests)
+					{
+						HttpClient->Tick();
+					}
+				}
+
+				if (!NextChunkRequest)
+				{
+					NextChunkRequest = HttpRequests.Dequeue();
 				}
 			}
 
 			{
-				TRACE_CPUPROFILER_EVENT_SCOPE(FOnDemandIoBackend::TickHttpEventLoop);
-				while (HttpClient->Tick());
+				// Keep processing pending connections until all work is complete or a new request if found
+				TRACE_CPUPROFILER_EVENT_SCOPE(FOnDemandIoBackend::TickHttp);
+				while (HttpClient->Tick() && !NextChunkRequest)
+				{
+					NextChunkRequest = HttpRequests.Dequeue();
+				}
 			}
 		}
 
