Removing spurious transform stuff, adding a reverse direction setup for the wheels

diff --git a/Engine/Source/Runtime/Experimental/ChaosVehicles/ChaosVehiclesCore/Private/SimModule/DeferredForcesModular.cpp b/Engine/Source/Runtime/Experimental/ChaosVehicles/ChaosVehiclesCore/Private/SimModule/DeferredForcesModular.cpp
--- a/Engine/Source/Runtime/Experimental/ChaosVehicles/ChaosVehiclesCore/Private/SimModule/DeferredForcesModular.cpp
+++ b/Engine/Source/Runtime/Experimental/ChaosVehicles/ChaosVehiclesCore/Private/SimModule/DeferredForcesModular.cpp
@@ -85,8 +85,7 @@ Chaos::FPBDRigidParticleHandle* FDeferredForcesModular::GetParticle(const FTrans
 	}
 
 	// local transform
-	FVector Offset = OffsetTransform.InverseTransformVector(Frame.GetTranslation());
-	TransformOut.SetTranslation(Offset);
+	TransformOut.SetTranslation(Frame.GetTranslation());
 	return ClusterParticles[SingleChassisIndex];
 }
 
@@ -134,8 +133,7 @@ void FDeferredForcesModular::Apply(TArray<Chaos::FPBDRigidClusteredParticleHandl
 
 
 
-void FDeferredForcesModular::Apply(const FTransform& OffsetTransform
-	, TArray<Chaos::FPBDRigidParticleHandle*>& Particles
+void FDeferredForcesModular::Apply(TArray<Chaos::FPBDRigidParticleHandle*>& Particles
 	, TArray<Chaos::FPBDRigidClusteredParticleHandle*>& ClusterParticles)
 {
 	if (!GCoreModularVehicleDebugParams.DisableForces)
@@ -143,7 +141,7 @@ void FDeferredForcesModular::Apply(const FTransform& OffsetTransform
 		FTransform RelativeTransform;
 		for (const FApplyForceData& Data : ApplyForceDatas)
 		{
-			Chaos::FPBDRigidParticleHandle* RigidHandle = GetParticle(OffsetTransform, Particles, ClusterParticles, Data.TransformIndex, RelativeTransform);
+			Chaos::FPBDRigidParticleHandle* RigidHandle = GetParticle(Data.OffsetTransform, Particles, ClusterParticles, Data.TransformIndex, RelativeTransform);
 			if (RigidHandle)
 			{
 				AddForce(RigidHandle, Data, RelativeTransform);
@@ -152,7 +150,7 @@ void FDeferredForcesModular::Apply(const FTransform& OffsetTransform
 
 		for (const FApplyForceAtPositionData& Data : ApplyForceAtPositionDatas)
 		{
-			Chaos::FPBDRigidParticleHandle* RigidHandle = GetParticle(OffsetTransform, Particles, ClusterParticles, Data.TransformIndex, RelativeTransform);
+			Chaos::FPBDRigidParticleHandle* RigidHandle = GetParticle(Data.OffsetTransform, Particles, ClusterParticles, Data.TransformIndex, RelativeTransform);
 			if (RigidHandle)
 			{
 				AddForceAtPosition(RigidHandle, Data, RelativeTransform);
@@ -161,7 +159,7 @@ void FDeferredForcesModular::Apply(const FTransform& OffsetTransform
 
 		for (const FAddTorqueInRadiansData& Data : ApplyTorqueDatas)
 		{
-			Chaos::FPBDRigidParticleHandle* RigidHandle = GetParticle(OffsetTransform, Particles, ClusterParticles, Data.TransformIndex, RelativeTransform);
+			Chaos::FPBDRigidParticleHandle* RigidHandle = GetParticle(Data.OffsetTransform, Particles, ClusterParticles, Data.TransformIndex, RelativeTransform);
 			if (RigidHandle)
 			{
 				AddTorque(RigidHandle, Data, RelativeTransform);
@@ -182,7 +180,7 @@ void AddForce_Implementation(Chaos::FPBDRigidParticleHandle* RigidHandle, const
 
 	// local to world
 	const FTransform WorldTM(RigidHandle->R(), RigidHandle->X());
-	Chaos::FVec3 Position = WorldTM.TransformPosition(ParticleOffsetTransform.TransformVector(OffsetTransform.GetLocation()));
+	Chaos::FVec3 Position = WorldTM.TransformPosition(OffsetTransform.GetLocation());
 	Chaos::FVec3 Force = WorldTM.TransformVector(ParticleOffsetTransform.TransformVector(OffsetTransform.TransformVector(LocalForce)));
 
 
@@ -229,7 +227,7 @@ void FDeferredForcesModular::AddForce(Chaos::FPBDRigidParticleHandle* RigidHandl
 {
 	if (ensure(RigidHandle))
 	{
-		AddForce_Implementation(RigidHandle, ParticleOffsetTransform, DataIn.Force, OffsetTransform, (DataIn.Flags & EForceFlags::LevelSlope) == EForceFlags::LevelSlope, DataIn.DebugColor);
+		AddForce_Implementation(RigidHandle, DataIn.OffsetTransform, DataIn.Force, OffsetTransform, (DataIn.Flags & EForceFlags::LevelSlope) == EForceFlags::LevelSlope, DataIn.DebugColor);
 	}
 
 }
diff --git a/Engine/Source/Runtime/Experimental/ChaosVehicles/ChaosVehiclesCore/Private/SimModule/SimModuleTree.cpp b/Engine/Source/Runtime/Experimental/ChaosVehicles/ChaosVehiclesCore/Private/SimModule/SimModuleTree.cpp
--- a/Engine/Source/Runtime/Experimental/ChaosVehicles/ChaosVehiclesCore/Private/SimModule/SimModuleTree.cpp
+++ b/Engine/Source/Runtime/Experimental/ChaosVehicles/ChaosVehiclesCore/Private/SimModule/SimModuleTree.cpp
@@ -16,13 +16,11 @@ namespace Chaos
 
 int FSimModuleTree::AddRoot(ISimulationModuleBase* SimModule)
 {
-	check(!bIsSimulating);
 	return AddNodeBelow(-1, SimModule);
 }
 
 void FSimModuleTree::Reparent(int AtIndex, int ParentIndex)
 {
-	check(!bIsSimulating);
 	check(AtIndex < SimulationModuleTree.Num());
 	check(ParentIndex < SimulationModuleTree.Num());
 
@@ -46,7 +44,6 @@ void FSimModuleTree::Reparent(int AtIndex, int ParentIndex)
 
 int FSimModuleTree::AddNodeBelow(int AtIndex, ISimulationModuleBase* SimModule)
 {
-	check(!bIsSimulating);
 	int NewIndex = GetNextIndex();
 	FSimModuleNode& Node = SimulationModuleTree[NewIndex];
 	SimModule->SetTreeIndex(NewIndex);
@@ -66,8 +63,6 @@ int FSimModuleTree::AddNodeBelow(int AtIndex, ISimulationModuleBase* SimModule)
 
 void FSimModuleTree::AppendTreeUpdates(const FSimTreeUpdates& TreeUpdates)
 {
-	ensure(!IsSimulating());
-
 	int TreeIndex = -1;
 	TMap<int, int> SimTreeMapping;
 
@@ -89,8 +84,6 @@ void FSimModuleTree::AppendTreeUpdates(const FSimTreeUpdates& TreeUpdates)
 			AddIndex = *AddIndexPtr;
 		}
 
-		ensure(!IsSimulating());
-
 		TreeIndex = AddNodeBelow(AddIndex, TreeUpdate.NewSimModule);
 		SimTreeMapping.Add(LocalIndex, TreeIndex);
 		LocalIndex++;
@@ -182,7 +175,6 @@ int FSimModuleTree::GetNextIndex()
 
 int FSimModuleTree::InsertNodeAbove(int AtIndex, ISimulationModuleBase* SimModule)
 {
-	check(!bIsSimulating);
 	int NewIndex = FSimModuleNode::INVALID_IDX;
 
 	if (ensure(AtIndex < SimulationModuleTree.Num()))
@@ -209,7 +201,6 @@ int FSimModuleTree::InsertNodeAbove(int AtIndex, ISimulationModuleBase* SimModul
 
 void FSimModuleTree::DeleteNode(int AtIndex)
 {
-	check(!bIsSimulating);
 	// multiple children might become equal parents?
 	
 	int ParentIndex = SimulationModuleTree[AtIndex].Parent;
@@ -399,9 +390,9 @@ void FSimModuleTree::UpdateModuleVelocites(FClusterUnionPhysicsProxy* PhysicsPro
 
 					if (Particle)
 					{
-						const FTransform& OffsetTransform = DeferredForces.GetOffsetTransform();
+						const FTransform& OffsetTransform = Module->GetComponentTransform();
 						FVector LocalPos = Module->GetParentRelativeTransform().GetLocation();
-						FVector WorldLocation = BodyTransform.TransformPosition(OffsetTransform.TransformVector(LocalPos));
+						FVector WorldLocation = BodyTransform.TransformPosition(LocalPos);
 						const Chaos::FVec3 Arm = WorldLocation - Particle->X();
 
 						//Chaos::FDebugDrawQueue::GetInstance().DrawDebugLine(Particle->X(), Particle->X() + Arm, FColor::Yellow, false, -1.f, 0, 2.f);
@@ -409,6 +400,9 @@ void FSimModuleTree::UpdateModuleVelocites(FClusterUnionPhysicsProxy* PhysicsPro
 						FVector WorldVelocity = Particle->V() - Chaos::FVec3::CrossProduct(Arm, Particle->W());
 						FVector LocalVelocity = OffsetTransform.InverseTransformVector(BodyTransform.InverseTransformVector(WorldVelocity));
 						Module->SetLocalVelocity(LocalVelocity);
+
+						//Chaos::FDebugDrawQueue::GetInstance().DrawDebugLine(WorldLocation, WorldLocation + WorldVelocity, FColor::White, false, -1.f, 0, 5.f);
+
 					}
 
 				}
diff --git a/Engine/Source/Runtime/Experimental/ChaosVehicles/ChaosVehiclesCore/Private/SimModule/SimulationModuleBase.cpp b/Engine/Source/Runtime/Experimental/ChaosVehicles/ChaosVehiclesCore/Private/SimModule/SimulationModuleBase.cpp
--- a/Engine/Source/Runtime/Experimental/ChaosVehicles/ChaosVehiclesCore/Private/SimModule/SimulationModuleBase.cpp
+++ b/Engine/Source/Runtime/Experimental/ChaosVehicles/ChaosVehiclesCore/Private/SimModule/SimulationModuleBase.cpp
@@ -14,7 +14,7 @@ void ISimulationModuleBase::AddLocalForceAtPosition(const FVector& Force, const
 	AppliedForce = Force;
 	if (SimModuleTree)
 	{
-		SimModuleTree->AccessDeferredForces().Add(FDeferredForcesModular::FApplyForceAtPositionData(TransformIndex, Force, Position, bAllowSubstepping, bIsLocalForce, bLevelSlope, DebugColorIn));
+		SimModuleTree->AccessDeferredForces().Add(FDeferredForcesModular::FApplyForceAtPositionData(ComponentTransform, TransformIndex, Force, Position, bAllowSubstepping, bIsLocalForce, bLevelSlope, DebugColorIn));
 	}
 }
 
@@ -27,7 +27,7 @@ void ISimulationModuleBase::AddLocalForce(const FVector& Force, bool bAllowSubst
 		//GetDebugString(DebugString);
 		//UE_LOG(LogInit, Warning, TEXT("AddLocalForce To %s"), *DebugString);
 
-		SimModuleTree->AccessDeferredForces().Add(FDeferredForcesModular::FApplyForceData(TransformIndex, Force, bAllowSubstepping, bIsLocalForce, bLevelSlope, DebugColorIn));
+		SimModuleTree->AccessDeferredForces().Add(FDeferredForcesModular::FApplyForceData(ComponentTransform, TransformIndex, Force, bAllowSubstepping, bIsLocalForce, bLevelSlope, DebugColorIn));
 	}
 }
 
diff --git a/Engine/Source/Runtime/Experimental/ChaosVehicles/ChaosVehiclesCore/Private/SimModule/SuspensionModule.cpp b/Engine/Source/Runtime/Experimental/ChaosVehicles/ChaosVehiclesCore/Private/SimModule/SuspensionModule.cpp
--- a/Engine/Source/Runtime/Experimental/ChaosVehicles/ChaosVehiclesCore/Private/SimModule/SuspensionModule.cpp
+++ b/Engine/Source/Runtime/Experimental/ChaosVehicles/ChaosVehiclesCore/Private/SimModule/SuspensionModule.cpp
@@ -37,8 +37,8 @@ namespace Chaos
 	{
 		FVector LocalDirection = Setup().SuspensionAxis;
 		FVector Local = GetParentRelativeTransform().GetLocation(); // change to just a vector and GetLocalLocation
-		FVector WorldLocation = BodyTransform.TransformPosition(GetRelativeOffsetTransform().TransformVector(Local/*Setup().LocalOffset*/));
-		FVector WorldDirection = BodyTransform.TransformVector(GetRelativeOffsetTransform().TransformVector(LocalDirection));
+		FVector WorldLocation = BodyTransform.TransformPosition(Local);
+		FVector WorldDirection = BodyTransform.TransformVector(LocalDirection);
 
 		OutTrace.Start = WorldLocation - WorldDirection * (Setup().MaxRaise);
 		OutTrace.End = WorldLocation + WorldDirection * (Setup().MaxDrop + WheelRadius);
diff --git a/Engine/Source/Runtime/Experimental/ChaosVehicles/ChaosVehiclesCore/Private/SimModule/WheelModule.cpp b/Engine/Source/Runtime/Experimental/ChaosVehicles/ChaosVehiclesCore/Private/SimModule/WheelModule.cpp
--- a/Engine/Source/Runtime/Experimental/ChaosVehicles/ChaosVehiclesCore/Private/SimModule/WheelModule.cpp
+++ b/Engine/Source/Runtime/Experimental/ChaosVehicles/ChaosVehiclesCore/Private/SimModule/WheelModule.cpp
@@ -50,6 +50,7 @@ namespace Chaos
 			FRotator SteeringRotator(0.f, SteerAngleDegrees, 0.f);
 			FVector Vel = SteeringRotator.UnrotateVector(ModuleLocalVelocity);
 			FVector LocalWheelVelocity = (Setup().Axis == EWheelAxis::X) ? FVector(Vel.X, Vel.Y, Vel.Z) : FVector(Vel.Y, Vel.X, Vel.Z); // Potential Axis Swap
+			LocalWheelVelocity = Setup().ReverseDirection ? -LocalWheelVelocity : LocalWheelVelocity;
 
 			float GroundAngularVelocity = LocalWheelVelocity.X / Re;
 			float Delta = GroundAngularVelocity - AngularVelocity;
@@ -173,6 +174,11 @@ namespace Chaos
 				ForceFromFriction.X = FinalLateralForce;
 			}
 
+			if (Setup().ReverseDirection)
+			{
+				ForceFromFriction = -ForceFromFriction;
+			}
+
  			AddLocalForce(SteeringRotator.RotateVector(ForceFromFriction));
 			TransmitTorque(VehicleModuleSystem, DriveTorque, BrakeTorque);
 
diff --git a/Engine/Source/Runtime/Experimental/ChaosVehicles/ChaosVehiclesCore/Public/SimModule/DeferredForcesModular.h b/Engine/Source/Runtime/Experimental/ChaosVehicles/ChaosVehiclesCore/Public/SimModule/DeferredForcesModular.h
--- a/Engine/Source/Runtime/Experimental/ChaosVehicles/ChaosVehiclesCore/Public/SimModule/DeferredForcesModular.h
+++ b/Engine/Source/Runtime/Experimental/ChaosVehicles/ChaosVehiclesCore/Public/SimModule/DeferredForcesModular.h
@@ -26,8 +26,9 @@ public:
 
 	struct FApplyForceData
 	{
-		FApplyForceData(int TransformIndexIn, const FVector& ForceIn, bool bAllowSubsteppingIn, bool bAccelChangeIn, bool bLevelSlope, const FColor& ColorIn)
-			: TransformIndex(TransformIndexIn)
+		FApplyForceData(const FTransform& OffsetTransformIn, int TransformIndexIn, const FVector& ForceIn, bool bAllowSubsteppingIn, bool bAccelChangeIn, bool bLevelSlope, const FColor& ColorIn)
+			: OffsetTransform(OffsetTransformIn)
+			, TransformIndex(TransformIndexIn)
 			, Force(ForceIn)
 			, Flags(EForceFlags::None)
 			, DebugColor(ColorIn)
@@ -37,6 +38,7 @@ public:
 			Flags |= bLevelSlope ? EForceFlags::LevelSlope : EForceFlags::None;
 		}
 
+		FTransform OffsetTransform;
 		int TransformIndex;
 		FVector Force;
 		EForceFlags Flags;
@@ -45,8 +47,9 @@ public:
 
 	struct FApplyForceAtPositionData
 	{
-		FApplyForceAtPositionData(int TransformIndexIn, const FVector& ForceIn, const FVector& PositionIn, bool bAllowSubsteppingIn, bool bIsLocalForceIn, bool bLevelSlope, const FColor& ColorIn)
-			: TransformIndex(TransformIndexIn)
+		FApplyForceAtPositionData(const FTransform& OffsetTransformIn, int TransformIndexIn, const FVector& ForceIn, const FVector& PositionIn, bool bAllowSubsteppingIn, bool bIsLocalForceIn, bool bLevelSlope, const FColor& ColorIn)
+			: OffsetTransform(OffsetTransformIn)
+			, TransformIndex(TransformIndexIn)
 			, Force(ForceIn)
 			, Position(PositionIn)
 			, Flags(EForceFlags::None)
@@ -57,6 +60,7 @@ public:
 			Flags |= bLevelSlope ? EForceFlags::LevelSlope : EForceFlags::None;
 		}
 
+		FTransform OffsetTransform;
 		int TransformIndex;
 		FVector Force;
 		FVector Position;
@@ -66,8 +70,9 @@ public:
 
 	struct FAddTorqueInRadiansData
 	{
-		FAddTorqueInRadiansData(int TransformIndexIn, const FVector& TorqueIn, bool bAllowSubsteppingIn, bool bAccelChangeIn)
-			: TransformIndex(TransformIndexIn)
+		FAddTorqueInRadiansData(const FTransform& OffsetTransformIn, int TransformIndexIn, const FVector& TorqueIn, bool bAllowSubsteppingIn, bool bAccelChangeIn)
+			: OffsetTransform(OffsetTransformIn)
+			, TransformIndex(TransformIndexIn)
 			, Torque(TorqueIn)
 			, Flags(EForceFlags::None)
 		{
@@ -75,6 +80,7 @@ public:
 			Flags |= bAccelChangeIn ? EForceFlags::AccelChange : EForceFlags::None;
 		}
 
+		FTransform OffsetTransform;
 		int TransformIndex;
 		FVector Torque;
 		EForceFlags Flags;
@@ -82,14 +88,16 @@ public:
 
 	struct FAddImpulseData
 	{
-		FAddImpulseData(int TransformIndexIn, const FVector& ImpulseIn, const bool bVelChangeIn)
-			: TransformIndex(TransformIndexIn)
+		FAddImpulseData(const FTransform& OffsetTransformIn, int TransformIndexIn, const FVector& ImpulseIn, const bool bVelChangeIn)
+			: OffsetTransform(OffsetTransformIn)
+			, TransformIndex(TransformIndexIn)
 			, Impulse(ImpulseIn)
 			, Flags(EForceFlags::None)
 		{
 			Flags |= bVelChangeIn ? EForceFlags::VelChange : EForceFlags::None;
 		}
 
+		FTransform OffsetTransform;
 		int TransformIndex;
 		FVector Impulse;
 		EForceFlags Flags;
@@ -97,14 +105,16 @@ public:
 
 	struct FAddImpulseAtPositionData
 	{
-		FAddImpulseAtPositionData(int TransformIndexIn, const FVector& ImpulseIn, const FVector& PositionIn)
-			: TransformIndex(TransformIndexIn)
+		FAddImpulseAtPositionData(const FTransform& OffsetTransformIn, int TransformIndexIn, const FVector& ImpulseIn, const FVector& PositionIn)
+			: OffsetTransform(OffsetTransformIn)
+			, TransformIndex(TransformIndexIn)
 			, Impulse(ImpulseIn)
 			, Position(PositionIn)
 		{
 
 		}
 
+		FTransform OffsetTransform;
 		int TransformIndex;
 		FVector Impulse;
 		FVector Position;
@@ -155,12 +165,9 @@ public:
 		, const TManagedArray<FTransform>& CollectionMassToLocal
 		, const TManagedArray<int32>& Parent);
 	
-	void Apply(const FTransform& OffsetTransform
-		, TArray<Chaos::FPBDRigidParticleHandle*>& Particles
+	void Apply(TArray<Chaos::FPBDRigidParticleHandle*>& Particles
 		, TArray<Chaos::FPBDRigidClusteredParticleHandle*>& ClusterParticles);
 
-	void SetOffsetTransform(const FTransform& TransformIn) { ParticleOffsetTransform = TransformIn; }
-	const FTransform& GetOffsetTransform() const { return ParticleOffsetTransform; }
 private:
 
 	void AddForceAtPosition(Chaos::FPBDRigidParticleHandle* RigidHandle, const FApplyForceAtPositionData& DataIn, const FTransform& OffsetTransform);
diff --git a/Engine/Source/Runtime/Experimental/ChaosVehicles/ChaosVehiclesCore/Public/SimModule/SimModuleTree.h b/Engine/Source/Runtime/Experimental/ChaosVehicles/ChaosVehiclesCore/Public/SimModule/SimModuleTree.h
--- a/Engine/Source/Runtime/Experimental/ChaosVehicles/ChaosVehiclesCore/Public/SimModule/SimModuleTree.h
+++ b/Engine/Source/Runtime/Experimental/ChaosVehicles/ChaosVehiclesCore/Public/SimModule/SimModuleTree.h
@@ -166,9 +166,6 @@ namespace Chaos
 			return LargestIndex;
 		}
 
-		void SetIsSimulating(bool IsSimulatingIn) { bIsSimulating = IsSimulatingIn; }
-		bool IsSimulating() { return bIsSimulating; }
-
 		void GenerateReplicationStructure(Chaos::FModuleNetDataArray& NetData);
 		void SetNetState(Chaos::FModuleNetDataArray& ModuleDatas);
 		void SetSimState(const Chaos::FModuleNetDataArray& ModuleDatas);
@@ -190,8 +187,6 @@ namespace Chaos
 		FDeferredForcesModular DeferredForces;
 
 		Chaos::FAllInputs AllInputs;
-
-		bool bIsSimulating = false; //#TODO : temp while testing - remove
 	};
 
 
diff --git a/Engine/Source/Runtime/Experimental/ChaosVehicles/ChaosVehiclesCore/Public/SimModule/SimulationModuleBase.h b/Engine/Source/Runtime/Experimental/ChaosVehicles/ChaosVehiclesCore/Public/SimModule/SimulationModuleBase.h
--- a/Engine/Source/Runtime/Experimental/ChaosVehicles/ChaosVehiclesCore/Public/SimModule/SimulationModuleBase.h
+++ b/Engine/Source/Runtime/Experimental/ChaosVehicles/ChaosVehiclesCore/Public/SimModule/SimulationModuleBase.h
@@ -12,7 +12,7 @@ struct CHAOSVEHICLESCORE_API FCoreModularVehicleDebugParams
 {
 	bool ShowMass = false;
 	bool ShowForces = false;
-	float DrawForceScaling = 0.002f;
+	float DrawForceScaling = 0.0004f;
 	float LevelSlopeThreshold = 0.96f; // ~16 degrees
 	bool DisableForces = false;
 };
@@ -248,9 +248,6 @@ namespace Chaos
 		void SetComponentTransform(const FTransform& TransformIn) { ComponentTransform = TransformIn; }
 		const FTransform& GetComponentTransform() const { return ComponentTransform; }
 
-		void SetRelativeOffsetTransform(const FTransform& TransformIn) { RelativeOffsetTransform = TransformIn; }
-		const FTransform& GetRelativeOffsetTransform() const { return RelativeOffsetTransform; }
-
 		/**
 		 * Set the COM relative transform of module when it is broken off, so relative to itself
 		 */
diff --git a/Engine/Source/Runtime/Experimental/ChaosVehicles/ChaosVehiclesCore/Public/SimModule/WheelModule.h b/Engine/Source/Runtime/Experimental/ChaosVehicles/ChaosVehiclesCore/Public/SimModule/WheelModule.h
--- a/Engine/Source/Runtime/Experimental/ChaosVehicles/ChaosVehiclesCore/Public/SimModule/WheelModule.h
+++ b/Engine/Source/Runtime/Experimental/ChaosVehicles/ChaosVehiclesCore/Public/SimModule/WheelModule.h
@@ -36,6 +36,7 @@ namespace Chaos
 			, HandbrakeTorque(3000)
 			, MaxRotationVel(100.0f)
 			, Axis(EWheelAxis::X)
+			, ReverseDirection(false)
 
 		{
 
@@ -62,8 +63,7 @@ namespace Chaos
 
 		float MaxRotationVel;
 		EWheelAxis Axis;
-
-
+		bool ReverseDirection;
 	};
 
 	class CHAOSVEHICLESCORE_API FWheelSimModule : public FTorqueSimModule, public TSimModuleSettings<FWheelSettings>
