Reduce ensure spam when spectating in Iris

diff --git a/Engine/Source/Runtime/Engine/Private/Net/Iris/ReplicationSystem/ActorReplicationBridge.cpp b/Engine/Source/Runtime/Engine/Private/Net/Iris/ReplicationSystem/ActorReplicationBridge.cpp
--- a/Engine/Source/Runtime/Engine/Private/Net/Iris/ReplicationSystem/ActorReplicationBridge.cpp
+++ b/Engine/Source/Runtime/Engine/Private/Net/Iris/ReplicationSystem/ActorReplicationBridge.cpp
@@ -247,7 +247,7 @@ UE::Net::FNetRefHandle UActorReplicationBridge::BeginReplication(AActor* Actor,
 	CreateNetRefHandleParams.PollFrequency = Actor->NetUpdateFrequency;
 
 #if !UE_BUILD_SHIPPING
-	ensureAlwaysMsgf(!(Actor->bAlwaysRelevant || Actor->bOnlyRelevantToOwner) || CreateNetRefHandleParams.StaticPriority >= 1.0f, TEXT("Very low NetPriority %.02f for always relevant or owner relevant Actor %s. Set it to 1.0f or higher."), Actor->NetPriority, ToCStr(Actor->GetName()));
+	ensureMsgf(!(Actor->bAlwaysRelevant || Actor->bOnlyRelevantToOwner) || CreateNetRefHandleParams.StaticPriority >= 1.0f, TEXT("Very low NetPriority %.02f for always relevant or owner relevant Actor %s. Set it to 1.0f or higher."), Actor->NetPriority, ToCStr(Actor->GetName()));
 #endif
 
 	FNetRefHandle ActorRefHandle = Super::BeginReplication(Actor, CreateNetRefHandleParams);
@@ -464,7 +464,7 @@ bool UActorReplicationBridge::WriteCreationHeader(UE::Net::FNetSerializationCont
 	using namespace UE::Net::Private;
 
 	FNetBitStreamWriter* Writer = Context.GetBitStreamWriter();
-	ensureAlways(IsReplicatedHandle(Handle));
+	ensure(IsReplicatedHandle(Handle));
 	const UObject* Object = GetReplicatedObject(Handle);
 	if (const AActor* Actor = Cast<AActor>(Object))
 	{
@@ -493,7 +493,7 @@ bool UActorReplicationBridge::WriteCreationHeader(UE::Net::FNetSerializationCont
 		return !Writer->IsOverflown();
 	}
 
-	ensureAlwaysMsgf(false, TEXT("UActorReplicationBridge::WriteCreationHeader Failed to write creationHeader for NetRefHandle (Id=%u)"), Handle.GetId());
+	ensureMsgf(false, TEXT("UActorReplicationBridge::WriteCreationHeader Failed to write creationHeader for NetRefHandle (Id=%u)"), Handle.GetId());
 
 	return false;
 }
@@ -793,7 +793,7 @@ void UActorReplicationBridge::DestroyInstanceFromRemote(UObject* Instance, ERepl
 		}
 
 		AActor* Owner = Cast<AActor>(Instance->GetOuter());
-		if (ensureAlwaysMsgf(IsValid(Owner) && !Owner->IsUnreachable(), TEXT("UActorReplicationBridge::DestroyInstanceFromRemote Destroyed subobject after owner %s"), *Instance->GetPathName()))
+		if (ensureMsgf(IsValid(Owner) && !Owner->IsUnreachable(), TEXT("UActorReplicationBridge::DestroyInstanceFromRemote Destroyed subobject after owner %s"), *Instance->GetPathName()))
 		{
 			Owner->OnSubobjectDestroyFromReplication(Instance);
 		}
diff --git a/Engine/Source/Runtime/Experimental/Iris/Core/Private/Iris/ReplicationSystem/Filtering/ReplicationFiltering.cpp b/Engine/Source/Runtime/Experimental/Iris/Core/Private/Iris/ReplicationSystem/Filtering/ReplicationFiltering.cpp
--- a/Engine/Source/Runtime/Experimental/Iris/Core/Private/Iris/ReplicationSystem/Filtering/ReplicationFiltering.cpp
+++ b/Engine/Source/Runtime/Experimental/Iris/Core/Private/Iris/ReplicationSystem/Filtering/ReplicationFiltering.cpp
@@ -410,7 +410,7 @@ bool FReplicationFiltering::SetFilter(FInternalNetRefIndex ObjectIndex, FNetObje
 {
 	if (Filter == ConnectionFilterHandle)
 	{
-		ensureAlwaysMsgf(false, TEXT("Use SetConnectionFilter to enable connection filtering of objects. Cause of ensure must be fixed!"));
+		ensureMsgf(false, TEXT("Use SetConnectionFilter to enable connection filtering of objects. Cause of ensure must be fixed!"));
 		return false;
 	}
 
@@ -422,12 +422,12 @@ bool FReplicationFiltering::SetFilter(FInternalNetRefIndex ObjectIndex, FNetObje
 	// Validate the filter
 	if (bWantsToUseDynamicFilter && (NewDynamicFilterIndex >= (uint32)DynamicFilterInfos.Num()))
 	{
-		ensureAlwaysMsgf(false, TEXT("Invalid dynamic filter 0x%08X. Filter is not being changed. Cause of ensure must be fixed!"), NewDynamicFilterIndex);
+		ensureMsgf(false, TEXT("Invalid dynamic filter 0x%08X. Filter is not being changed. Cause of ensure must be fixed!"), NewDynamicFilterIndex);
 		return false;
 	}
 	else if (!bWantsToUseDynamicFilter && (Filter != InvalidNetObjectFilterHandle) && (Filter != ToOwnerFilterHandle))
 	{
-		ensureAlwaysMsgf(false, TEXT("Invalid static filter 0x%08X. Filter is not being changed. Cause of ensure must be fixed!"), Filter);
+		ensureMsgf(false, TEXT("Invalid static filter 0x%08X. Filter is not being changed. Cause of ensure must be fixed!"), Filter);
 		return false;
 	}
 
@@ -1613,7 +1613,7 @@ void FReplicationFiltering::SetSubObjectFilterStatus(FNetObjectGroupHandle Group
 		return;
 	}
 
-	if (ensureAlways(ValidConnections.GetBit(ConnectionId) && SubObjectFilterGroups.GetBit(GroupHandle)))
+	if (ensure(ValidConnections.GetBit(ConnectionId) && SubObjectFilterGroups.GetBit(GroupHandle)))
 	{
 		UE_LOG(LogIrisFiltering, Verbose, TEXT("ReplicationFiltering::SetSubObjectFilterStatus GroupIndex: %u, ConnectionId: %u, FilterStatus: %u"), GroupHandle, ConnectionId, ReplicationStatus == ENetFilterStatus::Allow ? 1U : 0U);
 		FPerObjectInfo* FilterInfo = GetPerObjectInfo(GroupInfos[GroupHandle].ConnectionStateIndex);
@@ -2021,7 +2021,7 @@ void FReplicationFiltering::InitFilters()
 		UClass* NetObjectFilterClass = Cast<UClass>(StaticFindObject(UClass::StaticClass(), ClassOuter, ToCStr(FilterDefinition.ClassName.ToString()), bMatchExactClass));
 		if (NetObjectFilterClass == nullptr || !NetObjectFilterClass->IsChildOf(UNetObjectFilter::StaticClass()))
 		{
-			ensureAlwaysMsgf(NetObjectFilterClass != nullptr, TEXT("NetObjectFilter class is not a NetObjectFilter or could not be found: %s"), ToCStr(FilterDefinition.ClassName.ToString()));
+			ensureMsgf(NetObjectFilterClass != nullptr, TEXT("NetObjectFilter class is not a NetObjectFilter or could not be found: %s"), ToCStr(FilterDefinition.ClassName.ToString()));
 			continue;
 		}
 
@@ -2031,7 +2031,7 @@ void FReplicationFiltering::InitFilters()
 			NetObjectFilterConfigClass = Cast<UClass>(StaticFindObject(UClass::StaticClass(), ClassOuter, ToCStr(FilterDefinition.ConfigClassName.ToString()), bMatchExactClass));
 			if (NetObjectFilterConfigClass == nullptr || !NetObjectFilterConfigClass->IsChildOf(UNetObjectFilterConfig::StaticClass()))
 			{
-				ensureAlwaysMsgf(NetObjectFilterConfigClass != nullptr, TEXT("NetObjectFilterConfig class is not a NetObjectFilterConfig or could not be found: %s"), ToCStr(FilterDefinition.ConfigClassName.ToString()));
+				ensureMsgf(NetObjectFilterConfigClass != nullptr, TEXT("NetObjectFilterConfig class is not a NetObjectFilterConfig or could not be found: %s"), ToCStr(FilterDefinition.ConfigClassName.ToString()));
 				continue;
 			}
 		}
diff --git a/Engine/Source/Runtime/Experimental/Iris/Core/Private/Iris/ReplicationSystem/ReplicationWriter.cpp b/Engine/Source/Runtime/Experimental/Iris/Core/Private/Iris/ReplicationSystem/ReplicationWriter.cpp
--- a/Engine/Source/Runtime/Experimental/Iris/Core/Private/Iris/ReplicationSystem/ReplicationWriter.cpp
+++ b/Engine/Source/Runtime/Experimental/Iris/Core/Private/Iris/ReplicationSystem/ReplicationWriter.cpp
@@ -1826,7 +1826,7 @@ FReplicationWriter::EWriteObjectStatus FReplicationWriter::WriteObjectAndSubObje
 	const EReplicatedObjectState State = Info.GetState();
 
 	// As an object might still have subobjects pending destroy in the list of subobjects 
-	if (State == EReplicatedObjectState::Invalid || !ensureAlwaysMsgf(State > EReplicatedObjectState::Invalid && State < EReplicatedObjectState::PendingDestroy, TEXT("Unsupported state %s"), ToCStr(LexToString(State))))
+	if (State == EReplicatedObjectState::Invalid || !ensureMsgf(State > EReplicatedObjectState::Invalid && State < EReplicatedObjectState::PendingDestroy, TEXT("Unsupported state %s"), ToCStr(LexToString(State))))
 	{
 		return EWriteObjectStatus::InvalidState;
 	}
