Chaos - add support for visiting disabled collisions without also visiting out of date collisions

diff --git a/Engine/Source/Runtime/Experimental/Chaos/Private/Chaos/Collision/PBDCollisionConstraint.cpp b/Engine/Source/Runtime/Experimental/Chaos/Private/Chaos/Collision/PBDCollisionConstraint.cpp
--- a/Engine/Source/Runtime/Experimental/Chaos/Private/Chaos/Collision/PBDCollisionConstraint.cpp
+++ b/Engine/Source/Runtime/Experimental/Chaos/Private/Chaos/Collision/PBDCollisionConstraint.cpp
@@ -426,7 +426,7 @@ namespace Chaos
 
 	void FPBDCollisionConstraint::Activate()
 	{
-		Flags.bIsActivated = true;
+		Flags.bIsCurrent = true;
 		Flags.bDisabled = false;
 
 		AssignSavedManifoldPoints();
diff --git a/Engine/Source/Runtime/Experimental/Chaos/Private/Chaos/Collision/ParticlePairMidPhase.cpp b/Engine/Source/Runtime/Experimental/Chaos/Private/Chaos/Collision/ParticlePairMidPhase.cpp
--- a/Engine/Source/Runtime/Experimental/Chaos/Private/Chaos/Collision/ParticlePairMidPhase.cpp
+++ b/Engine/Source/Runtime/Experimental/Chaos/Private/Chaos/Collision/ParticlePairMidPhase.cpp
@@ -944,7 +944,7 @@ namespace Chaos
 				return ECollisionVisitorResult::Stop;
 			}
 			return ECollisionVisitorResult::Continue;
-		}, ECollisionVisitorFlags::VisitAll);
+		}, ECollisionVisitorFlags::VisitAllCurrent);
 		return bInGraph;
 	}
 
diff --git a/Engine/Source/Runtime/Experimental/Chaos/Private/Chaos/ContactModification.cpp b/Engine/Source/Runtime/Experimental/Chaos/Private/Chaos/ContactModification.cpp
--- a/Engine/Source/Runtime/Experimental/Chaos/Private/Chaos/ContactModification.cpp
+++ b/Engine/Source/Runtime/Experimental/Chaos/Private/Chaos/ContactModification.cpp
@@ -534,6 +534,11 @@ namespace Chaos
 		: Modifier(InModifier)
 		, Particle(InParticle)
 	{
+		// NOTE: We only visit constraints that were activated this tick because we use the ActiveConstraintIndex to determine which
+		// collisions we have already seen (see FVisitedContactPairsTracker::Visit) and this only works for collisions detected this tick.
+		// We explicitly do not want to visit sleeping collisions.
+		const ECollisionVisitorFlags VisitFlags = ECollisionVisitorFlags::VisitActiveAwake | ECollisionVisitorFlags::VisitDisabled;
+
 		FParticleCollisions& ParticleCollisions = Particle->ParticleCollisions();
 		ParticleCollisions.VisitCollisions([this](FPBDCollisionConstraint& Constraint)
 		{
@@ -542,13 +547,13 @@ namespace Chaos
 				Constraints.Add(&Constraint);
 			}
 			return ECollisionVisitorResult::Continue;
-		}, ECollisionVisitorFlags::VisitActiveAwake);
+		}, VisitFlags);
 	}
 
 	bool FVisitedContactPairsTracker::Visit(const FContactPairModifier& ContactPair)
 	{
 		// We use the index in the ActiveConstraints list as an ID for this constraint. This requires
-		// that we are only visiting the active constraints (and not sleeping or non-activated constraints)
+		// that we are only visiting the active constraints (and not sleeping or expired constraints)
 		// in the VisitCollisions call in FContactPairModifierParticleRange
 		const int32 ConstraintIndex = ContactPair.GetConstraintContainerCookie().ConstraintIndex;
 		
diff --git a/Engine/Source/Runtime/Experimental/Chaos/Private/Chaos/PBDRigidsEvolutionGBF.cpp b/Engine/Source/Runtime/Experimental/Chaos/Private/Chaos/PBDRigidsEvolutionGBF.cpp
--- a/Engine/Source/Runtime/Experimental/Chaos/Private/Chaos/PBDRigidsEvolutionGBF.cpp
+++ b/Engine/Source/Runtime/Experimental/Chaos/Private/Chaos/PBDRigidsEvolutionGBF.cpp
@@ -1088,7 +1088,7 @@ void FPBDRigidsEvolutionGBF::ParticleMaterialChanged(FGeometryParticleHandle* Pa
 		Collision.ClearMaterialProperties();
 		return ECollisionVisitorResult::Continue;
 
-	}, ECollisionVisitorFlags::VisitAll);
+	}, ECollisionVisitorFlags::VisitAllCurrentAndExpired);
 
 	// The graph caches some sleep thresholds etc 
 	GetIslandManager().UpdateParticleMaterial(Particle);
diff --git a/Engine/Source/Runtime/Experimental/Chaos/Public/Chaos/Collision/CollisionConstraintAllocator.h b/Engine/Source/Runtime/Experimental/Chaos/Public/Chaos/Collision/CollisionConstraintAllocator.h
--- a/Engine/Source/Runtime/Experimental/Chaos/Public/Chaos/Collision/CollisionConstraintAllocator.h
+++ b/Engine/Source/Runtime/Experimental/Chaos/Public/Chaos/Collision/CollisionConstraintAllocator.h
@@ -617,6 +617,14 @@ namespace Chaos
 			}
 		}
 
+		if (!(VisitFlags & ECollisionVisitorFlags::VisitExpired))
+		{
+			if (!Constraint->IsCurrent())
+			{
+				return false;
+			}
+		}
+
 		return true;
 	}
 
diff --git a/Engine/Source/Runtime/Experimental/Chaos/Public/Chaos/Collision/CollisionVisitor.h b/Engine/Source/Runtime/Experimental/Chaos/Public/Chaos/Collision/CollisionVisitor.h
--- a/Engine/Source/Runtime/Experimental/Chaos/Public/Chaos/Collision/CollisionVisitor.h
+++ b/Engine/Source/Runtime/Experimental/Chaos/Public/Chaos/Collision/CollisionVisitor.h
@@ -24,9 +24,11 @@ namespace Chaos
 		VisitActiveAwake	= 0,		// Collisions that we detected this tick and have not been disabled or put to sleep
 		VisitSleeping		= 1 << 1,	// Collisions on body pairs that are sleeping
 		VisitDisabled		= 1 << 2,	// Collisions that have been disabled by the user, or were not enabled this tick
+		VisitExpired		= 1 << 3,	// Collisions that are still in memory for possible later reuse but not "current" (activated this frame)
 
-		VisitDefault		= VisitActiveAwake | VisitSleeping,
-		VisitAll			= VisitActiveAwake | VisitSleeping | VisitDisabled
+		VisitDefault					= VisitActiveAwake | VisitSleeping,
+		VisitAllCurrent					= VisitActiveAwake | VisitSleeping | VisitDisabled,
+		VisitAllCurrentAndExpired		= VisitActiveAwake | VisitSleeping | VisitDisabled | VisitExpired,
 	};
 	ENUM_CLASS_FLAGS(ECollisionVisitorFlags);
 
diff --git a/Engine/Source/Runtime/Experimental/Chaos/Public/Chaos/Collision/PBDCollisionConstraint.h b/Engine/Source/Runtime/Experimental/Chaos/Public/Chaos/Collision/PBDCollisionConstraint.h
--- a/Engine/Source/Runtime/Experimental/Chaos/Public/Chaos/Collision/PBDCollisionConstraint.h
+++ b/Engine/Source/Runtime/Experimental/Chaos/Public/Chaos/Collision/PBDCollisionConstraint.h
@@ -351,11 +351,10 @@ namespace Chaos
 		FReal GetPhi() const { return (ClosestManifoldPointIndex != INDEX_NONE) ? ManifoldPoints[ClosestManifoldPointIndex].ContactPoint.Phi : TNumericLimits<FReal>::Max(); }
 
 		// Was this constraint activated this frame? It will be activated if the shapes are within CullDistance of each other.
-		// NOTE: All Active constraints are in the ActiveConstraints list on the CollisionConstraintAllocator. A constraint
-		// remains "activated" in this respect even if disabled by the user (via SetDisabled()) in a callback.
-		// This flag is only really useful if you are iterating over constraints that have been kept in memory
-		// for optimization reasons, but are not currently in use. Normally you would only need to consider GetDisabled()
-		bool IsActivated() const { return Flags.bIsActivated; }
+		// NOTE: All Current constraints are in the ActiveConstraints list on the CollisionConstraintAllocator. This remains true 
+		// even if disabled by the user (via SetDisabled()) in a callback. We only care about "not current" constraints
+		// for debug visualization/reporting. Normally you would only need to consider GetDisabled()
+		bool IsCurrent() const { return Flags.bIsCurrent; }
 
 		// Allow the user to disable this constraint. @see GetActive().
 		void SetDisabled(bool bInDisabled) { Flags.bDisabled = bInDisabled; }
@@ -824,7 +823,7 @@ namespace Chaos
 			FFlags() : Bits(0) {}
 			struct
 			{
-				uint16 bIsActivated : 1;				// Was this constraint activated this tick and therefore in the current tick's active list (note: it may subsequently be disabled)
+				uint16 bIsCurrent : 1;					// Was this constraint activated this tick and therefore in the current tick's active list (note: it may subsequently be disabled for various reasons)
 				uint16 bDisabled : 1;					// Is this contact disabled (by the user or because cull distance is exceeded)
 				uint16 bUseManifold : 1;				// Should we use contact manifolds or single points (faster but poor behaviour)
 				uint16 bUseIncrementalManifold : 1;		// Do we need to run incremental collision detection (only LavelSets now)
