early out for the check for the main particle in an island in a cluster union should only happen if a main particle was already found

diff --git a/Engine/Source/Runtime/Experimental/Chaos/Private/Chaos/PBDRigidClustering.cpp b/Engine/Source/Runtime/Experimental/Chaos/Private/Chaos/PBDRigidClustering.cpp
--- a/Engine/Source/Runtime/Experimental/Chaos/Private/Chaos/PBDRigidClustering.cpp
+++ b/Engine/Source/Runtime/Experimental/Chaos/Private/Chaos/PBDRigidClustering.cpp
@@ -1268,7 +1268,7 @@ namespace Chaos
 								ParticlesToRemove.Add(ChildParticle);
 							}
 						}
-						else
+						else if (bHasMainParticle)
 						{
 							break;
 						}
