Server doesn't need need Geometry Collection custom renderer visual representation.

diff --git a/Engine/Source/Runtime/Experimental/GeometryCollectionEngine/Private/GeometryCollection/GeometryCollectionComponent.cpp b/Engine/Source/Runtime/Experimental/GeometryCollectionEngine/Private/GeometryCollection/GeometryCollectionComponent.cpp
--- a/Engine/Source/Runtime/Experimental/GeometryCollectionEngine/Private/GeometryCollection/GeometryCollectionComponent.cpp
+++ b/Engine/Source/Runtime/Experimental/GeometryCollectionEngine/Private/GeometryCollection/GeometryCollectionComponent.cpp
@@ -4875,7 +4875,8 @@ bool UGeometryCollectionComponent::CanUseCustomRenderer() const
 
 void UGeometryCollectionComponent::RefreshCustomRenderer()
 {
-	if (CanUseCustomRenderer())
+	// Don't refresh the custom renderer on the server.
+	if (CanUseCustomRenderer() && !IsNetMode(NM_DedicatedServer))
 	{
 		if (IGeometryCollectionExternalRenderInterface* RendererInterface = Cast<IGeometryCollectionExternalRenderInterface>(CustomRenderer))
 		{
