Fix mass calculation for geometry collections when overriding materials on the component

diff --git a/Engine/Source/Runtime/Experimental/Chaos/Public/Chaos/MassProperties.h b/Engine/Source/Runtime/Experimental/Chaos/Public/Chaos/MassProperties.h
--- a/Engine/Source/Runtime/Experimental/Chaos/Public/Chaos/MassProperties.h
+++ b/Engine/Source/Runtime/Experimental/Chaos/Public/Chaos/MassProperties.h
@@ -113,6 +113,12 @@ namespace Chaos
 		return Density / (T)1000;
 	}
 
+	template<typename T>
+	T KgCm3ToGCm3(T Density)
+	{
+		return Density * (T)1000;
+	}
+
 	template <typename T, int d>
 	using TMassProperties UE_DEPRECATED(4.27, "Deprecated. this class is to be deleted, use FMassProperties instead") = FMassProperties;
 }
diff --git a/Engine/Source/Runtime/Experimental/GeometryCollectionEngine/Private/GeometryCollection/DerivedDataGeometryCollectionCooker.cpp b/Engine/Source/Runtime/Experimental/GeometryCollectionEngine/Private/GeometryCollection/DerivedDataGeometryCollectionCooker.cpp
--- a/Engine/Source/Runtime/Experimental/GeometryCollectionEngine/Private/GeometryCollection/DerivedDataGeometryCollectionCooker.cpp
+++ b/Engine/Source/Runtime/Experimental/GeometryCollectionEngine/Private/GeometryCollection/DerivedDataGeometryCollectionCooker.cpp
@@ -39,6 +39,9 @@ bool FDerivedDataGeometryCollectionCooker::Build(TArray<uint8>& OutData)
 		Chaos::FErrorReporter ErrorReporter(GeometryCollection.GetName());
 
 		BuildSimulationData(ErrorReporter, *Collection, SharedParams);
+		// important : this is necessary to make sure we compute mass scale on the instances properly
+		// sadly we cannot call this in BuildSimulationData because we have no access to the asset
+		GeometryCollection.CacheMaterialDensity();
 		Collection->Serialize(ChaosAr);
 
 		if (false && ErrorReporter.EncounteredAnyErrors())
diff --git a/Engine/Source/Runtime/Experimental/GeometryCollectionEngine/Private/GeometryCollection/GeometryCollectionComponent.cpp b/Engine/Source/Runtime/Experimental/GeometryCollectionEngine/Private/GeometryCollection/GeometryCollectionComponent.cpp
--- a/Engine/Source/Runtime/Experimental/GeometryCollectionEngine/Private/GeometryCollection/GeometryCollectionComponent.cpp
+++ b/Engine/Source/Runtime/Experimental/GeometryCollectionEngine/Private/GeometryCollection/GeometryCollectionComponent.cpp
@@ -3127,25 +3127,12 @@ void UGeometryCollectionComponent::RegisterAndInitializePhysicsProxy()
 		if (ensure(EnginePhysicalMaterial))
 		{
 			SimulationParameters.PhysicalMaterialHandle = EnginePhysicalMaterial->GetPhysicsMaterial();
-
-			// mass properties are cooked in the GC asset, but component can override physics material with a different density
-			// if the user wants the density to be set from the physics material we need adjust by scaling mass by the ratio override/cooked defined properties
-			SimulationParameters.MaterialOverrideMassScaleMultiplier = 1.0f;
-			if (RestCollection && RestCollection->bDensityFromPhysicsMaterial)
-			{
-				bool bMassAsDensity = false;
-				const float AssetMassOrDensity = RestCollection->GetMassOrDensity(bMassAsDensity);
-				if (ensureMsgf(bMassAsDensity, TEXT("Need a density to be able to compute the mass scale multiplier, this may result in incorrect mass properties")))
-				{
-					if (ensureMsgf(AssetMassOrDensity > SMALL_NUMBER, TEXT("Asset density is set to a too small number, ignoring it for mass adjustment")))
-					{
-						const float OverrideMaterialDensity = Chaos::GCm3ToKgCm3(EnginePhysicalMaterial->Density);
-						SimulationParameters.MaterialOverrideMassScaleMultiplier = OverrideMaterialDensity / AssetMassOrDensity;
-					}
-				}
-				
-			}
 		}
+
+		// mass properties are cooked in the GC asset, but component can override physics material with a different density
+		// if the user wants the density to be set from the physics material we need adjust by scaling mass by the ratio override/cooked defined properties
+		SimulationParameters.MaterialOverrideMassScaleMultiplier = ComputeMassScaleRelativeToAsset();
+
 		GetInitializationCommands(SimulationParameters.InitializationCommands);
 	}
 
@@ -5469,6 +5456,30 @@ void UGeometryCollectionComponent::GetMassAndExtents(int32 ItemIndex, float& Out
 	}
 }
 
+float UGeometryCollectionComponent::ComputeMassScaleRelativeToAsset() const
+{
+	// note this only acount for material override density , in the future we could certainly take in account the scale of the component
+	float MassScaleMultiplier = 1.0f;
+	UPhysicalMaterial* EnginePhysicalMaterial = GetPhysicalMaterial();
+	if (ensure(EnginePhysicalMaterial))
+	{
+		if (RestCollection && RestCollection->bDensityFromPhysicsMaterial)
+		{
+			bool bMassAsDensity = false;
+			const float AssetMassOrDensity = RestCollection->GetMassOrDensity(bMassAsDensity);
+			if (ensureMsgf(bMassAsDensity, TEXT("Need a density to be able to compute the mass scale multiplier, this may result in incorrect mass properties")))
+			{
+				if (ensureMsgf(AssetMassOrDensity > SMALL_NUMBER, TEXT("Asset density is set to a too small number, ignoring it for mass adjustment")))
+				{
+					const float OverrideMaterialDensity = Chaos::GCm3ToKgCm3(EnginePhysicalMaterial->Density);
+					MassScaleMultiplier = OverrideMaterialDensity / AssetMassOrDensity;
+				}
+			}
+		}
+	}
+	return MassScaleMultiplier;
+}
+
 float UGeometryCollectionComponent::GetMass() const
 {
 	float OutMass{ 0 };
@@ -5485,7 +5496,8 @@ float UGeometryCollectionComponent::GetMass() const
 			}
 		}
 	}
-	return OutMass;
+	// need to multiply by the instance specific mass scale is any
+	return OutMass * ComputeMassScaleRelativeToAsset();
 }
 
 float UGeometryCollectionComponent::CalculateMass(FName BoneName)
@@ -5504,7 +5516,8 @@ float UGeometryCollectionComponent::CalculateMass(FName BoneName)
 			}
 		}
 	}
-	return OutMass;
+	// need to multiply by the instance specific mass scale is any
+	return OutMass * ComputeMassScaleRelativeToAsset();
 }
 
 bool UGeometryCollectionComponent::CalculateInnerSphere(int32 TransformIndex, UE::Math::TSphere<double>& SphereOut) const
diff --git a/Engine/Source/Runtime/Experimental/GeometryCollectionEngine/Private/GeometryCollection/GeometryCollectionObject.cpp b/Engine/Source/Runtime/Experimental/GeometryCollectionEngine/Private/GeometryCollection/GeometryCollectionObject.cpp
--- a/Engine/Source/Runtime/Experimental/GeometryCollectionEngine/Private/GeometryCollection/GeometryCollectionObject.cpp
+++ b/Engine/Source/Runtime/Experimental/GeometryCollectionEngine/Private/GeometryCollection/GeometryCollectionObject.cpp
@@ -110,6 +110,7 @@ UGeometryCollection::UGeometryCollection(const FObjectInitializer& ObjectInitial
 	, CollisionObjectReductionPercentage_DEPRECATED(0.0f)
 #endif
 	, bDensityFromPhysicsMaterial(false)
+	, CachedDensityFromPhysicsMaterialInGCm3(0) // <=0 value means not cached yet 
 	, bMassAsDensity(true)
 	, Mass(2500.0f)
 	, MinimumMassClamp(0.1f)
@@ -359,23 +360,51 @@ void UGeometryCollection::PostInitProperties()
 	Super::PostInitProperties();
 }
 
+void UGeometryCollection::CacheMaterialDensity()
+{
+	CachedDensityFromPhysicsMaterialInGCm3 = 0;
+
+	UPhysicalMaterial* PhysicsMaterialForDensity = PhysicsMaterial;
+	if (!PhysicsMaterialForDensity)
+	{
+		PhysicsMaterialForDensity = GEngine ? GEngine->DefaultPhysMaterial : nullptr;
+	}
+	if (PhysicsMaterialForDensity)
+	{
+		CachedDensityFromPhysicsMaterialInGCm3 = PhysicsMaterial->Density;
+	}
+}
+
 float UGeometryCollection::GetMassOrDensity(bool& bOutIsDensity) const
+{
+	return GetMassOrDensityInternal(bOutIsDensity, /* bCached */ true);
+}
+
+float UGeometryCollection::GetMassOrDensityInternal(bool& bOutIsDensity, bool bCached) const
 {
 	bOutIsDensity = bMassAsDensity;
 	float MassOrDensity = bMassAsDensity ? Chaos::KgM3ToKgCm3(Mass) : Mass;
 	
 	if (bDensityFromPhysicsMaterial)
 	{
-		UPhysicalMaterial* PhysicsMaterialForDensity = PhysicsMaterial;
-		if (!PhysicsMaterialForDensity)
+		if (bCached && CachedDensityFromPhysicsMaterialInGCm3 > 0)
 		{
-			PhysicsMaterialForDensity = GEngine ? GEngine->DefaultPhysMaterial : nullptr;
+			bOutIsDensity = true;
+			MassOrDensity = Chaos::GCm3ToKgCm3(CachedDensityFromPhysicsMaterialInGCm3);
 		}
-		if (ensureMsgf(PhysicsMaterialForDensity, TEXT("bDensityFromPhysicsMaterial is true but no physics material has been set (and engine default cannot be found )")))
+		else
 		{
-			// materials only provide density
-			bOutIsDensity = true;
-			MassOrDensity = Chaos::GCm3ToKgCm3(PhysicsMaterial->Density);
+			UPhysicalMaterial* PhysicsMaterialForDensity = PhysicsMaterial;
+			if (!PhysicsMaterialForDensity)
+			{
+				PhysicsMaterialForDensity = GEngine ? GEngine->DefaultPhysMaterial : nullptr;
+			}
+			if (ensureMsgf(PhysicsMaterialForDensity, TEXT("bDensityFromPhysicsMaterial is true but no physics material has been set (and engine default cannot be found )")))
+			{
+				// materials only provide density
+				bOutIsDensity = true;
+				MassOrDensity = Chaos::GCm3ToKgCm3(PhysicsMaterial->Density);
+			}
 		}
 	}
 	return MassOrDensity;
@@ -385,7 +414,8 @@ void UGeometryCollection::GetSharedSimulationParams(FSharedSimulationParameters&
 {
 	const FGeometryCollectionSizeSpecificData& SizeSpecificDefault = GetDefaultSizeSpecificData();
 
-	OutParams.Mass = GetMassOrDensity(OutParams.bMassAsDensity);
+	// we grab the non cached version because this is going to be used to generate the mass attribute which will eventually cache the density value if necessary
+	OutParams.Mass = GetMassOrDensityInternal(OutParams.bMassAsDensity, false);
 	OutParams.MinimumMassClamp = MinimumMassClamp;
 
 	FGeometryCollectionSizeSpecificData InfSize;
diff --git a/Engine/Source/Runtime/Experimental/GeometryCollectionEngine/Public/GeometryCollection/GeometryCollectionComponent.h b/Engine/Source/Runtime/Experimental/GeometryCollectionEngine/Public/GeometryCollection/GeometryCollectionComponent.h
--- a/Engine/Source/Runtime/Experimental/GeometryCollectionEngine/Public/GeometryCollection/GeometryCollectionComponent.h
+++ b/Engine/Source/Runtime/Experimental/GeometryCollectionEngine/Public/GeometryCollection/GeometryCollectionComponent.h
@@ -1423,6 +1423,8 @@ private:
 	/** backward compatibility method, until we can remove GlobalMatrices */
 	GEOMETRYCOLLECTIONENGINE_API TArray<FMatrix> ComputeGlobalMatricesFromComponentSpaceTransforms() const;
 
+	float ComputeMassScaleRelativeToAsset() const;
+
 	/** The clusters we need to replicate */
 	TUniquePtr<TSet<Chaos::FPBDRigidClusteredParticleHandle*>> ClustersToRep;
 
diff --git a/Engine/Source/Runtime/Experimental/GeometryCollectionEngine/Public/GeometryCollection/GeometryCollectionObject.h b/Engine/Source/Runtime/Experimental/GeometryCollectionEngine/Public/GeometryCollection/GeometryCollectionObject.h
--- a/Engine/Source/Runtime/Experimental/GeometryCollectionEngine/Public/GeometryCollection/GeometryCollectionObject.h
+++ b/Engine/Source/Runtime/Experimental/GeometryCollectionEngine/Public/GeometryCollection/GeometryCollectionObject.h
@@ -484,12 +484,18 @@ public:
 	GEOMETRYCOLLECTIONENGINE_API void GetSharedSimulationParams(FSharedSimulationParameters& OutParams) const;
 
 	/**
-	* Get Mass or density as set by the asset 
+	* Get Mass or density as set by the asset ( this is the value used by to compute the cached attributes )
 	* Mass is return in Kg and Density is returned in Kg/Cm3
 	* @param bOutIsDensity  is set to true by the function if the returned value is to be treated as a density
 	*/
 	GEOMETRYCOLLECTIONENGINE_API float GetMassOrDensity(bool& bOutIsDensity) const;
 
+	/*
+	* cache the material density used to compute attribute
+	* Warning : this should only be called after recomputing the mass based on those values
+	*/
+	GEOMETRYCOLLECTIONENGINE_API void CacheMaterialDensity();
+
 	/** Accessors for the two guids used to identify this collection */
 	GEOMETRYCOLLECTIONENGINE_API FGuid GetIdGuid() const;
 	GEOMETRYCOLLECTIONENGINE_API FGuid GetStateGuid() const;
@@ -657,6 +663,13 @@ public:
 	UPROPERTY(EditAnywhere, BlueprintReadWrite, Category = "Collisions", meta = (EditCondition = "PhysicsMaterial != nullptr"))
 	bool bDensityFromPhysicsMaterial;
 
+	/**
+	* Cached Material density value used to compute the Mass attribute  ( In gram per cm3 )
+	* this is necessary because the material properties could be changed after without causing the mass attribute to be recomputed ( because the GC asset will not get notified )
+	*/
+	UPROPERTY()
+	float CachedDensityFromPhysicsMaterialInGCm3;
+
 	/**
 	* Mass As Density, units are in kg/m^3 ( only enabled if physics material is not set )
 	*/
@@ -857,4 +870,6 @@ private:
 	/** Array of user data stored with the asset */
 	UPROPERTY(EditAnywhere, AdvancedDisplay, Instanced, Category = AssetUserData)
 	TArray<TObjectPtr<UAssetUserData>> AssetUserData;
+
+	float GetMassOrDensityInternal(bool& bOutIsDensity, bool bCached) const;
 };
