Iris
* Fix copying data for all polled objects instead of only the dirty ones.

#rnx

diff --git a/Engine/Source/Runtime/Experimental/Iris/Core/Private/Iris/ReplicationSystem/Filtering/ReplicationFiltering.cpp b/Engine/Source/Runtime/Experimental/Iris/Core/Private/Iris/ReplicationSystem/Filtering/ReplicationFiltering.cpp
--- a/Engine/Source/Runtime/Experimental/Iris/Core/Private/Iris/ReplicationSystem/Filtering/ReplicationFiltering.cpp
+++ b/Engine/Source/Runtime/Experimental/Iris/Core/Private/Iris/ReplicationSystem/Filtering/ReplicationFiltering.cpp
@@ -318,21 +318,20 @@ void FReplicationFiltering::FilterPrePoll()
 		});
 	}
 
-	if (bCVarRepFilterCullNonRelevant)
-	{
-		FilterNonRelevantObjects();
-	}
-	else
-	{
-		// Make every object in the global scope part of the relevant list
-		NetRefHandleManager->GetRelevantObjectsInternalIndices().Copy(MakeNetBitArrayView(NetRefHandleManager->GetScopableInternalIndices()));
-	}
+	FilterNonRelevantObjects();
 }
 
 void FReplicationFiltering::FilterNonRelevantObjects()
 {
 	IRIS_PROFILER_SCOPE(FReplicationFiltering_FilterNonRelevantObjects);
 
+	if (!bCVarRepFilterCullNonRelevant)
+	{
+		// Make every object in the global scope part of the relevant list
+		NetRefHandleManager->GetRelevantObjectsInternalIndices().Copy(MakeNetBitArrayView(NetRefHandleManager->GetScopableInternalIndices()));
+		return;
+	}
+
 	const uint32* const GlobalScopeData = NetRefHandleManager->GetScopableInternalIndices().GetData();
 	const uint32* const WithOwnerData = ObjectsWithOwnerFilter.GetData();
 	const uint32* const ConnectionFiltersData = AllConnectionFilteredObjects.GetData();
@@ -359,6 +358,8 @@ void FReplicationFiltering::FilterNonRelevantObjects()
 	};
 	
 	ValidConnections.ForAllSetBits(MergeConnectionScopes);
+
+	//$IRIS TODO: Need to ensure newly relevant objects are immediately polled similar to calling ForceNetUpdate.
 }
 
 void FReplicationFiltering::UpdateDynamicFilters(ENetFilterType FilterPass)
diff --git a/Engine/Source/Runtime/Experimental/Iris/Core/Private/Iris/ReplicationSystem/NetRefHandleManager.cpp b/Engine/Source/Runtime/Experimental/Iris/Core/Private/Iris/ReplicationSystem/NetRefHandleManager.cpp
--- a/Engine/Source/Runtime/Experimental/Iris/Core/Private/Iris/ReplicationSystem/NetRefHandleManager.cpp
+++ b/Engine/Source/Runtime/Experimental/Iris/Core/Private/Iris/ReplicationSystem/NetRefHandleManager.cpp
@@ -25,6 +25,7 @@ FNetRefHandleManager::FNetRefHandleManager(FReplicationProtocolManager& InReplic
 , PrevFrameScopableInternalIndices(MaxActiveObjectCount)
 , RelevantObjectsInternalIndices(MaxActiveObjectCount)
 , PolledObjectsInternalIndices(MaxActiveObjectCount)
+, DirtyObjectsToCopy(MaxActiveObjectCount)
 , AssignedInternalIndices(MaxActiveObjectCount)
 , SubObjectInternalIndices(MaxActiveObjectCount)
 , DependentObjectInternalIndices(MaxActiveObjectCount)
diff --git a/Engine/Source/Runtime/Experimental/Iris/Core/Private/Iris/ReplicationSystem/NetRefHandleManager.h b/Engine/Source/Runtime/Experimental/Iris/Core/Private/Iris/ReplicationSystem/NetRefHandleManager.h
--- a/Engine/Source/Runtime/Experimental/Iris/Core/Private/Iris/ReplicationSystem/NetRefHandleManager.h
+++ b/Engine/Source/Runtime/Experimental/Iris/Core/Private/Iris/ReplicationSystem/NetRefHandleManager.h
@@ -197,6 +197,9 @@ public:
 	/** List of objects that we polled this frame */
 	FNetBitArrayView GetPolledObjectsInternalIndices() const { return MakeNetBitArrayView(PolledObjectsInternalIndices); }
 
+	/** List of objects that need to copy their state data */
+	FNetBitArrayView GetDirtyObjectsToCopy() const { return MakeNetBitArrayView(DirtyObjectsToCopy); }
+
 	// Get bitarray for all internal indices that currently are assigned
 	const FNetBitArray& GetAssignedInternalIndices() const { return AssignedInternalIndices; }
 
@@ -319,6 +322,9 @@ private:
 	/** List of objects that we polled this frame */
 	FNetBitArray PolledObjectsInternalIndices;
 
+	/** List of the objects that are considered dirty and for whom we will copy their state data */
+	FNetBitArray DirtyObjectsToCopy;
+
 	// Bitset containing all internal indices that are assigned
 	FNetBitArray AssignedInternalIndices;
 
diff --git a/Engine/Source/Runtime/Experimental/Iris/Core/Private/Iris/ReplicationSystem/Polling/ObjectPoller.cpp b/Engine/Source/Runtime/Experimental/Iris/Core/Private/Iris/ReplicationSystem/Polling/ObjectPoller.cpp
--- a/Engine/Source/Runtime/Experimental/Iris/Core/Private/Iris/ReplicationSystem/Polling/ObjectPoller.cpp
+++ b/Engine/Source/Runtime/Experimental/Iris/Core/Private/Iris/ReplicationSystem/Polling/ObjectPoller.cpp
@@ -23,6 +23,7 @@ FObjectPoller::FObjectPoller(const FInitParams& InitParams)
 	, LocalNetRefHandleManager(ReplicationSystemInternal->GetNetRefHandleManager())
 	, ReplicatedInstances(LocalNetRefHandleManager.GetReplicatedInstances())
 	, AccumulatedDirtyObjects(ReplicationSystemInternal->GetDirtyNetObjectTracker().GetAccumulatedDirtyNetObjects())
+	, DirtyObjectsToCopy(LocalNetRefHandleManager.GetDirtyObjectsToCopy())
 {
 	GarbageCollectionAffectedObjects = MakeNetBitArrayView(ObjectReplicationBridge->GarbageCollectionAffectedObjects);
 
@@ -108,13 +109,20 @@ void FObjectPoller::ForcePollObject(FInternalNetRefIndex ObjectIndex)
 		EReplicationFragmentPollFlags PollOptions = EReplicationFragmentPollFlags::PollAllState;
 		PollOptions |= bIsGCAffectedObject ? EReplicationFragmentPollFlags::ForceRefreshCachedObjectReferencesAfterGC : EReplicationFragmentPollFlags::None;
 
-		const bool bWasMarkedDirty = FReplicationInstanceOperations::PollAndRefreshCachedPropertyData(ObjectData.InstanceProtocol, PollOptions);
-		if (bWasMarkedDirty)
+		const bool bWasAlreadyDirty = DirtyObjectsThisFrame.IsBitSet(ObjectIndex);
+		const bool bPollFoundDirty = FReplicationInstanceOperations::PollAndRefreshCachedPropertyData(ObjectData.InstanceProtocol, PollOptions);
+		if (bWasAlreadyDirty || bPollFoundDirty)
 		{
+			DirtyObjectsToCopy.SetBit(ObjectIndex);
 			DirtyObjectsThisFrame.SetBit(ObjectIndex);
 		}
 		++PollStats.PolledObjectCount;
 	}
+	else
+	{
+		DirtyObjectsToCopy.SetBit(ObjectIndex);
+		DirtyObjectsThisFrame.SetBit(ObjectIndex);
+	}
 }
 
 void FObjectPoller::PushModelPollObject(FInternalNetRefIndex ObjectIndex)
@@ -149,11 +157,18 @@ void FObjectPoller::PushModelPollObject(FInternalNetRefIndex ObjectIndex)
 		if (bNeedsPoll && !bIsDirtyObject)
 		{
 			bIsDirtyObject = FGlobalDirtyNetObjectTracker::IsNetObjectStateDirty(ObjectData.NetHandle);
-			if (bIsDirtyObject)
-			{
-				DirtyObjectsThisFrame.SetBit(bIsDirtyObject);
-			}
 		}
+
+		if (bIsDirtyObject)
+		{
+			DirtyObjectsToCopy.SetBit(ObjectIndex);
+			DirtyObjectsThisFrame.SetBit(ObjectIndex);
+		}
+	}
+	else if (bIsDirtyObject)
+	{
+		DirtyObjectsToCopy.SetBit(ObjectIndex);
+		DirtyObjectsThisFrame.SetBit(ObjectIndex);
 	}
 
 	const bool bIsGCAffectedObject = GarbageCollectionAffectedObjects.GetBit(ObjectIndex);
@@ -172,7 +187,7 @@ void FObjectPoller::PushModelPollObject(FInternalNetRefIndex ObjectIndex)
 	ObjectData.bWantsFullPoll = 0U;
 	
 	// If the object is fully push model we only need to poll it if it's dirty, unless it's a new object or was garbage collected.
-	bool bWasMarkedDirty = false;
+	bool bPollFoundDirty = false;
 	if (EnumHasAnyFlags(InstanceTraits, EReplicationInstanceProtocolTraits::HasFullPushBasedDirtiness))
 	{
 		if (bIsDirtyObject || bWantsFullPoll)
@@ -180,14 +195,14 @@ void FObjectPoller::PushModelPollObject(FInternalNetRefIndex ObjectIndex)
 			// We need to do a poll if object is marked as dirty
 			EReplicationFragmentPollFlags PollOptions = EReplicationFragmentPollFlags::PollAllState;
 			PollOptions |= bIsGCAffectedObject ? EReplicationFragmentPollFlags::ForceRefreshCachedObjectReferencesAfterGC : EReplicationFragmentPollFlags::None;
-			bWasMarkedDirty = FReplicationInstanceOperations::PollAndRefreshCachedPropertyData(InstanceProtocol, EReplicationFragmentTraits::None, PollOptions);
+			bPollFoundDirty = FReplicationInstanceOperations::PollAndRefreshCachedPropertyData(InstanceProtocol, EReplicationFragmentTraits::None, PollOptions);
 			++PollStats.PolledObjectCount;
 		}
 		else if (bIsGCAffectedObject)
 		{
 			// If this object might have been affected by GC, only refresh cached references
 			const EReplicationFragmentTraits RequiredTraits = EReplicationFragmentTraits::HasPushBasedDirtiness;
-			bWasMarkedDirty = FReplicationInstanceOperations::PollAndRefreshCachedObjectReferences(InstanceProtocol, RequiredTraits);
+			bPollFoundDirty = FReplicationInstanceOperations::PollAndRefreshCachedObjectReferences(InstanceProtocol, RequiredTraits);
 			++PollStats.PolledReferencesObjectCount;
 		}
 	}
@@ -201,7 +216,7 @@ void FObjectPoller::PushModelPollObject(FInternalNetRefIndex ObjectIndex)
 		{
 			// Only states which has push based dirtiness need to be updated as the other states will be polled in full anyway.
 			const EReplicationFragmentTraits RequiredTraits = EReplicationFragmentTraits::HasPushBasedDirtiness;
-			bWasMarkedDirty = FReplicationInstanceOperations::PollAndRefreshCachedObjectReferences(InstanceProtocol, RequiredTraits);
+			bPollFoundDirty = FReplicationInstanceOperations::PollAndRefreshCachedObjectReferences(InstanceProtocol, RequiredTraits);
 			++PollStats.PolledReferencesObjectCount;
 		}
 
@@ -211,12 +226,13 @@ void FObjectPoller::PushModelPollObject(FInternalNetRefIndex ObjectIndex)
 
 		// If the object is not new or dirty at this point we only need to poll non-push based fragments as we know that pushed based states have not been modified
 		const EReplicationFragmentTraits ExcludeTraits = (bIsDirtyObject || bWantsFullPoll) ? EReplicationFragmentTraits::None : EReplicationFragmentTraits::HasPushBasedDirtiness;
-		bWasMarkedDirty |= FReplicationInstanceOperations::PollAndRefreshCachedPropertyData(InstanceProtocol, ExcludeTraits, PollOptions);
+		bPollFoundDirty |= FReplicationInstanceOperations::PollAndRefreshCachedPropertyData(InstanceProtocol, ExcludeTraits, PollOptions);
 		++PollStats.PolledObjectCount;
 	}
 
-	if (bWasMarkedDirty)
+	if (bIsDirtyObject || bPollFoundDirty)
 	{
+		DirtyObjectsToCopy.SetBit(ObjectIndex);
 		DirtyObjectsThisFrame.SetBit(ObjectIndex);
 	}
 }
diff --git a/Engine/Source/Runtime/Experimental/Iris/Core/Private/Iris/ReplicationSystem/Polling/ObjectPoller.h b/Engine/Source/Runtime/Experimental/Iris/Core/Private/Iris/ReplicationSystem/Polling/ObjectPoller.h
--- a/Engine/Source/Runtime/Experimental/Iris/Core/Private/Iris/ReplicationSystem/Polling/ObjectPoller.h
+++ b/Engine/Source/Runtime/Experimental/Iris/Core/Private/Iris/ReplicationSystem/Polling/ObjectPoller.h
@@ -72,6 +72,7 @@ private:
 
 	const FNetBitArrayView AccumulatedDirtyObjects;
 
+	FNetBitArrayView DirtyObjectsToCopy;
 	FNetBitArrayView DirtyObjectsThisFrame;
 	FNetBitArrayView GarbageCollectionAffectedObjects;
 
diff --git a/Engine/Source/Runtime/Experimental/Iris/Core/Private/Iris/ReplicationSystem/ReplicationSystem.cpp b/Engine/Source/Runtime/Experimental/Iris/Core/Private/Iris/ReplicationSystem/ReplicationSystem.cpp
--- a/Engine/Source/Runtime/Experimental/Iris/Core/Private/Iris/ReplicationSystem/ReplicationSystem.cpp
+++ b/Engine/Source/Runtime/Experimental/Iris/Core/Private/Iris/ReplicationSystem/ReplicationSystem.cpp
@@ -350,16 +350,16 @@ public:
 
 		SerializationContext.SetInternalContext(&InternalContext);
 
-		// Copy the state data of objects that got polled this frame.
-		const FNetBitArrayView ObjectsToCopy = NetRefHandleManager.GetPolledObjectsInternalIndices();
-		
+		// Copy the state data of objects that were dirty this frame.
+		FNetBitArrayView DirtyObjectsToCopy = NetRefHandleManager.GetDirtyObjectsToCopy();
+
 		auto CopyFunction = [&ChangeMaskWriter, &Cache, &NetRefHandleManager, &CopiedObjectCount, &SerializationContext](uint32 DirtyIndex)
 		{
 			CopiedObjectCount += FReplicationInstanceOperationsInternal::CopyObjectStateData(ChangeMaskWriter, Cache, NetRefHandleManager, SerializationContext, DirtyIndex);
 		};
 
-		// Only copy both dirty and relevant objects
-		ObjectsToCopy.ForAllSetBits(CopyFunction);
+		DirtyObjectsToCopy.ForAllSetBits(CopyFunction);
+		DirtyObjectsToCopy.Reset();
 
 		const uint32 ReplicationSystemId = ReplicationSystem->GetId();
 		UE_NET_TRACE_FRAME_STATSCOUNTER(ReplicationSystemId, ReplicationSystem.CopiedObjectCount, CopiedObjectCount, ENetTraceVerbosity::Trace);
@@ -371,7 +371,7 @@ public:
 
 		FNetRefHandleManager& NetRefHandleManager = ReplicationSystemInternal.GetNetRefHandleManager();
 
-		// Clean the objects that got polled and copied this frame
+		// Clean the objects that got polled this frame
 		const FNetBitArrayView ObjectsToClean = NetRefHandleManager.GetPolledObjectsInternalIndices();
 
 		// Reset object dirtyness
