Chaos: Add global limit "p.Chaos.Clustering.PerAdvanceBreaksAllowed" to number of processed breaks in BreakingModel for strained particles. All particles still receive strain and all strain modifiers are still executed for all strained particles but the breaking model will not release children beyond the specified limits.

Chaos: Added CVar"p.Chaos.Clustering.DumpClusterAndReleaseStats" to dump out per-frame stats on number of processed clusters and how many children are released from clusters

diff --git a/Engine/Source/Runtime/Experimental/Chaos/Private/Chaos/PBDRigidClustering.cpp b/Engine/Source/Runtime/Experimental/Chaos/Private/Chaos/PBDRigidClustering.cpp
--- a/Engine/Source/Runtime/Experimental/Chaos/Private/Chaos/PBDRigidClustering.cpp
+++ b/Engine/Source/Runtime/Experimental/Chaos/Private/Chaos/PBDRigidClustering.cpp
@@ -1086,6 +1086,8 @@ namespace Chaos
 			}
 		}
 
+		FrameReleasedChildren += ActivatedChildren.Num();
+
 		return ActivatedChildren;
 	}
 
@@ -1370,11 +1372,19 @@ namespace Chaos
 	FAutoConsoleVariableRef CVarBreakMode(TEXT("p.chaos.clustering.breakonlystrained"), GClusterBreakOnlyStrained, 
 										  TEXT("If enabled we only process strained clusters for breaks, if disabled all clusters are traversed and checked"));
 
+	static int32 GPerAdvanceBreaksAllowed = TNumericLimits<int32>::Max();
+	FAutoConsoleVariableRef CVarPerAdvanceBreaksAllowed(TEXT("p.Chaos.Clustering.PerAdvanceBreaksAllowed"), GPerAdvanceBreaksAllowed,
+		TEXT("Number of breaks allowed to occur for each invokation of AdvanceClustering"));
+
+	static int32 GDumpClusterAndReleaseStats = 0;
+	FAutoConsoleVariableRef CVarDumpClusterAndReleaseStats(TEXT("p.Chaos.Clustering.DumpClusterAndReleaseStats"), GDumpClusterAndReleaseStats,
+		TEXT("Report the number of cluster processes and released particles per frame, on/off 1/0"));
+
 	DECLARE_CYCLE_STAT(TEXT("TPBDRigidClustering<>::AdvanceClustering"), STAT_AdvanceClustering, STATGROUP_Chaos);
 	DECLARE_CYCLE_STAT(TEXT("TPBDRigidClustering<>::Update Impulse from Strain"), STAT_UpdateImpulseStrain, STATGROUP_Chaos);
 	DECLARE_CYCLE_STAT(TEXT("TPBDRigidClustering<>::Update Dirty Impulses"), STAT_UpdateDirtyImpulses, STATGROUP_Chaos);
-	void 
-	FRigidClustering::AdvanceClustering(
+	
+	void FRigidClustering::AdvanceClustering(
 		const FReal Dt, 
 		FPBDCollisionConstraints& CollisionRule)
 	{
@@ -1385,6 +1395,24 @@ namespace Chaos
 		FDurationTimer Timer(Time);
 		Timer.Start();
 
+		if(GDumpClusterAndReleaseStats == 1)
+		{
+			if(AdvanceCount > 0 && FrameReleasedChildren > 0)
+			{
+				UE_LOG(LogChaos, Display, TEXT("Clustering | Frame %.5u | Clusters: %.3u, Released: %.4u (TotalClusters: %.3u, TotalReleased: %.4u)"),
+					   AdvanceCount,
+					   FrameProcessedClusters,
+					   FrameReleasedChildren,
+					   TotalProcessedClusters,
+					   TotalReleasedChildren);
+			}
+			AdvanceCount++;
+			TotalProcessedClusters += FrameProcessedClusters;
+			TotalReleasedChildren += FrameReleasedChildren;
+		}
+		FrameProcessedClusters = 0;
+		FrameReleasedChildren = 0;
+
 		if(MChildren.Num())
 		{
 			//
@@ -1515,7 +1543,14 @@ namespace Chaos
 				// #TODO convert to visitor pattern to avoid TArray allocations above.
 				if(GClusterBreakOnlyStrained == 1)
 				{
-					BreakingModel(ParticlesToProcess);
+					// Restrict particle view to allowed breaks. We still build the whole array to allow the strain
+					// modifies (executed above) to all strained particles, even though they may not release
+					TArrayView<FPBDRigidClusteredParticleHandle*> ToProcessView{ 
+						ParticlesToProcess.GetData(), 
+						FMath::Min(ParticlesToProcess.Num(), GPerAdvanceBreaksAllowed) 
+					};
+
+					BreakingModel(ToProcessView);
 				}
 				else
 				{
@@ -1606,9 +1641,16 @@ namespace Chaos
 	}
 	
 	void FRigidClustering::BreakingModel(TArray<FPBDRigidClusteredParticleHandle*>& InParticles)
+	{
+		BreakingModel(MakeArrayView(InParticles));
+	}
+
+	void FRigidClustering::BreakingModel(TArrayView<FPBDRigidClusteredParticleHandle*> InParticles)
 	{
 		SCOPE_CYCLE_COUNTER(STAT_BreakingModel);
 
+		FrameProcessedClusters += InParticles.Num();
+
 		// Clear the set tracking breaking collisions
 		BreakingCollisions.Empty();
 
@@ -1624,7 +1666,7 @@ namespace Chaos
 		// we'll only update the cluster properties once here (connection graph, geometry, etc.).
 		ClusterUnionManager.HandleDeferredClusterUnionUpdateProperties();
 		// Restore some of the momentum of objects that were touching rigid clusters that broke
-		if (RestoreBreakingMomentumPercent > 0.f)
+		if(RestoreBreakingMomentumPercent > 0.f)
 		{
 			RestoreBreakingMomentum();
 		}
diff --git a/Engine/Source/Runtime/Experimental/Chaos/Public/Chaos/PBDRigidClustering.h b/Engine/Source/Runtime/Experimental/Chaos/Public/Chaos/PBDRigidClustering.h
--- a/Engine/Source/Runtime/Experimental/Chaos/Public/Chaos/PBDRigidClustering.h
+++ b/Engine/Source/Runtime/Experimental/Chaos/Public/Chaos/PBDRigidClustering.h
@@ -210,6 +210,7 @@ public:
 	*/
 	CHAOS_API void BreakingModel();
 	CHAOS_API void BreakingModel(TArray<FPBDRigidClusteredParticleHandle*>& InParticles);
+	CHAOS_API void BreakingModel(TArrayView<FPBDRigidClusteredParticleHandle*> InParticles);
 	
 	//
 	// Access
@@ -501,6 +502,13 @@ public:
 		const TMap<FPBDRigidParticleHandle*, FPBDRigidParticleHandle*>& ChildToParentMap);
 private:
 
+	// Cluster release stats for debugging with CVar p.Chaos.Clustering.DumpClusterAndReleaseStats
+	uint32 AdvanceCount = 0;
+	uint32 TotalProcessedClusters = 0;
+	uint32 TotalReleasedChildren = 0;
+	uint32 FrameProcessedClusters = 0;
+	uint32 FrameReleasedChildren = 0;
+
 	FRigidEvolution& MEvolution;
 	FPBDRigidClusteredParticles& MParticles;
 	TSet<Chaos::FPBDRigidClusteredParticleHandle*> TopLevelClusterParents;
