Stop creating the Geometry Collection's Dynamic Component for Editor Preview Worlds (the worlds used for tool preview viewports, and thumbnails) since the Dynamic Component has data that isn't tracked for undo, so can't be active in editors contexts where undo can happen.

This fixes an issue where stale data on the preview thumbnail caused an ensure on undo of any fracture operation that changed the number of bones in the collection.

diff --git a/Engine/Source/Runtime/Experimental/GeometryCollectionEngine/Private/GeometryCollection/GeometryCollectionComponent.cpp b/Engine/Source/Runtime/Experimental/GeometryCollectionEngine/Private/GeometryCollection/GeometryCollectionComponent.cpp
--- a/Engine/Source/Runtime/Experimental/GeometryCollectionEngine/Private/GeometryCollection/GeometryCollectionComponent.cpp
+++ b/Engine/Source/Runtime/Experimental/GeometryCollectionEngine/Private/GeometryCollection/GeometryCollectionComponent.cpp
@@ -2850,7 +2850,7 @@ void UGeometryCollectionComponent::ResetDynamicCollection()
 	bCreateDynamicCollection = false;
 	if (UWorld* World = GetWorld())
 	{
-		if(World->IsGameWorld() || World->IsPreviewWorld())
+		if(World->IsGameWorld())
 		{
 			bCreateDynamicCollection = true;
 		}
