[ChaosVD] Enabling compilation of CVD (Chaos Visual Debugger) by default in non-shipping builds.

Also fixed compiling errors caused due to changes in internal collision data structures in Chaos (which were not being detected because the tool is compiled out), and fixed some editor test cases failing.

diff --git a/Engine/Plugins/Experimental/ChaosVD/ChaosVD.uplugin b/Engine/Plugins/Experimental/ChaosVD/ChaosVD.uplugin
--- a/Engine/Plugins/Experimental/ChaosVD/ChaosVD.uplugin
+++ b/Engine/Plugins/Experimental/ChaosVD/ChaosVD.uplugin
@@ -10,7 +10,7 @@
 	"DocsURL": "",
 	"MarketplaceURL": "",
 	"SupportURL": "",
-	"EnabledByDefault": false,
+	"EnabledByDefault": true,
 	"CanContainContent": true,
 	"IsBetaVersion": false,
 	"IsExperimentalVersion": true,
diff --git a/Engine/Plugins/Experimental/ChaosVD/Source/ChaosVD/Public/ChaosVDEditorSettings.h b/Engine/Plugins/Experimental/ChaosVD/Source/ChaosVD/Public/ChaosVDEditorSettings.h
--- a/Engine/Plugins/Experimental/ChaosVD/Source/ChaosVD/Public/ChaosVDEditorSettings.h
+++ b/Engine/Plugins/Experimental/ChaosVD/Source/ChaosVD/Public/ChaosVDEditorSettings.h
@@ -34,13 +34,13 @@ class UChaosVDEditorSettings : public UObject
 	GENERATED_BODY()
 public:
 
-	UPROPERTY(EditAnywhere, Category = "Viewport Visualization Flags", meta = (Bitmask, BitmaskEnum = EChaosVDParticleDataVisualizationFlags))
+	UPROPERTY(EditAnywhere, Category = "Viewport Visualization Flags", meta = (Bitmask, BitmaskEnum = "/Script/ChaosVD.EChaosVDParticleDataVisualizationFlags"))
 	uint8 GlobalParticleDataVisualizationFlags = 0;
 	
-	UPROPERTY(EditAnywhere, Category = "Viewport Visualization Flags", meta = (Bitmask, BitmaskEnum = EChaosVDCollisionVisualizationFlags))
+	UPROPERTY(EditAnywhere, Category = "Viewport Visualization Flags", meta = (Bitmask, BitmaskEnum = "/Script/ChaosVD.EChaosVDCollisionVisualizationFlags"))
 	uint8 GlobalCollisionDataVisualizationFlags = 0;
 
-	UPROPERTY(EditAnywhere, Category = "Viewport Visualization Flags", meta = (Bitmask, BitmaskEnum = EChaosVDCollisionVisualizationFlags))
+	UPROPERTY(EditAnywhere, Category = "Viewport Visualization Flags", meta = (Bitmask, BitmaskEnum = "/Script/ChaosVD.EChaosVDCollisionVisualizationFlags"))
 	bool bShowDebugText = false;
 
 	UPROPERTY(EditAnywhere, Category = "Viewport Tracking")
diff --git a/Engine/Plugins/Experimental/ChaosVD/Source/ChaosVD/Public/ChaosVDParticleActor.h b/Engine/Plugins/Experimental/ChaosVD/Source/ChaosVD/Public/ChaosVDParticleActor.h
--- a/Engine/Plugins/Experimental/ChaosVD/Source/ChaosVD/Public/ChaosVDParticleActor.h
+++ b/Engine/Plugins/Experimental/ChaosVD/Source/ChaosVD/Public/ChaosVDParticleActor.h
@@ -65,10 +65,10 @@ public:
 
 protected:
 	
-	UPROPERTY(EditAnywhere, Category = "Viewport Visualization Flags", meta = (Bitmask, BitmaskEnum = EChaosVDParticleDataVisualizationFlags))
+	UPROPERTY(EditAnywhere, Category = "Viewport Visualization Flags", meta = (Bitmask, BitmaskEnum = "/Script/ChaosVD.EChaosVDParticleDataVisualizationFlags"))
 	uint8 LocalParticleDataVisualizationFlags;
 	
-	UPROPERTY(EditAnywhere, Category = "Viewport Visualization Flags", meta = (Bitmask, BitmaskEnum = EChaosVDCollisionVisualizationFlags))
+	UPROPERTY(EditAnywhere, Category = "Viewport Visualization Flags", meta = (Bitmask, BitmaskEnum = "/Script/ChaosVD.EChaosVDCollisionVisualizationFlags"))
 	uint8 LocalCollisionDataVisualizationFlags;
 
 	UPROPERTY(EditAnywhere, Category = "Viewport Visualization Flags")
diff --git a/Engine/Source/Programs/UnrealBuildTool/Configuration/TargetRules.cs b/Engine/Source/Programs/UnrealBuildTool/Configuration/TargetRules.cs
--- a/Engine/Source/Programs/UnrealBuildTool/Configuration/TargetRules.cs
+++ b/Engine/Source/Programs/UnrealBuildTool/Configuration/TargetRules.cs
@@ -826,7 +826,7 @@ namespace UnrealBuildTool
 		[XmlConfigFile(Name = "bCompileChaosVisualDebuggerSupport")]
 		[CommandLine("-CompileChaosVisualDebuggerSupport")]
 		[RequiresUniqueBuildEnvironment]
-		public bool bCompileChaosVisualDebuggerSupport = false;
+		public bool bCompileChaosVisualDebuggerSupport = true;
 
 		/// <summary>
 		/// Whether scene query acceleration is done by UE. The physx scene query structure is still created, but we do not use it.
diff --git a/Engine/Source/Runtime/Experimental/Chaos/Private/ChaosVisualDebugger/ChaosVDDataWrapperUtils.cpp b/Engine/Source/Runtime/Experimental/Chaos/Private/ChaosVisualDebugger/ChaosVDDataWrapperUtils.cpp
--- a/Engine/Source/Runtime/Experimental/Chaos/Private/ChaosVisualDebugger/ChaosVDDataWrapperUtils.cpp
+++ b/Engine/Source/Runtime/Experimental/Chaos/Private/ChaosVisualDebugger/ChaosVDDataWrapperUtils.cpp
@@ -70,6 +70,7 @@ FChaosVDConstraint FChaosVDDataWrapperUtils::BuildConstraintDataWrapperFromConst
 {
 	FChaosVDConstraint WrappedConstraintData;
 	
+	WrappedConstraintData.bIsCurrent = InConstraint.Flags.bIsCurrent;
 	WrappedConstraintData.bDisabled = InConstraint.Flags.bDisabled;
 	WrappedConstraintData.bUseManifold = InConstraint.Flags.bUseManifold;
 	WrappedConstraintData.bUseIncrementalManifold = InConstraint.Flags.bUseIncrementalManifold;
@@ -77,7 +78,6 @@ FChaosVDConstraint FChaosVDDataWrapperUtils::BuildConstraintDataWrapperFromConst
 	WrappedConstraintData.bWasManifoldRestored = InConstraint.Flags.bWasManifoldRestored;
 	WrappedConstraintData.bIsQuadratic0 = InConstraint.Flags.bIsQuadratic0;
 	WrappedConstraintData.bIsQuadratic1 = InConstraint.Flags.bIsQuadratic1;
-	WrappedConstraintData.bIsProbeUnmodified = InConstraint.Flags.bIsProbeUnmodified;
 	WrappedConstraintData.bIsProbe = InConstraint.Flags.bIsProbe;
 	WrappedConstraintData.bCCDEnabled = InConstraint.Flags.bCCDEnabled;
 	WrappedConstraintData.bCCDSweepEnabled = InConstraint.Flags.bCCDSweepEnabled;
@@ -156,7 +156,7 @@ FChaosVDParticlePairMidPhase FChaosVDDataWrapperUtils::BuildMidPhaseDataWrapperF
 		FChaosVDConstraint WrappedConstraintData = FChaosVDDataWrapperUtils::BuildConstraintDataWrapperFromConstraint(Constraint);
 		WrappedMidPhaseData.Constraints.Add(MoveTemp(WrappedConstraintData));
 		return Chaos::ECollisionVisitorResult::Continue;
-	}, false);
+	}, Chaos::ECollisionVisitorFlags::VisitAllCurrent);
 
 	return MoveTemp(WrappedMidPhaseData);
 }
diff --git a/Engine/Source/Runtime/Experimental/ChaosVisualDebugger/Private/DataWrappers/ChaosVDCollisionDataWrappers.cpp b/Engine/Source/Runtime/Experimental/ChaosVisualDebugger/Private/DataWrappers/ChaosVDCollisionDataWrappers.cpp
--- a/Engine/Source/Runtime/Experimental/ChaosVisualDebugger/Private/DataWrappers/ChaosVDCollisionDataWrappers.cpp
+++ b/Engine/Source/Runtime/Experimental/ChaosVisualDebugger/Private/DataWrappers/ChaosVDCollisionDataWrappers.cpp
@@ -37,6 +37,7 @@ bool FChaosVDManifoldPoint::Serialize(FArchive& Ar)
 
 bool FChaosVDConstraint::Serialize(FArchive& Ar)
 {
+	FArchive_Serialize_BitfieldBool(Ar, bIsCurrent);
 	FArchive_Serialize_BitfieldBool(Ar, bDisabled);
 	FArchive_Serialize_BitfieldBool(Ar, bUseManifold);
 	FArchive_Serialize_BitfieldBool(Ar, bUseIncrementalManifold);
@@ -44,7 +45,6 @@ bool FChaosVDConstraint::Serialize(FArchive& Ar)
 	FArchive_Serialize_BitfieldBool(Ar, bWasManifoldRestored);
 	FArchive_Serialize_BitfieldBool(Ar, bIsQuadratic0);
 	FArchive_Serialize_BitfieldBool(Ar, bIsQuadratic1);
-	FArchive_Serialize_BitfieldBool(Ar, bIsProbeUnmodified);
 	FArchive_Serialize_BitfieldBool(Ar, bIsProbe);
 	FArchive_Serialize_BitfieldBool(Ar, bCCDEnabled);
 	FArchive_Serialize_BitfieldBool(Ar, bCCDSweepEnabled);
diff --git a/Engine/Source/Runtime/Experimental/ChaosVisualDebugger/Public/DataWrappers/ChaosVDCollisionDataWrappers.h b/Engine/Source/Runtime/Experimental/ChaosVisualDebugger/Public/DataWrappers/ChaosVDCollisionDataWrappers.h
--- a/Engine/Source/Runtime/Experimental/ChaosVisualDebugger/Public/DataWrappers/ChaosVDCollisionDataWrappers.h
+++ b/Engine/Source/Runtime/Experimental/ChaosVisualDebugger/Public/DataWrappers/ChaosVDCollisionDataWrappers.h
@@ -155,6 +155,8 @@ struct CHAOSVDRUNTIME_API FChaosVDConstraint
 
 	inline static FStringView WrapperTypeName = TEXT("FChaosVDConstraint");
 
+	UPROPERTY(EditAnywhere, Category=Falgs)
+	uint8 bIsCurrent:1 = false;
 	UPROPERTY(EditAnywhere, Category=Falgs)
 	uint8 bDisabled:1 = false;
 	UPROPERTY(EditAnywhere, Category=Falgs)
@@ -170,8 +172,6 @@ struct CHAOSVDRUNTIME_API FChaosVDConstraint
 	UPROPERTY(EditAnywhere, Category=Falgs)
 	uint8 bIsQuadratic1:1 = false;
 	UPROPERTY(EditAnywhere, Category=Falgs)
-	uint8 bIsProbeUnmodified:1 = false;
-	UPROPERTY(EditAnywhere, Category=Falgs)
 	uint8 bIsProbe:1 = false;
 	UPROPERTY(EditAnywhere, Category=Falgs)
 	uint8 bCCDEnabled:1 = false;
