[ChaosVD] Moving IsTracing definition to header file to address CIS issue

diff --git a/Engine/Source/Runtime/Experimental/Chaos/Private/ChaosVisualDebugger/ChaosVisualDebuggerTrace.cpp b/Engine/Source/Runtime/Experimental/Chaos/Private/ChaosVisualDebugger/ChaosVisualDebuggerTrace.cpp
--- a/Engine/Source/Runtime/Experimental/Chaos/Private/ChaosVisualDebugger/ChaosVisualDebuggerTrace.cpp
+++ b/Engine/Source/Runtime/Experimental/Chaos/Private/ChaosVisualDebugger/ChaosVisualDebuggerTrace.cpp
@@ -101,7 +101,6 @@ FDelegateHandle FChaosVisualDebuggerTrace::RecordingFullCaptureRequestedHandle =
 FRWLock FChaosVisualDebuggerTrace::DeltaRecordingStatesLock = FRWLock();
 TSet<int32> FChaosVisualDebuggerTrace::SolverIDsForDeltaRecording = TSet<int32>();
 TSet<int32> FChaosVisualDebuggerTrace::RequestedFullCaptureSolverIDs = TSet<int32>();
-std::atomic<bool> FChaosVisualDebuggerTrace::bIsTracing = false;
 
 void FChaosVisualDebuggerTrace::TraceParticle(const Chaos::FGeometryParticleHandle* ParticleHandle)
 {
diff --git a/Engine/Source/Runtime/Experimental/Chaos/Public/ChaosVisualDebugger/ChaosVisualDebuggerTrace.h b/Engine/Source/Runtime/Experimental/Chaos/Public/ChaosVisualDebugger/ChaosVisualDebuggerTrace.h
--- a/Engine/Source/Runtime/Experimental/Chaos/Public/ChaosVisualDebugger/ChaosVisualDebuggerTrace.h
+++ b/Engine/Source/Runtime/Experimental/Chaos/Public/ChaosVisualDebugger/ChaosVisualDebuggerTrace.h
@@ -548,4 +548,7 @@ struct FChaosVDScopeSolverFrame
 
 	T& SolverRef;
 };
+
+std::atomic<bool> FChaosVisualDebuggerTrace::bIsTracing = false;
+
 #endif // CHAOS_VISUAL_DEBUGGER_ENABLED
