[ChaosVD] Adding missing class keywork to EChaosVDFullCaptureFlags enum, to address CIS issue.

diff --git a/Engine/Source/Runtime/Experimental/ChaosVisualDebugger/Public/ChaosVDRuntimeModule.h b/Engine/Source/Runtime/Experimental/ChaosVisualDebugger/Public/ChaosVDRuntimeModule.h
--- a/Engine/Source/Runtime/Experimental/ChaosVisualDebugger/Public/ChaosVDRuntimeModule.h
+++ b/Engine/Source/Runtime/Experimental/ChaosVisualDebugger/Public/ChaosVDRuntimeModule.h
@@ -11,7 +11,7 @@
 struct FChaosVDRecording;
 
 /* Option flags that controls what should be recorded when doing a full capture **/
-enum EChaosVDFullCaptureFlags : int32
+enum class EChaosVDFullCaptureFlags : int32
 {
 	Geometry = 1 << 0,
 	Particles = 1 << 1,
