[ChaosVD] Moving definition of FChaosVisualDebugger::IsTracing to the cpp file to avoit it being inlined

diff --git a/Engine/Source/Runtime/Experimental/Chaos/Private/ChaosVisualDebugger/ChaosVisualDebuggerTrace.cpp b/Engine/Source/Runtime/Experimental/Chaos/Private/ChaosVisualDebugger/ChaosVisualDebuggerTrace.cpp
--- a/Engine/Source/Runtime/Experimental/Chaos/Private/ChaosVisualDebugger/ChaosVisualDebuggerTrace.cpp
+++ b/Engine/Source/Runtime/Experimental/Chaos/Private/ChaosVisualDebugger/ChaosVisualDebuggerTrace.cpp
@@ -568,6 +568,11 @@ void FChaosVisualDebuggerTrace::TraceNonSolverTransform(const FTransform& InTran
 		<< ChaosVDNonSolverTransform.DebugName(DebugNameID.GetData(), DebugNameID.Len());
 }
 
+bool FChaosVisualDebuggerTrace::IsTracing()
+{
+	return bIsTracing;
+}
+
 void FChaosVisualDebuggerTrace::RegisterEventHandlers()
 {
 	{
diff --git a/Engine/Source/Runtime/Experimental/Chaos/Public/ChaosVisualDebugger/ChaosVisualDebuggerTrace.h b/Engine/Source/Runtime/Experimental/Chaos/Public/ChaosVisualDebugger/ChaosVisualDebuggerTrace.h
--- a/Engine/Source/Runtime/Experimental/Chaos/Public/ChaosVisualDebugger/ChaosVisualDebuggerTrace.h
+++ b/Engine/Source/Runtime/Experimental/Chaos/Public/ChaosVisualDebugger/ChaosVisualDebuggerTrace.h
@@ -446,7 +446,7 @@ public:
 	static FString GetDebugName(T& ObjectWithDebugName);
 
 	/** Returns true if a CVD trace is running */
-	static CHAOS_API bool IsTracing() { return bIsTracing;}
+	static CHAOS_API bool IsTracing();
 
 private:
 
