fix nullptr access on PhysicsMaterial instead of PhysicsMaterialForDensity in caching/computing mass/density in the geometry collection

diff --git a/Engine/Source/Runtime/Experimental/GeometryCollectionEngine/Private/GeometryCollection/GeometryCollectionObject.cpp b/Engine/Source/Runtime/Experimental/GeometryCollectionEngine/Private/GeometryCollection/GeometryCollectionObject.cpp
--- a/Engine/Source/Runtime/Experimental/GeometryCollectionEngine/Private/GeometryCollection/GeometryCollectionObject.cpp
+++ b/Engine/Source/Runtime/Experimental/GeometryCollectionEngine/Private/GeometryCollection/GeometryCollectionObject.cpp
@@ -371,7 +371,7 @@ void UGeometryCollection::CacheMaterialDensity()
 	}
 	if (PhysicsMaterialForDensity)
 	{
-		CachedDensityFromPhysicsMaterialInGCm3 = PhysicsMaterial->Density;
+		CachedDensityFromPhysicsMaterialInGCm3 = PhysicsMaterialForDensity->Density;
 	}
 }
 
@@ -403,7 +403,7 @@ float UGeometryCollection::GetMassOrDensityInternal(bool& bOutIsDensity, bool bC
 			{
 				// materials only provide density
 				bOutIsDensity = true;
-				MassOrDensity = Chaos::GCm3ToKgCm3(PhysicsMaterial->Density);
+				MassOrDensity = Chaos::GCm3ToKgCm3(PhysicsMaterialForDensity->Density);
 			}
 		}
 	}
