Fix uninitialized struct members
#horde 339584
#rnx

diff --git a/Engine/Source/Runtime/Experimental/ChaosVisualDebugger/Public/DataWrappers/ChaosVDCollisionDataWrappers.h b/Engine/Source/Runtime/Experimental/ChaosVisualDebugger/Public/DataWrappers/ChaosVDCollisionDataWrappers.h
--- a/Engine/Source/Runtime/Experimental/ChaosVisualDebugger/Public/DataWrappers/ChaosVDCollisionDataWrappers.h
+++ b/Engine/Source/Runtime/Experimental/ChaosVisualDebugger/Public/DataWrappers/ChaosVDCollisionDataWrappers.h
@@ -62,19 +62,19 @@ struct CHAOSVDRUNTIME_API FChaosVDContactPoint
 
 	// Shape-space contact normal on the second shape with direction that points away from shape 1
 	UPROPERTY(EditAnywhere, Category=Contact)
-	FVector  ShapeContactNormal;
+	FVector  ShapeContactNormal = FVector(ForceInit);
 
 	// Contact separation (negative for overlap)
 	UPROPERTY(EditAnywhere, Category=Contact)
-	float Phi;
+	float Phi = 0.f;
 
 	// Face index of the shape we hit. Only valid for Heightfield and Trimesh contact points, otherwise INDEX_NONE
 	UPROPERTY(EditAnywhere, Category=Contact)
-	int32 FaceIndex;
+	int32 FaceIndex = 0;
 
 	// Whether this is a vertex-plane contact, edge-edge contact etc.
 	UPROPERTY(EditAnywhere, Category=Contact)
-	EChaosVDContactPointType ContactType;
+	EChaosVDContactPointType ContactType = EChaosVDContactPointType::Unknown;
 
 	bool Serialize(FArchive& Ar);
 };
@@ -114,12 +114,12 @@ struct CHAOSVDRUNTIME_API FChaosVDManifoldPoint
 	uint8 bInsideStaticFrictionCone:1 = false;
 
 	UPROPERTY(EditAnywhere, Category=ManifoldPointResult)
-	FVector NetPushOut;
+	FVector NetPushOut = FVector(ForceInit);
 	UPROPERTY(EditAnywhere, Category=ManifoldPointResult)
-	FVector NetImpulse;
+	FVector NetImpulse = FVector(ForceInit);
 
 	UPROPERTY(EditAnywhere, Category=ManifoldPoint)
-	float TargetPhi;
+	float TargetPhi = 0.f;
 	UPROPERTY(EditAnywhere, Category=ManifoldPoint)
 	TArray<FVector> ShapeAnchorPoints;
 	UPROPERTY(EditAnywhere, Category=ManifoldPoint)
@@ -183,10 +183,10 @@ struct CHAOSVDRUNTIME_API FChaosVDConstraint
 	uint8 bMaterialSet:1 = false;
 	
 	UPROPERTY(EditAnywhere, Category=Collision)
-	FVector AccumulatedImpulse;
+	FVector AccumulatedImpulse = FVector(ForceInit);
 	
 	UPROPERTY(EditAnywhere, Category=Collision)
-	EChaosVDContactShapesType ShapesType;
+	EChaosVDContactShapesType ShapesType = EChaosVDContactShapesType::Unknown;
 	
 	UPROPERTY(EditAnywhere, Category=Collision)
 	TArray<FTransform> ShapeWorldTransforms;
@@ -195,37 +195,37 @@ struct CHAOSVDRUNTIME_API FChaosVDConstraint
 	TArray<FTransform> ImplicitTransforms;
 	
 	UPROPERTY(EditAnywhere, Category=Collision)
-	float CullDistance;
+	float CullDistance = 0.f;
 	
 	UPROPERTY(EditAnywhere, Category=Collision)
 	TArray<float> CollisionMargins;
 	
 	UPROPERTY(EditAnywhere, Category=Collision)
-	float CollisionTolerance;
+	float CollisionTolerance = 0.f;
 	
 	UPROPERTY(EditAnywhere, Category=Collision)
-	int32 ClosestManifoldPointIndex;
+	int32 ClosestManifoldPointIndex = 0;
 	
 	UPROPERTY(EditAnywhere, Category=Collision)
-	int32 ExpectedNumManifoldPoints;
+	int32 ExpectedNumManifoldPoints = 0;
 	
 	UPROPERTY(EditAnywhere, Category=Collision)
-	FVector LastShapeWorldPositionDelta;
+	FVector LastShapeWorldPositionDelta = FVector(ForceInit);
 	
 	UPROPERTY(EditAnywhere, Category=Collision)
-	FQuat LastShapeWorldRotationDelta;
+	FQuat LastShapeWorldRotationDelta = FQuat(ForceInit);
 	
 	UPROPERTY(EditAnywhere, Category=Collision)
-	float Stiffness;
+	float Stiffness = 0.f;
 	
 	UPROPERTY(EditAnywhere, Category=Collision)
-	float CCDTimeOfImpact;
+	float CCDTimeOfImpact = 0.f;
 	
 	UPROPERTY(EditAnywhere, Category=Collision)
-	float CCDEnablePenetration;
+	float CCDEnablePenetration = 0.f;
 	
 	UPROPERTY(EditAnywhere, Category=Collision)
-	float CCDTargetPenetration;
+	float CCDTargetPenetration = 0.f;
 
 	UPROPERTY(EditAnywhere, Category=Collision)
 	TArray<FChaosVDManifoldPoint> ManifoldPoints;
diff --git a/Engine/Source/Runtime/Experimental/ChaosVisualDebugger/Public/DataWrappers/ChaosVDParticleDataWrapper.h b/Engine/Source/Runtime/Experimental/ChaosVisualDebugger/Public/DataWrappers/ChaosVDParticleDataWrapper.h
--- a/Engine/Source/Runtime/Experimental/ChaosVisualDebugger/Public/DataWrappers/ChaosVDParticleDataWrapper.h
+++ b/Engine/Source/Runtime/Experimental/ChaosVisualDebugger/Public/DataWrappers/ChaosVDParticleDataWrapper.h
@@ -136,10 +136,10 @@ struct FChaosVDParticlePositionRotation : public FChaosVDParticleDataBase
 	CHAOSVDRUNTIME_API bool Serialize(FArchive& Ar);
 
 	UPROPERTY(EditAnywhere, Category= "Particle Velocities")
-	FVector MX;
+	FVector MX = FVector(ForceInit);
 
 	UPROPERTY(EditAnywhere, Category= "Particle Velocities")
-	FQuat MR;
+	FQuat MR = FQuat(ForceInit);
 };
 
 inline FArchive& operator<<(FArchive& Ar, FChaosVDParticlePositionRotation& Data)
@@ -180,10 +180,10 @@ struct FChaosVDParticleVelocities : public FChaosVDParticleDataBase
 	}
 	
 	UPROPERTY(EditAnywhere, Category= "Particle Velocities")
-	FVector MV;
+	FVector MV = FVector(ForceInit);
 
 	UPROPERTY(EditAnywhere, Category= "Particle Velocities")
-	FVector MW;
+	FVector MW = FVector(ForceInit);
 };
 
 inline FArchive& operator<<(FArchive& Ar, FChaosVDParticleVelocities& Data)
@@ -227,16 +227,16 @@ struct FChaosVDParticleDynamics : public FChaosVDParticleDataBase
 	}
 
 	UPROPERTY(EditAnywhere, Category= "Particle Dynamics")
-	FVector MAcceleration;
+	FVector MAcceleration = FVector(ForceInit);
 
 	UPROPERTY(EditAnywhere, Category= "Particle Dynamics")
-	FVector MAngularAcceleration;
+	FVector MAngularAcceleration = FVector(ForceInit);
 
 	UPROPERTY(EditAnywhere, Category= "Particle Dynamics")
-	FVector MLinearImpulseVelocity;
+	FVector MLinearImpulseVelocity = FVector(ForceInit);
 
 	UPROPERTY(EditAnywhere, Category= "Particle Dynamics")
-	FVector MAngularImpulseVelocity;
+	FVector MAngularImpulseVelocity = FVector(ForceInit);
 };
 
 inline FArchive& operator<<(FArchive& Ar, FChaosVDParticleDynamics& Data)
@@ -283,22 +283,22 @@ struct FChaosVDParticleMassProps : public FChaosVDParticleDataBase
 	}
 	
 	UPROPERTY(EditAnywhere, Category= "Particle Mass Props")
-	FVector MCenterOfMass;
+	FVector MCenterOfMass = FVector(ForceInit);
 
 	UPROPERTY(EditAnywhere, Category= "Particle Mass Props")
-	FQuat MRotationOfMass;
+	FQuat MRotationOfMass = FQuat(ForceInit);
 
 	UPROPERTY(EditAnywhere, Category= "Particle Mass Props")
-	FVector MI;
+	FVector MI = FVector(ForceInit);
 
 	UPROPERTY(EditAnywhere, Category= "Particle Mass Props")
-	FVector MInvI;
+	FVector MInvI = FVector(ForceInit);
 
 	UPROPERTY(EditAnywhere, Category= "Particle Mass Props")
-	double MM;
+	double MM = 0.0;
 
 	UPROPERTY(EditAnywhere, Category= "Particle Mass Props")
-	double MInvM;
+	double MInvM = 0.0;
 };
 
 inline FArchive& operator<<(FArchive& Ar, FChaosVDParticleMassProps& Data)
