Chaos Cloth - Fixed the weight map updates on all PBD material constraints that were using the wrong FPBDStiffness constructor.

#rnx

diff --git a/Engine/Source/Runtime/Experimental/Chaos/Private/Chaos/PBDSpringConstraints.cpp b/Engine/Source/Runtime/Experimental/Chaos/Private/Chaos/PBDSpringConstraints.cpp
--- a/Engine/Source/Runtime/Experimental/Chaos/Private/Chaos/PBDSpringConstraints.cpp
+++ b/Engine/Source/Runtime/Experimental/Chaos/Private/Chaos/PBDSpringConstraints.cpp
@@ -194,7 +194,12 @@ void FPBDEdgeSpringConstraints::SetProperties(
 		if (IsEdgeSpringStiffnessStringDirty(PropertyCollection))
 		{
 			const FString& WeightMapName = GetEdgeSpringStiffnessString(PropertyCollection);
-			Stiffness = FPBDStiffness(WeightedValue, WeightMaps.FindRef(WeightMapName), ParticleCount);
+			Stiffness = FPBDStiffness(
+				WeightedValue,
+				WeightMaps.FindRef(WeightMapName),
+				TConstArrayView<TVec2<int32>>(Constraints),
+				ParticleOffset,
+				ParticleCount);
 		}
 		else
 		{
@@ -213,7 +218,12 @@ void FPBDBendingSpringConstraints::SetProperties(
 		if (IsBendingSpringStiffnessStringDirty(PropertyCollection))
 		{
 			const FString& WeightMapName = GetBendingSpringStiffnessString(PropertyCollection);
-			Stiffness = FPBDStiffness(WeightedValue, WeightMaps.FindRef(WeightMapName), ParticleCount);
+			Stiffness = FPBDStiffness(
+				WeightedValue,
+				WeightMaps.FindRef(WeightMapName),
+				TConstArrayView<TVec2<int32>>(Constraints),
+				ParticleOffset,
+				ParticleCount);
 		}
 		else
 		{
