Fix invalid assert in FOnDemandIoBackend::FChunkRequests::Release

#rnx

- In rare cases, inbetween a call to ::Remove and ::Release on a request a new request could be made with the same ChunkKey but requesting a different portion of the data.
- This would cause the assert in ::Release to fire as we only check if there is currently a request inflight with the same ChunkKey as the request we are about to destroy, we don't check if it is actually the same request.
- Changed the assert to no longer make that assumption and checks both the key and the value before raising an issue.

diff --git a/Engine/Source/Runtime/Experimental/IoStoreOnDemand/Private/OnDemandIoDispatcherBackend.cpp b/Engine/Source/Runtime/Experimental/IoStoreOnDemand/Private/OnDemandIoDispatcherBackend.cpp
--- a/Engine/Source/Runtime/Experimental/IoStoreOnDemand/Private/OnDemandIoDispatcherBackend.cpp
+++ b/Engine/Source/Runtime/Experimental/IoStoreOnDemand/Private/OnDemandIoDispatcherBackend.cpp
@@ -896,7 +896,7 @@ class FOnDemandIoBackend final
 		void Release(FChunkRequest* Request)
 		{
 			FScopeLock _(&Mutex);
-			check(!Inflight.Contains(Request->Params.ChunkKey));
+			check(!IsInFlight(Request));
 			Allocator.Destroy(Request);
 		}
 		
@@ -907,6 +907,22 @@ class FOnDemandIoBackend final
 			Allocator.Destroy(Request);
 		}
 
+	private:
+
+		/** Helper intended to be called by methods that have already locked ::Mutex */
+		inline bool IsInFlight(const FChunkRequest* Request) const
+		{
+			const FChunkRequest* const* InFlightRequest = Inflight.Find(Request->Params.ChunkKey);
+			if (InFlightRequest == nullptr)
+			{
+				return false;
+			}
+			else
+			{
+				return *InFlightRequest == Request;
+			}
+		}
+
 		TSingleThreadedSlabAllocator<FChunkRequest, 128> Allocator;
 		TMap<FIoHash, FChunkRequest*> Inflight;
 		FCriticalSection Mutex;
