[ChaosVD] Fix Compilation Error when UE Trace is disabled

diff --git a/Engine/Source/Runtime/Experimental/Chaos/Private/ChaosModule.cpp b/Engine/Source/Runtime/Experimental/Chaos/Private/ChaosModule.cpp
--- a/Engine/Source/Runtime/Experimental/Chaos/Private/ChaosModule.cpp
+++ b/Engine/Source/Runtime/Experimental/Chaos/Private/ChaosModule.cpp
@@ -25,14 +25,14 @@ void FChaosEngineModule::StartupModule()
 		Chaos::GSingleThreadedPhysics = 1;
 	}
 
-#if WITH_CHAOS_VISUAL_DEBUGGER
+#if CHAOS_VISUAL_DEBUGGER_ENABLED
 	FChaosVisualDebuggerTrace::RegisterEventHandlers();
 #endif
 }
 
 void FChaosEngineModule::ShutdownModule()
 {
-#if WITH_CHAOS_VISUAL_DEBUGGER
+#if CHAOS_VISUAL_DEBUGGER_ENABLED
 	FChaosVisualDebuggerTrace::UnregisterEventHandlers();
 #endif
 }
