make FRigidClustering::RemoveParticlesFromCluster and FRigidClustering::AddParticlesToCluster do incremental updates (as much as possible) to the cluster particle's internal strain, cluster group, and physics proxy set instead of needing to re-build these properties from all its children

diff --git a/Engine/Source/Runtime/Experimental/Chaos/Private/Chaos/PBDRigidClustering.cpp b/Engine/Source/Runtime/Experimental/Chaos/Private/Chaos/PBDRigidClustering.cpp
--- a/Engine/Source/Runtime/Experimental/Chaos/Private/Chaos/PBDRigidClustering.cpp
+++ b/Engine/Source/Runtime/Experimental/Chaos/Private/Chaos/PBDRigidClustering.cpp
@@ -263,7 +263,13 @@ namespace Chaos
 		}
 
 		FRigidHandleArray& Children = MChildren.FindOrAdd(Cluster);
-		const int32 OldNumChildren = Children.Num();
+		
+		// Note that we want to compute the internal strain on the cluster the same if we build it up incrementally as well as if we
+		// build it all at the same time. The parent cluster's internal strain should be the average of all the child strains.
+		// The easy way to compute the new average is the multiply the old average by the number of old elements, add in the new strains,
+		// and then divide by the new total number of elements.
+		Cluster->SetInternalStrains(Cluster->GetInternalStrains() * static_cast<FRealSingle>(Children.Num()));
+
 		Children.Append(InChildren);
 
 		// Disable all the input children since they no longer need to be simulated.
@@ -286,20 +292,22 @@ namespace Chaos
 				TopLevelClusterParentsStrained.Remove(ClusteredChild);
 
 				ClusteredChild->ClusterIds().Id = Cluster;
+				Cluster->SetInternalStrains(Cluster->GetInternalStrains() + ClusteredChild->GetInternalStrains());
+				Cluster->SetCollisionGroup(FMath::Min(Cluster->CollisionGroup(), ClusteredChild->CollisionGroup()));
+			}
+
+			Cluster->AddPhysicsProxy(Handle->PhysicsProxy());
+			if (Cluster->PhysicsProxy() == nullptr)
+			{
+				Cluster->SetPhysicsProxy(Handle->PhysicsProxy());
 			}
 
 			MEvolution.DisableParticle(Handle);
 			MEvolution.GetParticles().MarkTransientDirtyParticle(Handle);
 		}
 
-		// Note that we want to compute the internal strain on the cluster the same if we build it up incrementally as well as if we
-		// build it all at the same time. The parent cluster's internal strain should be the average of all the child strains.
-		// The easy way to compute the new average is the multiply the old average by the number of old elements, add in the new strains,
-		// and then divide by the new total number of elements.
-		Cluster->SetInternalStrains(Cluster->GetInternalStrains() * static_cast<FRealSingle>(OldNumChildren));
 		Cluster->ClusterIds().NumChildren = Children.Num();
-
-		UpdateClusterParticlePropertiesFromChildren(Cluster, InChildren, ChildToParentMap);
+		Cluster->SetInternalStrains(Cluster->GetInternalStrains() / static_cast<FRealSingle>(Children.Num()));
 	}
 
 	DECLARE_CYCLE_STAT(TEXT("TPBDRigidClustering<>::RemoveParticlesFromCluster"), STAT_RemoveParticlesFromCluster, STATGROUP_Chaos);
@@ -311,11 +319,22 @@ namespace Chaos
 		SCOPE_CYCLE_COUNTER(STAT_RemoveParticlesFromCluster);
 
 		FRigidHandleArray& Children = MChildren.FindOrAdd(Cluster);
+		TSet<IPhysicsProxyBase*> RemovedProxies;
+
+		FRealSingle NewInternalStrain = Cluster->GetInternalStrains() * static_cast<FRealSingle>(Children.Num());
+
 		for (FPBDRigidParticleHandle* Child : InChildren)
 		{
 			if (int32 Index = Children.Find(Child); Child && Index != INDEX_NONE)
 			{
+				RemovedProxies.Add(Child->PhysicsProxy());
 				RemoveChildFromParent(Child, Cluster);
+
+				if (FPBDRigidClusteredParticleHandle* ChildCluster = Child->CastToClustered())
+				{
+					NewInternalStrain -= ChildCluster->GetInternalStrains();
+				}
+
 				Children.RemoveAtSwap(Index);
 				MEvolution.DirtyParticle(*Child);
 				MEvolution.GetParticles().MarkTransientDirtyParticle(Child);
@@ -323,12 +342,39 @@ namespace Chaos
 		}
 
 		Cluster->ClusterIds().NumChildren = Children.Num();
-		Cluster->SetInternalStrains(0.0);
-		Cluster->SetCollisionGroup(INT_MAX);
-		Cluster->ClearPhysicsProxies();
 
-		// We need to fully rebuild the cluster properties from the set of children.
-		UpdateClusterParticlePropertiesFromChildren(Cluster, Children, {});
+		// If we removed the last particle with a given physics proxy from a cluster, we need to remove that proxy from the proxy set.
+		if (Children.IsEmpty())
+		{
+			Cluster->SetInternalStrains(FRealSingle(0.0));
+			Cluster->ClearPhysicsProxies();
+		}
+		else
+		{
+			Cluster->SetInternalStrains(NewInternalStrain / static_cast<FRealSingle>(Children.Num()));
+
+			// Unfortunately we still need to iterate through every child in the cluster to see if a particular physics proxy in the set is still valid.
+			for (FPBDRigidParticleHandle* Child : Children)
+			{
+				RemovedProxies.Remove(Child->PhysicsProxy());
+				if (RemovedProxies.IsEmpty())
+				{
+					break;
+				}
+			}
+
+			for (IPhysicsProxyBase* Proxy : RemovedProxies)
+			{
+				Cluster->RemovePhysicsProxy(Proxy);
+			}
+
+			IPhysicsProxyBase* FallbackProxy = Children[0]->PhysicsProxy();
+			if (RemovedProxies.Contains(Cluster->PhysicsProxy()))
+			{
+				Cluster->SetPhysicsProxy(FallbackProxy);
+			}
+		}
+
 		MEvolution.DirtyParticle(*Cluster);
 	}
 
@@ -339,6 +385,7 @@ namespace Chaos
 		const FRigidHandleArray& Children,
 		const TMap<FPBDRigidParticleHandle*, FPBDRigidParticleHandle*>& ChildToParentMap)
 	{
+		SCOPE_CYCLE_COUNTER(STAT_UpdateClusterParticlePropertiesFromChildren);
 		// An initial pass through the children to transfer some of their cluster properties to their new parent.
 		for (FPBDRigidParticleHandle* Child : Children)
 		{
diff --git a/Engine/Source/Runtime/Experimental/Chaos/Public/Chaos/PBDRigidClustering.h b/Engine/Source/Runtime/Experimental/Chaos/Public/Chaos/PBDRigidClustering.h
--- a/Engine/Source/Runtime/Experimental/Chaos/Public/Chaos/PBDRigidClustering.h
+++ b/Engine/Source/Runtime/Experimental/Chaos/Public/Chaos/PBDRigidClustering.h
@@ -496,10 +496,12 @@ public:
 	 * incrementally or entirely rebuild the properties all over again. This all depends on whether the input children is
 	 * either 1) the new children or 2) all the children as well as what those initial properties are set to.
 	 */
+	UE_DEPRECATED(5.4, "This should be handled for you properly in AddParticlesToCluster and RemoveParticlesFromCluster. There is no need for an extra function call.")
 	CHAOS_API void UpdateClusterParticlePropertiesFromChildren(
 		FPBDRigidClusteredParticleHandle* Cluster,
 		const FRigidHandleArray& Children,
 		const TMap<FPBDRigidParticleHandle*, FPBDRigidParticleHandle*>& ChildToParentMap);
+
 private:
 
 	// Cluster release stats for debugging with CVar p.Chaos.Clustering.DumpClusterAndReleaseStats
