Constraints: dirty render state when changing scale sign

Note: MovieSceneTracksCustomAccessors::SetComponentTransform should probably dirty render state when inverting scale

#rnx

diff --git a/Engine/Source/Runtime/Experimental/Animation/Constraints/Private/TransformableHandle.cpp b/Engine/Source/Runtime/Experimental/Animation/Constraints/Private/TransformableHandle.cpp
--- a/Engine/Source/Runtime/Experimental/Animation/Constraints/Private/TransformableHandle.cpp
+++ b/Engine/Source/Runtime/Experimental/Animation/Constraints/Private/TransformableHandle.cpp
@@ -98,13 +98,26 @@ void UTransformableComponentHandle::SetGlobalTransform(const FTransform& InGloba
 {
 	if (Component.IsValid())
 	{
+		const FVector OldScale = GetGlobalTransform().GetScale3D();
+		const FVector NewScale = InGlobal.GetScale3D();
+		const bool bMarkRenderDirty =
+			FMath::Sign(OldScale[0]) != FMath::Sign(NewScale[0]) ||
+			FMath::Sign(OldScale[1]) != FMath::Sign(NewScale[1]) ||
+			FMath::Sign(OldScale[2]) != FMath::Sign(NewScale[2]);
+		
 		Component->SetWorldTransform(InGlobal);
+		
 		UWorld* World = Component->GetWorld();
 		if (World)
 		{
 			FConstraintsManagerController& Controller = FConstraintsManagerController::Get(World);
 			Controller.OnSceneComponentConstrained().Broadcast(Component.Get());
 		}
+
+		if (bMarkRenderDirty)
+		{
+			Component->MarkRenderStateDirty();
+		}
 	}
 }
 
