add p.Chaos.DebugDraw.ColorShapesByClusterUnion to be able to differentiate between cluster unions when drawing internal clusters

diff --git a/Engine/Source/Runtime/Experimental/Chaos/Private/Chaos/ChaosDebugDraw.cpp b/Engine/Source/Runtime/Experimental/Chaos/Private/Chaos/ChaosDebugDraw.cpp
--- a/Engine/Source/Runtime/Experimental/Chaos/Private/Chaos/ChaosDebugDraw.cpp
+++ b/Engine/Source/Runtime/Experimental/Chaos/Private/Chaos/ChaosDebugDraw.cpp
@@ -57,6 +57,9 @@ namespace Chaos
 		bool bChaosDebugDebugDrawColorShapesByInternalCluster = false;
 		FAutoConsoleVariableRef CVarChaosDebugDebugDrawColorShapesByInternalCluster(TEXT("p.Chaos.DebugDraw.ColorShapesByInternalCluster"), bChaosDebugDebugDrawColorShapesByInternalCluster, TEXT("Whether to check if the particle is an internal cluster to define its color (black : regular particle: red :internal cluster )"));
 
+		bool bChaosDebugDebugDrawColorShapesByClusterUnion = false;
+		FAutoConsoleVariableRef CVarChaosDebugDebugDrawColorShapesByClusterUnion(TEXT("p.Chaos.DebugDraw.ColorShapesByClusterUnion"), bChaosDebugDebugDrawColorShapesByClusterUnion, TEXT("An extension of the ColorShapesByInternalCluster option: instead of using a single color for every internal cluster, will use a unique color per cluster union. Non-cluster unions will be black."));
+
 		bool bChaosDebugDebugDrawColorBoundsByShapeType = false;
 		FAutoConsoleVariableRef CVarChaosDebugDebugDrawColorBoundsByShapeType(TEXT("p.Chaos.DebugDraw.ColorBoundsByShapeType"), bChaosDebugDebugDrawColorBoundsByShapeType, TEXT("Whether to use shape type to define the color of the bounds instead of using the particle state (if multiple shapes , will use the first one)"));
 
@@ -659,9 +662,19 @@ namespace Chaos
 					ShapeColor = FColor::Black;
 					if (const FPBDRigidClusteredParticleHandle* ClusteredParticle = Particle->CastToClustered())
 					{
-						if (ClusteredParticle->InternalCluster())
+						if (bChaosDebugDebugDrawColorShapesByClusterUnion)
+						{
+							if (ClusteredParticle->PhysicsProxy()->GetType() == EPhysicsProxyType::ClusterUnionProxy)
+							{
+								ShapeColor = GetIndexColor(FMath::Abs(ClusteredParticle->ClusterGroupIndex()));
+							}
+						}
+						else
 						{
-							ShapeColor = FColor::Purple;
+							if (ClusteredParticle->InternalCluster())
+							{
+								ShapeColor = FColor::Purple;
+							}
 						}
 					}
 				}
