GCs should not set themselves dormant to avoid client/server desync

diff --git a/Engine/Source/Runtime/Experimental/GeometryCollectionEngine/Private/GeometryCollection/GeometryCollectionComponent.cpp b/Engine/Source/Runtime/Experimental/GeometryCollectionEngine/Private/GeometryCollection/GeometryCollectionComponent.cpp
--- a/Engine/Source/Runtime/Experimental/GeometryCollectionEngine/Private/GeometryCollection/GeometryCollectionComponent.cpp
+++ b/Engine/Source/Runtime/Experimental/GeometryCollectionEngine/Private/GeometryCollection/GeometryCollectionComponent.cpp
@@ -1903,12 +1903,6 @@ void UGeometryCollectionComponent::UpdateRepData()
 				Owner->SetNetDormancy(DORM_Awake);
 			}
 		}
-		else if (!bFirstUpdate && Owner->NetDormancy == DORM_Awake)
-		{
-			// Clusters are no longer changing so we shouldn't need to replicate anymore.
-			// TODO: Wait a few frames?
-			Owner->SetNetDormancy(DORM_DormantAll);
-		}
 	}
 }
 
