Set MDC stat categories for some known subsystems.

diff --git a/Engine/Plugins/Experimental/Water/Source/Runtime/Private/WaterMeshComponent.cpp b/Engine/Plugins/Experimental/Water/Source/Runtime/Private/WaterMeshComponent.cpp
--- a/Engine/Plugins/Experimental/Water/Source/Runtime/Private/WaterMeshComponent.cpp
+++ b/Engine/Plugins/Experimental/Water/Source/Runtime/Private/WaterMeshComponent.cpp
@@ -64,6 +64,9 @@ UWaterMeshComponent::UWaterMeshComponent()
 	bHasPerInstanceHitProxies = true;
 
 	SetMobility(EComponentMobility::Static);
+
+	static FName NAME_StatsCategoryWater("Water");
+	SetMeshDrawCommandStatsCategory(NAME_StatsCategoryWater);
 }
 
 void UWaterMeshComponent::PostLoad()
diff --git a/Engine/Source/Runtime/Experimental/GeometryCollectionEngine/Private/GeometryCollection/GeometryCollectionISMPoolComponent.cpp b/Engine/Source/Runtime/Experimental/GeometryCollectionEngine/Private/GeometryCollection/GeometryCollectionISMPoolComponent.cpp
--- a/Engine/Source/Runtime/Experimental/GeometryCollectionEngine/Private/GeometryCollection/GeometryCollectionISMPoolComponent.cpp
+++ b/Engine/Source/Runtime/Experimental/GeometryCollectionEngine/Private/GeometryCollection/GeometryCollectionISMPoolComponent.cpp
@@ -109,6 +109,7 @@ FGeometryCollectionISM::FGeometryCollectionISM(AActor* InOwningActor, const FGeo
 	ISMC->MinLOD = MeshInstance.Desc.MinLod;
 	ISMC->SetLODDistanceScale(MeshInstance.Desc.LodScale);
 	ISMC->ComponentTags.Append(MeshInstance.Desc.Tags);
+	ISMC->SetMeshDrawCommandStatsCategory(MeshInstance.Desc.StatsCategory);
 
 	InOwningActor->AddInstanceComponent(ISMC);
 	ISMC->RegisterComponent();
diff --git a/Engine/Source/Runtime/Experimental/GeometryCollectionEngine/Public/GeometryCollection/GeometryCollectionISMPoolComponent.h b/Engine/Source/Runtime/Experimental/GeometryCollectionEngine/Public/GeometryCollection/GeometryCollectionISMPoolComponent.h
--- a/Engine/Source/Runtime/Experimental/GeometryCollectionEngine/Public/GeometryCollection/GeometryCollectionISMPoolComponent.h
+++ b/Engine/Source/Runtime/Experimental/GeometryCollectionEngine/Public/GeometryCollection/GeometryCollectionISMPoolComponent.h
@@ -118,6 +118,7 @@ struct FISMComponentDescription
 	int32 MinLod = 0;
 	float LodScale = 1.f;
 	TArray<FName> Tags;
+	FName StatsCategory;
 
 	bool operator==(const FISMComponentDescription& Other) const
 	{
@@ -127,7 +128,8 @@ struct FISMComponentDescription
 			EndCullDistance == Other.EndCullDistance &&
 			MinLod == Other.MinLod &&
 			LodScale == Other.LodScale &&
-			Tags == Other.Tags;
+			Tags == Other.Tags &&
+			StatsCategory == Other.StatsCategory;
 	}
 };
 
@@ -138,7 +140,8 @@ FORCEINLINE uint32 GetTypeHash(const FISMComponentDescription& Desc)
 	Hash = HashCombineFast(Hash, GetTypeHash(Desc.EndCullDistance));
 	Hash = HashCombineFast(Hash, GetTypeHash(Desc.MinLod));
 	Hash = HashCombineFast(Hash, GetTypeHash(Desc.LodScale));
-	return HashCombineFast(Hash, GetArrayHash(Desc.Tags.GetData(), Desc.Tags.Num()));
+	Hash = HashCombineFast(Hash, GetArrayHash(Desc.Tags.GetData(), Desc.Tags.Num()));
+	return HashCombineFast(Hash, GetTypeHash(Desc.StatsCategory));
 }
 
 /**
diff --git a/Engine/Source/Runtime/Landscape/Private/Landscape.cpp b/Engine/Source/Runtime/Landscape/Private/Landscape.cpp
--- a/Engine/Source/Runtime/Landscape/Private/Landscape.cpp
+++ b/Engine/Source/Runtime/Landscape/Private/Landscape.cpp
@@ -1101,6 +1101,9 @@ void ULandscapeComponent::PostLoad()
 		}
 	}
 
+	static FName NAME_StatsCategoryLandscape("Landscape");
+	SetMeshDrawCommandStatsCategory(NAME_StatsCategoryLandscape);
+
 #if WITH_EDITOR
 	ALandscapeProxy* LandscapeProxy = GetLandscapeProxy();
 	if (ensure(LandscapeProxy))
diff --git a/Engine/Source/Runtime/Landscape/Private/LandscapeGrass.cpp b/Engine/Source/Runtime/Landscape/Private/LandscapeGrass.cpp
--- a/Engine/Source/Runtime/Landscape/Private/LandscapeGrass.cpp
+++ b/Engine/Source/Runtime/Landscape/Private/LandscapeGrass.cpp
@@ -3768,6 +3768,9 @@ void ALandscapeProxy::UpdateGrass(const TArray<FVector>& Cameras, int32& InOutNu
 										GrassInstancedStaticMeshComponent->bEvaluateWorldPositionOffset = true;
 										GrassInstancedStaticMeshComponent->WorldPositionOffsetDisableDistance = GrassVariety.InstanceWorldPositionOffsetDisableDistance;
 										GrassInstancedStaticMeshComponent->ShadowCacheInvalidationBehavior = GrassVariety.ShadowCacheInvalidationBehavior;
+
+										static FName NAME_StatsCategoryGrass("Grass");
+										GrassInstancedStaticMeshComponent->SetMeshDrawCommandStatsCategory(NAME_StatsCategoryGrass);
 										
 										GrassInstancedStaticMeshComponent->PrecachePSOs();
 
