Geometry Collection : Add ability to disable root proxy rendering at runtime when using custom renderer

diff --git a/Engine/Source/Runtime/Experimental/GeometryCollectionEngine/Private/GeometryCollection/GeometryCollectionComponent.cpp b/Engine/Source/Runtime/Experimental/GeometryCollectionEngine/Private/GeometryCollection/GeometryCollectionComponent.cpp
--- a/Engine/Source/Runtime/Experimental/GeometryCollectionEngine/Private/GeometryCollection/GeometryCollectionComponent.cpp
+++ b/Engine/Source/Runtime/Experimental/GeometryCollectionEngine/Private/GeometryCollection/GeometryCollectionComponent.cpp
@@ -4849,6 +4849,12 @@ void UGeometryCollectionComponent::InitializeEmbeddedGeometry()
 	}
 }
 
+void UGeometryCollectionComponent::EnableRootProxyForCustomRenderer(bool bEnable)
+{ 
+	bEnableRootProxyForCustomRenderer = bEnable;
+	RefreshCustomRenderer();
+}
+
 bool UGeometryCollectionComponent::CanUseCustomRenderer() const 
 {
 	return bChaos_GC_UseCustomRenderer && CustomRenderer != nullptr && GetWorld()->IsGameWorld();
@@ -4867,15 +4873,11 @@ void UGeometryCollectionComponent::RefreshCustomRenderer()
 				const int32 RootIndex = GetRootIndex();
 
 				const bool bIsBroken = DynamicCollection ? !DynamicCollection->Active[RootIndex] : false;
+				const bool bRenderRootProxy = bEnableRootProxyForCustomRenderer && !bIsBroken;
 
-				RendererInterface->UpdateState(*RestCollection, ComponentTransform, bIsBroken, !bHiddenInGame);
+				RendererInterface->UpdateState(*RestCollection, ComponentTransform, !bRenderRootProxy, !bHiddenInGame);
 
-				if (bIsBroken)
-				{
-					CalculateGlobalMatrices();
-					RendererInterface->UpdateTransforms(*RestCollection, ComponentSpaceTransforms);
-				}
-				else
+				if (bRenderRootProxy)
 				{
 					// No need to compute the component space transform in that case, since the root is always in component space 
 					const FTransform RootTransform = 
@@ -4885,6 +4887,12 @@ void UGeometryCollectionComponent::RefreshCustomRenderer()
 					ComponentSpaceTransforms[RootIndex] = RootTransform;
 					RendererInterface->UpdateRootTransform(*RestCollection, RootTransform);
 				}
+				else
+				{
+					// render all individual pieces
+					CalculateGlobalMatrices();
+					RendererInterface->UpdateTransforms(*RestCollection, ComponentSpaceTransforms);
+				}
 			}
 		}
 	}
diff --git a/Engine/Source/Runtime/Experimental/GeometryCollectionEngine/Public/GeometryCollection/GeometryCollectionComponent.h b/Engine/Source/Runtime/Experimental/GeometryCollectionEngine/Public/GeometryCollection/GeometryCollectionComponent.h
--- a/Engine/Source/Runtime/Experimental/GeometryCollectionEngine/Public/GeometryCollection/GeometryCollectionComponent.h
+++ b/Engine/Source/Runtime/Experimental/GeometryCollectionEngine/Public/GeometryCollection/GeometryCollectionComponent.h
@@ -1097,6 +1097,9 @@ public:
 
 	GEOMETRYCOLLECTIONENGINE_API void SetUpdateNavigationInTick(const bool bUpdateInTick) { bUpdateNavigationInTick = bUpdateInTick; }
 
+	/** Enable or disable root proxy for custom rendering - this can be set at runtime */
+	GEOMETRYCOLLECTIONENGINE_API void EnableRootProxyForCustomRenderer(bool bEnable);
+
 	/** Force all GC components to reregister their custom renderer objects. */
 	static GEOMETRYCOLLECTIONENGINE_API void ReregisterAllCustomRenderers();
 
@@ -1194,6 +1197,9 @@ protected:
 	UPROPERTY(EditAnywhere, BlueprintReadWrite, Category = "ChaosPhysics|Rendering", meta = (editcondition = "bOverrideCustomRenderer", MustImplement = "/Script/GeometryCollectionEngine.GeometryCollectionExternalRenderInterface"))
 	TObjectPtr<UClass> CustomRendererType;
 
+	UPROPERTY()
+	bool bEnableRootProxyForCustomRenderer = true;
+
 	/** A custom renderer object created from CustomRenderType. */
 	UPROPERTY(Transient)
 	TObjectPtr<UGeometryCollectionExternalRenderInterface> CustomRenderer;
