ChaosFlesh: Dynamic collision springs

diff --git a/Engine/Plugins/Experimental/ChaosFlesh/Source/ChaosFleshEngine/Private/ChaosFlesh/ChaosDeformableSolverComponent.cpp b/Engine/Plugins/Experimental/ChaosFlesh/Source/ChaosFleshEngine/Private/ChaosFlesh/ChaosDeformableSolverComponent.cpp
--- a/Engine/Plugins/Experimental/ChaosFlesh/Source/ChaosFleshEngine/Private/ChaosFlesh/ChaosDeformableSolverComponent.cpp
+++ b/Engine/Plugins/Experimental/ChaosFlesh/Source/ChaosFleshEngine/Private/ChaosFlesh/ChaosDeformableSolverComponent.cpp
@@ -217,6 +217,9 @@ void UDeformableSolverComponent::Reset()
 			, SolverConstraints.GaussSeidelConstraints.bUseSOR
 			, SolverConstraints.GaussSeidelConstraints.OmegaSOR
 			, SolverConstraints.GaussSeidelConstraints.bUseGSNeohookean
+			, SolverConstraints.GaussSeidelConstraints.CollisionSpring.CollisionSearchRadius
+			, SolverConstraints.GaussSeidelConstraints.CollisionSpring.CollisionSpringStiffness
+			, SolverConstraints.GaussSeidelConstraints.CollisionSpring.bAllowSliding
 		}));
 
 		for (TObjectPtr<UDeformablePhysicsComponent>& DeformableComponent : ConnectedObjects.DeformableComponents)
diff --git a/Engine/Plugins/Experimental/ChaosFlesh/Source/ChaosFleshEngine/Public/ChaosFlesh/ChaosDeformableSolverComponent.h b/Engine/Plugins/Experimental/ChaosFlesh/Source/ChaosFleshEngine/Public/ChaosFlesh/ChaosDeformableSolverComponent.h
--- a/Engine/Plugins/Experimental/ChaosFlesh/Source/ChaosFleshEngine/Public/ChaosFlesh/ChaosDeformableSolverComponent.h
+++ b/Engine/Plugins/Experimental/ChaosFlesh/Source/ChaosFleshEngine/Public/ChaosFlesh/ChaosDeformableSolverComponent.h
@@ -88,10 +88,31 @@ struct FSolverGridBasedCollisionsGroup
 	GENERATED_USTRUCT_BODY()
 
 	UPROPERTY(EditAnywhere, Category = "GridBasedCollisions")
-		bool bUseGridBasedConstraints = false;
+	bool bUseGridBasedConstraints = false;
 
 	UPROPERTY(EditAnywhere, Category = "GridBasedCollisions")
-		float GridDx = 25.;
+	float GridDx = 25.;
+};
+
+USTRUCT(BlueprintType)
+struct FCollisionSpringGroup
+{
+	GENERATED_USTRUCT_BODY()
+	/**
+	* Search radius for point triangle pairs
+	*/
+	UPROPERTY(EditAnywhere, Category = "CollisionSpring")
+	float CollisionSearchRadius = 10.f;
+	/**
+	* Collision spring stiffness; larger value will stop penetration better
+	*/
+	UPROPERTY(EditAnywhere, Category = "CollisionSpring")
+	float CollisionSpringStiffness = 500.f;
+	/**
+	* Anisotropic springs will allow sliding on the triangle
+	*/
+	UPROPERTY(EditAnywhere, Category = "CollisionSpring")
+	bool bAllowSliding = true;
 };
 
 USTRUCT(BlueprintType)
@@ -103,27 +124,32 @@ struct FSolverGaussSeidelConstraintsGroup
 	* Enable the Gauss Seidel solver instead of the existing XPBD.
 	*/
 	UPROPERTY(EditAnywhere, Category = "GaussSeidelConstraints")
-		bool bUseGaussSeidelConstraints = false;
+	bool bUseGaussSeidelConstraints = false;
 
 	/**
 	* Enable another model that runs simulation faster.
 	*/
 	UPROPERTY(EditAnywhere, Category = "GaussSeidelConstraints")
-		bool bUseGSNeohookean = false;
+	bool bUseGSNeohookean = false;
 
 	/**
 	* Enable acceleration technique for Gauss Seidel solver to make simulation look better within a limited budget.
 	*/
 	UPROPERTY(EditAnywhere, Category = "GaussSeidelConstraints")
-		bool bUseSOR = true;
+	bool bUseSOR = true;
 
 	/**
 	* Acceleration related parameter. Tune it down if simulation becomes unstable. 
 	*/
 	UPROPERTY(EditAnywhere, Category = "GaussSeidelConstraints")
-		float OmegaSOR = 1.6f;
-};
+	float OmegaSOR = 1.6f;
 
+	/**
+	* Collsion detection radius and stiffness
+	*/
+	UPROPERTY(EditAnywhere, BlueprintReadOnly, Category = "GaussSeidelConstraints")
+	FCollisionSpringGroup CollisionSpring;
+};
 
 USTRUCT(BlueprintType)
 struct FSolverCollisionsGroup
diff --git a/Engine/Source/Runtime/Experimental/Chaos/Private/Chaos/Deformable/ChaosDeformableSolver.cpp b/Engine/Source/Runtime/Experimental/Chaos/Private/Chaos/Deformable/ChaosDeformableSolver.cpp
--- a/Engine/Source/Runtime/Experimental/Chaos/Private/Chaos/Deformable/ChaosDeformableSolver.cpp
+++ b/Engine/Source/Runtime/Experimental/Chaos/Private/Chaos/Deformable/ChaosDeformableSolver.cpp
@@ -1049,7 +1049,7 @@ namespace Chaos::Softs
 
 											for (int32 i = 0; i < BoneIndices.Num(); i++)
 											{
-												if (BoneIndices[i] > -1 && BoneIndices[i] < RestNum && BoneIndices[i] < TransformNum)
+												if (BoneIndices[i] > INDEX_NONE && BoneIndices[i] < RestNum && BoneIndices[i] < TransformNum)
 												{
 
 													// @todo(flesh) : Add the pre-cached component space rest transforms to the rest collection. 
@@ -1110,33 +1110,61 @@ namespace Chaos::Softs
 	void FDeformableSolver::InitializeSelfCollisionVariables()
 	{
 		PERF_SCOPE(STAT_ChaosDeformableSolver_InitializeSelfCollisionVariables);
+		int32 VertexOffset = 0;
+		SurfaceElements->SetNum(0);
+		int ComponentOffset = 0;
 
+		for (FThreadingProxy* InProxy : UninitializedProxys_Internal)
+		{
+			if (FFleshThreadingProxy* Proxy = InProxy->As<FFleshThreadingProxy>())
+			{
+				if (const FManagedArrayCollection* Rest = &Proxy->GetRestCollection())
+				{
+					if (const TManagedArray<FVector3f>* Vertex = Rest->FindAttribute<FVector3f>("Vertex", FGeometryCollection::VerticesGroup))
+					{
+						if (const TManagedArray<FIntVector>* Indices = Rest->FindAttribute<FIntVector>("Indices", FGeometryCollection::FacesGroup))
+						{
+							int32 SurfaceOffset = SurfaceElements->Num();
+							SurfaceElements->SetNum(SurfaceOffset + Indices->Num());
+							for (int32 i = 0; i < Indices->Num(); i++)
+							{
+								for (int32 j = 0; j < 3; j++)
+								{
+									(*SurfaceElements)[i + SurfaceOffset][j] = VertexOffset + (*Indices)[i][j];
+								}
+							}
+						}
+						int32 Offset = ParticleComponentIndex.Num();
+						ParticleComponentIndex.SetNum(ParticleComponentIndex.Num() + Vertex->Num());
+						for (int32 i = 0; i < Vertex->Num(); i++) {
+							ParticleComponentIndex[i + Offset] = ComponentOffset;
+						}
+						int NewComponentOffset = ComponentOffset;
+						if (const TManagedArray<int32>* ComponentIndex = Rest->FindAttribute<int32>("ComponentIndex", FGeometryCollection::VerticesGroup))
+						{
+							ensureMsgf(ComponentIndex->Num() == Vertex->Num(), TEXT("ComponentIndex size is not equal to vertex size"));
+							for (int32 i = 0; i < ComponentIndex->Num(); i++) {
+								if ((*ComponentIndex)[i] < 0)
+								{
+									ParticleComponentIndex[i + Offset] = (*ComponentIndex)[i]; //Isolated Nodes
+								}
+								else
+								{
+									ParticleComponentIndex[i + Offset] = ComponentOffset + (*ComponentIndex)[i];
+									NewComponentOffset = NewComponentOffset < ParticleComponentIndex[i + Offset] ? ParticleComponentIndex[i + Offset] : NewComponentOffset;
+								}
+							}
+						}
+						ComponentOffset = NewComponentOffset + 1;
+						VertexOffset += Vertex->Num();
+					}
+				}
+			}
+		}
 
-		int32 NumParticles = Evolution->Particles().Size();
 		SurfaceTriangleMesh->Init(*SurfaceElements);
 		TriangleMeshCollisions.Reset(new FPBDTriangleMeshCollisions(
 			0, Evolution->Particles().Size(), *SurfaceTriangleMesh, false, false));
-		TSet<Chaos::TVec2<int32>>* InDisabledCollisionElements = new TSet<Chaos::TVec2<int32>>();
-		for (int32 i = 0; i < (int32)NumParticles; i++)
-		{
-			Chaos::TVec2<int32> LocalEdge = { i, i };
-			InDisabledCollisionElements->Add(LocalEdge);
-		}
-		CollisionSpringConstraint.Reset(new FPBDCollisionSpringConstraints(0, (int32)NumParticles, *SurfaceTriangleMesh, nullptr, MoveTemp(*InDisabledCollisionElements), 1.f, 1.f));
-		int32 InitIndex1 = Evolution->AddConstraintInitRange(1, true);
-		Evolution->ConstraintInits()[InitIndex1] =
-			[this](FSolverParticles& InParticles, const FSolverReal Dt)
-		{
-			this->TriangleMeshCollisions->Init(InParticles);
-			TArray<FPBDTriangleMeshCollisions::FGIAColor> EmptyGIAColors;
-			this->CollisionSpringConstraint->Init(InParticles, TriangleMeshCollisions->GetSpatialHash(), static_cast<TConstArrayView<FPBDTriangleMeshCollisions::FGIAColor>>(EmptyGIAColors), EmptyGIAColors);
-		};
-		int32 ConstraintIndex1 = Evolution->AddConstraintRuleRange(1, true);
-		Evolution->ConstraintRules()[ConstraintIndex1] =
-			[this](FSolverParticles& InParticles, const FSolverReal Dt)
-		{
-			this->CollisionSpringConstraint->Apply(InParticles, Dt);
-		};
 	}
 
 	void FDeformableSolver::InitializeGridBasedConstraintVariables()
@@ -1250,11 +1278,19 @@ namespace Chaos::Softs
 
 		if (Property.bDoSelfCollision)
 		{
+			this->GSWeakConstraints->UpdateBoundaryVertices(*SurfaceElements);
+			
 			int32 InitIndex = Evolution->AddConstraintInitRange(1, true);
 			Evolution->ConstraintInits()[InitIndex] =
 				[this](FSolverParticles& InParticles, const FSolverReal Dt)
 			{
 				//TODO(Yushan & Joey): Add the collision detection code in here:
+				//this->GSWeakConstraints->CollisionDetection(this->Evolution->Particles(), *SurfaceElements, ParticleComponentIndex, Property.CollisionSearchRadius, Property.CollisionSpringStiffness, Property.bAllowSliding);
+				
+				this->TriangleMeshCollisions->Init(InParticles);
+				TArray<FPBDTriangleMeshCollisions::FGIAColor> EmptyGIAColors;
+				this->GSWeakConstraints->CollisionDetectionSpatialHash(this->Evolution->Particles(), *SurfaceTriangleMesh, ParticleComponentIndex, TriangleMeshCollisions->GetSpatialHash(), Property.CollisionSearchRadius, Property.CollisionSpringStiffness, Property.bAllowSliding);
+				
 			};
 
 			int32 InitIndex1 = Evolution->AddConstraintInitRange(1, true);
diff --git a/Engine/Source/Runtime/Experimental/Chaos/Public/Chaos/Deformable/ChaosDeformableSolver.h b/Engine/Source/Runtime/Experimental/Chaos/Public/Chaos/Deformable/ChaosDeformableSolver.h
--- a/Engine/Source/Runtime/Experimental/Chaos/Public/Chaos/Deformable/ChaosDeformableSolver.h
+++ b/Engine/Source/Runtime/Experimental/Chaos/Public/Chaos/Deformable/ChaosDeformableSolver.h
@@ -198,7 +198,7 @@ namespace Chaos::Softs
 		TUniquePtr <TArray<TArray<int32>>> AllSecondIndices;
 		TUniquePtr <TArray<FSolverReal>> AllWeights;
 		TUniquePtr <TArray<FSolverReal>> AllSecondWeights;
-
+		TArray<int32> ParticleComponentIndex;
 		//typedef TMap<int32, TTuple<float, Chaos::Softs::FPAndInvM, FVector3f>> TransientConstraintBufferMap;
 		typedef TMap<int32, TTuple<float, float, FVector3f>> TransientConstraintBufferMap;
 		TransientConstraintBufferMap TransientConstraintBuffer;
diff --git a/Engine/Source/Runtime/Experimental/Chaos/Public/Chaos/Deformable/ChaosDeformableSolverTypes.h b/Engine/Source/Runtime/Experimental/Chaos/Public/Chaos/Deformable/ChaosDeformableSolverTypes.h
--- a/Engine/Source/Runtime/Experimental/Chaos/Public/Chaos/Deformable/ChaosDeformableSolverTypes.h
+++ b/Engine/Source/Runtime/Experimental/Chaos/Public/Chaos/Deformable/ChaosDeformableSolverTypes.h
@@ -37,7 +37,10 @@ namespace Chaos::Softs
 			bool InbUseGaussSeidelConstraints = false, 
 			bool InbUseSOR = true,
 			FSolverReal InOmegaSOR = (FSolverReal)1.6, 
-			bool InbUseGSNeohookean = false)
+			bool InbUseGSNeohookean = false,
+			FSolverReal InCollisionSearchRadius = (FSolverReal)10.,
+			FSolverReal InCollisionSpringStiffness = (FSolverReal)500.,
+			bool InbAllowSliding = true)
 			: NumSolverSubSteps(InNumSolverSubSteps)
 			, NumSolverIterations(InNumSolverIterations)
 			, FixTimeStep(InFixTimeStep)
@@ -60,6 +63,9 @@ namespace Chaos::Softs
 			, bUseSOR(InbUseSOR)
 			, OmegaSOR(InOmegaSOR)
 			, bUseGSNeohookean(InbUseGSNeohookean)
+			, CollisionSearchRadius(InCollisionSearchRadius)
+			, CollisionSpringStiffness(InCollisionSpringStiffness)
+			, bAllowSliding(InbAllowSliding)
 		{}
 
 		int32 NumSolverSubSteps = 5;
@@ -84,6 +90,9 @@ namespace Chaos::Softs
 		bool bUseSOR = true;
 		FSolverReal OmegaSOR = (FSolverReal)1.6;
 		bool bUseGSNeohookean = false;
+		FSolverReal CollisionSearchRadius = (FSolverReal)10.;
+		FSolverReal CollisionSpringStiffness = (FSolverReal)500.;
+		bool bAllowSliding = true;
 	};
 
 
diff --git a/Engine/Source/Runtime/Experimental/Chaos/Public/Chaos/Deformable/GaussSeidelWeakConstraints.h b/Engine/Source/Runtime/Experimental/Chaos/Public/Chaos/Deformable/GaussSeidelWeakConstraints.h
--- a/Engine/Source/Runtime/Experimental/Chaos/Public/Chaos/Deformable/GaussSeidelWeakConstraints.h
+++ b/Engine/Source/Runtime/Experimental/Chaos/Public/Chaos/Deformable/GaussSeidelWeakConstraints.h
@@ -1,6 +1,7 @@
 // Copyright Epic Games, Inc. All Rights Reserved.
 #pragma once
 
+#include "Chaos/BoundingVolumeHierarchy.h"
 #include "Chaos/PBDSpringConstraintsBase.h"
 #include "Chaos/XPBDCorotatedConstraints.h"
 #include "ChaosStats.h"
@@ -11,7 +12,9 @@
 #include "Chaos/Utilities.h"
 #include "Chaos/DebugDrawQueue.h"
 #include "Chaos/XPBDWeakConstraints.h"
-
+#include "Chaos/Triangle.h"
+#include "Chaos/TriangleCollisionPoint.h"
+#include "Chaos/TriangleMesh.h"
 
 namespace Chaos::Softs
 {
@@ -51,7 +54,7 @@ namespace Chaos::Softs
 		void ComputeInitialWCData(const ParticleType& InParticles, const TArray<TVector<int32, 4>>& MeshConstraints, const TArray <TArray<int32>>& MeshIncidentElements, const TArray <TArray<int32>>& MeshIncidentElementsLocal, TArray<TArray<int32>>& ParticlesPerColor)
 		{
 			TArray<TVector<int32, 4>> ExtraConstraints;
-			ExtraConstraints.Init(TVector<int32, 4>(-1), Indices.Num());
+			ExtraConstraints.Init(TVector<int32, 4>(INDEX_NONE), Indices.Num());
 			for (int32 i = 0; i < Indices.Num(); i++)
 			{
 				ensureMsgf(Indices[i].Num() + SecondIndices[i].Num() == 4, TEXT("Currently does not support non point tri pair weak constraints"));
@@ -127,7 +130,8 @@ namespace Chaos::Softs
 					}
 				}
 			}
-
+			NoCollisionNodalWeights = NodalWeights;
+			NoCollisionConstraints = ExtraConstraints;
 			InitialWCSize = Indices.Num();
 		}
 
@@ -231,6 +235,55 @@ namespace Chaos::Softs
 			}
 
 			IsAnisotropic.Init(false, Indices.Num());
+			Normals.SetNum(Indices.Num());
+		}
+
+		void Resize(int32 Size)
+		{
+			Indices.SetNum(Size);
+			SecondIndices.SetNum(Size);
+			Stiffness.SetNum(Size);
+			Weights.SetNum(Size);
+			SecondWeights.SetNum(Size);
+			IsAnisotropic.SetNum(Size);
+			Normals.SetNum(Size);
+		}
+
+		void UpdatePointTriangleCollisionWCData(const FSolverParticles& Particles)
+		{	
+			TArray<TArray<int32>> IndicesTemp = Indices;
+			TArray<TArray<int32>> SecondIndicesTemp = SecondIndices;
+			TArray<T> StiffnessTemp = Stiffness;
+			TArray<TArray<T>> WeightsTemp = Weights;
+			TArray<TArray<T>> SecondWeightsTemp = SecondWeights;
+			TArray<bool> IsAnisotropicTemp = IsAnisotropic;
+			TArray<TVector<T, 3>> NormalsTemp = Normals;
+
+			Indices.SetNum(InitialWCSize);
+			SecondIndices.SetNum(InitialWCSize);
+			Stiffness.SetNum(InitialWCSize);
+			Weights.SetNum(InitialWCSize);
+			SecondWeights.SetNum(InitialWCSize);
+			IsAnisotropic.SetNum(InitialWCSize);
+			Normals.SetNum(InitialWCSize);
+
+			for (int32 i = InitialWCSize; i < IndicesTemp.Num(); i++)
+			{
+				ensureMsgf(IndicesTemp[i].Num() == 3, TEXT("Collision format is not point-triangle"));
+				ensureMsgf(SecondIndicesTemp[i].Num() == 1, TEXT("Collision format is not point-triangle"));
+				Chaos::TVector<float, 3> TriPos0(Particles.P(IndicesTemp[i][0])), TriPos1(Particles.P(IndicesTemp[i][1])), TriPos2(Particles.P(IndicesTemp[i][2])), ParticlePos(Particles.P(SecondIndicesTemp[i][0]));
+				Chaos::TVector<T, 3> Normal = FVector3f::CrossProduct(TriPos1 - TriPos0, TriPos2 - TriPos0);
+				if (FVector3f::DotProduct(ParticlePos - TriPos0, Normal) < 0.f) //not resolved, keep the spring
+				{
+					Indices.Add(IndicesTemp[i]);
+					SecondIndices.Add(SecondIndicesTemp[i]);
+					Stiffness.Add(StiffnessTemp[i]);
+					Weights.Add(WeightsTemp[i]);
+					SecondWeights.Add(SecondWeightsTemp[i]);
+					IsAnisotropic.Add(IsAnisotropicTemp[i]);
+					Normals.Add(NormalsTemp[i]);
+				}
+			}
 		}
 
 		void VisualizeAllBindings(const FSolverParticles& InParticles, const T Dt) const
@@ -287,20 +340,262 @@ namespace Chaos::Softs
 
 		}
 
+		void UpdateBoundaryVertices(const TArray<TVec3<int32>> Elements)
+		{
+			for (int32 i = 0; i < Elements.Num(); i++)
+			{
+				for (int32 j = 0; j < 3; j++)
+				{
+					BoundaryVertices.AddUnique(Elements[i][j]);
+				}
+			}
+		}
+
+		//CollisionDetectionSpatialHash should be faster than CollisionDetectionBVH
+		void CollisionDetectionBVH(const FSolverParticles& Particles, const TArray<TVec3<int32>>& SurfaceElements, const TArray<int32>& ComponentIndex, float DetectRadius = 1.f, float PositionTargetStiffness = 10000.f, bool UseAnisotropicSpring = true)
+		{
+			TRACE_CPUPROFILER_EVENT_SCOPE(STAT_ChaosGaussSeidelWeakConstraintsCollisionDetection);
+			Resize(InitialWCSize);
+
+			TArray<Chaos::TVector<int32, 3>> SurfaceElementsArray;
+			for (int32 i = 0; i < SurfaceElements.Num(); i++)
+			{
+				Chaos::TVector<int32, 3> CurrentSurfaceElements(0);
+				for (int32 j = 0; j < 3; j++)
+				{
+					CurrentSurfaceElements[j] = SurfaceElements[i][j];
+				}
+				if (CurrentSurfaceElements[0] != INDEX_NONE
+					&& CurrentSurfaceElements[1] != INDEX_NONE
+					&& CurrentSurfaceElements[2] != INDEX_NONE)
+				{
+					SurfaceElementsArray.Emplace(CurrentSurfaceElements);
+				}
+			}
+			TArray<TArray<int32>> LocalIndex;
+			TArray<TArray<int32>>* LocalIndexPtr = &LocalIndex;
+			TArray<TArray<int>> GlobalIndex = Chaos::Utilities::ComputeIncidentElements(SurfaceElementsArray, LocalIndexPtr);
+			int32 ActualParticleCount = 0;
+			for (int32 l = 0; l < GlobalIndex.Num(); l++)
+			{
+				if (GlobalIndex[l].Num() > 0)
+				{
+					ActualParticleCount += 1;
+				}
+			}
+			TArray<Chaos::TVector<float, 3>> SurfaceElementsPositions;
+			SurfaceElementsPositions.SetNum(ActualParticleCount);
+			TArray<int32> SurfaceElementsMap;
+			SurfaceElementsMap.SetNum(ActualParticleCount);
+			int32 CurrentParticleIndex = 0;
+			for (int32 i = 0; i < GlobalIndex.Num(); i++)
+			{
+				if (GlobalIndex[i].Num() > 0)
+				{
+					SurfaceElementsPositions[CurrentParticleIndex] = Particles.P(SurfaceElements[GlobalIndex[i][0]][LocalIndex[i][0]]);
+					SurfaceElementsMap[CurrentParticleIndex] = SurfaceElements[GlobalIndex[i][0]][LocalIndex[i][0]];
+					CurrentParticleIndex += 1;
+				}
+			}
+
+			TArray<Chaos::TSphere<Chaos::FReal, 3>*> VertexSpherePtrs;
+			TArray<Chaos::TSphere<Chaos::FReal, 3>> VertexSpheres;
+
+			VertexSpheres.Init(Chaos::TSphere<Chaos::FReal, 3>(Chaos::TVec3<Chaos::FReal>(0), DetectRadius), SurfaceElementsPositions.Num());
+			VertexSpherePtrs.SetNum(SurfaceElementsPositions.Num());
+
+			for (int32 i = 0; i < SurfaceElementsPositions.Num(); i++)
+			{
+				Chaos::TVec3<Chaos::FReal> SphereCenter(SurfaceElementsPositions[i]);
+				Chaos::TSphere<Chaos::FReal, 3> VertexSphere(SphereCenter, DetectRadius);
+				VertexSpheres[i] = Chaos::TSphere<Chaos::FReal, 3>(SphereCenter, DetectRadius);
+				VertexSpherePtrs[i] = &VertexSpheres[i];
+			}
+			Chaos::TBoundingVolumeHierarchy<
+				TArray<Chaos::TSphere<Chaos::FReal, 3>*>,
+				TArray<int32>,
+				Chaos::FReal,
+				3> VertexBVH(VertexSpherePtrs);
+
+			for (int32 i = 0; i < SurfaceElements.Num(); i++)
+			{
+				TArray<int32> TriangleIntersections0 = VertexBVH.FindAllIntersections(Particles.P(SurfaceElements[i][0]));
+				TArray<int32> TriangleIntersections1 = VertexBVH.FindAllIntersections(Particles.P(SurfaceElements[i][1]));
+				TArray<int32> TriangleIntersections2 = VertexBVH.FindAllIntersections(Particles.P(SurfaceElements[i][2]));
+				TriangleIntersections0.Sort();
+				TriangleIntersections1.Sort();
+				TriangleIntersections2.Sort();
+
+				TArray<int32> TriangleIntersections({});
+				for (int32 k = 0; k < TriangleIntersections0.Num(); k++)
+				{
+					if (TriangleIntersections1.Contains(TriangleIntersections0[k])
+						&& TriangleIntersections2.Contains(TriangleIntersections0[k]))
+					{
+						TriangleIntersections.Emplace(TriangleIntersections0[k]);
+					}
+				}
+
+				int32 TriangleIndex = ComponentIndex[SurfaceElements[i][0]];
+				int32 MinIndex = INDEX_NONE;
+				float MinDis = DetectRadius;
+				Chaos::TVector<float, 3> ClosestBary(0.f);
+				Chaos::TVector<float, 3> FaceNormal;
+				for (int32 j = 0; j < TriangleIntersections.Num(); j++)
+				{
+					if (ComponentIndex[SurfaceElementsMap[TriangleIntersections[j]]] >= 0 && TriangleIndex >= 0 && ComponentIndex[SurfaceElementsMap[TriangleIntersections[j]]] != TriangleIndex)
+					{
+						Chaos::TVector<float, 3> Bary, TriPos0(Particles.P(SurfaceElements[i][0])), TriPos1(Particles.P(SurfaceElements[i][1])), TriPos2(Particles.P(SurfaceElements[i][2])), ParticlePos(Particles.P(SurfaceElementsMap[TriangleIntersections[j]]));
+						Chaos::TVector<Chaos::FRealSingle, 3> ClosestPoint = Chaos::FindClosestPointAndBaryOnTriangle(TriPos0, TriPos1, TriPos2, ParticlePos, Bary);
+						Chaos::FRealSingle CurrentDistance = (Particles.P(SurfaceElementsMap[TriangleIntersections[j]]) - ClosestPoint).Size();
+						if (CurrentDistance < MinDis)
+						{
+							Chaos::TVector<T, 3> Normal = FVector3f::CrossProduct(TriPos2 - TriPos0, TriPos1 - TriPos0); //The normal needs to point outwards of the geometry
+							if (FVector3f::DotProduct(ParticlePos - TriPos0, Normal) < 0.f)
+							{
+								Normal.SafeNormalize(1e-8f);
+								MinDis = CurrentDistance;
+								MinIndex = SurfaceElementsMap[TriangleIntersections[j]];
+								ClosestBary = Bary;
+								FaceNormal = Normal;
+							}
+						}
+
+					}
+				}
+				if (MinIndex != INDEX_NONE
+					&& MinIndex != SurfaceElements[i][0]
+					&& MinIndex != SurfaceElements[i][1]
+					&& MinIndex != SurfaceElements[i][2])
+				{
+					Indices.SetNum(Indices.Num() + 1);
+					SecondIndices.SetNum(SecondIndices.Num() + 1);
+					Weights.SetNum(Weights.Num() + 1);
+					SecondWeights.SetNum(SecondWeights.Num() + 1);
+					Indices[Indices.Num() - 1].Add(SurfaceElements[i][0]);
+					Indices[Indices.Num() - 1].Add(SurfaceElements[i][1]);
+					Indices[Indices.Num() - 1].Add(SurfaceElements[i][2]);
+					SecondIndices[SecondIndices.Num() - 1].Add(MinIndex);
+					Weights[Weights.Num() - 1].Add(ClosestBary[0]);
+					Weights[Weights.Num() - 1].Add(ClosestBary[1]);
+					Weights[Weights.Num() - 1].Add(ClosestBary[2]);
+					SecondWeights[SecondWeights.Num() - 1].Add(1.f);
+
+					float SpringStiffness = 0.f;
+					for (int32 k = 0; k < 3; k++)
+					{
+						SpringStiffness += ClosestBary[k] * PositionTargetStiffness * Particles.M(SurfaceElements[i][k]);
+					}
+					SpringStiffness += PositionTargetStiffness * Particles.M(MinIndex);
+					Stiffness.Add(SpringStiffness);
+					IsAnisotropic.Add(UseAnisotropicSpring);
+					Normals.Add(FaceNormal);
+				}
+			}
+		}
+
+		template<typename SpatialAccelerator>
+		void CollisionDetectionSpatialHash(const FSolverParticles& Particles, const FTriangleMesh& TriangleMesh, const TArray<int32>& ComponentIndex, const SpatialAccelerator& Spatial, float DetectRadius = 1.f, float PositionTargetStiffness = 10000.f, bool UseAnisotropicSpring = true)
+		{
+			TRACE_CPUPROFILER_EVENT_SCOPE(STAT_ChaosGaussSeidelWeakConstraintsCollisionDetectionSpatialHash);
+			constexpr int32 MaxConnectionsPerPoint = 1;
+			Resize(InitialWCSize + Particles.Size() * MaxConnectionsPerPoint);
+			
+			std::atomic<int32> ConstraintIndex(InitialWCSize);
+			const TArray<TVec3<int32>>& Elements = TriangleMesh.GetSurfaceElements();
+			PhysicsParallelFor(BoundaryVertices.Num(),
+				[this, &Spatial, &Particles, &ConstraintIndex, MaxConnectionsPerPoint, &TriangleMesh, &Elements, &DetectRadius, &ComponentIndex, &PositionTargetStiffness, &UseAnisotropicSpring](int32 i)
+				{
+					const int32 Index = BoundaryVertices[i];
+
+					TArray< TTriangleCollisionPoint<FSolverReal> > Result;
+					if (TriangleMesh.PointProximityQuery(Spatial, static_cast<const TArrayView<const FSolverVec3>&>(Particles.XArray()), Index, Particles.X(Index), DetectRadius, DetectRadius,
+						[this, &ComponentIndex, &Elements](const int32 PointIndex, const int32 TriangleIndex)->bool
+						{
+							const TVector<int32, 3>& Elem = Elements[TriangleIndex];
+
+							if (ComponentIndex[PointIndex] == ComponentIndex[Elements[TriangleIndex][0]])
+							{
+								return false;
+							}
+
+							return true;
+						},
+						Result))
+					{
+
+						if (Result.Num() > MaxConnectionsPerPoint)
+						{
+							// TODO: once we have a PartialSort, use that instead here.
+							Result.Sort(
+								[](const TTriangleCollisionPoint<FSolverReal>& First, const TTriangleCollisionPoint<FSolverReal>& Second)->bool
+								{
+									return First.Phi < Second.Phi;
+								}
+							);
+							Result.SetNum(MaxConnectionsPerPoint, false /*bAllowShrinking*/);
+						}
+
+						for (const TTriangleCollisionPoint<FSolverReal>& CollisionPoint : Result)
+						{
+							const TVector<int32, 3>& Elem = Elements[CollisionPoint.Indices[1]];
+						
+
+							// NOTE: CollisionPoint.Normal has already been flipped to point toward the Point, so need to recalculate here.
+							const TTriangle<FSolverReal> Triangle(Particles.X(Elem[0]), Particles.X(Elem[1]), Particles.X(Elem[2]));
+							if ((Particles.X(Index) - CollisionPoint.Location).Dot(-Triangle.GetNormal()) < 0) //Is point inside boundary? Normal should point outwards
+							{
+								const int32 IndexToWrite = ConstraintIndex.fetch_add(1);
+
+								Indices[IndexToWrite] = { Elem[0], Elem[1] ,Elem[2] };
+								SecondIndices[IndexToWrite] = { Index };
+								Weights[IndexToWrite] = { CollisionPoint.Bary[1], CollisionPoint.Bary[2], CollisionPoint.Bary[3] };
+								SecondWeights[IndexToWrite] = { 1.f };
+								/*
+								Indices[IndexToWrite].Add(Elem[0]);
+								Indices[IndexToWrite].Add(Elem[1]);
+								Indices[IndexToWrite].Add(Elem[2]);
+								SecondIndices[IndexToWrite].Add(Index);
+								Weights[IndexToWrite].Add(ClosestBary[0]);
+								Weights[IndexToWrite].Add(ClosestBary[1]);
+								Weights[IndexToWrite].Add(ClosestBary[2]);
+								SecondWeights[SecondWeights.Num() - 1].Add(1.f);
+								*/
+								float SpringStiffness = 0.f;
+								for (int32 k = 0; k < 3; k++)
+								{
+									SpringStiffness += Weights[IndexToWrite][k] * PositionTargetStiffness * Particles.M(Elem[k]);
+								}
+								SpringStiffness += PositionTargetStiffness * Particles.M(Index);
+								Stiffness[IndexToWrite] = SpringStiffness;
+								IsAnisotropic[IndexToWrite] = UseAnisotropicSpring;
+								Normals[IndexToWrite] = Triangle.GetNormal();
+							}							
+						}
+					}
+				}
+			);
+
+			// Shrink the arrays to the actual number of found constraints.
+			const int32 ConstraintNum = ConstraintIndex.load();
+			Resize(ConstraintNum);
+		}
+
 		void ComputeCollisionWCData(const ParticleType& InParticles, const TArray<TVector<int32, 4>>& MeshConstraints, const TArray <TArray<int32>>& MeshIncidentElements, const TArray <TArray<int32>>& MeshIncidentElementsLocal, TArray<TArray<int32>>& ParticlesPerColor)
 		{
+			ensureMsgf(Indices.Num() >= InitialWCSize, TEXT("The size of Indices is smaller than InitialWCSize"));
 			TArray<TVector<int32, 4>> ExtraConstraints;
-			ExtraConstraints.Init(TVector<int32, 4>(-1), Indices.Num());
+			ExtraConstraints.Init(TVector<int32, 4>(INDEX_NONE), Indices.Num() - InitialWCSize);
 			for (int32 i = InitialWCSize; i < Indices.Num(); i++)
 			{
 				ensureMsgf(Indices[i].Num() + SecondIndices[i].Num() == 4, TEXT("Currently does not support non point tri pair weak constraints"));
 				for (int32 j = 0; j < Indices[i].Num(); j++)
 				{
-					ExtraConstraints[i][j] = Indices[i][j];
+					ExtraConstraints[i - InitialWCSize][j] = Indices[i][j];
 				}
 				for (int32 j = 0; j < SecondIndices[i].Num(); j++)
 				{
-					ExtraConstraints[i][j + Indices[i].Num()] = SecondIndices[i][j];
+					ExtraConstraints[i - InitialWCSize][j + Indices[i].Num()] = SecondIndices[i][j];
 				}
 			}
 
@@ -310,14 +605,15 @@ namespace Chaos::Softs
 			TArray<TVector<int32, 4>> TotalConstraints = NoCollisionConstraints;
 			TotalConstraints += ExtraConstraints;
 			//TODO (Yizhou): Make the following more efficient by computing only extra incident elements in the future. 
-			WCIncidentElements = Chaos::Utilities::ComputeIncidentElements(TotalConstraints, &WCIncidentElementsLocal);
-
+			TArray<TArray<int32>> WCIncidentElementsLocalTemp;
+			WCIncidentElements = Chaos::Utilities::ComputeIncidentElements(TotalConstraints, &WCIncidentElementsLocalTemp);
+			WCIncidentElementsLocal = WCIncidentElementsLocalTemp;
 			Particle2WCIndices.Init(INDEX_NONE, InParticles.Size());
 			for (int32 i = 0; i < WCIncidentElements.Num(); i++)
 			{
 				if (WCIncidentElements[i].Num() > 0)
 				{
-					int32 p = ExtraConstraints[WCIncidentElements[i][0]][WCIncidentElementsLocal[i][0]];
+					int32 p = TotalConstraints[WCIncidentElements[i][0]][WCIncidentElementsLocal[i][0]];
 					Particle2WCIndices[p] = i;
 				}
 			}
@@ -331,31 +627,31 @@ namespace Chaos::Softs
 
 
 			NodalWeights = NoCollisionNodalWeights;
-			for (int32 i = 0; i < ExtraWCIncidentElements.Num(); i++) 
+			for (int32 i = 0; i < ExtraWCIncidentElements.Num(); i++)
 			{
 				if (ExtraWCIncidentElements[i].Num() > 0)
 				{
 					int32 p = ExtraConstraints[ExtraWCIncidentElements[i][0]][ExtraWCIncidentElementsLocal[i][0]];
-					if (NodalWeights[p].Num() == 0) 
+					if (NodalWeights[p].Num() == 0)
 					{
 						NodalWeights[p].Init(T(0), 6);
 					}
-					for (int32 j = 0; j < ExtraWCIncidentElements[i].Num(); j++) 
+					for (int32 j = 0; j < ExtraWCIncidentElements[i].Num(); j++)
 					{
 						int32 LocalIndex = ExtraWCIncidentElementsLocal[i][j];
 						int32 ConstraintIndex = ExtraWCIncidentElements[i][j] + InitialWCSize;
 						T weight = T(0);
-						if (LocalIndex >= Indices[ConstraintIndex].Num()) 
+						if (LocalIndex >= Indices[ConstraintIndex].Num())
 						{
 							weight = SecondWeights[ConstraintIndex][LocalIndex - Indices[ConstraintIndex].Num()];
 						}
-						else 
+						else
 						{
 							weight = Weights[ConstraintIndex][LocalIndex];
 						}
-						if (IsAnisotropic[ConstraintIndex]) 
+						if (IsAnisotropic[ConstraintIndex])
 						{
-							for (int32 alpha = 0; alpha < 3; alpha++) 
+							for (int32 alpha = 0; alpha < 3; alpha++)
 							{
 								NodalWeights[p][alpha] += Normals[ConstraintIndex][alpha] * Normals[ConstraintIndex][alpha] * weight * weight * Stiffness[ConstraintIndex];
 							}
@@ -364,8 +660,9 @@ namespace Chaos::Softs
 							NodalWeights[p][4] += Normals[ConstraintIndex][0] * Normals[ConstraintIndex][2] * weight * weight * Stiffness[ConstraintIndex];
 							NodalWeights[p][5] += Normals[ConstraintIndex][1] * Normals[ConstraintIndex][2] * weight * weight * Stiffness[ConstraintIndex];
 						}
-						else {
-							for (int32 alpha = 0; alpha < 3; alpha++) 
+						else 
+						{
+							for (int32 alpha = 0; alpha < 3; alpha++)
 							{
 								NodalWeights[p][alpha] += weight * weight * Stiffness[ConstraintIndex];
 							}
@@ -373,11 +670,9 @@ namespace Chaos::Softs
 					}
 				}
 			}
-
 		}
 
 
-
 		TArray<TArray<int32>> Indices;
 		TArray<TArray<T>> Weights;
 		TArray<TVector<T, 3>> Constraints;
@@ -397,6 +692,9 @@ namespace Chaos::Softs
 		int32 InitialWCSize;
 		TArray<TArray<T>> NoCollisionNodalWeights;
 		TArray<TVector<int32, 4>> NoCollisionConstraints;
+		//For debugging
+		bool Detected = false;
+		TArray<int32> BoundaryVertices;
 	};
 
 
