Remove invalid (and seemingly unneeded) #include
#rnx

diff --git a/Engine/Source/Runtime/Experimental/Chaos/Public/Chaos/ClusterUnionManager.h b/Engine/Source/Runtime/Experimental/Chaos/Public/Chaos/ClusterUnionManager.h
--- a/Engine/Source/Runtime/Experimental/Chaos/Public/Chaos/ClusterUnionManager.h
+++ b/Engine/Source/Runtime/Experimental/Chaos/Public/Chaos/ClusterUnionManager.h
@@ -1,7 +1,6 @@
 // Copyright Epic Games, Inc. All Rights Reserved.
 #pragma once
 
-#include "../../../../../../Restricted/NotForLicensees/Plugins/GeneSplicer/Source/GeneSplicerLib/Private/genesplicer/Macros.h"
 #include "Chaos/ClusterCreationParameters.h"
 #include "Chaos/ParticleHandleFwd.h"
 #include "Containers/Array.h"
