I/O - delete all container partition(s) and not just the base container after upload

#rnx

diff --git a/Engine/Source/Runtime/Core/Private/IO/IoStore.cpp b/Engine/Source/Runtime/Core/Private/IO/IoStore.cpp
--- a/Engine/Source/Runtime/Core/Private/IO/IoStore.cpp
+++ b/Engine/Source/Runtime/Core/Private/IO/IoStore.cpp
@@ -2813,6 +2813,23 @@ public:
 		}
 	}
 
+	void GetContainerFilePaths(TArray<FString>& OutPaths)
+	{
+		TStringBuilder<256> Sb;
+
+		for (uint32 PartitionIndex = 0; PartitionIndex < Toc.GetTocResource().Header.PartitionCount; ++PartitionIndex)
+		{
+			Sb.Reset();
+			Sb.Append(ContainerPath);
+			if (PartitionIndex > 0)
+			{
+				Sb.Append(FString::Printf(TEXT("_s%d"), PartitionIndex));
+			}
+			Sb.Append(TEXT(".ucas"));
+			OutPaths.Add(Sb.ToString());
+		}
+	}
+
 private:
 
 
@@ -2914,6 +2931,11 @@ void FIoStoreReader::EnumerateCompressedBlocksForChunk(const FIoChunkId& Chunk,
 	Impl->EnumerateCompressedBlocksForChunk(Chunk, MoveTemp(Callback));
 }
 
+void FIoStoreReader::GetContainerFilePaths(TArray<FString>& OutPaths)
+{
+	Impl->GetContainerFilePaths(OutPaths);
+}
+
 FIoStatus FIoStoreTocResource::Read(const TCHAR* TocFilePath, EIoStoreTocReadOptions ReadOptions, FIoStoreTocResource& OutTocResource)
 {
 	check(TocFilePath != nullptr);
diff --git a/Engine/Source/Runtime/Core/Public/IO/IoDispatcher.h b/Engine/Source/Runtime/Core/Public/IO/IoDispatcher.h
--- a/Engine/Source/Runtime/Core/Public/IO/IoDispatcher.h
+++ b/Engine/Source/Runtime/Core/Public/IO/IoDispatcher.h
@@ -744,6 +744,9 @@ public:
 	CORE_API void EnumerateCompressedBlocks(TFunction<bool(const FIoStoreTocCompressedBlockInfo&)>&& Callback) const;
 	CORE_API void EnumerateCompressedBlocksForChunk(const FIoChunkId& Chunk, TFunction<bool(const FIoStoreTocCompressedBlockInfo&)>&& Callback) const;
 
+	// Returns the .ucas file path and all partition(s) ({containername}_s1.ucas, {containername}_s2.ucas)
+	CORE_API void GetContainerFilePaths(TArray<FString>& OutPaths);
+
 private:
 	FIoStoreReaderImpl* Impl;
 };
diff --git a/Engine/Source/Runtime/Experimental/IoStoreOnDemand/Private/IoStoreOnDemand.cpp b/Engine/Source/Runtime/Experimental/IoStoreOnDemand/Private/IoStoreOnDemand.cpp
--- a/Engine/Source/Runtime/Experimental/IoStoreOnDemand/Private/IoStoreOnDemand.cpp
+++ b/Engine/Source/Runtime/Experimental/IoStoreOnDemand/Private/IoStoreOnDemand.cpp
@@ -513,7 +513,8 @@ TIoStatusOr<FIoStoreUploadParams> FIoStoreUploadParams::Parse(const TCHAR* Comma
 	FParse::Value(CommandLine, TEXT("SessionToken="), Params.SessionToken);
 	FParse::Value(CommandLine, TEXT("CredentialsFile="), Params.CredentialsFile);
 	FParse::Value(CommandLine, TEXT("CredentialsFileKeyName="), Params.CredentialsFileKeyName);
-	Params.bDeleteContainerFiles = FParse::Param(CommandLine, TEXT("KeepUploadedContainers")) == false;
+	Params.bDeleteContainerFiles = !FParse::Param(CommandLine, TEXT("KeepContainerFiles"));
+	Params.bDeletePakFiles = !FParse::Param(CommandLine, TEXT("KeepPakFiles"));
 
 	if (Params.AccessKey.IsEmpty() &&
 		Params.SecretKey.IsEmpty() &&
@@ -660,7 +661,7 @@ TIoStatusOr<FIoStoreUploadResult> UploadContainerFiles(
 	FOnDemandToc OnDemandToc;
 	OnDemandToc.Header.ChunksDirectory = FString::Printf(TEXT("IoChunksV%u"), EOnDemandChunkVersion::Latest).ToLower();
 
-	TArray<FString> UploadedFiles;
+	TArray<FString> FilesToDelete;
 	for (const FString& Path : ContainerFiles)
 	{
 		FIoStoreReader ContainerFileReader;
@@ -679,7 +680,7 @@ TIoStatusOr<FIoStoreUploadResult> UploadContainerFiles(
 			continue;
 		}
 		
-		UE_LOG(LogIas, Display, TEXT("Uploading container '%s/.ucas'"), *Path);
+		UE_LOG(LogIas, Display, TEXT("Uploading container '%s'"), *Path);
 
 		const uint32 BlockSize = ContainerFileReader.GetCompressionBlockSize();
 		if (OnDemandToc.Header.BlockSize == 0)
@@ -810,7 +811,17 @@ TIoStatusOr<FIoStoreUploadResult> UploadContainerFiles(
 			}
 		}
 		
-		UploadedFiles.Add(Path);
+		if (UploadParams.bDeleteContainerFiles)
+		{
+			FilesToDelete.Add(Path);
+			ContainerFileReader.GetContainerFilePaths(FilesToDelete);
+
+			if (UploadParams.bDeletePakFiles)
+			{
+				FilesToDelete.Add(FPaths::ChangeExtension(Path, TEXT(".pak")));
+				FilesToDelete.Add(FPaths::ChangeExtension(Path, TEXT(".sig")));
+			}
+		}
 	}
 
 	if (OnDemandToc.Containers.IsEmpty())
@@ -853,20 +864,12 @@ TIoStatusOr<FIoStoreUploadResult> UploadContainerFiles(
 		}
 	}
 
-	if (UploadParams.bDeleteContainerFiles)
+	for (const FString& Path : FilesToDelete)
 	{
-		for (const FString& TocPath : UploadedFiles)
+		if (IFileManager::Get().FileExists(*Path))
 		{
-			const FString CasPath = FPaths::ChangeExtension(TocPath, TEXT(".ucas"));
-			const TCHAR* FilePaths[] {*TocPath, *CasPath};
-			for (const TCHAR* FilePath : FilePaths)
-			{
-				if (IFileManager::Get().FileExists(FilePath))
-				{
-					UE_LOG(LogIas, Display, TEXT("Deleting '%s'"), FilePath); 
-					IFileManager::Get().Delete(FilePath);
-				}
-			}
+			UE_LOG(LogIas, Display, TEXT("Deleting '%s'"), *Path); 
+			IFileManager::Get().Delete(*Path);
 		}
 	}
 
diff --git a/Engine/Source/Runtime/Experimental/IoStoreOnDemand/Public/IO/IoStoreOnDemand.h b/Engine/Source/Runtime/Experimental/IoStoreOnDemand/Public/IO/IoStoreOnDemand.h
--- a/Engine/Source/Runtime/Experimental/IoStoreOnDemand/Public/IO/IoStoreOnDemand.h
+++ b/Engine/Source/Runtime/Experimental/IoStoreOnDemand/Public/IO/IoStoreOnDemand.h
@@ -121,6 +121,7 @@ struct FIoStoreUploadParams
 	FString CredentialsFileKeyName;
 	int32 MaxConcurrentUploads = 16;
 	bool bDeleteContainerFiles = true;
+	bool bDeletePakFiles = true;
 	
 	UE_API static TIoStatusOr<FIoStoreUploadParams> Parse(const TCHAR* CommandLine);
 };
