Chaos Cloth:
- Fixed the CollectionPropertyFacade on the cloth component that wasn't updated after the config reinit and was doing nothing when ClearDirtyFlags was called on it.
- Correctly initialize the FVelocityAndPressureField property indices when initializing/changing the mesh or collection.
- Call missing ClearDirtyFlags on all configs' property facades after legacy code solver update.

#rnx

diff --git a/Engine/Plugins/Experimental/ChaosCloth/Source/Chaos/Private/ChaosCloth/ChaosClothingSimulation.cpp b/Engine/Plugins/Experimental/ChaosCloth/Source/Chaos/Private/ChaosCloth/ChaosClothingSimulation.cpp
--- a/Engine/Plugins/Experimental/ChaosCloth/Source/Chaos/Private/ChaosCloth/ChaosClothingSimulation.cpp
+++ b/Engine/Plugins/Experimental/ChaosCloth/Source/Chaos/Private/ChaosCloth/ChaosClothingSimulation.cpp
@@ -532,6 +532,11 @@ void FClothingSimulation::Simulate(IClothingSimulationContext* InContext)
 		if(Solver->GetEnableSolver() || (!Context->CacheData.HasData() && Context->CachedPositions.Num() == 0))
 		{
 			Solver->Update(SmoothedDeltaTime);
+
+			for (TUniquePtr<FClothingSimulationConfig>& Config : Configs)
+			{
+				Config->GetProperties().ClearDirtyFlags();
+			}
 		}
 		else
 		{
diff --git a/Engine/Plugins/Experimental/ChaosClothAsset/Source/ChaosClothAssetEngine/Private/ChaosClothAsset/ClothComponent.cpp b/Engine/Plugins/Experimental/ChaosClothAsset/Source/ChaosClothAssetEngine/Private/ChaosClothAsset/ClothComponent.cpp
--- a/Engine/Plugins/Experimental/ChaosClothAsset/Source/ChaosClothAssetEngine/Private/ChaosClothAsset/ClothComponent.cpp
+++ b/Engine/Plugins/Experimental/ChaosClothAsset/Source/ChaosClothAssetEngine/Private/ChaosClothAsset/ClothComponent.cpp
@@ -82,6 +82,9 @@ void UChaosClothComponent::ResetConfigProperties()
 			PropertyCollection->Reset();
 			CollectionPropertyMutableFacade.DefineSchema();
 		}
+
+		// Now the component property facade must be updated
+		*CollectionPropertyFacade = MoveTemp(CollectionPropertyMutableFacade);
 	}
 	else
 	{
diff --git a/Engine/Source/Runtime/Experimental/Chaos/Private/Chaos/VelocityField.cpp b/Engine/Source/Runtime/Experimental/Chaos/Private/Chaos/VelocityField.cpp
--- a/Engine/Source/Runtime/Experimental/Chaos/Private/Chaos/VelocityField.cpp
+++ b/Engine/Source/Runtime/Experimental/Chaos/Private/Chaos/VelocityField.cpp
@@ -100,6 +100,13 @@ void FVelocityAndPressureField::SetGeometry(
 	const TMap<FString, TConstArrayView<FRealSingle>>& Weightmaps,
 	FSolverReal WorldScale)
 {
+	// Reinit indices
+	DragIndex = FDragIndex(PropertyCollection);
+	LiftIndex = FLiftIndex(PropertyCollection);
+	FluidDensityIndex = FFluidDensityIndex(PropertyCollection);
+	PressureIndex = FPressureIndex(PropertyCollection);
+
+	// Reset geometry, properties, and weight maps
 	SetGeometry(TriangleMesh);
 	SetProperties(
 		FSolverVec2(GetWeightedFloatDrag(PropertyCollection, 0.f)),  // If these properties don't exist, set their values to 0, not to DefaultCoefficients!
