Allow IoStoreOnDemand to be enabled in the editor via the config file option Engine:[Ias]:EnableInEditor (false by default)

#rnx

diff --git a/Engine/Config/BaseEngine.ini b/Engine/Config/BaseEngine.ini
--- a/Engine/Config/BaseEngine.ini
+++ b/Engine/Config/BaseEngine.ini
@@ -3313,4 +3313,7 @@ bUpdateBuildGraphPropertiesFile=false
 ; Members of this struct property are inter-dependent and their initialization can be deferred, which
 ; means native defaults can match BP defaults at BP compile time while still requiring post-construct
 ; initialization, if other members that they are dependent on differ in Blueprint class default data.
-+RequiresCompleteValueForPostConstruction=/Script/Engine.BodyInstance:CollisionResponses
\ No newline at end of file
++RequiresCompleteValueForPostConstruction=/Script/Engine.BodyInstance:CollisionResponses
+
+[Ias]
+EnableInEditor=False
\ No newline at end of file
diff --git a/Engine/Source/Runtime/Experimental/IoStoreOnDemand/Private/IoStoreOnDemand.cpp b/Engine/Source/Runtime/Experimental/IoStoreOnDemand/Private/IoStoreOnDemand.cpp
--- a/Engine/Source/Runtime/Experimental/IoStoreOnDemand/Private/IoStoreOnDemand.cpp
+++ b/Engine/Source/Runtime/Experimental/IoStoreOnDemand/Private/IoStoreOnDemand.cpp
@@ -1247,7 +1247,16 @@ void FIoStoreOnDemandModule::StartupModule()
 {
 	using namespace UE::IO::Private;
 
-#if !WITH_EDITOR
+#if WITH_EDITOR
+	bool bEnabledInEditor = false;
+	GConfig->GetBool(TEXT("Ias"), TEXT("EnableInEditor"), bEnabledInEditor, GEngineIni);
+
+	if (!bEnabledInEditor)
+	{
+		return;
+	}
+#endif //WITH_EDITOR
+
 	const TCHAR* CommandLine = FCommandLine::Get();
 
 	UE::FOnDemandEndpoint Endpoint;
@@ -1358,7 +1367,6 @@ void FIoStoreOnDemandModule::StartupModule()
 	}
 #endif
 	FIoDispatcher::Get().Mount(Backend.ToSharedRef(), BackendPriority);
-#endif // !WITH_EDITOR
 }
 
 void FIoStoreOnDemandModule::ShutdownModule()
