Add the option to use a new ConvexHull3-based simplifier in the chaos convex builder, and if so disable the default simplification which can leave the hull in an invalid state (MergeFaces/MergeColinearEdges)

diff --git a/Engine/Source/Runtime/Experimental/Chaos/Private/Chaos/CollisionConvexMesh.cpp b/Engine/Source/Runtime/Experimental/Chaos/Private/Chaos/CollisionConvexMesh.cpp
--- a/Engine/Source/Runtime/Experimental/Chaos/Private/Chaos/CollisionConvexMesh.cpp
+++ b/Engine/Source/Runtime/Experimental/Chaos/Private/Chaos/CollisionConvexMesh.cpp
@@ -16,6 +16,7 @@ namespace Chaos
 	int32 FConvexBuilder::ComputeHorizonEpsilonFromMeshExtends = 1;
 
 	bool FConvexBuilder::bUseGeometryTConvexHull3 = true;
+	bool FConvexBuilder::bUseSimplifierForTConvexHull3 = false;
 
 	FAutoConsoleVariableRef CVarConvexGeometryCheckEnable(TEXT("p.Chaos.ConvexGeometryCheckEnable"), FConvexBuilder::PerformGeometryCheck, TEXT("Perform convex geometry complexity check for Chaos physics."));
 
@@ -24,15 +25,20 @@ namespace Chaos
 	FAutoConsoleVariableRef CVarConvexParticlesWarningThreshold(TEXT("p.Chaos.ConvexParticlesWarningThreshold"), FConvexBuilder::VerticesThreshold, TEXT("Threshold beyond which we warn about collision geometry complexity."));
 	
 	FAutoConsoleVariableRef CvarUseGeometryTConvexHull3(TEXT("p.Chaos.Convex.UseTConvexHull3Builder"), FConvexBuilder::bUseGeometryTConvexHull3, TEXT("Use the newer Geometry Tools code path for generating convex hulls when default build method is set.[def:true]"));
+	FAutoConsoleVariableRef CvarUseSimplifierForGeometryTConvexHull3(TEXT("p.Chaos.Convex.UseSimplifierForTConvexHull3Builder"), FConvexBuilder::bUseSimplifierForTConvexHull3, TEXT("If default build is using the Geometry Tools convex hull algorithm, also use the corresponding simplifier. [def:false]"));
 
 
 	bool FConvexBuilder::UseConvexHull3(FConvexBuilder::EBuildMethod BuildMethod)
 	{
-		if (BuildMethod == EBuildMethod::Default && FConvexBuilder::bUseGeometryTConvexHull3)
-		{
-			return true;
-		}
-		return BuildMethod == EBuildMethod::ConvexHull3;
+		return (BuildMethod == EBuildMethod::Default && FConvexBuilder::bUseGeometryTConvexHull3)
+			|| (BuildMethod == EBuildMethod::ConvexHull3 || BuildMethod == EBuildMethod::ConvexHull3Simplified);
+	}
+
+	bool FConvexBuilder::UseConvexHull3Simplifier(EBuildMethod BuildMethod)
+	{
+
+		return (BuildMethod == EBuildMethod::Default && FConvexBuilder::bUseGeometryTConvexHull3 && FConvexBuilder::bUseSimplifierForTConvexHull3)
+			|| BuildMethod == EBuildMethod::ConvexHull3Simplified;
 	}
 
 	void FConvexBuilder::BuildIndices(const TArray<FVec3Type>& InVertices, TArray<int32>& OutResultIndexData, EBuildMethod BuildMethod)
@@ -40,16 +46,45 @@ namespace Chaos
 		if (UseConvexHull3(BuildMethod))
 		{
 			UE::Geometry::TConvexHull3<FRealType> HullCompute;
+			const bool bUseSimplifier = UseConvexHull3Simplifier(BuildMethod);
+			if (bUseSimplifier)
+			{
+				HullCompute.bSaveTriangleNeighbors = true;
+				HullCompute.DegenerateEdgeTolerance = (FRealType)UE_DOUBLE_SMALL_NUMBER;
+			}
 			if (HullCompute.Solve<FVec3Type>(InVertices))
 			{
-				const TArray<UE::Geometry::FIndex3i>& HullTriangles = HullCompute.GetTriangles();
-				OutResultIndexData.Reserve(HullTriangles.Num() * 3);
-				for (const UE::Geometry::FIndex3i& Tri : HullTriangles)
+				if (bUseSimplifier)
+				{
+					// Note: currently these tolerances are chosen to be close to the tolerances used by the Convex Builder's calls to MergeFaces
+					// Note: A reasonable alternative for DistanceTolerance would be FRealType(1e-3) * OutLocalBounds.Extents().GetMax();
+					const FRealType AngleToleranceDeg = (FRealType)1;
+					const FRealType DistanceTolerance = (FRealType)1;
+					using FPolygonFace = UE::Geometry::TConvexHull3<FRealType>::FPolygonFace;
+					TArray<FPolygonFace> Polygons;
+					HullCompute.GetSimplifiedFaces(Polygons, [&InVertices](int32 Idx) { return InVertices[Idx]; }, AngleToleranceDeg, DistanceTolerance);
+					for (const FPolygonFace& Face : Polygons)
+					{
+						OutResultIndexData.Reserve(OutResultIndexData.Num() + 3 * (Face.Num() - 2));
+						for (int32 Idx = 1; Idx + 1 < Face.Num(); ++Idx)
+						{
+							OutResultIndexData.Add(Face[0]);
+							OutResultIndexData.Add(Face[Idx + 1]);
+							OutResultIndexData.Add(Face[Idx]);
+						}
+					}
+				}
+				else
 				{
-					// Winding is backwards from what Chaos expects
-					OutResultIndexData.Add(Tri[0]);
-					OutResultIndexData.Add(Tri[2]);
-					OutResultIndexData.Add(Tri[1]);
+					const TArray<UE::Geometry::FIndex3i>& HullTriangles = HullCompute.GetTriangles();
+					OutResultIndexData.Reserve(HullTriangles.Num() * 3);
+					for (const UE::Geometry::FIndex3i& Tri : HullTriangles)
+					{
+						// Winding is backwards from what Chaos expects
+						OutResultIndexData.Add(Tri[0]);
+						OutResultIndexData.Add(Tri[2]);
+						OutResultIndexData.Add(Tri[1]);
+					}
 				}
 			}
 		}
@@ -136,40 +171,58 @@ namespace Chaos
 		{
 			if (UseConvexHull3(BuildMethod)) // Use the newer Geometry Tools code path for generating convex hulls.
 			{
+				bool bUseSimplifier = UseConvexHull3Simplifier(BuildMethod);
 				UE::Geometry::TConvexHull3<FRealType> HullCompute;
 				HullCompute.bSaveTriangleNeighbors = true;
+				if (bUseSimplifier)
+				{
+					HullCompute.DegenerateEdgeTolerance = (FRealType)UE_DOUBLE_SMALL_NUMBER;
+				}
 				if (HullCompute.Solve<FVec3Type>(*VerticesToUse))
 				{
 					// Get Planes, FaceIndices, Vertices, and LocalBoundingBox
 					TMap<int32, int32> HullVertMap;
-					HullCompute.GetFaces([&VerticesToUse, &HullVertMap, &OutPlanes, &OutFaceIndices, &OutVertices, &OutLocalBounds](TArray<int32>& FaceIndices, FVector3f FaceNormal)
-						{
-							TArray<int32>& OutFace = OutFaceIndices.Emplace_GetRef();
-							OutFace.SetNumUninitialized(FaceIndices.Num());
+					auto AddFace = [&VerticesToUse, &HullVertMap, &OutPlanes, &OutFaceIndices, &OutVertices, &OutLocalBounds](TArray<int32>& FaceIndices, FVector3f FaceNormal)
+					{
+						TArray<int32>& OutFace = OutFaceIndices.Emplace_GetRef();
+						OutFace.SetNumUninitialized(FaceIndices.Num());
 
-							// Note winding is backwards from what Chaos expects, so this fills OutFace in reverse order
-							for (int32 j = 0, n = FaceIndices.Num(); j < n; ++j)
+						// Note winding is backwards from what Chaos expects, so this fills OutFace in reverse order
+						for (int32 j = 0, n = FaceIndices.Num(); j < n; ++j)
+						{
+							int32 Index = FaceIndices[j];
+							int32* MappedVert = HullVertMap.Find(Index);
+							if (!MappedVert)
 							{
-								int32 Index = FaceIndices[j];
-								int32* MappedVert = HullVertMap.Find(Index);
-								if (!MappedVert)
-								{
-									const FVec3Type& OrigPos = (*VerticesToUse)[Index];
-									int32 NewVID = OutVertices.Num();
-									OutVertices.Add(OrigPos);
-									HullVertMap.Add(Index, NewVID);
-									OutFace[n - j - 1] = NewVID;
-								}
-								else
-								{
-									OutFace[n - j - 1] = *MappedVert;
-								}
+								const FVec3Type& OrigPos = (*VerticesToUse)[Index];
+								int32 NewVID = OutVertices.Num();
+								OutVertices.Add(OrigPos);
+								HullVertMap.Add(Index, NewVID);
+								OutFace[n - j - 1] = NewVID;
 							}
-
-							const FVec3Type& Normal = FaceNormal;
-							const FVec3Type& Point = OutVertices[OutFace[0]];
-							OutPlanes.Add(FPlaneType{ Point, Normal });
-						}, [&VerticesToUse](int32 Idx) { return (*VerticesToUse)[Idx]; });
+							else
+							{
+								OutFace[n - j - 1] = *MappedVert;
+							}
+						}
+
+						const FVec3Type& Normal = FaceNormal;
+						const FVec3Type& Point = OutVertices[OutFace[0]];
+						OutPlanes.Add(FPlaneType{ Point, Normal });
+					};
+					auto GetVertex = [&VerticesToUse](int32 Idx) { return (*VerticesToUse)[Idx]; };
+					if (bUseSimplifier)
+					{
+						// Note: currently these tolerances are chosen to be close to the tolerances used by the Convex Builder's calls to MergeFaces
+						// Note: A reasonable alternative for DistanceTolerance would be FRealType(1e-3) * OutLocalBounds.Extents().GetMax();
+						const FRealType AngleToleranceDeg = (FRealType)1;
+						const FRealType DistanceTolerance = (FRealType)1;
+						HullCompute.GetSimplifiedFaces(AddFace, GetVertex, AngleToleranceDeg, DistanceTolerance);
+					}
+					else
+					{
+						HullCompute.GetFaces(AddFace, GetVertex);
+					}
 
 					OutLocalBounds = FAABB3Type(OutVertices[0], OutVertices[0]);
 					for (int32 VertexIndex = 1; VertexIndex < OutVertices.Num(); ++VertexIndex)
diff --git a/Engine/Source/Runtime/Experimental/Chaos/Public/Chaos/CollisionConvexMesh.h b/Engine/Source/Runtime/Experimental/Chaos/Public/Chaos/CollisionConvexMesh.h
--- a/Engine/Source/Runtime/Experimental/Chaos/Public/Chaos/CollisionConvexMesh.h
+++ b/Engine/Source/Runtime/Experimental/Chaos/Public/Chaos/CollisionConvexMesh.h
@@ -45,6 +45,7 @@ namespace Chaos
 			Default = 0,	// use what is set as default from cvars
 			Original = 1,	// original method, fast but has issues
 			ConvexHull3 = 2,// new method, slower but less issues
+			ConvexHull3Simplified = 3, // same as above new method, plus new simplification method
 		};
 
 		class Params
@@ -263,6 +264,8 @@ namespace Chaos
 
 		static bool UseConvexHull3(EBuildMethod BuildMethod);
 
+		static bool UseConvexHull3Simplifier(EBuildMethod BuildMethod);
+
 		static CHAOS_API void BuildIndices(const TArray<FVec3Type>& InVertices, TArray<int32>& OutResultIndexData, EBuildMethod BuildMethod = EBuildMethod::Default);
 		
 		static void BuildConvexHull(const TArray<FVec3Type>& InVertices, TArray<TVec3<int32>>& OutIndices, const Params& InParams = Params())
@@ -577,8 +580,8 @@ namespace Chaos
 		// Find edge pairs that are colinear and remove the unnecessary vertex to make a single edge.
 		// If we are left with an invalid face (2 verts or less), remove it.
 		// NOTE: a vertex in the middle of two colinear edges on one face may still be required by some other face, 
-		// although technically those faces could (and should) have been merged.
-		static void MergeColinearEdges(TArray<FPlaneType>& InOutPlanes, TArray<TArray<int32>>& InOutFaceVertexIndices, TArray<FVec3Type>& InOutVertices, FRealType AngleToleranceRad)
+		// although if we are lucky those faces would have been merged by the MergeFaces() function, depending on the tolerances used
+		static void MergeColinearEdges(TArray<FPlaneType>& InOutPlanes, TArray<TArray<int32>>& InOutFaceVertexIndices, TArray<FVec3Type>& InOutVertices, FRealType AngleToleranceCos)
 		{
 			check(InOutPlanes.Num() == InOutFaceVertexIndices.Num());
 
@@ -612,7 +615,7 @@ namespace Chaos
 						const FReal CosAngle = FVec3::DotProduct(Edge0, Edge1);
 
 						// See if we need the vertex.
-						if (CosAngle < (FReal(1) - AngleToleranceRad))
+						if (CosAngle < (FReal(1) - AngleToleranceCos))
 						{
 							VertexIndexMap[VertexIndex1] = VertexIndex1;
 						}
@@ -811,6 +814,7 @@ namespace Chaos
 		static CHAOS_API int32 VerticesThreshold;
 		static CHAOS_API int32 ComputeHorizonEpsilonFromMeshExtends;
 		static CHAOS_API bool bUseGeometryTConvexHull3;
+		static CHAOS_API bool bUseSimplifierForTConvexHull3;
 
 	private:
 
diff --git a/Engine/Source/Runtime/Experimental/Chaos/Public/Chaos/Convex.h b/Engine/Source/Runtime/Experimental/Chaos/Public/Chaos/Convex.h
--- a/Engine/Source/Runtime/Experimental/Chaos/Public/Chaos/Convex.h
+++ b/Engine/Source/Runtime/Experimental/Chaos/Public/Chaos/Convex.h
@@ -116,13 +116,16 @@ namespace Chaos
 			// @todo(chaos): this only works with triangles. Fix that an we can run MergeFaces before calling this
 			CalculateVolumeAndCenterOfMass(Vertices, FaceIndices, Volume, CenterOfMass);
 
-			// it appears that this code path can leave the convex in an undefined state.
-			// @todo(chaos): Tolerances should be based on size, or passed in
-			const FRealType DistanceTolerance = 1.0f;
-			const FRealType AngleTolerance = 1.e-6f;
-			FConvexBuilder::MergeFaces(Planes, FaceIndices, Vertices, DistanceTolerance);
-			FConvexBuilder::MergeColinearEdges(Planes, FaceIndices, Vertices, AngleTolerance);
-			CHAOS_ENSURE(Planes.Num() == FaceIndices.Num());
+			if (BuildMethod != FConvexBuilder::EBuildMethod::ConvexHull3Simplified)
+			{
+				// it appears that this code path can leave the convex in an undefined state.
+				// @todo(chaos): Tolerances should be based on size, or passed 
+				const FRealType DistanceTolerance = 1.0f;
+				const FRealType AngleTolerance = 1.e-6f;
+				FConvexBuilder::MergeFaces(Planes, FaceIndices, Vertices, DistanceTolerance);
+				FConvexBuilder::MergeColinearEdges(Planes, FaceIndices, Vertices, AngleTolerance);
+				CHAOS_ENSURE(Planes.Num() == FaceIndices.Num());
+			}
 
 			CreateStructureData(MoveTemp(FaceIndices));
 
