turn p.Chaos.OnlyUseInterclusterEdgesAttachedToMainParticles on by default

diff --git a/Engine/Source/Runtime/Experimental/Chaos/Private/Chaos/PBDRigidClustering.cpp b/Engine/Source/Runtime/Experimental/Chaos/Private/Chaos/PBDRigidClustering.cpp
--- a/Engine/Source/Runtime/Experimental/Chaos/Private/Chaos/PBDRigidClustering.cpp
+++ b/Engine/Source/Runtime/Experimental/Chaos/Private/Chaos/PBDRigidClustering.cpp
@@ -39,7 +39,7 @@ namespace Chaos
 	bool bCheckForInterclusterEdgesOnRelease = true;
 	FAutoConsoleVariableRef CVarCheckForInterclusterEdgesOnRelease(TEXT("p.Chaos.CheckForInterclusterEdgesOnRelease"), bCheckForInterclusterEdgesOnRelease, TEXT("Whether to check for intercluster edges when removing a child from its parent cluster so that we can add the particle back into a cluster union."));
 
-	bool bOnlyUseInterclusterEdgesAttachedToMainParticles = false;
+	bool bOnlyUseInterclusterEdgesAttachedToMainParticles = true;
 	FAutoConsoleVariableRef CVarOnlyUseInterclusterEdgesAttachedToMainParticles(TEXT("p.Chaos.OnlyUseInterclusterEdgesAttachedToMainParticles"), bOnlyUseInterclusterEdgesAttachedToMainParticles, TEXT("If true, an intercluster edge must be directly attached to a main particle for the particle to remain a part of the cluster union."));
 
 	int32 ComputeClusterCollisionStrains = 1;
