IAS content control interface. This lets you disable bulk optional assets or streaming entirely.

diff --git a/Engine/Source/Runtime/Experimental/IoStoreOnDemand/Private/IoStoreOnDemand.cpp b/Engine/Source/Runtime/Experimental/IoStoreOnDemand/Private/IoStoreOnDemand.cpp
--- a/Engine/Source/Runtime/Experimental/IoStoreOnDemand/Private/IoStoreOnDemand.cpp
+++ b/Engine/Source/Runtime/Experimental/IoStoreOnDemand/Private/IoStoreOnDemand.cpp
@@ -1231,18 +1231,25 @@ FIoStatus DownloadContainerFiles(const FIoStoreDownloadParams& DownloadParams, c
 } // namespace UE
 
 ////////////////////////////////////////////////////////////////////////////////
-class FIoStoreOnDemandModule
-	: public IModuleInterface
-{
-public:
-	virtual void StartupModule() override;
-	virtual void ShutdownModule() override;
-
-private:
-};
 
 IMPLEMENT_MODULE(FIoStoreOnDemandModule, IoStoreOnDemand);
 
+void FIoStoreOnDemandModule::SetBulkOptionalEnabled(bool bInEnabled)
+{
+	if (Backend.IsValid())
+	{
+		Backend->SetBulkOptionalEnabled(bInEnabled);
+	}
+}
+
+void FIoStoreOnDemandModule::SetEnabled(bool bInEnabled)
+{
+	if (Backend.IsValid())
+	{
+		Backend->SetEnabled(bInEnabled);
+	}
+}
+
 void FIoStoreOnDemandModule::StartupModule()
 {
 	using namespace UE::IO::Private;
@@ -1356,7 +1363,7 @@ void FIoStoreOnDemandModule::StartupModule()
 		UE_LOG(LogIas, Log, TEXT("File cache disabled. Streaming only."));
 	}
 
-	TSharedPtr<UE::IOnDemandIoDispatcherBackend> Backend = UE::MakeOnDemandIoDispatcherBackend(Cache);
+	Backend = UE::MakeOnDemandIoDispatcherBackend(Cache);
 	Backend->Mount(Endpoint);
 	int32 BackendPriority = -10;
 #if !UE_BUILD_SHIPPING
diff --git a/Engine/Source/Runtime/Experimental/IoStoreOnDemand/Private/OnDemandIoDispatcherBackend.cpp b/Engine/Source/Runtime/Experimental/IoStoreOnDemand/Private/OnDemandIoDispatcherBackend.cpp
--- a/Engine/Source/Runtime/Experimental/IoStoreOnDemand/Private/OnDemandIoDispatcherBackend.cpp
+++ b/Engine/Source/Runtime/Experimental/IoStoreOnDemand/Private/OnDemandIoDispatcherBackend.cpp
@@ -56,6 +56,13 @@ static FAutoConsoleVariableRef CVar_IoDispatcherMaxHttpPollTimeoutMs (
 	TEXT("Tick() poll timeout in milliseconds")
 );
 
+bool GIoDispatcherBulkOptionalEnabled = true;
+static FAutoConsoleVariableRef CVar_IoDispatcherBulkOptionalEnabled(
+	TEXT("ias.BulkOptionalEnabled"),
+	GIoDispatcherBulkOptionalEnabled,
+	TEXT("Enables bulk optional requests.")
+);
+
 namespace UE::IO::Private
 {
 
@@ -945,6 +952,8 @@ public:
 
 	// I/O Http backend
 	virtual void Mount(const FOnDemandEndpoint& Endpoint) override;
+	virtual void SetBulkOptionalEnabled(bool bInEnabled) override;
+	virtual void SetEnabled(bool bInEnabled) override;
 
 	// Runnable
 	virtual bool Init() override { return true; }
@@ -971,6 +980,8 @@ private:
 	FOnDemandIoBackendStats Stats;
 	FRWLock Lock;
 	std::atomic_bool bStopRequested{false};
+	std::atomic_bool bEnableBulkOptional{true};
+	std::atomic_bool bEnabled{true};
 };
 
 ///////////////////////////////////////////////////////////////////////////////
@@ -1081,6 +1092,11 @@ void FOnDemandIoBackend::CompleteRequest(FChunkRequest* ChunkRequest)
 bool FOnDemandIoBackend::Resolve(FIoRequestImpl* Request)
 {
 	using namespace UE::Tasks;
+
+	if (!bEnabled)
+	{
+		return false;
+	}
 	
 	FOnDemandIoStore::FChunkInfo ChunkInfo = IoStore->GetChunkInfo(Request->ChunkId);
 	if (!ChunkInfo.IsValid())
@@ -1088,6 +1104,11 @@ bool FOnDemandIoBackend::Resolve(FIoRequestImpl* Request)
 		return false;
 	}
 
+	if ((!GIoDispatcherBulkOptionalEnabled || !bEnableBulkOptional) && Request->ChunkId.GetChunkType() == EIoChunkType::OptionalBulkData)
+	{
+		return false;
+	}
+
 	Stats.OnIoRequestEnqueue();
 	FChunkRequestParams RequestParams = FChunkRequestParams::Create(Request, ChunkInfo);
 	FChunkRequest* ChunkRequest = ChunkRequests.Create(Request, RequestParams);
@@ -1161,6 +1182,16 @@ TIoStatusOr<uint64> FOnDemandIoBackend::GetSizeForChunk(const FIoChunkId& ChunkI
 {
 	if (IoStore.IsValid())
 	{
+		if (!bEnabled)
+		{
+			return FIoStatus(EIoErrorCode::UnknownChunkID);
+		}
+
+		if ((!GIoDispatcherBulkOptionalEnabled || !bEnableBulkOptional) && ChunkId.GetChunkType() == EIoChunkType::OptionalBulkData)
+		{
+			return FIoStatus(EIoErrorCode::UnknownChunkID);
+		}
+
 		return IoStore->GetChunkSize(ChunkId);
 	}
 
@@ -1262,6 +1293,16 @@ void FOnDemandIoBackend::Mount(const FOnDemandEndpoint& Endpoint)
 	}
 }
 
+void FOnDemandIoBackend::SetBulkOptionalEnabled(bool bInEnabled)
+{
+	bEnableBulkOptional = bInEnabled;
+}
+
+void FOnDemandIoBackend::SetEnabled(bool bInEnabled)
+{
+	bEnabled = bInEnabled;
+}
+
 TIoStatusOr<FOnDemandToc> FOnDemandIoBackend::GetToc(FHttpClient& HttpClient, const FString& TocPath)
 {
 	TRACE_CPUPROFILER_EVENT_SCOPE(FOnDemandIoBackend::GetToc);
diff --git a/Engine/Source/Runtime/Experimental/IoStoreOnDemand/Private/OnDemandIoDispatcherBackend.h b/Engine/Source/Runtime/Experimental/IoStoreOnDemand/Private/OnDemandIoDispatcherBackend.h
--- a/Engine/Source/Runtime/Experimental/IoStoreOnDemand/Private/OnDemandIoDispatcherBackend.h
+++ b/Engine/Source/Runtime/Experimental/IoStoreOnDemand/Private/OnDemandIoDispatcherBackend.h
@@ -34,6 +34,8 @@ class IOnDemandIoDispatcherBackend
 public:
 	virtual ~IOnDemandIoDispatcherBackend() = default;
 	virtual void Mount(const FOnDemandEndpoint& Endpoint) = 0;
+	virtual void SetBulkOptionalEnabled(bool bInEnabled) = 0;
+	virtual void SetEnabled(bool bInEnabled) = 0;
 };
 
 UE_API TSharedPtr<IOnDemandIoDispatcherBackend> MakeOnDemandIoDispatcherBackend(TSharedPtr<IIoCache> Cache);
diff --git a/Engine/Source/Runtime/Experimental/IoStoreOnDemand/Public/IO/IoStoreOnDemand.h b/Engine/Source/Runtime/Experimental/IoStoreOnDemand/Public/IO/IoStoreOnDemand.h
--- a/Engine/Source/Runtime/Experimental/IoStoreOnDemand/Public/IO/IoStoreOnDemand.h
+++ b/Engine/Source/Runtime/Experimental/IoStoreOnDemand/Public/IO/IoStoreOnDemand.h
@@ -9,6 +9,7 @@
 #include "Misc/Guid.h"
 #include "Templates/SharedPointer.h"
 #include "UObject/NameTypes.h"
+#include "Modules/ModuleInterface.h"
 
 #if (IS_PROGRAM || WITH_EDITOR)
 #include "Misc/AES.h"
@@ -160,4 +161,24 @@ UE_API FIoStatus DownloadContainerFiles(const FIoStoreDownloadParams& DownloadPa
 
 } // namespace UE
 
-#undef UE_API
+namespace UE
+{
+	class IOnDemandIoDispatcherBackend;
+}
+
+class FIoStoreOnDemandModule
+	: public IModuleInterface
+{
+	TSharedPtr<UE::IOnDemandIoDispatcherBackend> Backend;
+
+public:
+	UE_API void SetBulkOptionalEnabled(bool bInEnabled);
+	UE_API void SetEnabled(bool bInEnabled);
+
+	virtual void StartupModule() override;
+	virtual void ShutdownModule() override;
+
+private:
+};
+
+#undef UE_API
\ No newline at end of file
