CIS: resolve NumExtraFilterBits conflict in PhysicsFiltering.h and ParticleHandle.h

diff --git a/Engine/Source/Runtime/Experimental/Chaos/Public/Chaos/Collision/CollisionFilter.h b/Engine/Source/Runtime/Experimental/Chaos/Public/Chaos/Collision/CollisionFilter.h
--- a/Engine/Source/Runtime/Experimental/Chaos/Public/Chaos/Collision/CollisionFilter.h
+++ b/Engine/Source/Runtime/Experimental/Chaos/Public/Chaos/Collision/CollisionFilter.h
@@ -5,6 +5,7 @@
 
 #include "Chaos/Core.h"
 #include "Chaos/Collision/CollisionConstraintFlags.h"
+#include "Chaos/Collision/CollisionFilterBits.h"
 #include "Chaos/CollisionFilterData.h"
 #include "Chaos/GeometryParticles.h"
 #include "Chaos/ImplicitObjectType.h"
@@ -13,8 +14,6 @@
 namespace Chaos
 {
 	typedef uint8 FMaskFilter;
-	enum { NumExtraFilterBits = 6 };
-	enum { NumCollisionChannelBits = 5 };
 
 	/**
 	 * Check whether the two particles need to be considered in the broadphase.
diff --git a/Engine/Source/Runtime/Experimental/Chaos/Public/Chaos/Collision/CollisionFilterBits.h b/Engine/Source/Runtime/Experimental/Chaos/Public/Chaos/Collision/CollisionFilterBits.h
new file mode 100644
--- /dev/null
+++ b/Engine/Source/Runtime/Experimental/Chaos/Public/Chaos/Collision/CollisionFilterBits.h
@@ -0,0 +1,8 @@
+// Copyright Epic Games, Inc. All Rights Reserved.
+#pragma once
+
+namespace Chaos
+{
+	enum { ChaosNumExtraFilterBits = 6 };
+	enum { ChaosNumCollisionChannelBits = 5 };
+}
\ No newline at end of file
diff --git a/Engine/Source/Runtime/Experimental/Chaos/Public/Chaos/ParticleHandle.h b/Engine/Source/Runtime/Experimental/Chaos/Public/Chaos/ParticleHandle.h
--- a/Engine/Source/Runtime/Experimental/Chaos/Public/Chaos/ParticleHandle.h
+++ b/Engine/Source/Runtime/Experimental/Chaos/Public/Chaos/ParticleHandle.h
@@ -9,6 +9,7 @@
 #include "Chaos/ParticleHandleFwd.h"
 #include "Chaos/ParticleIterator.h"
 #include "Chaos/Properties.h"
+#include "Chaos/Collision/CollisionFilterBits.h"
 #include "ChaosCheck.h"
 #include "Chaos/ChaosDebugDrawDeclares.h"
 #if CHAOS_DEBUG_DRAW
@@ -160,11 +161,8 @@ FORCEINLINE_DEBUGGABLE bool PrePreQueryFilterImp(const FCollisionFilterData& Que
 
 FORCEINLINE_DEBUGGABLE uint32 GetChaosCollisionChannelAndExtraFilter(uint32 Word3, uint8& OutMaskFilter)
 {
-	enum { NumExtraFilterBits = 6 };
-	enum { NumCollisionChannelBits = 5 };
-
-	uint32 ChannelMask = (Word3 << NumExtraFilterBits) >> (32 - NumCollisionChannelBits);
-	OutMaskFilter = Word3 >> (32 - NumExtraFilterBits);
+	uint32 ChannelMask = (Word3 << ChaosNumExtraFilterBits) >> (32 - ChaosNumCollisionChannelBits);
+	OutMaskFilter = Word3 >> (32 - ChaosNumExtraFilterBits);
 	return (uint32)ChannelMask;
 }
 
