Fix scaled implicit object when implicit not created from TRefCountPtr

diff --git a/Engine/Source/Runtime/Experimental/Chaos/Private/Chaos/TriangleMeshImplicitObject.cpp b/Engine/Source/Runtime/Experimental/Chaos/Private/Chaos/TriangleMeshImplicitObject.cpp
--- a/Engine/Source/Runtime/Experimental/Chaos/Private/Chaos/TriangleMeshImplicitObject.cpp
+++ b/Engine/Source/Runtime/Experimental/Chaos/Private/Chaos/TriangleMeshImplicitObject.cpp
@@ -38,9 +38,7 @@ template <typename QueryGeomType>
 static auto MakeScaledHelper(const QueryGeomType& B, const FVec3& InvScale)
 {
 	// TODO: Fixup code using this and remove it.
-
-	TRefCountPtr<QueryGeomType> HackBPtr(const_cast<QueryGeomType*>(&B));	//todo: hack, need scaled object to accept raw ptr similar to transformed implicit
-	TImplicitObjectScaled<QueryGeomType> ScaledB(HackBPtr, InvScale);
+	TImplicitObjectScaled<QueryGeomType,true> ScaledB(const_cast<QueryGeomType*>(&B), InvScale);
 	return ScaledB;
 }
 
diff --git a/Engine/Source/Runtime/Experimental/Chaos/Public/Chaos/ImplicitObject.h b/Engine/Source/Runtime/Experimental/Chaos/Public/Chaos/ImplicitObject.h
--- a/Engine/Source/Runtime/Experimental/Chaos/Public/Chaos/ImplicitObject.h
+++ b/Engine/Source/Runtime/Experimental/Chaos/Public/Chaos/ImplicitObject.h
@@ -1,6 +1,7 @@
 // Copyright Epic Games, Inc. All Rights Reserved.
 #pragma once
 
+#include <atomic>
 #include "Chaos/Pair.h"
 #include "Chaos/Serializable.h"
 #include "Chaos/Core.h"
@@ -102,8 +103,55 @@ struct TImplicitTypeInfo
 // This is a compiler-dependent behavior, so if you are not seeing any other compile time errors about sizeof(FImplicitObject) + offsetof(...) with this disabled,
 // you should be OK.
 #define DISALLOW_FIMPLICIT_OBJECT_TAIL_PADDING INTEL_ISPC
+
+// Chaos ref counted object
+class FChaosRefCountedObject
+{
+public:
+	FChaosRefCountedObject() : NumRefs(0) {}
+	virtual ~FChaosRefCountedObject() { check(NumRefs.GetValue() == 0); }
+	FChaosRefCountedObject(const FChaosRefCountedObject& Rhs) = delete;
+	FChaosRefCountedObject& operator=(const FChaosRefCountedObject& Rhs) = delete;
+	uint32 AddRef() const
+	{
+		return uint32(NumRefs.Increment());
+	}
+	uint32 Release() const
+	{
+		uint32 Refs = uint32(NumRefs.Decrement());
+		if (Refs == 0)
+		{
+			if(bTransientFlag)
+			{ 
+				delete this;
+			}
+			else
+			{
+				bTransientFlag = true;
+			}
+		}
+		return Refs;
+	}
+	uint32 GetRefCount() const
+	{
+		return uint32(NumRefs.GetValue());
+	}
+
+	void MakePersistent() const
+	{
+		bTransientFlag = false;
+	}
 	
-class FImplicitObject : public FThreadSafeRefCountedObject
+private:
+	// Number of refs onto the object
+	mutable FThreadSafeCounter NumRefs;
+
+	// Transient flag to trigger or not the automatic deletion
+	mutable std::atomic<bool> bTransientFlag = true;
+};
+	
+
+class FImplicitObject : public FChaosRefCountedObject
 {
 public:
 	using TType = FReal;
diff --git a/Engine/Source/Runtime/Experimental/Chaos/Public/Chaos/ImplicitObjectScaled.h b/Engine/Source/Runtime/Experimental/Chaos/Public/Chaos/ImplicitObjectScaled.h
--- a/Engine/Source/Runtime/Experimental/Chaos/Public/Chaos/ImplicitObjectScaled.h
+++ b/Engine/Source/Runtime/Experimental/Chaos/Public/Chaos/ImplicitObjectScaled.h
@@ -446,22 +446,21 @@ public:
 	
 	TImplicitObjectScaled(ObjectType Object, const FVec3& Scale, FReal InMargin = 0)
 		: FImplicitObjectScaled(EImplicitObject::HasBoundingBox, Object->GetType())
-		, MObject(MoveTemp(Object))
+		, MObject(Object)
 	{
-		ensureMsgf((IsScaled(MObject->GetType()) == false), TEXT("Scaled objects should not contain each other."));
-		ensureMsgf((IsInstanced(MObject->GetType()) == false), TEXT("Scaled objects should not contain instances."));
-		switch (MObject->GetType())
+		InitScaledImplicit(Scale, InMargin);
+	}
+	
+	TImplicitObjectScaled(TConcrete* Object, const FVec3& Scale, FReal InMargin = 0)
+		: FImplicitObjectScaled(EImplicitObject::HasBoundingBox, Object->GetType())
+		, MObject(Object)
+	{
+		// Transient if raw pointer not already stored in a ref counted one
+		if(Object && (Object->GetRefCount() == 1))
 		{
-		case ImplicitObjectType::Transformed:
-		case ImplicitObjectType::Union:
-			check(false);	//scale is only supported for concrete types like sphere, capsule, convex, levelset, etc... Nothing that contains other objects
-			default:
-				break;
+			Object->MakePersistent();
 		}
-		this->bIsConvex = MObject->IsConvex();
-		this->bDoCollide = MObject->GetDoCollide();
-		this->OuterMargin = InMargin;
-		SetScale(Scale);
+		InitScaledImplicit(Scale, InMargin);
 	}
 
 	TImplicitObjectScaled(const TImplicitObjectScaled<TConcrete, bInstanced>& Other) = delete;
@@ -480,7 +479,7 @@ public:
 		this->MLocalBoundingBox = Other.MLocalBoundingBox;
 		SetMargin(Other.GetMargin());
 	}
-	~TImplicitObjectScaled() {}
+	~TImplicitObjectScaled(){}
 
 	static constexpr EImplicitObjectType StaticType()
 	{
@@ -1016,6 +1015,24 @@ private:
 		return new TImplicitObjectScaled<TConcrete, false>(reinterpret_cast<TRefCountPtr<TConcrete>&&>(DuplicatedShape), Obj->MScale, Obj->OuterMargin);
 	}
 
+	void InitScaledImplicit(const FVec3& Scale, FReal InMargin = 0)
+	{
+		ensureMsgf((IsScaled(MObject->GetType()) == false), TEXT("Scaled objects should not contain each other."));
+		ensureMsgf((IsInstanced(MObject->GetType()) == false), TEXT("Scaled objects should not contain instances."));
+		switch (MObject->GetType())
+		{
+		case ImplicitObjectType::Transformed:
+		case ImplicitObjectType::Union:
+			check(false);	//scale is only supported for concrete types like sphere, capsule, convex, levelset, etc... Nothing that contains other objects
+		default:
+			break;
+		}
+		this->bIsConvex = MObject->IsConvex();
+		this->bDoCollide = MObject->GetDoCollide();
+		this->OuterMargin = InMargin;
+		SetScale(Scale);
+	}
+
 	void UpdateBounds()
 	{
 		const FAABB3 UnscaledBounds = MObject->BoundingBox();
@@ -1028,13 +1045,7 @@ private:
 	template <typename QueryGeomType>
 	static auto MakeScaledHelper(const QueryGeomType& B, const TVector<T,d>& InvScale )
 	{
-		// TODO: Fixup code using this and remove it.
-
-		// Could be quite dangerous since if you create an implicit without storing it in a refcountptr
-		// Destroying the HackBPtr will probably set the refcount to 0 triggering the auto deletion
-		TRefCountPtr<QueryGeomType> HackBPtr(const_cast<QueryGeomType*>(&B));	//todo: hack, need scaled object to accept raw ptr similar to transformed implicit
-
-		TImplicitObjectScaled<QueryGeomType> ScaledB(HackBPtr, InvScale);
+		TImplicitObjectScaled<QueryGeomType, true> ScaledB(const_cast<QueryGeomType*>(&B), InvScale);
 		return ScaledB;
 	}
 
