Remove modifiers/watchers which are pending delete to fix crashes when calbacks are deleted

diff --git a/Engine/Source/Runtime/Experimental/Chaos/Private/PBDRigidsSolver.cpp b/Engine/Source/Runtime/Experimental/Chaos/Private/PBDRigidsSolver.cpp
--- a/Engine/Source/Runtime/Experimental/Chaos/Private/PBDRigidsSolver.cpp
+++ b/Engine/Source/Runtime/Experimental/Chaos/Private/PBDRigidsSolver.cpp
@@ -1647,22 +1647,65 @@ namespace Chaos
 			RemovedCallbackObject->bPendingDelete = true;
 		}
 
+		for (int32 Idx = SimCallbackObjects.Num() - 1; Idx >= 0; --Idx)
+		{
+			if (SimCallbackObjects[Idx]->bPendingDelete)
+			{
+				SimCallbackObjects.RemoveAtSwap(Idx, 1, false);
+			}
+		}
+
+		for (int32 Idx = MidPhaseModifiers.Num() - 1; Idx >= 0; --Idx)
+		{
+			if (MidPhaseModifiers[Idx]->bPendingDelete)
+			{
+				//will also be in SimCallbackObjects so we'll delete it in that loop
+				MidPhaseModifiers.RemoveAtSwap(Idx, 1, false);
+			}
+		}
+
+		for (int32 Idx = CCDModifiers.Num() - 1; Idx >= 0; --Idx)
+		{
+			if (CCDModifiers[Idx]->bPendingDelete)
+			{
+				//will also be in SimCallbackObjects so we'll delete it in that loop
+				CCDModifiers.RemoveAtSwap(Idx, 1, false);
+			}
+		}
+
+		for (int32 Idx = StrainModifiers.Num() - 1; Idx >= 0; --Idx)
+		{
+			if (StrainModifiers[Idx]->bPendingDelete)
+			{
+				//will also be in SimCallbackObjects so we'll delete it in that loop
+				StrainModifiers.RemoveAtSwap(Idx, 1, false);
+			}
+		}
+
 		for (int32 Idx = ContactModifiers.Num() - 1; Idx >= 0; --Idx)
 		{
-			ISimCallbackObject* Callback = ContactModifiers[Idx];
-			if (Callback->bPendingDelete)
+			if (ContactModifiers[Idx]->bPendingDelete)
 			{
 				//will also be in SimCallbackObjects so we'll delete it in that loop
 				ContactModifiers.RemoveAtSwap(Idx, 1, false);
 			}
 		}
 
-		for (int32 Idx = SimCallbackObjects.Num() - 1; Idx >= 0; --Idx)
+		for (int32 Idx = RegistrationWatchers.Num() - 1; Idx >= 0; --Idx)
 		{
-			ISimCallbackObject* Callback = SimCallbackObjects[Idx];
-			if (Callback->bPendingDelete)
+			if (RegistrationWatchers[Idx]->bPendingDelete)
 			{
-				SimCallbackObjects.RemoveAtSwap(Idx, 1, false);
+				//will also be in SimCallbackObjects so we'll delete it in that loop
+				RegistrationWatchers.RemoveAtSwap(Idx, 1, false);
+			}
+		}
+
+		for (int32 Idx = UnregistrationWatchers.Num() - 1; Idx >= 0; --Idx)
+		{
+			if (UnregistrationWatchers[Idx]->bPendingDelete)
+			{
+				//will also be in SimCallbackObjects so we'll delete it in that loop
+				UnregistrationWatchers.RemoveAtSwap(Idx, 1, false);
 			}
 		}
 
