ChaosFlesh : Add dynamic coloring for Gauss Seidel weak constraints.

diff --git a/Engine/Source/Runtime/Experimental/Chaos/Private/Chaos/Deformable/ChaosDeformableSolver.cpp b/Engine/Source/Runtime/Experimental/Chaos/Private/Chaos/Deformable/ChaosDeformableSolver.cpp
--- a/Engine/Source/Runtime/Experimental/Chaos/Private/Chaos/Deformable/ChaosDeformableSolver.cpp
+++ b/Engine/Source/Runtime/Experimental/Chaos/Private/Chaos/Deformable/ChaosDeformableSolver.cpp
@@ -1219,7 +1219,7 @@ namespace Chaos::Softs
 			if (Property.bEnablePositionTargets)
 			{
 				TArray<TArray<int32>> ParticlesPerColor;
-				GSWeakConstraints->ComputeInitialWCData(Evolution->Particles(), GSNeohookeanConstraints->GetMeshConstraints(), GSNeohookeanConstraints->GetIncidentElements(), GSNeohookeanConstraints->GetIncidentElementsLocal(), ParticlesPerColor);
+				GSWeakConstraints->ComputeInitialWCData(Evolution->Particles(), GSNeohookeanConstraints->GetMeshArray(), GSNeohookeanConstraints->GetIncidentElements(), GSNeohookeanConstraints->GetIncidentElementsLocal(), ParticlesPerColor);
 				GSNeohookeanConstraints->SetParticlesPerColor(MoveTemp(ParticlesPerColor));
 
 				GSNeohookeanConstraints->AddAdditionalRes = [this](const FSolverParticles& InParticles, const int32 p, const FSolverReal Dt, TVec3<FSolverReal>& res)
@@ -1253,7 +1253,7 @@ namespace Chaos::Softs
 			if (Property.bEnablePositionTargets)
 			{
 				TArray<TArray<int32>> ParticlesPerColor;
-				GSWeakConstraints->ComputeInitialWCData(Evolution->Particles(), GSCorotatedConstraints->GetMeshConstraints(), GSCorotatedConstraints->GetIncidentElements(), GSCorotatedConstraints->GetIncidentElementsLocal(), ParticlesPerColor);
+				GSWeakConstraints->ComputeInitialWCData(Evolution->Particles(), GSCorotatedConstraints->GetMeshArray(), GSCorotatedConstraints->GetIncidentElements(), GSCorotatedConstraints->GetIncidentElementsLocal(), ParticlesPerColor);
 				GSCorotatedConstraints->SetParticlesPerColor(MoveTemp(ParticlesPerColor));
 
 				GSCorotatedConstraints->AddAdditionalRes = [this](const FSolverParticles& InParticles, const int32 p, const FSolverReal Dt, TVec3<FSolverReal>& res)
@@ -1300,7 +1300,7 @@ namespace Chaos::Softs
 					[this](FSolverParticles& InParticles, const FSolverReal Dt)
 				{
 					TArray<TArray<int32>> ParticlesPerColor;
-					this->GSWeakConstraints->ComputeCollisionWCData(this->Evolution->Particles(), this->GSNeohookeanConstraints->GetMeshConstraints(), this->GSNeohookeanConstraints->GetIncidentElements(), this->GSNeohookeanConstraints->GetIncidentElementsLocal(), ParticlesPerColor);
+					this->GSWeakConstraints->ComputeCollisionWCData(this->Evolution->Particles(), this->GSNeohookeanConstraints->GetMeshArray(), this->GSNeohookeanConstraints->GetIncidentElements(), this->GSNeohookeanConstraints->GetIncidentElementsLocal(), ParticlesPerColor);
 					this->GSNeohookeanConstraints->SetParticlesPerColor(MoveTemp(ParticlesPerColor));
 				};
 			}
@@ -1310,7 +1310,7 @@ namespace Chaos::Softs
 					[this](FSolverParticles& InParticles, const FSolverReal Dt)
 				{
 					TArray<TArray<int32>> ParticlesPerColor;
-					this->GSWeakConstraints->ComputeCollisionWCData(this->Evolution->Particles(), this->GSCorotatedConstraints->GetMeshConstraints(), this->GSCorotatedConstraints->GetIncidentElements(), this->GSCorotatedConstraints->GetIncidentElementsLocal(), ParticlesPerColor);
+					this->GSWeakConstraints->ComputeCollisionWCData(this->Evolution->Particles(), this->GSCorotatedConstraints->GetMeshArray(), this->GSCorotatedConstraints->GetIncidentElements(), this->GSCorotatedConstraints->GetIncidentElementsLocal(), ParticlesPerColor);
 					this->GSCorotatedConstraints->SetParticlesPerColor(MoveTemp(ParticlesPerColor));
 				};
 			}
diff --git a/Engine/Source/Runtime/Experimental/Chaos/Private/Chaos/GraphColoring.cpp b/Engine/Source/Runtime/Experimental/Chaos/Private/Chaos/GraphColoring.cpp
--- a/Engine/Source/Runtime/Experimental/Chaos/Private/Chaos/GraphColoring.cpp
+++ b/Engine/Source/Runtime/Experimental/Chaos/Private/Chaos/GraphColoring.cpp
@@ -3,6 +3,7 @@
 #include "Chaos/Array.h"
 #include "ChaosLog.h"
 #include "Chaos/Framework/Parallel.h"
+#include "Containers/BitArray.h"
 
 template<typename T>
 static bool VerifyGraph(TArray<TArray<int32>> ColorGraph, const TArray<Chaos::TVec2<int32>>& Graph, const Chaos::TDynamicParticles<T, 3>& InParticles)
@@ -330,7 +331,7 @@ static bool VerifyNodalColoring(const TArray<Chaos::TVec4<int32>>& Graph, const
 template <typename T>
 static bool VerifyNodalColoring(const TArray<TArray<int32>>& Graph, const Chaos::TDynamicParticles<T, 3>& InParticles, const int32 GraphParticlesStart, const int32 GraphParticlesEnd, const TArray<TArray<int32>>& IncidentElements, const TArray<TArray<int32>>& IncidentElementsLocalIndex, const TArray<TArray<int32>> ParticlesPerColor)
 {
-	checkSlow(GraphParticlesStart <= GraphParticlesEnd);
+ 	checkSlow(GraphParticlesStart <= GraphParticlesEnd);
 	checkSlow(GraphParticlesEnd <= (int32)InParticles.Size());
 	TArray<bool> ParticleIsIncluded;
 	ParticleIsIncluded.Init(false, InParticles.Size());
@@ -396,6 +397,74 @@ static bool VerifyNodalColoring(const TArray<TArray<int32>>& Graph, const Chaos:
 	return true;
 }
 
+template <typename T>
+static bool VerifyExtraNodalColoring(const Chaos::TDynamicParticles<T, 3>& InParticles, const TArray<TArray<int32>>& Graph, const TArray<TArray<int32>>& ExtraGraph, const TArray<TArray<int32>>& IncidentElements, const TArray<TArray<int32>>& ExtraIncidentElements,const TArray<int32>& ParticleColors,const TArray<TArray<int32>>& ParticlesPerColor)
+{
+	TBitArray ParticleIsIncluded(false, InParticles.Size());
+	for (int32 i = 0; i < ParticlesPerColor.Num(); i++)
+	{
+		for (int32 j = 0; j < ParticlesPerColor[i].Num(); j++)
+		{
+			ParticleIsIncluded[ParticlesPerColor[i][j]] = true;
+		}
+	}
+
+	for (int32 ParticleIndex = 0; ParticleIndex < (int32)InParticles.Size(); ParticleIndex++)
+	{
+		if (InParticles.InvM(ParticleIndex) != (T)0.)
+		{
+			if (!ParticleIsIncluded[ParticleIndex])
+			{
+				return false;
+			}
+		}
+	}
+
+
+	for (int32 i = 0; i < ParticlesPerColor.Num(); i++)
+	{
+		TSet<int32> IncidentParticles;
+		IncidentParticles.Reserve(ParticlesPerColor[i].Num());
+		for (int32 j = 0; j < ParticlesPerColor[i].Num(); j++)
+		{
+			int32 ParticleIndex = ParticlesPerColor[i][j];
+			TSet<int32> LocalIncidentParticles;
+			for (int32 k = 0; k < IncidentElements[ParticleIndex].Num(); k++)
+			{
+				int32 ElementIndex = IncidentElements[ParticleIndex][k];
+				for (int32 ie = 0; ie < Graph[ElementIndex].Num(); ie++)
+				{
+					LocalIncidentParticles.Add(Graph[ElementIndex][ie]);
+				}
+			}
+			if (ParticleIndex < ExtraIncidentElements.Num())
+			{
+				for (int32 k = 0; k < ExtraIncidentElements[ParticleIndex].Num(); k++)
+				{
+					int32 ElementIndex = ExtraIncidentElements[ParticleIndex][k];
+					for (int32 ie = 0; ie < ExtraGraph[ElementIndex].Num(); ie++)
+					{
+						LocalIncidentParticles.Add(ExtraGraph[ElementIndex][ie]);
+					}
+				}
+			}
+			if (IncidentParticles.Contains(ParticleIndex))
+			{
+				return false;
+			}
+			else
+			{
+				for (int32 LocalParticle : LocalIncidentParticles)
+				{
+					IncidentParticles.Add(LocalParticle);
+				}
+			}
+		}
+	}
+
+	return true;
+}
+
 template<typename T>
 TArray<TArray<int32>> Chaos::FGraphColoring::ComputeGraphColoring(const TArray<Chaos::TVec2<int32>>& Graph, const Chaos::TDynamicParticles<T, 3>& InParticles, const int32 GraphParticlesStart, const int32 GraphParticlesEnd)
 {
@@ -1528,7 +1597,7 @@ TArray<TArray<int32>> Chaos::ComputeNodalColoring(const TArray<TVec4<int32>>& Gr
 
 
 template<typename T>
-TArray<TArray<int32>> Chaos::ComputeNodalColoring(const TArray<TArray<int32>>& Graph, const Chaos::TDynamicParticles<T, 3>& InParticles, const int32 GraphParticlesStart, const int32 GraphParticlesEnd, const TArray<TArray<int32>>& IncidentElements, const TArray<TArray<int32>>& IncidentElementsLocalIndex)
+TArray<TArray<int32>> Chaos::ComputeNodalColoring(const TArray<TArray<int32>>& Graph, const Chaos::TDynamicParticles<T, 3>& InParticles, const int32 GraphParticlesStart, const int32 GraphParticlesEnd, const TArray<TArray<int32>>& IncidentElements, const TArray<TArray<int32>>& IncidentElementsLocalIndex, TArray<int32>* ParticleColorsOut)
 {
 	using namespace Chaos;
 
@@ -1620,12 +1689,105 @@ TArray<TArray<int32>> Chaos::ComputeNodalColoring(const TArray<TArray<int32>>& G
 		}
 	}
 
+	if (ParticleColorsOut)
+	{
+		*ParticleColorsOut = ParticleColors;
+	}
+
 	checkSlow(VerifyNodalColoring<T>(Graph, InParticles, GraphParticlesStart, GraphParticlesEnd, IncidentElements, IncidentElementsLocalIndex, ParticlesPerColor));
 
 	return ParticlesPerColor;
 }
 
 
+template<typename T>
+void Chaos::ComputeExtraNodalColoring(const TArray<TArray<int32>>& Graph, const TArray<TArray<int32>>& ExtraGraph, const Chaos::TDynamicParticles<T, 3>& InParticles, const TArray<TArray<int32>>& IncidentElements, const TArray<TArray<int32>>& ExtraIncidentElements, TArray<int32>& ParticleColors, TArray<TArray<int32>>& ParticlesPerColor)
+{
+	TBitArray ParticleIsAffected(false, IncidentElements.Num());
+	for (int32 i = 0; i < ExtraIncidentElements.Num(); i++)
+	{
+		if (ExtraIncidentElements[i].Num() > 0)
+		{
+			ParticleIsAffected[i] = true;
+		}
+	}
+
+	TSet<int32> UsedColors;
+	for (int32 i = 0; i < ExtraIncidentElements.Num(); i++)
+	{
+		if (ParticleIsAffected[i])
+		{
+			int32 OriginalColor = ParticleColors[i];
+			ParticleColors[i] = INDEX_NONE;
+			if (OriginalColor != INDEX_NONE)
+			{
+				UsedColors.Reset();
+				UsedColors.Reserve(ExtraIncidentElements[i].Num() + IncidentElements[i].Num());
+				for (int32 j = 0; j < IncidentElements[i].Num(); j++)
+				{
+					int32 ConstraintIndex = IncidentElements[i][j];
+					for (int32 ie = 0; ie < Graph[ConstraintIndex].Num(); ie++)
+					{
+						UsedColors.Add(ParticleColors[Graph[ConstraintIndex][ie]]);
+					}
+				}
+				for (int32 j = 0; j < ExtraIncidentElements[i].Num(); j++)
+				{
+					int32 ConstraintIndex = ExtraIncidentElements[i][j];
+					for (int32 ie = 0; ie < ExtraGraph[ConstraintIndex].Num(); ie++)
+					{
+						UsedColors.Add(ParticleColors[ExtraGraph[ConstraintIndex][ie]]);
+					}
+				}
+				if (UsedColors.Contains(OriginalColor))
+				{
+					OriginalColor = 0;
+					while (UsedColors.Contains(OriginalColor))
+					{
+						OriginalColor++;
+					}
+					ParticleColors[i] = OriginalColor;
+				}
+				else
+				{
+					ParticleColors[i] = OriginalColor;
+					ParticleIsAffected[i] = false;
+				}
+			}
+		}
+	}
+
+
+	PhysicsParallelFor(ParticlesPerColor.Num(), [&](const int32 i)
+		{
+			int32 CurrentIndex = 0;
+			for (int32 j = 0; j < ParticlesPerColor[i].Num(); j++)
+			{
+				if (!ParticleIsAffected[ParticlesPerColor[i][j]])
+				{
+					if (CurrentIndex != j)
+					{
+						ParticlesPerColor[i][CurrentIndex] = ParticlesPerColor[i][j];
+					}
+					CurrentIndex += 1;
+				}
+			}
+			ParticlesPerColor[i].SetNum(CurrentIndex);
+		}, ParticlesPerColor[0].Num() < 1000);
+
+	ParticlesPerColor.SetNum(ParticleColors.Max() + 1);
+	for (int32 i = 0; i < ExtraIncidentElements.Num(); i++)
+	{
+		if (ParticleIsAffected[i])
+		{
+			ParticlesPerColor[ParticleColors[i]].Emplace(i);
+		}
+	}
+
+	checkSlow(VerifyExtraNodalColoring<T>(InParticles, Graph, ExtraGraph, IncidentElements, ExtraIncidentElements, ParticleColors, ParticlesPerColor));
+}
+
+
 
 
 template CHAOS_API TArray<TArray<int32>> Chaos::FGraphColoring::ComputeGraphColoring<Chaos::FRealSingle>(const TArray<Chaos::TVector<int32, 2>>&, const Chaos::TDynamicParticles<Chaos::FRealSingle, 3>&, const int32 GraphParticlesStart, const int32 GraphParticlesEnd);
@@ -1642,5 +1804,7 @@ template CHAOS_API void Chaos::ComputeWeakConstraintsColoring<Chaos::FRealSingle
 template CHAOS_API void Chaos::ComputeWeakConstraintsColoring<Chaos::FRealDouble>(const TArray<TArray<int32>>& Indices, const TArray<TArray<int32>>& SecondIndices, const Chaos::TDynamicParticles<Chaos::FRealDouble, 3>& InParticles, TArray<TArray<int32>>& ConstraintsPerColor);
 template CHAOS_API TArray<TArray<int32>> Chaos::ComputeNodalColoring<Chaos::FRealSingle>(const TArray<TVec4<int32>>& Graph, const Chaos::TDynamicParticles<Chaos::FRealSingle, 3>& InParticles, const int32 GraphParticlesStart, const int32 GraphParticlesEnd, const TArray<TArray<int32>>& IncidentElements, const TArray<TArray<int32>>& IncidentElementsLocalIndex);
 template CHAOS_API TArray<TArray<int32>> Chaos::ComputeNodalColoring<Chaos::FRealDouble>(const TArray<TVec4<int32>>& Graph, const Chaos::TDynamicParticles<Chaos::FRealDouble, 3>& InParticles, const int32 GraphParticlesStart, const int32 GraphParticlesEnd, const TArray<TArray<int32>>& IncidentElements, const TArray<TArray<int32>>& IncidentElementsLocalIndex);
-template CHAOS_API TArray<TArray<int32>> Chaos::ComputeNodalColoring<Chaos::FRealSingle>(const TArray<TArray<int32>>& Graph, const Chaos::TDynamicParticles<Chaos::FRealSingle, 3>& InParticles, const int32 GraphParticlesStart, const int32 GraphParticlesEnd, const TArray<TArray<int32>>& IncidentElements, const TArray<TArray<int32>>& IncidentElementsLocalIndex);
-template CHAOS_API TArray<TArray<int32>> Chaos::ComputeNodalColoring<Chaos::FRealDouble>(const TArray<TArray<int32>>& Graph, const Chaos::TDynamicParticles<Chaos::FRealDouble, 3>& InParticles, const int32 GraphParticlesStart, const int32 GraphParticlesEnd, const TArray<TArray<int32>>& IncidentElements, const TArray<TArray<int32>>& IncidentElementsLocalIndex);
\ No newline at end of file
+template CHAOS_API TArray<TArray<int32>> Chaos::ComputeNodalColoring<Chaos::FRealSingle>(const TArray<TArray<int32>>& Graph, const Chaos::TDynamicParticles<Chaos::FRealSingle, 3>& InParticles, const int32 GraphParticlesStart, const int32 GraphParticlesEnd, const TArray<TArray<int32>>& IncidentElements, const TArray<TArray<int32>>& IncidentElementsLocalIndex, TArray<int32>* ParticleColorsOut);
+template CHAOS_API TArray<TArray<int32>> Chaos::ComputeNodalColoring<Chaos::FRealDouble>(const TArray<TArray<int32>>& Graph, const Chaos::TDynamicParticles<Chaos::FRealDouble, 3>& InParticles, const int32 GraphParticlesStart, const int32 GraphParticlesEnd, const TArray<TArray<int32>>& IncidentElements, const TArray<TArray<int32>>& IncidentElementsLocalIndex, TArray<int32>* ParticleColorsOut);
+template CHAOS_API void Chaos::ComputeExtraNodalColoring<Chaos::FRealSingle>(const TArray<TArray<int32>>& Graph, const TArray<TArray<int32>>& ExtraGraph, const Chaos::TDynamicParticles<Chaos::FRealSingle, 3>& InParticles, const TArray<TArray<int32>>& IncidentElements, const TArray<TArray<int32>>& ExtraIncidentElements, TArray<int32>& ParticleColors, TArray<TArray<int32>>& ParticlesPerColor);
+template CHAOS_API void Chaos::ComputeExtraNodalColoring<Chaos::FRealDouble>(const TArray<TArray<int32>>& Graph, const TArray<TArray<int32>>& ExtraGraph, const Chaos::TDynamicParticles<Chaos::FRealDouble, 3>& InParticles, const TArray<TArray<int32>>& IncidentElements, const TArray<TArray<int32>>& ExtraIncidentElements, TArray<int32>& ParticleColors, TArray<TArray<int32>>& ParticlesPerColor);
\ No newline at end of file
diff --git a/Engine/Source/Runtime/Experimental/Chaos/Public/Chaos/Deformable/GaussSeidelCorotatedConstraints.h b/Engine/Source/Runtime/Experimental/Chaos/Public/Chaos/Deformable/GaussSeidelCorotatedConstraints.h
--- a/Engine/Source/Runtime/Experimental/Chaos/Public/Chaos/Deformable/GaussSeidelCorotatedConstraints.h
+++ b/Engine/Source/Runtime/Experimental/Chaos/Public/Chaos/Deformable/GaussSeidelCorotatedConstraints.h
@@ -75,6 +75,15 @@ namespace Chaos::Softs
 			}
 			InitColor(InParticles);
 			InitializeLambdas();
+			MeshArray.SetNum(MeshConstraints.Num());
+			for (int32 i = 0; i < MeshConstraints.Num(); i++)
+			{
+				MeshArray[i].SetNum(4);
+				for (int32 ie = 0; ie < 4; ie++)
+				{
+					MeshArray[i][ie] = MeshConstraints[i][ie];
+				}
+			}
 		}
 
 		virtual ~FGaussSeidelCorotatedConstraints() {}
@@ -127,6 +136,7 @@ namespace Chaos::Softs
 		TArray<TArray<int32>> GetIncidentElements() const { return IncidentElements; }
 		TArray<TArray<int32>> GetIncidentElementsLocal() const { return IncidentElementsLocal; }
 		TArray<TVector<int32, 4>> GetMeshConstraints() const { return MeshConstraints; }
+		TArray<TArray<int32>> GetMeshArray() const { return MeshArray; }
 		void SetParticlesPerColor(TArray<TArray<int32>>&& InParticlesPerColor) {ParticlesPerColor = MoveTemp(InParticlesPerColor); }
 
 
@@ -331,6 +341,7 @@ namespace Chaos::Softs
 
 	protected:	
 
+		TArray<TArray<int32>> MeshArray;
 		TArray<int32> Particle2Incident;
 		TArray<TArray<int32>> IncidentElements;
 		TArray<TArray<int32>> IncidentElementsLocal;
diff --git a/Engine/Source/Runtime/Experimental/Chaos/Public/Chaos/Deformable/GaussSeidelWeakConstraints.h b/Engine/Source/Runtime/Experimental/Chaos/Public/Chaos/Deformable/GaussSeidelWeakConstraints.h
--- a/Engine/Source/Runtime/Experimental/Chaos/Public/Chaos/Deformable/GaussSeidelWeakConstraints.h
+++ b/Engine/Source/Runtime/Experimental/Chaos/Public/Chaos/Deformable/GaussSeidelWeakConstraints.h
@@ -38,7 +38,6 @@ namespace Chaos::Softs
 
 			for (int32 i = 0; i < Indices.Num(); i++)
 			{
-				ensureMsgf(Indices[i].Num() + SecondIndices[i].Num() == 4, TEXT("Currently does not support non point tri pair weak constraints"));
 				TSet<int32> IndicesSet = TSet<int32>(Indices[i]);
 				for (int32 j = 0; j < SecondIndices[i].Num(); j++)
 				{
@@ -51,13 +50,13 @@ namespace Chaos::Softs
 
 		virtual ~FGaussSeidelWeakConstraints() {}
 
-		void ComputeInitialWCData(const ParticleType& InParticles, const TArray<TVector<int32, 4>>& MeshConstraints, const TArray <TArray<int32>>& MeshIncidentElements, const TArray <TArray<int32>>& MeshIncidentElementsLocal, TArray<TArray<int32>>& ParticlesPerColor)
+		void ComputeInitialWCData(const ParticleType& InParticles, const TArray<TArray<int32>>& MeshConstraints, const TArray <TArray<int32>>& MeshIncidentElements, const TArray <TArray<int32>>& MeshIncidentElementsLocal, TArray<TArray<int32>>& ParticlesPerColor)
 		{
-			TArray<TVector<int32, 4>> ExtraConstraints;
-			ExtraConstraints.Init(TVector<int32, 4>(INDEX_NONE), Indices.Num());
+			TArray<TArray<int32>> ExtraConstraints;
+			ExtraConstraints.Init(TArray<int32>(), Indices.Num());
 			for (int32 i = 0; i < Indices.Num(); i++)
 			{
-				ensureMsgf(Indices[i].Num() + SecondIndices[i].Num() == 4, TEXT("Currently does not support non point tri pair weak constraints"));
+				ExtraConstraints[i].SetNum(Indices[i].Num() + SecondIndices[i].Num());
 				for (int32 j = 0; j < Indices[i].Num(); j++)
 				{
 					ExtraConstraints[i][j] = Indices[i][j];
@@ -67,7 +66,7 @@ namespace Chaos::Softs
 					ExtraConstraints[i][j+Indices[i].Num()] = SecondIndices[i][j];
 				}
 			}
-			TArray<TVector<int32, 4>> TotalConstraints = MeshConstraints;
+			TArray<TArray<int32>> TotalConstraints = MeshConstraints;
 			TotalConstraints += ExtraConstraints;
 			TArray<TArray<int32>> TotalIncidentElements = MeshIncidentElements, TotalIncidentElementsLocal = MeshIncidentElementsLocal;
 			WCIncidentElements = Chaos::Utilities::ComputeIncidentElements(ExtraConstraints, &WCIncidentElementsLocal);
@@ -84,7 +83,7 @@ namespace Chaos::Softs
 				}
 			}
 
-			ParticlesPerColor = ComputeNodalColoring(TotalConstraints, InParticles, 0, InParticles.Size(), TotalIncidentElements, TotalIncidentElementsLocal);
+			ParticlesPerColor = ComputeNodalColoring(TotalConstraints, InParticles, 0, InParticles.Size(), TotalIncidentElements, TotalIncidentElementsLocal, &ParticleColors);
 
 			NodalWeights.Init({}, InParticles.Size());
 
@@ -133,6 +132,20 @@ namespace Chaos::Softs
 			NoCollisionNodalWeights = NodalWeights;
 			NoCollisionConstraints = ExtraConstraints;
 			InitialWCSize = Indices.Num();
+
+			NoCollisionWCIncidentElements = WCIncidentElements;
+			NoCollisionWCIncidentElementsLocal = WCIncidentElementsLocal;
+
+			StaticConstraints = MeshConstraints;
+			StaticConstraints += NoCollisionConstraints;
+			StaticIncidentElements = MeshIncidentElements;
+			for (int32 i = 0; i < NoCollisionWCIncidentElements.Num(); i++)
+			{
+				if (NoCollisionWCIncidentElements[i].Num() > 0)
+				{
+					StaticIncidentElements[i] += NoCollisionWCIncidentElements[i];
+				}
+			}
 		}
 
 		void AddWCResidual(const ParticleType& InParticles, const int32 p, const T Dt, TVec3<T>& res)
@@ -581,14 +594,14 @@ namespace Chaos::Softs
 			Resize(ConstraintNum);
 		}
 
-		void ComputeCollisionWCData(const ParticleType& InParticles, const TArray<TVector<int32, 4>>& MeshConstraints, const TArray <TArray<int32>>& MeshIncidentElements, const TArray <TArray<int32>>& MeshIncidentElementsLocal, TArray<TArray<int32>>& ParticlesPerColor)
+		void ComputeCollisionWCData(const ParticleType& InParticles, const TArray<TArray<int32>>& MeshConstraints, const TArray <TArray<int32>>& MeshIncidentElements, const TArray <TArray<int32>>& MeshIncidentElementsLocal, TArray<TArray<int32>>& ParticlesPerColor)
 		{
 			ensureMsgf(Indices.Num() >= InitialWCSize, TEXT("The size of Indices is smaller than InitialWCSize"));
-			TArray<TVector<int32, 4>> ExtraConstraints;
-			ExtraConstraints.Init(TVector<int32, 4>(INDEX_NONE), Indices.Num() - InitialWCSize);
+			TArray<TArray<int32>> ExtraConstraints;
+			ExtraConstraints.Init(TArray<int32>(), Indices.Num());
 			for (int32 i = InitialWCSize; i < Indices.Num(); i++)
 			{
-				ensureMsgf(Indices[i].Num() + SecondIndices[i].Num() == 4, TEXT("Currently does not support non point tri pair weak constraints"));
+				ExtraConstraints[i - InitialWCSize].SetNum(Indices[i].Num() + SecondIndices[i].Num());
 				for (int32 j = 0; j < Indices[i].Num(); j++)
 				{
 					ExtraConstraints[i - InitialWCSize][j] = Indices[i][j];
@@ -602,29 +615,52 @@ namespace Chaos::Softs
 			TArray<TArray<int32>> ExtraWCIncidentElements, ExtraWCIncidentElementsLocal;
 			ExtraWCIncidentElements = Chaos::Utilities::ComputeIncidentElements(ExtraConstraints, &ExtraWCIncidentElementsLocal);
 
-			TArray<TVector<int32, 4>> TotalConstraints = NoCollisionConstraints;
-			TotalConstraints += ExtraConstraints;
+			//TArray<TVector<int32, 4>> TotalConstraints = NoCollisionConstraints;
+			//TotalConstraints += ExtraConstraints;
 			//TODO (Yizhou): Make the following more efficient by computing only extra incident elements in the future. 
-			TArray<TArray<int32>> WCIncidentElementsLocalTemp;
-			WCIncidentElements = Chaos::Utilities::ComputeIncidentElements(TotalConstraints, &WCIncidentElementsLocalTemp);
-			WCIncidentElementsLocal = WCIncidentElementsLocalTemp;
+			//WCIncidentElements = Chaos::Utilities::ComputeIncidentElements(TotalConstraints, &WCIncidentElementsLocal);
+
+			//TArray<TVector<int32, 4>> TotalConstraints = NoCollisionConstraints;
+			//TotalConstraints += ExtraConstraints;
+			//TODO (Yizhou): Make the following more efficient by computing only extra incident elements in the future. 
+			//TArray<TArray<int32>> WCIncidentElementsLocalTemp;
+			////WCIncidentElements = Chaos::Utilities::ComputeIncidentElements(TotalConstraints, &WCIncidentElementsLocalTemp);
+			//WCIncidentElementsLocal = WCIncidentElementsLocalTemp;
+			WCIncidentElements = NoCollisionWCIncidentElements;
+			WCIncidentElementsLocal = NoCollisionWCIncidentElementsLocal;
+
+			if (NoCollisionWCIncidentElements.Num() < ExtraWCIncidentElements.Num())
+			{
+				WCIncidentElements.SetNum(ExtraWCIncidentElements.Num());
+				WCIncidentElementsLocal.SetNum(ExtraWCIncidentElementsLocal.Num());
+			}
+
+			for (int32 i = 0; i < ExtraWCIncidentElements.Num(); i++)
+			{
+				if (ExtraWCIncidentElements[i].Num() > 0)
+				{
+					TArray<int32> ExtraWCIncidentElementsTemp = ExtraWCIncidentElements[i];
+					for (int32 j = 0; j < ExtraWCIncidentElements[i].Num(); j++)
+					{
+						ExtraWCIncidentElementsTemp[j] += InitialWCSize;
+					}
+					WCIncidentElements[i] += ExtraWCIncidentElementsTemp;
+					WCIncidentElementsLocal[i] += ExtraWCIncidentElementsLocal[i];
+				}
+			}
+
+			//TODO(Yizhou): Check if the following variable is really necessary:
 			Particle2WCIndices.Init(INDEX_NONE, InParticles.Size());
 			for (int32 i = 0; i < WCIncidentElements.Num(); i++)
 			{
 				if (WCIncidentElements[i].Num() > 0)
 				{
-					int32 p = TotalConstraints[WCIncidentElements[i][0]][WCIncidentElementsLocal[i][0]];
+					int32 p = ExtraConstraints[WCIncidentElements[i][0]][WCIncidentElementsLocal[i][0]];
 					Particle2WCIndices[p] = i;
 				}
 			}
 
-			//TODO(Yizhou): Write the collision per timestep coloring to make the following thing much faster:
-			TArray<TVector<int32, 4>> TrueAllConstraints = MeshConstraints;
-			TrueAllConstraints += TotalConstraints;
-			TArray<TArray<int32>> TrueAllIncidentElements, TrueAllIncidentElementsLocal;
-			TrueAllIncidentElements = Chaos::Utilities::ComputeIncidentElements(TrueAllConstraints, &TrueAllIncidentElementsLocal);
-			ParticlesPerColor = ComputeNodalColoring(TrueAllConstraints, InParticles, 0, InParticles.Size(), TrueAllIncidentElements, TrueAllIncidentElementsLocal);
-
+			Chaos::ComputeExtraNodalColoring(StaticConstraints, ExtraConstraints, InParticles, StaticIncidentElements, ExtraWCIncidentElements, ParticleColors, ParticlesPerColor);
 
 			NodalWeights = NoCollisionNodalWeights;
 			for (int32 i = 0; i < ExtraWCIncidentElements.Num(); i++)
@@ -691,10 +727,19 @@ namespace Chaos::Softs
 
 		int32 InitialWCSize;
 		TArray<TArray<T>> NoCollisionNodalWeights;
-		TArray<TVector<int32, 4>> NoCollisionConstraints;
+		//TArray<TVector<int32, 4>> NoCollisionConstraints;
 		//For debugging
 		bool Detected = false;
 		TArray<int32> BoundaryVertices;
+
+		TArray<TArray<int32>> NoCollisionConstraints;
+		TArray<TArray<int32>> NoCollisionWCIncidentElements;
+		TArray<TArray<int32>> NoCollisionWCIncidentElementsLocal;
+
+		TArray<TArray<int32>> StaticConstraints;
+		TArray<TArray<int32>> StaticIncidentElements;
+
+		TArray<int32> ParticleColors; 
 	};
 
 
diff --git a/Engine/Source/Runtime/Experimental/Chaos/Public/Chaos/GraphColoring.h b/Engine/Source/Runtime/Experimental/Chaos/Public/Chaos/GraphColoring.h
--- a/Engine/Source/Runtime/Experimental/Chaos/Public/Chaos/GraphColoring.h
+++ b/Engine/Source/Runtime/Experimental/Chaos/Public/Chaos/GraphColoring.h
@@ -102,8 +102,9 @@ template<typename T>
 CHAOS_API TArray<TArray<int32>> ComputeNodalColoring(const TArray<TVec4<int32>>& Graph, const Chaos::TDynamicParticles<T, 3>& InParticles, const int32 GraphParticlesStart, const int32 GraphParticlesEnd, const TArray<TArray<int32>>& IncidentElements, const TArray<TArray<int32>>& IncidentElementsLocalIndex);
 
 template<typename T>
-CHAOS_API TArray<TArray<int32>> ComputeNodalColoring(const TArray<TArray<int32>>& Graph, const Chaos::TDynamicParticles<T, 3>& InParticles, const int32 GraphParticlesStart, const int32 GraphParticlesEnd, const TArray<TArray<int32>>& IncidentElements, const TArray<TArray<int32>>& IncidentElementsLocalIndex);
-
+CHAOS_API TArray<TArray<int32>> ComputeNodalColoring(const TArray<TArray<int32>>& Graph, const Chaos::TDynamicParticles<T, 3>& InParticles, const int32 GraphParticlesStart, const int32 GraphParticlesEnd, const TArray<TArray<int32>>& IncidentElements, const TArray<TArray<int32>>& IncidentElementsLocalIndex, TArray<int32>* ParticleColorsOut = nullptr);
 
+template<typename T>
+CHAOS_API void ComputeExtraNodalColoring(const TArray<TArray<int32>>& Graph, const TArray<TArray<int32>>& ExtraGraph, const Chaos::TDynamicParticles<T, 3>& InParticles, const TArray<TArray<int32>>& IncidentElements, const TArray<TArray<int32>>& ExtraIncidentElements, TArray<int32>& ParticleColors, TArray<TArray<int32>>& ParticlesPerColor);
 
 }
diff --git a/Engine/Source/Runtime/Experimental/Chaos/Public/Chaos/Utilities.h b/Engine/Source/Runtime/Experimental/Chaos/Public/Chaos/Utilities.h
--- a/Engine/Source/Runtime/Experimental/Chaos/Public/Chaos/Utilities.h
+++ b/Engine/Source/Runtime/Experimental/Chaos/Public/Chaos/Utilities.h
@@ -1051,7 +1051,7 @@ namespace Chaos
 			ConnectedComponentsDFSIterative(AdjacencyList, ConnectedComponents);
 		}
 
-		inline TArray<TArray<int>> ComputeIncidentElements(const TArray<TArray<int32>>& Constraints, TArray<TArray<int32>>* LocalIndex = nullptr)
+		inline TArray<TArray<int32>> ComputeIncidentElements(const TArray<TArray<int32>>& Constraints, TArray<TArray<int32>>* LocalIndex = nullptr)
 		{
 			int32 MaxIdx = 0;
 			for (int32 i = 0; i < Constraints.Num(); i++)
@@ -1085,5 +1085,17 @@ namespace Chaos
 
 			return IncidentElements;
 		}
+
+		inline void MergeIncidentElements(const TArray<TArray<int32>>& ExtraIncidentElements, const TArray<TArray<int32>>& ExtraIncidentElementsLocal, TArray<TArray<int32>>& IncidentElements, TArray<TArray<int32>>& IncidentElementsLocal)
+		{
+			if (ensureMsgf(IncidentElements.Num() == ExtraIncidentElements.Num() && IncidentElementsLocal.Num() == ExtraIncidentElementsLocal.Num(), TEXT("Input incident elements are of different size")))
+			{
+				for (int32 i = 0; i < IncidentElements.Num(); i++)
+				{
+					IncidentElements[i] += ExtraIncidentElements[i];
+					IncidentElementsLocal[i] += ExtraIncidentElementsLocal[i];
+				}
+			}
+		}
 	} // namespace Utilities
 } // namespace Chaos
\ No newline at end of file
