Fix implicit scale crash

diff --git a/Engine/Source/Runtime/Experimental/Chaos/Private/Chaos/TriangleMeshImplicitObject.cpp b/Engine/Source/Runtime/Experimental/Chaos/Private/Chaos/TriangleMeshImplicitObject.cpp
--- a/Engine/Source/Runtime/Experimental/Chaos/Private/Chaos/TriangleMeshImplicitObject.cpp
+++ b/Engine/Source/Runtime/Experimental/Chaos/Private/Chaos/TriangleMeshImplicitObject.cpp
@@ -38,8 +38,7 @@ template <typename QueryGeomType>
 static auto MakeScaledHelper(const QueryGeomType& B, const FVec3& InvScale)
 {
 	// TODO: Fixup code using this and remove it.
-	TImplicitObjectScaled<QueryGeomType,true> ScaledB(const_cast<QueryGeomType*>(&B), InvScale);
-	return ScaledB;
+	return TImplicitObjectScaled<QueryGeomType,true>(const_cast<QueryGeomType*>(&B), InvScale);
 }
 
 template <typename QueryGeomType>
diff --git a/Engine/Source/Runtime/Experimental/Chaos/Public/Chaos/ImplicitObject.h b/Engine/Source/Runtime/Experimental/Chaos/Public/Chaos/ImplicitObject.h
--- a/Engine/Source/Runtime/Experimental/Chaos/Public/Chaos/ImplicitObject.h
+++ b/Engine/Source/Runtime/Experimental/Chaos/Public/Chaos/ImplicitObject.h
@@ -125,10 +125,6 @@ public:
 			{ 
 				delete this;
 			}
-			else
-			{
-				bTransientFlag = true;
-			}
 		}
 		return Refs;
 	}
diff --git a/Engine/Source/Runtime/Experimental/Chaos/Public/Chaos/ImplicitObjectScaled.h b/Engine/Source/Runtime/Experimental/Chaos/Public/Chaos/ImplicitObjectScaled.h
--- a/Engine/Source/Runtime/Experimental/Chaos/Public/Chaos/ImplicitObjectScaled.h
+++ b/Engine/Source/Runtime/Experimental/Chaos/Public/Chaos/ImplicitObjectScaled.h
@@ -1045,8 +1045,7 @@ private:
 	template <typename QueryGeomType>
 	static auto MakeScaledHelper(const QueryGeomType& B, const TVector<T,d>& InvScale )
 	{
-		TImplicitObjectScaled<QueryGeomType, true> ScaledB(const_cast<QueryGeomType*>(&B), InvScale);
-		return ScaledB;
+		return TImplicitObjectScaled<QueryGeomType,true>(const_cast<QueryGeomType*>(&B), InvScale);
 	}
 
 	template <typename QueryGeomType>
