Chaos Cloth - Fixed the SetBackstop BP interactor function and allows to enable/disable property on the fly.

#rnx

diff --git a/Engine/Plugins/Experimental/ChaosCloth/Source/Chaos/Private/ChaosCloth/ChaosClothingSimulationCloth.cpp b/Engine/Plugins/Experimental/ChaosCloth/Source/Chaos/Private/ChaosCloth/ChaosClothingSimulationCloth.cpp
--- a/Engine/Plugins/Experimental/ChaosCloth/Source/Chaos/Private/ChaosCloth/ChaosClothingSimulationCloth.cpp
+++ b/Engine/Plugins/Experimental/ChaosCloth/Source/Chaos/Private/ChaosCloth/ChaosClothingSimulationCloth.cpp
@@ -1138,7 +1138,7 @@ void FClothingSimulationCloth::SetCollisionProperties(FRealSingle InCollisionThi
 
 void FClothingSimulationCloth::SetBackstopProperties(bool bInEnableBackstop)
 {
-	Config->GetProperties().SetEnabled(TEXT("BackstopDistance"), bInEnableBackstop);  // BackstopScale controls whether the backstop is enabled or not
+	Config->GetProperties().SetEnabled(TEXT("BackstopRadius"), bInEnableBackstop);  // BackstopRadius controls whether the backstop is enabled or not
 }
 
 void FClothingSimulationCloth::SetDampingProperties(FRealSingle InDampingCoefficient, FRealSingle InLocalDampingCoefficient)
diff --git a/Engine/Plugins/Experimental/ChaosCloth/Source/Chaos/Private/ChaosCloth/ChaosClothingSimulationInteractor.cpp b/Engine/Plugins/Experimental/ChaosCloth/Source/Chaos/Private/ChaosCloth/ChaosClothingSimulationInteractor.cpp
--- a/Engine/Plugins/Experimental/ChaosCloth/Source/Chaos/Private/ChaosCloth/ChaosClothingSimulationInteractor.cpp
+++ b/Engine/Plugins/Experimental/ChaosCloth/Source/Chaos/Private/ChaosCloth/ChaosClothingSimulationInteractor.cpp
@@ -95,7 +95,7 @@ void UChaosClothingInteractor::SetBackstop(bool bEnabled)
 {
 	ConfigCommands.Add(FChaosClothingInteractorConfigCommand::CreateLambda([bEnabled](Chaos::FClothingSimulationConfig* Config, int32 LODIndex)
 	{
-		Config->GetProperties(LODIndex).SetEnabled(TEXT("BackstopScale"), bEnabled);  // BackstopScale controls whether the backstop is enabled or not
+		Config->GetProperties(LODIndex).SetEnabled(TEXT("BackstopRadius"), bEnabled);  // BackstopRadius controls whether the backstop is enabled or not
 	}));
 }
 void UChaosClothingInteractor::SetDamping(float DampingCoefficient, float LocalDampingCoefficient)
diff --git a/Engine/Source/Runtime/Experimental/Chaos/Private/Chaos/PBDSphericalConstraint.cpp b/Engine/Source/Runtime/Experimental/Chaos/Private/Chaos/PBDSphericalConstraint.cpp
--- a/Engine/Source/Runtime/Experimental/Chaos/Private/Chaos/PBDSphericalConstraint.cpp
+++ b/Engine/Source/Runtime/Experimental/Chaos/Private/Chaos/PBDSphericalConstraint.cpp
@@ -81,6 +81,9 @@ void FPBDSphericalBackstopConstraint::SetProperties(
 		{
 			SphereRadii = WeightMaps.FindRef(GetBackstopRadiusString(PropertyCollection));
 		}
+
+		// Update enable status using the radius property
+		SetEnabled(IsBackstopRadiusEnabled(PropertyCollection));
 	}
 
 	if (IsBackstopDistanceMutable(PropertyCollection))
diff --git a/Engine/Source/Runtime/Experimental/Chaos/Public/Chaos/CollectionPropertyFacade.h b/Engine/Source/Runtime/Experimental/Chaos/Public/Chaos/CollectionPropertyFacade.h
--- a/Engine/Source/Runtime/Experimental/Chaos/Public/Chaos/CollectionPropertyFacade.h
+++ b/Engine/Source/Runtime/Experimental/Chaos/Public/Chaos/CollectionPropertyFacade.h
@@ -164,7 +164,6 @@ namespace Chaos::Softs
 		}
 
 	protected:
-
 		// No init constructor for FCollectionPropertyFacade
 		CHAOS_API FCollectionPropertyConstFacade(const TSharedPtr<const FManagedArrayCollection>& InManagedArrayCollection, ENoInit);
 
@@ -437,7 +436,6 @@ PRAGMA_ENABLE_DEPRECATION_WARNINGS
 		 */
 		CHAOS_API void Copy(const FManagedArrayCollection& InManagedArrayCollection);
 
-
 		/**
 		 * Update all properties and values from an existing collection to this property collection.
 		 * This won't copy any other groups, only data from PropertyGroup.
@@ -445,7 +443,6 @@ PRAGMA_ENABLE_DEPRECATION_WARNINGS
 		 */
 		CHAOS_API void Update(const TSharedPtr<const FManagedArrayCollection>& InManagedArrayCollection, ECollectionPropertyUpdateFlags UpdateFlags);
 
-
 		//~ Add values
 		template<typename T, TEMPLATE_REQUIRES(TIsWeightedType<T>::Value)>
 		inline int32 AddWeightedValue(const FString& Key, const T& LowValue, const T& HighValue, ECollectionPropertyFlags Flags = ECollectionPropertyFlags::Enabled);
@@ -607,7 +604,7 @@ PRAGMA_ENABLE_DEPRECATION_WARNINGS \
 	bool Is##PropertyName##Mutable(const FCollectionPropertyConstFacade& PropertyCollection) const \
 	{ \
 		return PropertyName##Index != INDEX_NONE && \
-			PropertyCollection.IsEnabled(PropertyName##Index) && PropertyCollection.IsAnimatable(PropertyName##Index) && (PropertyCollection.IsDirty(PropertyName##Index) || PropertyCollection.IsStringDirty(PropertyName##Index)); \
+			PropertyCollection.IsAnimatable(PropertyName##Index) && (PropertyCollection.IsDirty(PropertyName##Index) || PropertyCollection.IsStringDirty(PropertyName##Index)); \
 	} \
 	struct F##PropertyName##Index \
 	{ \
diff --git a/Engine/Source/Runtime/Experimental/Chaos/Public/Chaos/PBDSphericalConstraint.h b/Engine/Source/Runtime/Experimental/Chaos/Public/Chaos/PBDSphericalConstraint.h
--- a/Engine/Source/Runtime/Experimental/Chaos/Public/Chaos/PBDSphericalConstraint.h
+++ b/Engine/Source/Runtime/Experimental/Chaos/Public/Chaos/PBDSphericalConstraint.h
@@ -193,8 +193,8 @@ class FPBDSphericalBackstopConstraint final
 public:
 	static bool IsEnabled(const FCollectionPropertyConstFacade& PropertyCollection)
 	{
-		return IsBackstopDistanceEnabled(PropertyCollection, false) ||
-			IsBackstopDistanceAnimatable(PropertyCollection, false);  // Backstop can be re-enabled if animated
+		return IsBackstopRadiusEnabled(PropertyCollection, false) ||  // Radius makes more sense than distance to enable the constraint here, since without any radius there isn't a backstop
+			IsBackstopRadiusAnimatable(PropertyCollection, false);  // Backstop can be re-enabled if animated
 	}
 
 	FPBDSphericalBackstopConstraint(
