Fix ImplicitScaled move constructor (incoming ref no longer valid)

diff --git a/Engine/Source/Runtime/Experimental/Chaos/Public/Chaos/ImplicitObjectScaled.h b/Engine/Source/Runtime/Experimental/Chaos/Public/Chaos/ImplicitObjectScaled.h
--- a/Engine/Source/Runtime/Experimental/Chaos/Public/Chaos/ImplicitObjectScaled.h
+++ b/Engine/Source/Runtime/Experimental/Chaos/Public/Chaos/ImplicitObjectScaled.h
@@ -90,8 +90,8 @@ public:
 	{
 		ensureMsgf((IsScaled(MObject->GetType()) == false), TEXT("Scaled objects should not contain each other."));
 		ensureMsgf((IsInstanced(MObject->GetType()) == false), TEXT("Scaled objects should not contain instances."));
-		this->bIsConvex = Other.MObject->IsConvex();
-		this->bDoCollide = Other.MObject->GetDoCollide();
+		this->bIsConvex = MObject->IsConvex();
+		this->bDoCollide = MObject->GetDoCollide();
 		this->OuterMargin = Other.OuterMargin;
 		SetMargin(Other.GetMargin());
 	}
@@ -470,8 +470,8 @@ public:
 	{
 		ensureMsgf((IsScaled(MObject->GetType()) == false), TEXT("Scaled objects should not contain each other."));
 		ensureMsgf((IsInstanced(MObject->GetType()) == false), TEXT("Scaled objects should not contain instances."));
-		this->bIsConvex = Other.MObject->IsConvex();
-		this->bDoCollide = Other.MObject->GetDoCollide();
+		this->bIsConvex = MObject->IsConvex();
+		this->bDoCollide = MObject->GetDoCollide();
 		this->OuterMargin = Other.OuterMargin;
 		this->MScale = Other.MScale;
 		this->MInvScale = Other.MInvScale;
