geometry collection needs to compute its component space transforms even on the server otherwise will cause desync issues between client/server for character basing

diff --git a/Engine/Source/Runtime/Experimental/GeometryCollectionEngine/Private/GeometryCollection/GeometryCollectionComponent.cpp b/Engine/Source/Runtime/Experimental/GeometryCollectionEngine/Private/GeometryCollection/GeometryCollectionComponent.cpp
--- a/Engine/Source/Runtime/Experimental/GeometryCollectionEngine/Private/GeometryCollection/GeometryCollectionComponent.cpp
+++ b/Engine/Source/Runtime/Experimental/GeometryCollectionEngine/Private/GeometryCollection/GeometryCollectionComponent.cpp
@@ -4862,9 +4862,11 @@ bool UGeometryCollectionComponent::CanUseCustomRenderer() const
 
 void UGeometryCollectionComponent::RefreshCustomRenderer()
 {
-	// Don't refresh the custom renderer on the server.
-	if (CanUseCustomRenderer() && !IsNetMode(NM_DedicatedServer))
+	if (CanUseCustomRenderer())
 	{
+		// Don't refresh the custom renderer on the server but we still need to do the work of computing component space transforms.
+		const bool bUpdateRenderer = !IsNetMode(NM_DedicatedServer);
+
 		if (IGeometryCollectionExternalRenderInterface* RendererInterface = Cast<IGeometryCollectionExternalRenderInterface>(CustomRenderer))
 		{
 			if (RestCollection != nullptr)
@@ -4875,7 +4877,10 @@ void UGeometryCollectionComponent::RefreshCustomRenderer()
 				const bool bIsBroken = DynamicCollection ? !DynamicCollection->Active[RootIndex] : false;
 				const bool bRenderRootProxy = bEnableRootProxyForCustomRenderer && !bIsBroken;
 
-				RendererInterface->UpdateState(*RestCollection, ComponentTransform, !bRenderRootProxy, !bHiddenInGame);
+				if (bUpdateRenderer)
+				{
+					RendererInterface->UpdateState(*RestCollection, ComponentTransform, !bRenderRootProxy, !bHiddenInGame);
+				}
 
 				if (bRenderRootProxy)
 				{
@@ -4885,13 +4890,21 @@ void UGeometryCollectionComponent::RefreshCustomRenderer()
 						? DynamicCollection->Transform[RootIndex]
 						: FTransform::Identity;
 					ComponentSpaceTransforms[RootIndex] = RootTransform;
-					RendererInterface->UpdateRootTransform(*RestCollection, RootTransform);
+
+					if (bUpdateRenderer)
+					{
+						RendererInterface->UpdateRootTransform(*RestCollection, RootTransform);
+					}
 				}
 				else
 				{
 					// render all individual pieces
 					CalculateGlobalMatrices();
-					RendererInterface->UpdateTransforms(*RestCollection, ComponentSpaceTransforms);
+
+					if (bUpdateRenderer)
+					{
+						RendererInterface->UpdateTransforms(*RestCollection, ComponentSpaceTransforms);
+					}
 				}
 			}
 		}
