Geometry collection : fix bounds computation when custom renderer is being used
- this reverts CL 26423456 that introduced a bug trying to fix the same issue
- now component space transform are updated even when the GC is not yet broken , making sure that  all the leaf transforms are up to date for the bounds to be computed from them

diff --git a/Engine/Source/Runtime/Experimental/GeometryCollectionEngine/Private/GeometryCollection/GeometryCollectionComponent.cpp b/Engine/Source/Runtime/Experimental/GeometryCollectionEngine/Private/GeometryCollection/GeometryCollectionComponent.cpp
--- a/Engine/Source/Runtime/Experimental/GeometryCollectionEngine/Private/GeometryCollection/GeometryCollectionComponent.cpp
+++ b/Engine/Source/Runtime/Experimental/GeometryCollectionEngine/Private/GeometryCollection/GeometryCollectionComponent.cpp
@@ -687,9 +687,6 @@ FBoxSphereBounds UGeometryCollectionComponent::CalcBounds(const FTransform& Loca
 {	
 	SCOPE_CYCLE_COUNTER(STAT_GCCUpdateBounds);
 
-	const int32 RootIndex = GetRootIndex();
-	const FTransform FinalTransform = (DynamicCollection ? DynamicCollection->Transform[RootIndex] : FTransform::Identity) * LocalToWorldIn;
-	
 	if (bChaos_GC_CacheComponentSpaceBounds)
 	{
 		bool NeedBoundsUpdate = false;
@@ -707,10 +704,10 @@ FBoxSphereBounds UGeometryCollectionComponent::CalcBounds(const FTransform& Loca
 			NeedBoundsUpdate = false;
 		}
 
-		return ComponentSpaceBounds.TransformBy(FinalTransform);
+		return ComponentSpaceBounds.TransformBy(LocalToWorldIn);
 	}
 
-	return FBoxSphereBounds(ComputeBounds(FinalTransform));
+	return FBoxSphereBounds(ComputeBounds(LocalToWorldIn));
 }
 
 int32 UGeometryCollectionComponent::GetNumElements(FName Group) const
@@ -4884,16 +4881,14 @@ void UGeometryCollectionComponent::RefreshCustomRenderer()
 
 				if (bRenderRootProxy)
 				{
-					// No need to compute the component space transform in that case, since the root is always in component space 
-					const FTransform RootTransform = 
-						(DynamicCollection && DynamicCollection->Transform.IsValidIndex(RootIndex))
-						? DynamicCollection->Transform[RootIndex]
-						: FTransform::Identity;
-					ComponentSpaceTransforms[RootIndex] = RootTransform;
+					// even thoiugh we only need the root transform, we need to compute all the Component space transform 
+					// because the calclBounds relies on it
+					// @todo(chaos) : if this is a performance problem , we'll need to optimize the transform computation
+					CalculateGlobalMatrices();
 
 					if (bUpdateRenderer)
 					{
-						RendererInterface->UpdateRootTransform(*RestCollection, RootTransform);
+						RendererInterface->UpdateRootTransform(*RestCollection, ComponentSpaceTransforms[RootIndex]);
 					}
 				}
 				else
