Chaos Cloth - Tidied up the FCollectionPropertyFacade flags API, to allow for using the new enum typed flags and deprecate the uint8 flags functions, and fix the StringDirty behavior.

#rnx

diff --git a/Engine/Source/Runtime/Experimental/Chaos/Private/Chaos/CollectionPropertyFacade.cpp b/Engine/Source/Runtime/Experimental/Chaos/Private/Chaos/CollectionPropertyFacade.cpp
--- a/Engine/Source/Runtime/Experimental/Chaos/Private/Chaos/CollectionPropertyFacade.cpp
+++ b/Engine/Source/Runtime/Experimental/Chaos/Private/Chaos/CollectionPropertyFacade.cpp
@@ -101,12 +101,22 @@ namespace Chaos::Softs
 		UpdateArrays();
 		RebuildKeyIndices();
 	}
-	
+
 	FCollectionPropertyFacade::FCollectionPropertyFacade(const TSharedPtr<FManagedArrayCollection>& InManagedArrayCollection, ENoInit)
 		: FCollectionPropertyConstFacade(InManagedArrayCollection, NoInit)
 	{
 	}
 
+	void FCollectionPropertyFacade::SetFlags(int32 KeyIndex, ECollectionPropertyFlags Flags)
+	{
+		// Cannot set string dirty without also dirtying the property
+		Flags |= EnumHasAnyFlags(Flags, ECollectionPropertyFlags::StringDirty) ? ECollectionPropertyFlags::Dirty : ECollectionPropertyFlags::None;
+		// Cannot remove the dirty flags
+		Flags |= GetFlagsArray()[KeyIndex] & (ECollectionPropertyFlags::Dirty | ECollectionPropertyFlags::StringDirty);
+
+		SetValue(KeyIndex, GetFlagsArray(), Flags);
+	}
+
 	void FCollectionPropertyFacade::ClearDirtyFlags()
 	{
 		for (ECollectionPropertyFlags& Flags : GetFlagsArray())
@@ -266,12 +276,16 @@ namespace Chaos::Softs
 				for (int32 InKeyIndex = 0; InKeyIndex < NumInKeys; ++InKeyIndex)
 				{
 					const FString& PropertyName = InPropertyFacade.GetKey(InKeyIndex);
+PRAGMA_DISABLE_DEPRECATION_WARNINGS
+					// TODO: GetFlags needs to return an ECollectionPropertyFlags, not an uint8, but the uint8 getter needs to be deprecated first
+					const ECollectionPropertyFlags PropertyFlags = (ECollectionPropertyFlags)InPropertyFacade.GetFlags(InKeyIndex);
+PRAGMA_ENABLE_DEPRECATION_WARNINGS
 					int32 NewPropertyIndex = GetKeyIndex(PropertyName);
 					if (NewPropertyIndex == INDEX_NONE)
 					{
 						if (bAppendNewProperties)
 						{
-							NewPropertyIndex = AddProperty(PropertyName, (ECollectionPropertyFlags)InPropertyFacade.GetFlags(InKeyIndex));
+							NewPropertyIndex = AddProperty(PropertyName, PropertyFlags);
 						}
 						else
 						{
@@ -284,7 +298,7 @@ namespace Chaos::Softs
 					}
 					else
 					{
-						SetFlags(NewPropertyIndex, InPropertyFacade.GetFlags(InKeyIndex));
+						SetFlags(NewPropertyIndex, PropertyFlags);
 					}
 
 					// Setting as FVector3f since that is the underlying type
diff --git a/Engine/Source/Runtime/Experimental/Chaos/Public/Chaos/CollectionPropertyFacade.h b/Engine/Source/Runtime/Experimental/Chaos/Public/Chaos/CollectionPropertyFacade.h
--- a/Engine/Source/Runtime/Experimental/Chaos/Public/Chaos/CollectionPropertyFacade.h
+++ b/Engine/Source/Runtime/Experimental/Chaos/Public/Chaos/CollectionPropertyFacade.h
@@ -90,6 +90,7 @@ namespace Chaos::Softs
 		UE_DEPRECATED(5.3, "Use GetStringValue(int32) or GetStringValue(const FString&, const FString&, int32*) instead.")
 		const FString& GetStringValue(int32 KeyIndex, const FString& Default) const { return GetValue<const FString&>(KeyIndex, StringValueArray); }
 
+		UE_DEPRECATED(5.3, "uint8 GetFlags(int32) is deprecated and will soon be replaced by ECollectionPropertyFlags GetFlags(int32).")
 		uint8 GetFlags(int32 KeyIndex) const { return FlagsArray[KeyIndex]; }
 
 		bool IsEnabled(int32 KeyIndex) const { return HasAnyFlags(KeyIndex, ECollectionPropertyFlags::Enabled); }
@@ -133,9 +134,12 @@ namespace Chaos::Softs
 			return SafeGet(Key, [this](int32 KeyIndex)->FString { return GetStringValue(KeyIndex); }, Default, OutKeyIndex);
 		}
 
+		UE_DEPRECATED(5.3, "uint8 GetFlags(const FString&, uint8, int32*) is deprecated and will soon be replaced by ECollectionPropertyFlags GetFlags(const FString&, uint8, int32*).")
 		uint8 GetFlags(const FString& Key, uint8 Default = 0, int32* OutKeyIndex = nullptr) const
 		{
+PRAGMA_DISABLE_DEPRECATION_WARNINGS
 			return SafeGet(Key, [this](int32 KeyIndex)->uint8 { return GetFlags(KeyIndex); }, Default, OutKeyIndex);
+PRAGMA_ENABLE_DEPRECATION_WARNINGS
 		}
 
 		bool IsEnabled(const FString& Key, bool bDefault = false, int32* OutKeyIndex = nullptr) const
@@ -251,13 +255,10 @@ namespace Chaos::Softs
 
 		void SetStringValue(int32 KeyIndex, const FString& Value) { if (GetStringValueArray()[KeyIndex] != Value) { GetStringValueArray()[KeyIndex] = Value; SetStringDirty(KeyIndex); } }
 
-		/** SetFlags cannot be used to remove dirty flags. Use ClearDirtyFlags to remove dirty flags.*/
-		void SetFlags(int32 KeyIndex, uint8 Flags) 
-		{ 
-			const ECollectionPropertyFlags DirtyFlags = (ECollectionPropertyFlags)GetFlags(KeyIndex) & (ECollectionPropertyFlags::Dirty | ECollectionPropertyFlags::StringDirty);
-
-			SetValue(KeyIndex, GetFlagsArray(), (ECollectionPropertyFlags)Flags | DirtyFlags);
-		}
+		/** SetFlags cannot be used to remove Dirty or StringDirty flags. Use ClearDirtyFlags to remove dirty flags. */
+		void SetFlags(int32 KeyIndex, ECollectionPropertyFlags Flags);
+		UE_DEPRECATED(5.3, "Use SetFlags(int32, ECollectionPropertyFlags) instead.")
+		void SetFlags(int32 KeyIndex, uint8 Flags) { return SetFlags(KeyIndex, (ECollectionPropertyFlags)Flags); }
 
 		void SetEnabled(int32 KeyIndex, bool bEnabled) { EnableFlags(KeyIndex, ECollectionPropertyFlags::Enabled, bEnabled); }
 		void SetAnimatable(int32 KeyIndex, bool bAnimatable) { EnableFlags(KeyIndex, ECollectionPropertyFlags::Animatable, bAnimatable); }
@@ -265,7 +266,7 @@ namespace Chaos::Softs
 		void SetDirty(int32 KeyIndex) { EnableFlags(KeyIndex, ECollectionPropertyFlags::Dirty, true); }
 		UE_DEPRECATED(5.3, "SetDirty can only be set, to unset use ClearDirtyFlags instead.")
 		void SetDirty(int32 KeyIndex, bool bDirty) { EnableFlags(KeyIndex, ECollectionPropertyFlags::Dirty, bDirty); }
-		void SetStringDirty(int32 KeyIndex) { EnableFlags(KeyIndex, ECollectionPropertyFlags::StringDirty | ECollectionPropertyFlags::Dirty, true); }
+		void SetStringDirty(int32 KeyIndex) { EnableFlags(KeyIndex, ECollectionPropertyFlags::StringDirty, true); }
 
 		//~ Values set per key
 		template<typename T, TEMPLATE_REQUIRES(TIsWeightedType<T>::Value)>
@@ -302,10 +303,12 @@ namespace Chaos::Softs
 			return SafeSet(Key, [this, &Value](int32 KeyIndex) { SetStringValue(KeyIndex, Value); });
 		}
 
-		int32 SetFlags(const FString& Key, uint8 Flags)
+		int32 SetFlags(const FString& Key, ECollectionPropertyFlags Flags)
 		{
 			return SafeSet(Key, [this, Flags](int32 KeyIndex) { SetFlags(KeyIndex, Flags); });
 		}
+		UE_DEPRECATED(5.3, "Use SetFlags(const FString&, ECollectionPropertyFlags) instead.")
+		int32 SetFlags(const FString& Key, uint8 Flags) { return SetFlags(Key, (ECollectionPropertyFlags)Flags); }
 
 		int32 SetEnabled(const FString& Key, bool bEnabled)
 		{
@@ -604,7 +607,7 @@ PRAGMA_ENABLE_DEPRECATION_WARNINGS \
 	bool Is##PropertyName##Mutable(const FCollectionPropertyConstFacade& PropertyCollection) const \
 	{ \
 		return PropertyName##Index != INDEX_NONE && \
-			PropertyCollection.IsAnimatable(PropertyName##Index) && (PropertyCollection.IsDirty(PropertyName##Index) || PropertyCollection.IsStringDirty(PropertyName##Index)); \
+			PropertyCollection.IsAnimatable(PropertyName##Index) && (PropertyCollection.IsDirty(PropertyName##Index)); \
 	} \
 	struct F##PropertyName##Index \
 	{ \
