Chaos: Limit cluster children GT sync to children of dynamic clusters. Kinematic clusters aren't updating non-dynamic children and do not need to force a transform recomputation.

diff --git a/Engine/Source/Runtime/Experimental/Chaos/Private/PBDRigidsSolver.cpp b/Engine/Source/Runtime/Experimental/Chaos/Private/PBDRigidsSolver.cpp
--- a/Engine/Source/Runtime/Experimental/Chaos/Private/PBDRigidsSolver.cpp
+++ b/Engine/Source/Runtime/Experimental/Chaos/Private/PBDRigidsSolver.cpp
@@ -1951,25 +1951,36 @@ namespace Chaos
 										ActiveClusterUnions.AddUnique((FClusterUnionPhysicsProxy*)(Proxy));
 									}
 
-									const TSet<IPhysicsProxyBase*> Proxies = ClusterParticle->PhysicsProxies();
-									for (IPhysicsProxyBase* ClusterProxy : Proxies)
+									const bool bActivateChildren =
+										ClusterParticle->ObjectState() == EObjectStateType::Dynamic ||
+										Proxy->GetType() != EPhysicsProxyType::ClusterUnionProxy;
+
+									// If the cluster is dynamic, all children will likely need a transform update
+									// otherwise we should not need a full update (if the cluster is still we only
+									// need to consider de-clustered objects which should be in other parts of the 
+									// dirty view)
+									if(bActivateChildren)
 									{
-										if (!ClusterProxy)
+										const TSet<IPhysicsProxyBase*> Proxies = ClusterParticle->PhysicsProxies();
+										for(IPhysicsProxyBase* ClusterProxy : Proxies)
 										{
-											continue;
-										}
-
-										switch (ClusterProxy->GetType())
-										{
-										case EPhysicsProxyType::SingleParticleProxy:
-											ActiveRigid.AddUnique((FSingleParticlePhysicsProxy*)ClusterProxy);
-											break;
-										case EPhysicsProxyType::GeometryCollectionType:
-											ActiveGC.AddUnique((FGeometryCollectionPhysicsProxy*)(ClusterProxy));
-											break;
-										default:
-											ensure(false);
-											break;
+											if(!ClusterProxy)
+											{
+												continue;
+											}
+
+											switch(ClusterProxy->GetType())
+											{
+											case EPhysicsProxyType::SingleParticleProxy:
+												ActiveRigid.AddUnique((FSingleParticlePhysicsProxy*)ClusterProxy);
+												break;
+											case EPhysicsProxyType::GeometryCollectionType:
+												ActiveGC.AddUnique((FGeometryCollectionPhysicsProxy*)(ClusterProxy));
+												break;
+											default:
+												ensure(false);
+												break;
+											}
 										}
 									}
 								}
