UE-191196
- Zero length sweeps will now return correct normals and hit positions
- Overlaps can now return hit positions, and the normals will be in global space for all shape types

diff --git a/Engine/Source/Programs/HeadlessChaos/Private/HeadlessChaosTestOverlap.cpp b/Engine/Source/Programs/HeadlessChaos/Private/HeadlessChaosTestOverlap.cpp
--- a/Engine/Source/Programs/HeadlessChaos/Private/HeadlessChaosTestOverlap.cpp
+++ b/Engine/Source/Programs/HeadlessChaos/Private/HeadlessChaosTestOverlap.cpp
@@ -1078,6 +1078,10 @@ namespace ChaosTest
 					EXPECT_EQ(bResult, true);
 					bool bResultMTD = TriangleMesh->OverlapGeom(ScaledBox, QueryTM, 0.0, &MTDInfo, TriMeshScale);
 					EXPECT_EQ(bResult, bResultMTD);
+					FReal ErrorMargin = 0.01; // Good enough for engineering 
+					EXPECT_VECTOR_NEAR(MTDInfo.Normal, FVec3(0.1961f , 0.0f, 0.9806f), ErrorMargin);
+					EXPECT_VECTOR_NEAR(MTDInfo.Position, FVec3(49.2308f, 0.0f, 10.1538f), ErrorMargin);
+					EXPECT_NEAR(MTDInfo.Penetration, 1.1767, ErrorMargin);
 				}
 				{
 					// Inside mesh
diff --git a/Engine/Source/Runtime/Experimental/Chaos/Private/Chaos/HeightField.cpp b/Engine/Source/Runtime/Experimental/Chaos/Private/Chaos/HeightField.cpp
--- a/Engine/Source/Runtime/Experimental/Chaos/Private/Chaos/HeightField.cpp
+++ b/Engine/Source/Runtime/Experimental/Chaos/Private/Chaos/HeightField.cpp
@@ -1938,6 +1938,7 @@ namespace Chaos
 				{
 					InnerMTD->Penetration = Penetration;
 					InnerMTD->Normal = TriangleNormal;
+					InnerMTD->Position = ClosestA;
 				}
 				return true;
 			}
@@ -1978,6 +1979,7 @@ namespace Chaos
 			{
 				OutMTD->Normal = FVec3(0);
 				OutMTD->Penetration = TNumericLimits<FReal>::Lowest();
+				OutMTD->Position = FVec3(0);
 				FVec3 Points[4];
 				FAABB3 CellBounds;
 				for (const TVec2<int32>& Cell : Intersections)
diff --git a/Engine/Source/Runtime/Experimental/Chaos/Private/Chaos/Levelset.cpp b/Engine/Source/Runtime/Experimental/Chaos/Private/Chaos/Levelset.cpp
--- a/Engine/Source/Runtime/Experimental/Chaos/Private/Chaos/Levelset.cpp
+++ b/Engine/Source/Runtime/Experimental/Chaos/Private/Chaos/Levelset.cpp
@@ -1770,6 +1770,7 @@ bool FLevelSet::OverlapGeomImp(const QueryGeomType& QueryGeom, const FRigidTrans
 		{
 			OutMTD->Penetration = -TempPhi;
 			OutMTD->Normal = TempNormal;
+			OutMTD->Position = Transformed + OutMTD->Penetration * OutMTD->Normal;
 			bResult = true;
 		}
 		else
diff --git a/Engine/Source/Runtime/Experimental/Chaos/Private/Chaos/TriangleMeshImplicitObject.cpp b/Engine/Source/Runtime/Experimental/Chaos/Private/Chaos/TriangleMeshImplicitObject.cpp
--- a/Engine/Source/Runtime/Experimental/Chaos/Private/Chaos/TriangleMeshImplicitObject.cpp
+++ b/Engine/Source/Runtime/Experimental/Chaos/Private/Chaos/TriangleMeshImplicitObject.cpp
@@ -1096,6 +1096,7 @@ bool FTriangleMeshImplicitObject::OverlapGeomImp(const QueryGeomType& QueryGeom,
 					{
 						InnerMTD->Penetration = Penetration;
 						InnerMTD->Normal = TriangleNormal;
+						InnerMTD->Position = ClosestA;
 					}
 				}
 			}
diff --git a/Engine/Source/Runtime/Experimental/Chaos/Public/Chaos/GeometryQueries.h b/Engine/Source/Runtime/Experimental/Chaos/Public/Chaos/GeometryQueries.h
--- a/Engine/Source/Runtime/Experimental/Chaos/Public/Chaos/GeometryQueries.h
+++ b/Engine/Source/Runtime/Experimental/Chaos/Public/Chaos/GeometryQueries.h
@@ -26,6 +26,7 @@ namespace Chaos
 	struct FMTDInfo
 	{
 		FVec3 Normal;
+		FVec3 Position;
 		FReal Penetration;
 	};
 
@@ -62,11 +63,12 @@ namespace Chaos
 			{
 				return Utilities::CastHelper(A, BToATM, [&](const auto& AConcrete, const auto& BToAFullTM)
 				{
-					FVec3 LocalA,LocalB,LocalNormal;
+					FVec3 LocalA, LocalB, LocalNormal;
 					int32 ClosestVertexIndexA, ClosestVertexIndexB;
-					if(GJKPenetration<false, FReal>(AConcrete,B,BToAFullTM,OutMTD->Penetration,LocalA,LocalB,LocalNormal,ClosestVertexIndexA,ClosestVertexIndexB,Thickness,0.,Offset.SizeSquared() < 1e-4 ? FVec3(1,0,0) : Offset))
+					if(GJKPenetration<false, FReal>(AConcrete,B,BToAFullTM,OutMTD->Penetration, LocalA, LocalB, LocalNormal,ClosestVertexIndexA,ClosestVertexIndexB,Thickness,0.,Offset.SizeSquared() < 1e-4 ? FVec3(1,0,0) : Offset))
 					{
 						OutMTD->Normal = ATM.TransformVectorNoScale(LocalNormal);
+						OutMTD->Position = ATM.TransformPosition(LocalA);
 						return true;
 					}
 
@@ -80,22 +82,26 @@ namespace Chaos
 		}
 		else
 		{
+			bool bOverlap = false;
 			switch (AType)
 			{
 				case ImplicitObjectType::HeightField:
 				{
 					const FHeightField& AHeightField = static_cast<const FHeightField&>(A);
-					return AHeightField.OverlapGeom(B, BToATM, Thickness, OutMTD);
+					bOverlap = AHeightField.OverlapGeom(B, BToATM, Thickness, OutMTD);
+					break;
 				}
 				case ImplicitObjectType::TriangleMesh:
 				{
 					const FTriangleMeshImplicitObject& ATriangleMesh = static_cast<const FTriangleMeshImplicitObject&>(A);
-					return ATriangleMesh.OverlapGeom(B, BToATM, Thickness, OutMTD);
+					bOverlap = ATriangleMesh.OverlapGeom(B, BToATM, Thickness, OutMTD);
+					break;
 				}
 				case ImplicitObjectType::LevelSet:
 				{
 					const FLevelSet& ALevelSet = static_cast<const FLevelSet&>(A);
-					return ALevelSet.OverlapGeom(B, BToATM, Thickness, OutMTD);
+					bOverlap = ALevelSet.OverlapGeom(B, BToATM, Thickness, OutMTD);
+					break;
 				}
 				case ImplicitObjectType::Union:
 				case ImplicitObjectType::UnionClustered:
@@ -123,12 +129,12 @@ namespace Chaos
 					if(IsScaled(AType))
 					{
 						const auto& AScaled = TImplicitObjectScaled<FTriangleMeshImplicitObject>::AsScaledChecked(A);
-						return AScaled.LowLevelOverlapGeom(B, BToATM, Thickness, OutMTD);
+						bOverlap =  AScaled.LowLevelOverlapGeom(B, BToATM, Thickness, OutMTD);
 					}
 					else if(IsInstanced(AType))
 					{
 						const auto& AInstanced = TImplicitObjectInstanced<FTriangleMeshImplicitObject>::AsInstancedChecked(A);
-						return AInstanced.LowLevelOverlapGeom(B, BToATM, Thickness, OutMTD);
+						bOverlap = AInstanced.LowLevelOverlapGeom(B, BToATM, Thickness, OutMTD);
 					}
 					else
 					{
@@ -136,6 +142,13 @@ namespace Chaos
 					}
 				}
 			}
+
+			if (OutMTD && bOverlap)
+			{
+				OutMTD->Normal = ATM.TransformVectorNoScale(OutMTD->Normal);
+				OutMTD->Position = ATM.TransformPosition(OutMTD->Position);
+			}
+			return bOverlap;
 		}
 
 		return false;
diff --git a/Engine/Source/Runtime/PhysicsCore/Public/SQVisitor.h b/Engine/Source/Runtime/PhysicsCore/Public/SQVisitor.h
--- a/Engine/Source/Runtime/PhysicsCore/Public/SQVisitor.h
+++ b/Engine/Source/Runtime/PhysicsCore/Public/SQVisitor.h
@@ -237,7 +237,7 @@ private:
 
 				bool bHit = false;
 
-				FVec3 WorldPosition, WorldNormal;
+				FVec3 WorldPosition{ 0.0f }, WorldNormal{0.0f};
 				Chaos::FReal Distance = 0;	//not needed but fixes compiler warning for overlap
 				int32 FaceIdx = INDEX_NONE;	//not needed but fixes compiler warning for overlap
 				FVec3 FaceNormal = FVec3::ZeroVector;
@@ -269,7 +269,9 @@ private:
 						bHit = OverlapQuery(*Geom, ActorTM, *QueryGeom, StartTM, /*Thickness=*/0, &MTDInfo);
 						if (bHit)
 						{
-							WorldNormal = MTDInfo.Normal * MTDInfo.Penetration;
+							WorldNormal = MTDInfo.Normal;
+							WorldPosition = MTDInfo.Position;
+							Distance = -MTDInfo.Penetration;
 						}
 					}
 					else
