Iris
* ProtocolManager logs now compiled by default in test builds.

#rnx

diff --git a/Engine/Source/Runtime/Experimental/Iris/Core/Private/Iris/IrisConfigInternal.h b/Engine/Source/Runtime/Experimental/Iris/Core/Private/Iris/IrisConfigInternal.h
--- a/Engine/Source/Runtime/Experimental/Iris/Core/Private/Iris/IrisConfigInternal.h
+++ b/Engine/Source/Runtime/Experimental/Iris/Core/Private/Iris/IrisConfigInternal.h
@@ -22,14 +22,6 @@
 #endif
 #endif
 
-#ifndef UE_NET_ENABLE_PROTOCOLMANAGER_LOG
-#if (UE_BUILD_SHIPPING || UE_BUILD_TEST)
-#	define UE_NET_ENABLE_PROTOCOLMANAGER_LOG 0
-#else
-#	define UE_NET_ENABLE_PROTOCOLMANAGER_LOG 1
-#endif
-#endif
-
 // Misc config
 #ifndef UE_NET_VALIDATE_REPLICATION_RECORD
 #if (UE_BUILD_SHIPPING || UE_BUILD_TEST)
diff --git a/Engine/Source/Runtime/Experimental/Iris/Core/Private/Iris/ReplicationSystem/ReplicationProtocolManager.cpp b/Engine/Source/Runtime/Experimental/Iris/Core/Private/Iris/ReplicationSystem/ReplicationProtocolManager.cpp
--- a/Engine/Source/Runtime/Experimental/Iris/Core/Private/Iris/ReplicationSystem/ReplicationProtocolManager.cpp
+++ b/Engine/Source/Runtime/Experimental/Iris/Core/Private/Iris/ReplicationSystem/ReplicationProtocolManager.cpp
@@ -14,6 +14,11 @@
 #include "Containers/StringFwd.h"
 #include "Logging/LogScopedVerbosityOverride.h"
 
+#ifndef UE_NET_ENABLE_PROTOCOLMANAGER_LOG
+// Default is to enable protocol logs in non-shipping
+#	define UE_NET_ENABLE_PROTOCOLMANAGER_LOG !(UE_BUILD_SHIPPING)
+#endif
+
 #if UE_NET_ENABLE_PROTOCOLMANAGER_LOG
 #	define UE_LOG_PROTOCOLMANAGER(Log, Format, ...)  UE_LOG(LogIris, Log, Format, ##__VA_ARGS__)
 #else
