Packed Chaos::FShapeInstanceProxy to fit into 192b MB3 bin

diff --git a/Engine/Source/Runtime/Experimental/Chaos/Private/Chaos/ShapeInstance.cpp b/Engine/Source/Runtime/Experimental/Chaos/Private/Chaos/ShapeInstance.cpp
--- a/Engine/Source/Runtime/Experimental/Chaos/Private/Chaos/ShapeInstance.cpp
+++ b/Engine/Source/Runtime/Experimental/Chaos/Private/Chaos/ShapeInstance.cpp
@@ -6,6 +6,8 @@
 
 namespace Chaos
 {
+	static_assert(sizeof(FShapeInstanceProxy) <= 192, "FShapeInstanceProxy was optimized to fit into 192b bin of MB3 to prevent excess memory waste");
+
 	// Create or reuse the shapes in the shapes array and populate them with the Geometry.
 	// If we have a Union it will be unpacked into the ShapesArray.
 	// On the Physics Thread we set bAllowCachedLeafInfo which caches the shapes world space state to optimize collision detection,
diff --git a/Engine/Source/Runtime/Experimental/Chaos/Public/Chaos/Collision/CollisionFilter.h b/Engine/Source/Runtime/Experimental/Chaos/Public/Chaos/Collision/CollisionFilter.h
--- a/Engine/Source/Runtime/Experimental/Chaos/Public/Chaos/Collision/CollisionFilter.h
+++ b/Engine/Source/Runtime/Experimental/Chaos/Public/Chaos/Collision/CollisionFilter.h
@@ -127,11 +127,11 @@ namespace Chaos
 		//
 		if (Shape0)
 		{
-			if (Implicit0Type == ImplicitObjectType::TriangleMesh && Shape0->GetCollisionTraceType() != Chaos_CTF_UseComplexAsSimple)
+			if (Implicit0Type == ImplicitObjectType::TriangleMesh && Shape0->GetCollisionTraceType() != EChaosCollisionTraceFlag::Chaos_CTF_UseComplexAsSimple)
 			{
 				return false;
 			}
-			else if (Shape0->GetCollisionTraceType() == Chaos_CTF_UseComplexAsSimple && Implicit0Type != ImplicitObjectType::TriangleMesh)
+			else if (Shape0->GetCollisionTraceType() == EChaosCollisionTraceFlag::Chaos_CTF_UseComplexAsSimple && Implicit0Type != ImplicitObjectType::TriangleMesh)
 			{
 				return false;
 			}
@@ -143,11 +143,11 @@ namespace Chaos
 
 		if (Shape1)
 		{
-			if (Implicit1Type == ImplicitObjectType::TriangleMesh && Shape1->GetCollisionTraceType() != Chaos_CTF_UseComplexAsSimple)
+			if (Implicit1Type == ImplicitObjectType::TriangleMesh && Shape1->GetCollisionTraceType() != EChaosCollisionTraceFlag::Chaos_CTF_UseComplexAsSimple)
 			{
 				return false;
 			}
-			else if (Shape1->GetCollisionTraceType() == Chaos_CTF_UseComplexAsSimple && Implicit1Type != ImplicitObjectType::TriangleMesh)
+			else if (Shape1->GetCollisionTraceType() == EChaosCollisionTraceFlag::Chaos_CTF_UseComplexAsSimple && Implicit1Type != ImplicitObjectType::TriangleMesh)
 			{
 				return false;
 			}
diff --git a/Engine/Source/Runtime/Experimental/Chaos/Public/Chaos/Particles.h b/Engine/Source/Runtime/Experimental/Chaos/Public/Chaos/Particles.h
--- a/Engine/Source/Runtime/Experimental/Chaos/Public/Chaos/Particles.h
+++ b/Engine/Source/Runtime/Experimental/Chaos/Public/Chaos/Particles.h
@@ -220,7 +220,7 @@ namespace Chaos
 		Count
 	};
 
-	enum EChaosCollisionTraceFlag
+	enum class EChaosCollisionTraceFlag : int8
 	{
 		/** Use project physics settings (DefaultShapeComplexity) */
 		Chaos_CTF_UseDefault,
diff --git a/Engine/Source/Runtime/Experimental/Chaos/Public/Chaos/ShapeInstance.h b/Engine/Source/Runtime/Experimental/Chaos/Public/Chaos/ShapeInstance.h
--- a/Engine/Source/Runtime/Experimental/Chaos/Public/Chaos/ShapeInstance.h
+++ b/Engine/Source/Runtime/Experimental/Chaos/Public/Chaos/ShapeInstance.h
@@ -35,7 +35,7 @@ namespace Chaos
 	class FPerShapeData
 	{
 	protected:
-		enum class EPerShapeDataType : uint8
+		enum class EPerShapeDataType : uint32
 		{
 			Proxy,
 			Sim,
@@ -201,9 +201,10 @@ namespace Chaos
 
 		virtual void SerializeMaterials(FChaosArchive& Ar) = 0;
 
-		EPerShapeDataType Type;
-		uint8 bIsSingleMaterial : 1;	// For use by FShapeInstance (here because the space is available for free)
-		int32 ShapeIdx;
+		EPerShapeDataType Type : 2;
+		uint32 bIsSingleMaterial : 1;	// For use by FShapeInstance (here because the space is available for free)
+		uint32 ShapeIdx : 29;
+		FShapeDirtyFlags DirtyFlags;	// For use by FShapeInstanceProxy as there's 4 bytes of padding here
 		FImplicitObjectPtr Geometry;
 		TAABB<FReal, 3> WorldSpaceInflatedShapeBounds;
 	};
@@ -458,7 +459,6 @@ namespace Chaos
 		explicit FShapeInstanceProxy(int32 InShapeIdx)
 			: FPerShapeData(EPerShapeDataType::Proxy, InShapeIdx)
 			, Proxy(nullptr)
-			, DirtyFlags()
 			, CollisionData()
 			, Materials()
 		{
@@ -468,7 +468,6 @@ namespace Chaos
 		FShapeInstanceProxy(int32 InShapeIdx, TSerializablePtr<FImplicitObject> InGeometry)
 			: FPerShapeData(EPerShapeDataType::Proxy, InShapeIdx)
 			, Proxy(nullptr)
-			, DirtyFlags()
 			, CollisionData()
 			, Materials()
 		{
@@ -478,7 +477,6 @@ namespace Chaos
 		FShapeInstanceProxy(int32 InShapeIdx, const FImplicitObjectPtr& InGeometry)
 			: FPerShapeData(EPerShapeDataType::Proxy, InShapeIdx, InGeometry)
 			, Proxy(nullptr)
-			, DirtyFlags()
 			, CollisionData()
 			, Materials()
 		{
@@ -488,7 +486,6 @@ namespace Chaos
 
 
 		IPhysicsProxyBase* Proxy;
-		FShapeDirtyFlags DirtyFlags;
 
 		TShapeProperty<FCollisionData, EShapeProperty::CollisionData> CollisionData;
 		TShapeProperty<FMaterialData, EShapeProperty::Materials> Materials;
