Make dataflow proximity node optionally use convex hulls that are already on a geometry collection, rather than always re-computing the convex hulls from the geometry

diff --git a/Engine/Plugins/Experimental/GeometryCollectionPlugin/Source/GeometryCollectionNodes/Private/Dataflow/GeometryCollectionNodes.cpp b/Engine/Plugins/Experimental/GeometryCollectionPlugin/Source/GeometryCollectionNodes/Private/Dataflow/GeometryCollectionNodes.cpp
--- a/Engine/Plugins/Experimental/GeometryCollectionPlugin/Source/GeometryCollectionNodes/Private/Dataflow/GeometryCollectionNodes.cpp
+++ b/Engine/Plugins/Experimental/GeometryCollectionPlugin/Source/GeometryCollectionNodes/Private/Dataflow/GeometryCollectionNodes.cpp
@@ -695,10 +695,17 @@ void FProximityDataflowNode::Evaluate(Dataflow::FContext& Context, const FDatafl
 
 			GeomCollection->SetProximityProperties(Properties);
 
+			UE::GeometryCollectionConvexUtility::FConvexHulls TransformedExistingHulls;
+			bool bUseExistingHulls = false;
+			if (!bRecomputeConvexHulls)
+			{
+				bUseExistingHulls = UE::GeometryCollectionConvexUtility::GetExistingConvexHullsInSharedSpace(GeomCollection.Get(), TransformedExistingHulls, true);
+			}
+
 			// Invalidate proximity
 			FGeometryCollectionProximityUtility ProximityUtility(GeomCollection.Get());
 			ProximityUtility.InvalidateProximity();
-			ProximityUtility.UpdateProximity();
+			ProximityUtility.UpdateProximity(bUseExistingHulls ? &TransformedExistingHulls : nullptr);
 
 			SetValue<const FManagedArrayCollection&>(Context, *GeomCollection, &Collection);
 		}
diff --git a/Engine/Plugins/Experimental/GeometryCollectionPlugin/Source/GeometryCollectionNodes/Public/Dataflow/GeometryCollectionNodes.h b/Engine/Plugins/Experimental/GeometryCollectionPlugin/Source/GeometryCollectionNodes/Public/Dataflow/GeometryCollectionNodes.h
--- a/Engine/Plugins/Experimental/GeometryCollectionPlugin/Source/GeometryCollectionNodes/Public/Dataflow/GeometryCollectionNodes.h
+++ b/Engine/Plugins/Experimental/GeometryCollectionPlugin/Source/GeometryCollectionNodes/Public/Dataflow/GeometryCollectionNodes.h
@@ -1109,6 +1109,11 @@ public:
 	UPROPERTY(EditAnywhere, Category = "Proximity")
 	EConnectionContactAreaMethodEnum ContactAreaMethod = EConnectionContactAreaMethodEnum::Dataflow_ConnectionContactAreaMethod_None;
 
+	/** Whether to compute new convex hulls for proximity, or use the pre-existing hulls on the Collection, when using convex hulls to determine proximity */
+	UPROPERTY(EditAnywhere, Category = "Proximity", meta = (
+		EditCondition = "ProximityMethod == EProximityMethodEnum::Dataflow_ProximityMethod_ConvexHull || FilterContactMethod == EProximityContactFilteringMethodEnum::Dataflow_ProximityContactFilteringMethod_ConvexHullSharp || FilterContactMethod == EProximityContactFilteringMethodEnum::Dataflow_ProximityContactFilteringMethod_ConvexHullArea || ContactAreaMethod = EConnectionContactAreaMethodEnum::Dataflow_ProximityContactFilteringMethod_ConvexHullArea"))
+	bool bRecomputeConvexHulls = true;
+
 	/** GeometryCollection to update the proximity graph on */
 	UPROPERTY(meta = (DataflowInput, DataflowOutput, DataflowPassthrough = "Collection", DataflowIntrinsic))
 	FManagedArrayCollection Collection;
diff --git a/Engine/Source/Runtime/Experimental/Chaos/Private/GeometryCollection/GeometryCollectionConvexUtility.cpp b/Engine/Source/Runtime/Experimental/Chaos/Private/GeometryCollection/GeometryCollectionConvexUtility.cpp
--- a/Engine/Source/Runtime/Experimental/Chaos/Private/GeometryCollection/GeometryCollectionConvexUtility.cpp
+++ b/Engine/Source/Runtime/Experimental/Chaos/Private/GeometryCollection/GeometryCollectionConvexUtility.cpp
@@ -3443,4 +3443,70 @@ namespace UE::GeometryCollectionConvexUtility
 		FilterHullPoints(HullPolygons.Vertices, SimplificationDistanceThreshold);
 		*ResultHull = Chaos::FConvex(HullPolygons.Vertices, UpdateHull->GetMargin());
 	}
+
+	bool CHAOS_API GetExistingConvexHullsInSharedSpace(const FManagedArrayCollection* Collection, FConvexHulls& OutConvexHulls, bool bLeafOnly)
+	{
+		if (!Collection->HasAttribute("TransformToConvexIndices", FTransformCollection::TransformGroup) ||
+			!Collection->HasAttribute(FGeometryCollection::ConvexHullAttribute, FGeometryCollection::ConvexGroup))
+		{
+			return false;
+		}
+		
+		const TManagedArray<TSet<int32>>& OrigTransformToConvexIndices = Collection->GetAttribute<TSet<int32>>("TransformToConvexIndices", FTransformCollection::TransformGroup);
+		const TManagedArray<Chaos::FConvexPtr>& OrigConvexHulls = Collection->GetAttribute<Chaos::FConvexPtr>(FGeometryCollection::ConvexHullAttribute, FGeometryCollection::ConvexGroup);
+		const TManagedArray<int32>* SimulationType = Collection->FindAttribute<int32>("SimulationType", FTransformCollection::TransformGroup);
+		if (bLeafOnly && !SimulationType)
+		{
+			return false;
+		}
+
+		const int32 NumTransform = OrigTransformToConvexIndices.Num();
+		OutConvexHulls.TransformToHullsIndices.SetNum(NumTransform);
+
+		GeometryCollection::Facades::FCollectionTransformFacade TransformFacade(*Collection);
+		TArray<FTransform> GlobalTransformArray = TransformFacade.ComputeCollectionSpaceTransforms();
+		
+		TArray<int32> NewConvexToTransformIndices;
+		TArray<int32> NewConvexToOrigConvexIndices;
+		NewConvexToTransformIndices.Reserve(OrigConvexHulls.Num());
+		NewConvexToOrigConvexIndices.Reserve(OrigConvexHulls.Num());
+		int32 NumNewConvex = 0;
+		for (int32 TransformIdx = 0; TransformIdx < OrigTransformToConvexIndices.Num(); ++TransformIdx)
+		{
+			if (bLeafOnly && (*SimulationType)[TransformIdx] != FGeometryCollection::ESimulationTypes::FST_Rigid)
+			{
+				continue;
+			}
+			for (int32 OrigConvexIdx : OrigTransformToConvexIndices[TransformIdx])
+			{
+				int32 NewConvexIdx = NumNewConvex++;
+				NewConvexToTransformIndices.Add(TransformIdx);
+				NewConvexToOrigConvexIndices.Add(OrigConvexIdx);
+				checkSlow(NewConvexToTransformIndices.Num() == NumNewConvex);
+				checkSlow(NewConvexToOrigConvexIndices.Num() == NumNewConvex);
+				OutConvexHulls.TransformToHullsIndices[TransformIdx].Add(NewConvexIdx);
+			}
+		}
+
+		OutConvexHulls.Hulls.SetNum(NumNewConvex);
+		OutConvexHulls.OverlapRemovalShrinkPercent = 0;
+		OutConvexHulls.Pivots.Reset(); // Note: No scaling is applied so pivots are not used
+
+		ParallelFor(NewConvexToOrigConvexIndices.Num(), [&](int32 NewConvexIdx)
+		{
+			int32 TransformIdx = NewConvexToTransformIndices[NewConvexIdx];
+			FTransform Transform = GlobalTransformArray[TransformIdx];
+			int32 OrigConvexIdx = NewConvexToOrigConvexIndices[NewConvexIdx];
+			TArray<Chaos::FConvex::FVec3Type> HullPts;
+			HullPts.Reserve(OrigConvexHulls[OrigConvexIdx]->GetVertices().Num());
+			for (const Chaos::FConvex::FVec3Type& P : OrigConvexHulls[OrigConvexIdx]->GetVertices())
+			{
+				FVector PVec(P);
+				HullPts.Add((Chaos::FConvex::FVec3Type)(Transform.TransformPosition(PVec)));
+			}
+			OutConvexHulls.Hulls[NewConvexIdx] = new Chaos::FConvex(HullPts, UE_KINDA_SMALL_NUMBER);
+		});
+
+		return true;
+	}
 }
diff --git a/Engine/Source/Runtime/Experimental/Chaos/Public/GeometryCollection/GeometryCollectionConvexUtility.h b/Engine/Source/Runtime/Experimental/Chaos/Public/GeometryCollection/GeometryCollectionConvexUtility.h
--- a/Engine/Source/Runtime/Experimental/Chaos/Public/GeometryCollection/GeometryCollectionConvexUtility.h
+++ b/Engine/Source/Runtime/Experimental/Chaos/Public/GeometryCollection/GeometryCollectionConvexUtility.h
@@ -75,6 +75,10 @@ namespace UE::GeometryCollectionConvexUtility
 	void CHAOS_API IntersectConvexHulls(::Chaos::FConvex* ResultHull, const ::Chaos::FConvex* ClipHull, float ClipHullOffset, const ::Chaos::FConvex* UpdateHull,
 		const FTransform* ClipHullTransform = nullptr, const FTransform* UpdateHullTransform = nullptr, const FTransform* UpdateToResultTransform = nullptr, double SimplificationDistanceThreshold = 0.0);
 
+	// Get the existing convex hulls from the collection, transformed to the global/shared space of the overall collection
+	// @param bLeafOnly		Only include the convex hulls of leaf (rigid) nodes
+	bool CHAOS_API GetExistingConvexHullsInSharedSpace(const FManagedArrayCollection* Collection, FConvexHulls& OutConvexHulls, bool bLeafOnly = false);
+
 }
 
 class FGeometryCollectionConvexUtility
