Fix rare hardlock that can be encountered when an IAS request fails.

#rnx

- When a request fails it can be put back in the pending request queue (HttpRequests) so that it is run again. The callback that deals with this is driven by the call to HttpClient->Tick(). If the failed request was the only one in flight, Tick would return 0 and we would exit the final while loop with NextChunkRequest as nullptr indicating that there was no more work to do and with no more work to do, the thread would go back to waiting on TickBackendEvent.
- As a quick fix we now check the HttpRequests queue one last time to see if there are any more requests we could do. We'd like to clean up this entire method so it's not worth doing much more than this at the moment.

diff --git a/Engine/Source/Runtime/Experimental/IoStoreOnDemand/Private/OnDemandIoDispatcherBackend.cpp b/Engine/Source/Runtime/Experimental/IoStoreOnDemand/Private/OnDemandIoDispatcherBackend.cpp
--- a/Engine/Source/Runtime/Experimental/IoStoreOnDemand/Private/OnDemandIoDispatcherBackend.cpp
+++ b/Engine/Source/Runtime/Experimental/IoStoreOnDemand/Private/OnDemandIoDispatcherBackend.cpp
@@ -1410,6 +1410,8 @@ uint32 FOnDemandIoBackend::Run()
 								{
 									Stats.OnHttpRetry();
 									Stats.OnHttpEnqueue();
+
+									// Note there is no need to trigger TickBackendEvent as this callback will occur within FOnDemandIoBackend::Run
 									return HttpRequests.Enqueue(ChunkRequest); 
 								}
 							}
@@ -1450,12 +1452,18 @@ uint32 FOnDemandIoBackend::Run()
 			}
 
 			{
-				// Keep processing pending connections until all work is complete or a new request if found
+				// Keep processing pending connections until all requests are completed or a new one is issued
 				TRACE_CPUPROFILER_EVENT_SCOPE(FOnDemandIoBackend::TickHttp);
 				while (HttpClient->Tick() && !NextChunkRequest)
 				{
 					NextChunkRequest = HttpRequests.Dequeue();
 				}
+
+				// Tick can cause a new request to be added to HttpRequests, so we should try one last time
+				if (NextChunkRequest == nullptr)
+				{
+					NextChunkRequest = HttpRequests.Dequeue();
+				}
 			}
 		}
 
