[NetPhysics] Handle breaking CUs with resimulation by blocking rewinds past unhandled CU changes.

diff --git a/Engine/Source/Runtime/Experimental/Chaos/Private/Chaos/ClusterUnionManager.cpp b/Engine/Source/Runtime/Experimental/Chaos/Private/Chaos/ClusterUnionManager.cpp
--- a/Engine/Source/Runtime/Experimental/Chaos/Private/Chaos/ClusterUnionManager.cpp
+++ b/Engine/Source/Runtime/Experimental/Chaos/Private/Chaos/ClusterUnionManager.cpp
@@ -10,6 +10,7 @@
 #include "Chaos/PBDRigidsEvolution.h"
 #include "Chaos/PBDRigidsEvolutionGBF.h"
 #include "ChaosStats.h"
+#include "RewindData.h"
 
 namespace Chaos
 {
@@ -336,6 +337,12 @@ namespace Chaos
 				}
 			}
 			PendingClusterIndexOperations.Empty();
+
+			if (FRewindData* RewindData = MEvolution.GetRewindData())
+			{
+				// Temp, since rewind doesn't handle altered cluster unions we tell it to never rewind past this frame for now, only relevant when resimulation is used.
+				RewindData->BlockResim();
+			}
 		}
 
 		HandleDeferredClusterUnionUpdateProperties();
diff --git a/Engine/Source/Runtime/Experimental/Chaos/Private/Chaos/Collision/ParticlePairMidPhase.cpp b/Engine/Source/Runtime/Experimental/Chaos/Private/Chaos/Collision/ParticlePairMidPhase.cpp
--- a/Engine/Source/Runtime/Experimental/Chaos/Private/Chaos/Collision/ParticlePairMidPhase.cpp
+++ b/Engine/Source/Runtime/Experimental/Chaos/Private/Chaos/Collision/ParticlePairMidPhase.cpp
@@ -1845,7 +1845,9 @@ namespace Chaos
 	void FGenericParticlePairMidPhase::InjectCollisionImpl(const FPBDCollisionConstraint& Constraint, const FCollisionContext& Context)
 	{
 		// @todo(chaos): support rewind/resim here
-		ensure(false);
+#if !UE_BUILD_SHIPPING
+		UE_LOG(LogChaos, Warning, TEXT("Unsupported, handle Rewind/Resim in FGenericParticlePairMidPhase::InjectCollisionImpl"));
+#endif
 	}
 }
 
diff --git a/Engine/Source/Runtime/Experimental/Chaos/Private/Chaos/PBDRigidsEvolution.cpp b/Engine/Source/Runtime/Experimental/Chaos/Private/Chaos/PBDRigidsEvolution.cpp
--- a/Engine/Source/Runtime/Experimental/Chaos/Private/Chaos/PBDRigidsEvolution.cpp
+++ b/Engine/Source/Runtime/Experimental/Chaos/Private/Chaos/PBDRigidsEvolution.cpp
@@ -1081,7 +1081,16 @@ namespace Chaos
 		const EObjectStateType InitialState = Particle->ObjectState();
 		Particle->SetSleepType(InSleepType);
 		const EObjectStateType ObjectState = Particle->ObjectState();
-		Particles.SetDynamicParticleSOA(Particle);
+		
+		if (FPBDRigidClusteredParticleHandle* ClusteredParticle = Particle->CastToClustered())
+		{
+			Particles.SetClusteredParticleSOA(ClusteredParticle);
+		}
+		else
+		{
+			Particles.SetDynamicParticleSOA(Particle);
+		}
+
 		if (InitialState != ObjectState)
 		{
 			if(ObjectState != EObjectStateType::Dynamic)
diff --git a/Engine/Source/Runtime/Experimental/Chaos/Private/RewindData.cpp b/Engine/Source/Runtime/Experimental/Chaos/Private/RewindData.cpp
--- a/Engine/Source/Runtime/Experimental/Chaos/Private/RewindData.cpp
+++ b/Engine/Source/Runtime/Experimental/Chaos/Private/RewindData.cpp
@@ -465,7 +465,6 @@ void FRewindData::FinishFrame()
 		FinishHelper(DirtyJoints);
 	}
 	
-
 	++CurFrame;
 	LatestFrame = FMath::Max(LatestFrame, CurFrame);
 }
@@ -836,25 +835,23 @@ void FRewindData::ExtendHistoryWithFrame(const int32 Frame)
 int32 FRewindData::FindValidResimFrame(const int32 RequestedFrame)
 {
 	int32 ValidFrame = INDEX_NONE;
-	if(RequestedFrame > 0)
+	if (RequestedFrame > BlockResimFrame)
 	{
 		EnsureIsInPhysicsThreadContext();
 
 		// First frame of the history datas
-		const int32 EarliestFrame = FMath::Max(GetEarliestFrame_Internal(), 0);
-	
-		for(ValidFrame = RequestedFrame; ValidFrame >= EarliestFrame; --ValidFrame)
+		const int32 EarliestFrame = FMath::Max(GetEarliestFrame_Internal(), BlockResimFrame);
+
+		for (ValidFrame = RequestedFrame; ValidFrame >= EarliestFrame; --ValidFrame)
 		{
 			bool bHasTargetHistory = true;
 			for (FDirtyParticleInfo& DirtyParticleInfo : DirtyParticles)
 			{
-				FGeometryParticleHandle* PTParticle = DirtyParticleInfo.GetObjectPtr();
 				FGeometryParticleStateBase& History = DirtyParticleInfo.GetHistory();
-
 				const bool bResimAsFollower = DirtyParticleInfo.bResimAsFollower;
 
 				const FFrameAndPhase FrameAndPhase{ ValidFrame, FFrameAndPhase::PostPushData };
-				if(const FParticleDynamicMisc* DynamicMisc = History.DynamicsMisc.Read(FrameAndPhase, PropertiesPool))
+				if (const FParticleDynamicMisc* DynamicMisc = History.DynamicsMisc.Read(FrameAndPhase, PropertiesPool))
 				{
 					if (!DynamicMisc->Disabled() && (DynamicMisc->ObjectState() == EObjectStateType::Dynamic) && !History.TargetPositions.IsEmpty() && !History.TargetVelocities.IsEmpty() && !History.TargetStates.IsEmpty())
 					{
@@ -866,18 +863,18 @@ int32 FRewindData::FindValidResimFrame(const int32 RequestedFrame)
 					}
 				}
 			}
-			if(bHasTargetHistory)
+			if (bHasTargetHistory)
 			{
-				for(auto& InputsHistory : InputsHistories)
+				for (auto& InputsHistory : InputsHistories)
 				{
-					if(!InputsHistory.Pin().Get()->HasValidDatas(ValidFrame))
+					if (!InputsHistory.Pin().Get()->HasValidDatas(ValidFrame))
 					{
 						bHasTargetHistory = false;
 						break;
 					}
 				}
 			}
-			if(bHasTargetHistory)
+			if (bHasTargetHistory)
 			{
 				for (auto& StatesHistory : StatesHistories)
 				{
@@ -889,7 +886,7 @@ int32 FRewindData::FindValidResimFrame(const int32 RequestedFrame)
 				}
 			}
 
-			if(bHasTargetHistory)
+			if (bHasTargetHistory)
 			{
 				break;
 			}
@@ -983,4 +980,12 @@ void FRewindData::SetTargetStateAtFrame(FGeometryParticleHandle& Handle, const i
 		Position, Quaternion, LinVelocity, AngVelocity, bShouldSleep);
 }
 
+void FRewindData::BlockResim()
+{
+	if (LatestFrame > BlockResimFrame)
+	{
+		BlockResimFrame = LatestFrame;
+	}
 }
+
+}
\ No newline at end of file
diff --git a/Engine/Source/Runtime/Experimental/Chaos/Public/Chaos/PBDRigidsEvolution.h b/Engine/Source/Runtime/Experimental/Chaos/Public/Chaos/PBDRigidsEvolution.h
--- a/Engine/Source/Runtime/Experimental/Chaos/Public/Chaos/PBDRigidsEvolution.h
+++ b/Engine/Source/Runtime/Experimental/Chaos/Public/Chaos/PBDRigidsEvolution.h
@@ -935,6 +935,11 @@ public:
 		MRewindData = RewindData;
 	}
 
+	FRewindData* GetRewindData()
+	{
+		return MRewindData;
+	}
+
 	CHAOS_API void DisableParticleWithRemovalEvent(FGeometryParticleHandle* Particle);
 	const TArray<FRemovalData>& GetAllRemovals() { return MAllRemovals; }
 	void ResetAllRemovals() { MAllRemovals.Reset(); }
diff --git a/Engine/Source/Runtime/Experimental/Chaos/Public/RewindData.h b/Engine/Source/Runtime/Experimental/Chaos/Public/RewindData.h
--- a/Engine/Source/Runtime/Experimental/Chaos/Public/RewindData.h
+++ b/Engine/Source/Runtime/Experimental/Chaos/Public/RewindData.h
@@ -1268,7 +1268,7 @@ public:
 	: Managers(NumFrames+1)	//give 1 extra for saving at head
 	, Solver(InSolver)
 	, CurFrame(InCurrentFrame)
-	, LatestFrame(-1)
+	, LatestFrame(InCurrentFrame)
 	, FramesSaved(0)
 	, DataIdxOffset(0)
 	, bNeedsSave(false)
@@ -1280,6 +1280,7 @@ public:
 	{
 		Solver = InSolver;
 		CurFrame = InCurrentFrame;
+		LatestFrame = InCurrentFrame;
 		bResimOptimization = InResimOptimization;
 		Managers = TCircularBuffer<FFrameManagerInfo>(NumFrames + 1);
 	}
@@ -1450,8 +1451,10 @@ public:
 	int32 GetResimFrame() { return ResimFrame; }
 	void SetResimFrame(int32 Frame) { ResimFrame = Frame; }
 
-private:
+	/** This blocks any future resimulation to rewind back past the frame this is called on */
+	void BlockResim();
 
+private:
 	friend class FPBDRigidsSolver;
 
 	void CHAOS_API AdvanceFrameImp(IResimCacheBase* ResimCache);
@@ -1628,6 +1631,9 @@ private:
 	bool bResimOptimization;
 	int32 ResimFrame = INDEX_NONE;
 
+	// Used to block rewinding past a physics change we currently don't handle
+	int32 BlockResimFrame = INDEX_NONE;
+
 	template <typename TObj>
 	bool IsResimAndInSync(const TObj& Handle) const { return IsResim() && Handle.SyncState() == ESyncState::InSync; }
 
