I/O - set error message from HTTP request

#rnx

diff --git a/Engine/Source/Runtime/Experimental/IoStoreOnDemand/Private/OnDemandIoDispatcherBackend.cpp b/Engine/Source/Runtime/Experimental/IoStoreOnDemand/Private/OnDemandIoDispatcherBackend.cpp
--- a/Engine/Source/Runtime/Experimental/IoStoreOnDemand/Private/OnDemandIoDispatcherBackend.cpp
+++ b/Engine/Source/Runtime/Experimental/IoStoreOnDemand/Private/OnDemandIoDispatcherBackend.cpp
@@ -392,7 +392,7 @@ void FHttpClient::Issue(FAnsiStringView Url, FIoReadCallback&& Callback, FIoOffs
 			{
 				const uint64 Duration = (uint64)FPlatformTime::ToMilliseconds64(FPlatformTime::Cycles64() - StartTime);
 				LogHttpResult(*Url, StatusCode, Duration, 0, Offset, Status.GetErrorReason());
-				Callback(FIoStatus(EIoErrorCode::ReadError, TEXTVIEW("HTTP Error")));
+				Callback(FIoStatus(EIoErrorCode::ReadError, FString(Status.GetErrorReason())));
 			}
 		};
 
