I/O - change block hash from 20B to 4B to reduce TOC size

#rnx

diff --git a/Engine/Source/Runtime/Experimental/IoStoreOnDemand/Private/IoStoreOnDemand.cpp b/Engine/Source/Runtime/Experimental/IoStoreOnDemand/Private/IoStoreOnDemand.cpp
--- a/Engine/Source/Runtime/Experimental/IoStoreOnDemand/Private/IoStoreOnDemand.cpp
+++ b/Engine/Source/Runtime/Experimental/IoStoreOnDemand/Private/IoStoreOnDemand.cpp
@@ -236,7 +236,7 @@ FCbWriter& operator<<(FCbWriter& Writer, const FOnDemandTocContainerEntry& Conta
 	Writer.EndArray();
 
 	Writer.BeginArray(UTF8TEXTVIEW("BlockHashes"));
-	for (const FIoHash& BlockHash : ContainerEntry.BlockHashes)
+	for (uint32 BlockHash : ContainerEntry.BlockHashes)
 	{
 		Writer << BlockHash;
 	}
@@ -277,7 +277,15 @@ bool LoadFromCompactBinary(FCbFieldView Field, FOnDemandTocContainerEntry& OutCo
 		OutContainer.BlockHashes.Reserve(int32(BlockHashes.Num()));
 		for (FCbFieldView ArrayField : BlockHashes)
 		{
-			OutContainer.BlockHashes.Add(ArrayField.AsHash());
+			if (ArrayField.IsHash())
+			{
+				const FIoHash BlockHash = ArrayField.AsHash();
+				OutContainer.BlockHashes.Add(*reinterpret_cast<const uint32*>(&BlockHash));
+			}
+			else
+			{
+				OutContainer.BlockHashes.Add(ArrayField.AsUInt32());
+			}
 		}
 
 		OutContainer.UTocHash = Obj["UTocHash"].AsHash();
@@ -726,7 +734,8 @@ TIoStatusOr<FIoStoreUploadResult> UploadContainerFiles(
 
 				FMemoryView EncodedBlock = EncodedBlocks.Left(EncodedBlockSize);
 				EncodedBlocks += EncodedBlock.GetSize();
-				ContainerEntry.BlockHashes.Add(FIoHash::HashBuffer(EncodedBlock));
+				const FIoHash BlockHash = FIoHash::HashBuffer(EncodedBlock);
+				ContainerEntry.BlockHashes.Add(*reinterpret_cast<const uint32*>(&BlockHash));
 
 				EncodedChunkSize += EncodedBlockSize;
 				RawChunkSize += BlockInfo.UncompressedSize;
@@ -851,6 +860,7 @@ TIoStatusOr<FIoStoreUploadResult> UploadContainerFiles(
 		Key << LexToString(UploadResult.TocHash) << TEXT(".iochunktoc");
 
 		UploadResult.TocPath = Key.ToString();
+		UploadResult.TocSize = Ar.TotalSize();
 
 		const FS3PutObjectResponse Response = Client.TryPutObject(FS3PutObjectRequest{UploadParams.Bucket, Key.ToString(), Ar.GetView()});
 		if (Response.IsOk())
@@ -877,21 +887,22 @@ TIoStatusOr<FIoStoreUploadResult> UploadContainerFiles(
 		const double Duration = FPlatformTime::Seconds() - StartTime;
 		
 		UE_LOG(LogIas, Display, TEXT(""));
-		UE_LOG(LogIas, Display, TEXT("---------------------------------------- Upload Summary --------------------------------------"));
-		UE_LOG(LogIas, Display, TEXT("%-40s: %s"), TEXT("Service URL"), *UploadParams.ServiceUrl);
-		UE_LOG(LogIas, Display, TEXT("%-40s: %s"), TEXT("Bucket"), *UploadParams.Bucket);
-		UE_LOG(LogIas, Display, TEXT("%-40s: %s"), TEXT("TOC"), *UploadResult.TocPath);
-		UE_LOG(LogIas, Display, TEXT("%-40s: %.2lf second(s)"), TEXT("Duration"), Duration);
+		UE_LOG(LogIas, Display, TEXT("--------------------------------------------- Upload Summary -------------------------------------------"));
+		UE_LOG(LogIas, Display, TEXT("%-20s: %s"), TEXT("Service URL"), *UploadParams.ServiceUrl);
+		UE_LOG(LogIas, Display, TEXT("%-20s: %s"), TEXT("Bucket"), *UploadParams.Bucket);
+		UE_LOG(LogIas, Display, TEXT("%-20s: %s"), TEXT("TOC path"), *UploadResult.TocPath);
+		UE_LOG(LogIas, Display, TEXT("%-20s: %.2lf KiB"), TEXT("TOC size"), double(UploadResult.TocSize) / 1024.0);
+		UE_LOG(LogIas, Display, TEXT("%-20s: %.2lf second(s)"), TEXT("Duration"), Duration);
 		UE_LOG(LogIas, Display, TEXT(""));
 		
-		UE_LOG(LogIas, Display, TEXT("%-25s %15s %15s %15s %15s"),
+		UE_LOG(LogIas, Display, TEXT("%-40s %15s %15s %15s %15s"),
 			TEXT("Container"), TEXT("Chunk(s)"), TEXT("Size (MiB)"), TEXT("Uploaded"), TEXT("Uploaded (MiB)"));
-		UE_LOG(LogIas, Display, TEXT("----------------------------------------------------------------------------------------------"));
+		UE_LOG(LogIas, Display, TEXT("--------------------------------------------------------------------------------------------------------"));
 
 		FContainerStats TotalStats;
 		for (const TTuple<FString, FContainerStats>& Kv : ContainerSummary)
 		{
-			UE_LOG(LogIas, Display, TEXT("%-25s %15llu %15.2lf %15llu %15.2lf"),
+			UE_LOG(LogIas, Display, TEXT("%-40s %15llu %15.2lf %15llu %15.2lf"),
 				*Kv.Key, Kv.Value.ChunkCount, double(Kv.Value.TotalBytes) / 1024.0 / 1024.0, Kv.Value.UploadedChunkCount, double(Kv.Value.UploadedBytes) / 1024.0 / 1024.0);
 			
 			TotalStats.ChunkCount += Kv.Value.ChunkCount;
@@ -899,17 +910,17 @@ TIoStatusOr<FIoStoreUploadResult> UploadContainerFiles(
 			TotalStats.UploadedChunkCount += Kv.Value.UploadedChunkCount;
 			TotalStats.UploadedBytes += Kv.Value.UploadedBytes;
 		}
-		UE_LOG(LogIas, Display, TEXT("----------------------------------------------------------------------------------------------"));
-		UE_LOG(LogIas, Display, TEXT("%-25s %15llu %15.2lf %15llu %15.2lf"),
+		UE_LOG(LogIas, Display, TEXT("--------------------------------------------------------------------------------------------------------"));
+		UE_LOG(LogIas, Display, TEXT("%-40s %15llu %15.2lf %15llu %15.2lf"),
 			TEXT("Total"),TotalStats.ChunkCount, double(TotalStats.TotalBytes) / 1024.0 / 1024.0, TotalStats.UploadedChunkCount, double(TotalStats.UploadedBytes) / 1024.0 / 1024.0);
 		UE_LOG(LogIas, Display, TEXT(""));
 		
-		UE_LOG(LogIas, Display, TEXT("%-25s %15s %15s %15s"), TEXT("Bucket"), TEXT("TOC(s)"), TEXT("Chunk(s)"), TEXT("MiB"));
-		UE_LOG(LogIas, Display, TEXT("----------------------------------------------------------------------------------------------"));
-		UE_LOG(LogIas, Display, TEXT("%-25s %15llu %15d %15.2lf"), TEXT("Existing"), TotalExistingTocs, ExistingChunks.Num(), double(TotalExistingBytes) / 1024.0 / 1024.0);
-		UE_LOG(LogIas, Display, TEXT("%-25s %15llu %15d %15.2lf"), TEXT("Uploaded"), 1, TotalStats.UploadedChunkCount, double(TotalStats.UploadedBytes) / 1024.0 / 1024.0);
-		UE_LOG(LogIas, Display, TEXT("----------------------------------------------------------------------------------------------"));
-		UE_LOG(LogIas, Display, TEXT("%-25s %15llu %15d %15.2lf"), TEXT("Total"),
+		UE_LOG(LogIas, Display, TEXT("%-40s %15s %15s %15s"), TEXT("Bucket"), TEXT("TOC(s)"), TEXT("Chunk(s)"), TEXT("MiB"));
+		UE_LOG(LogIas, Display, TEXT("--------------------------------------------------------------------------------------------------------"));
+		UE_LOG(LogIas, Display, TEXT("%-40s %15llu %15d %15.2lf"), TEXT("Existing"), TotalExistingTocs, ExistingChunks.Num(), double(TotalExistingBytes) / 1024.0 / 1024.0);
+		UE_LOG(LogIas, Display, TEXT("%-40s %15llu %15d %15.2lf"), TEXT("Uploaded"), 1, TotalStats.UploadedChunkCount, double(TotalStats.UploadedBytes) / 1024.0 / 1024.0);
+		UE_LOG(LogIas, Display, TEXT("--------------------------------------------------------------------------------------------------------"));
+		UE_LOG(LogIas, Display, TEXT("%-40s %15llu %15d %15.2lf"), TEXT("Total"),
 			TotalExistingTocs + 1, ExistingChunks.Num() + TotalStats.UploadedChunkCount, double(TotalExistingBytes + TotalStats.UploadedBytes) / 1024.0 / 1024.0);
 		UE_LOG(LogIas, Display, TEXT(""));
 	}
diff --git a/Engine/Source/Runtime/Experimental/IoStoreOnDemand/Public/IO/IoStoreOnDemand.h b/Engine/Source/Runtime/Experimental/IoStoreOnDemand/Public/IO/IoStoreOnDemand.h
--- a/Engine/Source/Runtime/Experimental/IoStoreOnDemand/Public/IO/IoStoreOnDemand.h
+++ b/Engine/Source/Runtime/Experimental/IoStoreOnDemand/Public/IO/IoStoreOnDemand.h
@@ -35,6 +35,7 @@ enum class EOnDemandTocVersion : uint32
 	Invalid			= 0,
 	Initial			= 1,
 	UTocHash		= 2,
+	BlockHash32		= 3,
 
 	LatestPlusOne,
 	Latest			= (LatestPlusOne - 1)
@@ -86,7 +87,7 @@ struct FOnDemandTocContainerEntry
 	FString EncryptionKeyGuid;
 	TArray<FOnDemandTocEntry> Entries;
 	TArray<uint32> BlockSizes;
-	TArray<FIoHash> BlockHashes;
+	TArray<uint32> BlockHashes;
 	FIoHash UTocHash;
 
 	UE_API friend FCbWriter& operator<<(FCbWriter& Writer, const FOnDemandTocContainerEntry& ContainerEntry);
@@ -131,6 +132,7 @@ struct FIoStoreUploadResult
 {
 	FIoHash TocHash;
 	FString TocPath;
+	uint64 TocSize = 0;
 };
 
 UE_API TIoStatusOr<FIoStoreUploadResult> UploadContainerFiles(
@@ -181,4 +183,4 @@ public:
 private:
 };
 
-#undef UE_API
\ No newline at end of file
+#undef UE_API
