Change AppendGeometry()'s initial check() that ensures the geometry collection to be appended to an if(), and return INDEX_NONE if it fails.

diff --git a/Engine/Source/Runtime/Experimental/Chaos/Private/GeometryCollection/GeometryCollection.cpp b/Engine/Source/Runtime/Experimental/Chaos/Private/GeometryCollection/GeometryCollection.cpp
--- a/Engine/Source/Runtime/Experimental/Chaos/Private/GeometryCollection/GeometryCollection.cpp
+++ b/Engine/Source/Runtime/Experimental/Chaos/Private/GeometryCollection/GeometryCollection.cpp
@@ -179,7 +179,10 @@ void FGeometryCollection::SetDefaults(FName Group, uint32 StartSize, uint32 NumE
 int32 FGeometryCollection::AppendGeometry(const FGeometryCollection & Element, int32 MaterialIDOffset, bool ReindexAllMaterials, const FTransform& TransformRoot)
 {
 	// until we support a transform hierarchy this is just one.
-	check(Element.NumElements(FGeometryCollection::TransformGroup) > 0);
+	if (Element.NumElements(FGeometryCollection::TransformGroup) == 0)
+	{
+		return INDEX_NONE;
+	}
 
 	int NumTransforms = NumElements(FTransformCollection::TransformGroup);
 	int NumNewTransforms = Element.NumElements(FTransformCollection::TransformGroup);
