Plugged memory leak

#rnx

diff --git a/Engine/Source/Runtime/Experimental/IoStoreOnDemand/Private/FileIoCache.cpp b/Engine/Source/Runtime/Experimental/IoStoreOnDemand/Private/FileIoCache.cpp
--- a/Engine/Source/Runtime/Experimental/IoStoreOnDemand/Private/FileIoCache.cpp
+++ b/Engine/Source/Runtime/Experimental/IoStoreOnDemand/Private/FileIoCache.cpp
@@ -800,13 +800,13 @@ void FFileIoCache::FileWriterThreadInner()
 	}
 
 	// Copy data to be cached into a buffer
-	auto* Buffer = (uint8*)FMemory::Malloc(PendingSize);
-	auto* Cursor = (uint8*)Buffer;
+	auto Buffer = TUniquePtr<uint8[]>(new uint8[PendingSize]);
+	auto* Cursor = Buffer.Get();
 	for (FCacheEntry& Entry : Entries)
 	{
 		const FMemoryView& View = Entry.Data.GetView();
 
-		Entry.SerialOffset = WriteCursorPos + uint32(ptrdiff_t(Cursor - Buffer));
+		Entry.SerialOffset = WriteCursorPos + uint32(ptrdiff_t(Cursor - Buffer.Get()));
 		Entry.SerialSize = View.GetSize();
 		Entry.Hash = FIoHash::HashBuffer(View);
 		Entry.State = ECacheEntryState::Writing;
@@ -839,7 +839,7 @@ void FFileIoCache::FileWriterThreadInner()
 			return;
 		}
 		FileHandle->Seek(WriteCursorPos);
-		FileHandle->Write(Buffer, PendingSize);
+		FileHandle->Write(Buffer.Get(), PendingSize);
 		FileHandle->Flush();
 	}
 	WriteCursorPos += PendingSize;
