I/O - removed raw hash from on demand TOC and converted to binary instead of CbObject

#rnx

diff --git a/Engine/Source/Runtime/Experimental/IoStoreOnDemand/Private/IoStoreOnDemand.cpp b/Engine/Source/Runtime/Experimental/IoStoreOnDemand/Private/IoStoreOnDemand.cpp
--- a/Engine/Source/Runtime/Experimental/IoStoreOnDemand/Private/IoStoreOnDemand.cpp
+++ b/Engine/Source/Runtime/Experimental/IoStoreOnDemand/Private/IoStoreOnDemand.cpp
@@ -15,6 +15,7 @@
 #include "Misc/CoreDelegates.h"
 #include "Misc/Parse.h"
 #include "Modules/ModuleManager.h"
+#include "Serialization/Archive.h"
 #include "Serialization/CompactBinaryWriter.h"
 #include "Serialization/CompactBinarySerialization.h"
 #include "Serialization/LargeMemoryWriter.h"
@@ -143,6 +144,18 @@ namespace UE
 {
 
 ////////////////////////////////////////////////////////////////////////////////
+FArchive& operator<<(FArchive& Ar, FOnDemandTocHeader& Header)
+{
+	Ar << Header.Magic;
+	Ar << Header.Version;
+	Ar << Header.ChunkVersion;
+	Ar << Header.BlockSize;
+	Ar << Header.CompressionFormat;
+	Ar << Header.ChunksDirectory;
+
+	return Ar;
+}
+
 FCbWriter& operator<<(FCbWriter& Writer, const FOnDemandTocHeader& Header)
 {
 	Writer.BeginObject();
@@ -175,11 +188,22 @@ bool LoadFromCompactBinary(FCbFieldView Field, FOnDemandTocHeader& OutTocHeader)
 	return false;
 }
 
+FArchive& operator<<(FArchive& Ar, FOnDemandTocEntry& Entry)
+{
+	Ar << Entry.Hash;
+	Ar << Entry.ChunkId;
+	Ar << Entry.RawSize;
+	Ar << Entry.EncodedSize;
+	Ar << Entry.BlockOffset;
+	Ar << Entry.BlockCount;
+
+	return Ar;
+}
+
 FCbWriter& operator<<(FCbWriter& Writer, const FOnDemandTocEntry& Entry)
 {
 	Writer.BeginObject();
 	Writer.AddHash(UTF8TEXTVIEW("Hash"), Entry.Hash);
-	Writer.AddHash(UTF8TEXTVIEW("RawHash"), Entry.RawHash);
 	Writer << UTF8TEXTVIEW("ChunkId") << Entry.ChunkId;
 	Writer.AddInteger(UTF8TEXTVIEW("RawSize"), Entry.RawSize);
 	Writer.AddInteger(UTF8TEXTVIEW("EncodedSize"), Entry.EncodedSize);
@@ -200,7 +224,6 @@ bool LoadFromCompactBinary(FCbFieldView Field, FOnDemandTocEntry& OutTocEntry)
 		}
 
 		OutTocEntry.Hash = Obj["Hash"].AsHash();
-		OutTocEntry.RawHash = Obj["RawHash"].AsHash();
 		OutTocEntry.RawSize = Obj["RawSize"].AsUInt64(~uint64(0));
 		OutTocEntry.EncodedSize = Obj["EncodedSize"].AsUInt64(~uint64(0));
 		OutTocEntry.BlockOffset = Obj["BlockOffset"].AsUInt32(~uint32(0));
@@ -215,6 +238,18 @@ bool LoadFromCompactBinary(FCbFieldView Field, FOnDemandTocEntry& OutTocEntry)
 	return false;
 }
 
+FArchive& operator<<(FArchive& Ar, FOnDemandTocContainerEntry& ContainerEntry)
+{
+	Ar << ContainerEntry.ContainerName;
+	Ar << ContainerEntry.EncryptionKeyGuid;
+	Ar << ContainerEntry.Entries;
+	Ar << ContainerEntry.BlockSizes;
+	Ar << ContainerEntry.BlockHashes;
+	Ar << ContainerEntry.UTocHash;
+
+	return Ar;
+}
+
 FCbWriter& operator<<(FCbWriter& Writer, const FOnDemandTocContainerEntry& ContainerEntry)
 {
 	Writer.BeginObject();
@@ -296,13 +331,29 @@ bool LoadFromCompactBinary(FCbFieldView Field, FOnDemandTocContainerEntry& OutCo
 	return false;
 }
 
-FCbWriter& operator<<(FCbWriter& Writer, const FOnDemandToc& TocResource)
+FArchive& operator<<(FArchive& Ar, FOnDemandToc& Toc)
+{
+	Ar << Toc.Header;
+	
+	if (Toc.Header.Magic != FOnDemandTocHeader::ExpectedMagic || static_cast<EOnDemandTocVersion>(Toc.Header.Version) == EOnDemandTocVersion::Invalid)
+	{
+		Ar.SetError();
+		return Ar;
+	}
+
+	Ar.SetCustomVersion(Toc.VersionGuid, int32(Toc.Header.Version), TEXT("OnDemandToc"));
+	Ar << Toc.Containers;
+
+	return Ar;
+}
+
+FCbWriter& operator<<(FCbWriter& Writer, const FOnDemandToc& Toc)
 {
 	Writer.BeginObject();
-	Writer << UTF8TEXTVIEW("Header") << TocResource.Header;
+	Writer << UTF8TEXTVIEW("Header") << Toc.Header;
 
 	Writer.BeginArray(UTF8TEXTVIEW("Containers"));
-	for (const FOnDemandTocContainerEntry& Container : TocResource.Containers)
+	for (const FOnDemandTocContainerEntry& Container : Toc.Containers)
 	{
 		Writer << Container;
 	}
@@ -312,39 +363,7 @@ FCbWriter& operator<<(FCbWriter& Writer, const FOnDemandToc& TocResource)
 	return Writer;
 }
 
-TIoStatusOr<FString> FOnDemandToc::Save(const TCHAR* Directory, const FOnDemandToc& TocResource)
-{
-	if (TocResource.Header.Magic != FOnDemandTocHeader::ExpectedMagic)
-	{
-		return FIoStatus(EIoErrorCode::CorruptToc);
-	}
-
-	if (TocResource.Header.CompressionFormat.IsEmpty())
-	{
-		return FIoStatus(EIoErrorCode::CorruptToc);
-	}
-
-	FCbWriter Writer;
-	Writer << TocResource;
-
-	FLargeMemoryWriter Ar;
-	SaveCompactBinary(Ar, Writer.Save());
-
-	const FIoHash TocHash = FIoHash::HashBuffer(Ar.GetView());
-	const FString FilePath = FString(Directory) / LexToString(TocHash) + TEXT(".iochunktoc");
-
-	IFileManager& FileMgr = IFileManager::Get();
-	if (TUniquePtr<FArchive> FileAr(FileMgr.CreateFileWriter(*FilePath)); FileAr.IsValid())
-	{
-		FileAr->Serialize(Ar.GetData(), Ar.TotalSize());
-		FileAr->Flush();
-		FileAr->Close();
-
-		return FilePath;
-	}
-
-	return FIoStatus(EIoErrorCode::WriteError);
-}
+FGuid FOnDemandToc::VersionGuid = FGuid("C43DD98F353F499D9A0767F6EA0155EB");
 
 bool LoadFromCompactBinary(FCbFieldView Field, FOnDemandToc& OutToc)
 {
@@ -754,7 +773,6 @@ TIoStatusOr<FIoStoreUploadResult> UploadContainerFiles(
 			FOnDemandTocEntry& TocEntry = ContainerEntry.Entries.AddDefaulted_GetRef();
 			TocEntry.ChunkId = ChunkInfo.Id;
 			TocEntry.Hash = ChunkHash;
-			TocEntry.RawHash = ChunkInfo.Hash.ToIoHash();
 			TocEntry.RawSize = RawChunkSize;
 			TocEntry.EncodedSize = EncodedChunkSize;
 			TocEntry.BlockOffset = BlockOffset;
@@ -845,11 +863,8 @@ TIoStatusOr<FIoStoreUploadResult> UploadContainerFiles(
 
 	FIoStoreUploadResult UploadResult;
 	{
-		FCbWriter Writer;
-		Writer << OnDemandToc;
-
 		FLargeMemoryWriter Ar;
-		SaveCompactBinary(Ar, Writer.Save());
+		Ar << OnDemandToc;
 
 		UploadResult.TocHash = FIoHash::HashBuffer(Ar.GetView());
 		TStringBuilder<256> Key;
@@ -887,22 +902,22 @@ TIoStatusOr<FIoStoreUploadResult> UploadContainerFiles(
 		const double Duration = FPlatformTime::Seconds() - StartTime;
 		
 		UE_LOG(LogIas, Display, TEXT(""));
-		UE_LOG(LogIas, Display, TEXT("--------------------------------------------- Upload Summary -------------------------------------------"));
-		UE_LOG(LogIas, Display, TEXT("%-20s: %s"), TEXT("Service URL"), *UploadParams.ServiceUrl);
-		UE_LOG(LogIas, Display, TEXT("%-20s: %s"), TEXT("Bucket"), *UploadParams.Bucket);
-		UE_LOG(LogIas, Display, TEXT("%-20s: %s"), TEXT("TOC path"), *UploadResult.TocPath);
-		UE_LOG(LogIas, Display, TEXT("%-20s: %.2lf KiB"), TEXT("TOC size"), double(UploadResult.TocSize) / 1024.0);
-		UE_LOG(LogIas, Display, TEXT("%-20s: %.2lf second(s)"), TEXT("Duration"), Duration);
+		UE_LOG(LogIas, Display, TEXT("------------------------------------------------- Upload Summary -------------------------------------------------"));
+		UE_LOG(LogIas, Display, TEXT("%-15s: %s"), TEXT("Service URL"), *UploadParams.ServiceUrl);
+		UE_LOG(LogIas, Display, TEXT("%-15s: %s"), TEXT("Bucket"), *UploadParams.Bucket);
+		UE_LOG(LogIas, Display, TEXT("%-15s: %s"), TEXT("TOC path"), *UploadResult.TocPath);
+		UE_LOG(LogIas, Display, TEXT("%-15s: %.2lf KiB"), TEXT("TOC size"), double(UploadResult.TocSize) / 1024.0);
+		UE_LOG(LogIas, Display, TEXT("%-15s: %.2lf second(s)"), TEXT("Duration"), Duration);
 		UE_LOG(LogIas, Display, TEXT(""));
 		
-		UE_LOG(LogIas, Display, TEXT("%-40s %15s %15s %15s %15s"),
+		UE_LOG(LogIas, Display, TEXT("%-50s %15s %15s %15s %15s"),
 			TEXT("Container"), TEXT("Chunk(s)"), TEXT("Size (MiB)"), TEXT("Uploaded"), TEXT("Uploaded (MiB)"));
-		UE_LOG(LogIas, Display, TEXT("--------------------------------------------------------------------------------------------------------"));
+		UE_LOG(LogIas, Display, TEXT("-------------------------------------------------------------------------------------------------------------------"));
 
 		FContainerStats TotalStats;
 		for (const TTuple<FString, FContainerStats>& Kv : ContainerSummary)
 		{
-			UE_LOG(LogIas, Display, TEXT("%-40s %15llu %15.2lf %15llu %15.2lf"),
+			UE_LOG(LogIas, Display, TEXT("%-50s %15llu %15.2lf %15llu %15.2lf"),
 				*Kv.Key, Kv.Value.ChunkCount, double(Kv.Value.TotalBytes) / 1024.0 / 1024.0, Kv.Value.UploadedChunkCount, double(Kv.Value.UploadedBytes) / 1024.0 / 1024.0);
 			
 			TotalStats.ChunkCount += Kv.Value.ChunkCount;
@@ -910,17 +925,17 @@ TIoStatusOr<FIoStoreUploadResult> UploadContainerFiles(
 			TotalStats.UploadedChunkCount += Kv.Value.UploadedChunkCount;
 			TotalStats.UploadedBytes += Kv.Value.UploadedBytes;
 		}
-		UE_LOG(LogIas, Display, TEXT("--------------------------------------------------------------------------------------------------------"));
-		UE_LOG(LogIas, Display, TEXT("%-40s %15llu %15.2lf %15llu %15.2lf"),
+		UE_LOG(LogIas, Display, TEXT("-------------------------------------------------------------------------------------------------------------------"));
+		UE_LOG(LogIas, Display, TEXT("%-50s %15llu %15.2lf %15llu %15.2lf"),
 			TEXT("Total"),TotalStats.ChunkCount, double(TotalStats.TotalBytes) / 1024.0 / 1024.0, TotalStats.UploadedChunkCount, double(TotalStats.UploadedBytes) / 1024.0 / 1024.0);
 		UE_LOG(LogIas, Display, TEXT(""));
 		
-		UE_LOG(LogIas, Display, TEXT("%-40s %15s %15s %15s"), TEXT("Bucket"), TEXT("TOC(s)"), TEXT("Chunk(s)"), TEXT("MiB"));
-		UE_LOG(LogIas, Display, TEXT("--------------------------------------------------------------------------------------------------------"));
-		UE_LOG(LogIas, Display, TEXT("%-40s %15llu %15d %15.2lf"), TEXT("Existing"), TotalExistingTocs, ExistingChunks.Num(), double(TotalExistingBytes) / 1024.0 / 1024.0);
-		UE_LOG(LogIas, Display, TEXT("%-40s %15llu %15d %15.2lf"), TEXT("Uploaded"), 1, TotalStats.UploadedChunkCount, double(TotalStats.UploadedBytes) / 1024.0 / 1024.0);
-		UE_LOG(LogIas, Display, TEXT("--------------------------------------------------------------------------------------------------------"));
-		UE_LOG(LogIas, Display, TEXT("%-40s %15llu %15d %15.2lf"), TEXT("Total"),
+		UE_LOG(LogIas, Display, TEXT("%-50s %15s %15s %15s"), TEXT("Bucket"), TEXT("TOC(s)"), TEXT("Chunk(s)"), TEXT("MiB"));
+		UE_LOG(LogIas, Display, TEXT("-------------------------------------------------------------------------------------------------------------------"));
+		UE_LOG(LogIas, Display, TEXT("%-50s %15llu %15d %15.2lf"), TEXT("Existing"), TotalExistingTocs, ExistingChunks.Num(), double(TotalExistingBytes) / 1024.0 / 1024.0);
+		UE_LOG(LogIas, Display, TEXT("%-50s %15llu %15d %15.2lf"), TEXT("Uploaded"), 1, TotalStats.UploadedChunkCount, double(TotalStats.UploadedBytes) / 1024.0 / 1024.0);
+		UE_LOG(LogIas, Display, TEXT("-------------------------------------------------------------------------------------------------------------------"));
+		UE_LOG(LogIas, Display, TEXT("%-50s %15llu %15d %15.2lf"), TEXT("Total"),
 			TotalExistingTocs + 1, ExistingChunks.Num() + TotalStats.UploadedChunkCount, double(TotalExistingBytes + TotalStats.UploadedBytes) / 1024.0 / 1024.0);
 		UE_LOG(LogIas, Display, TEXT(""));
 	}
diff --git a/Engine/Source/Runtime/Experimental/IoStoreOnDemand/Private/OnDemandIoDispatcherBackend.cpp b/Engine/Source/Runtime/Experimental/IoStoreOnDemand/Private/OnDemandIoDispatcherBackend.cpp
--- a/Engine/Source/Runtime/Experimental/IoStoreOnDemand/Private/OnDemandIoDispatcherBackend.cpp
+++ b/Engine/Source/Runtime/Experimental/IoStoreOnDemand/Private/OnDemandIoDispatcherBackend.cpp
@@ -30,6 +30,7 @@
 #include "Serialization/CompactBinarySerialization.h"
 #include "Serialization/JsonReader.h"
 #include "Serialization/JsonSerializer.h"
+#include "Serialization/MemoryReader.h"
 #include "Tasks/Task.h"
 
 #include <atomic>
@@ -1319,9 +1320,13 @@ TIoStatusOr<FOnDemandToc> FOnDemandIoBackend::GetToc(FHttpClient& HttpClient, co
 		{
 			if (Response.IsOk())
 			{
+				TRACE_CPUPROFILER_EVENT_SCOPE(FOnDemandIoBackend::SerializeToc);
 				FIoBuffer Buffer = Response.ConsumeValueOrDie();
 				FOnDemandToc NewToc;	
-				if (LoadFromCompactBinary(FCbFieldView(Buffer.GetData()), NewToc))
+
+				FMemoryReaderView Ar(Buffer.GetView());
+				Ar << NewToc;
+				if (!Ar.IsError())
 				{
 					Toc = TIoStatusOr<FOnDemandToc>(MoveTemp(NewToc));
 				}
diff --git a/Engine/Source/Runtime/Experimental/IoStoreOnDemand/Public/IO/IoStoreOnDemand.h b/Engine/Source/Runtime/Experimental/IoStoreOnDemand/Public/IO/IoStoreOnDemand.h
--- a/Engine/Source/Runtime/Experimental/IoStoreOnDemand/Public/IO/IoStoreOnDemand.h
+++ b/Engine/Source/Runtime/Experimental/IoStoreOnDemand/Public/IO/IoStoreOnDemand.h
@@ -18,6 +18,7 @@
 
 #define UE_API IOSTOREONDEMAND_API
 
+class FArchive;
 class FCbWriter;
 class FCbFieldView;
 class IIoStoreWriter;
@@ -36,6 +37,7 @@ enum class EOnDemandTocVersion : uint32
 	Initial			= 1,
 	UTocHash		= 2,
 	BlockHash32		= 3,
+	NoRawHash		= 4,
 
 	LatestPlusOne,
 	Latest			= (LatestPlusOne - 1)
@@ -61,6 +63,7 @@ struct FOnDemandTocHeader
 	FString CompressionFormat;
 	FString ChunksDirectory;
 	
+	UE_API friend FArchive& operator<<(FArchive& Ar, FOnDemandTocHeader& Header);
 	UE_API friend FCbWriter& operator<<(FCbWriter& Writer, const FOnDemandTocHeader& Header);
 };
 
@@ -69,13 +72,13 @@ UE_API bool LoadFromCompactBinary(FCbFieldView Field, FOnDemandTocHeader& OutToc
 struct FOnDemandTocEntry
 {
 	FIoHash Hash = FIoHash::Zero;
-	FIoHash RawHash = FIoHash::Zero;
 	FIoChunkId ChunkId = FIoChunkId::InvalidChunkId;
 	uint64 RawSize = 0;
 	uint64 EncodedSize = 0;
 	uint32 BlockOffset = ~uint32(0);
 	uint32 BlockCount = 0; 
 	
+	UE_API friend FArchive& operator<<(FArchive& Ar, FOnDemandTocEntry& Entry);
 	UE_API friend FCbWriter& operator<<(FCbWriter& Writer, const FOnDemandTocEntry& Entry);
 };
 
@@ -90,6 +93,7 @@ struct FOnDemandTocContainerEntry
 	TArray<uint32> BlockHashes;
 	FIoHash UTocHash;
 
+	UE_API friend FArchive& operator<<(FArchive& Ar, FOnDemandTocContainerEntry& ContainerEntry);
 	UE_API friend FCbWriter& operator<<(FCbWriter& Writer, const FOnDemandTocContainerEntry& ContainerEntry);
 };
 
@@ -99,10 +103,11 @@ struct FOnDemandToc
 {
 	FOnDemandTocHeader Header;
 	TArray<FOnDemandTocContainerEntry> Containers;
-	
-	UE_API friend FCbWriter& operator<<(FCbWriter& Writer, const FOnDemandToc& TocResource);
 
-	UE_NODISCARD UE_API static TIoStatusOr<FString> Save(const TCHAR* Directory, const FOnDemandToc& TocResource);
+	UE_API friend FArchive& operator<<(FArchive& Ar, FOnDemandToc& Toc);
+	UE_API friend FCbWriter& operator<<(FCbWriter& Writer, const FOnDemandToc& Toc);
+
+	static FGuid VersionGuid;
 };
 
 UE_API bool LoadFromCompactBinary(FCbFieldView Field, FOnDemandToc& OutToc);
