Minor optimizations to FOnDemandIoStore::AddToc

#rnx

- Move the container name to avoid a string alloc, the source string will be destroyed at the end of the method anyway.
- Moved the creation of the CompressionFormat FName out of the loop and only calculate it once.

diff --git a/Engine/Source/Runtime/Experimental/IoStoreOnDemand/Private/OnDemandIoDispatcherBackend.cpp b/Engine/Source/Runtime/Experimental/IoStoreOnDemand/Private/OnDemandIoDispatcherBackend.cpp
--- a/Engine/Source/Runtime/Experimental/IoStoreOnDemand/Private/OnDemandIoDispatcherBackend.cpp
+++ b/Engine/Source/Runtime/Experimental/IoStoreOnDemand/Private/OnDemandIoDispatcherBackend.cpp
@@ -432,7 +432,7 @@ public:
 	{
 		const FToc* Toc = nullptr;
 		FAES::FAESKey EncryptionKey;
-		FString Name;
+		FString Name; // TODO: Consider removing when NO_LOGGING == 1 (only used for logging at the moment)
 		FString EncryptionKeyGuid;
 		FString ChunksDirectory;
 		FName CompressionFormat;
@@ -515,14 +515,16 @@ void FOnDemandIoStore::AddToc(const FOnDemandEndpoint& Ep, FOnDemandToc&& Toc)
 		FToc* NewToc = new(Tocs) FToc{Ep};
 		NewToc->Containers.SetNum(Toc.Containers.Num()); // List of containers can never change
 
+		const FName CompressionFormat(Header.CompressionFormat);
 		int32 ContainerIndex = 0;
+
 		for (FOnDemandTocContainerEntry& Container : Toc.Containers)
 		{
 			FContainer* NewContainer = &NewToc->Containers[ContainerIndex++];
 			NewContainer->Toc = NewToc;
-			NewContainer->Name = Container.ContainerName;
+			NewContainer->Name = MoveTemp(Container.ContainerName);
 			NewContainer->ChunksDirectory = (Prefix.IsEmpty() ? Header.ChunksDirectory : Prefix / Header.ChunksDirectory).ToLower();
-			NewContainer->CompressionFormat = FName(Header.CompressionFormat);
+			NewContainer->CompressionFormat = CompressionFormat;
 			NewContainer->BlockSize = Header.BlockSize;
 			NewContainer->EncryptionKeyGuid = Container.EncryptionKeyGuid;
 			
