Geometry collection : fix kinematically animated geometry collection not properly updating their GT state from PT if they are part of a cluster union

diff --git a/Engine/Source/Runtime/Experimental/Chaos/Private/PhysicsProxy/GeometryCollectionPhysicsProxy.cpp b/Engine/Source/Runtime/Experimental/Chaos/Private/PhysicsProxy/GeometryCollectionPhysicsProxy.cpp
--- a/Engine/Source/Runtime/Experimental/Chaos/Private/PhysicsProxy/GeometryCollectionPhysicsProxy.cpp
+++ b/Engine/Source/Runtime/Experimental/Chaos/Private/PhysicsProxy/GeometryCollectionPhysicsProxy.cpp
@@ -2896,6 +2896,9 @@ void FGeometryCollectionPhysicsProxy::PushToPhysicsState()
 
 							DeferredClusterUnionParticleUpdates.Add(Handle);
 							DeferredClusterUnionChildToParentUpdates.Add(RelativeTransform);
+
+							// Make sure the particle is mark dirty to make sure its proxy will properly update the transforms
+							RigidSolver->GetEvolution()->GetParticles().MarkTransientDirtyParticle(Handle);
 						}
 					}
 				}
