Reshuffled fields in FPBDCollisionConstraint to avoid padding

diff --git a/Engine/Source/Runtime/Experimental/Chaos/Private/Chaos/Collision/PBDCollisionConstraint.cpp b/Engine/Source/Runtime/Experimental/Chaos/Private/Chaos/Collision/PBDCollisionConstraint.cpp
--- a/Engine/Source/Runtime/Experimental/Chaos/Private/Chaos/Collision/PBDCollisionConstraint.cpp
+++ b/Engine/Source/Runtime/Experimental/Chaos/Private/Chaos/Collision/PBDCollisionConstraint.cpp
@@ -181,19 +181,18 @@ namespace Chaos
 	}
 
 	FPBDCollisionConstraint::FPBDCollisionConstraint()
-		: ImplicitTransform{ FRigidTransform3(), FRigidTransform3() }
+		: CollisionMargins{ 0, 0 }
+		, ImplicitTransform{ FRigidTransform3(), FRigidTransform3() }
 		, Particle{ nullptr, nullptr }
 		, Implicit{ nullptr, nullptr }
 		, Shape{ nullptr, nullptr }
 		, Simplicial{ nullptr, nullptr }
 		, AccumulatedImpulse(0)
-		, ContainerCookie()
-		, ShapesType(EContactShapesType::Unknown)
-		, ShapeWorldTransforms{ FRigidTransform3(), FRigidTransform3() }
 		, CullDistance(TNumericLimits<FRealSingle>::Max())
-		, CollisionMargins{ 0, 0 }
+		, ContainerCookie()
 		, CollisionTolerance(0)
 		, ClosestManifoldPointIndex(INDEX_NONE)
+		, ShapeWorldTransforms{ FRigidTransform3(), FRigidTransform3() }
 		, ExpectedNumManifoldPoints(0)
 		, LastShapeWorldPositionDelta()
 		, LastShapeWorldRotationDelta()
@@ -201,10 +200,11 @@ namespace Chaos
 		, Material()
 		, Stiffness(1)
 		, Flags()
+		, ShapesType(EContactShapesType::Unknown)
 		, SavedManifoldPoints()
 		, ManifoldPoints()
-		, SolverBodies{ nullptr, nullptr }
 		, CCDTimeOfImpact(0)
+		, SolverBodies{ nullptr, nullptr }
 		, CCDEnablePenetration(0)
 		, CCDTargetPenetration(0)
 	{
@@ -219,19 +219,18 @@ namespace Chaos
 		const FImplicitObject* Implicit1,
 		const FPerShapeData* Shape1,
 		const FBVHParticles* Simplicial1)
-		: ImplicitTransform{ FRigidTransform3(), FRigidTransform3() }
+		: CollisionMargins{ 0, 0 }
+		, ImplicitTransform{ FRigidTransform3(), FRigidTransform3() }
 		, Particle{ Particle0, Particle1 }
 		, Implicit{ Implicit0, Implicit1 }
 		, Shape{ nullptr, nullptr }
 		, Simplicial{ Simplicial0, Simplicial1 }
 		, AccumulatedImpulse(0)
-		, ContainerCookie()
-		, ShapesType(EContactShapesType::Unknown)
-		, ShapeWorldTransforms{ FRigidTransform3(), FRigidTransform3() }
 		, CullDistance(TNumericLimits<FRealSingle>::Max())
-		, CollisionMargins{ 0, 0 }
+		, ContainerCookie()
 		, CollisionTolerance(0)
 		, ClosestManifoldPointIndex(INDEX_NONE)
+		, ShapeWorldTransforms{ FRigidTransform3(), FRigidTransform3() }
 		, ExpectedNumManifoldPoints(0)
 		, LastShapeWorldPositionDelta()
 		, LastShapeWorldRotationDelta()
@@ -239,10 +238,11 @@ namespace Chaos
 		, Material()
 		, Stiffness(1)
 		, Flags()
+		, ShapesType(EContactShapesType::Unknown)
 		, SavedManifoldPoints()
 		, ManifoldPoints()
-		, SolverBodies{ nullptr, nullptr }
 		, CCDTimeOfImpact(0)
+		, SolverBodies{ nullptr, nullptr }
 		, CCDEnablePenetration(0)
 		, CCDTargetPenetration(0)
 	{
diff --git a/Engine/Source/Runtime/Experimental/Chaos/Public/Chaos/Collision/PBDCollisionConstraint.h b/Engine/Source/Runtime/Experimental/Chaos/Public/Chaos/Collision/PBDCollisionConstraint.h
--- a/Engine/Source/Runtime/Experimental/Chaos/Public/Chaos/Collision/PBDCollisionConstraint.h
+++ b/Engine/Source/Runtime/Experimental/Chaos/Public/Chaos/Collision/PBDCollisionConstraint.h
@@ -134,10 +134,10 @@ namespace Chaos
 	public:
 		FPBDCollisionConstraintContainerCookie()
 			: MidPhase(nullptr)
-			, bIsMultiShapePair(false)
 			, CreationEpoch(INDEX_NONE)
 			, LastUsedEpoch(INDEX_NONE)
 			, ConstraintIndex(INDEX_NONE)
+			, bIsMultiShapePair(false)
 			, CCDConstraintIndex(INDEX_NONE)
 		{
 		}
@@ -157,11 +157,6 @@ namespace Chaos
 		// The constraint owner - set when the constraint is created
 		FParticlePairMidPhase* MidPhase;
 
-		// Used by the MidPhase when a constraint is reactivated from a Resim cache
-		// If true, indicates that the constraint was created from the recursive collision detection
-		// path rather than the prefiltered shape-pair loop
-		bool bIsMultiShapePair;
-
 		// The Epoch when then constraint was initially created
 		int32 CreationEpoch;
 
@@ -169,7 +164,12 @@ namespace Chaos
 		int32 LastUsedEpoch;
 
 		// The index in the container - this changes every tick (is valid for all constraints, including CCD)
-		int32 ConstraintIndex;
+		int32 ConstraintIndex : 31;
+
+		// Used by the MidPhase when a constraint is reactivated from a Resim cache
+		// If true, indicates that the constraint was created from the recursive collision detection
+		// path rather than the prefiltered shape-pair loop
+		uint32 bIsMultiShapePair : 1;
 
 		// The CCD index in the container - this changes every tick (is INDEX_NONE for non-CCD constraints)
 		int32 CCDConstraintIndex;
@@ -841,6 +841,10 @@ namespace Chaos
 		};
 		static_assert(sizeof(FFlags) == 2, "Unexpected size for FPBDCollisionConstraint::FFLags");
 
+		// The margins to use during collision detection. We don't always use the margins on the shapes directly.
+		// E.g., we use the smallest non-zero margin for 2 convex shapes. See InitMarginsAndTolerances
+		FRealSingle CollisionMargins[2];
+
 		// Local-space transforms of the shape (relative to particle)
 		FRigidTransform3 ImplicitTransform[2];
 		
@@ -853,19 +857,11 @@ namespace Chaos
 		FVec3f AccumulatedImpulse;					// @todo(chaos): we need to accumulate angular impulse separately
 
 	private:
-		FPBDCollisionConstraintContainerCookie ContainerCookie;
-		Private::FCollisionSortKey CollisionSortKey;
-		EContactShapesType ShapesType;
-
-		// The shape transforms at the current particle transforms
-		FRigidTransform3 ShapeWorldTransforms[2];
-
 		// The separation distance at which we don't track contacts
 		FRealSingle CullDistance;
 
-		// The margins to use during collision detection. We don't always use the margins on the shapes directly.
-		// E.g., we use the smallest non-zero margin for 2 convex shapes. See InitMarginsAndTolerances
-		FRealSingle CollisionMargins[2];
+		FPBDCollisionConstraintContainerCookie ContainerCookie;
+		Private::FCollisionSortKey CollisionSortKey;
 
 		// The collision tolerance is used to determine whether a new contact matches an old on. It is derived from the
 		// margins of the two shapes, as well as their types
@@ -874,6 +870,9 @@ namespace Chaos
 		// The index into ManifoldPoints of the point with the lowest Phi
 		int32 ClosestManifoldPointIndex;
 
+		// The shape transforms at the current particle transforms
+		FRigidTransform3 ShapeWorldTransforms[2];
+
 		// Used by manifold point injection to see how many points were in the manifold before UpdateAndTryRestore
 		int32 ExpectedNumManifoldPoints;
 
@@ -890,17 +889,18 @@ namespace Chaos
 
 		FFlags Flags;
 		Private::FImplicitBoundsTestFlags BoundsTestFlags;
+		EContactShapesType ShapesType;
 
 		TCArray<FSavedManifoldPoint, MaxManifoldPoints> SavedManifoldPoints;
 		TCArray<FManifoldPoint, MaxManifoldPoints> ManifoldPoints;
 		TCArray<FManifoldPointResult, MaxManifoldPoints> ManifoldPointResults;
 
-		// @todo(chaos): These are only needed here to support incremental collision detection for LevelSet. Fix this.
-		const FSolverBody* SolverBodies[2];
-
 		// Value in range [0,1] used to interpolate P between [X,P] that we will rollback to when solving at time of impact.
 		FRealSingle CCDTimeOfImpact;
 
+		// @todo(chaos): These are only needed here to support incremental collision detection for LevelSet. Fix this.
+		const FSolverBody* SolverBodies[2];
+
 		// The penetration at which CCD contacts get processed following a CCD sweep test.
 		// NOTE: also see GeometryParticle CCDAxisThreshold, which is used to determine when to enable sweeping during collision detection.
 		// Calculated from particle pair properties when constraint is created.
