IAS: Fix bug dropping cache entry data before being serialized to disk

#rnx

diff --git a/Engine/Source/Runtime/Experimental/IoStoreOnDemand/Private/FileIoCache.cpp b/Engine/Source/Runtime/Experimental/IoStoreOnDemand/Private/FileIoCache.cpp
--- a/Engine/Source/Runtime/Experimental/IoStoreOnDemand/Private/FileIoCache.cpp
+++ b/Engine/Source/Runtime/Experimental/IoStoreOnDemand/Private/FileIoCache.cpp
@@ -815,13 +815,6 @@ void FFileIoCache::FileWriterThreadInner()
 		Cursor += View.GetSize();
 	}
 
-	// Drop buffers.
-	for (FCacheEntry& Entry : Entries)
-	{
-		FIoBuffer& Data = Entry.Data;
-		Data = FIoBuffer();
-	}
-
 	// Open cache file and write the block to it.
 	UE_LOG(LogIasCache, VeryVerbose, TEXT("Write; cursor:%llu size:%d"), WriteCursorPos, PendingSize);
 	TRACE_CPUPROFILER_EVENT_SCOPE(FFileIoCache::WriteCacheEntry);
