Move governor into a separate .inl file

#rnx

diff --git a/Engine/Source/Runtime/Experimental/IoStoreOnDemand/Private/FileIoCache.cpp b/Engine/Source/Runtime/Experimental/IoStoreOnDemand/Private/FileIoCache.cpp
--- a/Engine/Source/Runtime/Experimental/IoStoreOnDemand/Private/FileIoCache.cpp
+++ b/Engine/Source/Runtime/Experimental/IoStoreOnDemand/Private/FileIoCache.cpp
@@ -21,6 +21,8 @@
 #include "ProfilingDebugging/CpuProfilerTrace.h"
 #include "Templates/UniquePtr.h"
 
+#include "Governor.inl"
+
 #include <atomic>
 
 #define IAS_FORCE_TOC_SAVE
@@ -395,121 +397,6 @@ FIoStatus FCacheMap::Save(const FString& FilePath, const uint64 CursorPos)
 
 
 
-////////////////////////////////////////////////////////////////////////////////
-class FGovernorExternal
-{
-public:
-				FGovernorExternal();
-	void		Set(...) {}
-	uint32		TickAllowance();
-	void		Return(int32) {}
-
-private:
-	int64		CycleThreshold;
-	int64		CycleLast;
-};
-
-////////////////////////////////////////////////////////////////////////////////
-FGovernorExternal::FGovernorExternal()
-{
-	CycleThreshold = int64(1.0 / FPlatformTime::GetSecondsPerCycle());
-	CycleThreshold = (CycleThreshold * 3) / 4;
-	CycleLast = FPlatformTime::Cycles() - CycleThreshold;
-}
-
-////////////////////////////////////////////////////////////////////////////////
-uint32 FGovernorExternal::TickAllowance()
-{
-	// We'll only check the platform layer for our allowance every now and again
-	int64 Cycle = FPlatformTime::Cycles64();
-	if (Cycle - CycleLast < CycleThreshold)
-	{
-		return 0;
-	}
-
-	CycleLast = Cycle;
-
-	IPlatformFile& Ipf = IPlatformFile::GetPlatformPhysical();
-	return Ipf.GetAllowedBytesToWriteThrottledStorage();
-}
-
-
-
-////////////////////////////////////////////////////////////////////////////////
-class FGovernorInternal
-{
-public:
-	void		Set(uint32 InAllowance, uint32 InOps, uint32 Seconds);
-	uint32		TickAllowance();
-	void		Return(int32 LeftOver);
-
-private:
-	uint32		TickAllowance(int64 Cycles);
-	int64		CycleThreshold;
-	int64		CyclePrev;
-	uint32		Allowance;
-	int32		RunOff;
-};
-
-////////////////////////////////////////////////////////////////////////////////
-void FGovernorInternal::Set(uint32 InAllowance, uint32 Ops, uint32 Seconds)
-{
-	CycleThreshold = int64(1.0 / FPlatformTime::GetSecondsPerCycle());
-
-	CycleThreshold = (CycleThreshold * Seconds) / Ops;
-	Allowance = InAllowance / Ops;
-
-	CyclePrev = FPlatformTime::Cycles() - CycleThreshold;
-}
-
-////////////////////////////////////////////////////////////////////////////////
-uint32 FGovernorInternal::TickAllowance()
-{
-	int64 Cycle = FPlatformTime::Cycles64();
-	uint32 Ret = TickAllowance(Cycle) + RunOff;
-	RunOff = 0;
-	return Ret;
-}
-
-////////////////////////////////////////////////////////////////////////////////
-uint32 FGovernorInternal::TickAllowance(int64 Cycle)
-{
-	int64 CycleDelta = FMath::Max(0ll, Cycle - CyclePrev);
-	if (CycleDelta < CycleThreshold)
-	{
-		return 0;
-	}
-
-	// A crude guard against runaway.
-	for (; CycleDelta > CycleThreshold; CycleDelta -= CycleThreshold);
-	CyclePrev = Cycle + CycleDelta;
-
-	return Allowance;
-}
-
-////////////////////////////////////////////////////////////////////////////////
-void FGovernorInternal::Return(int32 LeftOver)
-{
-	// Roughly keep to some arbitrary limit so as to avoid excessive growth
-	int32 OnePointFive = Allowance + (Allowance >> 1);
-	RunOff = FMath::Min(OnePointFive, LeftOver);
-}
-
-
-
-////////////////////////////////////////////////////////////////////////////////
-#if !defined(UE_USE_PLATFORM_GOVERNOR)
-#	define UE_USE_PLATFORM_GOVERNOR 0
-#endif
-
-#if UE_USE_PLATFORM_GOVERNOR
-	using FGovernor = FGovernorExternal;
-#else
-	using FGovernor = FGovernorInternal;
-#endif
-
-
-
 ////////////////////////////////////////////////////////////////////////////////
 class FFileIoCache final
 	: public FRunnable
diff --git a/Engine/Source/Runtime/Experimental/IoStoreOnDemand/Private/Governor.inl b/Engine/Source/Runtime/Experimental/IoStoreOnDemand/Private/Governor.inl
new file mode 100644
--- /dev/null
+++ b/Engine/Source/Runtime/Experimental/IoStoreOnDemand/Private/Governor.inl
@@ -0,0 +1,122 @@
+// Copyright Epic Games, Inc. All Rights Reserved.
+
+#pragma once
+
+#include "HAL/PlatformTime.h"
+
+namespace UE::IO::Private {
+
+////////////////////////////////////////////////////////////////////////////////
+class FGovernorExternal
+{
+public:
+				FGovernorExternal();
+	void		Set(...) {}
+	uint32		TickAllowance();
+	void		Return(int32) {}
+
+private:
+	int64		CycleThreshold;
+	int64		CycleLast;
+};
+
+////////////////////////////////////////////////////////////////////////////////
+inline FGovernorExternal::FGovernorExternal()
+{
+	CycleThreshold = int64(1.0 / FPlatformTime::GetSecondsPerCycle());
+	CycleThreshold = (CycleThreshold * 3) / 4;
+	CycleLast = FPlatformTime::Cycles() - CycleThreshold;
+}
+
+////////////////////////////////////////////////////////////////////////////////
+inline uint32 FGovernorExternal::TickAllowance()
+{
+	// We'll only check the platform layer for our allowance every now and again
+	int64 Cycle = FPlatformTime::Cycles64();
+	if (Cycle - CycleLast < CycleThreshold)
+	{
+		return 0;
+	}
+
+	CycleLast = Cycle;
+
+	IPlatformFile& Ipf = IPlatformFile::GetPlatformPhysical();
+	return Ipf.GetAllowedBytesToWriteThrottledStorage();
+}
+
+
+
+////////////////////////////////////////////////////////////////////////////////
+class FGovernorInternal
+{
+public:
+	void		Set(uint32 InAllowance, uint32 InOps, uint32 Seconds);
+	uint32		TickAllowance();
+	void		Return(int32 LeftOver);
+
+private:
+	uint32		TickAllowance(int64 Cycles);
+	int64		CycleThreshold;
+	int64		CyclePrev;
+	uint32		Allowance;
+	int32		RunOff;
+};
+
+////////////////////////////////////////////////////////////////////////////////
+inline void FGovernorInternal::Set(uint32 InAllowance, uint32 Ops, uint32 Seconds)
+{
+	CycleThreshold = int64(1.0 / FPlatformTime::GetSecondsPerCycle());
+
+	CycleThreshold = (CycleThreshold * Seconds) / Ops;
+	Allowance = InAllowance / Ops;
+
+	CyclePrev = FPlatformTime::Cycles() - CycleThreshold;
+}
+
+////////////////////////////////////////////////////////////////////////////////
+inline uint32 FGovernorInternal::TickAllowance()
+{
+	int64 Cycle = FPlatformTime::Cycles64();
+	uint32 Ret = TickAllowance(Cycle) + RunOff;
+	RunOff = 0;
+	return Ret;
+}
+
+////////////////////////////////////////////////////////////////////////////////
+inline uint32 FGovernorInternal::TickAllowance(int64 Cycle)
+{
+	int64 CycleDelta = FMath::Max(0ll, Cycle - CyclePrev);
+	if (CycleDelta < CycleThreshold)
+	{
+		return 0;
+	}
+
+	// A crude guard against runaway.
+	for (; CycleDelta > CycleThreshold; CycleDelta -= CycleThreshold);
+	CyclePrev = Cycle + CycleDelta;
+
+	return Allowance;
+}
+
+////////////////////////////////////////////////////////////////////////////////
+inline void FGovernorInternal::Return(int32 LeftOver)
+{
+	// Roughly keep to some arbitrary limit so as to avoid excessive growth
+	int32 OnePointFive = Allowance + (Allowance >> 1);
+	RunOff = FMath::Min(OnePointFive, LeftOver);
+}
+
+
+
+////////////////////////////////////////////////////////////////////////////////
+#if !defined(UE_USE_PLATFORM_GOVERNOR)
+#	define UE_USE_PLATFORM_GOVERNOR 0
+#endif
+
+#if UE_USE_PLATFORM_GOVERNOR
+	using FGovernor = FGovernorExternal;
+#else
+	using FGovernor = FGovernorInternal;
+#endif
+
+} // namespace UE::IO::Private
