A journaled IAS file cache with throttled writes

#rnx

diff --git a/Engine/Source/Runtime/Experimental/IoStoreOnDemand/Private/FileIoCache.cpp b/Engine/Source/Runtime/Experimental/IoStoreOnDemand/Private/FileIoCache.cpp
--- a/Engine/Source/Runtime/Experimental/IoStoreOnDemand/Private/FileIoCache.cpp
+++ b/Engine/Source/Runtime/Experimental/IoStoreOnDemand/Private/FileIoCache.cpp
@@ -741,7 +741,15 @@ void FFileIoCache::FileWriterThreadInner()
 
 } // namespace UE::IO::Private
 
+
+
 TUniquePtr<IIoCache> MakeFileIoCache(const FFileIoCacheConfig& Config)
 {
-	return MakeUnique<UE::IO::Private::FFileIoCache>(Config);
+	if (Config.UseLegacy)
+	{
+		return MakeUnique<UE::IO::Private::FFileIoCache>(Config);
+	}
+
+	extern TUniquePtr<IIoCache> MakeJournaledCache(const FFileIoCacheConfig&);
+	return MakeJournaledCache(Config);
 }
diff --git a/Engine/Source/Runtime/Experimental/IoStoreOnDemand/Private/FileIoCache.h b/Engine/Source/Runtime/Experimental/IoStoreOnDemand/Private/FileIoCache.h
--- a/Engine/Source/Runtime/Experimental/IoStoreOnDemand/Private/FileIoCache.h
+++ b/Engine/Source/Runtime/Experimental/IoStoreOnDemand/Private/FileIoCache.h
@@ -37,15 +37,18 @@ struct FFileIoCacheConfig
 {
 	struct FRate
 	{
-		uint32 Allowance = ~0u;
-		uint32 Ops = 256;
-		uint32 Seconds = 1;
+		uint32	Allowance = 16 << 20;
+		uint32	Ops = 32;
+		uint32	Seconds = 60;
 	};
 
-	uint64 DiskQuota = 1ull << 30;
-	uint32 MemoryQuota = 3 << 20;
-	FRate WriteRate;
-	bool DropCache = false;
+	uint64		DiskQuota = 512ull << 20;
+	uint32		MemoryQuota = 2 << 20;
+	uint32		JournalQuota = 4 << 20;
+	uint32		JournalFlushInterval = 4;
+	uint32		UseLegacy = 0;
+	FRate		WriteRate;
+	bool		DropCache = false;
 };
 
 TUniquePtr<IIoCache> MakeFileIoCache(const FFileIoCacheConfig& Config);
diff --git a/Engine/Source/Runtime/Experimental/IoStoreOnDemand/Private/IoStoreOnDemand.cpp b/Engine/Source/Runtime/Experimental/IoStoreOnDemand/Private/IoStoreOnDemand.cpp
--- a/Engine/Source/Runtime/Experimental/IoStoreOnDemand/Private/IoStoreOnDemand.cpp
+++ b/Engine/Source/Runtime/Experimental/IoStoreOnDemand/Private/IoStoreOnDemand.cpp
@@ -120,6 +120,8 @@ static FFileIoCacheConfig GetFileIoCacheConfig(const TCHAR* CommandLine)
 	GetConfigInt(TEXT("WriteBytesPerPeriod"),	Ret.WriteRate.Allowance);
 	GetConfigInt(TEXT("DiskQuota"),				Ret.DiskQuota);
 	GetConfigInt(TEXT("MemoryQuota"),			Ret.MemoryQuota);
+	GetConfigInt(TEXT("JournalQuota"),			Ret.JournalQuota);
+	GetConfigInt(TEXT("JournalFlushInterval"),	Ret.JournalFlushInterval);
 #undef GetConfigInt
 
 #if !UE_BUILD_SHIPPING
diff --git a/Engine/Source/Runtime/Experimental/IoStoreOnDemand/Private/JournaledCache.cpp b/Engine/Source/Runtime/Experimental/IoStoreOnDemand/Private/JournaledCache.cpp
new file mode 100644
--- /dev/null
+++ b/Engine/Source/Runtime/Experimental/IoStoreOnDemand/Private/JournaledCache.cpp
@@ -0,0 +1,1449 @@
+// Copyright Epic Games, Inc. All Rights Reserved.
+
+#include "FileIoCache.h"
+#include "IO/IoStoreOnDemand.h"
+#include "Statistics.h"
+
+#include "CancellationToken.h"
+#include "Containers/Array.h"
+#include "Containers/Map.h"
+#include "Containers/UnrealString.h"
+#include "HAL/PlatformFile.h"
+#include "HAL/PlatformProcess.h"
+#include "HAL/Runnable.h"
+#include "HAL/RunnableThread.h"
+#include "IO/IoDispatcher.h"
+#include "IO/IoBuffer.h"
+#include "IO/IoHash.h"
+#include "Math/UnrealMath.h"
+#include "Misc/Paths.h"
+#include "Misc/ScopeRWLock.h"
+#include "Misc/StringBuilder.h"
+#include "Tasks/Task.h"
+#include "Templates/UniquePtr.h"
+
+#include "Governor.inl"
+
+#include <atomic>
+
+namespace UE::IO::Private
+{
+
+namespace JournaledCache
+{
+
+// {{{1 structs ................................................................
+
+////////////////////////////////////////////////////////////////////////////////
+struct FDebugCacheEntry
+{
+	using Callback = void(void* Param, const FDebugCacheEntry&);
+	uint64 Key;
+	uint32 Size;
+	uint32 IsPending : 1;
+	uint32 _Unused : 31;
+};
+
+////////////////////////////////////////////////////////////////////////////////
+template <bool IsExclusive>
+class FAccessScope
+{
+public:
+	FAccessScope(FRWLock& InLock)
+	: Lock(&InLock)
+	{
+		if (IsExclusive) Lock->WriteLock(); else Lock->ReadLock();
+	}
+
+	~FAccessScope()
+	{
+		if (Lock == nullptr) return;
+		if (IsExclusive) Lock->WriteUnlock(); else Lock->ReadUnlock();
+	}
+
+					FAccessScope() = default;
+					FAccessScope(FAccessScope&& Rhs) { Swap(Lock, Rhs.Lock); }
+	FAccessScope&	operator = (FAccessScope&& Rhs) { Swap(Lock, Rhs.Lock); return *this; }
+
+private:
+	FRWLock*		Lock = nullptr;
+					FAccessScope(const FAccessScope&) = delete;
+	FAccessScope	operator = (const FAccessScope&) = delete;
+};
+using FReadAccess	= FAccessScope<false>;
+using FWriteAccess	= FAccessScope<true>;
+
+////////////////////////////////////////////////////////////////////////////////
+using EntryHandle = UPTRINT;
+
+
+
+// {{{1 pending ................................................................
+
+////////////////////////////////////////////////////////////////////////////////
+class FPending
+{
+public:
+	struct FItem
+	{
+		uint64		Key;
+		FIoBuffer	Data;
+	};
+
+	using PeelItems = TArray<FItem>;
+
+					FPending(uint32 InMaxSize=64 << 10);
+	void			Reset();
+	uint32			GetCount() const	{ return Items.Num(); }
+	uint32			GetUsed() const		{ return UsedSize; }
+	uint32			GetMax() const		{ return MaxSize; }
+	EntryHandle		Get(uint64 Key) const;
+	bool			Materialize(EntryHandle Handle, FIoBuffer& Out, uint32 Offset=0) const;
+	bool			Put(uint64 Key, FIoBuffer&& Data);
+	int32			Peel(int32 PeelSize, PeelItems& Out);
+	uint32			DebugVisit(void* Param, FDebugCacheEntry::Callback* Callback);
+
+private:
+	template <typename Lambda>
+	int32			DropImpl(uint32 Size, Lambda&& Callback);
+	int32			Drop(uint32 Size);
+	uint32			MaxSize;
+	uint32			UsedSize = 0;
+	TArray<FItem>	Items;
+};
+
+////////////////////////////////////////////////////////////////////////////////
+FPending::FPending(uint32 InMaxSize)
+: MaxSize(InMaxSize)
+{
+}
+
+////////////////////////////////////////////////////////////////////////////////
+void FPending::Reset()
+{
+	Items.Reset();
+	UsedSize = 0;
+}
+
+////////////////////////////////////////////////////////////////////////////////
+EntryHandle FPending::Get(uint64 Key) const
+{
+	for (auto& Item : Items)
+	{
+		if (Item.Key == Key)
+		{
+			return UPTRINT(&Item);
+		}
+	}
+
+	return 0;
+}
+
+////////////////////////////////////////////////////////////////////////////////
+bool FPending::Materialize(EntryHandle Handle, FIoBuffer& Out, uint32 Offset) const
+{
+	FItem& Item = *(FItem*)Handle;
+
+	if (Out.GetData() == nullptr)
+	{
+		Out = Item.Data;
+		return true;
+	}
+
+	uint32 Size = uint32(Out.GetSize());
+
+	FMemoryView View = Item.Data.GetView();
+	View = View.RightChop(Offset);
+	if (View.GetSize() > Size)
+	{
+		View = View.Mid(0, Size);
+	}
+
+	Out.GetMutableView().CopyFrom(View);
+
+	return true;
+}
+
+////////////////////////////////////////////////////////////////////////////////
+bool FPending::Put(uint64 Key, FIoBuffer&& Data)
+{
+	if (Get(Key) != 0)
+	{
+		uint32 DataSize = uint32(Data.GetSize());
+		FOnDemandIoBackendStats::Get()->OnCachePutExisting(DataSize);
+		return true;
+	}
+
+	uint32 Size = Data.GetSize();
+	if (Size == 0 || MaxSize < Size)
+	{
+		return false;
+	}
+
+	if (UsedSize + Size > MaxSize)
+	{
+		int32 DroppedSize = Drop(Size);
+		FOnDemandIoBackendStats::Get()->OnCachePutReject(DroppedSize);
+	}
+
+	Items.Add({ Key, MoveTemp(Data) });
+	UsedSize += Size;
+
+	FOnDemandIoBackendStats::Get()->OnCachePut();
+	FOnDemandIoBackendStats::Get()->OnCachePendingBytes(UsedSize);
+	return true;
+}
+
+////////////////////////////////////////////////////////////////////////////////
+int32 FPending::Peel(int32 PeelSize, PeelItems& Out)
+{
+	return DropImpl(PeelSize, [&Out] (FItem&& Item) {
+		Out.Add(MoveTemp(Item));
+	});
+}
+
+////////////////////////////////////////////////////////////////////////////////
+uint32 FPending::DebugVisit(void* Param, FDebugCacheEntry::Callback* Callback)
+{
+	FDebugCacheEntry Out = {};
+	Out.IsPending = 1;
+	for (auto& Item : Items)
+	{
+		Out.Key = Item.Key;
+		Out.Size = Item.Data.GetSize();
+		Callback(Param, Out);
+	}
+	return Items.Num();
+}
+
+////////////////////////////////////////////////////////////////////////////////
+template <typename Lambda>
+int32 FPending::DropImpl(uint32 Size, Lambda&& Callback)
+{
+	int32 DropSize = 0;
+	int32 TargetSize = FMath::Min<int32>(Size, UsedSize);
+	for (int32 n = Items.Num(); --n >= 0;)
+	{
+		if (DropSize >= TargetSize)
+		{
+			break;
+		}
+
+		uint32 Index = n ? (Size * 0x0'a9e0'493) % n : 0;
+
+		Size = Items[Index].Data.GetSize();
+		UsedSize -= Size;
+		DropSize += Size;
+
+		Callback(MoveTemp(Items[Index]));
+
+		Items[Index] = MoveTemp(Items.Last());
+		Items.Pop();
+	}
+
+	FOnDemandIoBackendStats::Get()->OnCachePendingBytes(UsedSize);
+	return DropSize;
+}
+
+////////////////////////////////////////////////////////////////////////////////
+int32 FPending::Drop(uint32 Size)
+{
+	return DropImpl(Size, [] (FItem&&) {});
+}
+
+// {{{1 journal ................................................................
+
+////////////////////////////////////////////////////////////////////////////////
+class FJournal
+{
+private:
+	static const uint32 MAGIC = 0x04930001;
+	static const uint32 SIZE_BITS = 25;
+
+	using MarkerType = uint32;
+
+	struct FDataEntry
+	{
+		uint64		Key;
+		uint64		Offset : 23;
+		uint64		Size : SIZE_BITS;
+		uint64		EntryCount : 16;
+	};
+
+	struct FPhraseDesc
+	{
+		uint32		Magic;
+		MarkerType	Marker;
+		uint64		EntryCount : 16;
+		uint64		DataCursor : 48;
+		FDataEntry	Entries[];
+	};
+
+	static_assert(sizeof(FPhraseDesc) == 16);
+	static_assert(sizeof(FPhraseDesc) == sizeof(FDataEntry));
+
+public:
+	struct FPhrase
+	{
+		bool				Add(uint64 Key, FIoBuffer&& Data);
+		TUniquePtr<uint8[]>	Buffer;
+		TArray<FDataEntry>*	Entries;
+		uint32				Cursor;
+		uint32				Index;
+		int32				MaxEntries;
+	};
+
+							FJournal(FString&& Path, uint64 InMaxDataSize, uint32 InJournalSize);
+	void					Reset();
+	FPhrase					OpenPhrase(uint32 DataSize);
+	void					ClosePhrase(FPhrase& Phrase);
+	EntryHandle				Get(uint64 Key) const;
+	bool					Materialize(EntryHandle Handle, FIoBuffer& Out, uint32 Offset=0) const;
+	int32					Flush();
+	void					Drop();
+	bool					Load();
+	uint32					DebugVisit(void* Param, FDebugCacheEntry::Callback* Callback);
+
+private:
+	struct FMapEntry
+	{
+		uint64				DataCursor : 38;
+		uint64				Size : SIZE_BITS;
+		uint64				First : 1;
+	};
+	static_assert(sizeof(FMapEntry) == sizeof(uint64));
+
+	using					FDataMap = TMap<uint64, FMapEntry>;
+	void					Spam();
+	void					PhraseReset();
+	uint64					Insert(uint64 DataBase, const FDataEntry& Entry);
+	uint64					Insert(const FPhraseDesc* Phrase);
+	void					Prune(uint64 DataBase, uint32 Size);
+	TArray<FDataEntry>		Entries;
+	FString					BinPath;
+	FDataMap				DataMap;
+	uint64					MappedBytes;
+	uint64					MaxDataSize;
+	uint64					DataCursor;
+	uint32					JournalSize;
+	uint32					JournalCursor;
+	uint32					OverRemoval;
+	MarkerType				Marker = 0;
+};
+
+////////////////////////////////////////////////////////////////////////////////
+bool FJournal::FPhrase::Add(uint64 Key, FIoBuffer&& Data)
+{
+	check(MaxEntries > 0);
+	uint32 DataSize = uint32(Data.GetSize());
+	check(DataSize < (1 << SIZE_BITS));
+	Entries->Add({Key, Cursor, DataSize});
+	std::memcpy(Buffer.Get() + Cursor, Data.GetData(), DataSize);
+	Cursor += DataSize;
+	return (--MaxEntries > 0);
+}
+
+////////////////////////////////////////////////////////////////////////////////
+FJournal::FJournal(FString&& Path, uint64 InMaxDataSize, uint32 InJournalSize)
+: BinPath(Path)
+, MaxDataSize(InMaxDataSize)
+, JournalSize(InJournalSize)
+{
+	// Align down to keep to some assumptions
+	JournalSize &= ~(sizeof(FDataEntry) - 1);
+	MaxDataSize = (MaxDataSize - JournalSize) & ~((1ull << 20) - 1);
+
+	Reset();
+}
+
+////////////////////////////////////////////////////////////////////////////////
+void FJournal::Reset()
+{
+	DataMap.Reset();
+	MappedBytes = 0;
+	DataCursor = 0;
+	// Marker = 0; // We'll just lets this roll along in its own little world
+	JournalCursor = 0;
+	OverRemoval = 0;
+	PhraseReset();
+}
+
+////////////////////////////////////////////////////////////////////////////////
+FJournal::FPhrase FJournal::OpenPhrase(uint32 DataSize)
+{
+	check((JournalCursor & (sizeof(FDataEntry) - 1)) == 0);
+
+	FPhrase Ret;
+	Ret.Buffer = TUniquePtr<uint8[]>(new uint8[DataSize + sizeof(MarkerType)]);
+	Ret.Entries = &Entries;
+	Ret.Cursor = sizeof(MarkerType);
+	Ret.Index = Entries.Num();
+	Ret.MaxEntries = int32((JournalSize - JournalCursor) / sizeof(FDataEntry)) - 1;
+
+	Entries.Add(FDataEntry{});
+
+	return Ret;
+}
+
+////////////////////////////////////////////////////////////////////////////////
+void FJournal::ClosePhrase(FPhrase& Phrase)
+{
+	int32 EntryCount = Entries.Num() - Phrase.Index - 1;
+	if (EntryCount <= 0)
+	{
+		Entries.Pop();
+		return;
+	}
+
+	Entries.Last().EntryCount = uint16(EntryCount);
+
+	uint32 WriteSize = Phrase.Cursor;
+	if (DataCursor + WriteSize > MaxDataSize)
+	{
+		OverRemoval = 0;
+		DataCursor = 0;
+	}
+
+	auto& Header = (FPhraseDesc&)(Entries[Phrase.Index]);
+	Header.Magic = MAGIC;
+	Header.Marker = Marker;
+	Header.DataCursor = DataCursor;
+	Header.EntryCount = uint16(EntryCount);
+
+	uint8* Buffer = Phrase.Buffer.Get();
+
+	*(MarkerType*)Buffer = Marker;
+	++Marker;
+
+	IPlatformFile& Ipf = IPlatformFile::GetPlatformPhysical();
+	if (TUniquePtr<IFileHandle> File(Ipf.OpenWrite(*BinPath, true, false)); File.IsValid())
+	{
+		File->Seek(DataCursor);
+		File->Write(Buffer, WriteSize);
+		File->Flush();
+		File.Reset();
+
+		Prune(DataCursor, WriteSize);
+		Insert(&Header);
+
+		DataCursor += WriteSize;
+	}
+
+	Phrase.Buffer.Reset();
+}
+
+////////////////////////////////////////////////////////////////////////////////
+EntryHandle FJournal::Get(uint64 Key) const
+{
+	return UPTRINT(DataMap.Find(Key));
+}
+
+////////////////////////////////////////////////////////////////////////////////
+bool FJournal::Materialize(EntryHandle Handle, FIoBuffer& Out, uint32 Offset) const
+{
+	const FMapEntry& Entry = *(FMapEntry*)Handle;
+
+	uint32 ReadSize = Entry.Size - Offset;
+
+	if (Out.GetData() == nullptr)
+	{
+		Out = FIoBuffer(ReadSize);
+	}
+
+	ReadSize = FMath::Min<uint32>(Out.GetSize(), ReadSize);
+
+	IPlatformFile& Ipf = IPlatformFile::GetPlatformPhysical();
+	TUniquePtr<IFileHandle> File(Ipf.OpenRead(*BinPath, false));
+	if (!File.IsValid())
+	{
+		return false;
+	}
+
+	File->Seek(Entry.DataCursor + Offset);
+	return File->Read(Out.GetData(), ReadSize);
+}
+
+////////////////////////////////////////////////////////////////////////////////
+uint64 FJournal::Insert(uint64 DataBase, const FDataEntry& Entry)
+{
+	FMapEntry Value;
+	Value.DataCursor = DataBase + Entry.Offset;
+	Value.Size = Entry.Size;
+	Value.First = (Entry.Offset == sizeof(MarkerType));
+	DataMap.Add(Entry.Key, Value);
+	return Entry.Size;
+}
+
+////////////////////////////////////////////////////////////////////////////////
+uint64 FJournal::Insert(const FPhraseDesc* Phrase)
+{
+	uint64 TotalSize = 0;
+	uint64 DataBase = Phrase->DataCursor;
+	for (uint32 i = 0; i < Phrase->EntryCount; ++i)
+	{
+		TotalSize += Insert(DataBase, Phrase->Entries[i]);
+	}
+
+	MappedBytes += TotalSize;
+	FOnDemandIoBackendStats::Get()->OnCachePersistedBytes(MappedBytes);
+	return MappedBytes;
+}
+
+////////////////////////////////////////////////////////////////////////////////
+void FJournal::Prune(uint64 DataBase, uint32 Size)
+{
+	int32 BytesRemoved = OverRemoval;
+	if (BytesRemoved >= int32(Size))
+	{
+		OverRemoval -= Size;
+		return;
+	}
+	OverRemoval = 0;
+
+	int64 A[] = { int64(DataBase), int64(DataBase + Size) };
+
+	int64 Overage = 0;
+	for (auto Iter = DataMap.CreateIterator(); Iter; ++Iter)
+	{
+		const FMapEntry& Candidate = Iter.Value();
+
+		int64 B[] = {
+			int64(Candidate.DataCursor) - (int64(Candidate.First) << 2),
+			int64(Candidate.DataCursor + Candidate.Size)
+		};
+
+		int32 Outside = (B[0] >= A[1]) | (B[1] <= A[0]);
+		if (Outside)
+		{
+			continue;
+		}
+
+		Iter.RemoveCurrent();
+		MappedBytes -= Candidate.Size;
+
+		Overage = FMath::Max(Overage, B[1] - A[1]);
+		BytesRemoved += int32(B[1] - B[0]);
+		if (BytesRemoved - Overage >= int32(Size))
+		{
+			check(BytesRemoved - Overage == int32(Size));
+			OverRemoval = int32(Overage);
+			break;
+		}
+	}
+}
+
+////////////////////////////////////////////////////////////////////////////////
+int32 FJournal::Flush()
+{
+	if (Entries.IsEmpty())
+	{
+		return 0;
+	}
+
+	uint32 Size = Entries.Num() * sizeof(Entries[0]);
+
+	if (JournalCursor + Size > JournalSize)
+	{
+		JournalCursor = 0;
+	}
+
+	IPlatformFile& Ipf = IPlatformFile::GetPlatformPhysical();
+
+	TStringBuilder<265> JrnPath;
+	JrnPath << BinPath;
+	JrnPath << TEXT(".jrn");
+	if (TUniquePtr<IFileHandle> File(Ipf.OpenWrite(*JrnPath, true, false)); File.IsValid())
+	{
+		File->Seek(JournalCursor);
+		File->Write((uint8*)(Entries.GetData()), Size);
+		File->Flush();
+		JournalCursor += Size;
+	}
+
+	PhraseReset();
+
+	Spam();
+	return Size;
+}
+
+////////////////////////////////////////////////////////////////////////////////
+void FJournal::Drop()
+{
+	IPlatformFile& Ipf = IPlatformFile::GetPlatformPhysical();
+	Ipf.DeleteFile(*BinPath);
+	Ipf.DeleteFile(*(BinPath + TEXT(".jrn")));
+	Reset();
+}
+
+////////////////////////////////////////////////////////////////////////////////
+bool FJournal::Load()
+{
+	IPlatformFile& Ipf = IPlatformFile::GetPlatformPhysical();
+
+	uint32 DataSize = 0;
+	TUniquePtr<uint8[]> Data;
+
+	TStringBuilder<256> JrnPath;
+	JrnPath << BinPath;
+	JrnPath << TEXT(".jrn");
+	if (TUniquePtr<IFileHandle> File(Ipf.OpenRead(*JrnPath, false)); File.IsValid())
+	{
+		DataSize = uint32(File->Size());
+		Data = TUniquePtr<uint8[]>(new uint8[DataSize]);
+		File->Read(Data.Get(), DataSize);
+	}
+
+	if (DataSize == 0)
+	{
+		return false;
+	}
+
+	UE_LOG(LogIas, VeryVerbose, TEXT("JournaledCache: %u byte journal found"), DataSize);
+
+	struct FParagraph
+	{
+		const FPhraseDesc*	Phrase;
+		uint32				Marker;
+		uint32				DataSize;
+	};
+
+	auto IsOob = [&Data, DataSize] (const void* Address)
+	{
+		return (UPTRINT(Address) - UPTRINT(Data.Get())) > DataSize;
+	};
+
+	auto ReadPhrases = [&IsOob] (
+		const uint8* Cursor,
+		FParagraph& Out) -> const uint8*
+	{
+		// Only proceed if we can read at least three integers
+		if (IsOob(Cursor + sizeof(FPhraseDesc)))
+		{
+			return nullptr;
+		}
+
+		const auto* Header = (FPhraseDesc*)Cursor;
+		if (Header->Magic != MAGIC)
+		{
+			return nullptr;
+		}
+
+		uint32 ToConsume = sizeof(FDataEntry) * (Header->EntryCount + 1);
+		if (IsOob(Cursor + ToConsume))
+		{
+			return nullptr;
+		}
+
+		Cursor += ToConsume;
+
+		const auto* LastEntry = (FDataEntry*)Cursor - 1;
+		if (LastEntry->EntryCount != Header->EntryCount)
+		{
+			return nullptr;
+		}
+
+		Out.Phrase = Header;
+		Out.Marker = Header->Marker;
+		Out.DataSize = LastEntry->Offset + LastEntry->Size + sizeof(MarkerType);
+
+		return Cursor;
+	};
+
+	TArray<FParagraph> Paragraphs;
+	FParagraph Paragraph;
+
+	// Read from the front
+	const uint8* Left = Data.Get();
+	for (const uint8* Next; ; Left = Next)
+	{
+		Next = ReadPhrases(Left, Paragraph);
+		if (Next == nullptr)
+		{
+			break;
+		}
+		Paragraphs.Add(Paragraph);
+	}
+
+	// Read from the back
+	const uint8* Right = Data.Get() + DataSize;
+	while (Right >= (Left + (sizeof(FDataEntry) * 2)))
+	{
+		const auto* Entry = (FDataEntry*)Right - 1;
+
+		Entry -= Entry->EntryCount;
+		if (IsOob(Entry))
+		{
+			break;
+		}
+
+		const auto* Next = (uint8*)Entry;
+		if (ReadPhrases(Next, Paragraph) == nullptr)
+		{
+			break;
+		}
+
+		Paragraphs.Add(Paragraph);
+		Right = Next;
+	}
+
+	UE_LOG(LogIas, VeryVerbose, TEXT("JournaledCache: %d paragraphs discovered"), Paragraphs.Num());
+
+	if (Paragraphs.IsEmpty())
+	{
+		return false;
+	}
+
+	auto LessWithWrap = [] (
+		const FParagraph& Lhs,
+		const FParagraph& Rhs)
+	{
+        uint32 L = Lhs.Marker;
+        uint32 R = Rhs.Marker;
+        enum : uint32 { LowQuarter = 1u << 30, HighQuarter = 3u << 30 };
+        int32 Wrap = (L < LowQuarter) & (R >= HighQuarter);
+        Wrap |= (R < LowQuarter) & (L >= HighQuarter);
+        return (L < R) != Wrap;
+	};
+	Algo::Sort(Paragraphs, LessWithWrap);
+
+	// Eliminate any discontinuities and find where data wrapped
+	int32 BasisIndex = 0;
+	int64 Remaining = MaxDataSize;
+	for (int32 i = Paragraphs.Num() - 2; i >= 0; BasisIndex = i--)
+	{
+		const FParagraph& Newer = Paragraphs[i + 1];
+
+		int32 PhraseDataSize = Newer.DataSize;
+		Remaining -= PhraseDataSize;
+		if (Remaining < 0)
+		{
+			break;
+		}
+
+		const FParagraph& Older = Paragraphs[i];
+		if (Newer.Marker != Older.Marker + 1)
+		{
+			break;
+		}
+	}
+
+	TUniquePtr<IFileHandle> File(Ipf.OpenRead(*BinPath, false));
+	if (!File.IsValid())
+	{
+		UE_LOG(LogIas, VeryVerbose, TEXT("JournaledCache: unable to open '%s'"), *BinPath);
+		return false;
+	}
+
+	if (uint64(File->Size()) > MaxDataSize)
+	{
+		UE_LOG(LogIas, VeryVerbose, TEXT("JournaledCache: Dropping - existing cache to bi; %llu/%llu"), uint64(File->Size()), MaxDataSize);
+		Drop();
+		return false;
+	}
+
+	// Detect data writes that are newer than any journal flushes.
+	auto ReadBinMarker = [File = File.Get()] (int64 Cursor)
+	{
+		uint32 Value;
+		File->Seek(Cursor);
+		File->Read((uint8*)(&Value), sizeof(Value));
+		return Value;
+	};
+
+	for (; BasisIndex < Paragraphs.Num(); ++BasisIndex)
+	{
+		const FParagraph& Stock = Paragraphs[BasisIndex];
+
+		uint64 DataBase = Stock.Phrase->DataCursor;
+		if (DataBase + Stock.DataSize > MaxDataSize)
+		{
+			DataBase = 0;
+		}
+
+		uint32 DataMark = ReadBinMarker(DataBase);
+		if (DataMark == Stock.Marker)
+		{
+			break;
+		}
+	}
+
+	// Add known entries into the tree.
+	for (uint32 i = BasisIndex, n = Paragraphs.Num(); i < n; ++i)
+	{
+		const FParagraph& Holm = Paragraphs[i];
+		Insert(Holm.Phrase);
+	}
+	
+	// Prime the journal's state
+	const FParagraph& LastPara = Paragraphs.Last();
+	const FPhraseDesc* LastPhrase = LastPara.Phrase;
+	Marker = LastPara.Marker + 1;
+
+	if (DataSize <= JournalSize)
+	{
+		JournalCursor = uint32(UPTRINT(LastPhrase + LastPhrase->EntryCount + 1) - UPTRINT(Data.Get()));
+	}
+	else
+	{
+		Ipf.DeleteFile(*JrnPath);
+	}
+
+	DataCursor = LastPhrase->DataCursor + LastPara.DataSize;
+	if (DataCursor > MaxDataSize)
+	{
+		UE_LOG(LogIas, VeryVerbose, TEXT("JournaledCache: Dropping - DataCursor too big; %llu/%llu"), DataCursor, MaxDataSize);
+		Drop();
+		return false;
+	}
+
+	Spam();
+	return true;
+}
+
+////////////////////////////////////////////////////////////////////////////////
+void FJournal::Spam()
+{
+	UE_LOG(LogIas, VeryVerbose,
+		TEXT("JournaledCache: MappedKiB=%llu Entries=%d DataCur=%llu JournalCur=%u Marker=%u)"),
+		(MappedBytes >> 10),
+		DataMap.Num(),
+		DataCursor,
+		JournalCursor,
+		Marker
+	);
+}
+
+////////////////////////////////////////////////////////////////////////////////
+uint32 FJournal::DebugVisit(void* Param, FDebugCacheEntry::Callback* Callback)
+{
+	FDebugCacheEntry Out = {};
+	for (const auto& Entry : DataMap)
+	{
+		Out.Key = Entry.Key;
+		Out.Size = Entry.Value.Size;
+		Callback(Param, Out);
+	}
+	return DataMap.Num();
+}
+
+////////////////////////////////////////////////////////////////////////////////
+void FJournal::PhraseReset()
+{
+	Entries.Reset();
+}
+
+// {{{1 cache ..................................................................
+
+////////////////////////////////////////////////////////////////////////////////
+class FCache
+{
+public:
+	struct FConfig
+		: public FFileIoCacheConfig
+	{
+		FString	Path;
+	};
+
+	enum class EHit
+	{
+		None, Memory, Disk
+	};
+
+	class FEntry
+	{
+	public:
+					FEntry() = default;
+					FEntry(FEntry&& Rhs)		= default;
+		FEntry&		operator = (FEntry&& Rhs)	= default;
+		bool		IsHit() const		{ return GetHitType() != EHit::None; }
+		EHit		GetHitType() const	{ return HitType; }
+		bool		Materialize(FIoBuffer& Out, uint32 Offset=0);
+
+	private:
+		friend		FCache;
+		EntryHandle	Handle;
+		EHit		HitType = EHit::None;
+		FReadAccess	Lock;
+		const void*	Owner = nullptr;
+					FEntry(const FEntry& Rhs) = delete;
+		FEntry&		operator = (const FEntry& Rhs) = delete;
+	};
+
+					FCache(FConfig&& Config);
+	void			Reset();
+	bool			Load();
+	FEntry			Get(uint64 Key) const;
+	bool			Put(uint64 Key, FIoBuffer& Data);
+	int32			Flush(int32 Allowance);
+	uint32			DebugVisit(void* Param, FDebugCacheEntry::Callback* Callback);
+
+private:
+	mutable FRWLock	Lock;
+	FPending		Pending;
+	FJournal		Journal;
+	uint32			FlushIndex = 0;
+	uint32			FlushPeriod = 4;
+};
+
+////////////////////////////////////////////////////////////////////////////////
+bool FCache::FEntry::Materialize(FIoBuffer& Out, uint32 Offset)
+{
+	if (Owner == nullptr)
+	{
+		return false;
+	}
+
+	switch (HitType)
+	{
+	case EHit::Memory:	((const FPending*)Owner)->Materialize(Handle, Out, Offset); break;
+	case EHit::Disk:	((const FJournal*)Owner)->Materialize(Handle, Out, Offset); break;
+	case EHit::None:	return false;
+	}
+
+	FOnDemandIoBackendStats::Get()->OnCacheGet(Out.GetSize());
+	return true;
+}
+
+////////////////////////////////////////////////////////////////////////////////
+FCache::FCache(FConfig&& Config)
+: Pending(Config.MemoryQuota)
+, Journal(MoveTemp(Config.Path), Config.DiskQuota, Config.JournalQuota)
+, FlushPeriod(Config.JournalFlushInterval)
+{
+	if (Config.DropCache)
+	{
+		Journal.Drop();
+	}
+}
+
+////////////////////////////////////////////////////////////////////////////////
+void FCache::Reset()
+{
+	FWriteAccess _(Lock);
+	Pending.Reset();
+	Journal.Reset();
+	FlushIndex = 0;
+}
+
+////////////////////////////////////////////////////////////////////////////////
+bool FCache::Load()
+{
+	FWriteAccess _(Lock);
+	return Journal.Load();
+}
+
+////////////////////////////////////////////////////////////////////////////////
+FCache::FEntry FCache::Get(uint64 Key) const
+{
+	FEntry Ret;
+	Ret.Lock = FReadAccess(Lock);
+
+	if (Ret.Handle = Journal.Get(Key); Ret.Handle)
+	{
+		Ret.HitType = EHit::Disk;
+		Ret.Owner = &Journal;
+	}
+		
+	else if (Ret.Handle = Pending.Get(Key); Ret.Handle)
+	{
+		Ret.HitType = EHit::Memory;
+		Ret.Owner = &Pending;
+	}
+
+	return Ret;
+}
+
+////////////////////////////////////////////////////////////////////////////////
+bool FCache::Put(uint64 Key, FIoBuffer& Data)
+{
+	FIoBuffer Cloned = Data;
+	FWriteAccess _(Lock);
+	return Pending.Put(Key, MoveTemp(Cloned));
+}
+
+////////////////////////////////////////////////////////////////////////////////
+int32 FCache::Flush(int32 Allowance)
+{
+	FlushIndex += 1;
+	if (FlushIndex >= FlushPeriod)
+	{
+		FlushIndex -= FlushPeriod;
+		FWriteAccess _(Lock);
+		return Allowance - Journal.Flush();
+	}
+
+	if (Pending.GetUsed() == 0)
+	{
+		FlushIndex -= (FlushIndex == 1);
+		FWriteAccess _(Lock);
+		return Allowance - Journal.Flush();
+	}
+
+	FPending::PeelItems PeelItems;
+
+	FWriteAccess _(Lock);
+
+	int32 PendingSize = Pending.Peel(Allowance, PeelItems);
+
+	FJournal::FPhrase Phrase = Journal.OpenPhrase(PendingSize);
+	int32 PeelIndex = -1;
+	for (int32 i = 0, n = PeelItems.Num(); i < n; ++i)
+	{
+		auto& [Key, Data] = PeelItems[i];
+		if (!Phrase.Add(Key, MoveTemp(Data)))
+		{
+			PeelIndex = i;
+			break;
+		}
+	}
+	Journal.ClosePhrase(Phrase);
+
+	if (PeelIndex >= 0)
+	{
+		/* end of journal reached so not all peeled items could be added, may
+		 * we can re-add leftover peeled items back to pending? */
+	}
+
+	return Allowance - PendingSize;
+}
+
+////////////////////////////////////////////////////////////////////////////////
+uint32 FCache::DebugVisit(void* Param, FDebugCacheEntry::Callback* Callback)
+{
+	FReadAccess _(Lock);
+	uint32 Count = 0;
+	Count += Pending.DebugVisit(Param, Callback);
+	Count += Journal.DebugVisit(Param, Callback);
+	return Count;
+}
+
+// }}}
+
+} // namespace JournaledCache
+
+
+
+// {{{1 journaled-cache ........................................................
+
+////////////////////////////////////////////////////////////////////////////////
+class FJournaledCache
+	: public IIoCache
+	, public FRunnable
+{
+public:
+	using GetRetType = UE::Tasks::TTask<TIoStatusOr<FIoBuffer>>;
+
+								FJournaledCache(const FFileIoCacheConfig& Config);
+								~FJournaledCache();
+	virtual bool				ContainsChunk(const FIoHash& Key) const override;
+	virtual GetRetType			Get(const FIoHash& Key, const FIoReadOptions& Options, const FIoCancellationToken* CancellationToken) override;
+	virtual FIoStatus			Put(const FIoHash& Key, FIoBuffer& Data) override;
+
+private:
+	using FCacheInner = JournaledCache::FCache;
+
+	void						Update();
+	static uint64				ReduceKey(const FIoHash& Key);
+	TUniquePtr<FCacheInner>		Cache;
+	FGovernor					Governor;
+
+	// FRunnable
+	virtual uint32				Run() override;
+	virtual void				Stop() override;
+	void						StartThread();
+	TUniquePtr<FRunnableThread>	Thread;
+	FEventRef					WakeEvent;
+	std::atomic<bool>			Running = true;
+};
+
+////////////////////////////////////////////////////////////////////////////////
+FJournaledCache::FJournaledCache(const FFileIoCacheConfig& Config)
+{
+	TStringBuilder<256> CachePath;
+	CachePath << FPaths::ProjectPersistentDownloadDir();
+	CachePath << TEXT("ias.cache.0");
+
+	FCacheInner::FConfig EventualConfig;
+	static_cast<FFileIoCacheConfig&>(EventualConfig) = Config;
+	EventualConfig.Path = CachePath;
+	Cache = MakeUnique<FCacheInner>(MoveTemp(EventualConfig));
+	Cache->Load();
+
+	const FFileIoCacheConfig::FRate& WriteRate = Config.WriteRate;
+	Governor.Set(WriteRate.Allowance, WriteRate.Ops, WriteRate.Seconds);
+
+	StartThread();
+}
+
+////////////////////////////////////////////////////////////////////////////////
+void FJournaledCache::StartThread()
+{
+	auto* Inst = FRunnableThread::Create(this, TEXT("Ias.FileCache"), 0, TPri_BelowNormal);
+	Thread.Reset(Inst);
+}
+
+////////////////////////////////////////////////////////////////////////////////
+FJournaledCache::~FJournaledCache()
+{
+}
+
+////////////////////////////////////////////////////////////////////////////////
+void FJournaledCache::Update()
+{
+	int32 WriteAllowance = Governor.TickAllowance();
+	if (!WriteAllowance)
+	{
+		return;
+	}
+
+	int32 AllowanceUsed = Cache->Flush(WriteAllowance);
+
+	Governor.Return(WriteAllowance - AllowanceUsed);
+}
+
+////////////////////////////////////////////////////////////////////////////////
+uint32 FJournaledCache::Run()
+{
+	while (Running.load(std::memory_order_relaxed))
+	{
+		Update();
+		WakeEvent->Wait(127);
+	}
+	return 0;
+}
+
+////////////////////////////////////////////////////////////////////////////////
+void FJournaledCache::Stop()
+{
+	Running.store(false, std::memory_order_relaxed);
+	WakeEvent->Trigger();
+}
+
+////////////////////////////////////////////////////////////////////////////////
+bool FJournaledCache::ContainsChunk(const FIoHash& Key) const
+{
+	uint64 InnerKey = ReduceKey(Key);
+	return Cache->Get(InnerKey).IsHit();
+}
+
+////////////////////////////////////////////////////////////////////////////////
+FJournaledCache::GetRetType	FJournaledCache::Get(
+	const FIoHash& Key,
+	const FIoReadOptions& Options,
+	const FIoCancellationToken* CancelToken)
+{
+	uint64 InnerKey = ReduceKey(Key);
+	return UE::Tasks::Launch(TEXT("IasCacheGet"), [this, InnerKey, Options, CancelToken] () {
+		FCacheInner::FEntry Entry = Cache->Get(InnerKey);
+		if (!Entry.IsHit())
+		{
+			return TIoStatusOr<FIoBuffer>(FIoStatus(EIoErrorCode::Unknown));
+		}
+
+		if (CancelToken != nullptr && CancelToken->IsCancelled())
+		{
+			return TIoStatusOr<FIoBuffer>(FIoStatus(EIoErrorCode::Cancelled));
+		}
+
+		FIoBuffer Buffer;
+
+		uint64 Offset = Options.GetOffset();
+		int64 Size = Options.GetSize();
+		if (void* DestAddr = Options.GetTargetVa(); DestAddr != nullptr)
+		{
+			check(Size > 0);
+			Buffer = FIoBuffer(FIoBuffer::Wrap, DestAddr, Size);
+		}
+		else if (Size >= 0)
+		{
+			Buffer = FIoBuffer(Size);
+		}
+
+		if (!Entry.Materialize(Buffer, Offset))
+		{
+			return TIoStatusOr<FIoBuffer>(FIoStatus(EIoErrorCode::ReadError));
+		}
+
+		return TIoStatusOr<FIoBuffer>(Buffer);
+	});
+}
+
+////////////////////////////////////////////////////////////////////////////////
+FIoStatus FJournaledCache::Put(const FIoHash& Key, FIoBuffer& Data)
+{
+	uint64 InnerKey = ReduceKey(Key);
+	bool Ok = Cache->Put(InnerKey, Data);
+	return Ok ? FIoStatus::Ok : FIoStatus(EIoErrorCode::Unknown);
+}
+
+////////////////////////////////////////////////////////////////////////////////
+uint64 FJournaledCache::ReduceKey(const FIoHash& Key)
+{
+	const uint8* Bytes = Key.GetBytes();
+	uint64 Ret[2];
+	std::memcpy(Ret + 0, Bytes + 0, sizeof(uint64));
+	std::memcpy(Ret + 1, Bytes + 8, sizeof(uint64));
+	return Ret[0] /*bad...*/ ^ /*...mixer*/ Ret[1];
+}
+
+// }}}
+
+} // namespace UE::IO::Private
+
+
+
+// {{{1 test ...................................................................
+
+#if defined(IAS_JOURNALED_CACHE_TEST)
+
+namespace IasJournaledFileCacheTest
+{
+
+////////////////////////////////////////////////////////////////////////////////
+static constexpr uint64 operator ""_Ki (unsigned long long value) { return value << 10; }
+static constexpr uint64 operator ""_Mi (unsigned long long value) { return value << 20; }
+
+////////////////////////////////////////////////////////////////////////////////
+static uint64 KeyGen(const uint8* Data, uint32 Size)
+{
+	uint64 Ret = 0x0'a9e0'493;
+	for (; Size; Ret = (Data[--Size] + Ret) * 0x369dea0f31a53f85ull);
+	return Ret;
+}
+
+static uint64 KeyGen(const FIoBuffer& Data)
+{
+	return KeyGen(Data.GetData(), Data.GetSize());
+}
+
+////////////////////////////////////////////////////////////////////////////////
+IOSTOREONDEMAND_API void Tests()
+{
+	using namespace UE::IO::Private::JournaledCache;
+
+	// Some randomness
+	uint64 Th = 0x0'a9e0'493; // prime!
+	auto MixTh = [&] {
+		return Th *= 0x369dea0f31a53f85ull;
+	};
+
+	// Some data
+	uint32 WorkingSize = 1_Mi;
+	TUniquePtr<uint8[]> WorkingScope(new uint8[WorkingSize]);
+	uint8* Working = WorkingScope.Get();
+	for (uint32 i = 0; i < WorkingSize; i += 8)
+	{
+		*(uint64*)(Working + i) = MixTh();
+	}
+	auto DummyData = [&] (uint32 Size) {
+		uint32 Offset = MixTh() % (WorkingSize - Size);
+		return FIoBuffer(FIoBuffer::Wrap, Working + Offset, Size);
+	};
+
+	// Pending {{{2
+	{
+		struct {
+			int32 Size;
+			int32 Expected;
+		} TestCases[] = {
+			{ 0, 0 },
+			{ 10, 0 },
+			{ 1023, 511 },
+			{ 1024, 1024 },
+			{ 1025, 1024 },
+		};
+
+		for (auto& [Size, Expected] : TestCases)
+		{
+			FPending Pending(Size);
+
+			Pending.Put(0x493, DummyData(0));
+			check(Pending.GetCount() == 0);
+
+			Pending.Put(0x493, DummyData(513));
+			Pending.Put(0xa9e, DummyData(511));
+			check(Pending.GetUsed() == Expected);
+
+			Pending.Put(0x49e, DummyData(11));
+			Expected = (Expected == 0) ? 0 : (511 + 11);
+			check(Pending.GetUsed() == Expected);
+		}
+
+		FPending::PeelItems Peeled;
+
+		FPending Pending(64);
+		Pending.Put(1, DummyData(1));
+		check(Pending.Peel(0, Peeled) == 0);
+		check(Peeled.Num() == 0);
+		check(Pending.Peel(64, Peeled) == 1);
+		check(Peeled.Num() == 1);
+		check(Pending.GetUsed() == 0);
+		Peeled.Reset();
+
+		Pending = FPending(64);
+		for (int32 i = 0; i < 64; ++i)
+		{
+			Pending.Put(i + 1, DummyData(1));
+		}
+
+		check(Pending.Peel(32, Peeled) == 32);
+		check(Peeled.Num() == 32);
+		check(Pending.GetUsed() == 32);
+		for (auto& [Key, _] : Peeled)
+		{
+			FIoBuffer Data;
+			check(Pending.Get(Key) == 0);
+		}
+		Peeled.Reset();
+	} // }}}
+
+	// Cache {{{2
+	{
+		FCache::FConfig Config;
+		Config.Path = FPaths::ProjectPersistentDownloadDir() / TEXT("ias_cache_test");
+		Config.MemoryQuota = 512_Ki;
+		Config.DiskQuota = 8_Mi;
+		Config.JournalQuota = 7_Ki;
+		Config.JournalFlushInterval = 4;
+		Config.DropCache = true;
+		FCache Cache(MoveTemp(Config));
+
+		auto PrimePuts = [&] (int32 PutMax) {
+			TMap<uint64, FIoBuffer> Ret;
+			while (true)
+			{
+				uint32 Size = MixTh() & ((128_Ki) - 1);
+				if ((PutMax -= Size) < 0)
+				{
+					break;
+				}
+				FIoBuffer Data = DummyData(Size);
+				uint64 Key = KeyGen(Data);
+				Cache.Put(Key, Data);
+				Ret.Add(Key, Data);
+			}
+			return Ret;
+		};
+
+		uint32 WriteAllowance;
+
+		// no-op
+		WriteAllowance = 1_Ki;
+		Cache.Flush(WriteAllowance);
+		Cache.Flush(WriteAllowance);
+		Cache.Reset();
+
+		WriteAllowance = 512_Ki;
+		PrimePuts(WriteAllowance);
+		Cache.Flush(0);
+		Cache.Flush(WriteAllowance);
+		Cache.Flush(WriteAllowance);
+		Cache.Flush(WriteAllowance);
+		Cache.Reset();
+
+		auto Validate = [&] () {
+			struct FVisitState {
+				FCache& Cache;
+				const uint8* WorkRange[2];
+			};
+			auto Visitor = [] (void* Param, const FDebugCacheEntry& Entry) {
+				auto* State = (FVisitState*)Param;
+				
+				auto GetEntry = State->Cache.Get(Entry.Key);
+				check(GetEntry.IsHit());
+				FIoBuffer Data;
+				check(GetEntry.Materialize(Data));
+				check(Data.GetSize() == Entry.Size);
+				check(KeyGen(Data) == Entry.Key);
+
+				int32 IsFromDisk = 0;
+				IsFromDisk |= Data.GetData() >= State->WorkRange[1];
+				IsFromDisk |= (Data.GetData() + Data.GetSize()) <= State->WorkRange[0];
+				check(Entry.IsPending != IsFromDisk);
+			};
+			FVisitState State = {Cache, {Working, Working + WorkingSize}};
+			return Cache.DebugVisit(&State, Visitor);
+		};
+		check(Validate() == 0);
+
+		WriteAllowance = 512_Ki;
+
+		// simple
+		for (uint32 i : {1, 2, 4, 7, 11})
+		{
+			for (uint32 j = 0; j < i; ++j)
+			{
+				FIoBuffer Data = DummyData(32);
+				Cache.Put(KeyGen(Data), Data);
+			}
+			Cache.Flush(WriteAllowance);
+			check(Validate() == i);
+
+			Cache.Reset();
+			Cache.Load();
+			check(Validate() <= i);
+		}
+		Cache.Reset();
+
+		// general
+		for (int32 i : {1, 4, 136, 137})
+		{
+			for (int32 j = 0; j < i; ++j)
+			{
+				PrimePuts(WriteAllowance);
+				Cache.Flush(WriteAllowance);
+			}
+			uint32 Count = Validate();
+
+			Cache.Reset();
+			Cache.Load();
+			check(Validate() <= Count);
+		}
+		Cache.Reset();
+
+		// power 2
+		for (int32 i : {74, 75})
+		{
+			while (i--)
+			{
+				for (int32 j = 0; j < 3; ++j)
+				{
+					FIoBuffer Data = DummyData(64_Ki - 4);
+					Cache.Put(KeyGen(Data), Data);
+					Cache.Flush(WriteAllowance);
+				}
+				Cache.Flush(WriteAllowance);
+			}
+			Validate();
+			Cache.Reset();
+		}
+
+		// marker wrap
+		// one-phrase journal
+		// journal paragraphs that are all the same size
+		// phrases with no entries
+
+		// cache items larger than pending memory
+		// cache items larger than write allowance
+
+		// journal wrapping without truncation
+
+		// changes in max data/journal size
+
+		// don't load-and-sort so many paragraphs (only need max-data size)
+	} // }}}
+}
+
+} // namespace IasJournaledFileCacheTest
+
+#endif // IAS_JOURNALED_CACHE_TEST
+
+// }}}
+
+////////////////////////////////////////////////////////////////////////////////
+TUniquePtr<IIoCache> MakeJournaledCache(const FFileIoCacheConfig& Config)
+{
+	return MakeUnique<UE::IO::Private::FJournaledCache>(Config);
+}
+
+/* vim: set noet foldlevel=1 : */
