Packed Chaos::Private::FImplicitBVHObject to avoid 10% of padding waste in it

diff --git a/Engine/Source/Runtime/Experimental/Chaos/Private/Chaos/ImplicitObjectBVH.cpp b/Engine/Source/Runtime/Experimental/Chaos/Private/Chaos/ImplicitObjectBVH.cpp
--- a/Engine/Source/Runtime/Experimental/Chaos/Private/Chaos/ImplicitObjectBVH.cpp
+++ b/Engine/Source/Runtime/Experimental/Chaos/Private/Chaos/ImplicitObjectBVH.cpp
@@ -12,6 +12,8 @@ namespace Chaos
 
 	namespace Private
 	{
+		static_assert(sizeof(FImplicitBVHObject) == 72, "FImplicitBVHObject was packed to avoid any padding in it");
+
 		FImplicitBVHObject::FImplicitBVHObject()
 		{
 		}
@@ -23,11 +25,11 @@ namespace Chaos
 			const FAABB3& InBounds,
 			const int32 InRootObjectIndex,
 			const int32 InObjectIndex)
-			: R(FRotation3f(InR))
+			: R{ (float)InR.X, (float)InR.Y, (float)InR.Z, (float)InR.W }
 			, X(FVec3f(InX))
 			, Bounds(FAABB3f(InBounds))
-			, Geometry(InGeometry)
 			, RootObjectIndex(InRootObjectIndex)
+			, Geometry(InGeometry)
 			, ObjectIndex(InObjectIndex)
 		{
 		}
@@ -319,7 +321,7 @@ namespace Chaos
 		{
 			Ar << Geometry;
 			Ar << X;
-			Ar << R;
+			Ar << R[0]; Ar << R[1]; Ar << R[2]; Ar << R[3];
 			Ar << RootObjectIndex;
 			return Ar;
 		}
diff --git a/Engine/Source/Runtime/Experimental/Chaos/Public/Chaos/ImplicitObjectBVH.h b/Engine/Source/Runtime/Experimental/Chaos/Public/Chaos/ImplicitObjectBVH.h
--- a/Engine/Source/Runtime/Experimental/Chaos/Public/Chaos/ImplicitObjectBVH.h
+++ b/Engine/Source/Runtime/Experimental/Chaos/Public/Chaos/ImplicitObjectBVH.h
@@ -31,7 +31,7 @@ namespace Chaos
 
 			const FVec3f& GetX() const { return X; }
 
-			const FRotation3f& GetR() const { return R; }
+			const FRotation3f& GetR() const { return *(FRotation3f*)R; }
 
 			const FAABB3f& GetBounds() const { return Bounds; }
 
@@ -55,16 +55,16 @@ namespace Chaos
 			FChaosArchive& Serialize(FChaosArchive& Ar);
 
 			// Transform and bounds in the space of the BVH owner (Union Implicit)
-			FRotation3f R;
+			float R[4];	//an alias for FRotation3f to avoid 16b alignment requirement and avoid 12b of padding
 			FVec3f X;
 			FAABB3f Bounds;
 
-			// The leaf geometry stripped of decorators (but not Instanced or Scaled)
-			TSerializablePtr<FImplicitObject> Geometry;
-
 			// The index of our ancestor in the array of RootObjects that was provided when creating the BVH
 			int32 RootObjectIndex;
 
+			// The leaf geometry stripped of decorators (but not Instanced or Scaled)
+			TSerializablePtr<FImplicitObject> Geometry;
+
 			// Our index in the hierarchy. This could be used to uniquely identity copies of the same implicit in the hierarchy
 			int32 ObjectIndex;
 		};
