Pad the duration field when logging IAS requests to keep exceptionally long request (>1second) aligned with the other requests.

#rnx

diff --git a/Engine/Source/Runtime/Experimental/IoStoreOnDemand/Private/OnDemandIoDispatcherBackend.cpp b/Engine/Source/Runtime/Experimental/IoStoreOnDemand/Private/OnDemandIoDispatcherBackend.cpp
--- a/Engine/Source/Runtime/Experimental/IoStoreOnDemand/Private/OnDemandIoDispatcherBackend.cpp
+++ b/Engine/Source/Runtime/Experimental/IoStoreOnDemand/Private/OnDemandIoDispatcherBackend.cpp
@@ -71,7 +71,7 @@ namespace UE::IO::Private
 static void LogHttpResult(const TCHAR* Url, uint32 StatusCode, uint32 Duration, uint32 Size, uint32 Offset, const char* Memo="ok")
 {
 	Size >>= 10;
-	UE_LOG(LogIas, VeryVerbose, TEXT("http-%3u: %4ums %5uKiB [%7u] '%S' %s"), StatusCode, Duration, Size, Offset, Memo, Url);
+	UE_LOG(LogIas, VeryVerbose, TEXT("http-%3u: %5ums %5uKiB [%7u] '%S' %s"), StatusCode, Duration, Size, Offset, Memo, Url);
 };
 
 using namespace UE::Tasks;
@@ -915,7 +915,7 @@ static void LogIoResult(
 		}
 		return bCached ? TEXT("io-cache") : TEXT("io-http ");
 	}();
-	UE_LOG(LogIas, VeryVerbose, TEXT("%s: %4ums %5uKiB [%7u] %s:%s:%s|%u"),
+	UE_LOG(LogIas, VeryVerbose, TEXT("%s: %5ums %5uKiB [%7u] %s:%s:%s|%u"),
 		Prefix,
 		Duration,
 		UncompressedSize >> 10,
