Updating simulation check to include sleeping but physically simulated objects

diff --git a/Engine/Source/Runtime/Engine/Private/PrimitiveComponentPhysics.cpp b/Engine/Source/Runtime/Engine/Private/PrimitiveComponentPhysics.cpp
--- a/Engine/Source/Runtime/Engine/Private/PrimitiveComponentPhysics.cpp
+++ b/Engine/Source/Runtime/Engine/Private/PrimitiveComponentPhysics.cpp
@@ -964,7 +964,7 @@ bool UPrimitiveComponent::IsSimulatingPhysics(FName BoneName) const
 			}
 		);
 
-		return Interface->AreAllDynamic(PhysicsObjects);
+		return Interface->AreAllDynamicOrSleeping(PhysicsObjects);
 	}
 
 	return false;
diff --git a/Engine/Source/Runtime/Experimental/Chaos/Private/Chaos/PhysicsObjectInterface.cpp b/Engine/Source/Runtime/Experimental/Chaos/Private/Chaos/PhysicsObjectInterface.cpp
--- a/Engine/Source/Runtime/Experimental/Chaos/Private/Chaos/PhysicsObjectInterface.cpp
+++ b/Engine/Source/Runtime/Experimental/Chaos/Private/Chaos/PhysicsObjectInterface.cpp
@@ -528,6 +528,17 @@ namespace Chaos
 		return bCheck;
 	}
 
+	template<EThreadContext Id>
+	bool FReadPhysicsObjectInterface<Id>::AreAllDynamicOrSleeping(TArrayView<const FConstPhysicsObjectHandle> InObjects)
+	{
+		bool bCheck = !InObjects.IsEmpty();
+		for (const FConstPhysicsObjectHandle Object : InObjects)
+		{
+			bCheck &= (Object && Object->IsValid() && (Object->ObjectState<Id>() == EObjectStateType::Dynamic || Object->ObjectState<Id>() == EObjectStateType::Sleeping));
+		}
+		return bCheck;
+	}
+
 	template<EThreadContext Id>
 	bool FReadPhysicsObjectInterface<Id>::AreAllDisabled(TArrayView<const FConstPhysicsObjectHandle> InObjects)
 	{
diff --git a/Engine/Source/Runtime/Experimental/Chaos/Public/Chaos/PhysicsObjectInterface.h b/Engine/Source/Runtime/Experimental/Chaos/Public/Chaos/PhysicsObjectInterface.h
--- a/Engine/Source/Runtime/Experimental/Chaos/Public/Chaos/PhysicsObjectInterface.h
+++ b/Engine/Source/Runtime/Experimental/Chaos/Public/Chaos/PhysicsObjectInterface.h
@@ -95,6 +95,7 @@ namespace Chaos
 		CHAOS_API bool AreAllSleeping(TArrayView<const FConstPhysicsObjectHandle> InObjects);
 		CHAOS_API bool AreAllRigidBody(TArrayView<const FConstPhysicsObjectHandle> InObjects);
 		CHAOS_API bool AreAllDynamic(TArrayView<const FConstPhysicsObjectHandle> InObjects);
+		CHAOS_API bool AreAllDynamicOrSleeping(TArrayView<const FConstPhysicsObjectHandle> InObjects);
 		CHAOS_API bool AreAllDisabled(TArrayView<const FConstPhysicsObjectHandle> InObjects);
 		CHAOS_API bool AreAllShapesQueryEnabled(TArrayView<const FConstPhysicsObjectHandle> InObjects);
 		CHAOS_API float GetMass(TArrayView<const FConstPhysicsObjectHandle> InObjects);
