Added some trace scopes.

#rnx

diff --git a/Engine/Source/Runtime/Experimental/IoStoreOnDemand/Private/JournaledCache.cpp b/Engine/Source/Runtime/Experimental/IoStoreOnDemand/Private/JournaledCache.cpp
--- a/Engine/Source/Runtime/Experimental/IoStoreOnDemand/Private/JournaledCache.cpp
+++ b/Engine/Source/Runtime/Experimental/IoStoreOnDemand/Private/JournaledCache.cpp
@@ -19,6 +19,7 @@
 #include "Misc/Paths.h"
 #include "Misc/ScopeRWLock.h"
 #include "Misc/StringBuilder.h"
+#include "ProfilingDebugging/CpuProfilerTrace.h"
 #include "Tasks/Task.h"
 #include "Templates/UniquePtr.h"
 
@@ -142,6 +143,8 @@ EntryHandle FPending::Get(uint64 Key) const
 ////////////////////////////////////////////////////////////////////////////////
 bool FPending::Materialize(EntryHandle Handle, FIoBuffer& Out, uint32 Offset) const
 {
+	TRACE_CPUPROFILER_EVENT_SCOPE(IasCache::Materialize_Memory);
+
 	FItem& Item = *(FItem*)Handle;
 
 	if (Out.GetData() == nullptr)
@@ -418,6 +421,8 @@ void FJournal::ClosePhrase(FPhrase& Phrase)
 	IPlatformFile& Ipf = IPlatformFile::GetPlatformPhysical();
 	if (TUniquePtr<IFileHandle> File(Ipf.OpenWrite(*BinPath, true, false)); File.IsValid())
 	{
+		TRACE_CPUPROFILER_EVENT_SCOPE(IasCache::DataWrite);
+
 		File->Seek(DataCursor);
 		File->Write(Buffer, WriteSize);
 		File->Flush();
@@ -441,6 +446,8 @@ EntryHandle FJournal::Get(uint64 Key) const
 ////////////////////////////////////////////////////////////////////////////////
 bool FJournal::Materialize(EntryHandle Handle, FIoBuffer& Out, uint32 Offset) const
 {
+	TRACE_CPUPROFILER_EVENT_SCOPE(IasCache::Materialize_Disk);
+
 	const FMapEntry& Entry = *(FMapEntry*)Handle;
 
 	uint32 ReadSize = Entry.Size - Offset;
@@ -492,6 +499,8 @@ uint64 FJournal::Insert(const FPhraseDesc* Phrase)
 ////////////////////////////////////////////////////////////////////////////////
 void FJournal::Prune(uint64 DataBase, uint32 Size)
 {
+	TRACE_CPUPROFILER_EVENT_SCOPE(IasCache::Prune);
+
 	int32 BytesRemoved = OverRemoval;
 	if (BytesRemoved >= int32(Size))
 	{
@@ -540,6 +549,8 @@ int32 FJournal::Flush()
 		return 0;
 	}
 
+	TRACE_CPUPROFILER_EVENT_SCOPE(IasCache::Flush_Journal);
+
 	uint32 Size = Entries.Num() * sizeof(Entries[0]);
 
 	if (JournalCursor + Size > JournalSize)
@@ -554,6 +565,8 @@ int32 FJournal::Flush()
 	JrnPath << TEXT(".jrn");
 	if (TUniquePtr<IFileHandle> File(Ipf.OpenWrite(*JrnPath, true, false)); File.IsValid())
 	{
+		TRACE_CPUPROFILER_EVENT_SCOPE(IasCache::JournalWrite);
+
 		File->Seek(JournalCursor);
 		File->Write((uint8*)(Entries.GetData()), Size);
 		File->Flush();
@@ -979,6 +992,8 @@ int32 FCache::Flush(int32 Allowance)
 		return Allowance - Journal.Flush();
 	}
 
+	TRACE_CPUPROFILER_EVENT_SCOPE(IasCache::Flush_Pending);
+
 	FPending::PeelItems PeelItems;
 
 	FWriteAccess _(Lock);
@@ -1096,6 +1111,8 @@ void FJournaledCache::Update()
 		return;
 	}
 
+	TRACE_CPUPROFILER_EVENT_SCOPE(IasCache::Update);
+
 	int32 AllowanceUsed = Cache->Flush(WriteAllowance);
 
 	Governor.Return(WriteAllowance - AllowanceUsed);
