Add the request priority when logging a IAS Io request result.

#rnx

diff --git a/Engine/Source/Runtime/Experimental/IoStoreOnDemand/Private/OnDemandIoDispatcherBackend.cpp b/Engine/Source/Runtime/Experimental/IoStoreOnDemand/Private/OnDemandIoDispatcherBackend.cpp
--- a/Engine/Source/Runtime/Experimental/IoStoreOnDemand/Private/OnDemandIoDispatcherBackend.cpp
+++ b/Engine/Source/Runtime/Experimental/IoStoreOnDemand/Private/OnDemandIoDispatcherBackend.cpp
@@ -905,6 +905,7 @@ static void LogIoResult(
 	uint32 UncompressedSize,
 	uint32 UncompressedOffset,
 	uint32 CompressedOffset,
+	int32 Priority,
 	bool bCached)
 {
 	const TCHAR* Prefix = [bCached, UncompressedSize]() -> const TCHAR*
@@ -915,7 +916,7 @@ static void LogIoResult(
 		}
 		return bCached ? TEXT("io-cache") : TEXT("io-http ");
 	}();
-	UE_LOG(LogIas, VeryVerbose, TEXT("%s: %5ums %5uKiB [%7u] %s:%s:%s|%u"),
+	UE_LOG(LogIas, VeryVerbose, TEXT("%s: %5ums %5uKiB [%7u] %s:%s:%s|%u (%d)"),
 		Prefix,
 		Duration,
 		UncompressedSize >> 10,
@@ -923,7 +924,8 @@ static void LogIoResult(
 		*LexToString(ChunkId),
 		*LexToString(CacheKey),
 		*LexToString(UrlHash),
-		CompressedOffset);
+		CompressedOffset,
+		Priority);
 };
 
 ///////////////////////////////////////////////////////////////////////////////
@@ -1196,7 +1198,7 @@ void FOnDemandIoBackend::CompleteRequest(FChunkRequest* ChunkRequest)
 			Stats.OnIoRequestComplete(Request->GetBuffer().GetSize());
 			LogIoResult(Request->ChunkId, ChunkRequest->Params.GetUrlHash(), ChunkRequest->Params.ChunkKey, Duration,
 				Request->GetBuffer().DataSize(), Request->Options.GetOffset(),
-				ChunkRequest->Params.ChunkRange.GetOffset(), ChunkRequest->bCached);
+				ChunkRequest->Params.ChunkRange.GetOffset(), ChunkRequest->Priority, ChunkRequest->bCached);
 				
 		}
 		else
@@ -1207,7 +1209,7 @@ void FOnDemandIoBackend::CompleteRequest(FChunkRequest* ChunkRequest)
 			Stats.OnIoRequestFail();
 			LogIoResult(Request->ChunkId, ChunkRequest->Params.GetUrlHash(), ChunkRequest->Params.ChunkKey, Duration,
 				0, Request->Options.GetOffset(),
-				ChunkRequest->Params.ChunkRange.GetOffset(), ChunkRequest->bCached);
+				ChunkRequest->Params.ChunkRange.GetOffset(), ChunkRequest->Priority, ChunkRequest->bCached);
 		}
 
 		CompletedRequests.Enqueue(Request);
