Ias: Fix "RequestHead == nullptr" assert in FChunkRequest::RemoveDispatcherRequest (fallout from 26911030).

#rnx

diff --git a/Engine/Source/Runtime/Experimental/IoStoreOnDemand/Private/OnDemandIoDispatcherBackend.cpp b/Engine/Source/Runtime/Experimental/IoStoreOnDemand/Private/OnDemandIoDispatcherBackend.cpp
--- a/Engine/Source/Runtime/Experimental/IoStoreOnDemand/Private/OnDemandIoDispatcherBackend.cpp
+++ b/Engine/Source/Runtime/Experimental/IoStoreOnDemand/Private/OnDemandIoDispatcherBackend.cpp
@@ -828,6 +828,12 @@ struct FChunkRequest
 				It = It->NextRequest;
 			}
 
+			if (RequestTail == It)
+			{
+				check(It->NextRequest == nullptr);
+				RequestTail = Request;
+			}
+
 			Request->NextRequest = It->NextRequest;
 			It->NextRequest = Request;
 		}
