Placate MSVC analyzer

#rnx

diff --git a/Engine/Source/Runtime/Experimental/IoStoreOnDemand/Private/JournaledCache.cpp b/Engine/Source/Runtime/Experimental/IoStoreOnDemand/Private/JournaledCache.cpp
--- a/Engine/Source/Runtime/Experimental/IoStoreOnDemand/Private/JournaledCache.cpp
+++ b/Engine/Source/Runtime/Experimental/IoStoreOnDemand/Private/JournaledCache.cpp
@@ -1197,10 +1197,11 @@ FIoStatus FJournaledCache::Put(const FIoHash& Key, FIoBuffer& Data)
 uint64 FJournaledCache::ReduceKey(const FIoHash& Key)
 {
 	const uint8* Bytes = Key.GetBytes();
-	uint64 Ret[2];
-	std::memcpy(Ret + 0, Bytes + 0, sizeof(uint64));
-	std::memcpy(Ret + 1, Bytes + 8, sizeof(uint64));
-	return Ret[0] /*bad...*/ ^ /*...mixer*/ Ret[1];
+	uint64 Ret[3] = {};
+	std::memcpy(Ret + 0, Bytes +  0, sizeof(uint64));
+	std::memcpy(Ret + 1, Bytes +  8, sizeof(uint64));
+	std::memcpy(Ret + 2, Bytes + 16, sizeof(uint32));
+	return (Ret[0] + Ret[2]) ^ Ret[1];
 }
 
 // }}}
