Make module animation a seperate function and don't animate on dedicated server.

diff --git a/Engine/Source/Runtime/Experimental/ChaosVehicles/ChaosVehiclesCore/Private/SimModule/AerofoilModule.cpp b/Engine/Source/Runtime/Experimental/ChaosVehicles/ChaosVehiclesCore/Private/SimModule/AerofoilModule.cpp
--- a/Engine/Source/Runtime/Experimental/ChaosVehicles/ChaosVehiclesCore/Private/SimModule/AerofoilModule.cpp
+++ b/Engine/Source/Runtime/Experimental/ChaosVehicles/ChaosVehiclesCore/Private/SimModule/AerofoilModule.cpp
@@ -31,7 +31,7 @@ namespace Chaos
 		}
 
 		// needs local velocity at aerofoil to calculate lift and drag
-		FVector LocalForce = GetForce(ModuleLocalVelocity * CmToMScaling(), CmToM(Altitude), DeltaTime);
+		FVector LocalForce = GetForce(LocalLinearVelocity * CmToMScaling(), CmToM(Altitude), DeltaTime);
 		FVector LocalOffset = GetCenterOfLiftOffset() * MToCmScaling();
 
 		AddLocalForceAtPosition(LocalForce * MToCmScaling(), LocalOffset, true, false, false, FColor::Yellow);
@@ -120,7 +120,7 @@ namespace Chaos
 		AirflowNormal.Normalize();
 
 		// determine angle of attack for control surface
-		AngleOfAttack = CalcAngleOfAttackDegrees(Setup().UpAxis, AirflowNormal);
+		AngleOfAttack = CalcAngleOfAttackDegrees(Setup().ForceAxis, AirflowNormal);
 
 		// Aerofoil Camber and Control Surface just lumped together
 		float TotalControlAngle = ControlSurfaceAngle + Setup().Camber;
@@ -133,11 +133,23 @@ namespace Chaos
 		float DragCoef = CalcDragCoefficient(AngleOfAttack, TotalControlAngle) * Setup().DragMultiplier;
 
 		// Combine to create a single force vector
-		Force = Setup().UpAxis * (Common * LiftCoef) + AirflowNormal * (Common * DragCoef);
+		Force = Setup().ForceAxis * (Common * LiftCoef) + AirflowNormal * (Common * DragCoef);
 
 		return Force;
 	}
 
+	void FAerofoilSimModule::Animate(Chaos::FClusterUnionPhysicsProxy* Proxy)
+	{
+		if (FPBDRigidClusteredParticleHandle* ClusterChild = GetClusterParticle(Proxy))
+		{
+			FQuat ControlRotation = FQuat(Setup().ControlRotationAxis, FMath::DegreesToRadians(ControlSurfaceAngle));
+
+			FTransform InitialTransform = GetInitialParticleTransform();
+			ClusterChild->ChildToParent().SetRotation(InitialTransform.GetRotation() * ControlRotation);
+		}
+
+	}
+
 } // namespace Chaos
 
 #if VEHICLE_DEBUGGING_ENABLED
diff --git a/Engine/Source/Runtime/Experimental/ChaosVehicles/ChaosVehiclesCore/Private/SimModule/ChassisModule.cpp b/Engine/Source/Runtime/Experimental/ChaosVehicles/ChaosVehiclesCore/Private/SimModule/ChassisModule.cpp
--- a/Engine/Source/Runtime/Experimental/ChaosVehicles/ChaosVehiclesCore/Private/SimModule/ChassisModule.cpp
+++ b/Engine/Source/Runtime/Experimental/ChaosVehicles/ChaosVehiclesCore/Private/SimModule/ChassisModule.cpp
@@ -3,6 +3,9 @@
 #include "SimModule/ChassisModule.h"
 #include "SimModule/SimModuleTree.h"
 
+#if VEHICLE_DEBUGGING_ENABLED
+UE_DISABLE_OPTIMIZATION
+#endif
 
 namespace Chaos
 {
@@ -15,8 +18,21 @@ namespace Chaos
 
 	void FChassisSimModule::Simulate(float DeltaTime, const FAllInputs& Inputs, FSimModuleTree& VehicleModuleSystem)
 	{
-
+		FVector VelocityMs = Chaos::CmToM(LocalLinearVelocity);
+		float DragForceMagnitude = -VelocityMs.SizeSquared() * Setup().DensityOfMedium * 0.5f * Setup().AreaMetresSquared * Setup().DragCoefficient * Chaos::MToCmScaling();
+		FVector ForceVector = LocalLinearVelocity.GetSafeNormal() * DragForceMagnitude;
+		ForceVector.X *= Setup().XAxisMultiplier;
+		ForceVector.Y *= Setup().YAxisMultiplier;
+		AddLocalForce(ForceVector);
+
+		float Value = Setup().AngularDamping * Chaos::MToCmScaling();
+		FVector Torque = LocalAngularVelocity * -Value;
+		AddLocalTorque(Torque);
 	}
 
 
 } // namespace Chaos
+
+#if VEHICLE_DEBUGGING_ENABLED
+UE_ENABLE_OPTIMIZATION
+#endif
diff --git a/Engine/Source/Runtime/Experimental/ChaosVehicles/ChaosVehiclesCore/Private/SimModule/DeferredForcesModular.cpp b/Engine/Source/Runtime/Experimental/ChaosVehicles/ChaosVehiclesCore/Private/SimModule/DeferredForcesModular.cpp
--- a/Engine/Source/Runtime/Experimental/ChaosVehicles/ChaosVehiclesCore/Private/SimModule/DeferredForcesModular.cpp
+++ b/Engine/Source/Runtime/Experimental/ChaosVehicles/ChaosVehiclesCore/Private/SimModule/DeferredForcesModular.cpp
@@ -23,6 +23,7 @@ FAutoConsoleVariableRef CVarChaosModularVehiclesDisableForces(TEXT("p.ModularVeh
 Chaos::FPBDRigidParticleHandle* FDeferredForcesModular::GetParticle(TArray<Chaos::FPBDRigidClusteredParticleHandle*>& Particles
 		, TArray<Chaos::FPBDRigidClusteredParticleHandle*>& ClusterParticles
 		, int TransformIndex
+		, const FVector& PositionalOffset
 		, const TManagedArray<FTransform>& Transforms
 		, const TManagedArray<FTransform>& CollectionMassToLocal
 		, const TManagedArray<int32>& Parent
@@ -35,11 +36,12 @@ Chaos::FPBDRigidParticleHandle* FDeferredForcesModular::GetParticle(TArray<Chaos
 	if (SimParticle && !SimParticle->Disabled())
 	{
 		TransformOut = FTransform::Identity;
+		TransformOut.AddToTranslation(PositionalOffset);
 		return SimParticle;		// apply to broken off child
 	}
 	else if (SimClusterParticle && !SimClusterParticle->Disabled())
 	{
-		TransformOut.SetLocation(CollectionMassToLocal[Parent[TransformIndex]].InverseTransformPosition(Transforms[TransformIndex].GetLocation()));
+		TransformOut.SetLocation(CollectionMassToLocal[Parent[TransformIndex]].InverseTransformPosition(Transforms[TransformIndex].GetLocation() + PositionalOffset));
 		TransformOut.SetRotation(Transforms[TransformIndex].GetRotation());
 
 		return SimClusterParticle; // apply to intact cluster
@@ -52,6 +54,7 @@ Chaos::FPBDRigidParticleHandle* FDeferredForcesModular::GetParticle(const FTrans
 	, TArray<Chaos::FPBDRigidParticleHandle*>& Particles
 	, TArray<Chaos::FPBDRigidClusteredParticleHandle*>& ClusterParticles
 	, int TransformIndex
+	, const FVector& PositionalOffset
 	, FTransform& TransformOut)
 {
 	using namespace Chaos;
@@ -61,7 +64,6 @@ Chaos::FPBDRigidParticleHandle* FDeferredForcesModular::GetParticle(const FTrans
 	TransformOut = FTransform::Identity;
 
 	const FRigidTransform3 ClusterWorldTM(ClusterParticles[SingleChassisIndex]->X(), ClusterParticles[SingleChassisIndex]->R());
-	//const FRigidTransform3 ClusterWorldTM(ClusterParticles[SingleChassisIndex]->GetTransformXRCom());
 	FRigidTransform3 Frame = FRigidTransform3::Identity;
 
 	//UE_LOG(LogInit, Warning, TEXT("TransformIndex %d  vs  Particles.Num() %d  vs ClusterParticles.Num() %d"), TransformIndex, Particles.Num(), ClusterParticles.Num());
@@ -79,6 +81,7 @@ Chaos::FPBDRigidParticleHandle* FDeferredForcesModular::GetParticle(const FTrans
 			//const FRigidTransform3 ChildWorldTM(ChildT.GetLocation(), ChildT.GetRotation());
 			const FRigidTransform3 ChildWorldTM(Child->X(), Child->R());
 			Frame = ChildWorldTM.GetRelativeTransform(ClusterWorldTM);
+			Frame.SetLocation(Frame.GetLocation() + OffsetTransform.TransformVector(PositionalOffset));
 
 			//UE_LOG(LogChaos, Warning, TEXT("ApplyForce Data TIdx %d %s"), TransformIndex, *Frame.GetTranslation().ToString());
 		}
@@ -100,7 +103,7 @@ void FDeferredForcesModular::Apply(TArray<Chaos::FPBDRigidClusteredParticleHandl
 		FTransform RelativeTransform;
 		for (const FApplyForceData& Data : ApplyForceDatas)
 		{
-			Chaos::FPBDRigidParticleHandle* RigidHandle = GetParticle(Particles, ClusterParticles, Data.TransformIndex, Transforms, CollectionMassToLocal, Parent, RelativeTransform);
+			Chaos::FPBDRigidParticleHandle* RigidHandle = GetParticle(Particles, ClusterParticles, Data.TransformIndex, FVector::ZeroVector, Transforms, CollectionMassToLocal, Parent, RelativeTransform);
 			if (RigidHandle)
 			{
 				AddForce(RigidHandle, Data, RelativeTransform);
@@ -109,7 +112,7 @@ void FDeferredForcesModular::Apply(TArray<Chaos::FPBDRigidClusteredParticleHandl
 
 		for (const FApplyForceAtPositionData& Data : ApplyForceAtPositionDatas)
 		{
-			Chaos::FPBDRigidParticleHandle* RigidHandle = GetParticle(Particles, ClusterParticles, Data.TransformIndex, Transforms, CollectionMassToLocal, Parent, RelativeTransform);
+			Chaos::FPBDRigidParticleHandle* RigidHandle = GetParticle(Particles, ClusterParticles, Data.TransformIndex, Data.Position, Transforms, CollectionMassToLocal, Parent, RelativeTransform);
 			if (RigidHandle)
 			{
 				AddForceAtPosition(RigidHandle, Data, RelativeTransform);
@@ -118,7 +121,7 @@ void FDeferredForcesModular::Apply(TArray<Chaos::FPBDRigidClusteredParticleHandl
 
 		for (const FAddTorqueInRadiansData& Data : ApplyTorqueDatas)
 		{
-			Chaos::FPBDRigidParticleHandle* RigidHandle = GetParticle(Particles, ClusterParticles, Data.TransformIndex, Transforms, CollectionMassToLocal, Parent, RelativeTransform);
+			Chaos::FPBDRigidParticleHandle* RigidHandle = GetParticle(Particles, ClusterParticles, Data.TransformIndex, FVector::ZeroVector, Transforms, CollectionMassToLocal, Parent, RelativeTransform);
 			if (RigidHandle)
 			{
 				AddTorque(RigidHandle, Data, RelativeTransform);
@@ -141,7 +144,7 @@ void FDeferredForcesModular::Apply(TArray<Chaos::FPBDRigidParticleHandle*>& Part
 		FTransform RelativeTransform;
 		for (const FApplyForceData& Data : ApplyForceDatas)
 		{
-			Chaos::FPBDRigidParticleHandle* RigidHandle = GetParticle(Data.OffsetTransform, Particles, ClusterParticles, Data.TransformIndex, RelativeTransform);
+			Chaos::FPBDRigidParticleHandle* RigidHandle = GetParticle(Data.OffsetTransform, Particles, ClusterParticles, Data.TransformIndex, FVector::ZeroVector, RelativeTransform);
 			if (RigidHandle)
 			{
 				AddForce(RigidHandle, Data, RelativeTransform);
@@ -150,7 +153,7 @@ void FDeferredForcesModular::Apply(TArray<Chaos::FPBDRigidParticleHandle*>& Part
 
 		for (const FApplyForceAtPositionData& Data : ApplyForceAtPositionDatas)
 		{
-			Chaos::FPBDRigidParticleHandle* RigidHandle = GetParticle(Data.OffsetTransform, Particles, ClusterParticles, Data.TransformIndex, RelativeTransform);
+			Chaos::FPBDRigidParticleHandle* RigidHandle = GetParticle(Data.OffsetTransform, Particles, ClusterParticles, Data.TransformIndex, Data.Position, RelativeTransform);
 			if (RigidHandle)
 			{
 				AddForceAtPosition(RigidHandle, Data, RelativeTransform);
@@ -159,7 +162,7 @@ void FDeferredForcesModular::Apply(TArray<Chaos::FPBDRigidParticleHandle*>& Part
 
 		for (const FAddTorqueInRadiansData& Data : ApplyTorqueDatas)
 		{
-			Chaos::FPBDRigidParticleHandle* RigidHandle = GetParticle(Data.OffsetTransform, Particles, ClusterParticles, Data.TransformIndex, RelativeTransform);
+			Chaos::FPBDRigidParticleHandle* RigidHandle = GetParticle(Data.OffsetTransform, Particles, ClusterParticles, Data.TransformIndex, FVector::ZeroVector, RelativeTransform);
 			if (RigidHandle)
 			{
 				AddTorque(RigidHandle, Data, RelativeTransform);
@@ -194,7 +197,7 @@ void AddForce_Implementation(Chaos::FPBDRigidParticleHandle* RigidHandle, const
 	Chaos::FVec3 Arm = Position - WorldCOM;
 
 	Chaos::FVec3 WorldTorque = Chaos::FVec3::CrossProduct(Arm, Force);
-	auto* RP = RigidHandle->CastToRigidParticle();
+	Chaos::FPBDRigidParticleHandle *RP = RigidHandle->CastToRigidParticle();
 	if (RP)
 	{
 		RP->AddForce(Force);
@@ -223,6 +226,21 @@ void AddForce_Implementation(Chaos::FPBDRigidParticleHandle* RigidHandle, const
 
 }
 
+void AddTorque_Implementation(Chaos::FPBDRigidParticleHandle* RigidHandle, const FTransform& ParticleOffsetTransform, const FVector& LocalTorque, const FTransform& OffsetTransform, const FColor& DebugColor)
+{
+	Chaos::FPBDRigidParticleHandle* RP = RigidHandle->CastToRigidParticle();
+	if (RP)
+	{
+		// local to world
+		const FTransform WorldTM(RigidHandle->R(), RigidHandle->X());
+		FVector T1 = OffsetTransform.TransformVector(LocalTorque);
+		FVector T2 = ParticleOffsetTransform.TransformVector(T1);
+		Chaos::FVec3 WorldTorque = WorldTM.TransformVector(T2);
+		RP->AddTorque(WorldTorque);
+	}
+
+}
+
 void FDeferredForcesModular::AddForce(Chaos::FPBDRigidParticleHandle* RigidHandle, const FApplyForceData& DataIn, const FTransform& OffsetTransform)
 {
 	if (ensure(RigidHandle))
@@ -232,55 +250,15 @@ void FDeferredForcesModular::AddForce(Chaos::FPBDRigidParticleHandle* RigidHandl
 
 }
 
-//void FDeferredForcesModular::AddForce(Chaos::FPBDRigidClusteredParticleHandle* RigidHandle, const FApplyForceData& DataIn)
-//{
-//	check(false);
-//	if (ensure(RigidHandle))
-//	{
-//		Chaos::EObjectStateType ObjectState = RigidHandle->ObjectState();
-//		if (CHAOS_ENSURE(ObjectState == Chaos::EObjectStateType::Dynamic || ObjectState == Chaos::EObjectStateType::Sleeping))
-//		{
-//			if (DataIn.Flags == EForceFlags::AccelChange)
-//			{
-//				const float RigidMass = RigidHandle->M();
-//				const Chaos::FVec3 Acceleration = DataIn.Force * RigidMass;
-//				RigidHandle->AddForce(Acceleration, false);
-//			}
-//			else
-//			{
-//				RigidHandle->AddForce(DataIn.Force, false);
-//			}
-//
-//		}
-//	}
-//}
-
 void FDeferredForcesModular::AddForceAtPosition(Chaos::FPBDRigidParticleHandle* RigidHandle, const FApplyForceAtPositionData& DataIn, const FTransform& OffsetTransform)
 {
 	if (ensure(RigidHandle))
 	{
-		FTransform ShiftedTransform = OffsetTransform;
-		ShiftedTransform.SetLocation(ShiftedTransform.GetLocation() + DataIn.Position);
-		AddForce_Implementation(RigidHandle, ParticleOffsetTransform, DataIn.Force, ShiftedTransform, (DataIn.Flags & EForceFlags::LevelSlope) == EForceFlags::LevelSlope, DataIn.DebugColor);
+		AddForce_Implementation(RigidHandle, DataIn.OffsetTransform, DataIn.Force, OffsetTransform, (DataIn.Flags & EForceFlags::LevelSlope) == EForceFlags::LevelSlope, DataIn.DebugColor);
 	}
 }
 
 void FDeferredForcesModular::AddTorque(Chaos::FPBDRigidParticleHandle* RigidHandle, const FAddTorqueInRadiansData& DataIn, const FTransform& OffsetTransform)
 {	
-	check(false); //#TODO - incorporate OffsetTransform as this is currently wrong
-	//if (ensure(RigidHandle))
-	//{
-	//	Chaos::EObjectStateType ObjectState = RigidHandle->ObjectState();
-	//	if (CHAOS_ENSURE(ObjectState == Chaos::EObjectStateType::Dynamic || ObjectState == Chaos::EObjectStateType::Sleeping))
-	//	{
-	//		if (DataIn.Flags == EForceFlags::AccelChange)
-	//		{
-	//			RigidHandle->AddTorque(Chaos::FParticleUtilitiesXR::GetWorldInertia(RigidHandle) * DataIn.Torque, false);
-	//		}
-	//		else
-	//		{
-	//			RigidHandle->AddTorque(DataIn.Torque, false);
-	//		}
-	//	}
-	//}
+	AddTorque_Implementation(RigidHandle, DataIn.OffsetTransform, DataIn.Torque, OffsetTransform, DataIn.DebugColor);
 }
diff --git a/Engine/Source/Runtime/Experimental/ChaosVehicles/ChaosVehiclesCore/Private/SimModule/SimModuleTree.cpp b/Engine/Source/Runtime/Experimental/ChaosVehicles/ChaosVehiclesCore/Private/SimModule/SimModuleTree.cpp
--- a/Engine/Source/Runtime/Experimental/ChaosVehicles/ChaosVehiclesCore/Private/SimModule/SimModuleTree.cpp
+++ b/Engine/Source/Runtime/Experimental/ChaosVehicles/ChaosVehiclesCore/Private/SimModule/SimModuleTree.cpp
@@ -243,7 +243,7 @@ void FSimModuleTree::Simulate(float DeltaTime, FAllInputs& Inputs, FClusterUnion
 
 	for (int RootIndex : RootNodes)
 	{
-		SimulateNode(DeltaTime, Inputs, RootIndex);
+		SimulateNode(DeltaTime, Inputs, RootIndex, PhysicsProxy);
 	}
 }
 
@@ -260,23 +260,28 @@ void FSimModuleTree::Simulate(float DeltaTime, FAllInputs& Inputs, FGeometryColl
 
 	for (int RootIndex : RootNodes)
 	{
-		SimulateNode(DeltaTime, Inputs, RootIndex);
+		SimulateNode(DeltaTime, Inputs, RootIndex, nullptr);
 	}
 
 }
 
-void FSimModuleTree::SimulateNode(float DeltaTime, FAllInputs& Inputs, int NodeIndex)
+void FSimModuleTree::SimulateNode(float DeltaTime, FAllInputs& Inputs, int NodeIndex, FClusterUnionPhysicsProxy* PhysicsProxy)
 {
 	if (ISimulationModuleBase* Module = AccessSimModule(NodeIndex))
 	{
 		if (Module->IsEnabled())
 		{
 			Module->Simulate(DeltaTime, Inputs, *this);
+
+			if (IsAnimationEnabled() && Module->IsAnimationEnabled())
+			{
+				Module->Animate(PhysicsProxy);
+			}
 		}
 
 		for (int ChildIdx : GetChildren(NodeIndex))
 		{
-			SimulateNode(DeltaTime, Inputs, ChildIdx);
+			SimulateNode(DeltaTime, Inputs, ChildIdx, PhysicsProxy);
 		}
 	}
 }
@@ -350,10 +355,14 @@ void FSimModuleTree::UpdateModuleVelocites(FGeometryCollectionPhysicsProxy* Phys
 						const Chaos::FVec3 Arm = WorldLocation - Particle->X();
 
 						FVector WorldVelocity = Particle->V() - Chaos::FVec3::CrossProduct(Arm, Particle->W());
-						FVector LocalVelocity = BodyTransform.InverseTransformVector(WorldVelocity);
-						LocalVelocity = Module->GetClusteredTransform().InverseTransformVector(LocalVelocity);
+						FVector LocalLinearVelocity = BodyTransform.InverseTransformVector(WorldVelocity);
+						LocalLinearVelocity = Module->GetClusteredTransform().InverseTransformVector(LocalLinearVelocity);
+
+						FVector LocalAngular = BodyTransform.InverseTransformVector(Particle->W());
+						LocalAngular = Module->GetClusteredTransform().InverseTransformVector(LocalAngular);
 
-						Module->SetLocalVelocity(LocalVelocity);
+						Module->SetLocalLinearVelocity(LocalLinearVelocity);
+						Module->SetLocalAngularVelocity(LocalAngular);
 					}
 			
 				}
@@ -399,7 +408,12 @@ void FSimModuleTree::UpdateModuleVelocites(FClusterUnionPhysicsProxy* PhysicsPro
 
 						FVector WorldVelocity = Particle->V() - Chaos::FVec3::CrossProduct(Arm, Particle->W());
 						FVector LocalVelocity = OffsetTransform.InverseTransformVector(BodyTransform.InverseTransformVector(WorldVelocity));
-						Module->SetLocalVelocity(LocalVelocity);
+
+						FVector LocalAngular = BodyTransform.InverseTransformVector(Particle->W());
+						LocalAngular = Module->GetClusteredTransform().InverseTransformVector(LocalAngular);
+
+						Module->SetLocalLinearVelocity(LocalVelocity);
+						Module->SetLocalAngularVelocity(LocalAngular);
 
 						//Chaos::FDebugDrawQueue::GetInstance().DrawDebugLine(WorldLocation, WorldLocation + WorldVelocity, FColor::White, false, -1.f, 0, 5.f);
 
@@ -447,7 +461,7 @@ void FSimModuleTree::SetSimState(const Chaos::FModuleNetDataArray& ModuleDatas)
 {
 	for (const TSharedPtr<FModuleNetData>& DataElement : ModuleDatas)
 	{
-		if (!SimulationModuleTree.IsEmpty() && SimulationModuleTree[DataElement->SimArrayIndex].SimModule)
+		if (!SimulationModuleTree.IsEmpty() && DataElement->SimArrayIndex < SimulationModuleTree.Num() && SimulationModuleTree[DataElement->SimArrayIndex].SimModule)
 		{
 			DataElement->FillSimState(SimulationModuleTree[DataElement->SimArrayIndex].SimModule);
 		}
@@ -466,3 +480,4 @@ void FSimModuleTree::InterpolateState(const float LerpFactor, Chaos::FModuleNetD
 #if VEHICLE_DEBUGGING_ENABLED
 UE_ENABLE_OPTIMIZATION
 #endif
+
diff --git a/Engine/Source/Runtime/Experimental/ChaosVehicles/ChaosVehiclesCore/Private/SimModule/SimulationModuleBase.cpp b/Engine/Source/Runtime/Experimental/ChaosVehicles/ChaosVehiclesCore/Private/SimModule/SimulationModuleBase.cpp
--- a/Engine/Source/Runtime/Experimental/ChaosVehicles/ChaosVehiclesCore/Private/SimModule/SimulationModuleBase.cpp
+++ b/Engine/Source/Runtime/Experimental/ChaosVehicles/ChaosVehiclesCore/Private/SimModule/SimulationModuleBase.cpp
@@ -3,6 +3,7 @@
 #include "SimModule/SimulationModuleBase.h"
 #include "SimModule/SimModuleTree.h"
 #include "SimModule/DeferredForcesModular.h"
+#include "PhysicsProxy/ClusterUnionPhysicsProxy.h"
 
 DEFINE_LOG_CATEGORY(LogSimulationModule);
 
@@ -23,14 +24,18 @@ void ISimulationModuleBase::AddLocalForce(const FVector& Force, bool bAllowSubst
 	AppliedForce = Force;
 	if (SimModuleTree)
 	{
-		//FString DebugString;
-		//GetDebugString(DebugString);
-		//UE_LOG(LogInit, Warning, TEXT("AddLocalForce To %s"), *DebugString);
-
 		SimModuleTree->AccessDeferredForces().Add(FDeferredForcesModular::FApplyForceData(ComponentTransform, TransformIndex, Force, bAllowSubstepping, bIsLocalForce, bLevelSlope, DebugColorIn));
 	}
 }
 
+void ISimulationModuleBase::AddLocalTorque(const FVector& Torque, bool bAllowSubstepping, bool bAccelChangeIn, const FColor& DebugColorIn)
+{
+	if (SimModuleTree)
+	{
+		SimModuleTree->AccessDeferredForces().Add(FDeferredForcesModular::FAddTorqueInRadiansData(ComponentTransform, TransformIndex, Torque, bAllowSubstepping, bAccelChangeIn, DebugColorIn));
+	}
+}
+
 ISimulationModuleBase* ISimulationModuleBase::GetParent()
 {
 	return (SimModuleTree != nullptr) ? SimModuleTree->AccessSimModule(SimModuleTree->GetParent(SimTreeIndex)) : nullptr;
@@ -49,6 +54,29 @@ ISimulationModuleBase* ISimulationModuleBase::GetFirstChild()
 	return nullptr;
 }
 
+FPBDRigidClusteredParticleHandle* ISimulationModuleBase::GetClusterParticle(Chaos::FClusterUnionPhysicsProxy* Proxy)
+{ 
+	// TODO: should store what we need rather than search for it all the time
+	FPBDRigidClusteredParticleHandle* ClusterChild = nullptr;
+
+	FPBDRigidsEvolutionGBF& Evolution = *static_cast<FPBDRigidsSolver*>(Proxy->GetSolver<FPBDRigidsSolver>())->GetEvolution();
+	FClusterUnionManager& ClusterUnionManager = Evolution.GetRigidClustering().GetClusterUnionManager();
+	const FClusterUnionIndex& CUI = Proxy->GetClusterUnionIndex();
+
+	if (FClusterUnion* ClusterUnion = ClusterUnionManager.FindClusterUnion(CUI))
+	{
+		FPBDRigidClusteredParticleHandle* ClusterHandle = ClusterUnion->InternalCluster;
+		TArray<FPBDRigidParticleHandle*> Particles = ClusterUnion->ChildParticles;
+
+		if (TransformIndex < Particles.Num())
+		{
+			ClusterChild = Particles[TransformIndex]->CastToClustered();
+		}
+	}
+
+	return ClusterChild;
+}
+
 bool ISimulationModuleBase::GetDebugString(FString& StringOut) const 
 {
 	StringOut += FString::Format(TEXT("{0}: TreeIndex {1}, Enabled {2}, InCluster {3}, TFormIdx {4}, ")
diff --git a/Engine/Source/Runtime/Experimental/ChaosVehicles/ChaosVehiclesCore/Private/SimModule/SuspensionModule.cpp b/Engine/Source/Runtime/Experimental/ChaosVehicles/ChaosVehiclesCore/Private/SimModule/SuspensionModule.cpp
--- a/Engine/Source/Runtime/Experimental/ChaosVehicles/ChaosVehiclesCore/Private/SimModule/SuspensionModule.cpp
+++ b/Engine/Source/Runtime/Experimental/ChaosVehicles/ChaosVehiclesCore/Private/SimModule/SuspensionModule.cpp
@@ -66,7 +66,7 @@ namespace Chaos
 		}
 
 		// tell wheels how much they are being pressed into the ground
-		if (SimModuleTree)
+		if (SimModuleTree && WheelSimTreeIndex != INVALID_IDX)
 		{
 			if (Chaos::ISimulationModuleBase* Module = SimModuleTree->AccessSimModule(WheelSimTreeIndex))
 			{
@@ -75,8 +75,22 @@ namespace Chaos
 
 				Wheel->SetForceIntoSurface(ForceIntoSurface);
 			}
-			
+		
+		}
+	}
+
+	void FSuspensionSimModule::Animate(Chaos::FClusterUnionPhysicsProxy* Proxy)
+	{
+		if (FPBDRigidClusteredParticleHandle* ClusterChild = GetClusterParticle(Proxy))
+		{
+			float CurrentSpringLength = GetSpringLength();
+
+			FVector RestPos = GetInitialParticleTransform().GetTranslation();
 
+			FVector Movement = Setup().SuspensionAxis * (Setup().MaxRaise + CurrentSpringLength);
+			Movement = GetComponentTransform().TransformVector(Movement);
+			FVector NewPos = RestPos - Movement;
+			ClusterChild->ChildToParent().SetTranslation(NewPos); // local frame for module
 		}
 	}
 
diff --git a/Engine/Source/Runtime/Experimental/ChaosVehicles/ChaosVehiclesCore/Private/SimModule/ThrusterModule.cpp b/Engine/Source/Runtime/Experimental/ChaosVehicles/ChaosVehiclesCore/Private/SimModule/ThrusterModule.cpp
--- a/Engine/Source/Runtime/Experimental/ChaosVehicles/ChaosVehiclesCore/Private/SimModule/ThrusterModule.cpp
+++ b/Engine/Source/Runtime/Experimental/ChaosVehicles/ChaosVehiclesCore/Private/SimModule/ThrusterModule.cpp
@@ -19,9 +19,26 @@ namespace Chaos
 
 	void FThrusterSimModule::Simulate(float DeltaTime, const FAllInputs& Inputs, FSimModuleTree& VehicleModuleSystem)
 	{
+		SteerAngleDegrees = 0.0f;
+		if (Setup().SteeringEnabled)
+		{
+			SteerAngleDegrees = Setup().SteeringEnabled ? Inputs.ControlInputs.Steering * Setup().MaxSteeringAngle : 0.0f;
+		}
+
 		// applies continuous force
 		FVector Force = Setup().ForceAxis * Setup().MaxThrustForce * Inputs.ControlInputs.Throttle;
-		AddLocalForceAtPosition(Force, Setup().ForceOffset, true, false, false, FColor::Magenta);
+		FQuat Steer = FQuat(Setup().SteeringAxis, FMath::DegreesToRadians(SteerAngleDegrees) * Setup().SteeringForceEffect);
+		AddLocalForceAtPosition(Steer.RotateVector(Force), Setup().ForceOffset, true, false, false, FColor::Magenta);
+	}
+
+	void FThrusterSimModule::Animate(Chaos::FClusterUnionPhysicsProxy* Proxy)
+	{
+		if (FPBDRigidClusteredParticleHandle* ClusterChild = GetClusterParticle(Proxy))
+		{
+			FQuat Steer = FQuat(Setup().SteeringAxis, FMath::DegreesToRadians(SteerAngleDegrees));
+
+			ClusterChild->ChildToParent().SetRotation(GetInitialParticleTransform().GetRotation() * Steer);
+		}
 
 	}
 
diff --git a/Engine/Source/Runtime/Experimental/ChaosVehicles/ChaosVehiclesCore/Private/SimModule/WheelModule.cpp b/Engine/Source/Runtime/Experimental/ChaosVehicles/ChaosVehiclesCore/Private/SimModule/WheelModule.cpp
--- a/Engine/Source/Runtime/Experimental/ChaosVehicles/ChaosVehiclesCore/Private/SimModule/WheelModule.cpp
+++ b/Engine/Source/Runtime/Experimental/ChaosVehicles/ChaosVehiclesCore/Private/SimModule/WheelModule.cpp
@@ -27,7 +27,7 @@ namespace Chaos
 	void FWheelSimModule::Simulate(float DeltaTime, const FAllInputs& Inputs, FSimModuleTree& VehicleModuleSystem)
 	{	
 		float Re = Setup().Radius;
-		float K = 0.4f; // #TODO: 0.4 * number of wheels - fix the amount of ForceRequiredToStop wrt number of active contact points
+		float K = 0.4f;
 		float TorqueScaling = 1.0f;
 		float TractionControlAndABSScaling = 0.98f;	// how close to perfection is the system working
 
@@ -38,7 +38,7 @@ namespace Chaos
 
 		// TODO: think about doing this properly, stops vehicles rolling around on their own too much
 		// i.e. an auto handbrake feature
-		if (Inputs.ControlInputs.Brake < SMALL_NUMBER && Inputs.ControlInputs.Throttle < SMALL_NUMBER && ModuleLocalVelocity.X < 10.0f)
+		if (Inputs.ControlInputs.Brake < SMALL_NUMBER && Inputs.ControlInputs.Throttle < SMALL_NUMBER && LocalLinearVelocity.X < 10.0f)
 		{
 			BrakeTorque = Setup().HandbrakeTorque;
 		}
@@ -48,7 +48,7 @@ namespace Chaos
 		if (bTouchingGround)
 		{
 			FRotator SteeringRotator(0.f, SteerAngleDegrees, 0.f);
-			FVector Vel = SteeringRotator.UnrotateVector(ModuleLocalVelocity);
+			FVector Vel = SteeringRotator.UnrotateVector(LocalLinearVelocity);
 			FVector LocalWheelVelocity = (Setup().Axis == EWheelAxis::X) ? FVector(Vel.X, Vel.Y, Vel.Z) : FVector(Vel.Y, Vel.X, Vel.Z); // Potential Axis Swap
 			LocalWheelVelocity = Setup().ReverseDirection ? -LocalWheelVelocity : LocalWheelVelocity;
 
@@ -204,6 +204,18 @@ namespace Chaos
 		IntegrateAngularVelocity(DeltaTime, Setup().WheelInertia, Setup().MaxRotationVel);
 	}
 
+	void FWheelSimModule::Animate(Chaos::FClusterUnionPhysicsProxy* Proxy)
+	{
+		if (FPBDRigidClusteredParticleHandle* ClusterChild = GetClusterParticle(Proxy))
+		{
+			float Direction = Setup().ReverseDirection ? -1.0f : 1.0f;
+			FQuat Rot = (Setup().Axis == Chaos::EWheelAxis::Y) ? FQuat(FVector(1, 0, 0), -GetAngularPosition() * Direction) : FQuat(FVector(1, 0, 0), GetAngularPosition() * Direction);
+			FQuat Steer = FQuat(FVector(0, 0, 1), FMath::DegreesToRadians(GetSteerAngleDegrees()));
+
+			FTransform InitialTransform = GetInitialParticleTransform();
+			ClusterChild->ChildToParent().SetRotation(InitialTransform.GetRotation() * Steer * Rot);
+		}
+	}
 
 	bool FWheelSimModule::GetDebugString(FString& StringOut) const
 	{
diff --git a/Engine/Source/Runtime/Experimental/ChaosVehicles/ChaosVehiclesCore/Public/SimModule/AerofoilModule.h b/Engine/Source/Runtime/Experimental/ChaosVehicles/ChaosVehiclesCore/Public/SimModule/AerofoilModule.h
--- a/Engine/Source/Runtime/Experimental/ChaosVehicles/ChaosVehiclesCore/Public/SimModule/AerofoilModule.h
+++ b/Engine/Source/Runtime/Experimental/ChaosVehicles/ChaosVehiclesCore/Public/SimModule/AerofoilModule.h
@@ -2,8 +2,8 @@
 
 #pragma once
 
-#include "SimModule/TorqueSimModule.h"
-
+#include "VehicleUtility.h"
+#include "SimModule/SimulationModuleBase.h"
 
 namespace Chaos
 {
@@ -22,13 +22,12 @@ namespace Chaos
 	{
 		FAerofoilSettings()
 			: Offset(FVector::ZeroVector)
-			, UpAxis(FVector(0.f, 0.f, 1.f))
+			, ForceAxis(FVector(0.f, 0.f, 1.f))
+			, ControlRotationAxis(FVector(0.f, 1.f, 0.f))
 			, Area(5.0f)
 			, Camber(3.0f)
 			, MaxControlAngle(1.f)
 			, StallAngle(16.0f)
-			, MaxCeiling(1E30)
-			, MinCeiling(-1E30)
 			, Type(EAerofoil::Fixed)
 			, LiftMultiplier(1.0f)
 			, DragMultiplier(1.0f)
@@ -36,15 +35,13 @@ namespace Chaos
 		}
 
 		FVector Offset;
-		FVector UpAxis;
+		FVector ForceAxis;
+		FVector ControlRotationAxis;
 		float Area;
 		float Camber;
 		float MaxControlAngle;
 		float StallAngle;
 
-		float MaxCeiling;
-		float MinCeiling;
-
 		EAerofoil Type;
 		float LiftMultiplier;
 		float DragMultiplier;
@@ -77,6 +74,8 @@ namespace Chaos
 
 		virtual void Simulate(float DeltaTime, const FAllInputs& Inputs, FSimModuleTree& VehicleModuleSystem) override;
 
+		virtual void Animate(Chaos::FClusterUnionPhysicsProxy* Proxy) override;
+
 		void SetDensityOfMedium(float InDensity)
 		{
 			CurrentAirDensity = InDensity;
diff --git a/Engine/Source/Runtime/Experimental/ChaosVehicles/ChaosVehiclesCore/Public/SimModule/ChassisModule.h b/Engine/Source/Runtime/Experimental/ChaosVehicles/ChaosVehiclesCore/Public/SimModule/ChassisModule.h
--- a/Engine/Source/Runtime/Experimental/ChaosVehicles/ChaosVehiclesCore/Public/SimModule/ChassisModule.h
+++ b/Engine/Source/Runtime/Experimental/ChaosVehicles/ChaosVehiclesCore/Public/SimModule/ChassisModule.h
@@ -3,6 +3,7 @@
 #pragma once
 
 #include "SimModule/SimulationModuleBase.h"
+#include "VehicleUtility.h"
 
 
 namespace Chaos
@@ -16,10 +17,20 @@ namespace Chaos
 	struct CHAOSVEHICLESCORE_API FChassisSettings
 	{
 		FChassisSettings()
+			: AreaMetresSquared(2.0f)
+			, DragCoefficient(0.5f)
+			, DensityOfMedium(RealWorldConsts::AirDensity())
+			, XAxisMultiplier(1.0f)
+			, YAxisMultiplier(1.0f)
+			, AngularDamping(100000.0f)
 		{
-
 		}
-
+		float AreaMetresSquared;	// [meters squared]
+		float DragCoefficient;		// always positive
+		float DensityOfMedium;
+		float XAxisMultiplier;
+		float YAxisMultiplier;
+		float AngularDamping;
 	};
 
 	/// <summary>
@@ -37,7 +48,7 @@ namespace Chaos
 
 		virtual const FString GetDebugName() const { return TEXT("Chassis"); }
 
-		virtual bool IsBehaviourType(eSimModuleTypeFlags InType) const override { return (InType & NonFunctional); }
+		virtual bool IsBehaviourType(eSimModuleTypeFlags InType) const override { return (InType & Velocity); }
 
 		virtual void Simulate(float DeltaTime, const FAllInputs& Inputs, FSimModuleTree& VehicleModuleSystem) override;
 
diff --git a/Engine/Source/Runtime/Experimental/ChaosVehicles/ChaosVehiclesCore/Public/SimModule/DeferredForcesModular.h b/Engine/Source/Runtime/Experimental/ChaosVehicles/ChaosVehiclesCore/Public/SimModule/DeferredForcesModular.h
--- a/Engine/Source/Runtime/Experimental/ChaosVehicles/ChaosVehiclesCore/Public/SimModule/DeferredForcesModular.h
+++ b/Engine/Source/Runtime/Experimental/ChaosVehicles/ChaosVehiclesCore/Public/SimModule/DeferredForcesModular.h
@@ -70,11 +70,12 @@ public:
 
 	struct FAddTorqueInRadiansData
 	{
-		FAddTorqueInRadiansData(const FTransform& OffsetTransformIn, int TransformIndexIn, const FVector& TorqueIn, bool bAllowSubsteppingIn, bool bAccelChangeIn)
+		FAddTorqueInRadiansData(const FTransform& OffsetTransformIn, int TransformIndexIn, const FVector& TorqueIn, bool bAllowSubsteppingIn, bool bAccelChangeIn, const FColor& ColorIn)
 			: OffsetTransform(OffsetTransformIn)
 			, TransformIndex(TransformIndexIn)
 			, Torque(TorqueIn)
 			, Flags(EForceFlags::None)
+			, DebugColor(ColorIn)
 		{
 			Flags |= bAllowSubsteppingIn ? EForceFlags::AllowSubstepping : EForceFlags::None;
 			Flags |= bAccelChangeIn ? EForceFlags::AccelChange : EForceFlags::None;
@@ -84,6 +85,7 @@ public:
 		int TransformIndex;
 		FVector Torque;
 		EForceFlags Flags;
+		FColor DebugColor;
 	};
 
 	struct FAddImpulseData
@@ -148,6 +150,7 @@ public:
 	Chaos::FPBDRigidParticleHandle* GetParticle(TArray<Chaos::FPBDRigidClusteredParticleHandle*>& Particles
 			, TArray<Chaos::FPBDRigidClusteredParticleHandle*>& ClusterParticles
 			, int TransformIndex
+			, const FVector& PositionalOffset
 			, const TManagedArray<FTransform>& Transforms
 			, const TManagedArray<FTransform>& CollectionMassToLocal
 			, const TManagedArray<int32>& Parent
@@ -157,6 +160,7 @@ public:
 		, TArray<Chaos::FPBDRigidParticleHandle*>& Particles
 		, TArray<Chaos::FPBDRigidClusteredParticleHandle*>& ClusterParticles
 		, int TransformIndex
+		, const FVector& PositionalOffset
 		, FTransform& TransformOut);
 
 	void Apply(TArray<Chaos::FPBDRigidClusteredParticleHandle*>& Particles
@@ -182,3 +186,4 @@ private:
 
 	FTransform ParticleOffsetTransform; // Odd rotation coming through from CU physics bodies
 };
+
diff --git a/Engine/Source/Runtime/Experimental/ChaosVehicles/ChaosVehiclesCore/Public/SimModule/SimModuleTree.h b/Engine/Source/Runtime/Experimental/ChaosVehicles/ChaosVehiclesCore/Public/SimModule/SimModuleTree.h
--- a/Engine/Source/Runtime/Experimental/ChaosVehicles/ChaosVehiclesCore/Public/SimModule/SimModuleTree.h
+++ b/Engine/Source/Runtime/Experimental/ChaosVehicles/ChaosVehiclesCore/Public/SimModule/SimModuleTree.h
@@ -90,6 +90,7 @@ namespace Chaos
 
 		FSimModuleTree()
 		{
+			bAnimationEnabled = false;
 		}
 
 		~FSimModuleTree()
@@ -132,6 +133,9 @@ namespace Chaos
 		const FDeferredForcesModular& GetDeferredForces() const { return DeferredForces; }
 		const TArray<FSimModuleNode>& GetSimulationModuleTree() { return SimulationModuleTree; }
 
+		void SetAnimationEnabled(bool bInEnabled) { bAnimationEnabled = bInEnabled; }
+		bool IsAnimationEnabled() { return bAnimationEnabled; }
+
 		FControlInputs& GetControlInputs()
 		{
 			return AllInputs.ControlInputs;
@@ -172,7 +176,7 @@ namespace Chaos
 		void InterpolateState(const float LerpFactor, Chaos::FModuleNetDataArray& LerpDatas, const Chaos::FModuleNetDataArray& MinDatas, const Chaos::FModuleNetDataArray& MaxDatas);
 
 	protected:
-		void SimulateNode(float DeltaTime, FAllInputs& Inputs, int NodeIdx);
+		void SimulateNode(float DeltaTime, FAllInputs& Inputs, int NodeIdx, FClusterUnionPhysicsProxy* PhysicsProxy);
 
 		void DeleteNodesBelow(int NodeIdx);
 
@@ -187,6 +191,8 @@ namespace Chaos
 		FDeferredForcesModular DeferredForces;
 
 		Chaos::FAllInputs AllInputs;
+
+		bool bAnimationEnabled;
 	};
 
 
diff --git a/Engine/Source/Runtime/Experimental/ChaosVehicles/ChaosVehiclesCore/Public/SimModule/SimulationModuleBase.h b/Engine/Source/Runtime/Experimental/ChaosVehicles/ChaosVehiclesCore/Public/SimModule/SimulationModuleBase.h
--- a/Engine/Source/Runtime/Experimental/ChaosVehicles/ChaosVehiclesCore/Public/SimModule/SimulationModuleBase.h
+++ b/Engine/Source/Runtime/Experimental/ChaosVehicles/ChaosVehiclesCore/Public/SimModule/SimulationModuleBase.h
@@ -4,6 +4,7 @@
 
 #include "CoreMinimal.h"
 #include "Logging/LogMacros.h"
+#include "Chaos/ParticleHandleFwd.h"
 
 
 DECLARE_LOG_CATEGORY_EXTERN(LogSimulationModule, Warning, All);
@@ -21,6 +22,7 @@ namespace Chaos
 {
 	class FSimModuleTree;
 	struct FModuleNetData;
+	class FClusterUnionPhysicsProxy;
 
 	struct CHAOSVEHICLESCORE_API FControlInputs
 	{
@@ -35,6 +37,7 @@ namespace Chaos
 			, Roll(0)
 			, Pitch(0)
 			, Yaw(0)
+			, Boost(0)
 			, ChangeUp(false)
 			, ChangeDown(false)
 			, GearNumber(0)
@@ -51,12 +54,11 @@ namespace Chaos
 		float Roll;
 		float Pitch;
 		float Yaw;
+		float Boost;
 		bool ChangeUp;
 		bool ChangeDown;
 		int GearNumber;
 		int InputDebugIndex;
-
-		int GetChecksum() { return (int)((Throttle+Brake+Steering)*100000.0f);}
 	};
 
 	struct CHAOSVEHICLESCORE_API FModuleHitResults
@@ -110,8 +112,8 @@ namespace Chaos
 
 	enum eSimModuleTypeFlags
 	{
-		NonFunctional	= (1<<0),	// bitmask 1,2,4,8
-		Raycast			= (1<<1),	// requires raycast data
+		NonFunctional	= (1 << 0),	// bitmask 1,2,4,8
+		Raycast			= (1 << 1),	// requires raycast data
 		TorqueBased		= (1 << 2),	// performs torque calculations
 		Velocity		= (1 << 3),	// requires velocity data
 	};
@@ -119,8 +121,8 @@ namespace Chaos
 	enum eSimType
 	{
 		Undefined = 0,
-		Chassis,		// no simulation effect
-		Thruster,		// applies force
+		Chassis,		// linear/angular damping can be applied here
+		Thruster,		// applies force (can be steerable)
 		Aerofoil,		// applied drag and lift forces
 		Wheel,			// a wheel will simply roll if it has no power source
 		Suspension,		// associated with a wheel
@@ -149,8 +151,10 @@ namespace Chaos
 			, SimTreeIndex(INVALID_IDX)
 			, StateFlags(Enabled)
 			, TransformIndex(INVALID_IDX)
-			, ModuleLocalVelocity(FVector::ZeroVector)
+			, LocalLinearVelocity(FVector::ZeroVector)
+			, LocalAngularVelocity(FVector::ZeroVector)
 			, bClustered(true)
+			, bAnimationEnabled(true)
 			, AppliedForce(FVector::ZeroVector)
 			, Guid(INDEX_NONE)
 		{}
@@ -187,7 +191,15 @@ namespace Chaos
 		/**
 		 * The main Simulation function that is called from the physics async callback thread
 		 */
-		virtual void Simulate(float DeltaTime, const FAllInputs& Inputs, FSimModuleTree& VehicleModuleSystem) {};
+		virtual void Simulate(float DeltaTime, const FAllInputs& Inputs, FSimModuleTree& VehicleModuleSystem) {}
+
+		/**
+		 * Animate/modify the childToParent transforms, to say rotate a wheel, or rudder, etc
+		 */
+		virtual void Animate(Chaos::FClusterUnionPhysicsProxy* Proxy) {}
+
+		void SetAnimationEnabled(bool bInEnabled) { bAnimationEnabled = bInEnabled; }
+		bool IsAnimationEnabled() { return bAnimationEnabled; }
 
 		/**
 		 * Option to draw debug for this module requires CVar p.Chaos.DebugDraw.Enabled 1
@@ -228,6 +240,12 @@ namespace Chaos
 		 */
 		void AddLocalForce(const FVector& Force, bool bAllowSubstepping = true, bool bIsLocalForce = false, bool bLevelSlope = false, const FColor& DebugColorIn = FColor::Blue);
 
+		/**
+		 * Torque application function
+		 * Note: forces are applied in local coordinates of the module
+		 */
+		void AddLocalTorque(const FVector& Torque, bool bAllowSubstepping = true, bool bAccelChangeIn = true, const FColor& DebugColorIn = FColor::Magenta);
+
 		//---
 
 		/**
@@ -263,8 +281,10 @@ namespace Chaos
 		/**
 		 * Update the module with its current velocity
 		 */
-		void SetLocalVelocity(const FVector& VelocityIn) { ModuleLocalVelocity = VelocityIn; }
-		const FVector& GetLocalVelocity() const { return ModuleLocalVelocity; }
+		void SetLocalLinearVelocity(const FVector& VelocityIn) { LocalLinearVelocity = VelocityIn; }
+		const FVector& GetLocalLinearVelocity() const { return LocalLinearVelocity; }
+		void SetLocalAngularVelocity(const FVector& VelocityIn) { LocalAngularVelocity = VelocityIn; }
+		const FVector& GetLocalAngularVelocity() const { return LocalAngularVelocity; }
 
 		ISimulationModuleBase* GetParent();
 		ISimulationModuleBase* GetFirstChild();
@@ -275,6 +295,9 @@ namespace Chaos
 		// this is the replication datas
 		virtual TSharedPtr<FModuleNetData> GenerateNetData(int NodeArrayIndex) const = 0;
 
+		//void SetClusterParticle(FPBDRigidClusteredParticleHandle* ParticleIn) { ClusterParticle = ParticleIn; }
+		Chaos::FPBDRigidClusteredParticleHandle* GetClusterParticle(Chaos::FClusterUnionPhysicsProxy* Proxy);
+
 	protected:
 
 		FSimModuleTree* SimModuleTree;	// A pointer back to the simulation tree where we are stored
@@ -288,13 +311,16 @@ namespace Chaos
 
 		FTransform ClusteredCOMRelativeTransform;
 		FTransform IntactCOMRelativeTransform;
-		FVector ModuleLocalVelocity;
+		FVector LocalLinearVelocity;
+		FVector LocalAngularVelocity;
 		bool bClustered;
+		bool bAnimationEnabled;
 
 		// for headless chaos testing
 		FVector AppliedForce;
 		int Guid; // needed a way of associating internal module with game thread.
 
+		//FPBDRigidClusteredParticleHandle* ClusterParticle;
 	};
 
 	/**
diff --git a/Engine/Source/Runtime/Experimental/ChaosVehicles/ChaosVehiclesCore/Public/SimModule/SuspensionModule.h b/Engine/Source/Runtime/Experimental/ChaosVehicles/ChaosVehiclesCore/Public/SimModule/SuspensionModule.h
--- a/Engine/Source/Runtime/Experimental/ChaosVehicles/ChaosVehiclesCore/Public/SimModule/SuspensionModule.h
+++ b/Engine/Source/Runtime/Experimental/ChaosVehicles/ChaosVehiclesCore/Public/SimModule/SuspensionModule.h
@@ -8,6 +8,7 @@ namespace Chaos
 {
 	struct FAllInputs;
 	class FSimModuleTree;
+	class FClusterUnionPhysicsProxy;
 
 	struct CHAOSVEHICLESCORE_API FSuspensionSimModuleDatas : public FModuleNetData
 	{
@@ -123,6 +124,8 @@ namespace Chaos
 
 		virtual void Simulate(float DeltaTime, const FAllInputs& Inputs, FSimModuleTree& VehicleModuleSystem) override;
 
+		virtual void Animate(Chaos::FClusterUnionPhysicsProxy* Proxy) override;
+
 		const FVector& GetRestLocation() const { return Setup().RestOffset; }
 
 		void SetWheelSimTreeIndex(int WheelTreeIndexIn) { WheelSimTreeIndex = WheelTreeIndexIn; }
diff --git a/Engine/Source/Runtime/Experimental/ChaosVehicles/ChaosVehiclesCore/Public/SimModule/ThrusterModule.h b/Engine/Source/Runtime/Experimental/ChaosVehicles/ChaosVehiclesCore/Public/SimModule/ThrusterModule.h
--- a/Engine/Source/Runtime/Experimental/ChaosVehicles/ChaosVehiclesCore/Public/SimModule/ThrusterModule.h
+++ b/Engine/Source/Runtime/Experimental/ChaosVehicles/ChaosVehiclesCore/Public/SimModule/ThrusterModule.h
@@ -18,14 +18,22 @@ namespace Chaos
 		FThrusterSettings()
 			: MaxThrustForce(0)
 			, ForceAxis(FVector(1.0f, 0.0f, 0.0f))
+			, SteeringAxis(FVector(0.0f, 0.0f, 1.0f))
 			, ForceOffset(FVector::ZeroVector)
+			, MaxSteeringAngle(0)
+			, SteeringForceEffect(1.0f)
+			, SteeringEnabled(false)
 		{
 
 		}
 
 		float MaxThrustForce;
 		FVector ForceAxis;
+		FVector SteeringAxis;
 		FVector ForceOffset;
+		float MaxSteeringAngle;
+		float SteeringForceEffect;
+		bool SteeringEnabled;
 	};
 
 	/// <summary>
@@ -47,7 +55,12 @@ namespace Chaos
 
 		virtual void Simulate(float DeltaTime, const FAllInputs& Inputs, FSimModuleTree& VehicleModuleSystem) override;
 
+		virtual void Animate(Chaos::FClusterUnionPhysicsProxy* Proxy) override;
+
+		float GetSteerAngleDegrees() const { return SteerAngleDegrees; }
+
 	private:
+		float SteerAngleDegrees;
 	};
 
 
diff --git a/Engine/Source/Runtime/Experimental/ChaosVehicles/ChaosVehiclesCore/Public/SimModule/WheelModule.h b/Engine/Source/Runtime/Experimental/ChaosVehicles/ChaosVehiclesCore/Public/SimModule/WheelModule.h
--- a/Engine/Source/Runtime/Experimental/ChaosVehicles/ChaosVehiclesCore/Public/SimModule/WheelModule.h
+++ b/Engine/Source/Runtime/Experimental/ChaosVehicles/ChaosVehiclesCore/Public/SimModule/WheelModule.h
@@ -8,6 +8,7 @@ namespace Chaos
 {
 	struct FAllInputs;
 	class FSimModuleTree;
+	class FClusterUnionPhysicsProxy;
 
 	enum CHAOSVEHICLESCORE_API EWheelAxis
 	{
@@ -37,7 +38,6 @@ namespace Chaos
 			, MaxRotationVel(100.0f)
 			, Axis(EWheelAxis::X)
 			, ReverseDirection(false)
-
 		{
 
 		}
@@ -80,6 +80,8 @@ namespace Chaos
 
 		virtual void Simulate(float DeltaTime, const FAllInputs& Inputs, FSimModuleTree& VehicleModuleSystem) override;
 
+		virtual void Animate(Chaos::FClusterUnionPhysicsProxy* Proxy) override;
+
 		virtual bool IsBehaviourType(eSimModuleTypeFlags InType) const override { return (InType & TorqueBased)||(InType & Velocity); }
 
 		void SetSuspensionSimTreeIndex(int IndexIn) { SuspensionSimTreeIndex = IndexIn; }
diff --git a/Engine/Source/Runtime/Experimental/ChaosVehicles/ChaosVehiclesCore/Public/VehicleUtility.h b/Engine/Source/Runtime/Experimental/ChaosVehicles/ChaosVehiclesCore/Public/VehicleUtility.h
--- a/Engine/Source/Runtime/Experimental/ChaosVehicles/ChaosVehiclesCore/Public/VehicleUtility.h
+++ b/Engine/Source/Runtime/Experimental/ChaosVehicles/ChaosVehiclesCore/Public/VehicleUtility.h
@@ -21,6 +21,11 @@ namespace Chaos
 
 	struct CHAOSVEHICLESCORE_API RealWorldConsts
 	{
+		FORCEINLINE static float WaterDensity()
+		{
+			return 997.0f; // kg / m3;
+		}
+
 		FORCEINLINE static float AirDensity()
 		{
 			return 1.225f; // kg / m3;
@@ -218,12 +223,24 @@ namespace Chaos
 		return Cm * 0.01f;
 	}
 
+	/** cm to meters */
+	FORCEINLINE FVector CmToM(const FVector& Cm)
+	{
+		return Cm * 0.01f;
+	}
+
 	/** meters to cm */
 	FORCEINLINE float MToCm(float M)
 	{
 		return M * 100.0f;
 	}
 
+	/** cm to meters */
+	FORCEINLINE FVector MToCm(const FVector& M)
+	{
+		return M * 100.0f;
+	}
+
 	/** cm to meters - warning loss of precision */
 	FORCEINLINE float CmToMiles(float Cm)
 	{
