GeomCollectionISMPool: Use mirrored identity primitive transform on ISMs that have been flagged for reverse culling.
Reverse culling is expected to be set when all instances have mirrored transforms.
Putting the mirror in the component transform will correctly set the primitve flags for use by shaders.

diff --git a/Engine/Source/Runtime/Experimental/GeometryCollectionEngine/Private/GeometryCollection/GeometryCollectionISMPoolComponent.cpp b/Engine/Source/Runtime/Experimental/GeometryCollectionEngine/Private/GeometryCollection/GeometryCollectionISMPoolComponent.cpp
--- a/Engine/Source/Runtime/Experimental/GeometryCollectionEngine/Private/GeometryCollection/GeometryCollectionISMPoolComponent.cpp
+++ b/Engine/Source/Runtime/Experimental/GeometryCollectionEngine/Private/GeometryCollection/GeometryCollectionISMPoolComponent.cpp
@@ -93,9 +93,15 @@ FGeometryCollectionISM::FGeometryCollectionISM(AActor* InOwningActor, const FGeo
 		ISMC->SetDefaultCustomPrimitiveDataFloat(DataIndex, MeshInstance.CustomPrimitiveData[DataIndex]);
 	}
 
+	if ((MeshInstance.Desc.Flags & FISMComponentDescription::ReverseCulling) != 0)
+	{
+		ISMC->SetReverseCulling(true);
+		// We put the mirror in the component transform so that PRIMITIVE_SCENE_DATA_FLAG_DETERMINANT_SIGN will be set for use by materials.
+		ISMC->SetRelativeScale3D(FVector(-1, 1, 1));
+	}
+
 	ISMC->SetRemoveSwap();
 	ISMC->NumCustomDataFloats = MeshInstance.Desc.NumCustomDataFloats;
-	ISMC->SetReverseCulling((MeshInstance.Desc.Flags & FISMComponentDescription::ReverseCulling) != 0);
 	ISMC->SetMobility((MeshInstance.Desc.Flags & FISMComponentDescription::StaticMobility) != 0 ? EComponentMobility::Static : EComponentMobility::Stationary);
 	ISMC->SetCullDistances(MeshInstance.Desc.StartCullDistance, MeshInstance.Desc.EndCullDistance);
 	ISMC->SetCastShadow((MeshInstance.Desc.Flags & FISMComponentDescription::AffectShadow) != 0);
