Rearranged data in Chaos::FTrimeshBVH::FChildData to avoid padding due to alignment requiremtns

diff --git a/Engine/Source/Runtime/Experimental/Chaos/Private/Chaos/TriangleMeshImplicitObject.cpp b/Engine/Source/Runtime/Experimental/Chaos/Private/Chaos/TriangleMeshImplicitObject.cpp
--- a/Engine/Source/Runtime/Experimental/Chaos/Private/Chaos/TriangleMeshImplicitObject.cpp
+++ b/Engine/Source/Runtime/Experimental/Chaos/Private/Chaos/TriangleMeshImplicitObject.cpp
@@ -2089,9 +2089,9 @@ void FTriangleMeshImplicitObject::RebuildFastBVHFromTree(const BVHType& TreeBVH)
 
 		// make the node
 		FTrimeshBVH::FNode& NewNode = FastBVH.Nodes.Emplace_GetRef();
-		NewNode.Children[0].SetChildOrFaceIndex(0);
-		NewNode.Children[0].SetFaceCount(Leaf.Elems.Num());
-		NewNode.Children[0].SetBounds(BoundingBox());
+		NewNode.Children.SetChildOrFaceIndex(0, 0);
+		NewNode.Children.SetFaceCount(0, Leaf.Elems.Num());
+		NewNode.Children.SetBounds(0, BoundingBox());
 		for (const TPayloadBoundsElement<int32, FRealSingle>& LeafPayload: Leaf.Elems)
 		{
 			FastBVH.FaceBounds.Add(FAABBVectorized(LeafPayload.Bounds));
@@ -2157,12 +2157,12 @@ void FTriangleMeshImplicitObject::RebuildFastBVHFromTree(const BVHType& TreeBVH)
 			for (int32 ChildIndex=1; ChildIndex>=0; --ChildIndex)
 			{
 				// common infos
-				FTrimeshBVH::FChildData& ChildData = NewNode.Children[ChildIndex];
-				ChildData.SetBounds(Node.ChildrenBounds[ChildIndex]);
+				FTrimeshBVH::FChildData& ChildData = NewNode.Children;
+				ChildData.SetBounds(ChildIndex, Node.ChildrenBounds[ChildIndex]);
 				const int32 ChildNodeIndex = Node.ChildrenNodes[ChildIndex];
 				// index in the original BVH space, remapping is done at the end of the process
 				// this may be overwritten if the child is a leaf
-				ChildData.SetChildOrFaceIndex(ChildNodeIndex);
+				ChildData.SetChildOrFaceIndex(ChildIndex, ChildNodeIndex);
 
 				if (Nodes.IsValidIndex(ChildNodeIndex))
 				{
@@ -2173,9 +2173,9 @@ void FTriangleMeshImplicitObject::RebuildFastBVHFromTree(const BVHType& TreeBVH)
 						const LeafType& Leaf = Leaves[*LeafIndex];
 
 						// store face range in the node 
-						ChildData.SetChildOrFaceIndex(FaceNum);
-						ChildData.SetFaceCount(Leaf.Elems.Num());
-						check(ChildData.GetFaceCount() > 0);
+						ChildData.SetChildOrFaceIndex(ChildIndex, FaceNum);
+						ChildData.SetFaceCount(ChildIndex, Leaf.Elems.Num());
+						check(ChildData.GetFaceCount(ChildIndex) > 0);
 
 						TMap<int32, int32> VertexReuse;
 						// copy indices in the linear face array
@@ -2212,15 +2212,15 @@ void FTriangleMeshImplicitObject::RebuildFastBVHFromTree(const BVHType& TreeBVH)
 		FTrimeshBVH::FNode& Node = FastBVH.Nodes[NodeIndex];
 		for (int32 ChildIndex=0; ChildIndex<2; ++ChildIndex)
 		{
-			FTrimeshBVH::FChildData& ChildData = Node.Children[ChildIndex];
-			const bool bHasFaces = ChildData.GetFaceCount() > 0;
+			FTrimeshBVH::FChildData& ChildData = Node.Children;
+			const bool bHasFaces = ChildData.GetFaceCount(ChildIndex) > 0;
 			if (!bHasFaces)
 			{
-				const int32 ChildNodeIndex = ChildData.GetChildOrFaceIndex();
+				const int32 ChildNodeIndex = ChildData.GetChildOrFaceIndex(ChildIndex);
 				const int32* FixedChildNodeIndex = BVHToFastBVHNodeIndexMap.Find(ChildNodeIndex);
 				if (ensure(FixedChildNodeIndex))
 				{
-					ChildData.SetChildOrFaceIndex( *FixedChildNodeIndex);
+					ChildData.SetChildOrFaceIndex(ChildIndex, *FixedChildNodeIndex);
 				}
 			}
 		}
diff --git a/Engine/Source/Runtime/Experimental/Chaos/Public/Chaos/TriangleMeshImplicitObject.h b/Engine/Source/Runtime/Experimental/Chaos/Public/Chaos/TriangleMeshImplicitObject.h
--- a/Engine/Source/Runtime/Experimental/Chaos/Public/Chaos/TriangleMeshImplicitObject.h
+++ b/Engine/Source/Runtime/Experimental/Chaos/Public/Chaos/TriangleMeshImplicitObject.h
@@ -50,53 +50,56 @@ namespace Chaos
 		struct alignas(16) FChildData
 		{
 			FChildData()
-				: ChildOrFaceIndex(INDEX_NONE)
-				, FaceCount(0)
+				: ChildOrFaceIndex{ INDEX_NONE, INDEX_NONE }
+				, FaceCount{ 0, 0 }
 			{
 			}
 
-			FORCEINLINE int32 GetChildOrFaceIndex() const
+			FORCEINLINE int32 GetChildOrFaceIndex(int ChildIndex) const
 			{
-				return ChildOrFaceIndex;
+				return ChildOrFaceIndex[ChildIndex];
 			}
 
-			FORCEINLINE int32 GetFaceCount() const
+			FORCEINLINE int32 GetFaceCount(int ChildIndex) const
 			{
-				return FaceCount;
+				return FaceCount[ChildIndex];
 			}
 
-			FORCEINLINE void SetChildOrFaceIndex(int32 InChildOrFaceIndex)
+			FORCEINLINE void SetChildOrFaceIndex(int ChildIndex, int32 InChildOrFaceIndex)
 			{
-				ChildOrFaceIndex = InChildOrFaceIndex;
+				ChildOrFaceIndex[ChildIndex] = InChildOrFaceIndex;
 			}
 
-			FORCEINLINE void SetFaceCount(int32 InFaceCount)
+			FORCEINLINE void SetFaceCount(int ChildIndex, int32 InFaceCount)
 			{
-				FaceCount = InFaceCount;
+				FaceCount[ChildIndex] = InFaceCount;
 			}
 
-			FORCEINLINE void SetBounds(const FAABB3& AABB)
+			FORCEINLINE void SetBounds(int ChildIndex, const FAABB3& AABB)
 			{
-				Bounds = FAABBVectorized(AABB);
+				Bounds[ChildIndex] = FAABBVectorized(AABB);
 			}
-			FORCEINLINE void SetBounds(const TAABB<FRealSingle, 3>& AABB)
+			FORCEINLINE void SetBounds(int ChildIndex, const TAABB<FRealSingle, 3>& AABB)
 			{
-				Bounds = FAABBVectorized(AABB);
+				Bounds[ChildIndex] = FAABBVectorized(AABB);
 			}
 
-			const FAABBVectorized& GetBounds() const { return Bounds; }
+			const FAABBVectorized& GetBounds(int ChildIndex) const { return Bounds[ChildIndex]; }
 
 			void Serialize(FArchive& Ar)
 			{
-				Ar << Bounds;
-				Ar << ChildOrFaceIndex;
-				Ar << FaceCount;
+				Ar << Bounds[0];
+				Ar << ChildOrFaceIndex[0];
+				Ar << FaceCount[0];
+				Ar << Bounds[1];
+				Ar << ChildOrFaceIndex[1];
+				Ar << FaceCount[1];
 			}
 
 		private:
-			FAABBVectorized Bounds;
-			int32 ChildOrFaceIndex;
-			int32 FaceCount;
+			FAABBVectorized Bounds[2];
+			int32 ChildOrFaceIndex[2];
+			int32 FaceCount[2];
 		};
 		
 		struct FNode
@@ -105,10 +108,9 @@ namespace Chaos
 
 			void Serialize(FArchive& Ar)
 			{
-				Children[0].Serialize(Ar);
-				Children[1].Serialize(Ar);
+				Children.Serialize(Ar);
 			}
-			FChildData Children[2];
+			FChildData Children;
 		};
 
 		template <typename SQVisitor>
@@ -242,13 +244,13 @@ namespace Chaos
 
 				for (int32 ChildIndex = 0; ChildIndex < 2; ++ChildIndex)
 				{
-					const FChildData& ChildData = Node.Children[ChildIndex]; 
-					const int32 FaceIndex = ChildData.GetChildOrFaceIndex();
+					const FChildData& ChildData = Node.Children; 
+					const int32 FaceIndex = ChildData.GetChildOrFaceIndex(ChildIndex);
 					if (FaceIndex != INDEX_NONE)
 					{
-						if (BoundsFilter(ChildData.GetBounds()) == EFilterResult::Keep)
+						if (BoundsFilter(ChildData.GetBounds(ChildIndex)) == EFilterResult::Keep)
 						{
-							const uint32 FaceCount = ChildData.GetFaceCount();
+							const uint32 FaceCount = ChildData.GetFaceCount(ChildIndex);
 							if (FaceCount > 0)
 							{
 								if (EVisitorResult::Stop == VisitFaces(FaceIndex, FaceCount, BoundsFilter, FaceVisitor))
