Add GetFirstClusteredPhyscialMaterial to get the first physical material of a child particle if it is a clustered particle with no material of its own.

diff --git a/Engine/Plugins/Experimental/Buoyancy/Source/Runtime/Private/BuoyancySubsystem.cpp b/Engine/Plugins/Experimental/Buoyancy/Source/Runtime/Private/BuoyancySubsystem.cpp
--- a/Engine/Plugins/Experimental/Buoyancy/Source/Runtime/Private/BuoyancySubsystem.cpp
+++ b/Engine/Plugins/Experimental/Buoyancy/Source/Runtime/Private/BuoyancySubsystem.cpp
@@ -225,9 +225,6 @@ void FBuoyancySubsystemSimCallback::OnPreSimulate_Internal()
 	{
 		return;
 	}
-
-
-
 }
 
 void FBuoyancySubsystemSimCallback::OnMidPhaseModification_Internal(Chaos::FMidPhaseModifierAccessor& MidPhaseAccessor)
@@ -310,17 +307,6 @@ void FBuoyancySubsystemSimCallback::OnMidPhaseModification_Internal(Chaos::FMidP
 			return;
 		}
 
-		// If the particle is in a cluster, get the actual root particle that we can apply forces to
-		FPBDRigidClusteredParticleHandle* RigidCluster = RigidParticle->CastToClustered();
-		if (RigidCluster)
-		{
-			while (FPBDRigidClusteredParticleHandle* ParentCluster = RigidCluster->Parent())
-			{
-				RigidCluster = ParentCluster;
-			}
-		}
-		FPBDRigidParticleHandle* RigidRoot = RigidCluster ? RigidCluster->CastToRigidParticle() : RigidParticle;
-
 		// Compute submerged volume and CoM
 		float SubmergedVol;
 		FVec3 SubmergedCoM;
@@ -376,7 +362,7 @@ void FBuoyancySubsystemSimCallback::OnMidPhaseModification_Internal(Chaos::FMidP
 		// However, the volumes or masses of each "true" shape are not accessible to
 		// us, so at the moment this is nearly the best estimate we'll be able to get.
 		FRealSingle SubmergedVol = Submersion.Vol;
-		if (const FChaosPhysicsMaterial* ParticleMaterial = Evolution->GetFirstPhysicsMaterial(Submersion.Particle))
+		if (const FChaosPhysicsMaterial* ParticleMaterial = Evolution->GetFirstClusteredPhysicsMaterial(Submersion.Particle))
 		{
 			const FRealSingle ParticleDensity = Chaos::GCm3ToKgCm3(ParticleMaterial->Density);
 			const FRealSingle ParticleMass = Submersion.Particle->M();
diff --git a/Engine/Source/Runtime/Experimental/Chaos/Private/Chaos/ClusterUnionManager.cpp b/Engine/Source/Runtime/Experimental/Chaos/Private/Chaos/ClusterUnionManager.cpp
--- a/Engine/Source/Runtime/Experimental/Chaos/Private/Chaos/ClusterUnionManager.cpp
+++ b/Engine/Source/Runtime/Experimental/Chaos/Private/Chaos/ClusterUnionManager.cpp
@@ -364,12 +364,23 @@ namespace Chaos
 		return ClusterUnions.Find(Index);
 	}
 
+	const FClusterUnion* FClusterUnionManager::FindClusterUnion(FClusterUnionIndex Index) const
+	{
+		return ClusterUnions.Find(Index);
+	}
+
 	FClusterUnion* FClusterUnionManager::FindClusterUnionFromParticle(FPBDRigidParticleHandle* Particle)
 	{
 		FClusterUnionIndex ClusterUnionIndex = FindClusterUnionIndexFromParticle(Particle);
 		return FindClusterUnion(ClusterUnionIndex);
 	}
 
+	const FClusterUnion* FClusterUnionManager::FindClusterUnionFromParticle(const FPBDRigidParticleHandle* Particle) const
+	{
+		FClusterUnionIndex ClusterUnionIndex = FindClusterUnionIndexFromParticle(Particle);
+		return FindClusterUnion(ClusterUnionIndex);
+	}
+
 	DECLARE_CYCLE_STAT(TEXT("FClusterUnionManager::HandleAddOperation"), STAT_HandleAddOperation, STATGROUP_Chaos);
 	void FClusterUnionManager::HandleAddOperation(FClusterUnionIndex ClusterIndex, const TArray<FPBDRigidParticleHandle*>& Particles, bool bReleaseClustersFirst)
 	{
@@ -777,7 +788,7 @@ namespace Chaos
 	}
 
 	DECLARE_CYCLE_STAT(TEXT("FClusterUnionManager::FindClusterUnionIndexFromParticle"), STAT_FindClusterUnionIndexFromParticle, STATGROUP_Chaos);
-	FClusterUnionIndex FClusterUnionManager::FindClusterUnionIndexFromParticle(FPBDRigidParticleHandle* ChildParticle)
+	FClusterUnionIndex FClusterUnionManager::FindClusterUnionIndexFromParticle(const FPBDRigidParticleHandle* ChildParticle) const
 	{
 		SCOPE_CYCLE_COUNTER(STAT_FindClusterUnionIndexFromParticle);
 		if (!ChildParticle)
@@ -785,9 +796,9 @@ namespace Chaos
 			return INDEX_NONE;
 		}
 
-		if (FPBDRigidClusteredParticleHandle* ChildClustered = ChildParticle->CastToClustered())
+		if (const FPBDRigidClusteredParticleHandle* ChildClustered = ChildParticle->CastToClustered())
 		{
-			if (FPBDRigidClusteredParticleHandle* Parent = ChildClustered->Parent())
+			if (const FPBDRigidClusteredParticleHandle* Parent = ChildClustered->Parent())
 			{
 				// Recursion should be fine here since the hierarchy should be fairly shallow.
 				return FindClusterUnionIndexFromParticle(Parent);
@@ -795,7 +806,7 @@ namespace Chaos
 
 			// The only other check to do here is to see if this is the cluster union particle itself.
 			const int32 ClusterGroupIndex = ChildClustered->ClusterGroupIndex();
-			if (FClusterUnion* ClusterUnion = FindClusterUnion(FMath::Abs(ClusterGroupIndex)))
+			if (const FClusterUnion* ClusterUnion = FindClusterUnion(FMath::Abs(ClusterGroupIndex)))
 			{
 				// This is a sanity check which may or may not be necessary.
 				// TODO: Can probably be safely removed once we deprecate cluster group indices on GCs.
diff --git a/Engine/Source/Runtime/Experimental/Chaos/Private/Chaos/PBDRigidsEvolutionGBF.cpp b/Engine/Source/Runtime/Experimental/Chaos/Private/Chaos/PBDRigidsEvolutionGBF.cpp
--- a/Engine/Source/Runtime/Experimental/Chaos/Private/Chaos/PBDRigidsEvolutionGBF.cpp
+++ b/Engine/Source/Runtime/Experimental/Chaos/Private/Chaos/PBDRigidsEvolutionGBF.cpp
@@ -1094,6 +1094,38 @@ void FPBDRigidsEvolutionGBF::ParticleMaterialChanged(FGeometryParticleHandle* Pa
 	GetIslandManager().UpdateParticleMaterial(Particle);
 }
 
+const FChaosPhysicsMaterial* FPBDRigidsEvolutionGBF::GetFirstClusteredPhysicsMaterial(const FGeometryParticleHandle* Particle) const
+{
+	if (const FChaosPhysicsMaterial* PhysicsMaterial = GetFirstPhysicsMaterial(Particle))
+	{
+		return PhysicsMaterial;
+	}
+	else if (const FPBDRigidClusteredParticleHandle* Cluster = Particle->CastToClustered())
+	{
+		if (const FClusterUnion* ClusterUnion = Clustering.GetClusterUnionManager().FindClusterUnionFromParticle(Cluster))
+		{
+			if (ClusterUnion->InternalCluster == Cluster)
+			{
+				for (const FPBDRigidParticleHandle* ChildParticle : ClusterUnion->ChildParticles)
+				{
+					// I'm not sure how a cluster might end up containing itself, but
+					// in order to avoid this at all costs we just add a check here to make
+					// sure the particle isn't it's own child before trying to get its
+					// material
+					if (ChildParticle && ensureMsgf(ChildParticle != Particle, TEXT("Clustered particle found which has itself as a child")))
+					{
+						if (const FChaosPhysicsMaterial* ChildPhysicalMaterial = GetFirstClusteredPhysicsMaterial(ChildParticle))
+						{
+							return ChildPhysicalMaterial;
+						}
+					}
+				}
+			}
+		}
+	}
+
+	return nullptr;
+}
 
 void FPBDRigidsEvolutionGBF::UpdateInertiaConditioning()
 {
diff --git a/Engine/Source/Runtime/Experimental/Chaos/Public/Chaos/ClusterUnionManager.h b/Engine/Source/Runtime/Experimental/Chaos/Public/Chaos/ClusterUnionManager.h
--- a/Engine/Source/Runtime/Experimental/Chaos/Public/Chaos/ClusterUnionManager.h
+++ b/Engine/Source/Runtime/Experimental/Chaos/Public/Chaos/ClusterUnionManager.h
@@ -159,8 +159,10 @@ namespace Chaos
 		// Access the cluster union externally.
 		CHAOS_API FClusterUnion* FindClusterUnionFromExplicitIndex(FClusterUnionExplicitIndex Index);
 		CHAOS_API FClusterUnion* FindClusterUnion(FClusterUnionIndex Index);
+		CHAOS_API const FClusterUnion* FindClusterUnion(FClusterUnionIndex Index) const;
 		CHAOS_API FClusterUnion* FindClusterUnionFromParticle(FPBDRigidParticleHandle* Particle);
-		CHAOS_API FClusterUnionIndex FindClusterUnionIndexFromParticle(FPBDRigidParticleHandle* Particle);
+		CHAOS_API const FClusterUnion* FindClusterUnionFromParticle(const FPBDRigidParticleHandle* Particle) const;
+		CHAOS_API FClusterUnionIndex FindClusterUnionIndexFromParticle(const FPBDRigidParticleHandle* Particle) const;
 
 		// An extension to FindClusterUnionIndexFromParticle to check whether or not the given particle is the cluster union particle itself.
 		CHAOS_API bool IsClusterUnionParticle(FPBDRigidClusteredParticleHandle* Particle);
diff --git a/Engine/Source/Runtime/Experimental/Chaos/Public/Chaos/PBDRigidsEvolutionGBF.h b/Engine/Source/Runtime/Experimental/Chaos/Public/Chaos/PBDRigidsEvolutionGBF.h
--- a/Engine/Source/Runtime/Experimental/Chaos/Public/Chaos/PBDRigidsEvolutionGBF.h
+++ b/Engine/Source/Runtime/Experimental/Chaos/Public/Chaos/PBDRigidsEvolutionGBF.h
@@ -181,6 +181,8 @@ namespace Chaos
 		// Called when a the material changes one or more shapes on a particle. Required because collisions cache material properties
 		void ParticleMaterialChanged(FGeometryParticleHandle* Particle);
 
+		CHAOS_API const FChaosPhysicsMaterial* GetFirstClusteredPhysicsMaterial(const FGeometryParticleHandle* Particle) const;
+
 		template<typename TParticleView>
 		void Integrate(const TParticleView& InParticles, FReal Dt)
 		{
