Rename AreaFromBoundingBoxOverlap to ComputeAreaFromBoundingBoxOverlap to avoid ambiguous calls

diff --git a/Engine/Source/Runtime/Experimental/Chaos/Private/Chaos/ClusterUnionStressSolver.cpp b/Engine/Source/Runtime/Experimental/Chaos/Private/Chaos/ClusterUnionStressSolver.cpp
--- a/Engine/Source/Runtime/Experimental/Chaos/Private/Chaos/ClusterUnionStressSolver.cpp
+++ b/Engine/Source/Runtime/Experimental/Chaos/Private/Chaos/ClusterUnionStressSolver.cpp
@@ -13,7 +13,7 @@ namespace
 	float ClusterUnionStressSolverStrengthScalar = 1.f;
 	FAutoConsoleVariableRef CVarClusterUnionStressSolverStrengthScalar(TEXT("ClusterUnion.StressSolver.StrengthScalar"), ClusterUnionStressSolverStrengthScalar, TEXT("Materioal strength scalar ( <1: weaker, >1: stronger)"));
 
-	float AreaFromBoundingBoxOverlap(const Chaos::FAABB3& BoxA, const Chaos::FAABB3& BoxB)
+	float ComputeAreaFromBoundingBoxOverlap(const Chaos::FAABB3& BoxA, const Chaos::FAABB3& BoxB)
 	{
 		// if the two box don't overlap, we'll get an inside out box 
 		// but we are still using it to compute the area as an approximation
@@ -250,7 +250,7 @@ namespace Chaos
 						// compute the actual overlap surface ( instead of using the strain of the edge - may not be the right area for now ) 
 						// todo : we shoud precompute this so we can just use the "strain" property of the connection  
 
-						const float ConnectionArea = AreaFromBoundingBoxOverlap(Node.Particle->WorldSpaceInflatedBounds(), OppositeNode->Particle->WorldSpaceInflatedBounds());
+						const float ConnectionArea = ComputeAreaFromBoundingBoxOverlap(Node.Particle->WorldSpaceInflatedBounds(), OppositeNode->Particle->WorldSpaceInflatedBounds());
 						TConnectivityEdge<FReal>& PropagatingConnection = PropagatingConnections.Add_GetRef(Connection);
 						PropagatingConnection.SetArea(ConnectionArea);
 						TotalPropagationArea += ConnectionArea;
