Add fatigue threshold based destruction logic

diff --git a/Engine/Source/Runtime/Experimental/Chaos/Private/Chaos/StrainModification.cpp b/Engine/Source/Runtime/Experimental/Chaos/Private/Chaos/StrainModification.cpp
--- a/Engine/Source/Runtime/Experimental/Chaos/Private/Chaos/StrainModification.cpp
+++ b/Engine/Source/Runtime/Experimental/Chaos/Private/Chaos/StrainModification.cpp
@@ -6,6 +6,16 @@
 #include "GeometryCollection/GeometryCollectionSimulationCoreTypes.h"
 #include "GeometryCollection/GeometryCollection.h"
 
+namespace
+{
+	float GetMaxAppliedStrain(const Chaos::FPBDRigidClusteredParticleHandle* Handle, const uint8 StrainTypes)
+	{
+		const float CollisionImpulses = (StrainTypes & Chaos::EStrainTypes::CollisionStrain) ? Handle->CollisionImpulses() : 0.f;
+		const float ExternalStrain = (StrainTypes & Chaos::EStrainTypes::ExternalStrain) ? Handle->GetExternalStrain() : 0.f;
+		return FMath::Max(CollisionImpulses, ExternalStrain);
+	}
+}
+
 const FGeometryCollectionPhysicsProxy* Chaos::FStrainedProxyModifier::GetProxy() const
 {
 	return Proxy;
@@ -40,19 +50,8 @@ int32 Chaos::FStrainedProxyModifier::GetNumBreakingStrains(const bool bDoubleCou
 		if (ChildHandle->Parent() == nullptr) { continue; }
 
 		// Get the applied and internal strain
-		//
-		// TODO: Make this a computation internal to the strain, or the connection
-		// graph or something... this logic is currently copied from
-		// FRigidClustering::ReleaseClusterParticlesImpl, but should be tucked behind
-		// an interface somewhere.
-		const Chaos::FReal CollisionImpulses
-			= (StrainTypes & Chaos::EStrainTypes::CollisionStrain)
-			? ChildHandle->CollisionImpulses() : FReal(0);
-		const Chaos::FReal ExternalStrain
-			= (StrainTypes & Chaos::EStrainTypes::ExternalStrain)
-			? ChildHandle->GetExternalStrain() : FReal(0);
-		const Chaos::FReal InternalStrain = ChildHandle->GetInternalStrains();
-		Chaos::FReal MaxAppliedStrain = FMath::Max(CollisionImpulses, ExternalStrain);
+		const float InternalStrain = ChildHandle->GetInternalStrains();
+		const float MaxAppliedStrain = GetMaxAppliedStrain(ChildHandle, StrainTypes);
 
 		if (bDoubleCount && InternalStrain > SMALL_NUMBER)
 		{
@@ -68,6 +67,56 @@ int32 Chaos::FStrainedProxyModifier::GetNumBreakingStrains(const bool bDoubleCou
 	return NumBreakingStrains;
 }
 
+float Chaos::FStrainedProxyModifier::GetMaxBreakStrainRatio(const float FatigueThreshold, const uint8 StrainTypes) const
+{
+	// Make sure we have a proxy and rest-children
+	if (Proxy == nullptr) { return 0; }
+	if (RestChildren == nullptr) { return 0; }
+
+	// Loop over each child, checking whether or not it will have been freed
+	// by the strain that it has accumulated
+	float MaxBreakStrainRatio = 0.f;
+	for (int32 RestChildIdx : *RestChildren)
+	{
+		Chaos::FPBDRigidClusteredParticleHandle* ChildHandle = Proxy->GetSolverParticleHandles()[RestChildIdx];
+		if (ChildHandle->Parent() == nullptr) { continue; }
+
+		// compute the strain ratio
+		const float InternalStrain = ChildHandle->GetInternalStrains();
+		const float MaxAppliedStrain = GetMaxAppliedStrain(ChildHandle, StrainTypes);
+		const float AdjustedAppliedStrain = (MaxAppliedStrain >= FatigueThreshold) ? MaxAppliedStrain : 0.f;
+		const float StrainRatio = (InternalStrain > SMALL_NUMBER)? (MaxAppliedStrain / InternalStrain) : 1.0f;
+
+		MaxBreakStrainRatio = FMath::Max(MaxBreakStrainRatio, StrainRatio);
+	}
+
+	return MaxBreakStrainRatio;
+}
+
+void Chaos::FStrainedProxyModifier::AdjustStrainForBreak(const float FatigueThreshold, const uint8 StrainTypes)
+{
+	// Make sure we have a proxy and rest-children
+	if (Proxy == nullptr) { return; }
+	if (RestChildren == nullptr) { return; }
+
+	// Loop over each child, checking whether or not it will have been freed
+	// by the strain that it has accumulated
+	float MaxBreakStrainRatio = 0.f;
+	for (int32 RestChildIdx : *RestChildren)
+	{
+		Chaos::FPBDRigidClusteredParticleHandle* ChildHandle = Proxy->GetSolverParticleHandles()[RestChildIdx];
+		if (ChildHandle->Parent() == nullptr) { continue; }
+
+		// compute the strain ratio
+		const float InternalStrain = ChildHandle->GetInternalStrains();
+		const float MaxAppliedStrain = GetMaxAppliedStrain(ChildHandle, StrainTypes);
+		if (MaxAppliedStrain >= FatigueThreshold)
+		{
+			RigidClustering.SetExternalStrain(ChildHandle, InternalStrain);
+		}
+	}
+}
+
 void Chaos::FStrainedProxyModifier::ClearStrains()
 {
 	// Make sure we have a proxy and rest-children
diff --git a/Engine/Source/Runtime/Experimental/Chaos/Public/Chaos/StrainModification.h b/Engine/Source/Runtime/Experimental/Chaos/Public/Chaos/StrainModification.h
--- a/Engine/Source/Runtime/Experimental/Chaos/Public/Chaos/StrainModification.h
+++ b/Engine/Source/Runtime/Experimental/Chaos/Public/Chaos/StrainModification.h
@@ -58,9 +58,19 @@ namespace Chaos
 		// break a connection N times.
 		CHAOS_API int32 GetNumBreakingStrains(bool bDoubleCount = true, const uint8 StrainTypes = EStrainTypes::ExternalStrain | EStrainTypes::CollisionStrain) const;
 
+		// Go through the children and find the one with the largest ratio of applied strain vs internal strain and return this ratio
+		// @param FatigueThreshold if the applied strains are below the fatgue threshold, they will be ignored and the returned ratio will be 0 
+		// @param StrainTypes type strain to account for ( collision and/or external )
+		CHAOS_API float GetMaxBreakStrainRatio(const float FatigueThreshold, const uint8 StrainTypes) const;
+
 		// Clear strains for all strained cluster children
 		CHAOS_API void ClearStrains();
 
+		// Adjust strain for all the strained children above the fatigue threshold so that their strain is large enough to break
+		// @param FatigueThreshold if the applied strains are below the fatgue threshold, they will be ignored 
+		// @param StrainTypes type strain to account for ( collision and/or external )
+		CHAOS_API void AdjustStrainForBreak(const float FatigueThreshold, const uint8 StrainTypes);
+
 	private:
 
 		static CHAOS_API Chaos::FPBDRigidClusteredParticleHandle* InitRootHandle(FGeometryCollectionPhysicsProxy* Proxy);
