Fix inverted geometry in ISMPool

diff --git a/Engine/Source/Runtime/Experimental/GeometryCollectionEngine/Private/GeometryCollection/GeometryCollectionISMPoolComponent.cpp b/Engine/Source/Runtime/Experimental/GeometryCollectionEngine/Private/GeometryCollection/GeometryCollectionISMPoolComponent.cpp
--- a/Engine/Source/Runtime/Experimental/GeometryCollectionEngine/Private/GeometryCollection/GeometryCollectionISMPoolComponent.cpp
+++ b/Engine/Source/Runtime/Experimental/GeometryCollectionEngine/Private/GeometryCollection/GeometryCollectionISMPoolComponent.cpp
@@ -95,8 +95,9 @@ FGeometryCollectionISM::FGeometryCollectionISM(AActor* InOwningActor, const FGeo
 
 	if ((MeshInstance.Desc.Flags & FISMComponentDescription::ReverseCulling) != 0)
 	{
-		ISMC->SetReverseCulling(true);
-		// We put the mirror in the component transform so that PRIMITIVE_SCENE_DATA_FLAG_DETERMINANT_SIGN will be set for use by materials.
+		// Instead of reverse culling, we put a mirror in the component transform so that 
+		// PRIMITIVE_SCENE_DATA_FLAG_DETERMINANT_SIGN will be set for use by materials.
+		//ISMC->SetReverseCulling(true);
 		ISMC->SetRelativeScale3D(FVector(-1, 1, 1));
 	}
 
