RunOff would cause early write ops

#rnx

diff --git a/Engine/Source/Runtime/Experimental/IoStoreOnDemand/Private/Governor.inl b/Engine/Source/Runtime/Experimental/IoStoreOnDemand/Private/Governor.inl
--- a/Engine/Source/Runtime/Experimental/IoStoreOnDemand/Private/Governor.inl
+++ b/Engine/Source/Runtime/Experimental/IoStoreOnDemand/Private/Governor.inl
@@ -77,8 +77,12 @@ inline void FGovernorInternal::Set(uint32 InAllowance, uint32 Ops, uint32 Second
 inline uint32 FGovernorInternal::TickAllowance()
 {
 	int64 Cycle = FPlatformTime::Cycles64();
-	uint32 Ret = TickAllowance(Cycle) + RunOff;
-	RunOff = 0;
+	uint32 Ret = TickAllowance(Cycle);
+	if (Ret != 0)
+	{
+		Ret += RunOff;
+		RunOff = 0;
+	}
 	return Ret;
 }
 
