Ias: Add LLM FileSystem scopes.
Rename some distribution endpoint functions and ass some cpu related scopes.

#rnx

diff --git a/Engine/Source/Runtime/Experimental/IoStoreOnDemand/Private/IoStoreOnDemand.cpp b/Engine/Source/Runtime/Experimental/IoStoreOnDemand/Private/IoStoreOnDemand.cpp
--- a/Engine/Source/Runtime/Experimental/IoStoreOnDemand/Private/IoStoreOnDemand.cpp
+++ b/Engine/Source/Runtime/Experimental/IoStoreOnDemand/Private/IoStoreOnDemand.cpp
@@ -6,6 +6,7 @@
 #include "LatencyInjector.h"
 
 #include "FileIoCache.h"
+#include "HAL/LowLevelMemTracker.h"
 #include "HAL/PlatformMisc.h"
 #include "HAL/PlatformTime.h"
 #include "HAL/FileManager.h"
@@ -1280,6 +1281,7 @@ void FIoStoreOnDemandModule::SetEnabled(bool bInEnabled)
 
 void FIoStoreOnDemandModule::StartupModule()
 {
+	LLM_SCOPE(ELLMTag::FileSystem);
 	using namespace UE::IO::Private;
 
 #if WITH_EDITOR
diff --git a/Engine/Source/Runtime/Experimental/IoStoreOnDemand/Private/OnDemandIoDispatcherBackend.cpp b/Engine/Source/Runtime/Experimental/IoStoreOnDemand/Private/OnDemandIoDispatcherBackend.cpp
--- a/Engine/Source/Runtime/Experimental/IoStoreOnDemand/Private/OnDemandIoDispatcherBackend.cpp
+++ b/Engine/Source/Runtime/Experimental/IoStoreOnDemand/Private/OnDemandIoDispatcherBackend.cpp
@@ -9,6 +9,7 @@
 #include "EncryptionKeyManager.h"
 #include "FileIoCache.h"
 #include "HAL/Event.h"
+#include "HAL/LowLevelMemTracker.h"
 #include "HAL/Platform.h"
 #include "HAL/PlatformTime.h"
 #include "HAL/PreprocessorHelpers.h"
@@ -95,7 +96,7 @@ public:
 	FDistributionEndpoints() = default;
 	~FDistributionEndpoints();
 
-	void Resolve(const FString& DistributionUrl, FOnEndpointResolved&& OnResolved);
+	void ResolveEndpoints(const FString& DistributionUrl, FOnEndpointResolved&& OnResolved);
 	void ResolveDeferredEndpoints();
 
 private:
@@ -112,9 +113,9 @@ private:
 		int32 RetryCount = 0;
 	};
 
-	void IssueRequests();
-	void CancelRequests();
-	void CompleteRequest(FResolveRequest& ResolveRequest, FHttpResponsePtr HttpResponse);
+	void IssueEndpointRequests();
+	void CancelEndpointRequests();
+	void CompleteEndpointRequest(FResolveRequest& ResolveRequest, FHttpResponsePtr HttpResponse);
 
 	TMap<FString, TUniquePtr<FResolvedEndpoint>> ResolvedEndpoints;
 	TMap<FString, TUniquePtr<FResolveRequest>> PendingRequests;
@@ -124,11 +125,12 @@ private:
 
 FDistributionEndpoints::~FDistributionEndpoints()
 {
-	CancelRequests();
+	CancelEndpointRequests();
 }
 
-void FDistributionEndpoints::Resolve(const FString& DistributionUrl, FOnEndpointResolved&& OnResolved)
+void FDistributionEndpoints::ResolveEndpoints(const FString& DistributionUrl, FOnEndpointResolved&& OnResolved)
 {
+	TRACE_CPUPROFILER_EVENT_SCOPE(FOnDemandIoBackend::ResolveEndpoints);
 	const FResolvedEndpoint* Ep = nullptr;
 	{
 		FReadScopeLock _(Lock);
@@ -157,22 +159,24 @@ void FDistributionEndpoints::Resolve(const FString& DistributionUrl, FOnEndpoint
 
 	if (bIssueRequest)
 	{
-		IssueRequests();
+		IssueEndpointRequests();
 	}
 }
 
 void FDistributionEndpoints::ResolveDeferredEndpoints()
 {
+	TRACE_CPUPROFILER_EVENT_SCOPE(FOnDemandIoBackend::ResolveDeferredEndpoints);
 	{
 		FWriteScopeLock _(Lock);
 		bInitialized = true;
 	}
 
-	IssueRequests();
+	IssueEndpointRequests();
 }
 
-void FDistributionEndpoints::IssueRequests()
+void FDistributionEndpoints::IssueEndpointRequests()
 {
+	TRACE_CPUPROFILER_EVENT_SCOPE(FOnDemandIoBackend::IssueEndpointRequests);
 	// Currenlty we need to use the HTTP module in order to resolve service endpoints due to HTTPS
 	FHttpModule& HttpModule = FModuleManager::LoadModuleChecked<FHttpModule>("HTTP");
 	const int32 MaxAttempts = GIoDispatcherMaxHttpRetryCount;
@@ -201,17 +205,18 @@ void FDistributionEndpoints::IssueRequests()
 				[this, &ResolveRequest, MaxAttempts]
 				(FHttpRequestPtr, FHttpResponsePtr Response, bool bOk)
 				{
+					LLM_SCOPE(ELLMTag::FileSystem);
 					FHttpRequestPtr Request = MoveTemp(ResolveRequest.HttpRequest);
 					if (Response->GetResponseCode() != 200)
 					{
 						if (++ResolveRequest.RetryCount <= MaxAttempts)
 						{
 							Request->OnProcessRequestComplete().Unbind();
-							return IssueRequests();
+							return IssueEndpointRequests();
 						}
 					}
 
-					CompleteRequest(ResolveRequest, Response);
+					CompleteEndpointRequest(ResolveRequest, Response);
 				});
 
 			ResolveRequest.HttpRequest = HttpRequest;
@@ -225,7 +230,7 @@ void FDistributionEndpoints::IssueRequests()
 	}
 }
 
-void FDistributionEndpoints::CancelRequests()
+void FDistributionEndpoints::CancelEndpointRequests()
 {
 	TArray<FHttpRequestPtr, TInlineAllocator<2>> HttpRequests;
 	{
@@ -250,9 +255,9 @@ void FDistributionEndpoints::CancelRequests()
 	}
 }
 
-void FDistributionEndpoints::CompleteRequest(FResolveRequest& ResolveRequest, FHttpResponsePtr HttpResponse)
+void FDistributionEndpoints::CompleteEndpointRequest(FResolveRequest& ResolveRequest, FHttpResponsePtr HttpResponse)
 {
-	TRACE_CPUPROFILER_EVENT_SCOPE(FOnDemandIoBackend::CompleteDistributionRequest);
+	TRACE_CPUPROFILER_EVENT_SCOPE(FOnDemandIoBackend::CompleteEndpointRequest);
 
 	using FJsonValuePtr = TSharedPtr<FJsonValue>;
 	using FJsonObjPtr = TSharedPtr<FJsonObject>;
@@ -607,6 +612,8 @@ void FOnDemandIoStore::AddDeferredContainers()
 
 void FOnDemandIoStore::OnEncryptionKeyAdded(const FGuid& Id, const FAES::FAESKey& Key)
 {
+	LLM_SCOPE(ELLMTag::FileSystem);
+	TRACE_CPUPROFILER_EVENT_SCOPE(FOnDemandIoBackend::OnEncryptionKeyAdded);
 	AddDeferredContainers();
 }
 
@@ -1140,6 +1147,8 @@ FOnDemandIoBackend::~FOnDemandIoBackend()
 
 void FOnDemandIoBackend::Initialize(TSharedRef<const FIoDispatcherBackendContext> Context)
 {
+	TRACE_CPUPROFILER_EVENT_SCOPE(FOnDemandIoBackend::Initialize);
+	LLM_SCOPE(ELLMTag::FileSystem);
 	UE_LOG(LogIas, Log, TEXT("Initializing on demand I/O dispatcher backend"));
 	BackendContext = Context;
 	DistributionEndpoints.ResolveDeferredEndpoints();
@@ -1162,6 +1171,7 @@ void FOnDemandIoBackend::Shutdown()
 
 void FOnDemandIoBackend::CompleteRequest(FChunkRequest* ChunkRequest)
 {
+	LLM_SCOPE(ELLMTag::FileSystem);
 	TRACE_CPUPROFILER_EVENT_SCOPE(FOnDemandIoBackend::CompleteRequest);
 	check(ChunkRequest != nullptr);
 	const bool bCancelled = ChunkRequest->CancellationToken.IsCancelled();
@@ -1280,6 +1290,7 @@ bool FOnDemandIoBackend::Resolve(FIoRequestImpl* Request)
 	const ETaskPriority TaskPriority = ChunkRequest->Priority > IoDispatcherPriority_Medium ? ETaskPriority::High : ETaskPriority::Normal;
 	Launch(UE_SOURCE_LOCATION, [this, ChunkRequest]()
 	{
+		LLM_SCOPE(ELLMTag::FileSystem);
 		TRACE_CPUPROFILER_EVENT_SCOPE(FOnDemandIoBackend::CompleteOrEnqueueHttpRequest);
 		if (ChunkRequest->CacheTask.IsValid())
 		{
@@ -1406,6 +1417,7 @@ FIoStatus FOnDemandIoBackend::MountDeferredEndpoints(const FString& Distribution
 
 void FOnDemandIoBackend::Mount(const FOnDemandEndpoint& Endpoint)
 {
+	LLM_SCOPE(ELLMTag::FileSystem);
 	TRACE_CPUPROFILER_EVENT_SCOPE(FOnDemandIoBackend::Mount);
 
 	if ((Endpoint.DistributionUrl.IsEmpty() && Endpoint.ServiceUrl.IsEmpty()) || Endpoint.TocPath.IsEmpty())
@@ -1423,7 +1435,7 @@ void FOnDemandIoBackend::Mount(const FOnDemandEndpoint& Endpoint)
 				DeferredEndpoints.Add(Endpoint);
 			}
 
-			DistributionEndpoints.Resolve(Endpoint.DistributionUrl, [this](const FString& DistributionUrl, TConstArrayView<FString> SerivceUrls)
+			DistributionEndpoints.ResolveEndpoints(Endpoint.DistributionUrl, [this](const FString& DistributionUrl, TConstArrayView<FString> SerivceUrls)
 			{
 				if (FIoStatus Status = MountDeferredEndpoints(DistributionUrl, SerivceUrls); !Status.IsOk())
 				{
@@ -1530,6 +1542,8 @@ FIoStatus FOnDemandIoBackend::AddToc(const FOnDemandEndpoint& Endpoint)
 
 uint32 FOnDemandIoBackend::Run()
 {
+	LLM_SCOPE(ELLMTag::FileSystem);
+
 	const int32 MaxConcurrentRequests = HttpClient->MaxConnectionCount();
 	FChunkRequest* NextChunkRequest = nullptr;
 	int32 NumConcurrentRequests = 0;
