Add a new cmdline option to UnrealPak allowing all IAS content to be primed on an endpoint, so that testers can have the same experience as an end users.

#rnx

### New Command
- Run UnrealPak.exe with the commandline "-PrimeEndPoint=ABC" where ABC is a path to a IoStoreOnDemand.ini containing the info about the end point to prime.
- The command will mount the correct backend, connect to the end point, download the utoc files there and then attempt to download each chunk present so that they will be cached by the CDN endpoint.
- The command does not use the IAS filecache and should not affect the state of the local machine.
- We also make use of the existing IAS code, so this path can also be used for profiling when the IAS requests are saturated etc.

### Code Changes

- The code for parsing the ini file in FIoStoreOnDemandModule::StartupModule has been moved to common utility code so that it can be reused.

diff --git a/Engine/Source/Developer/IoStoreUtilities/Private/IoStoreUtilities.cpp b/Engine/Source/Developer/IoStoreUtilities/Private/IoStoreUtilities.cpp
--- a/Engine/Source/Developer/IoStoreUtilities/Private/IoStoreUtilities.cpp
+++ b/Engine/Source/Developer/IoStoreUtilities/Private/IoStoreUtilities.cpp
@@ -8599,3 +8599,14 @@ bool DownloadIoStoreContainerFiles(const TCHAR* TocPath)
 
 	return Status.IsOk();
 }
+
+bool PrimeEndPoint(FStringView IoStoreOnDemandIniPath)
+{
+	FIoStatus Status = UE::PrimeEndPoint(IoStoreOnDemandIniPath);
+	if (Status.IsOk() == false)
+	{
+		UE_LOG(LogIoStore, Error, TEXT("Failed to prime end points, reason '%s'"), *Status.ToString());
+	}
+
+	return Status.IsOk();
+}
diff --git a/Engine/Source/Developer/IoStoreUtilities/Public/IoStoreUtilities.h b/Engine/Source/Developer/IoStoreUtilities/Public/IoStoreUtilities.h
--- a/Engine/Source/Developer/IoStoreUtilities/Public/IoStoreUtilities.h
+++ b/Engine/Source/Developer/IoStoreUtilities/Public/IoStoreUtilities.h
@@ -56,3 +56,10 @@ IOSTOREUTILITIES_API bool SignIoStoreContainer(const TCHAR* InContainerFilename,
 IOSTOREUTILITIES_API bool UploadIoStoreContainerFiles(const TCHAR* ContainerPathOrWildcard); 
 
 IOSTOREUTILITIES_API bool DownloadIoStoreContainerFiles(const TCHAR* TocPath);
+
+/**
+ * Attempts to connect to the endpoint in the provided ini file, download all
+ * available tocs and then download all chunks in those tocs. This will
+ * pre-warm the CDN and allow testing closer to the end users experience.
+ */
+IOSTOREUTILITIES_API bool PrimeEndPoint(FStringView IoStoreOnDemandIniPath);
diff --git a/Engine/Source/Developer/PakFileUtilities/Private/PakFileUtilities.cpp b/Engine/Source/Developer/PakFileUtilities/Private/PakFileUtilities.cpp
--- a/Engine/Source/Developer/PakFileUtilities/Private/PakFileUtilities.cpp
+++ b/Engine/Source/Developer/PakFileUtilities/Private/PakFileUtilities.cpp
@@ -5467,6 +5467,11 @@ bool ExecuteUnrealPak(const TCHAR* CmdLine)
 		{
 			return DownloadIoStoreContainerFiles(*IoStoreArg) == 0;
 		}
+
+		if (FParse::Value(CmdLine, TEXT("-PrimeEndPoint="), IoStoreArg))
+		{
+			return PrimeEndPoint(*IoStoreArg) == 0;
+		}
 	}
 
 	// Parse all the non-option arguments from the command line
diff --git a/Engine/Source/Runtime/Experimental/IoStoreOnDemand/Private/IoStoreOnDemand.cpp b/Engine/Source/Runtime/Experimental/IoStoreOnDemand/Private/IoStoreOnDemand.cpp
--- a/Engine/Source/Runtime/Experimental/IoStoreOnDemand/Private/IoStoreOnDemand.cpp
+++ b/Engine/Source/Runtime/Experimental/IoStoreOnDemand/Private/IoStoreOnDemand.cpp
@@ -14,6 +14,7 @@
 #include "Misc/CommandLine.h"
 #include "Misc/ConfigCacheIni.h"
 #include "Misc/CoreDelegates.h"
+#include "Misc/FileHelper.h"
 #include "Misc/Parse.h"
 #include "Modules/ModuleManager.h"
 #include "Serialization/Archive.h"
@@ -27,7 +28,6 @@
 #include "Async/Async.h"
 #include "HAL/PlatformFileManager.h"
 #include "IO/IoStore.h"
-#include "Misc/FileHelper.h"
 #include "S3/S3Client.h"
 #endif // (PLATFORM_DESKTOP && (IS_PROGRAM || WITH_EDITOR))
 
@@ -86,6 +86,71 @@ static bool ParseEncryptionKeyParam(const FString& Param, FGuid& OutKeyGuid, FAE
 	return false;
 }
 
+////////////////////////////////////////////////////////////////////////////////
+static bool ApplyEncryptionKeyFromString(const FString& GuidKeyPair)
+{
+	FGuid KeyGuid;
+	FAES::FAESKey Key;
+
+	if (ParseEncryptionKeyParam(GuidKeyPair, KeyGuid, Key))
+	{
+		// TODO: PAK and I/O store should share key manager
+		UE::FEncryptionKeyManager::Get().AddKey(KeyGuid, Key);
+		FCoreDelegates::GetRegisterEncryptionKeyMulticastDelegate().Broadcast(KeyGuid, Key);
+
+		return true;
+	}
+	else
+	{
+		return false;
+	}
+}
+
+////////////////////////////////////////////////////////////////////////////////
+static bool ParseConfigFile(const FString& ConfigPath, UE::FOnDemandEndpoint& OutEndpoint)
+{
+	FString ConfigContent;
+	if (!FFileHelper::LoadFileToString(ConfigContent, &IPlatformFile::GetPlatformPhysical(), *ConfigPath))
+	{
+		return false;
+	}
+
+	if (!ConfigContent.IsEmpty())
+	{
+		const FString ConfigFileName = FPaths::GetCleanFilename(ConfigPath);
+
+		FConfigFile Config;
+		Config.ProcessInputFileContents(ConfigContent, ConfigFileName);
+
+		Config.GetString(TEXT("Endpoint"), TEXT("DistributionUrl"), OutEndpoint.DistributionUrl);
+		Config.GetString(TEXT("Endpoint"), TEXT("ServiceUrl"), OutEndpoint.ServiceUrl);
+		Config.GetString(TEXT("Endpoint"), TEXT("TocPath"), OutEndpoint.TocPath);
+
+		if (OutEndpoint.DistributionUrl.EndsWith(TEXT("/")))
+		{
+			OutEndpoint.DistributionUrl = OutEndpoint.DistributionUrl.Left(OutEndpoint.DistributionUrl.Len() - 1);
+		}
+
+		if (OutEndpoint.ServiceUrl.EndsWith(TEXT("/")))
+		{
+			OutEndpoint.ServiceUrl = OutEndpoint.DistributionUrl.Left(OutEndpoint.ServiceUrl.Len() - 1);
+		}
+
+		if (OutEndpoint.TocPath.StartsWith(TEXT("/")))
+		{
+			OutEndpoint.TocPath.RightChopInline(1);
+		}
+
+		FString ContentKey;
+		if (Config.GetString(TEXT("Endpoint"), TEXT("ContentKey"), ContentKey))
+		{
+			ApplyEncryptionKeyFromString(ContentKey);
+		}
+	}
+
+	return OutEndpoint.IsValid();
+}
+
 ////////////////////////////////////////////////////////////////////////////////
 static FFileIoCacheConfig GetFileIoCacheConfig(const TCHAR* CommandLine)
 {
@@ -1255,6 +1320,81 @@ FIoStatus DownloadContainerFiles(const FIoStoreDownloadParams& DownloadParams, c
 
 	return FIoStatus::Ok;
 }
+
+class FMyRequest : public FIoRequestImpl
+{
+public:
+
+	static FMyRequest* CreateRequest(const FIoChunkId& ChunkId)
+	{
+		return nullptr;
+	}
+};
+
+FIoStatus PrimeEndPoint(FStringView IoStoreOnDemandIniPath)
+{
+	TRACE_CPUPROFILER_EVENT_SCOPE(PrimeEndPoint);
+
+	using namespace UE::IO::Private;
+
+	UE::FOnDemandEndpoint EndPoint;
+	if (!ParseConfigFile(FString(IoStoreOnDemandIniPath), EndPoint))
+	{
+		return FIoStatus(EIoErrorCode::Unknown, TEXT("Failed to parse config file"));
+	}
+
+	EndPoint.EndpointType = UE::EOnDemandEndpointType::CDN;
+
+	TSharedPtr <UE::IOnDemandIoDispatcherBackend> Backend = UE::MakeOnDemandIoDispatcherBackend(nullptr);
+	Backend->Mount(EndPoint);
+
+	check(!FIoDispatcher::IsInitialized()); // Assume this is only run in standalone programs that have not
+											// yet setup the IoDispatcher
+	FIoDispatcher::Initialize();
+	FIoDispatcher::InitializePostSettings();
+	ON_SCOPE_EXIT
+	{
+		FIoDispatcher::Shutdown();
+	};
+
+	FIoDispatcher::Get().Mount(Backend.ToSharedRef(), MAX_int32);
+
+	if (!Backend->FlushDeferedEndPoints(30.0))
+	{
+		// TODO: Real error value
+		return FIoStatus(EIoErrorCode::Unknown, TEXT("Unable to connect to endpoint"));
+	}
+
+	UE_LOG(LogIas, Display, TEXT("Finding on demand chunks..."));
+	TArray<FIoChunkId> Chunks = Backend->GetAllChunkIds();
+	
+	UE_LOG(LogIas, Display, TEXT("Found %d chunks"), Chunks.Num());
+
+	std::atomic<int32> CompletedRequests = 0;
+	FIoBatch Batch;
+	for (const FIoChunkId& ChunkId : Chunks)
+	{
+		Batch.ReadWithCallback(ChunkId, FIoReadOptions(), IoDispatcherPriority_Medium,
+			[&](TIoStatusOr<FIoBuffer> Result)
+			{
+				CompletedRequests++;
+			});
+	}
+
+	UE_LOG(LogIas, Display, TEXT("Priming all chunks for endpoint..."));
+
+	FEventRef BatchCompleted;
+	Batch.IssueAndTriggerEvent(BatchCompleted.Get());
+	
+	while (!BatchCompleted->Wait(FTimespan::FromSeconds(5.0)))
+	{
+		UE_LOG(LogIas, Display, TEXT("Completed %d/%d"), CompletedRequests.load(), Chunks.Num());
+	}
+
+	UE_LOG(LogIas, Display, TEXT("All chunks primed!"));
+
+	return FIoStatus::Ok;
+}
 #endif // (PLATFORM_DESKTOP && (IS_PROGRAM || WITH_EDITOR))
 
 } // namespace UE
@@ -1317,14 +1457,7 @@ void FIoStoreOnDemandModule::StartupModule()
 		FString EncryptionKey;
 		if (FParse::Value(CommandLine, TEXT("Ias.EncryptionKey="), EncryptionKey))
 		{
-			FGuid KeyGuid;
-			FAES::FAESKey Key;
-			if (ParseEncryptionKeyParam(EncryptionKey, KeyGuid, Key))
-			{
-				// TODO: PAK and I/O store should share key manager
-				UE::FEncryptionKeyManager::Get().AddKey(KeyGuid, Key);
-				FCoreDelegates::GetRegisterEncryptionKeyMulticastDelegate().Broadcast(KeyGuid, Key);
-			}
+			ApplyEncryptionKeyFromString(EncryptionKey);
 		}
 	}
 
@@ -1332,46 +1465,9 @@ void FIoStoreOnDemandModule::StartupModule()
 	{
 		Endpoint = UE::FOnDemandEndpoint();
 		FString ConfigFileName = TEXT("IoStoreOnDemand.ini");
-		FString ConfigPath = FPaths::Combine(TEXT("Cloud"), ConfigFileName);
-		FString ConfigContent = FPlatformMisc::LoadTextFileFromPlatformPackage(ConfigPath);
-
-		if (ConfigContent.Len())
-		{
-			FConfigFile Config;
-			Config.ProcessInputFileContents(ConfigContent, ConfigFileName);
-
-			Config.GetString(TEXT("Endpoint"), TEXT("DistributionUrl"), Endpoint.DistributionUrl);
-			Config.GetString(TEXT("Endpoint"), TEXT("ServiceUrl"), Endpoint.ServiceUrl);
-			Config.GetString(TEXT("Endpoint"), TEXT("TocPath"), Endpoint.TocPath);
-			
-			if (Endpoint.DistributionUrl.EndsWith(TEXT("/")))
-			{
-				Endpoint.DistributionUrl = Endpoint.DistributionUrl.Left(Endpoint.DistributionUrl.Len() - 1);
-			}
-			
-			if (Endpoint.ServiceUrl.EndsWith(TEXT("/")))
-			{
-				Endpoint.ServiceUrl = Endpoint.DistributionUrl.Left(Endpoint.ServiceUrl.Len() - 1);
-			}
-
-			if (Endpoint.TocPath.StartsWith(TEXT("/")))
-			{
-				Endpoint.TocPath.RightChopInline(1);
-			}
+		FString ConfigPath = FPaths::Combine(FPaths::RootDir(), TEXT("Cloud"), ConfigFileName);
 
-			FString ContentKey;
-			if (Config.GetString(TEXT("Endpoint"), TEXT("ContentKey"), ContentKey))
-			{
-				FGuid KeyGuid;
-				FAES::FAESKey Key;
-				if (ParseEncryptionKeyParam(ContentKey, KeyGuid, Key))
-				{
-					// TODO: PAK and I/O store should share key manager
-					UE::FEncryptionKeyManager::Get().AddKey(KeyGuid, Key);
-					FCoreDelegates::GetRegisterEncryptionKeyMulticastDelegate().Broadcast(KeyGuid, Key);
-				}
-			}
-		}
+		ParseConfigFile(ConfigPath, Endpoint);
 	}
 
 	if (!Endpoint.IsValid())
diff --git a/Engine/Source/Runtime/Experimental/IoStoreOnDemand/Private/OnDemandIoDispatcherBackend.cpp b/Engine/Source/Runtime/Experimental/IoStoreOnDemand/Private/OnDemandIoDispatcherBackend.cpp
--- a/Engine/Source/Runtime/Experimental/IoStoreOnDemand/Private/OnDemandIoDispatcherBackend.cpp
+++ b/Engine/Source/Runtime/Experimental/IoStoreOnDemand/Private/OnDemandIoDispatcherBackend.cpp
@@ -177,7 +177,7 @@ void FDistributionEndpoints::ResolveDeferredEndpoints()
 void FDistributionEndpoints::IssueEndpointRequests()
 {
 	TRACE_CPUPROFILER_EVENT_SCOPE(FOnDemandIoBackend::IssueEndpointRequests);
-	// Currenlty we need to use the HTTP module in order to resolve service endpoints due to HTTPS
+	// Currently we need to use the HTTP module in order to resolve service endpoints due to HTTPS
 	FHttpModule& HttpModule = FModuleManager::LoadModuleChecked<FHttpModule>("HTTP");
 	const int32 MaxAttempts = GIoDispatcherMaxHttpRetryCount;
 
@@ -475,6 +475,8 @@ public:
 	TIoStatusOr<uint64> GetChunkSize(const FIoChunkId& ChunkId);
 	FChunkInfo GetChunkInfo(const FIoChunkId& ChunkId);
 
+	TArray<FIoChunkId> GetAllChunkIds(bool bIncludeOptionalChunks);
+
 private:
 	void AddDeferredContainers();
 	void OnEncryptionKeyAdded(const FGuid& Id, const FAES::FAESKey& Key);
@@ -578,6 +580,33 @@ FOnDemandIoStore::FChunkInfo FOnDemandIoStore::GetChunkInfo(const FIoChunkId& Ch
 	return {};
 }
 
+TArray<FIoChunkId> FOnDemandIoStore::GetAllChunkIds(bool bIncludeOptionalChunks)
+{
+	FReadScopeLock _(Lock);
+
+	TArray<FIoChunkId> ChunkIds;
+	int32 NumChunks = 0;
+	for (const FContainer* Container : RegisteredContainers)
+	{
+		NumChunks += Container->TocEntries.Num();
+	}
+
+	ChunkIds.Reserve(NumChunks);
+
+	for (const FContainer* Container : RegisteredContainers)
+	{
+		for (const TPair<FIoChunkId, FTocEntry>& Entry : Container->TocEntries)
+		{
+			if (bIncludeOptionalChunks || Entry.Key.GetChunkType() != EIoChunkType::OptionalBulkData)
+			{
+				ChunkIds.Add(Entry.Key);
+			}
+		}
+	}
+
+	return ChunkIds;
+}
+
 void FOnDemandIoStore::AddDeferredContainers()
 {
 	FWriteScopeLock _(Lock);
@@ -1104,6 +1133,11 @@ public:
 	virtual void SetBulkOptionalEnabled(bool bInEnabled) override;
 	virtual void SetEnabled(bool bInEnabled) override;
 
+#if IS_PROGRAM || WITH_EDITOR
+	virtual bool FlushDeferedEndPoints(double TimeOut = 0.0) override;
+	virtual TArray<FIoChunkId> GetAllChunkIds() override;
+#endif // IS_PROGRAM || WITH_EDITOR
+
 	// Runnable
 	virtual bool Init() override { return true; }
 	virtual void Stop() override { bStopRequested = true; }
@@ -1465,6 +1499,45 @@ void FOnDemandIoBackend::SetEnabled(bool bInEnabled)
 	bEnabled = bInEnabled;
 }
 
+#if IS_PROGRAM || WITH_EDITOR
+bool FOnDemandIoBackend::FlushDeferedEndPoints(double TimeOut)
+{
+	TRACE_CPUPROFILER_EVENT_SCOPE(FOnDemandIoBackend::FlushDeferedEndPoints);
+
+	FHttpManager& HttpManager = FHttpModule::Get().GetHttpManager();
+
+	const double StartTime = FPlatformTime::Seconds();
+
+	while (!DeferredEndpoints.IsEmpty())
+	{
+		HttpManager.Tick(0.0);
+		FPlatformProcess::SleepNoStats(0.0f);
+
+		if (TimeOut > 0.0 && (FPlatformTime::Seconds() - StartTime) > TimeOut)
+		{
+			return false;
+		}
+	}
+
+	return true;
+}
+
+TArray<FIoChunkId> FOnDemandIoBackend::GetAllChunkIds()
+{
+	TRACE_CPUPROFILER_EVENT_SCOPE(FOnDemandIoBackend::GetAllChunkIds);
+
+	if (!IoStore.IsValid() || !bEnabled)
+	{
+		return TArray<FIoChunkId>();
+	}
+
+	const bool bAllowOptional = GIoDispatcherBulkOptionalEnabled && bEnableBulkOptional;
+
+	return IoStore->GetAllChunkIds(bAllowOptional);
+}
+
+#endif // IS_PROGRAM || WITH_EDITOR
+
 TIoStatusOr<FOnDemandToc> FOnDemandIoBackend::GetToc(FHttpClient& HttpClient, const FString& TocPath)
 {
 	TRACE_CPUPROFILER_EVENT_SCOPE(FOnDemandIoBackend::GetToc);
diff --git a/Engine/Source/Runtime/Experimental/IoStoreOnDemand/Private/OnDemandIoDispatcherBackend.h b/Engine/Source/Runtime/Experimental/IoStoreOnDemand/Private/OnDemandIoDispatcherBackend.h
--- a/Engine/Source/Runtime/Experimental/IoStoreOnDemand/Private/OnDemandIoDispatcherBackend.h
+++ b/Engine/Source/Runtime/Experimental/IoStoreOnDemand/Private/OnDemandIoDispatcherBackend.h
@@ -36,6 +36,11 @@ public:
 	virtual void Mount(const FOnDemandEndpoint& Endpoint) = 0;
 	virtual void SetBulkOptionalEnabled(bool bInEnabled) = 0;
 	virtual void SetEnabled(bool bInEnabled) = 0;
+
+#if IS_PROGRAM || WITH_EDITOR
+	virtual bool FlushDeferedEndPoints(double TimeOut) = 0;
+	virtual TArray<FIoChunkId> GetAllChunkIds() = 0;
+#endif // IS_PROGRAM || WITH_EDITOR
 };
 
 UE_API TSharedPtr<IOnDemandIoDispatcherBackend> MakeOnDemandIoDispatcherBackend(TSharedPtr<IIoCache> Cache);
diff --git a/Engine/Source/Runtime/Experimental/IoStoreOnDemand/Public/IO/IoStoreOnDemand.h b/Engine/Source/Runtime/Experimental/IoStoreOnDemand/Public/IO/IoStoreOnDemand.h
--- a/Engine/Source/Runtime/Experimental/IoStoreOnDemand/Public/IO/IoStoreOnDemand.h
+++ b/Engine/Source/Runtime/Experimental/IoStoreOnDemand/Public/IO/IoStoreOnDemand.h
@@ -164,6 +164,8 @@ struct FIoStoreDownloadParams
 
 UE_API FIoStatus DownloadContainerFiles(const FIoStoreDownloadParams& DownloadParams, const FString& TocPath);
 
+UE_API FIoStatus PrimeEndPoint(FStringView IoStoreOnDemandIniPath);
+
 #endif // (IS_PROGRAM || WITH_EDITOR)
 
 } // namespace UE
