Iris 1/2
- Fixed issue with gameplay code holding stale handles to NetObjectGroups after restarting ReplicationSystem due to HotFix.
- Added epoch to NetObjectGroupHandles to be able to detect stale groups

diff --git a/Engine/Plugins/Runtime/ReplicationSystemTestPlugin/Source/Private/Tests/ReplicationState/TestFastArrayReplicationState.cpp b/Engine/Plugins/Runtime/ReplicationSystemTestPlugin/Source/Private/Tests/ReplicationState/TestFastArrayReplicationState.cpp
--- a/Engine/Plugins/Runtime/ReplicationSystemTestPlugin/Source/Private/Tests/ReplicationState/TestFastArrayReplicationState.cpp
+++ b/Engine/Plugins/Runtime/ReplicationSystemTestPlugin/Source/Private/Tests/ReplicationState/TestFastArrayReplicationState.cpp
@@ -807,8 +807,8 @@ UE_NET_TEST_FIXTURE(FReplicationSystemServerClientTestFixture, TestPostReplicate
 	FastArrayEditor.Add(Item);
 
 	// Make sure that references are not replicated to client
-	Server->ReplicationSystem->AddToGroup(NotReplicatedNetObjectGroupHandle, ServerReferencedObjectA->NetRefHandle);
-	Server->ReplicationSystem->AddToGroup(NotReplicatedNetObjectGroupHandle, ServerReferencedObjectB->NetRefHandle);
+	Server->ReplicationSystem->AddToGroup(Server->GetReplicationSystem()->GetNotReplicatedNetObjectGroup(), ServerReferencedObjectA->NetRefHandle);
+	Server->ReplicationSystem->AddToGroup(Server->GetReplicationSystem()->GetNotReplicatedNetObjectGroup(), ServerReferencedObjectB->NetRefHandle);
 	
 	// Replicate
 	Server->PreSendUpdate();
@@ -837,7 +837,7 @@ UE_NET_TEST_FIXTURE(FReplicationSystemServerClientTestFixture, TestPostReplicate
 	ClientFastArray.bPostReplicatedReceiveWasHitWithUnresolvedReferences = false;
 
 	// Enable replication for ServerReferenceObjectA
-	Server->ReplicationSystem->RemoveFromGroup(NotReplicatedNetObjectGroupHandle, ServerReferencedObjectA->NetRefHandle);
+	Server->ReplicationSystem->RemoveFromGroup(Server->GetReplicationSystem()->GetNotReplicatedNetObjectGroup(), ServerReferencedObjectA->NetRefHandle);
 
 	// Replicate
 	Server->PreSendUpdate();
@@ -856,7 +856,7 @@ UE_NET_TEST_FIXTURE(FReplicationSystemServerClientTestFixture, TestPostReplicate
 	ClientFastArray.bPostReplicatedReceiveWasHitWithUnresolvedReferences = false;
 
 	// Enable replication for ServerReferenceObjectB 
-	Server->ReplicationSystem->RemoveFromGroup(NotReplicatedNetObjectGroupHandle, ServerReferencedObjectB->NetRefHandle);
+	Server->ReplicationSystem->RemoveFromGroup(Server->GetReplicationSystem()->GetNotReplicatedNetObjectGroup(), ServerReferencedObjectB->NetRefHandle);
 
 	// Replicate
 	Server->PreSendUpdate();
diff --git a/Engine/Plugins/Runtime/ReplicationSystemTestPlugin/Source/Private/Tests/ReplicationSystem/Filtering/TestGroups.cpp b/Engine/Plugins/Runtime/ReplicationSystemTestPlugin/Source/Private/Tests/ReplicationSystem/Filtering/TestGroups.cpp
--- a/Engine/Plugins/Runtime/ReplicationSystemTestPlugin/Source/Private/Tests/ReplicationSystem/Filtering/TestGroups.cpp
+++ b/Engine/Plugins/Runtime/ReplicationSystemTestPlugin/Source/Private/Tests/ReplicationSystem/Filtering/TestGroups.cpp
@@ -82,7 +82,7 @@ UE_NET_TEST_FIXTURE(FTestGroupsFixture, PublicGroupAPI)
 	{
 		FNetObjectGroupHandle GroupHandle = ReplicationSystem->CreateGroup();
 
-		UE_NET_ASSERT_FALSE(ReplicationSystem->IsValidGroup(0U));
+		UE_NET_ASSERT_FALSE(ReplicationSystem->IsValidGroup(FNetObjectGroupHandle()));
 		UE_NET_ASSERT_TRUE(ReplicationSystem->IsValidGroup(GroupHandle));
 
 		// Destroy
@@ -157,7 +157,7 @@ UE_NET_TEST_FIXTURE(FTestGroupsFixture, GroupFilterAPI)
 
 	// Create group
 	FNetObjectGroupHandle GroupHandle = Server->ReplicationSystem->CreateGroup();
-	UE_NET_ASSERT_NE(0U, (uint32)GroupHandle);	
+	UE_NET_ASSERT_EQ(true, GroupHandle.IsValid());	
 
 	// Add Objects to group
 	Server->ReplicationSystem->AddToGroup(GroupHandle, ServerObject0->NetRefHandle);
@@ -170,7 +170,7 @@ UE_NET_TEST_FIXTURE(FTestGroupsFixture, GroupFilterAPI)
 
 	// Create another group
 	FNetObjectGroupHandle GroupHandle2 = Server->ReplicationSystem->CreateGroup();
-	UE_NET_ASSERT_NE(0U, (uint32)GroupHandle2);	
+	UE_NET_ASSERT_EQ(true, GroupHandle2.IsValid());	
 
 	// Add Objects to group
 	Server->ReplicationSystem->AddToGroup(GroupHandle2, ServerObject3->NetRefHandle);
@@ -181,7 +181,7 @@ UE_NET_TEST_FIXTURE(FTestGroupsFixture, GroupFilterAPI)
 
 	// Create another group
 	FNetObjectGroupHandle GroupHandle3 = Server->ReplicationSystem->CreateGroup();
-	UE_NET_ASSERT_NE(0U, (uint32)GroupHandle3);	
+	UE_NET_ASSERT_TRUE(GroupHandle3.IsValid());	
 
 	// Update filters
 	Filter();
@@ -290,8 +290,8 @@ UE_NET_TEST_FIXTURE(FTestGroupsFixture, GroupFilterAPIObjectMemberOfMoreThanOneG
 	// Create groups
 	FNetObjectGroupHandle GroupHandle0 = Server->ReplicationSystem->CreateGroup();
 	FNetObjectGroupHandle GroupHandle1 = Server->ReplicationSystem->CreateGroup();
-	UE_NET_ASSERT_NE(0U, (uint32)GroupHandle0);	
-	UE_NET_ASSERT_NE(0U, (uint32)GroupHandle1);	
+	UE_NET_ASSERT_TRUE(GroupHandle0.IsValid());	
+	UE_NET_ASSERT_TRUE(GroupHandle1.IsValid());	
 
 	// Add Objects to groups
 	Server->ReplicationSystem->AddToGroup(GroupHandle0, ServerHandle);
@@ -348,8 +348,8 @@ UE_NET_TEST_FIXTURE(FTestGroupsFixture, GroupFilterAPIAddThenRemoveFilterTrait)
 	// Create groups
 	FNetObjectGroupHandle GroupHandle0 = Server->ReplicationSystem->CreateGroup();
 	FNetObjectGroupHandle GroupHandle1 = Server->ReplicationSystem->CreateGroup();
-	UE_NET_ASSERT_NE(0U, (uint32)GroupHandle0);
-	UE_NET_ASSERT_NE(0U, (uint32)GroupHandle1);
+	UE_NET_ASSERT_TRUE(GroupHandle0.IsValid());	
+	UE_NET_ASSERT_TRUE(GroupHandle1.IsValid());	
 
 	// Add Objects to groups
 	Server->ReplicationSystem->AddToGroup(GroupHandle0, ServerHandle);
@@ -410,9 +410,9 @@ UE_NET_TEST_FIXTURE(FTestGroupsFixture, GroupFilterTestObjectListViaTrait)
 	FNetObjectGroupHandle GroupHandle0 = Server->ReplicationSystem->CreateGroup();
 	FNetObjectGroupHandle GroupHandle1 = Server->ReplicationSystem->CreateGroup();
 	FNetObjectGroupHandle GroupHandle2 = Server->ReplicationSystem->CreateGroup();
-	UE_NET_ASSERT_NE(0U, (uint32)GroupHandle0);
-	UE_NET_ASSERT_NE(0U, (uint32)GroupHandle1);
-	UE_NET_ASSERT_NE(0U, (uint32)GroupHandle2);
+	UE_NET_ASSERT_TRUE(GroupHandle0.IsValid());	
+	UE_NET_ASSERT_TRUE(GroupHandle1.IsValid());	
+	UE_NET_ASSERT_TRUE(GroupHandle2.IsValid());
 
 	// Add Objects to groups
 	Server->ReplicationSystem->AddToGroup(GroupHandle0, ServerHandle);
@@ -480,9 +480,9 @@ UE_NET_TEST_FIXTURE(FTestGroupsFixture, GroupFilterTestObjectListViaMembership)
 	FNetObjectGroupHandle GroupHandle0 = Server->ReplicationSystem->CreateGroup();
 	FNetObjectGroupHandle GroupHandle1 = Server->ReplicationSystem->CreateGroup();
 	FNetObjectGroupHandle GroupHandle2 = Server->ReplicationSystem->CreateGroup();
-	UE_NET_ASSERT_NE(0U, (uint32)GroupHandle0);
-	UE_NET_ASSERT_NE(0U, (uint32)GroupHandle1);
-	UE_NET_ASSERT_NE(0U, (uint32)GroupHandle2);
+	UE_NET_ASSERT_TRUE(GroupHandle0.IsValid());	
+	UE_NET_ASSERT_TRUE(GroupHandle1.IsValid());	
+	UE_NET_ASSERT_TRUE(GroupHandle2.IsValid());	
 
 	// Add Objects to groups
 	Server->ReplicationSystem->AddToGroup(GroupHandle0, ServerHandle);
@@ -548,7 +548,7 @@ UE_NET_TEST_FIXTURE(FTestGroupsFixture, GroupFilterAPINotFilteredGroup)
 	UReplicatedTestObject* ServerObject2 = Server->CreateObject(0, 0);
 
 	// Add Objects to NotFilteredGroup
-	Server->ReplicationSystem->AddToGroup(NotReplicatedNetObjectGroupHandle, ServerObject0->NetRefHandle);
+	Server->ReplicationSystem->AddToGroup(Server->GetReplicationSystem()->GetNotReplicatedNetObjectGroup(), ServerObject0->NetRefHandle);
 
 	// Update filters
 	Filter();
@@ -559,7 +559,7 @@ UE_NET_TEST_FIXTURE(FTestGroupsFixture, GroupFilterAPINotFilteredGroup)
 	UE_NET_ASSERT_TRUE(ObjectGroupFilter.IsAnyBitSet());
 
 	// Remove object from NotReplicatedNetObjectGroup
-	Server->ReplicationSystem->RemoveFromGroup(NotReplicatedNetObjectGroupHandle, ServerObject0->NetRefHandle);
+	Server->ReplicationSystem->RemoveFromGroup(Server->GetReplicationSystem()->GetNotReplicatedNetObjectGroup(), ServerObject0->NetRefHandle);
 
 	// Update filters
 	Filter();
diff --git a/Engine/Plugins/Runtime/ReplicationSystemTestPlugin/Source/Private/Tests/ReplicationSystem/TestCancelPendingDestroy.cpp b/Engine/Plugins/Runtime/ReplicationSystemTestPlugin/Source/Private/Tests/ReplicationSystem/TestCancelPendingDestroy.cpp
--- a/Engine/Plugins/Runtime/ReplicationSystemTestPlugin/Source/Private/Tests/ReplicationSystem/TestCancelPendingDestroy.cpp
+++ b/Engine/Plugins/Runtime/ReplicationSystemTestPlugin/Source/Private/Tests/ReplicationSystem/TestCancelPendingDestroy.cpp
@@ -26,12 +26,12 @@ UE_NET_TEST_FIXTURE(FTestCancelPendingDestroyFixture, TestCancelPendingDestroyDu
 	Server->PostSendUpdate();
 
 	// Filter out object to cause a PendingDestroy
-	Server->GetReplicationSystem()->AddToGroup(NotReplicatedNetObjectGroupHandle, ServerObject->NetRefHandle);
+	Server->GetReplicationSystem()->AddToGroup(Server->GetReplicationSystem()->GetNotReplicatedNetObjectGroup(), ServerObject->NetRefHandle);
 	Server->PreSendUpdate();
 	Server->PostSendUpdate();
 
 	// Remove object from filter to cause object to end up in CancelPendingDestroy
-	Server->GetReplicationSystem()->RemoveFromGroup(NotReplicatedNetObjectGroupHandle, ServerObject->NetRefHandle);
+	Server->GetReplicationSystem()->RemoveFromGroup(Server->GetReplicationSystem()->GetNotReplicatedNetObjectGroup(), ServerObject->NetRefHandle);
 	Server->PreSendUpdate();
 	Server->PostSendUpdate();
 
@@ -67,12 +67,12 @@ UE_NET_TEST_FIXTURE(FTestCancelPendingDestroyFixture, TestCancelPendingDestroyDu
 	Server->PostSendUpdate();
 
 	// Filter out object to cause a PendingDestroy
-	Server->GetReplicationSystem()->AddToGroup(NotReplicatedNetObjectGroupHandle, ServerObject->NetRefHandle);
+	Server->GetReplicationSystem()->AddToGroup(Server->GetReplicationSystem()->GetNotReplicatedNetObjectGroup(), ServerObject->NetRefHandle);
 	Server->PreSendUpdate();
 	Server->PostSendUpdate();
 
 	// Remove object from filter to cause object to end up in CancelPendingDestroy
-	Server->GetReplicationSystem()->RemoveFromGroup(NotReplicatedNetObjectGroupHandle, ServerObject->NetRefHandle);
+	Server->GetReplicationSystem()->RemoveFromGroup(Server->GetReplicationSystem()->GetNotReplicatedNetObjectGroup(), ServerObject->NetRefHandle);
 	Server->PreSendUpdate();
 	Server->PostSendUpdate();
 
@@ -104,12 +104,12 @@ UE_NET_TEST_FIXTURE(FTestCancelPendingDestroyFixture, TestCancelPendingDestroyDu
 	UE_NET_ASSERT_NE(Client->GetReplicationBridge()->GetReplicatedObject(ServerObject->NetRefHandle), nullptr);
 
 	// Filter out object to cause a PendingDestroy
-	Server->GetReplicationSystem()->AddToGroup(NotReplicatedNetObjectGroupHandle, ServerObject->NetRefHandle);
+	Server->GetReplicationSystem()->AddToGroup(Server->GetReplicationSystem()->GetNotReplicatedNetObjectGroup(), ServerObject->NetRefHandle);
 	Server->PreSendUpdate();
 	Server->PostSendUpdate();
 
 	// Remove object from filter to cause object to end up in CancelPendingDestroy
-	Server->GetReplicationSystem()->RemoveFromGroup(NotReplicatedNetObjectGroupHandle, ServerObject->NetRefHandle);
+	Server->GetReplicationSystem()->RemoveFromGroup(Server->GetReplicationSystem()->GetNotReplicatedNetObjectGroup(), ServerObject->NetRefHandle);
 	Server->PreSendUpdate();
 	Server->PostSendUpdate();
 
@@ -139,13 +139,13 @@ UE_NET_TEST_FIXTURE(FTestCancelPendingDestroyFixture, TestCancelPendingDestroyDu
 	Server->PostSendUpdate();
 
 	// Filter out object and write packet to cause object to end up in WaitOnDestroyConfirmation
-	Server->GetReplicationSystem()->AddToGroup(NotReplicatedNetObjectGroupHandle, ServerObject->NetRefHandle);
+	Server->GetReplicationSystem()->AddToGroup(Server->GetReplicationSystem()->GetNotReplicatedNetObjectGroup(), ServerObject->NetRefHandle);
 	Server->PreSendUpdate();
 	Server->SendUpdate(Client->ConnectionIdOnServer);
 	Server->PostSendUpdate();
 
 	// Remove object from filter to cause object to end up in CancelPendingDestroy
-	Server->GetReplicationSystem()->RemoveFromGroup(NotReplicatedNetObjectGroupHandle, ServerObject->NetRefHandle);
+	Server->GetReplicationSystem()->RemoveFromGroup(Server->GetReplicationSystem()->GetNotReplicatedNetObjectGroup(), ServerObject->NetRefHandle);
 	Server->PreSendUpdate();
 	Server->PostSendUpdate();
 
@@ -181,13 +181,13 @@ UE_NET_TEST_FIXTURE(FTestCancelPendingDestroyFixture, TestCancelPendingDestroyDu
 	Server->PostSendUpdate();
 
 	// Filter out object and write packet to cause object to end up in WaitOnDestroyConfirmation
-	Server->GetReplicationSystem()->AddToGroup(NotReplicatedNetObjectGroupHandle, ServerObject->NetRefHandle);
+	Server->GetReplicationSystem()->AddToGroup(Server->GetReplicationSystem()->GetNotReplicatedNetObjectGroup(), ServerObject->NetRefHandle);
 	Server->PreSendUpdate();
 	Server->SendUpdate(Client->ConnectionIdOnServer);
 	Server->PostSendUpdate();
 
 	// Remove object from filter to cause object to end up in CancelPendingDestroy, also modify a property 
-	Server->GetReplicationSystem()->RemoveFromGroup(NotReplicatedNetObjectGroupHandle, ServerObject->NetRefHandle);
+	Server->GetReplicationSystem()->RemoveFromGroup(Server->GetReplicationSystem()->GetNotReplicatedNetObjectGroup(), ServerObject->NetRefHandle);
 	ServerObject->IntA ^= 1;
 	Server->PreSendUpdate();
 	Server->PostSendUpdate();
@@ -223,17 +223,17 @@ UE_NET_TEST_FIXTURE(FTestCancelPendingDestroyFixture, TestCancelingCancelPending
 	Server->PostSendUpdate();
 
 	// Filter out object to cause a PendingDestroy
-	Server->GetReplicationSystem()->AddToGroup(NotReplicatedNetObjectGroupHandle, ServerObject->NetRefHandle);
+	Server->GetReplicationSystem()->AddToGroup(Server->GetReplicationSystem()->GetNotReplicatedNetObjectGroup(), ServerObject->NetRefHandle);
 	Server->PreSendUpdate();
 	Server->PostSendUpdate();
 
 	// Remove object from filter to cause object to end up in CancelPendingDestroy...
-	Server->GetReplicationSystem()->RemoveFromGroup(NotReplicatedNetObjectGroupHandle, ServerObject->NetRefHandle);
+	Server->GetReplicationSystem()->RemoveFromGroup(Server->GetReplicationSystem()->GetNotReplicatedNetObjectGroup(), ServerObject->NetRefHandle);
 	Server->PreSendUpdate();
 	Server->PostSendUpdate();
 
 	// ... and cancel that thought! We want the object to be destroyed after all.
-	Server->GetReplicationSystem()->AddToGroup(NotReplicatedNetObjectGroupHandle, ServerObject->NetRefHandle);
+	Server->GetReplicationSystem()->AddToGroup(Server->GetReplicationSystem()->GetNotReplicatedNetObjectGroup(), ServerObject->NetRefHandle);
 	Server->PreSendUpdate();
 	Server->PostSendUpdate();
 
@@ -265,17 +265,17 @@ UE_NET_TEST_FIXTURE(FTestCancelPendingDestroyFixture, TestCancelingCancelPending
 	Server->PostSendUpdate();
 
 	// Filter out object to cause a PendingDestroy
-	Server->GetReplicationSystem()->AddToGroup(NotReplicatedNetObjectGroupHandle, ServerObject->NetRefHandle);
+	Server->GetReplicationSystem()->AddToGroup(Server->GetReplicationSystem()->GetNotReplicatedNetObjectGroup(), ServerObject->NetRefHandle);
 	Server->PreSendUpdate();
 	Server->PostSendUpdate();
 
 	// Remove object from filter to cause object to end up in CancelPendingDestroy...
-	Server->GetReplicationSystem()->RemoveFromGroup(NotReplicatedNetObjectGroupHandle, ServerObject->NetRefHandle);
+	Server->GetReplicationSystem()->RemoveFromGroup(Server->GetReplicationSystem()->GetNotReplicatedNetObjectGroup(), ServerObject->NetRefHandle);
 	Server->PreSendUpdate();
 	Server->PostSendUpdate();
 
 	// ... and cancel that thought! We want the object to be destroyed after all.
-	Server->GetReplicationSystem()->AddToGroup(NotReplicatedNetObjectGroupHandle, ServerObject->NetRefHandle);
+	Server->GetReplicationSystem()->AddToGroup(Server->GetReplicationSystem()->GetNotReplicatedNetObjectGroup(), ServerObject->NetRefHandle);
 	Server->PreSendUpdate();
 	Server->PostSendUpdate();
 
@@ -311,17 +311,17 @@ UE_NET_TEST_FIXTURE(FTestCancelPendingDestroyFixture, TestCancelingCancelPending
 	UE_NET_ASSERT_NE(Client->GetReplicationBridge()->GetReplicatedObject(ServerObject->NetRefHandle), nullptr);
 
 	// Filter out object to cause a PendingDestroy
-	Server->GetReplicationSystem()->AddToGroup(NotReplicatedNetObjectGroupHandle, ServerObject->NetRefHandle);
+	Server->GetReplicationSystem()->AddToGroup(Server->GetReplicationSystem()->GetNotReplicatedNetObjectGroup(), ServerObject->NetRefHandle);
 	Server->PreSendUpdate();
 	Server->PostSendUpdate();
 
 	// Remove object from filter to cause object to end up in CancelPendingDestroy...
-	Server->GetReplicationSystem()->RemoveFromGroup(NotReplicatedNetObjectGroupHandle, ServerObject->NetRefHandle);
+	Server->GetReplicationSystem()->RemoveFromGroup(Server->GetReplicationSystem()->GetNotReplicatedNetObjectGroup(), ServerObject->NetRefHandle);
 	Server->PreSendUpdate();
 	Server->PostSendUpdate();
 
 	// ... and cancel that thought! We want the object to be destroyed after all.
-	Server->GetReplicationSystem()->AddToGroup(NotReplicatedNetObjectGroupHandle, ServerObject->NetRefHandle);
+	Server->GetReplicationSystem()->AddToGroup(Server->GetReplicationSystem()->GetNotReplicatedNetObjectGroup(), ServerObject->NetRefHandle);
 	Server->PreSendUpdate();
 	Server->PostSendUpdate();
 
@@ -348,18 +348,18 @@ UE_NET_TEST_FIXTURE(FTestCancelPendingDestroyFixture, TestCancelingCancelPending
 	Server->PostSendUpdate();
 
 	// Filter out object and write packet to cause object to end up in WaitOnDestroyConfirmation
-	Server->GetReplicationSystem()->AddToGroup(NotReplicatedNetObjectGroupHandle, ServerObject->NetRefHandle);
+	Server->GetReplicationSystem()->AddToGroup(Server->GetReplicationSystem()->GetNotReplicatedNetObjectGroup(), ServerObject->NetRefHandle);
 	Server->PreSendUpdate();
 	Server->SendUpdate(Client->ConnectionIdOnServer);
 	Server->PostSendUpdate();
 
 	// Remove object from filter to cause object to end up in CancelPendingDestroy...
-	Server->GetReplicationSystem()->RemoveFromGroup(NotReplicatedNetObjectGroupHandle, ServerObject->NetRefHandle);
+	Server->GetReplicationSystem()->RemoveFromGroup(Server->GetReplicationSystem()->GetNotReplicatedNetObjectGroup(), ServerObject->NetRefHandle);
 	Server->PreSendUpdate();
 	Server->PostSendUpdate();
 
 	// ... and cancel that thought! We want the object to be destroyed after all.
-	Server->GetReplicationSystem()->AddToGroup(NotReplicatedNetObjectGroupHandle, ServerObject->NetRefHandle);
+	Server->GetReplicationSystem()->AddToGroup(Server->GetReplicationSystem()->GetNotReplicatedNetObjectGroup(), ServerObject->NetRefHandle);
 	Server->PreSendUpdate();
 	Server->PostSendUpdate();
 
@@ -392,18 +392,18 @@ UE_NET_TEST_FIXTURE(FTestCancelPendingDestroyFixture, TestCancelingCancelPending
 	Server->PostSendUpdate();
 
 	// Filter out object and write packet to cause object to end up in WaitOnDestroyConfirmation
-	Server->GetReplicationSystem()->AddToGroup(NotReplicatedNetObjectGroupHandle, ServerObject->NetRefHandle);
+	Server->GetReplicationSystem()->AddToGroup(Server->GetReplicationSystem()->GetNotReplicatedNetObjectGroup(), ServerObject->NetRefHandle);
 	Server->PreSendUpdate();
 	Server->SendUpdate(Client->ConnectionIdOnServer);
 	Server->PostSendUpdate();
 
 	// Remove object from filter to cause object to end up in CancelPendingDestroy...
-	Server->GetReplicationSystem()->RemoveFromGroup(NotReplicatedNetObjectGroupHandle, ServerObject->NetRefHandle);
+	Server->GetReplicationSystem()->RemoveFromGroup(Server->GetReplicationSystem()->GetNotReplicatedNetObjectGroup(), ServerObject->NetRefHandle);
 	Server->PreSendUpdate();
 	Server->PostSendUpdate();
 
 	// ... and cancel that thought! We want the object to be destroyed after all.
-	Server->GetReplicationSystem()->AddToGroup(NotReplicatedNetObjectGroupHandle, ServerObject->NetRefHandle);
+	Server->GetReplicationSystem()->AddToGroup(Server->GetReplicationSystem()->GetNotReplicatedNetObjectGroup(), ServerObject->NetRefHandle);
 	Server->PreSendUpdate();
 	Server->PostSendUpdate();
 
@@ -442,12 +442,12 @@ UE_NET_TEST_FIXTURE(FTestCancelPendingDestroyFixture, TestCancelPendingSubObject
 	UE_NET_ASSERT_NE(Client->GetReplicationBridge()->GetReplicatedObject(ServerSubObject->NetRefHandle), nullptr);
 
 	// Filter out subobject to cause a SubObjectPendingDestroy
-	Server->GetReplicationSystem()->AddToGroup(NotReplicatedNetObjectGroupHandle, ServerSubObject->NetRefHandle);
+	Server->GetReplicationSystem()->AddToGroup(Server->GetReplicationSystem()->GetNotReplicatedNetObjectGroup(), ServerSubObject->NetRefHandle);
 	Server->PreSendUpdate();
 	Server->PostSendUpdate();
 
 	// Remove subobject from filter to cause object to end up in CancelPendingDestroy
-	Server->GetReplicationSystem()->RemoveFromGroup(NotReplicatedNetObjectGroupHandle, ServerSubObject->NetRefHandle);
+	Server->GetReplicationSystem()->RemoveFromGroup(Server->GetReplicationSystem()->GetNotReplicatedNetObjectGroup(), ServerSubObject->NetRefHandle);
 	Server->PreSendUpdate();
 	Server->PostSendUpdate();
 
@@ -478,13 +478,13 @@ UE_NET_TEST_FIXTURE(FTestCancelPendingDestroyFixture, TestCancelPendingSubObject
 	Server->PostSendUpdate();
 
 	// Filter out object and write packet to cause subobject to end up in WaitOnDestroyConfirmation
-	Server->GetReplicationSystem()->AddToGroup(NotReplicatedNetObjectGroupHandle, ServerSubObject->NetRefHandle);
+	Server->GetReplicationSystem()->AddToGroup(Server->GetReplicationSystem()->GetNotReplicatedNetObjectGroup(), ServerSubObject->NetRefHandle);
 	Server->PreSendUpdate();
 	Server->SendUpdate(Client->ConnectionIdOnServer);
 	Server->PostSendUpdate();
 
 	// Remove subobject from filter to cause object to end up in CancelPendingDestroy
-	Server->GetReplicationSystem()->RemoveFromGroup(NotReplicatedNetObjectGroupHandle, ServerSubObject->NetRefHandle);
+	Server->GetReplicationSystem()->RemoveFromGroup(Server->GetReplicationSystem()->GetNotReplicatedNetObjectGroup(), ServerSubObject->NetRefHandle);
 	Server->PreSendUpdate();
 	Server->PostSendUpdate();
 
@@ -521,13 +521,13 @@ UE_NET_TEST_FIXTURE(FTestCancelPendingDestroyFixture, TestCancelPendingSubObject
 	Server->PostSendUpdate();
 
 	// Filter out object and write packet to cause subobject to end up in WaitOnDestroyConfirmation
-	Server->GetReplicationSystem()->AddToGroup(NotReplicatedNetObjectGroupHandle, ServerSubObject->NetRefHandle);
+	Server->GetReplicationSystem()->AddToGroup(Server->GetReplicationSystem()->GetNotReplicatedNetObjectGroup(), ServerSubObject->NetRefHandle);
 	Server->PreSendUpdate();
 	Server->SendUpdate(Client->ConnectionIdOnServer);
 	Server->PostSendUpdate();
 
 	// Remove object from filter to cause object to end up in CancelPendingDestroy, also modify a property 
-	Server->GetReplicationSystem()->RemoveFromGroup(NotReplicatedNetObjectGroupHandle, ServerSubObject->NetRefHandle);
+	Server->GetReplicationSystem()->RemoveFromGroup(Server->GetReplicationSystem()->GetNotReplicatedNetObjectGroup(), ServerSubObject->NetRefHandle);
 	ServerSubObject->IntA ^= 1;
 	Server->PreSendUpdate();
 	Server->PostSendUpdate();
@@ -567,17 +567,17 @@ UE_NET_TEST_FIXTURE(FTestCancelPendingDestroyFixture, TestCancelingCancelPending
 	UE_NET_ASSERT_NE(Client->GetReplicationBridge()->GetReplicatedObject(ServerSubObject->NetRefHandle), nullptr);
 
 	// Filter out object to cause a SubObjectPendingDestroy
-	Server->GetReplicationSystem()->AddToGroup(NotReplicatedNetObjectGroupHandle, ServerSubObject->NetRefHandle);
+	Server->GetReplicationSystem()->AddToGroup(Server->GetReplicationSystem()->GetNotReplicatedNetObjectGroup(), ServerSubObject->NetRefHandle);
 	Server->PreSendUpdate();
 	Server->PostSendUpdate();
 
 	// Remove object from filter to cause object to end up in CancelPendingDestroy...
-	Server->GetReplicationSystem()->RemoveFromGroup(NotReplicatedNetObjectGroupHandle, ServerSubObject->NetRefHandle);
+	Server->GetReplicationSystem()->RemoveFromGroup(Server->GetReplicationSystem()->GetNotReplicatedNetObjectGroup(), ServerSubObject->NetRefHandle);
 	Server->PreSendUpdate();
 	Server->PostSendUpdate();
 
 	// ... and cancel that thought! We want the subobject to be destroyed after all.
-	Server->GetReplicationSystem()->AddToGroup(NotReplicatedNetObjectGroupHandle, ServerSubObject->NetRefHandle);
+	Server->GetReplicationSystem()->AddToGroup(Server->GetReplicationSystem()->GetNotReplicatedNetObjectGroup(), ServerSubObject->NetRefHandle);
 	Server->PreSendUpdate();
 	Server->PostSendUpdate();
 
@@ -605,18 +605,18 @@ UE_NET_TEST_FIXTURE(FTestCancelPendingDestroyFixture, TestCancelingCancelPending
 	Server->PostSendUpdate();
 
 	// Filter out object and write packet to cause subobject to end up in WaitOnDestroyConfirmation
-	Server->GetReplicationSystem()->AddToGroup(NotReplicatedNetObjectGroupHandle, ServerSubObject->NetRefHandle);
+	Server->GetReplicationSystem()->AddToGroup(Server->GetReplicationSystem()->GetNotReplicatedNetObjectGroup(), ServerSubObject->NetRefHandle);
 	Server->PreSendUpdate();
 	Server->SendUpdate(Client->ConnectionIdOnServer);
 	Server->PostSendUpdate();
 
 	// Remove object from filter to cause subobject to end up in CancelPendingDestroy...
-	Server->GetReplicationSystem()->RemoveFromGroup(NotReplicatedNetObjectGroupHandle, ServerSubObject->NetRefHandle);
+	Server->GetReplicationSystem()->RemoveFromGroup(Server->GetReplicationSystem()->GetNotReplicatedNetObjectGroup(), ServerSubObject->NetRefHandle);
 	Server->PreSendUpdate();
 	Server->PostSendUpdate();
 
 	// ... and cancel that thought! We want the subobject to be destroyed after all.
-	Server->GetReplicationSystem()->AddToGroup(NotReplicatedNetObjectGroupHandle, ServerSubObject->NetRefHandle);
+	Server->GetReplicationSystem()->AddToGroup(Server->GetReplicationSystem()->GetNotReplicatedNetObjectGroup(), ServerSubObject->NetRefHandle);
 	Server->PreSendUpdate();
 	Server->PostSendUpdate();
 
@@ -650,18 +650,18 @@ UE_NET_TEST_FIXTURE(FTestCancelPendingDestroyFixture, TestCancelingCancelPending
 	Server->PostSendUpdate();
 
 	// Filter out subobject and write packet to cause subobject to end up in WaitOnDestroyConfirmation
-	Server->GetReplicationSystem()->AddToGroup(NotReplicatedNetObjectGroupHandle, ServerSubObject->NetRefHandle);
+	Server->GetReplicationSystem()->AddToGroup(Server->GetReplicationSystem()->GetNotReplicatedNetObjectGroup(), ServerSubObject->NetRefHandle);
 	Server->PreSendUpdate();
 	Server->SendUpdate(Client->ConnectionIdOnServer);
 	Server->PostSendUpdate();
 
 	// Remove subobject from filter to cause subobject to end up in CancelPendingDestroy...
-	Server->GetReplicationSystem()->RemoveFromGroup(NotReplicatedNetObjectGroupHandle, ServerSubObject->NetRefHandle);
+	Server->GetReplicationSystem()->RemoveFromGroup(Server->GetReplicationSystem()->GetNotReplicatedNetObjectGroup(), ServerSubObject->NetRefHandle);
 	Server->PreSendUpdate();
 	Server->PostSendUpdate();
 
 	// ... and cancel that thought! We want the subobject to be destroyed after all.
-	Server->GetReplicationSystem()->AddToGroup(NotReplicatedNetObjectGroupHandle, ServerSubObject->NetRefHandle);
+	Server->GetReplicationSystem()->AddToGroup(Server->GetReplicationSystem()->GetNotReplicatedNetObjectGroup(), ServerSubObject->NetRefHandle);
 	Server->PreSendUpdate();
 	Server->PostSendUpdate();
 
@@ -696,19 +696,19 @@ UE_NET_TEST_FIXTURE(FTestCancelPendingDestroyFixture, TestCancelingSubObjectDest
 	Server->PostSendUpdate();
 
 	// Filter out subobject and write packet to cause subobject to end up in WaitOnDestroyConfirmation
-	Server->GetReplicationSystem()->AddToGroup(NotReplicatedNetObjectGroupHandle, ServerSubObject2->NetRefHandle);
+	Server->GetReplicationSystem()->AddToGroup(Server->GetReplicationSystem()->GetNotReplicatedNetObjectGroup(), ServerSubObject2->NetRefHandle);
 	Server->PreSendUpdate();
 	Server->SendUpdate(Client->ConnectionIdOnServer);
 	Server->PostSendUpdate();
 
 	// Filter out parent and write packet
-	Server->GetReplicationSystem()->AddToGroup(NotReplicatedNetObjectGroupHandle, ServerObject->NetRefHandle);
+	Server->GetReplicationSystem()->AddToGroup(Server->GetReplicationSystem()->GetNotReplicatedNetObjectGroup(), ServerObject->NetRefHandle);
 	Server->PreSendUpdate();
 	Server->SendUpdate(Client->ConnectionIdOnServer);
 	Server->PostSendUpdate();
 
 	// Remove subobject from filter to cause subobject to end up in CancelPendingDestroy
-	Server->GetReplicationSystem()->RemoveFromGroup(NotReplicatedNetObjectGroupHandle, ServerSubObject2->NetRefHandle);
+	Server->GetReplicationSystem()->RemoveFromGroup(Server->GetReplicationSystem()->GetNotReplicatedNetObjectGroup(), ServerSubObject2->NetRefHandle);
 	Server->PreSendUpdate();
 	Server->PostSendUpdate();
 
diff --git a/Engine/Plugins/Runtime/ReplicationSystemTestPlugin/Source/Private/Tests/ReplicationSystem/TestObjectReferences.cpp b/Engine/Plugins/Runtime/ReplicationSystemTestPlugin/Source/Private/Tests/ReplicationSystem/TestObjectReferences.cpp
--- a/Engine/Plugins/Runtime/ReplicationSystemTestPlugin/Source/Private/Tests/ReplicationSystem/TestObjectReferences.cpp
+++ b/Engine/Plugins/Runtime/ReplicationSystemTestPlugin/Source/Private/Tests/ReplicationSystem/TestObjectReferences.cpp
@@ -861,9 +861,9 @@ UE_NET_TEST_FIXTURE(FReplicationSystemServerClientTestFixture, TestPartialResolv
 	ServerObject->Ref_TArray.Add(ServerReferencedObjectC);
 
 	// Make sure that references are not replicated to client
-	Server->ReplicationSystem->AddToGroup(NotReplicatedNetObjectGroupHandle, ServerReferencedObjectA->NetRefHandle);
-	Server->ReplicationSystem->AddToGroup(NotReplicatedNetObjectGroupHandle, ServerReferencedObjectB->NetRefHandle);
-	Server->ReplicationSystem->AddToGroup(NotReplicatedNetObjectGroupHandle, ServerReferencedObjectC->NetRefHandle);
+	Server->ReplicationSystem->AddToGroup(Server->GetReplicationSystem()->GetNotReplicatedNetObjectGroup(), ServerReferencedObjectA->NetRefHandle);
+	Server->ReplicationSystem->AddToGroup(Server->GetReplicationSystem()->GetNotReplicatedNetObjectGroup(), ServerReferencedObjectB->NetRefHandle);
+	Server->ReplicationSystem->AddToGroup(Server->GetReplicationSystem()->GetNotReplicatedNetObjectGroup(), ServerReferencedObjectC->NetRefHandle);
 	
 	// Replicate
 	Server->PreSendUpdate();
@@ -883,7 +883,7 @@ UE_NET_TEST_FIXTURE(FReplicationSystemServerClientTestFixture, TestPartialResolv
 	UE_NET_ASSERT_TRUE(ClientObject->Ref_TArray[2] == nullptr);
 
 	// Enable replication for ServerReferenceObjectA
-	Server->ReplicationSystem->RemoveFromGroup(NotReplicatedNetObjectGroupHandle, ServerReferencedObjectA->NetRefHandle);
+	Server->ReplicationSystem->RemoveFromGroup(Server->GetReplicationSystem()->GetNotReplicatedNetObjectGroup(), ServerReferencedObjectA->NetRefHandle);
 
 	// Replicate
 	Server->PreSendUpdate();
@@ -896,7 +896,7 @@ UE_NET_TEST_FIXTURE(FReplicationSystemServerClientTestFixture, TestPartialResolv
 	UE_NET_ASSERT_TRUE(ClientObject->Ref_TArray[2] == nullptr);
 
 	// Enable replication for ServerReferenceObjectB 
-	Server->ReplicationSystem->RemoveFromGroup(NotReplicatedNetObjectGroupHandle, ServerReferencedObjectB->NetRefHandle);
+	Server->ReplicationSystem->RemoveFromGroup(Server->GetReplicationSystem()->GetNotReplicatedNetObjectGroup(), ServerReferencedObjectB->NetRefHandle);
 
 	// Replicate
 	Server->PreSendUpdate();
@@ -909,7 +909,7 @@ UE_NET_TEST_FIXTURE(FReplicationSystemServerClientTestFixture, TestPartialResolv
 	UE_NET_ASSERT_TRUE(ClientObject->Ref_TArray[2] == nullptr);
 
 	// Enable replication for ServerReferenceObjectB 
-	Server->ReplicationSystem->RemoveFromGroup(NotReplicatedNetObjectGroupHandle, ServerReferencedObjectC->NetRefHandle);
+	Server->ReplicationSystem->RemoveFromGroup(Server->GetReplicationSystem()->GetNotReplicatedNetObjectGroup(), ServerReferencedObjectC->NetRefHandle);
 
 	// Replicate
 	Server->PreSendUpdate();
@@ -956,9 +956,9 @@ UE_NET_TEST_FIXTURE(FReplicationSystemServerClientTestFixture, TestOutOfOrderRes
 #endif
 
 	// Make sure that references are not replicated to client
-	Server->ReplicationSystem->AddToGroup(NotReplicatedNetObjectGroupHandle, ServerReferencedObjectA->NetRefHandle);
-	Server->ReplicationSystem->AddToGroup(NotReplicatedNetObjectGroupHandle, ServerReferencedObjectB->NetRefHandle);
-	Server->ReplicationSystem->AddToGroup(NotReplicatedNetObjectGroupHandle, ServerReferencedObjectC->NetRefHandle);
+	Server->ReplicationSystem->AddToGroup(Server->GetReplicationSystem()->GetNotReplicatedNetObjectGroup(), ServerReferencedObjectA->NetRefHandle);
+	Server->ReplicationSystem->AddToGroup(Server->GetReplicationSystem()->GetNotReplicatedNetObjectGroup(), ServerReferencedObjectB->NetRefHandle);
+	Server->ReplicationSystem->AddToGroup(Server->GetReplicationSystem()->GetNotReplicatedNetObjectGroup(), ServerReferencedObjectC->NetRefHandle);
 	
 	// Replicate
 	Server->PreSendUpdate();
@@ -980,7 +980,7 @@ UE_NET_TEST_FIXTURE(FReplicationSystemServerClientTestFixture, TestOutOfOrderRes
 	UE_NET_ASSERT_TRUE(ClientFastArray.Items[2].ObjectRef == nullptr);
 
 	// Enable replication for ServerReferenceObjectA
-	Server->ReplicationSystem->RemoveFromGroup(NotReplicatedNetObjectGroupHandle, ServerReferencedObjectA->NetRefHandle);
+	Server->ReplicationSystem->RemoveFromGroup(Server->GetReplicationSystem()->GetNotReplicatedNetObjectGroup(), ServerReferencedObjectA->NetRefHandle);
 
 	// Replicate
 	Server->PreSendUpdate();
@@ -997,7 +997,7 @@ UE_NET_TEST_FIXTURE(FReplicationSystemServerClientTestFixture, TestOutOfOrderRes
 	UE_NET_ASSERT_TRUE(ClientFastArray.Items[2].ObjectRef == nullptr);
 
 	// Enable replication for ServerReferenceObjectB 
-	Server->ReplicationSystem->RemoveFromGroup(NotReplicatedNetObjectGroupHandle, ServerReferencedObjectB->NetRefHandle);
+	Server->ReplicationSystem->RemoveFromGroup(Server->GetReplicationSystem()->GetNotReplicatedNetObjectGroup(), ServerReferencedObjectB->NetRefHandle);
 
 	// Replicate
 	Server->PreSendUpdate();
@@ -1010,7 +1010,7 @@ UE_NET_TEST_FIXTURE(FReplicationSystemServerClientTestFixture, TestOutOfOrderRes
 	UE_NET_ASSERT_TRUE(ClientFastArray.Items[2].ObjectRef == nullptr);
 
 	// Enable replication for ServerReferenceObjectC
-	Server->ReplicationSystem->RemoveFromGroup(NotReplicatedNetObjectGroupHandle, ServerReferencedObjectC->NetRefHandle);
+	Server->ReplicationSystem->RemoveFromGroup(Server->GetReplicationSystem()->GetNotReplicatedNetObjectGroup(), ServerReferencedObjectC->NetRefHandle);
 
 	// Replicate
 	Server->PreSendUpdate();
@@ -1051,9 +1051,9 @@ UE_NET_TEST_FIXTURE(FReplicationSystemServerClientTestFixture, TestOutOfOrderRes
 	ServerFastArrayEditor.Add(Item);
 
 	// Make sure that references are not replicated to client
-	Server->ReplicationSystem->AddToGroup(NotReplicatedNetObjectGroupHandle, ServerReferencedObjectA->NetRefHandle);
-	Server->ReplicationSystem->AddToGroup(NotReplicatedNetObjectGroupHandle, ServerReferencedObjectB->NetRefHandle);
-	Server->ReplicationSystem->AddToGroup(NotReplicatedNetObjectGroupHandle, ServerReferencedObjectC->NetRefHandle);
+	Server->ReplicationSystem->AddToGroup(Server->GetReplicationSystem()->GetNotReplicatedNetObjectGroup(), ServerReferencedObjectA->NetRefHandle);
+	Server->ReplicationSystem->AddToGroup(Server->GetReplicationSystem()->GetNotReplicatedNetObjectGroup(), ServerReferencedObjectB->NetRefHandle);
+	Server->ReplicationSystem->AddToGroup(Server->GetReplicationSystem()->GetNotReplicatedNetObjectGroup(), ServerReferencedObjectC->NetRefHandle);
 	
 	// Replicate
 	Server->PreSendUpdate();
@@ -1075,7 +1075,7 @@ UE_NET_TEST_FIXTURE(FReplicationSystemServerClientTestFixture, TestOutOfOrderRes
 	UE_NET_ASSERT_TRUE(ClientFastArray.Items[2].ObjectRef == nullptr);
 
 	// Enable replication for ServerReferenceObjectA
-	Server->ReplicationSystem->RemoveFromGroup(NotReplicatedNetObjectGroupHandle, ServerReferencedObjectA->NetRefHandle);
+	Server->ReplicationSystem->RemoveFromGroup(Server->GetReplicationSystem()->GetNotReplicatedNetObjectGroup(), ServerReferencedObjectA->NetRefHandle);
 
 	// Replicate
 	Server->PreSendUpdate();
@@ -1088,7 +1088,7 @@ UE_NET_TEST_FIXTURE(FReplicationSystemServerClientTestFixture, TestOutOfOrderRes
 	UE_NET_ASSERT_TRUE(ClientFastArray.Items[2].ObjectRef == nullptr);
 
 	// Enable replication for ServerReferenceObjectB 
-	Server->ReplicationSystem->RemoveFromGroup(NotReplicatedNetObjectGroupHandle, ServerReferencedObjectB->NetRefHandle);
+	Server->ReplicationSystem->RemoveFromGroup(Server->GetReplicationSystem()->GetNotReplicatedNetObjectGroup(), ServerReferencedObjectB->NetRefHandle);
 
 	// Replicate
 	Server->PreSendUpdate();
@@ -1101,7 +1101,7 @@ UE_NET_TEST_FIXTURE(FReplicationSystemServerClientTestFixture, TestOutOfOrderRes
 	UE_NET_ASSERT_TRUE(ClientFastArray.Items[2].ObjectRef == nullptr);
 
 	// Enable replication for ServerReferenceObjectC
-	Server->ReplicationSystem->RemoveFromGroup(NotReplicatedNetObjectGroupHandle, ServerReferencedObjectC->NetRefHandle);
+	Server->ReplicationSystem->RemoveFromGroup(Server->GetReplicationSystem()->GetNotReplicatedNetObjectGroup(), ServerReferencedObjectC->NetRefHandle);
 
 	// Replicate
 	Server->PreSendUpdate();
@@ -1146,9 +1146,9 @@ UE_NET_TEST_FIXTURE(FReplicationSystemServerClientTestFixture, TestOutOfOrderRes
 	ServerFastArrayEditor.Add(Item);
 
 	// Make sure that references are not replicated to client
-	Server->ReplicationSystem->AddToGroup(NotReplicatedNetObjectGroupHandle, ServerReferencedObjectA->NetRefHandle);
-	Server->ReplicationSystem->AddToGroup(NotReplicatedNetObjectGroupHandle, ServerReferencedObjectB->NetRefHandle);
-	Server->ReplicationSystem->AddToGroup(NotReplicatedNetObjectGroupHandle, ServerReferencedObjectC->NetRefHandle);
+	Server->ReplicationSystem->AddToGroup(Server->GetReplicationSystem()->GetNotReplicatedNetObjectGroup(), ServerReferencedObjectA->NetRefHandle);
+	Server->ReplicationSystem->AddToGroup(Server->GetReplicationSystem()->GetNotReplicatedNetObjectGroup(), ServerReferencedObjectB->NetRefHandle);
+	Server->ReplicationSystem->AddToGroup(Server->GetReplicationSystem()->GetNotReplicatedNetObjectGroup(), ServerReferencedObjectC->NetRefHandle);
 	
 	// Replicate
 	Server->PreSendUpdate();
@@ -1177,7 +1177,7 @@ UE_NET_TEST_FIXTURE(FReplicationSystemServerClientTestFixture, TestOutOfOrderRes
 	UE_NET_ASSERT_TRUE(ClientFastArray.Items[2].Ref_TArray[0] == nullptr);
 
 	// Enable replication for ServerReferenceObjectA
-	Server->ReplicationSystem->RemoveFromGroup(NotReplicatedNetObjectGroupHandle, ServerReferencedObjectA->NetRefHandle);
+	Server->ReplicationSystem->RemoveFromGroup(Server->GetReplicationSystem()->GetNotReplicatedNetObjectGroup(), ServerReferencedObjectA->NetRefHandle);
 
 	// Replicate
 	Server->PreSendUpdate();
@@ -1190,7 +1190,7 @@ UE_NET_TEST_FIXTURE(FReplicationSystemServerClientTestFixture, TestOutOfOrderRes
 	UE_NET_ASSERT_TRUE(ClientFastArray.Items[2].Ref_TArray[0] == nullptr);
 
 	// Enable replication for ServerReferenceObjectB 
-	Server->ReplicationSystem->RemoveFromGroup(NotReplicatedNetObjectGroupHandle, ServerReferencedObjectB->NetRefHandle);
+	Server->ReplicationSystem->RemoveFromGroup(Server->GetReplicationSystem()->GetNotReplicatedNetObjectGroup(), ServerReferencedObjectB->NetRefHandle);
 
 	// Replicate
 	Server->PreSendUpdate();
@@ -1203,7 +1203,7 @@ UE_NET_TEST_FIXTURE(FReplicationSystemServerClientTestFixture, TestOutOfOrderRes
 	UE_NET_ASSERT_TRUE(ClientFastArray.Items[2].Ref_TArray[0] == nullptr);
 
 	// Enable replication for ServerReferenceObjectC
-	Server->ReplicationSystem->RemoveFromGroup(NotReplicatedNetObjectGroupHandle, ServerReferencedObjectC->NetRefHandle);
+	Server->ReplicationSystem->RemoveFromGroup(Server->GetReplicationSystem()->GetNotReplicatedNetObjectGroup(), ServerReferencedObjectC->NetRefHandle);
 
 	// Replicate
 	Server->PreSendUpdate();
@@ -1244,7 +1244,7 @@ UE_NET_TEST_FIXTURE(FReplicationSystemServerClientTestFixture, TestRemovedRefere
 	UE_NET_ASSERT_NE(ClientObject->StructWithRef.ObjectRef, TObjectPtr<UObject>(nullptr));
 
 	// Disable replication for ServerReferenceObjectA
-	Server->ReplicationSystem->AddToGroup(NotReplicatedNetObjectGroupHandle, ServerReferencedObjectA->NetRefHandle);
+	Server->ReplicationSystem->AddToGroup(Server->GetReplicationSystem()->GetNotReplicatedNetObjectGroup(), ServerReferencedObjectA->NetRefHandle);
 
 	// Replicate to make sure object is destroyed on client
 	Server->PreSendUpdate();
@@ -1263,7 +1263,7 @@ UE_NET_TEST_FIXTURE(FReplicationSystemServerClientTestFixture, TestRemovedRefere
 	UE_NET_ASSERT_EQ(ClientObject->StructWithRef.ObjectRef, TObjectPtr<UObject>(nullptr));
 
 	// Enable replication for ServerReferenceObjectA
-	Server->ReplicationSystem->RemoveFromGroup(NotReplicatedNetObjectGroupHandle, ServerReferencedObjectA->NetRefHandle);
+	Server->ReplicationSystem->RemoveFromGroup(Server->GetReplicationSystem()->GetNotReplicatedNetObjectGroup(), ServerReferencedObjectA->NetRefHandle);
 
 	// Replicate
 	Server->PreSendUpdate();
@@ -1312,7 +1312,7 @@ UE_NET_TEST_FIXTURE(FReplicationSystemServerClientTestFixture, TestRemovedRefere
 	UE_NET_ASSERT_NE(ClientObject->StructWithRef_CArray[0].ObjectRef, TObjectPtr<UObject>(nullptr));
 
 	// Disable replication for ServerReferenceObjectA
-	Server->ReplicationSystem->AddToGroup(NotReplicatedNetObjectGroupHandle, ServerReferencedObjectA->NetRefHandle);
+	Server->ReplicationSystem->AddToGroup(Server->GetReplicationSystem()->GetNotReplicatedNetObjectGroup(), ServerReferencedObjectA->NetRefHandle);
 
 	// Replicate to make sure object is destroyed on client
 	Server->PreSendUpdate();
@@ -1333,7 +1333,7 @@ UE_NET_TEST_FIXTURE(FReplicationSystemServerClientTestFixture, TestRemovedRefere
 	UE_NET_ASSERT_EQ(ClientObject->StructWithRef_CArray[0].ObjectRef, TObjectPtr<UObject>(nullptr));
 
 	// Enable replication for ServerReferenceObjectA
-	Server->ReplicationSystem->RemoveFromGroup(NotReplicatedNetObjectGroupHandle, ServerReferencedObjectA->NetRefHandle);
+	Server->ReplicationSystem->RemoveFromGroup(Server->GetReplicationSystem()->GetNotReplicatedNetObjectGroup(), ServerReferencedObjectA->NetRefHandle);
 
 	// Replicate
 	Server->PreSendUpdate();
@@ -1398,9 +1398,9 @@ UE_NET_TEST_FIXTURE(FReplicationSystemServerClientTestFixture, TestRemovingRefer
 	UE_NET_ASSERT_NE(ClientTArray[3].ObjectRef, TObjectPtr<UObject>(nullptr));
 
 	// Disable replication for ServerReferenceObjectA, B and D.
-	Server->ReplicationSystem->AddToGroup(NotReplicatedNetObjectGroupHandle, ServerReferencedObjectA->NetRefHandle);
-	Server->ReplicationSystem->AddToGroup(NotReplicatedNetObjectGroupHandle, ServerReferencedObjectB->NetRefHandle);
-	Server->ReplicationSystem->AddToGroup(NotReplicatedNetObjectGroupHandle, ServerReferencedObjectD->NetRefHandle);
+	Server->ReplicationSystem->AddToGroup(Server->GetReplicationSystem()->GetNotReplicatedNetObjectGroup(), ServerReferencedObjectA->NetRefHandle);
+	Server->ReplicationSystem->AddToGroup(Server->GetReplicationSystem()->GetNotReplicatedNetObjectGroup(), ServerReferencedObjectB->NetRefHandle);
+	Server->ReplicationSystem->AddToGroup(Server->GetReplicationSystem()->GetNotReplicatedNetObjectGroup(), ServerReferencedObjectD->NetRefHandle);
 
 	// Replicate to make sure objects are destroyed on client
 	Server->PreSendUpdate();
@@ -1433,9 +1433,9 @@ UE_NET_TEST_FIXTURE(FReplicationSystemServerClientTestFixture, TestRemovingRefer
 	UE_NET_ASSERT_EQ(ClientTArray[3].ObjectRef, TObjectPtr<UObject>(nullptr));
 
 	// Enable replication for ServerReferenceObjectA, B and D.
-	Server->ReplicationSystem->RemoveFromGroup(NotReplicatedNetObjectGroupHandle, ServerReferencedObjectA->NetRefHandle);
-	Server->ReplicationSystem->RemoveFromGroup(NotReplicatedNetObjectGroupHandle, ServerReferencedObjectB->NetRefHandle);
-	Server->ReplicationSystem->RemoveFromGroup(NotReplicatedNetObjectGroupHandle, ServerReferencedObjectD->NetRefHandle);
+	Server->ReplicationSystem->RemoveFromGroup(Server->GetReplicationSystem()->GetNotReplicatedNetObjectGroup(), ServerReferencedObjectA->NetRefHandle);
+	Server->ReplicationSystem->RemoveFromGroup(Server->GetReplicationSystem()->GetNotReplicatedNetObjectGroup(), ServerReferencedObjectB->NetRefHandle);
+	Server->ReplicationSystem->RemoveFromGroup(Server->GetReplicationSystem()->GetNotReplicatedNetObjectGroup(), ServerReferencedObjectD->NetRefHandle);
 
 	// Replicate
 	Server->PreSendUpdate();
diff --git a/Engine/Plugins/Runtime/ReplicationSystemTestPlugin/Source/Private/Tests/ReplicationSystem/TestObjectSplitting.cpp b/Engine/Plugins/Runtime/ReplicationSystemTestPlugin/Source/Private/Tests/ReplicationSystem/TestObjectSplitting.cpp
--- a/Engine/Plugins/Runtime/ReplicationSystemTestPlugin/Source/Private/Tests/ReplicationSystem/TestObjectSplitting.cpp
+++ b/Engine/Plugins/Runtime/ReplicationSystemTestPlugin/Source/Private/Tests/ReplicationSystem/TestObjectSplitting.cpp
@@ -713,12 +713,12 @@ UE_NET_TEST_FIXTURE(FSplitObjectTestFixture, TestCancelPendingDestroyOfHugeObjec
 	}
 
 	// Filter out object to cause a PendingDestroy
-	Server->GetReplicationSystem()->AddToGroup(NotReplicatedNetObjectGroupHandle, ServerObject->NetRefHandle);
+	Server->GetReplicationSystem()->AddToGroup(Server->GetReplicationSystem()->GetNotReplicatedNetObjectGroup(), ServerObject->NetRefHandle);
 	Server->PreSendUpdate();
 	Server->PostSendUpdate();
 
 	// Remove object from filter to cause object to end up in CancelPendingDestroy
-	Server->GetReplicationSystem()->RemoveFromGroup(NotReplicatedNetObjectGroupHandle, ServerObject->NetRefHandle);
+	Server->GetReplicationSystem()->RemoveFromGroup(Server->GetReplicationSystem()->GetNotReplicatedNetObjectGroup(), ServerObject->NetRefHandle);
 	Server->PreSendUpdate();
 	Server->PostSendUpdate();
 
@@ -766,12 +766,12 @@ UE_NET_TEST_FIXTURE(FSplitObjectTestFixture, TestCancelPendingDestroyOfHugeObjec
 	}
 
 	// Filter out object to cause a PendingDestroy
-	Server->GetReplicationSystem()->AddToGroup(NotReplicatedNetObjectGroupHandle, ServerObject->NetRefHandle);
+	Server->GetReplicationSystem()->AddToGroup(Server->GetReplicationSystem()->GetNotReplicatedNetObjectGroup(), ServerObject->NetRefHandle);
 	Server->PreSendUpdate();
 	Server->PostSendUpdate();
 
 	// Remove object from filter to cause object to end up in CancelPendingDestroy
-	Server->GetReplicationSystem()->RemoveFromGroup(NotReplicatedNetObjectGroupHandle, ServerObject->NetRefHandle);
+	Server->GetReplicationSystem()->RemoveFromGroup(Server->GetReplicationSystem()->GetNotReplicatedNetObjectGroup(), ServerObject->NetRefHandle);
 	Server->PreSendUpdate();
 	Server->PostSendUpdate();
 
@@ -831,12 +831,12 @@ UE_NET_TEST_FIXTURE(FSplitObjectTestFixture, TestCancelPendingDestroyDuringHugeO
 	}
 
 	// Filter out object to cause a PendingDestroy
-	Server->GetReplicationSystem()->AddToGroup(NotReplicatedNetObjectGroupHandle, ServerObject->NetRefHandle);
+	Server->GetReplicationSystem()->AddToGroup(Server->GetReplicationSystem()->GetNotReplicatedNetObjectGroup(), ServerObject->NetRefHandle);
 	Server->PreSendUpdate();
 	Server->PostSendUpdate();
 
 	// Remove object from filter to cause object to end up in CancelPendingDestroy
-	Server->GetReplicationSystem()->RemoveFromGroup(NotReplicatedNetObjectGroupHandle, ServerObject->NetRefHandle);
+	Server->GetReplicationSystem()->RemoveFromGroup(Server->GetReplicationSystem()->GetNotReplicatedNetObjectGroup(), ServerObject->NetRefHandle);
 	Server->PreSendUpdate();
 	Server->PostSendUpdate();
 
@@ -898,7 +898,7 @@ UE_NET_TEST_FIXTURE(FSplitObjectTestFixture, TestReliableAttachmentIsDeliveredDe
 	}
 
 	// Filter out object to cause object to be set in state WaitOnFlush
-	Server->GetReplicationSystem()->AddToGroup(NotReplicatedNetObjectGroupHandle, ServerObject->NetRefHandle);
+	Server->GetReplicationSystem()->AddToGroup(Server->GetReplicationSystem()->GetNotReplicatedNetObjectGroup(), ServerObject->NetRefHandle);
 	Server->PreSendUpdate();
 	Server->SendUpdate(Client->ConnectionIdOnServer);
 	Server->PostSendUpdate();
@@ -962,13 +962,13 @@ UE_NET_TEST_FIXTURE(FSplitObjectTestFixture, TestHugeObjectIsFlushedAndNotDestro
 	}
 
 	// Filter out object to cause object to be set in state WaitOnFlush
-	Server->GetReplicationSystem()->AddToGroup(NotReplicatedNetObjectGroupHandle, ServerObject->NetRefHandle);
+	Server->GetReplicationSystem()->AddToGroup(Server->GetReplicationSystem()->GetNotReplicatedNetObjectGroup(), ServerObject->NetRefHandle);
 	Server->PreSendUpdate();
 	Server->SendUpdate(Client->ConnectionIdOnServer);
 	Server->PostSendUpdate();
 
 	// Remove object from filter to cause object to be set in state Created
-	Server->GetReplicationSystem()->RemoveFromGroup(NotReplicatedNetObjectGroupHandle, ServerObject->NetRefHandle);
+	Server->GetReplicationSystem()->RemoveFromGroup(Server->GetReplicationSystem()->GetNotReplicatedNetObjectGroup(), ServerObject->NetRefHandle);
 	Server->PreSendUpdate();
 	Server->SendUpdate(Client->ConnectionIdOnServer);
 	Server->PostSendUpdate();
@@ -1037,7 +1037,7 @@ UE_NET_TEST_FIXTURE(FSplitObjectTestFixture, TestReliableAttachmentAddedAfterSpl
 	}
 
 	// Filter out object to cause object to be set in state WaitOnFlush
-	Server->GetReplicationSystem()->AddToGroup(NotReplicatedNetObjectGroupHandle, ServerObject->NetRefHandle);
+	Server->GetReplicationSystem()->AddToGroup(Server->GetReplicationSystem()->GetNotReplicatedGroup(), ServerObject->NetRefHandle);
 	Server->PreSendUpdate();
 	Server->SendUpdate(Client->ConnectionIdOnServer);
 	Server->PostSendUpdate();
diff --git a/Engine/Source/Runtime/Engine/Classes/Engine/NetDriver.h b/Engine/Source/Runtime/Engine/Classes/Engine/NetDriver.h
--- a/Engine/Source/Runtime/Engine/Classes/Engine/NetDriver.h
+++ b/Engine/Source/Runtime/Engine/Classes/Engine/NetDriver.h
@@ -352,7 +352,7 @@ class UReplicationSystem;
 class UReplicationBridge;
 namespace UE::Net
 {
-	typedef uint16 FNetObjectGroupHandle;
+	class FNetObjectGroupHandle;
 }
 #endif // UE_WITH_IRIS
 
diff --git a/Engine/Source/Runtime/Engine/Private/Net/Iris/ReplicationSystem/ActorReplicationBridge.cpp b/Engine/Source/Runtime/Engine/Private/Net/Iris/ReplicationSystem/ActorReplicationBridge.cpp
--- a/Engine/Source/Runtime/Engine/Private/Net/Iris/ReplicationSystem/ActorReplicationBridge.cpp
+++ b/Engine/Source/Runtime/Engine/Private/Net/Iris/ReplicationSystem/ActorReplicationBridge.cpp
@@ -291,7 +291,7 @@ UE::Net::FNetRefHandle UActorReplicationBridge::BeginReplication(AActor* Actor,
 			const FName PackageName = LevelPackage->GetFName();
 
 			FNetObjectGroupHandle LevelGroup = GetLevelGroup(Level);
-			if (!LevelGroup)
+			if (!LevelGroup.IsValid())
 			{
 				LevelGroup = CreateLevelGroup(Level);
 
diff --git a/Engine/Source/Runtime/Engine/Private/NetConnection.cpp b/Engine/Source/Runtime/Engine/Private/NetConnection.cpp
--- a/Engine/Source/Runtime/Engine/Private/NetConnection.cpp
+++ b/Engine/Source/Runtime/Engine/Private/NetConnection.cpp
@@ -1450,7 +1450,7 @@ FNetLevelVisibilityTransactionId UNetConnection::UpdateLevelStreamStatusChangedT
 			if (const ULevel* Level = LevelObject->GetLoadedLevel())
 			{
 				const UReplicationBridge* ReplicationBridge = ReplicationSystem->GetReplicationBridge();
-				if (UE::Net::FNetObjectGroupHandle GroupHandle = ReplicationBridge->GetLevelGroup(Level); GroupHandle != UE::Net::InvalidNetObjectGroupHandle)
+				if (UE::Net::FNetObjectGroupHandle GroupHandle = ReplicationBridge->GetLevelGroup(Level); GroupHandle.IsValid())
 				{
 					ReplicationSystem->SetGroupFilterStatus(GroupHandle, GetConnectionId(), UE::Net::ENetFilterStatus::Disallow);
 				}
@@ -1601,7 +1601,7 @@ void UNetConnection::UpdateLevelVisibilityInternal(const FUpdateLevelVisibilityL
 				if (Level)
 				{
 					const UReplicationBridge* ReplicationBridge = ReplicationSystem->GetReplicationBridge();
-					if (UE::Net::FNetObjectGroupHandle GroupHandle = ReplicationBridge->GetLevelGroup(Level); GroupHandle != UE::Net::InvalidNetObjectGroupHandle)
+					if (UE::Net::FNetObjectGroupHandle GroupHandle = ReplicationBridge->GetLevelGroup(Level); GroupHandle.IsValid())
 					{
 						ReplicationSystem->SetGroupFilterStatus(GroupHandle, GetConnectionId(), UE::Net::ENetFilterStatus::Allow);
 					}
@@ -1698,7 +1698,7 @@ void UNetConnection::UpdateLevelVisibilityInternal(const FUpdateLevelVisibilityL
 					// The reason for this is that the client currently destroys the instances rather then managing this through the replication system
 					// If we implement a way to re-instantiate instances on the client we might be able to persist the state
 					const UReplicationBridge* ReplicationBridge = ReplicationSystem->GetReplicationBridge();
-					if (UE::Net::FNetObjectGroupHandle GroupHandle = ReplicationBridge->GetLevelGroup(Level); GroupHandle != UE::Net::InvalidNetObjectGroupHandle)
+					if (UE::Net::FNetObjectGroupHandle GroupHandle = ReplicationBridge->GetLevelGroup(Level); GroupHandle.IsValid())
 					{
 						ReplicationSystem->SetGroupFilterStatus(GroupHandle, GetConnectionId(), UE::Net::ENetFilterStatus::Disallow);
 					}
diff --git a/Engine/Source/Runtime/Engine/Private/NetDriver.cpp b/Engine/Source/Runtime/Engine/Private/NetDriver.cpp
--- a/Engine/Source/Runtime/Engine/Private/NetDriver.cpp
+++ b/Engine/Source/Runtime/Engine/Private/NetDriver.cpp
@@ -6232,7 +6232,7 @@ void UNetDriver::UpdateGroupFilterStatusForLevel(const ULevel* Level, UE::Net::F
 {
 	using namespace UE::Net;
 
-	if (ReplicationSystem == nullptr || LevelGroupHandle == InvalidNetObjectGroupHandle)
+	if (ReplicationSystem == nullptr || !ensure(ReplicationSystem->IsValidGroup(LevelGroupHandle)))
 	{
 		return;
 	}
diff --git a/Engine/Source/Runtime/Experimental/Iris/Core/Private/Iris/ReplicationSystem/Conditionals/ReplicationConditionals.cpp b/Engine/Source/Runtime/Experimental/Iris/Core/Private/Iris/ReplicationSystem/Conditionals/ReplicationConditionals.cpp
--- a/Engine/Source/Runtime/Experimental/Iris/Core/Private/Iris/ReplicationSystem/Conditionals/ReplicationConditionals.cpp
+++ b/Engine/Source/Runtime/Experimental/Iris/Core/Private/Iris/ReplicationSystem/Conditionals/ReplicationConditionals.cpp
@@ -474,18 +474,18 @@ void FReplicationConditionals::GetChildSubObjectsToReplicate(uint32 ReplicatingC
 						for (uint32 GroupIt = 0U; GroupIt < GroupCount; ++GroupIt)
 						{
 							const FNetObjectGroupHandle NetGroup = GroupMemberships[GroupIt];
-							if (NetGroup == NetGroupOwnerNetObjectGroupHandle)
+							if (NetGroup.IsNetGroupOwnerNetObjectGroup())
 							{
 								bShouldReplicateSubObject = LifetimeConditionals.IsConditionEnabled(COND_OwnerOnly);
 							}
-							else if (NetGroup == NetGroupReplayNetObjectGroupHandle)
+							else if (NetGroup.IsNetGroupReplayNetObjectGroup())
 							{
 								bShouldReplicateSubObject = LifetimeConditionals.IsConditionEnabled(COND_ReplayOnly);
 							}
 							else
 							{
 								ENetFilterStatus ReplicationStatus = ENetFilterStatus::Disallow;
-								ensureAlwaysMsgf(ReplicationFiltering->GetSubObjectFilterStatus(NetGroup, ReplicatingConnectionId, ReplicationStatus), TEXT("FReplicationConditionals::GetChildSubObjectsToReplicat Trying to filter with group %u that is not a SubObjectFilterGroup"), NetGroup);
+								ensureAlwaysMsgf(ReplicationFiltering->GetSubObjectFilterStatus(NetGroup, ReplicatingConnectionId, ReplicationStatus), TEXT("FReplicationConditionals::GetChildSubObjectsToReplicat Trying to filter with group %u that is not a SubObjectFilterGroup"), NetGroup.GetGroupIndex());
 								bShouldReplicateSubObject = ReplicationStatus != ENetFilterStatus::Disallow;
 							}
 						
diff --git a/Engine/Source/Runtime/Experimental/Iris/Core/Private/Iris/ReplicationSystem/Filtering/NetObjectGroups.cpp b/Engine/Source/Runtime/Experimental/Iris/Core/Private/Iris/ReplicationSystem/Filtering/NetObjectGroups.cpp
--- a/Engine/Source/Runtime/Experimental/Iris/Core/Private/Iris/ReplicationSystem/Filtering/NetObjectGroups.cpp
+++ b/Engine/Source/Runtime/Experimental/Iris/Core/Private/Iris/ReplicationSystem/Filtering/NetObjectGroups.cpp
@@ -10,6 +10,7 @@ namespace UE::Net::Private
 
 FNetObjectGroups::FNetObjectGroups()
 : MaxGroupCount(0U)
+, CurrentEpoch(++NextEpoch)
 {
 }
 
@@ -38,7 +39,7 @@ bool FNetObjectGroups::AddGroupMembership(FNetObjectGroupMembership& Target, FNe
 		{
 			return false;
 		}
-		else if (Target.Groups[It] == InvalidNetObjectGroupHandle)
+		else if (!Target.Groups[It].IsValid())
 		{
 			Target.Groups[It] = Group;
 			return true;
@@ -56,7 +57,7 @@ void FNetObjectGroups::RemoveGroupMembership(FNetObjectGroupMembership& Target,
 		if (Target.Groups[It] == Group)
 		{
 			FPlatformMemory::Memmove(&Target.Groups[It], &Target.Groups[It + 1], (LastIndex - It) * sizeof(FNetObjectGroupHandle));
-			Target.Groups[LastIndex] = InvalidNetObjectGroupHandle;
+			Target.Groups[LastIndex] = FNetObjectGroupHandle();
 
 			return;
 		}
@@ -84,28 +85,31 @@ FNetObjectGroupHandle FNetObjectGroups::CreateGroup()
 	if (ensure((uint32)Groups.Num() < MaxGroupCount))
 	{
 		const int Index = Groups.Add(FNetObjectGroup());
-		return FNetObjectGroupHandle(Index);
+		FNetObjectGroupHandle GroupHandle;
+		GroupHandle.Index = Index;
+		GroupHandle.Epoch = CurrentEpoch;
+		return GroupHandle;
 	}
 	else
 	{
-		return InvalidNetObjectGroupHandle;
+		return FNetObjectGroupHandle();
 	}
 }
 
 void FNetObjectGroups::DestroyGroup(FNetObjectGroupHandle GroupHandle)
 {
-	if (GroupHandle != InvalidNetObjectGroupHandle && Groups.IsValidIndex(GroupHandle))
+	if (IsValidGroup(GroupHandle))
 	{
 		ClearGroup(GroupHandle);
 
-		FNetObjectGroup& Group = Groups[GroupHandle];
+		FNetObjectGroup& Group = Groups[GroupHandle.GetGroupIndex()];
 
 		if (EnumHasAnyFlags(Group.Traits, ENetObjectGroupTraits::IsFindableByName))
 		{
 			NamedGroups.FindAndRemoveChecked(Group.GroupName);
 		}
 
-		Groups.RemoveAt(GroupHandle);
+		Groups.RemoveAt(GroupHandle.GetGroupIndex());
 	}
 }
 
@@ -113,7 +117,7 @@ void FNetObjectGroups::SetGroupName(FNetObjectGroupHandle GroupHandle, FName Gro
 {
 	if (FNetObjectGroup* Group = GetGroup(GroupHandle))
 	{
-		if (ensureAlwaysMsgf(EnumHasAnyFlags(Group->Traits, ENetObjectGroupTraits::IsFindableByName), TEXT("FNetObjectGroups::SetGroupName Cannot SetGroupName for grouphandle %u as it is FindableByName"), GroupHandle))
+		if (ensureAlwaysMsgf(EnumHasAnyFlags(Group->Traits, ENetObjectGroupTraits::IsFindableByName), TEXT("FNetObjectGroups::SetGroupName Cannot SetGroupName for grouphandle %u as it is FindableByName"), GroupHandle.GetGroupIndex()))
 		{
 			Group->GroupName = GroupName;
 		}
@@ -125,7 +129,7 @@ FNetObjectGroupHandle FNetObjectGroups::CreateNamedGroup(FName GroupName)
 	if (NamedGroups.Contains(GroupName))
 	{
 		ensureAlwaysMsgf(false, TEXT("FNetObjectGroups, trying to create named group %s that already exists"), *GroupName.ToString());
-		return InvalidNetObjectGroupHandle;
+		return FNetObjectGroupHandle();
 	}
 
 	FNetObjectGroupHandle GroupHandle = CreateGroup();
@@ -146,7 +150,7 @@ FNetObjectGroupHandle FNetObjectGroups::GetNamedGroupHandle(FName GroupName)
 		return *Group;
 	}
 
-	return InvalidNetObjectGroupHandle;
+	return FNetObjectGroupHandle();
 }
 
 void FNetObjectGroups::DestroyNamedGroup(FName GroupName)
@@ -156,9 +160,9 @@ void FNetObjectGroups::DestroyNamedGroup(FName GroupName)
 
 void FNetObjectGroups::ClearGroup(FNetObjectGroupHandle GroupHandle)
 {
-	if (GroupHandle != InvalidNetObjectGroupHandle && Groups.IsValidIndex(GroupHandle))
+	if (IsValidGroup(GroupHandle))
 	{
-		FNetObjectGroup& Group = Groups[GroupHandle];
+		FNetObjectGroup& Group = Groups[GroupHandle.GetGroupIndex()];
 
 		for (FInternalNetRefIndex InternalIndex : Group.Members)
 		{
@@ -172,12 +176,22 @@ void FNetObjectGroups::ClearGroup(FNetObjectGroupHandle GroupHandle)
 
 const FNetObjectGroup* FNetObjectGroups::GetGroup(FNetObjectGroupHandle GroupHandle) const
 {
-	return Groups.IsValidIndex(GroupHandle) ? &Groups[GroupHandle] : nullptr;
+	return IsValidGroup(GroupHandle) ? &Groups[GroupHandle.GetGroupIndex()] : nullptr;
 }
 
 FNetObjectGroup* FNetObjectGroups::GetGroup(FNetObjectGroupHandle GroupHandle)
 {
-	return Groups.IsValidIndex(GroupHandle) ? &Groups[GroupHandle] : nullptr;
+	return IsValidGroup(GroupHandle) ? &Groups[GroupHandle.GetGroupIndex()] : nullptr;
+}
+
+const FNetObjectGroup* FNetObjectGroups::GetGroupByIndex(FNetObjectGroupHandle::FGroupIndexType GroupIndex) const
+{
+	return (GroupIndex != FNetObjectGroupHandle::InvalidNetObjectGroupIndex && Groups.IsValidIndex(GroupIndex)) ? &Groups[GroupIndex] : nullptr;
+}
+
+FNetObjectGroup* FNetObjectGroups::GetGroupByIndex(FNetObjectGroupHandle::FGroupIndexType GroupIndex)
+{
+	return (GroupIndex != FNetObjectGroupHandle::InvalidNetObjectGroupIndex && Groups.IsValidIndex(GroupIndex)) ? &Groups[GroupIndex] : nullptr;
 }
 
 bool FNetObjectGroups::Contains(FNetObjectGroupHandle GroupHandle, FInternalNetRefIndex InternalIndex) const
@@ -282,12 +296,12 @@ bool FNetObjectGroups::IsInAnyFilterGroup(const FNetObjectGroupMembership& Group
 {
 	for (FNetObjectGroupHandle AssignedGroup : GroupMembership.Groups)
 	{
-		if (AssignedGroup == InvalidNetObjectGroupHandle)
+		if (!AssignedGroup.IsValid())
 		{
 			// Note: An invalid group means we found the end of the array
 			return false;
 		}
-		else if (IsFilterGroup(Groups[AssignedGroup]))
+		else if (IsFilterGroup(Groups[AssignedGroup.GetGroupIndex()]))
 		{
 			return true;
 		}
@@ -306,7 +320,7 @@ uint32 FNetObjectGroups::GetNumGroupMemberships(FInternalNetRefIndex InternalInd
 	uint32 Count = 0U;
 	for (uint32 It=0; It < FNetObjectGroupMembership::MaxAssignedGroupCount; ++It)
 	{
-		Count += GroupMemberships[InternalIndex].Groups[It] != 0U;
+		Count += GroupMemberships[InternalIndex].Groups[It].IsValid() ? 1U : 0U;
 	}
 	
 	return Count;
diff --git a/Engine/Source/Runtime/Experimental/Iris/Core/Private/Iris/ReplicationSystem/Filtering/NetObjectGroups.h b/Engine/Source/Runtime/Experimental/Iris/Core/Private/Iris/ReplicationSystem/Filtering/NetObjectGroups.h
--- a/Engine/Source/Runtime/Experimental/Iris/Core/Private/Iris/ReplicationSystem/Filtering/NetObjectGroups.h
+++ b/Engine/Source/Runtime/Experimental/Iris/Core/Private/Iris/ReplicationSystem/Filtering/NetObjectGroups.h
@@ -8,10 +8,10 @@
 #include "Containers/SparseArray.h"
 #include "UObject/NameTypes.h"
 #include "Net/Core/NetBitArray.h"
+#include "Iris/ReplicationSystem/NetObjectGroupHandle.h"
 
 namespace UE::Net
 {
-	typedef uint16 FNetObjectGroupHandle;
 	namespace Private
 	{
 		typedef uint32 FInternalNetRefIndex;
@@ -61,9 +61,14 @@ public:
 	
 	const FNetObjectGroup* GetGroup(FNetObjectGroupHandle GroupHandle) const;
 	FNetObjectGroup* GetGroup(FNetObjectGroupHandle GroupHandle);
+	
+	FNetObjectGroupHandle MakeNetObjectGroupHandle(FNetObjectGroupHandle::FGroupIndexType GroupIndex) const { return FNetObjectGroupHandle(GroupIndex, CurrentEpoch); }
+	const FNetObjectGroup* GetGroupByIndex(FNetObjectGroupHandle::FGroupIndexType GroupIndex) const;
+	FNetObjectGroup* GetGroupByIndex(FNetObjectGroupHandle::FGroupIndexType GroupIndex);
+
 	void SetGroupName(FNetObjectGroupHandle GroupHandle, FName GroupName);
 
-	inline bool IsValidGroup(FNetObjectGroupHandle GroupHandle) const { return GroupHandle && Groups.IsValidIndex(GroupHandle); }
+	inline bool IsValidGroup(FNetObjectGroupHandle GroupHandle) const { return GroupHandle.IsValid() && GroupHandle.Epoch == CurrentEpoch && Groups.IsValidIndex(GroupHandle.GetGroupIndex()); }
 
 	bool Contains(FNetObjectGroupHandle GroupHandle, FInternalNetRefIndex InternalIndex) const;
 	void AddToGroup(FNetObjectGroupHandle GroupHandle, FInternalNetRefIndex InternalIndex);
@@ -125,6 +130,10 @@ private:
 	FNetBitArray GroupFilteredObjects;
 
 	TMap<FName, FNetObjectGroupHandle> NamedGroups;
+
+	uint32 CurrentEpoch = 0U;
+
+	inline static uint32 NextEpoch = 1U;
 };
 
 inline bool FNetObjectGroups::IsFilterGroup(const FNetObjectGroup& Group) const
diff --git a/Engine/Source/Runtime/Experimental/Iris/Core/Private/Iris/ReplicationSystem/Filtering/ReplicationFiltering.cpp b/Engine/Source/Runtime/Experimental/Iris/Core/Private/Iris/ReplicationSystem/Filtering/ReplicationFiltering.cpp
--- a/Engine/Source/Runtime/Experimental/Iris/Core/Private/Iris/ReplicationSystem/Filtering/ReplicationFiltering.cpp
+++ b/Engine/Source/Runtime/Experimental/Iris/Core/Private/Iris/ReplicationSystem/Filtering/ReplicationFiltering.cpp
@@ -641,10 +641,9 @@ void FReplicationFiltering::RemoveConnection(uint32 ConnectionId)
 	// Reset SubObject filter for removed connection
 	SubObjectFilterGroups.ForAllSetBits([this, ConnectionId](uint32 GroupIndex)
 	{
-		FNetObjectGroupHandle GroupHandle(static_cast<uint16>(GroupIndex));
-		if (!IsReservedNetObjectGroupHandle(GroupHandle))
+		if (!FNetObjectGroupHandle::IsReservedNetObjectGroupIndex(GroupIndex))
 		{
-			SetConnectionFilterStatus(*GetPerObjectInfo(GroupInfos[GroupHandle].ConnectionStateIndex), ConnectionId, ENetFilterStatus::Disallow);
+			SetConnectionFilterStatus(*GetPerObjectInfo(GroupInfos[GroupIndex].ConnectionStateIndex), ConnectionId, ENetFilterStatus::Disallow);
 			DirtySubObjectFilterGroups.SetBit(GroupIndex);
 		}
 	});
@@ -688,16 +687,16 @@ void FReplicationFiltering::InitNewConnections()
 
 		// Update group filter by iterating over all registered groups for the connection
 		{
-			auto InitGroupFilterForConnection = [this, &ConnectionInfo, ConnectionId](uint32 GroupIndex)
+			auto InitGroupFilterForConnection = [this, &ConnectionInfo, ConnectionId](uint32 GroupIndexIn)
 			{
-				const FNetObjectGroupHandle GroupHandle = static_cast<uint16>(GroupIndex);
+				const FNetObjectGroupHandle::FGroupIndexType GroupIndex = static_cast<FNetObjectGroupHandle::FGroupIndexType>(GroupIndexIn);
 				const FPerObjectInfo* ConnectionState = GetPerObjectInfo(this->GroupInfos[GroupIndex].ConnectionStateIndex);
 
 				// Setup filter for connection
 				if (GetConnectionFilterStatus(*ConnectionState, ConnectionId) == ENetFilterStatus::Disallow)
 				{
 					// Apply filter
-					const FNetObjectGroup* Group = this->Groups->GetGroup(GroupHandle);
+					const FNetObjectGroup* Group = this->Groups->GetGroupByIndex(GroupIndex);
 					FNetBitArray& GroupFilteredOutObjects = ConnectionInfo.GroupFilteredOutObjects;
 					for (uint32 ObjectIndex : MakeArrayView(Group->Members.GetData(), Group->Members.Num()))
 					{
@@ -890,9 +889,9 @@ void FReplicationFiltering::UpdateObjectsInScope()
 					const FNetObjectGroupHandle* GroupHandles = Groups->GetGroupMemberships(ParentIndex, GroupMembershipCount);
 					for (FNetObjectGroupHandle GroupHandle : MakeArrayView(GroupHandles, GroupMembershipCount))
 					{
-						if (FilterGroups.GetBit(GroupHandle))
+						if (FilterGroups.GetBit(GroupHandle.GetGroupIndex()))
 						{
-							DirtyFilterGroups.SetBit(GroupHandle);
+							DirtyFilterGroups.SetBit(GroupHandle.GetGroupIndex());
 						}
 					}
 
@@ -1076,7 +1075,7 @@ void FReplicationFiltering::UpdateGroupFiltering()
 			IRIS_PROFILER_SCOPE(FReplicationFiltering_UpdateGroupFiltering);
 
 			const FPerObjectInfo* ConnectionStateInfo = GetPerObjectInfo(this->GroupInfos[GroupIndex].ConnectionStateIndex);
-			const FNetObjectGroup* Group = this->Groups->GetGroup(FNetObjectGroupHandle(GroupIndex));
+			const FNetObjectGroup* Group = this->Groups->GetGroupByIndex(FNetObjectGroupHandle::FGroupIndexType(GroupIndex));
 			FPerConnectionInfo* LocalConnectionInfos = this->ConnectionInfos.GetData();
 
 			auto UpdateGroupFilterForConnection = [this, ConnectionStateInfo, Group, LocalConnectionInfos](uint32 ConnectionId)
@@ -1493,35 +1492,37 @@ const FReplicationFiltering::FPerObjectInfo* FReplicationFiltering::GetPerObject
 
 void FReplicationFiltering::AddSubObjectFilter(FNetObjectGroupHandle GroupHandle)
 {
-	check(!FilterGroups.GetBit(GroupHandle));
+	const FNetObjectGroupHandle::FGroupIndexType GroupIndex = GroupHandle.GetGroupIndex();
+	check(!FilterGroups.GetBit(GroupIndex));
 
-	if (ensure(GroupHandle != InvalidNetObjectGroupHandle && !SubObjectFilterGroups.GetBit(GroupHandle)))
+	if (ensure(GroupHandle.IsValid() && !SubObjectFilterGroups.GetBit(GroupIndex)))
 	{
-		const bool bGroupExists = SubObjectFilterGroups.GetBit(GroupHandle);
+		const bool bGroupExists = SubObjectFilterGroups.GetBit(GroupIndex);
 		if (ensure(bGroupExists == false))
 		{
-			SubObjectFilterGroups.SetBit(GroupHandle);
-			GroupInfos[GroupHandle].ConnectionStateIndex = AllocPerObjectInfo();
+			SubObjectFilterGroups.SetBit(GroupIndex);
+			GroupInfos[GroupIndex].ConnectionStateIndex = AllocPerObjectInfo();
 
 			// By default filter is active, i.e. we filter out all connections
-			SetPerObjectInfoFilterStatus(*GetPerObjectInfo(GroupInfos[GroupHandle].ConnectionStateIndex), ENetFilterStatus::Disallow);
+			SetPerObjectInfoFilterStatus(*GetPerObjectInfo(GroupInfos[GroupIndex].ConnectionStateIndex), ENetFilterStatus::Disallow);
 
-			UE_LOG(LogIrisFiltering, Verbose, TEXT("ReplicationFiltering::AddSubObjectFilter (%s) GroupIndex: %u, FilterStatus: DisallowReplication"), *(Groups->GetGroup(GroupHandle)->GroupName.ToString()), GroupHandle);
+			UE_LOG(LogIrisFiltering, Verbose, TEXT("ReplicationFiltering::AddSubObjectFilter (%s) GroupIndex: %u, FilterStatus: DisallowReplication"), *(Groups->GetGroup(GroupHandle)->GroupName.ToString()), GroupIndex);
 		}
 	}
 }
 
 void FReplicationFiltering::RemoveSubObjectFilter(FNetObjectGroupHandle GroupHandle)
 {
-	if (GroupHandle != InvalidNetObjectGroupHandle  && SubObjectFilterGroups.GetBit(GroupHandle))
+	const FNetObjectGroupHandle::FGroupIndexType GroupIndex = GroupHandle.GetGroupIndex();
+	if (GroupHandle.IsValid() && SubObjectFilterGroups.GetBit(GroupIndex))
 	{
 		// Mark group as no longer a SubObjectFilter group
-		SubObjectFilterGroups.ClearBit(GroupHandle);
-		const PerObjectInfoIndexType ConnectionStateIndex = GroupInfos[GroupHandle].ConnectionStateIndex;
-		GroupInfos[GroupHandle].ConnectionStateIndex = 0U;
+		SubObjectFilterGroups.ClearBit(GroupIndex);
+		const PerObjectInfoIndexType ConnectionStateIndex = GroupInfos[GroupIndex].ConnectionStateIndex;
+		GroupInfos[GroupIndex].ConnectionStateIndex = 0U;
 		FreePerObjectInfo(ConnectionStateIndex);
 
-		UE_LOG(LogIrisFiltering, Verbose, TEXT("ReplicationFiltering::RemoveSubObjectFilter GroupIndex: %u"), GroupHandle);
+		UE_LOG(LogIrisFiltering, Verbose, TEXT("ReplicationFiltering::RemoveSubObjectFilter GroupIndex: %u"), GroupIndex);
 	}
 }
 
@@ -1530,10 +1531,10 @@ void FReplicationFiltering::UpdateSubObjectFilters()
 	// We want to remove all groups that have no members and no enabled connections
 	auto UpdateSubObjectFilterGroup = [this](uint32 GroupIndex)
 	{
-		FNetObjectGroupHandle GroupHandle(static_cast<uint16>(GroupIndex));
-		if (const FNetObjectGroup* Group = Groups->GetGroup(FNetObjectGroupHandle(GroupIndex)))
+		FNetObjectGroupHandle GroupHandle = Groups->MakeNetObjectGroupHandle(static_cast<FNetObjectGroupHandle::FGroupIndexType>(GroupIndex));
+		if (const FNetObjectGroup* Group = Groups->GetGroup(GroupHandle))
 		{
-			if (Group->Members.Num() == 0U && !IsConnectionFilterStatusAllowedForAnyConnection(*GetPerObjectInfo(GroupInfos[GroupHandle].ConnectionStateIndex)))
+			if (Group->Members.Num() == 0U && !IsConnectionFilterStatusAllowedForAnyConnection(*GetPerObjectInfo(GroupInfos[GroupHandle.GetGroupIndex()].ConnectionStateIndex)))
 			{
 				ReplicationSystem->DestroyGroup(GroupHandle);
 			}
@@ -1548,23 +1549,23 @@ void FReplicationFiltering::SetSubObjectFilterStatus(FNetObjectGroupHandle Group
 {
 	IRIS_PROFILER_SCOPE(SetSubObjectFilterStatus)
 
-	if (IsReservedNetObjectGroupHandle(GroupHandle))
+	if (GroupHandle.IsReservedNetObjectGroup())
 	{
-		UE_LOG(LogIrisFiltering, Warning, TEXT("FReplicationFiltering::SetSubObjectFilterStatus - Trying to set filter for reserved GroupIndex: %u which is not allowed."), GroupHandle);
+		UE_LOG(LogIrisFiltering, Warning, TEXT("FReplicationFiltering::SetSubObjectFilterStatus - Trying to set filter for reserved GroupIndex: %u which is not allowed."), GroupHandle.GetGroupIndex());
 		return;
 	}
 
-	if (!SubObjectFilterGroups.GetBit(GroupHandle))
+	if (!SubObjectFilterGroups.GetBit(GroupHandle.GetGroupIndex()))
 	{
-		UE_LOG(LogIrisFiltering, Warning, TEXT("FReplicationFiltering::SetSubObjectFilterStatus - Trying to ReplicationStatus for GroupIndex: %u that is not a SubObjectFilterGroup"), GroupHandle);
+		UE_LOG(LogIrisFiltering, Warning, TEXT("FReplicationFiltering::SetSubObjectFilterStatus - Trying to ReplicationStatus for GroupIndex: %u that is not a SubObjectFilterGroup"), GroupHandle.GetGroupIndex());
 		return;
 	}
 
-	FPerObjectInfo* FilterInfo = GetPerObjectInfo(GroupInfos[GroupHandle].ConnectionStateIndex);
+	FPerObjectInfo* FilterInfo = GetPerObjectInfo(GroupInfos[GroupHandle.GetGroupIndex()].ConnectionStateIndex);
 	SetPerObjectInfoFilterStatus(*FilterInfo, ReplicationStatus);
 	if (!IsConnectionFilterStatusAllowedForAnyConnection(*FilterInfo))
 	{
-		DirtySubObjectFilterGroups.SetBit(GroupHandle);
+		DirtySubObjectFilterGroups.SetBit(GroupHandle.GetGroupIndex());
 	}
 }
 
@@ -1572,19 +1573,19 @@ void FReplicationFiltering::SetSubObjectFilterStatus(FNetObjectGroupHandle Group
 {
 	if (ConnectionsBitArray.GetNumBits() > ValidConnections.GetNumBits())
 	{
-		UE_LOG(LogIrisFiltering, Warning, TEXT("FReplicationFiltering::SetSubObjectFilterStatus - Trying to set filter for %u, with invalid Connections parameters."), GroupHandle);
+		UE_LOG(LogIrisFiltering, Warning, TEXT("FReplicationFiltering::SetSubObjectFilterStatus - Trying to set filter for %u, with invalid Connections parameters."), GroupHandle.GetGroupIndex());
 		return;
 	}
 
-	if (IsReservedNetObjectGroupHandle(GroupHandle))
+	if (GroupHandle.IsReservedNetObjectGroup())
 	{
-		UE_LOG(LogIrisFiltering, Warning, TEXT("FReplicationFiltering::SetSubObjectFilterStatus - Trying to set filter for reserved GroupIndex: %u which is not allowed."), GroupHandle);
+		UE_LOG(LogIrisFiltering, Warning, TEXT("FReplicationFiltering::SetSubObjectFilterStatus - Trying to set filter for reserved GroupIndex: %u which is not allowed."), GroupHandle.GetGroupIndex());
 		return;
 	}
 
-	if (!SubObjectFilterGroups.GetBit(GroupHandle))
+	if (!SubObjectFilterGroups.GetBit(GroupHandle.GetGroupIndex()))
 	{
-		UE_LOG(LogIrisFiltering, Warning, TEXT("FReplicationFiltering::SetSubObjectFilterStatus - Trying to ReplicationStatus for GroupIndex: %u that is not a SubObjectFilterGroup"), GroupHandle);
+		UE_LOG(LogIrisFiltering, Warning, TEXT("FReplicationFiltering::SetSubObjectFilterStatus - Trying to ReplicationStatus for GroupIndex: %u that is not a SubObjectFilterGroup"), GroupHandle.GetGroupIndex());
 		return;
 	}
 
@@ -1608,32 +1609,32 @@ void FReplicationFiltering::SetSubObjectFilterStatus(FNetObjectGroupHandle Group
 
 void FReplicationFiltering::SetSubObjectFilterStatus(FNetObjectGroupHandle GroupHandle, uint32 ConnectionId, ENetFilterStatus ReplicationStatus)
 {
-	if (IsReservedNetObjectGroupHandle(GroupHandle))
+	if (GroupHandle.IsReservedNetObjectGroup())
 	{
-		UE_LOG(LogIrisFiltering, Warning, TEXT("FReplicationFiltering::SetSubObjectFilterStatus - Trying to set filter for reserved GroupIndex: %u which is not allowed."), GroupHandle);
+		UE_LOG(LogIrisFiltering, Warning, TEXT("FReplicationFiltering::SetSubObjectFilterStatus - Trying to set filter for reserved GroupIndex: %u which is not allowed."), GroupHandle.GetGroupIndex());
 		return;
 	}
 
-	if (ensure(ValidConnections.GetBit(ConnectionId) && SubObjectFilterGroups.GetBit(GroupHandle)))
+	if (ensure(ValidConnections.GetBit(ConnectionId) && SubObjectFilterGroups.GetBit(GroupHandle.GetGroupIndex())))
 	{
 		UE_LOG(LogIrisFiltering, Verbose, TEXT("ReplicationFiltering::SetSubObjectFilterStatus GroupIndex: %u, ConnectionId: %u, FilterStatus: %u"), GroupHandle, ConnectionId, ReplicationStatus == ENetFilterStatus::Allow ? 1U : 0U);
-		FPerObjectInfo* FilterInfo = GetPerObjectInfo(GroupInfos[GroupHandle].ConnectionStateIndex);
+		FPerObjectInfo* FilterInfo = GetPerObjectInfo(GroupInfos[GroupHandle.GetGroupIndex()].ConnectionStateIndex);
 		SetConnectionFilterStatus(*FilterInfo, ConnectionId, ReplicationStatus);
 		if (!IsConnectionFilterStatusAllowedForAnyConnection(*FilterInfo))
 		{
-			DirtySubObjectFilterGroups.SetBit(GroupHandle);
+			DirtySubObjectFilterGroups.SetBit(GroupHandle.GetGroupIndex());
 		}
 	}
 }
 
 bool FReplicationFiltering::GetSubObjectFilterStatus(FNetObjectGroupHandle GroupHandle, uint32 ConnectionId, ENetFilterStatus& OutReplicationStatus) const
 {
-	if (!(ValidConnections.GetBit(ConnectionId) && SubObjectFilterGroups.GetBit(GroupHandle)))
+	if (!(ValidConnections.GetBit(ConnectionId) && SubObjectFilterGroups.GetBit(GroupHandle.GetGroupIndex())))
 	{
 		return false;
 	}
 
-	const FPerObjectInfo* ConnectionState = GetPerObjectInfo(GroupInfos[GroupHandle].ConnectionStateIndex);
+	const FPerObjectInfo* ConnectionState = GetPerObjectInfo(GroupInfos[GroupHandle.GetGroupIndex()].ConnectionStateIndex);
 	OutReplicationStatus = GetConnectionFilterStatus(*ConnectionState, ConnectionId);
 
 	return true;
@@ -1641,27 +1642,26 @@ bool FReplicationFiltering::GetSubObjectFilterStatus(FNetObjectGroupHandle Group
 
 void FReplicationFiltering::AddGroupFilter(FNetObjectGroupHandle GroupHandle)
 {
-	const bool bIsValidGroup = ensureMsgf(GroupHandle != InvalidNetObjectGroupHandle, TEXT("AddGroupFilter received an invalid group handle"));
-
+	const bool bIsValidGroup = ensureMsgf(Groups->IsValidGroup(GroupHandle), TEXT("AddGroupFilter received an invalid group handle"));
 	if (bIsValidGroup)
 	{
-		const bool bIsFiltering = FilterGroups.GetBit(GroupHandle);
-		ensureMsgf(!bIsFiltering, TEXT("NetObjectGroup %u was asked to start filtering but it was already doing so."), GroupHandle);
+		const bool bIsFiltering = FilterGroups.GetBit(GroupHandle.GetGroupIndex());
+		ensureMsgf(!bIsFiltering, TEXT("NetObjectGroup %u was asked to start filtering but it was already doing so."), GroupHandle.GetGroupIndex());
 
-		check(!SubObjectFilterGroups.GetBit(GroupHandle));
+		check(!SubObjectFilterGroups.GetBit(GroupHandle.GetGroupIndex()));
 		
 		if (!bIsFiltering)
 		{
 			Groups->AddFilterTrait(GroupHandle);
 
-			FilterGroups.SetBit(GroupHandle);
-			DirtyFilterGroups.SetBit(GroupHandle);
-			GroupInfos[GroupHandle].ConnectionStateIndex = AllocPerObjectInfo();
+			FilterGroups.SetBit(GroupHandle.GetGroupIndex());
+			DirtyFilterGroups.SetBit(GroupHandle.GetGroupIndex());
+			GroupInfos[GroupHandle.GetGroupIndex()].ConnectionStateIndex = AllocPerObjectInfo();
 
 			// By default filter is active, i.e. we filter out all connections
-			SetPerObjectInfoFilterStatus(*GetPerObjectInfo(GroupInfos[GroupHandle].ConnectionStateIndex), ENetFilterStatus::Disallow);
+			SetPerObjectInfoFilterStatus(*GetPerObjectInfo(GroupInfos[GroupHandle.GetGroupIndex()].ConnectionStateIndex), ENetFilterStatus::Disallow);
 
-			UE_LOG(LogIrisFiltering, Verbose, TEXT("ReplicationFiltering::AddGroupFilter GroupIndex: %u, FilterStatus: DisallowReplication"), GroupHandle);
+			UE_LOG(LogIrisFiltering, Verbose, TEXT("ReplicationFiltering::AddGroupFilter GroupIndex: %u, FilterStatus: DisallowReplication"), GroupHandle.GetGroupIndex());
 		}
 	}
 }
@@ -1669,7 +1669,7 @@ void FReplicationFiltering::AddGroupFilter(FNetObjectGroupHandle GroupHandle)
 void FReplicationFiltering::RemoveGroupFilter(FNetObjectGroupHandle GroupHandle)
 {
 	// Remove the filter trait from the group
-	if (GroupHandle && FilterGroups.GetBit(GroupHandle))
+	if (Groups->IsValidGroup(GroupHandle) && FilterGroups.GetBit(GroupHandle.GetGroupIndex()))
 	{
 		ValidConnections.ForAllSetBits( [GroupHandle, this](uint32 ConnectionIndex) 
 		{ 
@@ -1677,14 +1677,14 @@ void FReplicationFiltering::RemoveGroupFilter(FNetObjectGroupHandle GroupHandle)
 		});
 
 		// Mark group as no longer a group filter group
-		FilterGroups.ClearBit(GroupHandle);
-		const PerObjectInfoIndexType ConnectionStateIndex = GroupInfos[GroupHandle].ConnectionStateIndex;
-		GroupInfos[GroupHandle].ConnectionStateIndex = 0U;
+		FilterGroups.ClearBit(GroupHandle.GetGroupIndex());
+		const PerObjectInfoIndexType ConnectionStateIndex = GroupInfos[GroupHandle.GetGroupIndex()].ConnectionStateIndex;
+		GroupInfos[GroupHandle.GetGroupIndex()].ConnectionStateIndex = 0U;
 		FreePerObjectInfo(ConnectionStateIndex);
 
 		Groups->RemoveFilterTrait(GroupHandle);
 
-		UE_LOG(LogIrisFiltering, Verbose, TEXT("ReplicationFiltering::RemoveGroupFilter GroupIndex: %u"), GroupHandle);
+		UE_LOG(LogIrisFiltering, Verbose, TEXT("ReplicationFiltering::RemoveGroupFilter GroupIndex: %u"), GroupHandle.GetGroupIndex());
 	}
 }
 
@@ -1692,15 +1692,15 @@ void FReplicationFiltering::SetGroupFilterStatus(FNetObjectGroupHandle GroupHand
 {
 	IRIS_PROFILER_SCOPE(SetGroupFilterStatus)
 
-	if (IsReservedNetObjectGroupHandle(GroupHandle))
+	if (GroupHandle.IsReservedNetObjectGroup())
 	{
-		UE_LOG(LogIrisFiltering, Warning, TEXT("FReplicationFiltering::SetGroupFilterStatus - Trying to set filter for reserved GroupIndex: %u which is not allowed."), GroupHandle);
+		UE_LOG(LogIrisFiltering, Warning, TEXT("FReplicationFiltering::SetGroupFilterStatus - Trying to set filter for reserved GroupIndex: %u which is not allowed."), GroupHandle.GetGroupIndex());
 		return;
 	}
 
-	if (!FilterGroups.GetBit(GroupHandle))
+	if (!FilterGroups.GetBit(GroupHandle.GetGroupIndex()))
 	{
-		UE_LOG(LogIrisFiltering, Warning, TEXT("FReplicationFiltering::SetGroupFilterStatus - Trying to set filter for invalid GroupIndex: %u, Make sure group is added to filtering"), GroupHandle);
+		UE_LOG(LogIrisFiltering, Warning, TEXT("FReplicationFiltering::SetGroupFilterStatus - Trying to set filter for invalid GroupIndex: %u, Make sure group is added to filtering"), GroupHandle.GetGroupIndex());
 		return;
 	}
 	
@@ -1715,21 +1715,21 @@ void FReplicationFiltering::SetGroupFilterStatus(FNetObjectGroupHandle GroupHand
 {
 	IRIS_PROFILER_SCOPE(SetGroupFilterStatus)
 		
-	if (IsReservedNetObjectGroupHandle(GroupHandle))
+	if (GroupHandle.IsReservedNetObjectGroup())
 	{
-		UE_LOG(LogIrisFiltering, Warning, TEXT("FReplicationFiltering::SetGroupFilterStatus - Trying to set filter for reserved GroupIndex: %u which is not allowed."), GroupHandle);		
+		UE_LOG(LogIrisFiltering, Warning, TEXT("FReplicationFiltering::SetGroupFilterStatus - Trying to set filter for reserved GroupIndex: %u which is not allowed."), GroupHandle.GetGroupIndex());		
 		return;
 	}
 
-	if (!FilterGroups.GetBit(GroupHandle))
+	if (!FilterGroups.GetBit(GroupHandle.GetGroupIndex()))
 	{
-		UE_LOG(LogIrisFiltering, Warning, TEXT("FReplicationFiltering::SetGroupFilterStatus - Trying to set filter for invalid GroupIndex: %u, Make sure group is added to filtering"), GroupHandle);
+		UE_LOG(LogIrisFiltering, Warning, TEXT("FReplicationFiltering::SetGroupFilterStatus - Trying to set filter for invalid GroupIndex: %u, Make sure group is added to filtering"), GroupHandle.GetGroupIndex());
 		return;
 	}
 
 	if (ConnectionsBitArray.GetNumBits() > ValidConnections.GetNumBits())
 	{
-		UE_LOG(LogIrisFiltering, Warning, TEXT("FReplicationFiltering::SetGroupFilterStatus - Trying to set filter for %u, with invalid Connections parameters."), GroupHandle);
+		UE_LOG(LogIrisFiltering, Warning, TEXT("FReplicationFiltering::SetGroupFilterStatus - Trying to set filter for %u, with invalid Connections parameters."), GroupHandle.GetGroupIndex());
 		return;
 	}
 
@@ -1749,15 +1749,15 @@ void FReplicationFiltering::SetGroupFilterStatus(FNetObjectGroupHandle GroupHand
 {
 	IRIS_PROFILER_SCOPE(SetGroupFilterStatus)
 
-	if (IsReservedNetObjectGroupHandle(GroupHandle))
+	if (GroupHandle.IsReservedNetObjectGroup())
 	{
-		UE_LOG(LogIrisFiltering, Warning, TEXT("FReplicationFiltering::SetGroupFilterStatus - Trying to set filter for reserved GroupIndex: %u which is not allowed."), GroupHandle);
+		UE_LOG(LogIrisFiltering, Warning, TEXT("FReplicationFiltering::SetGroupFilterStatus - Trying to set filter for reserved GroupIndex: %u which is not allowed."), GroupHandle.GetGroupIndex());
 		return;
 	}
 
-	if (!FilterGroups.GetBit(GroupHandle))
+	if (!FilterGroups.GetBit(GroupHandle.GetGroupIndex()))
 	{
-		UE_LOG(LogIrisFiltering, Warning, TEXT("FReplicationFiltering::SetGroupFilterStatus - Trying to set invalid filter ConnectionId: %u, GroupIndex: %u"), ConnectionId, GroupHandle);
+		UE_LOG(LogIrisFiltering, Warning, TEXT("FReplicationFiltering::SetGroupFilterStatus - Trying to set invalid filter ConnectionId: %u, GroupIndex: %u"), ConnectionId, GroupHandle.GetGroupIndex());
 		return;
 	}
 	
@@ -1776,9 +1776,9 @@ bool FReplicationFiltering::GetIsFilteredOutByAnyGroup(uint32 ObjectInternalInde
 
 	for (FNetObjectGroupHandle GroupHandle : MakeArrayView(GroupHandles, GroupMembershipCount))
 	{
-		if (FilterGroups.GetBit(GroupHandle))
+		if (FilterGroups.GetBit(GroupHandle.GetGroupIndex()))
 		{
-			const FPerObjectInfo* ConnectionState = GetPerObjectInfo(GroupInfos[GroupHandle].ConnectionStateIndex);
+			const FPerObjectInfo* ConnectionState = GetPerObjectInfo(GroupInfos[GroupHandle.GetGroupIndex()].ConnectionStateIndex);
 			if (GetConnectionFilterStatus(*ConnectionState, ConnectionId) == ENetFilterStatus::Disallow)
 			{
 				return true;
@@ -1823,22 +1823,22 @@ bool FReplicationFiltering::UpdateGroupFilterEffectsForObject(uint32 ObjectIndex
 
 void FReplicationFiltering::InternalSetGroupFilterStatus(FNetObjectGroupHandle GroupHandle, uint32 ConnectionId, ENetFilterStatus ReplicationStatus)
 {
-	if (GroupHandle == NotReplicatedNetObjectGroupHandle)
+	if (GroupHandle.IsNotReplicatedNetObjectGroup())
 	{
-		UE_LOG(LogIrisFiltering, Warning, TEXT("FReplicationFiltering::SetGroupFilterStatus - Trying to set filter NotReplicatedNetObjectGroupHandle GroupIndex: %u which is not allowed."), GroupHandle);		
+		UE_LOG(LogIrisFiltering, Warning, TEXT("FReplicationFiltering::SetGroupFilterStatus - Trying to set filter statua on NotReplicatedNetObjectGroupHandle GroupIndex: %u which is not allowed."), GroupHandle.GetGroupIndex());		
 		return;
 	}
 
-	if (!FilterGroups.GetBit(GroupHandle))
+	if (!FilterGroups.GetBit(GroupHandle.GetGroupIndex()))
 	{
-		UE_LOG(LogIrisFiltering, Warning, TEXT("FReplicationFiltering::SetGroupFilterStatus - Trying to set invalid filter ConnectionId: %u, GroupIndex: %u"), ConnectionId, GroupHandle);
+		UE_LOG(LogIrisFiltering, Warning, TEXT("FReplicationFiltering::SetGroupFilterStatus - Trying to set invalid filter ConnectionId: %u, GroupIndex: %u"), ConnectionId, GroupHandle.GetGroupIndex());
 		return;
 	}
 
 	if (ReplicationStatus == ENetFilterStatus::Disallow)
 	{
 		// Filter out objects in group
-		FPerObjectInfo* ConnectionState = GetPerObjectInfo(GroupInfos[GroupHandle].ConnectionStateIndex);
+		FPerObjectInfo* ConnectionState = GetPerObjectInfo(GroupInfos[GroupHandle.GetGroupIndex()].ConnectionStateIndex);
 		if (GetConnectionFilterStatus(*ConnectionState, ConnectionId) == ENetFilterStatus::Allow)
 		{
 			IRIS_PROFILER_SCOPE(FReplicationFiltering_InternalSetGroupFilterStatus_Disallow);
@@ -1846,7 +1846,7 @@ void FReplicationFiltering::InternalSetGroupFilterStatus(FNetObjectGroupHandle G
 			// Mark filter as active for the connection
 			SetConnectionFilterStatus(*ConnectionState, ConnectionId, ENetFilterStatus::Disallow);
 
-			UE_LOG(LogIrisFiltering, Verbose, TEXT("ReplicationFiltering::SetGroupFilterStatus GroupIndex: %u, ConnectionId: %u, FilterStatus: DisallowReplication"), GroupHandle, ConnectionId);
+			UE_LOG(LogIrisFiltering, Verbose, TEXT("ReplicationFiltering::SetGroupFilterStatus GroupIndex: %u, ConnectionId: %u, FilterStatus: DisallowReplication"), GroupHandle.GetGroupIndex(), ConnectionId);
 
 			if (ValidConnections.GetBit(ConnectionId))
 			{
@@ -1878,7 +1878,7 @@ void FReplicationFiltering::InternalSetGroupFilterStatus(FNetObjectGroupHandle G
 	else
 	{
 		// Clear filter effects
-		FPerObjectInfo* ConnectionState = GetPerObjectInfo(GroupInfos[GroupHandle].ConnectionStateIndex);
+		FPerObjectInfo* ConnectionState = GetPerObjectInfo(GroupInfos[GroupHandle.GetGroupIndex()].ConnectionStateIndex);
 		if (GetConnectionFilterStatus(*ConnectionState, ConnectionId) == ENetFilterStatus::Disallow)
 		{
 			IRIS_PROFILER_SCOPE(FReplicationFiltering_InternalSetGroupFilterStatus_Allow);
@@ -1886,7 +1886,7 @@ void FReplicationFiltering::InternalSetGroupFilterStatus(FNetObjectGroupHandle G
 			// Mark filter as no longer being active for the connection
 			SetConnectionFilterStatus(*ConnectionState, ConnectionId, ENetFilterStatus::Allow);
 
-			UE_LOG(LogIrisFiltering, Verbose, TEXT("ReplicationFiltering::SetGroupFilterStatus GroupIndex: %u, ConnectionId: %u, FilterStatus: AllowReplication"), GroupHandle, ConnectionId);
+			UE_LOG(LogIrisFiltering, Verbose, TEXT("ReplicationFiltering::SetGroupFilterStatus GroupIndex: %u, ConnectionId: %u, FilterStatus: AllowReplication"), GroupHandle.GetGroupIndex(), ConnectionId);
 
 			if (ValidConnections.GetBit(ConnectionId))
 			{
@@ -1910,12 +1910,12 @@ void FReplicationFiltering::InternalSetGroupFilterStatus(FNetObjectGroupHandle G
 
 bool FReplicationFiltering::GetGroupFilterStatus(FNetObjectGroupHandle GroupHandle, uint32 ConnectionId, ENetFilterStatus& OutReplicationStatus) const
 {
-	if (!(ValidConnections.GetBit(ConnectionId) && FilterGroups.GetBit(GroupHandle)))
+	if (!(ValidConnections.GetBit(ConnectionId) && FilterGroups.GetBit(GroupHandle.GetGroupIndex())))
 	{
 		return false;
 	}
 
-	const FPerObjectInfo* ConnectionState = GetPerObjectInfo(GroupInfos[GroupHandle].ConnectionStateIndex);
+	const FPerObjectInfo* ConnectionState = GetPerObjectInfo(GroupInfos[GroupHandle.GetGroupIndex()].ConnectionStateIndex);
 	OutReplicationStatus = GetConnectionFilterStatus(*ConnectionState, ConnectionId);
 
 	return true;
@@ -1924,15 +1924,15 @@ bool FReplicationFiltering::GetGroupFilterStatus(FNetObjectGroupHandle GroupHand
 void FReplicationFiltering::NotifyObjectAddedToGroup(FNetObjectGroupHandle GroupHandle, FInternalNetRefIndex ObjectIndex)
 {
 	// As we want to avoid the fundamentally bad case of doing tons of single object adds to a group and updating group filter for every call we deffer the update until next call to filter
-	if (FilterGroups.GetBit(GroupHandle))
+	if (FilterGroups.GetBit(GroupHandle.GetGroupIndex()))
 	{
-		DirtyFilterGroups.SetBit(GroupHandle);
+		DirtyFilterGroups.SetBit(GroupHandle.GetGroupIndex());
 
-		UE_LOG(LogIrisFiltering, Verbose, TEXT("ReplicationFiltering::NotifyObjectAddedToGroup Added %s to Groupfilter with GroupIndex: %u"), *(NetRefHandleManager->GetReplicatedObjectDataNoCheck(ObjectIndex).RefHandle.ToString()), GroupHandle);
+		UE_LOG(LogIrisFiltering, Verbose, TEXT("ReplicationFiltering::NotifyObjectAddedToGroup Added %s to Groupfilter with GroupIndex: %u"), *(NetRefHandleManager->GetReplicatedObjectDataNoCheck(ObjectIndex).RefHandle.ToString()), GroupHandle.GetGroupIndex());
 	}
-	else if (SubObjectFilterGroups.GetBit(GroupHandle))
+	else if (SubObjectFilterGroups.GetBit(GroupHandle.GetGroupIndex()))
 	{
-		UE_LOG(LogIrisFiltering, Verbose, TEXT("ReplicationFiltering::NotifyObjectAddedToGroup Added %s to SubObjectFilter with GroupIndex: %u"), *(NetRefHandleManager->GetReplicatedObjectDataNoCheck(ObjectIndex).RefHandle.ToString()), GroupHandle);
+		UE_LOG(LogIrisFiltering, Verbose, TEXT("ReplicationFiltering::NotifyObjectAddedToGroup Added %s to SubObjectFilter with GroupIndex: %u"), *(NetRefHandleManager->GetReplicatedObjectDataNoCheck(ObjectIndex).RefHandle.ToString()), GroupHandle.GetGroupIndex());
 	}
 }
 
@@ -1940,9 +1940,9 @@ void FReplicationFiltering::NotifyObjectRemovedFromGroup(FNetObjectGroupHandle G
 {
 	UE_LOG(LogIrisFiltering, Verbose, TEXT("ReplicationFiltering::NotifyObjectRemovedFromGroup GroupIndex: %u, %s"), GroupHandle, *(NetRefHandleManager->GetReplicatedObjectDataNoCheck(ObjectIndex).RefHandle.ToString()));
 
-	if (FilterGroups.GetBit(GroupHandle))
+	if (FilterGroups.GetBit(GroupHandle.GetGroupIndex()))
 	{
-		FPerObjectInfo* ConnectionState = GetPerObjectInfo(GroupInfos[GroupHandle].ConnectionStateIndex);
+		FPerObjectInfo* ConnectionState = GetPerObjectInfo(GroupInfos[GroupHandle.GetGroupIndex()].ConnectionStateIndex);
 		const FNetBitArray& ScopedObjects = NetRefHandleManager->GetScopableInternalIndices();
 
 		auto MarkGroupsDirty = [this, ObjectIndex, GroupHandle, ConnectionState, &ScopedObjects](uint32 ConnectionId)
@@ -1957,9 +1957,9 @@ void FReplicationFiltering::NotifyObjectRemovedFromGroup(FNetObjectGroupHandle G
 		// We need to update groupfilter status for all initialized connections
 		FNetBitArray::ForAllSetBits(ValidConnections, NewConnections, FNetBitArrayBase::AndNotOp, MarkGroupsDirty);
 	}
-	else if (SubObjectFilterGroups.GetBit(GroupHandle))
+	else if (SubObjectFilterGroups.GetBit(GroupHandle.GetGroupIndex()))
 	{		
-		DirtySubObjectFilterGroups.SetBit(GroupHandle);
+		DirtySubObjectFilterGroups.SetBit(GroupHandle.GetGroupIndex());
 	}
 }
 
diff --git a/Engine/Source/Runtime/Experimental/Iris/Core/Private/Iris/ReplicationSystem/Filtering/ReplicationFiltering.h b/Engine/Source/Runtime/Experimental/Iris/Core/Private/Iris/ReplicationSystem/Filtering/ReplicationFiltering.h
--- a/Engine/Source/Runtime/Experimental/Iris/Core/Private/Iris/ReplicationSystem/Filtering/ReplicationFiltering.h
+++ b/Engine/Source/Runtime/Experimental/Iris/Core/Private/Iris/ReplicationSystem/Filtering/ReplicationFiltering.h
@@ -5,6 +5,7 @@
 #include "CoreTypes.h"
 #include "Net/Core/NetBitArray.h"
 #include "Iris/ReplicationSystem/Filtering/NetObjectFilter.h"
+#include "Iris/ReplicationSystem/NetObjectGroupHandle.h"
 #include "Containers/Array.h"
 #include "UObject/StrongObjectPtr.h"
 
@@ -12,7 +13,6 @@ class UReplicationSystem;
 namespace UE::Net
 {
 	typedef uint32 FNetObjectFilterHandle;
-	typedef uint16 FNetObjectGroupHandle;
 	namespace Private
 	{
 		class FDeltaCompressionBaselineInvalidationTracker;
@@ -107,7 +107,7 @@ public:
 	// Group based filtering
 	void AddGroupFilter(FNetObjectGroupHandle GroupHandle);
 	void RemoveGroupFilter(FNetObjectGroupHandle GroupHandle);
-	bool IsGroupFilterGroup(FNetObjectGroupHandle GroupHandle) const { return GroupHandle && FilterGroups.GetBit(GroupHandle); }
+	bool IsGroupFilterGroup(FNetObjectGroupHandle GroupHandle) const { return GroupHandle.IsValid() && FilterGroups.GetBit(GroupHandle.GetGroupIndex()); }
 
 	void SetGroupFilterStatus(FNetObjectGroupHandle GroupHandle, ENetFilterStatus ReplicationStatus);
 	void SetGroupFilterStatus(FNetObjectGroupHandle GroupHandle, const FNetBitArrayView& ConnectionsBitArray, ENetFilterStatus);
@@ -123,7 +123,7 @@ public:
 	// SubObjectFilter status
 	void AddSubObjectFilter(FNetObjectGroupHandle GroupHandle);
 	void RemoveSubObjectFilter(FNetObjectGroupHandle GroupHandle);
-	bool IsSubObjectFilterGroup(FNetObjectGroupHandle GroupHandle) const { return GroupHandle && SubObjectFilterGroups.GetBit(GroupHandle); }
+	bool IsSubObjectFilterGroup(FNetObjectGroupHandle GroupHandle) const { return GroupHandle.IsValid() && SubObjectFilterGroups.GetBit(GroupHandle.GetGroupIndex()); }
 
 	void SetSubObjectFilterStatus(FNetObjectGroupHandle GroupHandle, ENetFilterStatus ReplicationStatus);
 	void SetSubObjectFilterStatus(FNetObjectGroupHandle GroupHandle, const FNetBitArrayView& ConnectionsBitArray, ENetFilterStatus);
diff --git a/Engine/Source/Runtime/Experimental/Iris/Core/Private/Iris/ReplicationSystem/ReplicationBridge.cpp b/Engine/Source/Runtime/Experimental/Iris/Core/Private/Iris/ReplicationSystem/ReplicationBridge.cpp
--- a/Engine/Source/Runtime/Experimental/Iris/Core/Private/Iris/ReplicationSystem/ReplicationBridge.cpp
+++ b/Engine/Source/Runtime/Experimental/Iris/Core/Private/Iris/ReplicationSystem/ReplicationBridge.cpp
@@ -550,9 +550,9 @@ void UReplicationBridge::RemoveDestructionInfosForGroup(UE::Net::FNetObjectGroup
 
 	FWorldLocations& WorldLocations = GetReplicationSystem()->GetReplicationSystemInternal()->GetWorldLocations();
 
-	if (GroupHandle)
+	if (GroupHandle.IsValid())
 	{
-		UE_LOG_REPLICATIONBRIDGE(Verbose, TEXT("RemoveDestructionInfosForGroup GroupIndex: %u"), GroupHandle);
+		UE_LOG_REPLICATIONBRIDGE(Verbose, TEXT("RemoveDestructionInfosForGroup GroupIndex: %u"), GroupHandle.GetGroupIndex());
 
 		FNetObjectGroup* Group = Groups->GetGroup(GroupHandle);
 		check(Group);
@@ -924,7 +924,7 @@ UE::Net::FNetObjectGroupHandle UReplicationBridge::CreateLevelGroup(const UObjec
 	using namespace UE::Net;
 
 	FNetObjectGroupHandle LevelGroupHandle = ReplicationSystem->CreateGroup();
-	if (ensure(LevelGroupHandle != InvalidNetObjectGroupHandle))
+	if (ensure(LevelGroupHandle.IsValid()))
 	{
 		ReplicationSystem->AddGroupFilter(LevelGroupHandle);
 		LevelGroups.Emplace(FObjectKey(Level), LevelGroupHandle);
@@ -936,7 +936,7 @@ UE::Net::FNetObjectGroupHandle UReplicationBridge::CreateLevelGroup(const UObjec
 UE::Net::FNetObjectGroupHandle UReplicationBridge::GetLevelGroup(const UObject* Level) const
 {
 	const UE::Net::FNetObjectGroupHandle* LevelGroupHandle = LevelGroups.Find(FObjectKey(Level));
-	return (LevelGroupHandle != nullptr ? *LevelGroupHandle : UE::Net::InvalidNetObjectGroupHandle);
+	return (LevelGroupHandle != nullptr ? *LevelGroupHandle : UE::Net::FNetObjectGroupHandle());
 }
 
 void UReplicationBridge::DestroyGlobalNetHandle(UE::Net::Private::FInternalNetRefIndex InternalReplicationIndex)
diff --git a/Engine/Source/Runtime/Experimental/Iris/Core/Private/Iris/ReplicationSystem/ReplicationSystem.cpp b/Engine/Source/Runtime/Experimental/Iris/Core/Private/Iris/ReplicationSystem/ReplicationSystem.cpp
--- a/Engine/Source/Runtime/Experimental/Iris/Core/Private/Iris/ReplicationSystem/ReplicationSystem.cpp
+++ b/Engine/Source/Runtime/Experimental/Iris/Core/Private/Iris/ReplicationSystem/ReplicationSystem.cpp
@@ -43,6 +43,9 @@ public:
 	UReplicationSystem* ReplicationSystem;
 	FReplicationSystemInternal ReplicationSystemInternal;
 	uint64 IrisDebugHelperDummy = 0U;
+	FNetObjectGroupHandle NotReplicatedNetObjectGroupHandle;
+	FNetObjectGroupHandle NetGroupOwnerNetObjectGroupHandle;
+	FNetObjectGroupHandle NetGroupReplayNetObjectGroupHandle;
 
 	explicit FReplicationSystemImpl(UReplicationSystem* InReplicationSystem, const UReplicationSystem::FReplicationSystemParams& Params)
 	: ReplicationSystem(InReplicationSystem)
@@ -56,16 +59,16 @@ public:
 
 	void InitDefaultFilteringGroups()
 	{
-		FNetObjectGroupHandle LocalNotReplicatedGroupHande = ReplicationSystem->CreateGroup();
-		check(LocalNotReplicatedGroupHande == NotReplicatedNetObjectGroupHandle);
+		NotReplicatedNetObjectGroupHandle = ReplicationSystem->CreateGroup();
+		check(NotReplicatedNetObjectGroupHandle.IsNotReplicatedNetObjectGroup());
 		ReplicationSystem->AddGroupFilter(NotReplicatedNetObjectGroupHandle);
 		
 		// Setup SubObjectFiltering groups
-		FNetObjectGroupHandle LocalNetGroupOwnerNetObjectGroupHandle = ReplicationSystem->GetOrCreateSubObjectFilter(UE::Net::NetGroupOwner);
-		check(LocalNetGroupOwnerNetObjectGroupHandle == NetGroupOwnerNetObjectGroupHandle);
+		NetGroupOwnerNetObjectGroupHandle = ReplicationSystem->GetOrCreateSubObjectFilter(UE::Net::NetGroupOwner);
+		check(NetGroupOwnerNetObjectGroupHandle.IsNetGroupOwnerNetObjectGroup());
 
-		FNetObjectGroupHandle LocalNetGroupReplayNetObjectGroupHandle = ReplicationSystem->GetOrCreateSubObjectFilter(UE::Net::NetGroupReplay);
-		check(LocalNetGroupReplayNetObjectGroupHandle == NetGroupReplayNetObjectGroupHandle);
+		NetGroupReplayNetObjectGroupHandle = ReplicationSystem->GetOrCreateSubObjectFilter(UE::Net::NetGroupReplay);
+		check(NetGroupReplayNetObjectGroupHandle.IsNetGroupReplayNetObjectGroup());
 	}
 
 	void Init(const UReplicationSystem::FReplicationSystemParams& Params)
@@ -699,7 +702,7 @@ void UReplicationSystem::CollectGarbage()
 
 void UReplicationSystem::ResetGameWorldState()
 {
-	Impl->ReplicationSystemInternal.GetReplicationBridge()->RemoveDestructionInfosForGroup(UE::Net::InvalidNetObjectGroupHandle);
+	Impl->ReplicationSystemInternal.GetReplicationBridge()->RemoveDestructionInfosForGroup(UE::Net::FNetObjectGroupHandle());
 }
 
 void UReplicationSystem::NotifyStreamingLevelUnload(const UObject* Level)
@@ -954,14 +957,14 @@ UE::Net::FNetObjectGroupHandle UReplicationSystem::GetOrCreateSubObjectFilter(FN
 	FReplicationFiltering& Filtering = Impl->ReplicationSystemInternal.GetFiltering();
 
 	FNetObjectGroupHandle GroupHandle = Groups.GetNamedGroupHandle(GroupName);
-	if (GroupHandle != InvalidNetObjectGroupHandle)
+	if (GroupHandle.IsValid())
 	{
 		check(Filtering.IsSubObjectFilterGroup(GroupHandle));
 		return GroupHandle;
 	}
 
 	GroupHandle = Groups.CreateNamedGroup(GroupName);
-	if (GroupHandle != InvalidNetObjectGroupHandle)
+	if (GroupHandle.IsValid())
 	{
 		Filtering.AddSubObjectFilter(GroupHandle);
 	}
@@ -977,14 +980,14 @@ UE::Net::FNetObjectGroupHandle UReplicationSystem::GetSubObjectFilterGroupHandle
 	FReplicationFiltering& Filtering = Impl->ReplicationSystemInternal.GetFiltering();
 
 	FNetObjectGroupHandle GroupHandle = Groups.GetNamedGroupHandle(GroupName);
-	if (GroupHandle != InvalidNetObjectGroupHandle)
+	if (GroupHandle.IsValid())
 	{
 		if (ensureAlwaysMsgf(Filtering.IsSubObjectFilterGroup(GroupHandle), TEXT("UReplicationSystem::GetSubObjectFilterGroupHandle Trying to lookup NetObjectGroupHandle for NetGroup %s that is not a subobject filter"), *GroupName.ToString()))
 		{
 			return GroupHandle;
 		}
 	}
-	return InvalidNetObjectGroupHandle;
+	return FNetObjectGroupHandle();
 }
 
 void UReplicationSystem::SetSubObjectFilterStatus(FName GroupName, uint32 ConnectionId, UE::Net::ENetFilterStatus ReplicationStatus)
@@ -999,7 +1002,7 @@ void UReplicationSystem::SetSubObjectFilterStatus(FName GroupName, uint32 Connec
 	}
 
 	FNetObjectGroupHandle GroupHandle = GetSubObjectFilterGroupHandle(GroupName);
-	if (GroupHandle != InvalidNetObjectGroupHandle)
+	if (GroupHandle.IsValid())
 	{
 		FReplicationFiltering& Filtering = Impl->ReplicationSystemInternal.GetFiltering();
 		Filtering.SetSubObjectFilterStatus(GroupHandle, ConnectionId, ReplicationStatus);
@@ -1015,7 +1018,7 @@ void UReplicationSystem::RemoveSubObjectFilter(FName GroupName)
 	FReplicationFiltering& Filtering = Impl->ReplicationSystemInternal.GetFiltering();
 
 	FNetObjectGroupHandle GroupHandle = GetSubObjectFilterGroupHandle(GroupName);
-	if (GroupHandle != InvalidNetObjectGroupHandle)
+	if (GroupHandle.IsValid())
 	{
 		Filtering.RemoveSubObjectFilter(GroupHandle);
 		Groups.DestroyGroup(GroupHandle);
@@ -1031,6 +1034,12 @@ UE::Net::FNetObjectGroupHandle UReplicationSystem::CreateGroup()
 
 void UReplicationSystem::AddToGroup(FNetObjectGroupHandle GroupHandle, FNetRefHandle Handle)
 {
+	// Early out if this is invalid group
+	if (!ensure(IsValidGroup(GroupHandle)))
+	{
+		return;
+	}
+
 	LLM_SCOPE_BYTAG(Iris);
 
 	using namespace UE::Net::Private;
@@ -1041,7 +1050,7 @@ void UReplicationSystem::AddToGroup(FNetObjectGroupHandle GroupHandle, FNetRefHa
 
 	const FInternalNetRefIndex ObjectInternalIndex = NetRefHandleManager.GetInternalIndex(Handle);
 
-	if (GroupHandle && ObjectInternalIndex)
+	if (ObjectInternalIndex)
 	{
 		Groups.AddToGroup(GroupHandle, ObjectInternalIndex);
 		Filtering.NotifyObjectAddedToGroup(GroupHandle, ObjectInternalIndex);
@@ -1049,13 +1058,19 @@ void UReplicationSystem::AddToGroup(FNetObjectGroupHandle GroupHandle, FNetRefHa
 }
 
 void UReplicationSystem::RemoveFromGroup(FNetObjectGroupHandle GroupHandle, FNetRefHandle Handle)
-{	
+{
+	// Early out if this is invalid group
+	if (!ensure(IsValidGroup(GroupHandle)))
+	{
+		return;
+	}
+
 	using namespace UE::Net::Private;
 
 	FNetRefHandleManager& NetRefHandleManager = Impl->ReplicationSystemInternal.GetNetRefHandleManager();
 
 	const FInternalNetRefIndex ObjectInternalIndex = NetRefHandleManager.GetInternalIndex(Handle);
-	if (GroupHandle && ObjectInternalIndex)
+	if (ObjectInternalIndex)
 	{
 		FNetObjectGroups& Groups = Impl->ReplicationSystemInternal.GetGroups();
 		FReplicationFiltering& Filtering = Impl->ReplicationSystemInternal.GetFiltering();	
@@ -1089,6 +1104,12 @@ void UReplicationSystem::RemoveFromAllGroups(FNetRefHandle Handle)
 
 bool UReplicationSystem::IsInGroup(FNetObjectGroupHandle GroupHandle, FNetRefHandle Handle) const
 {
+	// Early out if this is invalid group
+	if (!IsValidGroup(GroupHandle))
+	{
+		return false;
+	}
+
 	using namespace UE::Net::Private;
 
 	const FNetObjectGroups& Groups = Impl->ReplicationSystemInternal.GetGroups();
@@ -1103,21 +1124,20 @@ bool UReplicationSystem::IsValidGroup(FNetObjectGroupHandle GroupHandle) const
 {
 	const UE::Net::Private::FNetObjectGroups& Groups = Impl->ReplicationSystemInternal.GetGroups();
 
-	return GroupHandle && Groups.IsValidGroup(GroupHandle);
+	return GroupHandle.IsValid() && Groups.IsValidGroup(GroupHandle);
 }
 
 void UReplicationSystem::DestroyGroup(FNetObjectGroupHandle GroupHandle)
 {
-	using namespace UE::Net;
-	using namespace UE::Net::Private;
-
-	// We do not allow client code to remove reserved groups
-	if (IsReservedNetObjectGroupHandle(GroupHandle))
+	// Early out if this is invalid or reserved group
+	if (!ensure(IsValidGroup(GroupHandle) || GroupHandle.IsReservedNetObjectGroup()))
 	{
-		check(false);
 		return;
 	}
 
+	using namespace UE::Net;
+	using namespace UE::Net::Private;
+
 	FReplicationFiltering& Filtering = Impl->ReplicationSystemInternal.GetFiltering();	
 	FNetObjectGroups& Groups = Impl->ReplicationSystemInternal.GetGroups();
 
@@ -1127,32 +1147,77 @@ void UReplicationSystem::DestroyGroup(FNetObjectGroupHandle GroupHandle)
 	Groups.DestroyGroup(GroupHandle);
 }
 
+UE::Net::FNetObjectGroupHandle UReplicationSystem::GetNotReplicatedNetObjectGroup() const
+{
+	return Impl->NotReplicatedNetObjectGroupHandle;
+}
+
+UE::Net::FNetObjectGroupHandle UReplicationSystem::GetNetGroupOwnerNetObjectGroup() const
+{
+	return Impl->NetGroupOwnerNetObjectGroupHandle;
+}
+
+UE::Net::FNetObjectGroupHandle UReplicationSystem::GetNetGroupReplayNetObjectGroup() const
+{
+	return Impl->NetGroupReplayNetObjectGroupHandle;
+}
+
 void UReplicationSystem::AddGroupFilter(FNetObjectGroupHandle GroupHandle)
 {
+	// Early out if this is invalid group
+	if (!ensure(IsValidGroup(GroupHandle)))
+	{
+		return;
+	}
+
 	UE::Net::Private::FReplicationFiltering& Filtering = Impl->ReplicationSystemInternal.GetFiltering();	
 	Filtering.AddGroupFilter(GroupHandle);
 }
 
 void UReplicationSystem::RemoveGroupFilter(FNetObjectGroupHandle GroupHandle)
 {
+	// Early out if this is invalid group
+	if (!ensure(IsValidGroup(GroupHandle)))
+	{
+		return;
+	}
+
 	UE::Net::Private::FReplicationFiltering& Filtering = Impl->ReplicationSystemInternal.GetFiltering();	
 	Filtering.RemoveGroupFilter(GroupHandle);
 }
 
 void UReplicationSystem::SetGroupFilterStatus(FNetObjectGroupHandle GroupHandle, uint32 ConnectionId, UE::Net::ENetFilterStatus ReplicationStatus)
 {
+	// Early out if this is invalid group
+	if (!ensure(IsValidGroup(GroupHandle)))
+	{
+		return;
+	}
+
 	UE::Net::Private::FReplicationFiltering& Filtering = Impl->ReplicationSystemInternal.GetFiltering();	
 	Filtering.SetGroupFilterStatus(GroupHandle, ConnectionId, ReplicationStatus);
 }
 
 void UReplicationSystem::SetGroupFilterStatus(FNetObjectGroupHandle GroupHandle, const UE::Net::FNetBitArray& Connections, UE::Net::ENetFilterStatus ReplicationStatus)
 {
+	// Early out if this is invalid group
+	if (!ensure(IsValidGroup(GroupHandle)))
+	{
+		return;
+	}
+
 	UE::Net::Private::FReplicationFiltering& Filtering = Impl->ReplicationSystemInternal.GetFiltering();	
 	Filtering.SetGroupFilterStatus(GroupHandle, UE::Net::MakeNetBitArrayView(Connections), ReplicationStatus);
 }
 
 void UReplicationSystem::SetGroupFilterStatus(FNetObjectGroupHandle GroupHandle, UE::Net::ENetFilterStatus ReplicationStatus)
 {
+	// Early out if this is invalid group
+	if (!ensure(IsValidGroup(GroupHandle)))
+	{
+		return;
+	}
+
 	UE::Net::Private::FReplicationFiltering& Filtering = Impl->ReplicationSystemInternal.GetFiltering();	
 	Filtering.SetGroupFilterStatus(GroupHandle, ReplicationStatus);
 }
diff --git a/Engine/Source/Runtime/Experimental/Iris/Core/Public/Iris/ReplicationSystem/Filtering/NetObjectFilter.h b/Engine/Source/Runtime/Experimental/Iris/Core/Public/Iris/ReplicationSystem/Filtering/NetObjectFilter.h
--- a/Engine/Source/Runtime/Experimental/Iris/Core/Public/Iris/ReplicationSystem/Filtering/NetObjectFilter.h
+++ b/Engine/Source/Runtime/Experimental/Iris/Core/Public/Iris/ReplicationSystem/Filtering/NetObjectFilter.h
@@ -12,7 +12,6 @@ class UReplicationSystem;
 namespace UE::Net
 {
 	typedef uint32 FNetObjectFilterHandle;
-	typedef uint16 FNetObjectGroupHandle;
 	struct FReplicationInstanceProtocol;
 	struct FReplicationProtocol;
 }
@@ -25,21 +24,6 @@ constexpr FNetObjectFilterHandle ToOwnerFilterHandle = FNetObjectFilterHandle(1)
 /** ConnectionFilterHandle is for internal use only. */
 constexpr FNetObjectFilterHandle ConnectionFilterHandle = FNetObjectFilterHandle(2);
 
-/** Invalid group handle */
-constexpr FNetObjectGroupHandle InvalidNetObjectGroupHandle = FNetObjectGroupHandle(0);
-
-/** Special group, NetHandles assigned to this group will be filtered out for all connections */
-constexpr FNetObjectGroupHandle NotReplicatedNetObjectGroupHandle = FNetObjectGroupHandle(1);
-
-/** Special group, SubObjects assigned to this group will replicate to owner of RootParent */
-constexpr FNetObjectGroupHandle NetGroupOwnerNetObjectGroupHandle = FNetObjectGroupHandle(2);
-
-/** Special group, NetHandles assigned to this group will Replicate if replay netconditions is met  */
-constexpr FNetObjectGroupHandle NetGroupReplayNetObjectGroupHandle = FNetObjectGroupHandle(3);
-
-/** Returns true of the provided GroupHandle is a reserved NetObjectGroupHandle */
-static constexpr bool IsReservedNetObjectGroupHandle(FNetObjectGroupHandle GroupHandle) { return GroupHandle >= NotReplicatedNetObjectGroupHandle && GroupHandle <= NetGroupReplayNetObjectGroupHandle; }
-
 /** Used to control whether an object is allowed to be replicated or not. */
 enum class ENetFilterStatus : uint32
 {
diff --git a/Engine/Source/Runtime/Experimental/Iris/Core/Public/Iris/ReplicationSystem/NetObjectGroupHandle.h b/Engine/Source/Runtime/Experimental/Iris/Core/Public/Iris/ReplicationSystem/NetObjectGroupHandle.h
new file mode 100644
--- /dev/null
+++ b/Engine/Source/Runtime/Experimental/Iris/Core/Public/Iris/ReplicationSystem/NetObjectGroupHandle.h
@@ -0,0 +1,94 @@
+// Copyright Epic Games, Inc. All Rights Reserved.
+
+#pragma once
+
+#include "CoreTypes.h"
+#include "Iris/IrisConfig.h"
+#include "Containers/StringFwd.h"
+#include "Templates/TypeHash.h"
+
+// Forward declarations
+class FString;
+
+namespace UE::Net::Private
+{
+	class FNetObjectGroups;
+}
+
+namespace UE::Net
+{
+
+class FNetObjectGroupHandle
+{
+public:
+	typedef uint16 FGroupIndexType;
+
+	/** Reserved group indices */
+	enum : FGroupIndexType
+	{
+		InvalidNetObjectGroupIndex = 0,
+		NotReplicatedNetObjectGroupIndex,
+		NetGroupOwnerNetObjectGroupIndex,
+		NetGroupReplayNetObjectGroupIndex
+	};
+
+	FNetObjectGroupHandle() : Value(0u) {}
+
+	/** Returns true if the handle is valid, note this does not mean that the group is valid */
+	inline bool IsValid() const { return Value != 0u; }
+
+	/** Returns the GroupIndex of the group assosciated with the handle */
+	inline FGroupIndexType GetGroupIndex() const { return FGroupIndexType(Index); }
+
+	/** Returns true if the provided GroupIndex is a reserved NetObjectsGroupIndex */
+	static bool IsReservedNetObjectGroupIndex(FGroupIndexType GroupIndex) { return GroupIndex >= NotReplicatedNetObjectGroupIndex && GroupIndex <= NetGroupReplayNetObjectGroupIndex; }
+
+	/** Returns true if this is a reserved NetObjectsGroups */
+	bool IsReservedNetObjectGroup() const { return FNetObjectGroupHandle::IsReservedNetObjectGroupIndex(Index); }
+
+	/** Special group, NetHandles assigned to this group will be filtered out for all connections */
+	bool IsNotReplicatedNetObjectGroup() const { return Index == NotReplicatedNetObjectGroupIndex; }
+
+	/** Special group, SubObjects assigned to this group will replicate to owner of RootParent */
+	bool IsNetGroupOwnerNetObjectGroup() const { return Index == NetGroupOwnerNetObjectGroupIndex; }
+
+	/** Special group, SubObjects assigned to this group will replicate if replay netconditions is met  */
+	bool IsNetGroupReplayNetObjectGroup() const { return Index == NetGroupReplayNetObjectGroupIndex; }
+
+private:
+	friend UE::Net::Private::FNetObjectGroups;
+
+	FNetObjectGroupHandle(FGroupIndexType IndexIn, FGroupIndexType EpochIn)
+	{ 
+		if (IndexIn == InvalidNetObjectGroupIndex)
+		{
+			*this = FNetObjectGroupHandle();
+		}
+		else
+		{
+			Index = IndexIn;
+			Epoch = EpochIn;
+		}
+	}
+
+	union 
+	{
+		uint32 Value;
+		struct
+		{
+			FGroupIndexType Index;
+			FGroupIndexType Epoch;
+		};
+	};
+
+	friend inline bool operator==(const FNetObjectGroupHandle& Lhs, const FNetObjectGroupHandle& Rhs) { return Lhs.Value == Rhs.Value; }
+	friend inline bool operator!=(const FNetObjectGroupHandle& Lhs, const FNetObjectGroupHandle& Rhs) { return Lhs.Value != Rhs.Value; }
+};
+
+static_assert(sizeof(FNetObjectGroupHandle) == sizeof(uint32), "Bad packing");
+
+FORCEINLINE uint32 GetTypeHash(const FNetObjectGroupHandle Handle)
+{
+	return ::GetTypeHash(Handle.GetGroupIndex());
+}
+}
diff --git a/Engine/Source/Runtime/Experimental/Iris/Core/Public/Iris/ReplicationSystem/ObjectReplicationBridge.h b/Engine/Source/Runtime/Experimental/Iris/Core/Public/Iris/ReplicationSystem/ObjectReplicationBridge.h
--- a/Engine/Source/Runtime/Experimental/Iris/Core/Public/Iris/ReplicationSystem/ObjectReplicationBridge.h
+++ b/Engine/Source/Runtime/Experimental/Iris/Core/Public/Iris/ReplicationSystem/ObjectReplicationBridge.h
@@ -15,7 +15,6 @@ namespace UE::Net
 {
 	struct FNetObjectResolveContext;
 	typedef uint32 FNetObjectFilterHandle;
-	typedef uint16 FNetObjectGroupHandle;
 	typedef uint32 FNetObjectPrioritizerHandle;
 	class FNetObjectReference;
 	namespace Private
diff --git a/Engine/Source/Runtime/Experimental/Iris/Core/Public/Iris/ReplicationSystem/ReplicationBridge.h b/Engine/Source/Runtime/Experimental/Iris/Core/Public/Iris/ReplicationSystem/ReplicationBridge.h
--- a/Engine/Source/Runtime/Experimental/Iris/Core/Public/Iris/ReplicationSystem/ReplicationBridge.h
+++ b/Engine/Source/Runtime/Experimental/Iris/Core/Public/Iris/ReplicationSystem/ReplicationBridge.h
@@ -7,6 +7,7 @@
 #include "Containers/Map.h"
 #include "Iris/Core/NetObjectReference.h"
 #include "Iris/ReplicationSystem/NetRefHandle.h"
+#include "Iris/ReplicationSystem/NetObjectGroupHandle.h"
 #include "Misc/EnumClassFlags.h"
 #include "Net/Core/NetHandle/NetHandle.h"
 #include "UObject/ObjectKey.h"
@@ -21,7 +22,6 @@ class UNetDriver;
 namespace UE::Net
 {
 	struct FNetDependencyInfo;
-	typedef uint16 FNetObjectGroupHandle;
 	class FNetTokenStoreState;
 	class FReplicationFragment;
 	struct FReplicationInstanceProtocol;
diff --git a/Engine/Source/Runtime/Experimental/Iris/Core/Public/Iris/ReplicationSystem/ReplicationSystem.h b/Engine/Source/Runtime/Experimental/Iris/Core/Public/Iris/ReplicationSystem/ReplicationSystem.h
--- a/Engine/Source/Runtime/Experimental/Iris/Core/Public/Iris/ReplicationSystem/ReplicationSystem.h
+++ b/Engine/Source/Runtime/Experimental/Iris/Core/Public/Iris/ReplicationSystem/ReplicationSystem.h
@@ -8,6 +8,7 @@
 #include "Templates/RefCounting.h"
 #include "Iris/IrisConfig.h"
 #include "Iris/ReplicationSystem/NetRefHandle.h"
+#include "Iris/ReplicationSystem/NetObjectGroupHandle.h"
 #include "Net/Core/NetHandle/NetHandle.h"
 
 #include "ReplicationSystem.generated.h"
@@ -28,7 +29,6 @@ namespace UE::Net
 	class FNetObjectAttachment;
 	enum class ENetObjectDeltaCompressionStatus : unsigned;
 	typedef uint32 FNetObjectFilterHandle;
-	typedef uint16 FNetObjectGroupHandle;
 	typedef uint32 FNetObjectPrioritizerHandle;
 	class FNetObjectReference;
 	enum class EReplicationCondition : uint32;
@@ -306,6 +306,15 @@ public:
 	 */
 	IRISCORE_API bool IsValidGroup(FNetObjectGroupHandle GroupHandle) const;
 
+	/** Special group, NetHandles assigned to this group will be filtered out for all connections */
+	IRISCORE_API FNetObjectGroupHandle GetNotReplicatedNetObjectGroup() const;
+
+	/** Special group, SubObjects assigned to this group will replicate to owner of RootParent */
+	IRISCORE_API FNetObjectGroupHandle GetNetGroupOwnerNetObjectGroup() const;
+
+	/** Special group, SubObjects assigned to this group will replicate if replay netconditions is met  */
+	IRISCORE_API FNetObjectGroupHandle GetNetGroupReplayNetObjectGroup() const;
+
 	// Filtering
 
 	/**
