Iris - Fixed priority calculation.

#rnx

diff --git a/Engine/Source/Runtime/Experimental/Iris/Core/Private/Iris/ReplicationSystem/Prioritization/SphereNetObjectPrioritizer.cpp b/Engine/Source/Runtime/Experimental/Iris/Core/Private/Iris/ReplicationSystem/Prioritization/SphereNetObjectPrioritizer.cpp
--- a/Engine/Source/Runtime/Experimental/Iris/Core/Private/Iris/ReplicationSystem/Prioritization/SphereNetObjectPrioritizer.cpp
+++ b/Engine/Source/Runtime/Experimental/Iris/Core/Private/Iris/ReplicationSystem/Prioritization/SphereNetObjectPrioritizer.cpp
@@ -99,24 +99,17 @@ void USphereNetObjectPrioritizer::PrioritizeBatch(FBatchParams& BatchParams)
 }
 
 /**
-  * Priority falls off with the squared distance from the object position to the view position.
+  * Priority falls off linearly with distance from the object position to the view position.
   *
   * The equation is:
-  * OuterPriority - (OuterPriority - InnerPriority)*(Sqr((1.0 - Max(Distance(ObjPos, ViewPos), InnerRadius)/(OuterRadius - InnerRadius))
-  *
-  * This can be simplified as:
-  * DistanceToCenter = Max(Distance(ObjPos, ViewPos), InnerRadius)
-  * PrioFactor = (OuterPrio - InnerPrio)*ReciprocalSquare(OuterRadius - InnerRadius)
-  * OuterPriority - Square(OuterRadius - DistanceToCenter)*PrioFactor
-  *
-  * Note that PrioFactor only need to be calculated once for all objects.
+  * OuterPriority + (OuterPriority - InnerPriority)*(Clamp(Distance(ObjPos, ViewPos), InnerRadius, OuterRadius)/(OuterRadius - InnerRadius))
   */
 
 void USphereNetObjectPrioritizer::PrioritizeBatchForSingleView(FBatchParams& BatchParams)
 {
 	IRIS_PROFILER_SCOPE(USphereNetObjectPrioritizer_PrioritizeBatchForSingleView);
 	const FVector& ViewPosVector = BatchParams.View.Views[0].Pos;
-	const VectorRegister ViewPos = MakeVectorRegister(ViewPosVector.X, ViewPosVector.Y, ViewPosVector.Z, 0.0f);
+	const VectorRegister ViewPos = VectorLoadFloat3_W0(&ViewPosVector);
 
 	const VectorRegister* Positions = BatchParams.Positions;
 	float* Priorities = BatchParams.Priorities;
@@ -146,16 +139,15 @@ void USphereNetObjectPrioritizer::PrioritizeBatchForSingleView(FBatchParams& Bat
 		const VectorRegister ScalarDistSqr2323 = VectorSwizzle(VectorCombineHigh(ScalarDistSqr2, ScalarDistSqr3), 0, 2, 1, 3);
 		const VectorRegister ScalarDistSqr0123 = VectorCombineHigh(ScalarDistSqr0101, ScalarDistSqr2323);
 
-		const VectorRegister ScalarDist0123 = VectorReciprocalAccurate(VectorReciprocalSqrt(ScalarDistSqr0123));
-		const VectorRegister ClampedScalarDist0123 = VectorMax(ScalarDist0123, BatchParams.PriorityCalculationConstants.InnerRadius);
+		const VectorRegister ScalarDist0123 = VectorSqrt(ScalarDistSqr0123);
+		const VectorRegister ClampedScalarDist0123 = VectorMax(VectorSubtract(ScalarDist0123, BatchParams.PriorityCalculationConstants.InnerRadius), VectorZeroVectorRegister());
 
 		// Calculate priority assuming the object is inside the sphere
-		const VectorRegister DistToOuter = VectorSubtract(BatchParams.PriorityCalculationConstants.OuterRadius, ClampedScalarDist0123);
-		const VectorRegister DistToOuterSqr = VectorMultiply(DistToOuter, DistToOuter);
-		VectorRegister Priorities0123 = VectorSubtract(BatchParams.PriorityCalculationConstants.OuterPriority, VectorMultiply(DistToOuterSqr, BatchParams.PriorityCalculationConstants.PriorityFactor));
+		const VectorRegister RadiusFactor = VectorMultiply(ClampedScalarDist0123, BatchParams.PriorityCalculationConstants.InvRadiusDiff);
+		VectorRegister Priorities0123 = VectorMultiplyAdd(RadiusFactor, BatchParams.PriorityCalculationConstants.PriorityDiff, BatchParams.PriorityCalculationConstants.InnerPriority);
 
 		// If object is outside the sphere we use the OutsidePriority
-		const VectorRegister OutsideSphereMask = VectorCompareGT(ClampedScalarDist0123, BatchParams.PriorityCalculationConstants.OuterRadius);
+		const VectorRegister OutsideSphereMask = VectorCompareGT(ClampedScalarDist0123, BatchParams.PriorityCalculationConstants.RadiusDiff);
 		Priorities0123 = VectorSelect(OutsideSphereMask, BatchParams.PriorityCalculationConstants.OutsidePriority, Priorities0123);
 
 		// Store the max of our calculated priority and the provided priorities
@@ -168,10 +160,10 @@ void USphereNetObjectPrioritizer::PrioritizeBatchForDualView(FBatchParams& Batch
 {
 	IRIS_PROFILER_SCOPE(USphereNetObjectPrioritizer_PrioritizeBatchForDualView);
 	const FVector& ViewPos0Vector = BatchParams.View.Views[0].Pos;
-	const VectorRegister ViewPos0 = MakeVectorRegister(ViewPos0Vector.X, ViewPos0Vector.Y, ViewPos0Vector.Z, 0.0f);
+	const VectorRegister ViewPos0 = VectorLoadFloat3_W0(&ViewPos0Vector);
 
 	const FVector& ViewPos1Vector = BatchParams.View.Views[1].Pos;
-	const VectorRegister ViewPos1 = MakeVectorRegister(ViewPos1Vector.X, ViewPos1Vector.Y, ViewPos1Vector.Z, 0.0f);
+	const VectorRegister ViewPos1 = VectorLoadFloat3_W0(&ViewPos1Vector);
 
 	const VectorRegister* Positions = BatchParams.Positions;
 	float* Priorities = BatchParams.Priorities;
@@ -216,16 +208,15 @@ void USphereNetObjectPrioritizer::PrioritizeBatchForDualView(FBatchParams& Batch
 		const VectorRegister ScalarDistSqr2323 = VectorSwizzle(VectorCombineHigh(ScalarDistSqr2, ScalarDistSqr3), 0, 2, 1, 3);
 		const VectorRegister ScalarDistSqr0123 = VectorCombineHigh(ScalarDistSqr0101, ScalarDistSqr2323);
 
-		const VectorRegister ScalarDist0123 = VectorReciprocalAccurate(VectorReciprocalSqrt(ScalarDistSqr0123));
-		const VectorRegister ClampedScalarDist0123 = VectorMax(ScalarDist0123, BatchParams.PriorityCalculationConstants.InnerRadius);
+		const VectorRegister ScalarDist0123 = VectorSqrt(ScalarDistSqr0123);
+		const VectorRegister ClampedScalarDist0123 = VectorMax(VectorSubtract(ScalarDist0123, BatchParams.PriorityCalculationConstants.InnerRadius), VectorZeroVectorRegister());
 
 		// Calculate priority assuming the object is inside the sphere
-		const VectorRegister DistToOuter = VectorSubtract(BatchParams.PriorityCalculationConstants.OuterRadius, ClampedScalarDist0123);
-		const VectorRegister DistToOuterSqr = VectorMultiply(DistToOuter, DistToOuter);
-		VectorRegister Priorities0123 = VectorSubtract(BatchParams.PriorityCalculationConstants.OuterPriority, VectorMultiply(DistToOuterSqr, BatchParams.PriorityCalculationConstants.PriorityFactor));
+		const VectorRegister RadiusFactor = VectorMultiply(ClampedScalarDist0123, BatchParams.PriorityCalculationConstants.InvRadiusDiff);
+		VectorRegister Priorities0123 = VectorMultiplyAdd(RadiusFactor, BatchParams.PriorityCalculationConstants.PriorityDiff, BatchParams.PriorityCalculationConstants.InnerPriority);
 
 		// If object is outside the sphere we use the OutsidePriority
-		const VectorRegister OutsideSphereMask = VectorCompareGT(ClampedScalarDist0123, BatchParams.PriorityCalculationConstants.OuterRadius);
+		const VectorRegister OutsideSphereMask = VectorCompareGT(ClampedScalarDist0123, BatchParams.PriorityCalculationConstants.RadiusDiff);
 		Priorities0123 = VectorSelect(OutsideSphereMask, BatchParams.PriorityCalculationConstants.OutsidePriority, Priorities0123);
 
 		// Store the max of our calculated priority and the provided priorities
@@ -241,7 +232,7 @@ void USphereNetObjectPrioritizer::PrioritizeBatchForMultiView(FBatchParams& Batc
 	for (const UE::Net::FReplicationView::FView& View : BatchParams.View.Views)
 	{
 		const FVector& ViewPosVector = View.Pos;
-		ViewPositions.Add(MakeVectorRegister(ViewPosVector.X, ViewPosVector.Y, ViewPosVector.Z, 0.0f));
+		ViewPositions.Add(VectorLoadFloat3_W0(&ViewPosVector));
 	}
 	ensureMsgf(ViewPositions.Num() <= 8, TEXT("Performance warning: Global allocation was needed to accommodate %d views."), ViewPositions.Num());
 
@@ -282,16 +273,15 @@ void USphereNetObjectPrioritizer::PrioritizeBatchForMultiView(FBatchParams& Batc
 		const VectorRegister ScalarDistSqr2323 = VectorSwizzle(VectorCombineHigh(ScalarDistSqr2, ScalarDistSqr3), 0, 2, 1, 3);
 		const VectorRegister ScalarDistSqr0123 = VectorCombineHigh(ScalarDistSqr0101, ScalarDistSqr2323);
 
-		const VectorRegister ScalarDist0123 = VectorReciprocalAccurate(VectorReciprocalSqrt(ScalarDistSqr0123));
-		const VectorRegister ClampedScalarDist0123 = VectorMax(ScalarDist0123, BatchParams.PriorityCalculationConstants.InnerRadius);
+		const VectorRegister ScalarDist0123 = VectorSqrt(ScalarDistSqr0123);
+		const VectorRegister ClampedScalarDist0123 = VectorMax(VectorSubtract(ScalarDist0123, BatchParams.PriorityCalculationConstants.InnerRadius), VectorZeroVectorRegister());
 
 		// Calculate priority assuming the object is inside the sphere
-		const VectorRegister DistToOuter = VectorSubtract(BatchParams.PriorityCalculationConstants.OuterRadius, ClampedScalarDist0123);
-		const VectorRegister DistToOuterSqr = VectorMultiply(DistToOuter, DistToOuter);
-		VectorRegister Priorities0123 = VectorSubtract(BatchParams.PriorityCalculationConstants.OuterPriority, VectorMultiply(DistToOuterSqr, BatchParams.PriorityCalculationConstants.PriorityFactor));
+		const VectorRegister RadiusFactor = VectorMultiply(ClampedScalarDist0123, BatchParams.PriorityCalculationConstants.InvRadiusDiff);
+		VectorRegister Priorities0123 = VectorMultiplyAdd(RadiusFactor, BatchParams.PriorityCalculationConstants.PriorityDiff, BatchParams.PriorityCalculationConstants.InnerPriority);
 
 		// If object is outside the sphere we use the OutsidePriority
-		const VectorRegister OutsideSphereMask = VectorCompareGT(ClampedScalarDist0123, BatchParams.PriorityCalculationConstants.OuterRadius);
+		const VectorRegister OutsideSphereMask = VectorCompareGT(ClampedScalarDist0123, BatchParams.PriorityCalculationConstants.RadiusDiff);
 		Priorities0123 = VectorSelect(OutsideSphereMask, BatchParams.PriorityCalculationConstants.OutsidePriority, Priorities0123);
 
 		// Store the max of our calculated priority and the provided priorities
@@ -338,11 +328,14 @@ void USphereNetObjectPrioritizer::SetupCalculationConstants(FPriorityCalculation
 	const VectorRegister OutsidePriority = VectorSetFloat1(Config->OutsidePriority);
 
 	const VectorRegister RadiusDiff = VectorSubtract(OuterRadius, InnerRadius);
-	const VectorRegister PriorityFactor = VectorDivide(VectorSubtract(OuterPriority, InnerPriority), VectorMultiply(RadiusDiff, RadiusDiff));
+	const VectorRegister PriorityDiff = VectorSubtract(OuterPriority, InnerPriority);
 
 	OutConstants.InnerRadius = InnerRadius;
 	OutConstants.OuterRadius = OuterRadius;
+	OutConstants.RadiusDiff = RadiusDiff;
+	OutConstants.InvRadiusDiff = VectorReciprocalAccurate(RadiusDiff);
+	OutConstants.InnerPriority = InnerPriority;
 	OutConstants.OuterPriority = OuterPriority;
 	OutConstants.OutsidePriority = OutsidePriority;
-	OutConstants.PriorityFactor = PriorityFactor;
+	OutConstants.PriorityDiff = PriorityDiff;
 }
diff --git a/Engine/Source/Runtime/Experimental/Iris/Core/Public/Iris/ReplicationSystem/Prioritization/SphereNetObjectPrioritizer.h b/Engine/Source/Runtime/Experimental/Iris/Core/Public/Iris/ReplicationSystem/Prioritization/SphereNetObjectPrioritizer.h
--- a/Engine/Source/Runtime/Experimental/Iris/Core/Public/Iris/ReplicationSystem/Prioritization/SphereNetObjectPrioritizer.h
+++ b/Engine/Source/Runtime/Experimental/Iris/Core/Public/Iris/ReplicationSystem/Prioritization/SphereNetObjectPrioritizer.h
@@ -50,9 +50,14 @@ protected:
 	{
 		VectorRegister InnerRadius;
 		VectorRegister OuterRadius;
+		// OuterRadius - InnerRadius
+		VectorRegister RadiusDiff;
+		VectorRegister InvRadiusDiff;
+		VectorRegister InnerPriority;
 		VectorRegister OuterPriority;
 		VectorRegister OutsidePriority;
-		VectorRegister PriorityFactor;
+		// OuterPriority - InnerPriority 
+		VectorRegister PriorityDiff;
 	};
 
 	struct FBatchParams
