Ias: Strip CacheKey from LogIoResult lines

#rnx

diff --git a/Engine/Source/Runtime/Experimental/IoStoreOnDemand/Private/OnDemandIoDispatcherBackend.cpp b/Engine/Source/Runtime/Experimental/IoStoreOnDemand/Private/OnDemandIoDispatcherBackend.cpp
--- a/Engine/Source/Runtime/Experimental/IoStoreOnDemand/Private/OnDemandIoDispatcherBackend.cpp
+++ b/Engine/Source/Runtime/Experimental/IoStoreOnDemand/Private/OnDemandIoDispatcherBackend.cpp
@@ -942,7 +942,6 @@ struct FChunkRequest
 static void LogIoResult(
 	const FIoChunkId& ChunkId,
 	const FIoHash& UrlHash,
-	const FIoHash& CacheKey,
 	uint32 Duration,
 	uint32 UncompressedSize,
 	uint32 UncompressedOffset,
@@ -958,13 +957,12 @@ static void LogIoResult(
 		}
 		return bCached ? TEXT("io-cache") : TEXT("io-http ");
 	}();
-	UE_LOG(LogIas, VeryVerbose, TEXT("%s: %5ums %5uKiB [%7u] %s:%s:%s|%u (%d)"),
+	UE_LOG(LogIas, VeryVerbose, TEXT("%s: %5ums %5uKiB [%7u] %s:%s|%u (%d)"),
 		Prefix,
 		Duration,
 		UncompressedSize >> 10,
 		UncompressedOffset,
 		*LexToString(ChunkId),
-		*LexToString(CacheKey),
 		*LexToString(UrlHash),
 		CompressedOffset,
 		Priority);
@@ -1246,7 +1244,7 @@ void FOnDemandIoBackend::CompleteRequest(FChunkRequest* ChunkRequest)
 		if (bDecoded)
 		{
 			Stats.OnIoRequestComplete(Request->GetBuffer().GetSize());
-			LogIoResult(Request->ChunkId, ChunkRequest->Params.GetUrlHash(), ChunkRequest->Params.ChunkKey, Duration,
+			LogIoResult(Request->ChunkId, ChunkRequest->Params.GetUrlHash(), Duration,
 				Request->GetBuffer().DataSize(), Request->Options.GetOffset(),
 				ChunkRequest->Params.ChunkRange.GetOffset(), ChunkRequest->Priority, ChunkRequest->bCached);
 				
@@ -1257,7 +1255,7 @@ void FOnDemandIoBackend::CompleteRequest(FChunkRequest* ChunkRequest)
 			Request->SetFailed();
 
 			Stats.OnIoRequestFail();
-			LogIoResult(Request->ChunkId, ChunkRequest->Params.GetUrlHash(), ChunkRequest->Params.ChunkKey, Duration,
+			LogIoResult(Request->ChunkId, ChunkRequest->Params.GetUrlHash(), Duration,
 				0, Request->Options.GetOffset(),
 				ChunkRequest->Params.ChunkRange.GetOffset(), ChunkRequest->Priority, ChunkRequest->bCached);
 		}
