Add option to spatially divide ISMs in the ISMPool.
Enable ISM pooling by default so that spatially placed ISMs can be recycled easily when moving through a streaming world.

diff --git a/Engine/Source/Runtime/Experimental/GeometryCollectionEngine/Private/GeometryCollection/GeometryCollectionISMPoolComponent.cpp b/Engine/Source/Runtime/Experimental/GeometryCollectionEngine/Private/GeometryCollection/GeometryCollectionISMPoolComponent.cpp
--- a/Engine/Source/Runtime/Experimental/GeometryCollectionEngine/Private/GeometryCollection/GeometryCollectionISMPoolComponent.cpp
+++ b/Engine/Source/Runtime/Experimental/GeometryCollectionEngine/Private/GeometryCollection/GeometryCollectionISMPoolComponent.cpp
@@ -9,10 +9,10 @@
 
 #include UE_INLINE_GENERATED_CPP_BY_NAME(GeometryCollectionISMPoolComponent)
 
-// Using a FreeList forces the calling of UnregisterComponent which can be 
-// slow if called many times in a frame. Disable for now, but can maybe enable
-// if we defer and throttle UnregisterComponent calls.
-static bool GUseComponentFreeList = false;
+// Use the FreeLists to enable recycling of ISM components.
+// Might want to add a Tick function to keep the pool at a stable value over time.
+// We would always want some spare components for fast allocation, but want to clean up when numbers get too high.
+static bool GUseComponentFreeList = true;
 FAutoConsoleVariableRef CVarISMPoolUseComponentFreeList(
 	TEXT("r.ISMPool.UseComponentFreeList"),
 	GUseComponentFreeList,
@@ -57,71 +57,71 @@ void FGeometryCollectionMeshGroup::RemoveAllMeshes(FGeometryCollectionISMPool& I
 	Meshes.Empty();
 }
 
-FGeometryCollectionISM::FGeometryCollectionISM(AActor* InOwningActor, const FGeometryCollectionStaticMeshInstance& InMeshInstance)
+FGeometryCollectionISM::FGeometryCollectionISM(AActor* InOwningActor)
 {
-	MeshInstance = InMeshInstance;
-
-	check(MeshInstance.StaticMesh);
 	check(InOwningActor);
 
-	UHierarchicalInstancedStaticMeshComponent* HISMC = nullptr;
-	UInstancedStaticMeshComponent* ISMC = nullptr;
-	
 	if ((MeshInstance.Desc.Flags & FISMComponentDescription::UseHISM) != 0)
 	{
-		const FName ISMName = MakeUniqueObjectName(InOwningActor, UHierarchicalInstancedStaticMeshComponent::StaticClass(), MeshInstance.StaticMesh->GetFName());
-		ISMC = HISMC = NewObject<UHierarchicalInstancedStaticMeshComponent>(InOwningActor, ISMName, RF_Transient | RF_DuplicateTransient);
+		ISMComponent = NewObject<UHierarchicalInstancedStaticMeshComponent>(InOwningActor, NAME_None, RF_Transient | RF_DuplicateTransient);
 	}
 	else
 	{
-		const FName ISMName = MakeUniqueObjectName(InOwningActor, UInstancedStaticMeshComponent::StaticClass(), MeshInstance.StaticMesh->GetFName());
-		ISMC = NewObject<UInstancedStaticMeshComponent>(InOwningActor, ISMName, RF_Transient | RF_DuplicateTransient);
+		ISMComponent = NewObject<UInstancedStaticMeshComponent>(InOwningActor, NAME_None, RF_Transient | RF_DuplicateTransient);
 	}
 
-	if (!ensure(ISMC != nullptr))
-	{
-		return;
-	}
+	ISMComponent->SetRemoveSwap();
+	ISMComponent->SetCanEverAffectNavigation(false);
+	ISMComponent->SetCollisionEnabled(ECollisionEnabled::NoCollision);
+
+	InOwningActor->AddInstanceComponent(ISMComponent);
+	ISMComponent->RegisterComponent();
+}
 
-	ISMC->SetStaticMesh(MeshInstance.StaticMesh);
+void FGeometryCollectionISM::InitISM(const FGeometryCollectionStaticMeshInstance& InMeshInstance)
+{
+	MeshInstance = InMeshInstance;
+	check(MeshInstance.StaticMesh);
+	check(ISMComponent != nullptr);
+
+#if WITH_EDITOR
+	const FName ISMName = MakeUniqueObjectName(ISMComponent->GetOwner(), UInstancedStaticMeshComponent::StaticClass(), InMeshInstance.StaticMesh->GetFName());
+	const FString ISMNameString = ISMName.ToString();
+	ISMComponent->Rename(*ISMNameString);
+#endif
+
+	ISMComponent->SetStaticMesh(MeshInstance.StaticMesh);
+	ISMComponent->EmptyOverrideMaterials();
 	for (int32 MaterialIndex = 0; MaterialIndex < MeshInstance.MaterialsOverrides.Num(); MaterialIndex++)
 	{
-		ISMC->SetMaterial(MaterialIndex, MeshInstance.MaterialsOverrides[MaterialIndex]);
+		ISMComponent->SetMaterial(MaterialIndex, MeshInstance.MaterialsOverrides[MaterialIndex]);
 	}
+	
+	ISMComponent->NumCustomDataFloats = MeshInstance.Desc.NumCustomDataFloats;
 	for (int32 DataIndex = 0; DataIndex < MeshInstance.CustomPrimitiveData.Num(); DataIndex++)
 	{
-		ISMC->SetDefaultCustomPrimitiveDataFloat(DataIndex, MeshInstance.CustomPrimitiveData[DataIndex]);
-	}
-
-	if ((MeshInstance.Desc.Flags & FISMComponentDescription::ReverseCulling) != 0)
-	{
-		// Instead of reverse culling, we put a mirror in the component transform so that 
-		// PRIMITIVE_SCENE_DATA_FLAG_DETERMINANT_SIGN will be set for use by materials.
-		//ISMC->SetReverseCulling(true);
-		ISMC->SetRelativeScale3D(FVector(-1, 1, 1));
+		ISMComponent->SetDefaultCustomPrimitiveDataFloat(DataIndex, MeshInstance.CustomPrimitiveData[DataIndex]);
 	}
 
-	ISMC->SetRemoveSwap();
-	ISMC->NumCustomDataFloats = MeshInstance.Desc.NumCustomDataFloats;
-	ISMC->SetMobility((MeshInstance.Desc.Flags & FISMComponentDescription::StaticMobility) != 0 ? EComponentMobility::Static : EComponentMobility::Stationary);
-	ISMC->SetCullDistances(MeshInstance.Desc.StartCullDistance, MeshInstance.Desc.EndCullDistance);
-	ISMC->SetCastShadow((MeshInstance.Desc.Flags & FISMComponentDescription::AffectShadow) != 0);
-	ISMC->bAffectDynamicIndirectLighting = (MeshInstance.Desc.Flags & FISMComponentDescription::AffectDynamicIndirectLighting) != 0;
-	ISMC->bAffectDistanceFieldLighting = (MeshInstance.Desc.Flags & FISMComponentDescription::AffectDistanceFieldLighting) != 0;
-	ISMC->bWorldPositionOffsetWritesVelocity = (MeshInstance.Desc.Flags & FISMComponentDescription::WorldPositionOffsetWritesVelocity) != 0;
-	ISMC->bUseGpuLodSelection = (MeshInstance.Desc.Flags & FISMComponentDescription::GpuLodSelection) != 0;
-	ISMC->SetCanEverAffectNavigation(false);
-	ISMC->SetCollisionEnabled(ECollisionEnabled::NoCollision);
-	ISMC->bOverrideMinLOD = MeshInstance.Desc.MinLod > 0;	
-	ISMC->MinLOD = MeshInstance.Desc.MinLod;
-	ISMC->SetLODDistanceScale(MeshInstance.Desc.LodScale);
-	ISMC->ComponentTags.Append(MeshInstance.Desc.Tags);
-	ISMC->SetMeshDrawCommandStatsCategory(MeshInstance.Desc.StatsCategory);
-
-	InOwningActor->AddInstanceComponent(ISMC);
-	ISMC->RegisterComponent();
-	ISMC->SetVisibility(false);
-	ISMComponent = ISMC;
+	const bool bReverseCulling = (MeshInstance.Desc.Flags & FISMComponentDescription::ReverseCulling) != 0;
+	// Instead of reverse culling we put the mirror in the component transform so that PRIMITIVE_SCENE_DATA_FLAG_DETERMINANT_SIGN will be set for use by materials.
+	//ISMComponent->SetReverseCulling(bReverseCulling);
+	const FVector Scale = bReverseCulling ? FVector(-1, 1, 1) : FVector(1, 1, 1);
+	ISMComponent->SetRelativeTransform(FTransform(FQuat::Identity, MeshInstance.Desc.Position, Scale));
+
+	ISMComponent->SetMobility((MeshInstance.Desc.Flags & FISMComponentDescription::StaticMobility) != 0 ? EComponentMobility::Static : EComponentMobility::Stationary);
+	ISMComponent->SetCachedMaxDrawDistance(MeshInstance.Desc.EndCullDistance);
+	ISMComponent->SetCullDistances(MeshInstance.Desc.StartCullDistance, MeshInstance.Desc.EndCullDistance);
+	ISMComponent->SetCastShadow((MeshInstance.Desc.Flags & FISMComponentDescription::AffectShadow) != 0);
+	ISMComponent->bAffectDynamicIndirectLighting = (MeshInstance.Desc.Flags & FISMComponentDescription::AffectDynamicIndirectLighting) != 0;
+	ISMComponent->bAffectDistanceFieldLighting = (MeshInstance.Desc.Flags & FISMComponentDescription::AffectDistanceFieldLighting) != 0;
+	ISMComponent->bWorldPositionOffsetWritesVelocity = (MeshInstance.Desc.Flags & FISMComponentDescription::WorldPositionOffsetWritesVelocity) != 0;
+	ISMComponent->bUseGpuLodSelection = (MeshInstance.Desc.Flags & FISMComponentDescription::GpuLodSelection) != 0;
+	ISMComponent->bOverrideMinLOD = MeshInstance.Desc.MinLod > 0;
+	ISMComponent->MinLOD = MeshInstance.Desc.MinLod;
+	ISMComponent->SetLODDistanceScale(MeshInstance.Desc.LodScale);
+	ISMComponent->SetMeshDrawCommandStatsCategory(MeshInstance.Desc.StatsCategory);
+	ISMComponent->ComponentTags = MeshInstance.Desc.Tags;
 }
 
 FInstanceGroups::FInstanceGroupId FGeometryCollectionISM::AddInstanceGroup(int32 InstanceCount, TArrayView<const float> CustomDataFloats)
@@ -135,7 +135,6 @@ FInstanceGroups::FInstanceGroupId FGeometryCollectionISM::AddInstanceGroup(int32
 	TArray<FTransform> ZeroScaleTransforms;
 	ZeroScaleTransforms.Init(ZeroScaleTransform, InstanceCount);
 
-	ISMComponent->SetVisibility(true);
 	ISMComponent->PreAllocateInstancesMemory(InstanceCount);
 	TArray<int32> RenderInstances = ISMComponent->AddInstances(ZeroScaleTransforms, true, true);
 
@@ -174,19 +173,26 @@ FGeometryCollectionISMPool::FISMIndex FGeometryCollectionISMPool::AddISM(UGeomet
 		return *ISMIndexPtr;
 	}
 
+	// Take an ISM from the current FreeLists if available instead of allocating a new slot.
+	const bool bIsHISM = (MeshInstance.Desc.Flags & FISMComponentDescription::UseHISM) != 0;
+
 	FISMIndex ISMIndex = INDEX_NONE;
-	if (FreeList.Num())
+	if (bIsHISM && FreeListHISM.Num())
+	{
+		ISMIndex = FreeListHISM.Last();
+		FreeListHISM.RemoveAt(FreeListHISM.Num() - 1);
+	}
+	else if (!bIsHISM && FreeListISM.Num())
 	{
-		// Take an ISM from the current FreeList instead of allocating a new slot.
-		ISMIndex = FreeList.Last();
-		FreeList.RemoveAt(FreeList.Num() - 1);
-		ISMs[ISMIndex] = FGeometryCollectionISM(OwningComponent->GetOwner(), MeshInstance);
+		ISMIndex = FreeListISM.Last();
+		FreeListISM.RemoveAt(FreeListISM.Num() - 1);
 	}
 	else
 	{
-		ISMIndex = ISMs.Emplace(OwningComponent->GetOwner(), MeshInstance);
+		ISMIndex = ISMs.Emplace(OwningComponent->GetOwner());
 	}
 	
+	ISMs[ISMIndex].InitISM(MeshInstance);
 	MeshToISMIndex.Add(MeshInstance, ISMIndex);
 	return ISMIndex;
 }
@@ -206,8 +212,6 @@ bool FGeometryCollectionISMPool::BatchUpdateInstancesTransforms(FGeometryCollect
 
 bool FGeometryCollectionISMPool::BatchUpdateInstancesTransforms(FGeometryCollectionMeshInfo& MeshInfo, int32 StartInstanceIndex, TArrayView<const FTransform> NewInstancesTransforms, bool bWorldSpace, bool bMarkRenderStateDirty, bool bTeleport)
 {
-	constexpr bool bUseArrayView = true;
-
 	if (ISMs.IsValidIndex(MeshInfo.ISMIndex))
 	{
 		FGeometryCollectionISM& ISM = ISMs[MeshInfo.ISMIndex];
@@ -221,50 +225,24 @@ bool FGeometryCollectionISMPool::BatchUpdateInstancesTransforms(FGeometryCollect
 		int32 TransformIndex = 0;
 		int32 BatchCount = 1;
 
-		if constexpr (bUseArrayView)
+		for (int InstanceIndex = StartInstanceIndex + 1; InstanceIndex < NewInstancesTransforms.Num(); ++InstanceIndex)
 		{
-			for (int InstanceIndex = StartInstanceIndex + 1; InstanceIndex < NewInstancesTransforms.Num(); ++InstanceIndex)
+			// Flush batch for non-sequential instances.
+			int32 RenderIndex = ISM.InstanceIndexToRenderIndex[InstanceGroup.Start + InstanceIndex];
+			if (RenderIndex != (StartIndex + BatchCount))
 			{
-				// Flush batch for non-sequential instances.
-				int32 RenderIndex = ISM.InstanceIndexToRenderIndex[InstanceGroup.Start + InstanceIndex];
-				if (RenderIndex != (StartIndex + BatchCount))
-				{
-					TArrayView<const FTransform> BatchedTransformsView = MakeArrayView(NewInstancesTransforms.GetData() + TransformIndex, BatchCount);
-					ISM.ISMComponent->BatchUpdateInstancesTransforms(StartIndex, BatchedTransformsView, bWorldSpace, bMarkRenderStateDirty, bTeleport);
-					StartIndex = RenderIndex;
-					TransformIndex += BatchCount;
-					BatchCount = 0;
-				}
-				BatchCount++;
+				TArrayView<const FTransform> BatchedTransformsView = MakeArrayView(NewInstancesTransforms.GetData() + TransformIndex, BatchCount);
+				ISM.ISMComponent->BatchUpdateInstancesTransforms(StartIndex, BatchedTransformsView, bWorldSpace, bMarkRenderStateDirty, bTeleport);
+				StartIndex = RenderIndex;
+				TransformIndex += BatchCount;
+				BatchCount = 0;
 			}
-
-			// last one
-			TArrayView<const FTransform> BatchedTransformsView = MakeArrayView(NewInstancesTransforms.GetData() + TransformIndex, BatchCount);
-			return ISM.ISMComponent->BatchUpdateInstancesTransforms(StartIndex, BatchedTransformsView, bWorldSpace, bMarkRenderStateDirty, bTeleport);
+			BatchCount++;
 		}
-		else
-		{
-			TArray<FTransform> BatchTransforms; // Can't use TArrayView because blueprint function doesn't support that 
-			BatchTransforms.Reserve(NewInstancesTransforms.Num());
-			BatchTransforms.Add(NewInstancesTransforms[TransformIndex++]);
-			for (int InstanceIndex = StartInstanceIndex + 1; InstanceIndex < NewInstancesTransforms.Num(); ++InstanceIndex)
-			{
-				// Flush batch for non-sequential instances.
-				int32 RenderIndex = ISM.InstanceIndexToRenderIndex[InstanceGroup.Start + InstanceIndex];
-				if (RenderIndex != (StartIndex + BatchCount))
-				{
-					ISM.ISMComponent->BatchUpdateInstancesTransforms(StartIndex, BatchTransforms, bWorldSpace, bMarkRenderStateDirty, bTeleport);
-					StartIndex = RenderIndex;
-					BatchTransforms.SetNum(0, false);
-					BatchCount = 0;
-				}
-
-				BatchTransforms.Add(NewInstancesTransforms[TransformIndex++]);
-				BatchCount++;
-			}
 
-			return ISM.ISMComponent->BatchUpdateInstancesTransforms(StartIndex, BatchTransforms, bWorldSpace, bMarkRenderStateDirty, bTeleport);
-		}
+		// last one
+		TArrayView<const FTransform> BatchedTransformsView = MakeArrayView(NewInstancesTransforms.GetData() + TransformIndex, BatchCount);
+		return ISM.ISMComponent->BatchUpdateInstancesTransforms(StartIndex, BatchedTransformsView, bWorldSpace, bMarkRenderStateDirty, bTeleport);
 	}
 	UE_LOG(LogChaos, Warning, TEXT("UGeometryCollectionISMPoolComponent : Invalid ISM Id (%d) when updating the transform "), MeshInfo.ISMIndex);
 	return false;
@@ -312,20 +290,26 @@ void FGeometryCollectionISMPool::RemoveISM(const FGeometryCollectionMeshInfo& Me
 			ISM.InstanceGroups.Reset();
 			ISM.InstanceIndexToRenderIndex.Reset();
 			ISM.RenderIndexToInstanceIndex.Reset();
-			ISM.ISMComponent->SetVisibility(false);
 		}
 
 		if (GUseComponentFreeList && ISM.ISMComponent->PerInstanceSMData.Num() == 0)
 		{
 			// Remove component and push this ISM slot to the free list.
-			// todo: profile if it is better to push component into a free pool and recycle it.
-			ISM.ISMComponent->GetOwner()->RemoveInstanceComponent(ISM.ISMComponent);
-			ISM.ISMComponent->UnregisterComponent();
-			ISM.ISMComponent->DestroyComponent();
-			
 			MeshToISMIndex.Remove(ISM.MeshInstance);
-			FreeList.Add(MeshInfo.ISMIndex);
-			ISM.ISMComponent = nullptr;
+
+			const bool bIsHISM = (ISM.MeshInstance.Desc.Flags & FISMComponentDescription::UseHISM) != 0;
+			if (bIsHISM)
+			{
+				FreeListHISM.Add(MeshInfo.ISMIndex);
+			}
+			else
+			{
+				FreeListISM.Add(MeshInfo.ISMIndex);
+			}
+
+#if WITH_EDITOR
+			ISM.ISMComponent->Rename(nullptr);
+#endif
 		}
 	}
 }
@@ -333,7 +317,8 @@ void FGeometryCollectionISMPool::RemoveISM(const FGeometryCollectionMeshInfo& Me
 void FGeometryCollectionISMPool::Clear()
 {
 	MeshToISMIndex.Reset();
-	FreeList.Reset();
+	FreeListISM.Reset();
+	FreeListHISM.Reset();
 	if (ISMs.Num() > 0)
 	{
 		if (AActor* OwningActor = ISMs[0].ISMComponent->GetOwner())
@@ -413,7 +398,8 @@ void UGeometryCollectionISMPoolComponent::GetResourceSizeEx(FResourceSizeEx& Cum
 		MeshGroups.GetAllocatedSize()
 		+ Pool.MeshToISMIndex.GetAllocatedSize()
 		+ Pool.ISMs.GetAllocatedSize()
-		+ Pool.FreeList.GetAllocatedSize();
+		+ Pool.FreeListISM.GetAllocatedSize()
+		+ Pool.FreeListHISM.GetAllocatedSize();
 	
 	for (FGeometryCollectionISM ISM : Pool.ISMs)
 	{
diff --git a/Engine/Source/Runtime/Experimental/GeometryCollectionEngine/Public/GeometryCollection/GeometryCollectionISMPoolComponent.h b/Engine/Source/Runtime/Experimental/GeometryCollectionEngine/Public/GeometryCollection/GeometryCollectionISMPoolComponent.h
--- a/Engine/Source/Runtime/Experimental/GeometryCollectionEngine/Public/GeometryCollection/GeometryCollectionISMPoolComponent.h
+++ b/Engine/Source/Runtime/Experimental/GeometryCollectionEngine/Public/GeometryCollection/GeometryCollectionISMPoolComponent.h
@@ -113,6 +113,7 @@ struct FISMComponentDescription
 
 	uint32 Flags = WorldPositionOffsetWritesVelocity|AffectShadow;
 	int32 NumCustomDataFloats = 0;
+	FVector Position = FVector::ZeroVector;
 	int32 StartCullDistance = 0;
 	int32 EndCullDistance = 0;
 	int32 MinLod = 0;
@@ -124,7 +125,8 @@ struct FISMComponentDescription
 	{
 		return Flags == Other.Flags &&
 			NumCustomDataFloats == Other.NumCustomDataFloats &&
-			StartCullDistance == Other.StartCullDistance && 
+			Position == Other.Position &&
+			StartCullDistance == Other.StartCullDistance &&
 			EndCullDistance == Other.EndCullDistance &&
 			MinLod == Other.MinLod &&
 			LodScale == Other.LodScale &&
@@ -136,6 +138,7 @@ struct FISMComponentDescription
 FORCEINLINE uint32 GetTypeHash(const FISMComponentDescription& Desc)
 {
 	uint32 Hash = HashCombineFast(GetTypeHash(Desc.Flags), GetTypeHash(Desc.NumCustomDataFloats));
+	Hash = HashCombineFast(Hash, GetTypeHash(Desc.Position));
 	Hash = HashCombineFast(Hash, GetTypeHash(Desc.StartCullDistance));
 	Hash = HashCombineFast(Hash, GetTypeHash(Desc.EndCullDistance));
 	Hash = HashCombineFast(Hash, GetTypeHash(Desc.MinLod));
@@ -241,8 +244,10 @@ struct FGeometryCollectionMeshGroup
 /** Structure containting all info for a single ISM. */
 struct FGeometryCollectionISM
 {
-	FGeometryCollectionISM(AActor* OwmingActor, const FGeometryCollectionStaticMeshInstance& InMeshInstance);
+	FGeometryCollectionISM(AActor* InOwningActor);
 
+	/** Initialize the ISMComponent according to settings on the mesh instance. */
+	void InitISM(const FGeometryCollectionStaticMeshInstance& InMeshInstance);
 	/** Add a group to the ISM. Returns the group index. */
 	FInstanceGroups::FInstanceGroupId AddInstanceGroup(int32 InstanceCount, TArrayView<const float> CustomDataFloats);
 
@@ -277,7 +282,8 @@ struct FGeometryCollectionISMPool
 
 	TMap<FGeometryCollectionStaticMeshInstance, FISMIndex> MeshToISMIndex;
 	TArray<FGeometryCollectionISM> ISMs;
-	TArray<int32> FreeList;
+	TArray<int32> FreeListISM;
+	TArray<int32> FreeListHISM;
 };
 
 
