Split duplicated code in Chaos::DoCollide into its own function, so that individual shapes can be queried for collision

diff --git a/Engine/Source/Runtime/Experimental/Chaos/Public/Chaos/Collision/CollisionFilter.h b/Engine/Source/Runtime/Experimental/Chaos/Public/Chaos/Collision/CollisionFilter.h
--- a/Engine/Source/Runtime/Experimental/Chaos/Public/Chaos/Collision/CollisionFilter.h
+++ b/Engine/Source/Runtime/Experimental/Chaos/Public/Chaos/Collision/CollisionFilter.h
@@ -111,51 +111,43 @@ namespace Chaos
 		return (!Shape || (Shape->GetSimEnabled() && IsFilterValid(Shape->GetSimData())));
 	}
 
-	/**
-	* Sim Collision filter (i.e., not Query collision filter)
-	*/
-	inline bool DoCollide(EImplicitObjectType Implicit0Type, const FPerShapeData* Shape0, EImplicitObjectType Implicit1Type, const FPerShapeData* Shape1)
+	inline bool DoCollide(EImplicitObjectType ImplicitType, const FPerShapeData* Shape)
 	{
 		//
 		// Disabled shapes do not collide
 		//
-		if (!FilterHasSimEnabled(Shape0)) return false;
-		if (!FilterHasSimEnabled(Shape1)) return false;
+		if (!FilterHasSimEnabled(Shape)) return false;
 
 		//
 		// Triangle Mesh geometry is only used if the shape specifies UseComplexAsSimple
 		//
-		if (Shape0)
+		if (Shape)
 		{
-			if (Implicit0Type == ImplicitObjectType::TriangleMesh && Shape0->GetCollisionTraceType() != EChaosCollisionTraceFlag::Chaos_CTF_UseComplexAsSimple)
+			if (ImplicitType == ImplicitObjectType::TriangleMesh && Shape->GetCollisionTraceType() != EChaosCollisionTraceFlag::Chaos_CTF_UseComplexAsSimple)
 			{
 				return false;
 			}
-			else if (Shape0->GetCollisionTraceType() == EChaosCollisionTraceFlag::Chaos_CTF_UseComplexAsSimple && Implicit0Type != ImplicitObjectType::TriangleMesh)
+			else if (Shape->GetCollisionTraceType() == EChaosCollisionTraceFlag::Chaos_CTF_UseComplexAsSimple && ImplicitType != ImplicitObjectType::TriangleMesh)
 			{
 				return false;
 			}
 		}
-		else if (Implicit0Type == ImplicitObjectType::TriangleMesh)
+		else if (ImplicitType == ImplicitObjectType::TriangleMesh)
 		{
 			return false;
 		}
 
-		if (Shape1)
-		{
-			if (Implicit1Type == ImplicitObjectType::TriangleMesh && Shape1->GetCollisionTraceType() != EChaosCollisionTraceFlag::Chaos_CTF_UseComplexAsSimple)
-			{
-				return false;
-			}
-			else if (Shape1->GetCollisionTraceType() == EChaosCollisionTraceFlag::Chaos_CTF_UseComplexAsSimple && Implicit1Type != ImplicitObjectType::TriangleMesh)
-			{
-				return false;
-			}
-		}
-		else if (Implicit1Type == ImplicitObjectType::TriangleMesh)
-		{
-			return false;
-		}
+		return true;
+	}
+
+	/**
+	* Sim Collision filter (i.e., not Query collision filter)
+	*/
+	inline bool DoCollide(EImplicitObjectType Implicit0Type, const FPerShapeData* Shape0, EImplicitObjectType Implicit1Type, const FPerShapeData* Shape1)
+	{
+
+		if (!DoCollide(Implicit0Type, Shape0)) { return false; }
+		if (!DoCollide(Implicit1Type, Shape1)) { return false; }
 
 		//
 		// Shape Filtering
