Made sure the GC particle is marked dirty when updating its ChildToParent transform

diff --git a/Engine/Source/Runtime/Experimental/Chaos/Private/Chaos/ClusterUnionManager.cpp b/Engine/Source/Runtime/Experimental/Chaos/Private/Chaos/ClusterUnionManager.cpp
--- a/Engine/Source/Runtime/Experimental/Chaos/Private/Chaos/ClusterUnionManager.cpp
+++ b/Engine/Source/Runtime/Experimental/Chaos/Private/Chaos/ClusterUnionManager.cpp
@@ -904,6 +904,9 @@ namespace Chaos
 						ChildHandle->SetR(ParticleToWorld.GetRotation());
 						ChildHandle->SetQ(ParticleToWorld.GetRotation());
 
+						// We need to mark the child handle to be dirty so that its proxy gets a chance to sync back to the GC 
+						MEvolution.GetParticles().MarkTransientDirtyParticle(ChildHandle);
+
 						// A child to parent update needs to remove *and* add to the connectivity graph (in that order) since
 						// the child to parent update might move the node so far away as to make the old connectivity edges incorrect.
 						ClusterUnion->PendingConnectivityOperations.Add({ Particle, EClusterUnionConnectivityOperation::Remove });
