Add VisitConcreteObjects for executing a lambda on downcasted children of implicit unions.

diff --git a/Engine/Source/Runtime/Experimental/Chaos/Public/Chaos/CastingUtilities.h b/Engine/Source/Runtime/Experimental/Chaos/Public/Chaos/CastingUtilities.h
--- a/Engine/Source/Runtime/Experimental/Chaos/Public/Chaos/CastingUtilities.h
+++ b/Engine/Source/Runtime/Experimental/Chaos/Public/Chaos/CastingUtilities.h
@@ -5,6 +5,7 @@
 
 #include "Chaos/ImplicitObjectScaled.h"
 #include "Chaos/ImplicitObjectTransformed.h"
+#include "Chaos/ImplicitObjectUnion.h"
 #include "Chaos/Sphere.h"
 #include "Chaos/Capsule.h"
 #include "Chaos/Convex.h"
@@ -260,6 +261,45 @@ namespace Chaos
 			}
 		}
 
+		template <typename Lambda, bool bRootObject>
+		FORCEINLINE_DEBUGGABLE void VisitConcreteObjectsImpl(const FImplicitObject& Geom, const Lambda& Func, int32 RootObjectIndex)
+		{
+			if (const FImplicitObjectUnion* Union = Geom.AsA<FImplicitObjectUnion>())
+			{
+				for (const FImplicitObjectPtr& ObjectPtr : Union->GetObjects())
+				{
+					if (const FImplicitObject* Object = ObjectPtr.GetReference())
+					{
+						VisitConcreteObjectsImpl<Lambda, false>(*Object, Func, RootObjectIndex);
+					}
+
+					if constexpr (bRootObject)
+					{
+						++RootObjectIndex;
+					}
+				}
+			}
+			else
+			{
+				CastHelper(Geom, [&Func, RootObjectIndex](const auto& Concrete)
+				{
+					Func(Concrete, RootObjectIndex);
+				});
+			}
+		}
+
+		/**
+		 * @brief Similar to FImplicitObject::VisitLeafObjects, but provides only the root object index
+		 * (ie, the index into a particle's ShapeInstance array) and the concrete type of the leaf object.
+		 * This is useful to avoid requiring users to write lambdas inside lambdas or do manual type
+		 * checking if all the caller wants to do is perform some operation on each of the underlying geometries.
+		 */
+		template <typename Lambda>
+		FORCEINLINE_DEBUGGABLE void VisitConcreteObjects(const FImplicitObject& Geom, const Lambda& Func)
+		{
+			VisitConcreteObjectsImpl<Lambda, true>(Geom, Func, 0);
+		}
+
 	} // namespace Utilities
 
 } // namespace Chaos
