Chaos - fix duplicate particle in view bug

Cluster particles transitioning from kinematic to dynamic were not being removed from the moving kinematics list.

diff --git a/Engine/Source/Runtime/Experimental/Chaos/Public/Chaos/PBDRigidsSOAs.h b/Engine/Source/Runtime/Experimental/Chaos/Public/Chaos/PBDRigidsSOAs.h
--- a/Engine/Source/Runtime/Experimental/Chaos/Public/Chaos/PBDRigidsSOAs.h
+++ b/Engine/Source/Runtime/Experimental/Chaos/Public/Chaos/PBDRigidsSOAs.h
@@ -691,6 +691,11 @@ public:
 
 	void SetClusteredParticleSOA(FPBDRigidClusteredParticleHandle* ClusteredParticle)
 	{
+		if (ClusteredParticle->ObjectState() != EObjectStateType::Kinematic)
+		{
+			MovingKinematicsMapArray.Remove(ClusteredParticle);
+		}
+
 		if (TPBDGeometryCollectionParticleHandle<FReal, 3>* GCParticle = ClusteredParticle->CastToGeometryCollection())
 		{
 			// Geometry collection particles have their own arrays which are also included in the active view
