Add garbage collection to ISM pool.

diff --git a/Engine/Source/Runtime/Experimental/GeometryCollectionEngine/Private/GeometryCollection/GeometryCollectionISMPoolComponent.cpp b/Engine/Source/Runtime/Experimental/GeometryCollectionEngine/Private/GeometryCollection/GeometryCollectionISMPoolComponent.cpp
--- a/Engine/Source/Runtime/Experimental/GeometryCollectionEngine/Private/GeometryCollection/GeometryCollectionISMPoolComponent.cpp
+++ b/Engine/Source/Runtime/Experimental/GeometryCollectionEngine/Private/GeometryCollection/GeometryCollectionISMPoolComponent.cpp
@@ -16,8 +16,13 @@ static bool GUseComponentFreeList = true;
 FAutoConsoleVariableRef CVarISMPoolUseComponentFreeList(
 	TEXT("r.ISMPool.UseComponentFreeList"),
 	GUseComponentFreeList,
-	TEXT("Recycle ISM components in the Pool."));
+	TEXT("Recycle ISM components in the ISMPool."));
 
+static int32 GComponentFreeListTargetSize = 50;
+FAutoConsoleVariableRef CVarISMPoolComponentFreeListTargetSize(
+	TEXT("r.ISMPool.ComponentFreeListTargetSize"),
+	GComponentFreeListTargetSize,
+	TEXT("Target size for number of ISM components in the ISMPool."));
 
 FGeometryCollectionMeshGroup::FMeshId FGeometryCollectionMeshGroup::AddMesh(const FGeometryCollectionStaticMeshInstance& MeshInstance, int32 InstanceCount, const FGeometryCollectionMeshInfo& ISMInstanceInfo)
 {
@@ -57,11 +62,11 @@ void FGeometryCollectionMeshGroup::RemoveAllMeshes(FGeometryCollectionISMPool& I
 	Meshes.Empty();
 }
 
-FGeometryCollectionISM::FGeometryCollectionISM(AActor* InOwningActor)
+void FGeometryCollectionISM::CreateISM(AActor* InOwningActor, bool bInUseHISM)
 {
 	check(InOwningActor);
 
-	if ((MeshInstance.Desc.Flags & FISMComponentDescription::UseHISM) != 0)
+	if (bInUseHISM)
 	{
 		ISMComponent = NewObject<UHierarchicalInstancedStaticMeshComponent>(InOwningActor, NAME_None, RF_Transient | RF_DuplicateTransient);
 	}
@@ -187,12 +192,20 @@ FGeometryCollectionISMPool::FISMIndex FGeometryCollectionISMPool::AddISM(UGeomet
 		ISMIndex = FreeListISM.Last();
 		FreeListISM.RemoveAt(FreeListISM.Num() - 1);
 	}
+	else if (FreeList.Num())
+	{
+		ISMIndex = FreeList.Last();
+		FreeList.RemoveAt(FreeList.Num() - 1);
+		ISMs[ISMIndex].CreateISM(OwningComponent->GetOwner(), bIsHISM);
+	}
 	else
 	{
-		ISMIndex = ISMs.Emplace(OwningComponent->GetOwner());
+		ISMIndex = ISMs.AddDefaulted();
+		ISMs[ISMIndex].CreateISM(OwningComponent->GetOwner(), bIsHISM);
 	}
 	
 	ISMs[ISMIndex].InitISM(MeshInstance);
+	
 	MeshToISMIndex.Add(MeshInstance, ISMIndex);
 	return ISMIndex;
 }
@@ -317,6 +330,7 @@ void FGeometryCollectionISMPool::RemoveISM(const FGeometryCollectionMeshInfo& Me
 void FGeometryCollectionISMPool::Clear()
 {
 	MeshToISMIndex.Reset();
+	FreeList.Reset();
 	FreeListISM.Reset();
 	FreeListHISM.Reset();
 	if (ISMs.Num() > 0)
@@ -334,9 +348,50 @@ void FGeometryCollectionISMPool::Clear()
 	}
 }
 
+void FGeometryCollectionISMPool::GarbageCollect()
+{
+	// Release one component per call until we reach minimum pool size.
+	const int32 NumFreeISMSlots = FreeListISM.Num();
+	const int32 NumFreeHISMSlots = FreeListHISM.Num();
+
+	if (NumFreeISMSlots + NumFreeHISMSlots > GComponentFreeListTargetSize)
+	{
+		int32 ISMIndex = INDEX_NONE;
+		if (NumFreeHISMSlots >= NumFreeISMSlots)
+		{
+			ISMIndex = FreeListHISM.Last();
+			FreeListHISM.RemoveAt(FreeListHISM.Num() - 1);
+		}
+		else
+		{
+			ISMIndex = FreeListISM.Last();
+			FreeListISM.RemoveAt(FreeListISM.Num() - 1);
+		}
+		
+		UInstancedStaticMeshComponent* ISM = ISMs[ISMIndex].ISMComponent;
+		ISM->UnregisterComponent();
+		ISM->DestroyComponent();
+		ISM->GetOwner()->RemoveInstanceComponent(ISM);
+		ISMs[ISMIndex].ISMComponent = nullptr;
+
+		FreeList.Add(ISMIndex);
+	}
+}
+
+
 UGeometryCollectionISMPoolComponent::UGeometryCollectionISMPoolComponent(const FObjectInitializer& ObjectInitializer)
 	: NextMeshGroupId(0)
 {
+	PrimaryComponentTick.bCanEverTick = true;
+	PrimaryComponentTick.bStartWithTickEnabled = true;
+	PrimaryComponentTick.bAllowTickOnDedicatedServer = false;
+	PrimaryComponentTick.TickInterval = 0.25f;
+}
+
+void UGeometryCollectionISMPoolComponent::TickComponent(float DeltaTime, enum ELevelTick TickType, FActorComponentTickFunction* ThisTickFunction)
+{
+	Super::TickComponent(DeltaTime, TickType, ThisTickFunction);
+	Pool.GarbageCollect();
 }
 
 UGeometryCollectionISMPoolComponent::FMeshGroupId  UGeometryCollectionISMPoolComponent::CreateMeshGroup()
@@ -398,6 +453,7 @@ void UGeometryCollectionISMPoolComponent::GetResourceSizeEx(FResourceSizeEx& Cum
 		MeshGroups.GetAllocatedSize()
 		+ Pool.MeshToISMIndex.GetAllocatedSize()
 		+ Pool.ISMs.GetAllocatedSize()
+		+ Pool.FreeList.GetAllocatedSize()
 		+ Pool.FreeListISM.GetAllocatedSize()
 		+ Pool.FreeListHISM.GetAllocatedSize();
 	
diff --git a/Engine/Source/Runtime/Experimental/GeometryCollectionEngine/Public/GeometryCollection/GeometryCollectionISMPoolComponent.h b/Engine/Source/Runtime/Experimental/GeometryCollectionEngine/Public/GeometryCollection/GeometryCollectionISMPoolComponent.h
--- a/Engine/Source/Runtime/Experimental/GeometryCollectionEngine/Public/GeometryCollection/GeometryCollectionISMPoolComponent.h
+++ b/Engine/Source/Runtime/Experimental/GeometryCollectionEngine/Public/GeometryCollection/GeometryCollectionISMPoolComponent.h
@@ -244,8 +244,8 @@ struct FGeometryCollectionMeshGroup
 /** Structure containting all info for a single ISM. */
 struct FGeometryCollectionISM
 {
-	FGeometryCollectionISM(AActor* InOwningActor);
-
+	/** Create the ISMComponent according to settings on the mesh instance. */
+	void CreateISM(AActor* InOwningActor, bool bInUseHISM);
 	/** Initialize the ISMComponent according to settings on the mesh instance. */
 	void InitISM(const FGeometryCollectionStaticMeshInstance& InMeshInstance);
 	/** Add a group to the ISM. Returns the group index. */
@@ -268,21 +268,30 @@ struct FGeometryCollectionISMPool
 {
 	using FISMIndex = int32;
 
+	/** Add ISM contents. */
 	FISMIndex AddISM(UGeometryCollectionISMPoolComponent* OwningComponent, const FGeometryCollectionStaticMeshInstance& MeshInstance);
 	FGeometryCollectionMeshInfo AddISM(UGeometryCollectionISMPoolComponent* OwningComponent, const FGeometryCollectionStaticMeshInstance& MeshInstance, int32 InstanceCount, TArrayView<const float> CustomDataFloats);
-
+	/** Remove ISM contents. */
+	void RemoveISM(const FGeometryCollectionMeshInfo& MeshInfo);
+	/** Update ISM contents. */
 	bool BatchUpdateInstancesTransforms(FGeometryCollectionMeshInfo& MeshInfo, int32 StartInstanceIndex, const TArray<FTransform>& NewInstancesTransforms, bool bWorldSpace, bool bMarkRenderStateDirty, bool bTeleport);
-
 	bool BatchUpdateInstancesTransforms(FGeometryCollectionMeshInfo& MeshInfo, int32 StartInstanceIndex, TArrayView<const FTransform> NewInstancesTransforms, bool bWorldSpace, bool bMarkRenderStateDirty, bool bTeleport);
 
-	void RemoveISM(const FGeometryCollectionMeshInfo& MeshInfo);
-	
-	/** Clear all ISM components and associated data */
+	/** Clear all ISM components and associated data. */
 	void Clear();
+	/** Garbage collect free lists. */
+	void GarbageCollect();
 
-	TMap<FGeometryCollectionStaticMeshInstance, FISMIndex> MeshToISMIndex;
+	/** Array of ISM objects. */
 	TArray<FGeometryCollectionISM> ISMs;
+	/** Mapping from mesh description to ISMs array slot. */
+	TMap<FGeometryCollectionStaticMeshInstance, FISMIndex> MeshToISMIndex;
+	
+	/** Free list of indices in ISMs that are empty. */
+	TArray<int32> FreeList;
+	/** Free list of indices in ISMs that have registered ISM components. */
 	TArray<int32> FreeListISM;
+	/** Free list of indices in ISMs that have registered HISM components. */
 	TArray<int32> FreeListHISM;
 };
 
@@ -301,7 +310,8 @@ public:
 	using FMeshId = int32;
 
 	//~ Begin UActorComponent Interface
-	GEOMETRYCOLLECTIONENGINE_API virtual void GetResourceSizeEx(FResourceSizeEx& CumulativeResourceSize) override;
+	virtual void TickComponent(float DeltaTime, enum ELevelTick TickType, FActorComponentTickFunction* ThisTickFunction) override;
+	virtual void GetResourceSizeEx(FResourceSizeEx& CumulativeResourceSize) override;
 	//~ End UActorComponent Interface
 
 	/** 
@@ -311,7 +321,7 @@ public:
 	*/
 	GEOMETRYCOLLECTIONENGINE_API FMeshGroupId CreateMeshGroup();
 
-	/** destroy  a mesh group and its associated resources */
+	/** Destroy  a mesh group and its associated resources */
 	GEOMETRYCOLLECTIONENGINE_API void DestroyMeshGroup(FMeshGroupId MeshGroupId);
 
 	/** Add a static mesh for a mesh group */
