fix FClusterUnionManager::IsDirectlyConnectedToMainParticleInClusterUnion which would not consider the case where the particle itself is a main particle

diff --git a/Engine/Source/Runtime/Experimental/Chaos/Private/Chaos/ClusterUnionManager.cpp b/Engine/Source/Runtime/Experimental/Chaos/Private/Chaos/ClusterUnionManager.cpp
--- a/Engine/Source/Runtime/Experimental/Chaos/Private/Chaos/ClusterUnionManager.cpp
+++ b/Engine/Source/Runtime/Experimental/Chaos/Private/Chaos/ClusterUnionManager.cpp
@@ -921,32 +921,41 @@ namespace Chaos
 
 	bool FClusterUnionManager::IsDirectlyConnectedToMainParticleInClusterUnion(const FClusterUnion& ClusterUnion, FPBDRigidParticleHandle* Particle) const
 	{
-		bool bIsConnectedToMainParticle = false;
+		auto IsParticleMainParticle = [&ClusterUnion](const FPBDRigidClusteredParticleHandle* ClusterParticle)
+		{
+			if (!ClusterParticle)
+			{
+				return false;
+			}
+
+			if (const FClusterUnionParticleProperties* Props = ClusterUnion.ChildProperties.Find(ClusterParticle))
+			{
+				return !Props->bIsAuxiliaryParticle;
+			}
+
+			return true;
+		};
+
 		if (FPBDRigidClusteredParticleHandle* ClusterParticle = Particle->CastToClustered())
 		{
+			if (IsParticleMainParticle(ClusterParticle))
+			{
+				return true;
+			}
+
 			for (const FConnectivityEdge& Edge : ClusterParticle->ConnectivityEdges())
 			{
 				if (IsInterclusterEdge(*ClusterParticle, Edge))
 				{
-					if (const FClusterUnionParticleProperties* SiblingProps = ClusterUnion.ChildProperties.Find(Edge.Sibling))
-					{
-						if (!SiblingProps->bIsAuxiliaryParticle)
-						{
-							bIsConnectedToMainParticle = true;
-							break;
-						}
-					}
-					else
+					if (Edge.Sibling && IsParticleMainParticle(Edge.Sibling->CastToClustered()))
 					{
-						// No props = main particles.
-						bIsConnectedToMainParticle = true;
-						break;
+						return true;
 					}
 				}
 			}
 		}
 
-		return bIsConnectedToMainParticle;
+		return false;
 	}
 
 	void FClusterUnionManager::RequestDeferredClusterPropertiesUpdate(FClusterUnionIndex ClusterIndex, EUpdateClusterUnionPropertiesFlags Flags)
