IAS - load binary TOC's first and fallback to old CbObject format

#rnx

diff --git a/Engine/Source/Runtime/Experimental/IoStoreOnDemand/Private/IoStoreOnDemand.cpp b/Engine/Source/Runtime/Experimental/IoStoreOnDemand/Private/IoStoreOnDemand.cpp
--- a/Engine/Source/Runtime/Experimental/IoStoreOnDemand/Private/IoStoreOnDemand.cpp
+++ b/Engine/Source/Runtime/Experimental/IoStoreOnDemand/Private/IoStoreOnDemand.cpp
@@ -21,6 +21,7 @@
 #include "Serialization/CompactBinaryWriter.h"
 #include "Serialization/CompactBinarySerialization.h"
 #include "Serialization/LargeMemoryWriter.h"
+#include "Serialization/MemoryReader.h"
 #include "String/LexFromString.h"
 
 #if (PLATFORM_DESKTOP && (IS_PROGRAM || WITH_EDITOR))
@@ -214,8 +215,26 @@ namespace UE
 ////////////////////////////////////////////////////////////////////////////////
 FArchive& operator<<(FArchive& Ar, FOnDemandTocHeader& Header)
 {
+	if (Ar.IsLoading() && Ar.TotalSize() < sizeof(FOnDemandTocHeader))
+	{
+		Ar.SetError();
+		return Ar;
+	}
+
 	Ar << Header.Magic;
+	if (Header.Magic != FOnDemandTocHeader::ExpectedMagic)
+	{
+		Ar.SetError();
+		return Ar;
+	}
+
 	Ar << Header.Version;
+	if (static_cast<EOnDemandTocVersion>(Header.Version) == EOnDemandTocVersion::Invalid)
+	{
+		Ar.SetError();
+		return Ar;
+	}
+
 	Ar << Header.ChunkVersion;
 	Ar << Header.BlockSize;
 	Ar << Header.CompressionFormat;
@@ -402,10 +421,8 @@ bool LoadFromCompactBinary(FCbFieldView Field, FOnDemandTocContainerEntry& OutCo
 FArchive& operator<<(FArchive& Ar, FOnDemandToc& Toc)
 {
 	Ar << Toc.Header;
-	
-	if (Toc.Header.Magic != FOnDemandTocHeader::ExpectedMagic || static_cast<EOnDemandTocVersion>(Toc.Header.Version) == EOnDemandTocVersion::Invalid)
+	if (Ar.IsError())
 	{
-		Ar.SetError();
 		return Ar;
 	}
 
@@ -696,7 +713,7 @@ TIoStatusOr<FIoStoreUploadResult> UploadContainerFiles(
 		TStringBuilder<256> TocsKey;
 		TocsKey << UploadParams.BucketPrefix << "/";
 
-		UE_LOG(LogIas, Display, TEXT("Fetching TOC's '%s/%s/%s'"), *Client.GetConfig().ServiceUrl, *UploadParams.Bucket, TocsKey.ToString());
+		UE_LOG(LogIas, Display, TEXT("Fetching existing TOC's from '%s/%s/%s'"), *Client.GetConfig().ServiceUrl, *UploadParams.Bucket, TocsKey.ToString());
 		FS3ListObjectResponse Response = Client.ListObjects(FS3ListObjectsRequest
 		{
 			UploadParams.Bucket,
@@ -725,10 +742,17 @@ TIoStatusOr<FIoStoreUploadResult> UploadContainerFiles(
 			}
 
 			FOnDemandToc Toc;
-			if (UE::LoadFromCompactBinary(FCbFieldView(TocResponse.Body.GetData()), Toc) == false)
+			FMemoryReaderView Ar(TocResponse.Body.GetView());
+			Ar << Toc;
+
+			if (Ar.IsError()) 
 			{
-				UE_LOG(LogIas, Warning, TEXT("Failed to load TOC '%s/%s/%s'"), *Client.GetConfig().ServiceUrl, *UploadParams.Bucket, *TocInfo.Key);
-				continue;
+				Toc = FOnDemandToc{};
+				if (UE::LoadFromCompactBinary(FCbFieldView(TocResponse.Body.GetData()), Toc) == false)
+				{
+					UE_LOG(LogIas, Warning, TEXT("Failed to load TOC '%s/%s/%s'"), *Client.GetConfig().ServiceUrl, *UploadParams.Bucket, *TocInfo.Key);
+					continue;
+				}
 			}
 
 			for (const FOnDemandTocContainerEntry& ContainerEntry : Toc.Containers)
