Ias: Stats improvemtns.
Add http duration stats.
Issue csv totals every frame.

#rnx

diff --git a/Engine/Source/Runtime/Experimental/IoStoreOnDemand/Private/OnDemandIoDispatcherBackend.cpp b/Engine/Source/Runtime/Experimental/IoStoreOnDemand/Private/OnDemandIoDispatcherBackend.cpp
--- a/Engine/Source/Runtime/Experimental/IoStoreOnDemand/Private/OnDemandIoDispatcherBackend.cpp
+++ b/Engine/Source/Runtime/Experimental/IoStoreOnDemand/Private/OnDemandIoDispatcherBackend.cpp
@@ -69,10 +69,10 @@ namespace UE::IO::Private
 {
 
 ///////////////////////////////////////////////////////////////////////////////
-static void LogHttpResult(const TCHAR* Url, uint32 StatusCode, uint32 Duration, uint32 Size, uint32 Offset, const char* Memo="ok")
+static void LogHttpResult(const TCHAR* Url, uint32 StatusCode, uint32 DurationMs, uint32 Size, uint32 Offset, const char* Memo="ok")
 {
 	Size >>= 10;
-	UE_LOG(LogIas, VeryVerbose, TEXT("http-%3u: %5ums %5uKiB [%7u] '%S' %s"), StatusCode, Duration, Size, Offset, Memo, Url);
+	UE_LOG(LogIas, VeryVerbose, TEXT("http-%3u: %5ums %5uKiB [%7u] '%S' %s"), StatusCode, DurationMs, Size, Offset, Memo, Url);
 };
 
 using namespace UE::Tasks;
@@ -316,19 +316,21 @@ void FDistributionEndpoints::CompleteEndpointRequest(FResolveRequest& ResolveReq
 class FHttpClient
 {
 public:
+	using FGetCallback = TFunction<void(TIoStatusOr<FIoBuffer>, uint64 DurationMs)>;
+
 	FHttpClient(const FString& ServiceUrl, int32 MaxConnectionCount = 8);
 	~FHttpClient() = default;
 
 	const FString& ServiceUrl() const { return SvcsUrl; }
 	int32 MaxConnectionCount() const { return MaxConnections;}
-	void Get(FAnsiStringView Url, FIoReadCallback&& Callback);
-	void Get(FAnsiStringView Url, const FIoOffsetAndLength& Range, FIoReadCallback&& Callback);
+	void Get(FAnsiStringView Url, FGetCallback&& Callback);
+	void Get(FAnsiStringView Url, const FIoOffsetAndLength& Range, FGetCallback&& Callback);
 
 	/** @return True if the client has pending work otherwise false. */
 	bool Tick();
 
 private:
-	void Issue(FAnsiStringView Url, FIoReadCallback&& Callback, FIoOffsetAndLength Range = FIoOffsetAndLength());
+	void Issue(FAnsiStringView Url, FGetCallback&& Callback, FIoOffsetAndLength Range = FIoOffsetAndLength());
 
 	FString SvcsUrl;
 	int32 MaxConnections;
@@ -346,17 +348,17 @@ FHttpClient::FHttpClient(const FString& ServiceUrl, int32 MaxConnectionCount)
 	ConnectionPool = MakeUnique<HTTP::FConnectionPool>(Params);
 }
 
-void FHttpClient::Get(FAnsiStringView Url, const FIoOffsetAndLength& Range, FIoReadCallback&& Callback)
+void FHttpClient::Get(FAnsiStringView Url, const FIoOffsetAndLength& Range, FGetCallback&& Callback)
 {
 	Issue(Url, MoveTemp(Callback), Range);
 }
 
-void FHttpClient::Get(FAnsiStringView Url, FIoReadCallback&& Callback)
+void FHttpClient::Get(FAnsiStringView Url, FGetCallback&& Callback)
 {
 	Issue(Url, MoveTemp(Callback));
 }
 
-void FHttpClient::Issue(FAnsiStringView Url, FIoReadCallback&& Callback, FIoOffsetAndLength Range)
+void FHttpClient::Issue(FAnsiStringView Url, FGetCallback&& Callback, FIoOffsetAndLength Range)
 {
 	using namespace UE::HTTP;
 
@@ -377,26 +379,26 @@ void FHttpClient::Issue(FAnsiStringView Url, FIoReadCallback&& Callback, FIoOffs
 			}
 			else if (FTicketStatus::EId::Content == Status.GetId())
 			{
-				const uint64 Duration = (uint64)FPlatformTime::ToMilliseconds64(FPlatformTime::Cycles64() - StartTime);
+				const uint64 DurationMs = (uint64)FPlatformTime::ToMilliseconds64(FPlatformTime::Cycles64() - StartTime);
 				const FIoBuffer& Content = Status.GetContent(); 
 
-				LogHttpResult(*Url, StatusCode, Duration, Content.GetSize(), Offset);
+				LogHttpResult(*Url, StatusCode, DurationMs, Content.GetSize(), Offset);
 
 				const bool bSuccessful = StatusCode > 199 && StatusCode < 300;
 				if (bSuccessful && Content.GetSize() > 0)
 				{
-					Callback(Content);
+					Callback(Content, DurationMs);
 				}
 				else
 				{
-					Callback(FIoStatus(EIoErrorCode::NotFound, TEXTVIEW("Invalid Content")));
+					Callback(FIoStatus(EIoErrorCode::NotFound, TEXTVIEW("Invalid Content")), DurationMs);
 				}
 			}
 			else if (FTicketStatus::EId::Error == Status.GetId())
 			{
-				const uint64 Duration = (uint64)FPlatformTime::ToMilliseconds64(FPlatformTime::Cycles64() - StartTime);
-				LogHttpResult(*Url, StatusCode, Duration, 0, Offset, Status.GetErrorReason());
-				Callback(FIoStatus(EIoErrorCode::ReadError, FString(Status.GetErrorReason())));
+				const uint64 DurationMs = (uint64)FPlatformTime::ToMilliseconds64(FPlatformTime::Cycles64() - StartTime);
+				LogHttpResult(*Url, StatusCode, DurationMs, 0, Offset, Status.GetErrorReason());
+				Callback(FIoStatus(EIoErrorCode::ReadError, FString(Status.GetErrorReason())), DurationMs);
 			}
 		};
 
@@ -942,7 +944,7 @@ struct FChunkRequest
 static void LogIoResult(
 	const FIoChunkId& ChunkId,
 	const FIoHash& UrlHash,
-	uint32 Duration,
+	uint32 DurationMs,
 	uint32 UncompressedSize,
 	uint32 UncompressedOffset,
 	uint32 CompressedOffset,
@@ -959,7 +961,7 @@ static void LogIoResult(
 	}();
 	UE_LOG(LogIas, VeryVerbose, TEXT("%s: %5ums %5uKiB [%7u] %s:%s|%u (%d)"),
 		Prefix,
-		Duration,
+		DurationMs,
 		UncompressedSize >> 10,
 		UncompressedOffset,
 		*LexToString(ChunkId),
@@ -1238,13 +1240,13 @@ void FOnDemandIoBackend::CompleteRequest(FChunkRequest* ChunkRequest)
 			bDecoded = FIoChunkEncoding::Decode(DecodingParams, Chunk.GetView(), Request->GetBuffer().GetMutableView());
 		}
 		
-		const uint64 Duration = Request->GetStartTime() > 0 ?
+		const uint64 DurationMs = Request->GetStartTime() > 0 ?
 			(uint64)FPlatformTime::ToMilliseconds64(FPlatformTime::Cycles64() - Request->GetStartTime()) : 0;
 
 		if (bDecoded)
 		{
-			Stats.OnIoRequestComplete(Request->GetBuffer().GetSize());
-			LogIoResult(Request->ChunkId, ChunkRequest->Params.GetUrlHash(), Duration,
+			Stats.OnIoRequestComplete(Request->GetBuffer().GetSize(), DurationMs);
+			LogIoResult(Request->ChunkId, ChunkRequest->Params.GetUrlHash(), DurationMs,
 				Request->GetBuffer().DataSize(), Request->Options.GetOffset(),
 				ChunkRequest->Params.ChunkRange.GetOffset(), ChunkRequest->Priority, ChunkRequest->bCached);
 				
@@ -1254,8 +1256,8 @@ void FOnDemandIoBackend::CompleteRequest(FChunkRequest* ChunkRequest)
 			bCanCache = false;
 			Request->SetFailed();
 
-			Stats.OnIoRequestFail();
-			LogIoResult(Request->ChunkId, ChunkRequest->Params.GetUrlHash(), Duration,
+			Stats.OnIoRequestError();
+			LogIoResult(Request->ChunkId, ChunkRequest->Params.GetUrlHash(), DurationMs,
 				0, Request->Options.GetOffset(),
 				ChunkRequest->Params.ChunkRange.GetOffset(), ChunkRequest->Priority, ChunkRequest->bCached);
 		}
@@ -1269,7 +1271,6 @@ void FOnDemandIoBackend::CompleteRequest(FChunkRequest* ChunkRequest)
 		Cache->Put(ChunkRequest->Params.ChunkKey, Chunk);
 	}
 
-	Stats.OnChunkRequestRelease();
 	ChunkRequests.Release(ChunkRequest);
 }
 
@@ -1311,8 +1312,6 @@ bool FOnDemandIoBackend::Resolve(FIoRequestImpl* Request)
 		return true;
 	}
 
-	Stats.OnChunkRequestCreate();
-		
 	if (Cache.IsValid())
 	{
 		//TODO: Pass priority to cache
@@ -1548,7 +1547,7 @@ TIoStatusOr<FOnDemandToc> FOnDemandIoBackend::GetToc(FHttpClient& HttpClient, co
 		UE_LOG(LogIas, Log, TEXT("Fetching TOC '%s/%s' (#%d/%d)"), *HttpClient.ServiceUrl(), *TocPath, Attempt + 1, MaxAttempts);
 		
 		TIoStatusOr<FOnDemandToc> Toc;
-		HttpClient.Get(Url.ToView(), [&Toc](TIoStatusOr<FIoBuffer> Response)
+		HttpClient.Get(Url.ToView(), [&Toc](TIoStatusOr<FIoBuffer> Response, uint64 DurationMs)
 		{
 			if (Response.IsOk())
 			{
@@ -1639,7 +1638,7 @@ uint32 FOnDemandIoBackend::Run()
 
 					NumConcurrentRequests++;
 					HttpClient->Get(Url.ToView(), ChunkRequest->Params.ChunkRange,
-						[this, ChunkRequest, &NumConcurrentRequests](TIoStatusOr<FIoBuffer> Status)
+						[this, ChunkRequest, &NumConcurrentRequests](TIoStatusOr<FIoBuffer> Status, uint64 DurationMs)
 						{
 							NumConcurrentRequests--;
 
@@ -1659,13 +1658,10 @@ uint32 FOnDemandIoBackend::Run()
 							if (Status.IsOk())
 							{
 								ChunkRequest->Chunk = Status.ConsumeValueOrDie();
-								Stats.OnHttpGet(ChunkRequest->Chunk.DataSize());
+								Stats.OnHttpGet(ChunkRequest->Chunk.DataSize(), DurationMs);
 							}
 							else
 							{
-								TAnsiStringBuilder<256> Url;
-								ChunkRequest->Params.GetUrl(Url);
-								LogHttpResult(StringCast<TCHAR>(*Url).Get(), -1, 0, 0, ChunkRequest->Params.ChunkRange.GetOffset(), "HTTP FAILED");
 								Stats.OnHttpError();
 							}
 
diff --git a/Engine/Source/Runtime/Experimental/IoStoreOnDemand/Private/Statistics.cpp b/Engine/Source/Runtime/Experimental/IoStoreOnDemand/Private/Statistics.cpp
--- a/Engine/Source/Runtime/Experimental/IoStoreOnDemand/Private/Statistics.cpp
+++ b/Engine/Source/Runtime/Experimental/IoStoreOnDemand/Private/Statistics.cpp
@@ -1,6 +1,7 @@
 // Copyright Epic Games, Inc. All Rights Reserved.
 
 #include "Statistics.h"
+#include "Misc/CoreDelegates.h"
 
 #if IAS_WITH_STATISTICS
 
@@ -32,14 +33,11 @@ static int32 BytesToApproxKB(uint64 Bytes) { return int32(Bytes >> 10); }
 #endif 
 
 // iorequest stats
-FCounterInt			GIoRequestsMade(TEXT("Ias/IoRequestsMade"), TraceCounterDisplayHint_None);
-FCounterAtomicInt	GIoRequestsCompleted(TEXT("Ias/IoRequestsCompleted"), TraceCounterDisplayHint_None);
-FCounterAtomicInt	GIoRequestsCompletedSize(TEXT("Ias/IoRequestsCompletedSize"), TraceCounterDisplayHint_Memory);
-FCounterInt			GIoRequestsCancelled(TEXT("Ias/IoRequestsCancelled"), TraceCounterDisplayHint_None);
-FCounterAtomicInt	GIoRequestsFailed(TEXT("Ias/IoRequestsFailed"), TraceCounterDisplayHint_None);
-// chunkrequest stats
-FCounterInt			GReadRequestsCreated(TEXT("Ias/ReadRequestsCreated"), TraceCounterDisplayHint_None);
-FCounterAtomicInt	GReadRequestsRemoved(TEXT("Ias/ReadRequestsRemoved"), TraceCounterDisplayHint_None);
+FCounterInt			GIoRequestCount(TEXT("Ias/IoRequestCount"), TraceCounterDisplayHint_None);
+FCounterAtomicInt	GIoRequestReadCount(TEXT("Ias/IoRequestReadCount"), TraceCounterDisplayHint_None);
+FCounterAtomicInt	GIoRequestReadBytes(TEXT("Ias/IoRequestReadBytes"), TraceCounterDisplayHint_Memory);
+FCounterInt			GIoRequestCancelCount(TEXT("Ias/IoRequestCancelCount"), TraceCounterDisplayHint_None);
+FCounterAtomicInt	GIoRequestErrorCount(TEXT("Ias/IoRequestErrorCount"), TraceCounterDisplayHint_None);
 // cache stats
 FCounterAtomicInt	GCacheErrorCount(TEXT("Ias/CacheErrorCount"), TraceCounterDisplayHint_None);
 FCounterAtomicInt	GCacheGetCount(TEXT("Ias/CacheGetCount"), TraceCounterDisplayHint_None);
@@ -57,47 +55,74 @@ FCounterInt			GHttpRetryCount(TEXT("Ias/HttpRetryCount"), TraceCounterDisplayHin
 FCounterAtomicInt	GHttpPendingCount(TEXT("Ias/HttpPendingCount"), TraceCounterDisplayHint_None);
 FCounterInt			GHttpInflightCount(TEXT("Ias/HttpInflightCount"), TraceCounterDisplayHint_None);
 FCounterInt			GHttpDownloadedBytes(TEXT("Ias/HttpDownloadedBytes"), TraceCounterDisplayHint_Memory);
+FCounterInt			GHttpBandwidthMpbs(TEXT("Ias/HttpBandwidthMbps"), TraceCounterDisplayHint_None);
+FCounterInt			GHttpDurationMs(TEXT("Ias/HttpDurationMs"), TraceCounterDisplayHint_None);
+FCounterInt			GHttpDurationMsAvg(TEXT("Ias/HttpDurationMsAvg"), TraceCounterDisplayHint_None);
+FCounterInt			GHttpDurationMsMax(TEXT("Ias/HttpDurationMsMax"), TraceCounterDisplayHint_None);
+int64				GHttpDurationMsSum = 0;
 
 ////////////////////////////////////////////////////////////////////////////////
 // CSV STATS
 CSV_DEFINE_CATEGORY(Ias, true);
-// iorequest stats
-CSV_DEFINE_STAT(Ias, FrameIoRequestsMade);
-CSV_DEFINE_STAT(Ias, FrameIoRequestsCompleted);
-CSV_DEFINE_STAT(Ias, FrameIoRequestsCompletedSize);
-CSV_DEFINE_STAT(Ias, FrameIoRequestsCancelled);
-CSV_DEFINE_STAT(Ias, FrameIoRequestsFailed);
-// chunkrequest stats
-CSV_DEFINE_STAT(Ias, FrameReadRequestsCreated);
-CSV_DEFINE_STAT(Ias, FrameReadRequestsRemoved);
-// cache stats
-CSV_DEFINE_STAT(Ias, FrameCacheErrorCount);
-CSV_DEFINE_STAT(Ias, FrameCacheGetCount);
-CSV_DEFINE_STAT(Ias, FrameCachePutCount);
-CSV_DEFINE_STAT(Ias, FrameCachePutExistingCount);
-CSV_DEFINE_STAT(Ias, FrameCachePutRejectCount);
-CSV_DEFINE_STAT(Ias, FrameCacheCachedBytes);
-CSV_DEFINE_STAT(Ias, FrameCachePendingBytes);
-CSV_DEFINE_STAT(Ias, FrameCacheReadBytes);
-CSV_DEFINE_STAT(Ias, FrameCacheRejectBytes);
-// http stats
+// iorequest per frame stats
+CSV_DEFINE_STAT(Ias, FrameIoRequestCount);
+CSV_DEFINE_STAT(Ias, FrameIoRequestReadCount);
+CSV_DEFINE_STAT(Ias, FrameIoRequestReadMB);
+CSV_DEFINE_STAT(Ias, FrameIoRequestCancelCount);
+CSV_DEFINE_STAT(Ias, FrameIoRequestErrorCount);
+// cache stat totals
+CSV_DEFINE_STAT(Ias, CacheGetCount);
+CSV_DEFINE_STAT(Ias, CacheErrorCount);
+CSV_DEFINE_STAT(Ias, CachePutCount);
+CSV_DEFINE_STAT(Ias, CachePutExistingCount);
+CSV_DEFINE_STAT(Ias, CachePutRejectCount);
+CSV_DEFINE_STAT(Ias, CacheCachedMB);
+CSV_DEFINE_STAT(Ias, CacheReadMB);
+CSV_DEFINE_STAT(Ias, CacheRejectedMB);
+// http stat totals
 CSV_DEFINE_STAT(Ias, HttpGetCount);
 CSV_DEFINE_STAT(Ias, HttpRetryCount);
 CSV_DEFINE_STAT(Ias, HttpErrorCount);
 CSV_DEFINE_STAT(Ias, HttpPendingCount);
-CSV_DEFINE_STAT(Ias, HttpInflightCount);
-CSV_DEFINE_STAT(Ias, HttpDownloadKB);
+CSV_DEFINE_STAT(Ias, HttpDownloadedMB);
+CSV_DEFINE_STAT(Ias, HttpBandwidthMpbs);
+CSV_DEFINE_STAT(Ias, HttpDurationMsAvg);
+CSV_DEFINE_STAT(Ias, HttpDurationMsMax);
 
 static FOnDemandIoBackendStats* GStatistics = nullptr;
+static FDelegateHandle GStatisticsEndFrameDelegateHandle;
 
 FOnDemandIoBackendStats::FOnDemandIoBackendStats()
 {
 	check(GStatistics == nullptr);
 	GStatistics = this;
+	GStatisticsEndFrameDelegateHandle = FCoreDelegates::OnEndFrame.AddLambda([this]()
+	{
+		// cache stat totals
+		CSV_CUSTOM_STAT_DEFINED(CacheGetCount, int32(GCacheGetCount.Get()), ECsvCustomStatOp::Set);
+		CSV_CUSTOM_STAT_DEFINED(CacheErrorCount, int32(GCacheErrorCount.Get()), ECsvCustomStatOp::Set);
+		CSV_CUSTOM_STAT_DEFINED(CachePutCount, int32(GCachePutCount.Get()), ECsvCustomStatOp::Set);
+		CSV_CUSTOM_STAT_DEFINED(CachePutExistingCount, int32(GCachePutExistingCount.Get()), ECsvCustomStatOp::Set);
+		CSV_CUSTOM_STAT_DEFINED(CachePutRejectCount, int32(GCachePutRejectCount.Get()), ECsvCustomStatOp::Set);
+		CSV_CUSTOM_STAT_DEFINED(CacheCachedMB, BytesToApproxMB(GCacheCachedBytes.Get()), ECsvCustomStatOp::Set);
+		CSV_CUSTOM_STAT_DEFINED(CacheReadMB, BytesToApproxMB(GCacheReadBytes.Get()), ECsvCustomStatOp::Set);
+		CSV_CUSTOM_STAT_DEFINED(CacheRejectedMB, BytesToApproxMB(GCacheRejectBytes.Get()), ECsvCustomStatOp::Set);
+
+		// http stat totals
+		CSV_CUSTOM_STAT_DEFINED(HttpGetCount, int32(GHttpGetCount.Get()), ECsvCustomStatOp::Set);
+		CSV_CUSTOM_STAT_DEFINED(HttpRetryCount, int32(GHttpRetryCount.Get()), ECsvCustomStatOp::Set);
+		CSV_CUSTOM_STAT_DEFINED(HttpErrorCount, int32(GHttpErrorCount.Get()), ECsvCustomStatOp::Set);
+		CSV_CUSTOM_STAT_DEFINED(HttpPendingCount, int32(GHttpPendingCount.Get()), ECsvCustomStatOp::Set);
+		CSV_CUSTOM_STAT_DEFINED(HttpDownloadedMB, BytesToApproxMB(GHttpDownloadedBytes.Get()), ECsvCustomStatOp::Set);
+		CSV_CUSTOM_STAT_DEFINED(HttpBandwidthMpbs, int32(GHttpBandwidthMpbs.Get()), ECsvCustomStatOp::Set);
+		CSV_CUSTOM_STAT_DEFINED(HttpDurationMsAvg, int32(GHttpDurationMsAvg.Get()), ECsvCustomStatOp::Set);
+		CSV_CUSTOM_STAT_DEFINED(HttpDurationMsMax, int32(GHttpDurationMsMax.Get()), ECsvCustomStatOp::Set);
+	});
 }
 
 FOnDemandIoBackendStats::~FOnDemandIoBackendStats()
 {
+	FCoreDelegates::OnEndFrame.Remove(GStatisticsEndFrameDelegateHandle);
 	GStatistics = nullptr;
 }
 
@@ -108,114 +133,93 @@ FOnDemandIoBackendStats* FOnDemandIoBackendStats::Get()
 
 void FOnDemandIoBackendStats::OnIoRequestEnqueue()
 {
-	GIoRequestsMade.Add(1);
-	CSV_CUSTOM_STAT_DEFINED(FrameIoRequestsMade, int32(GIoRequestsMade.Get()), ECsvCustomStatOp::Set);
+	GIoRequestCount.Add(1);
+	CSV_CUSTOM_STAT_DEFINED(FrameIoRequestCount, int32(GIoRequestCount.Get()), ECsvCustomStatOp::Set);
 }
 
-void FOnDemandIoBackendStats::OnIoRequestComplete(uint64 RequestSize)
+void FOnDemandIoBackendStats::OnIoRequestComplete(uint64 Size, uint64 Duration)
 {
-	GIoRequestsCompleted.Add(1);
-	GIoRequestsCompletedSize.Add(RequestSize);
+	GIoRequestReadCount.Add(1);
+	GIoRequestReadBytes.Add(Size);
 
-	CSV_CUSTOM_STAT_DEFINED(FrameIoRequestsCompleted, int32(GIoRequestsCompleted.Get()), ECsvCustomStatOp::Set);
-	CSV_CUSTOM_STAT_DEFINED(FrameIoRequestsCompletedSize, BytesToApproxKB(GIoRequestsCompletedSize.Get()), ECsvCustomStatOp::Set);
+	CSV_CUSTOM_STAT_DEFINED(FrameIoRequestReadCount, int32(GIoRequestReadCount.Get()), ECsvCustomStatOp::Set);
+	CSV_CUSTOM_STAT_DEFINED(FrameIoRequestReadMB, BytesToApproxMB(GIoRequestReadBytes.Get()), ECsvCustomStatOp::Set);
 }
 
 void FOnDemandIoBackendStats::OnIoRequestCancel()
 {
-	GIoRequestsCancelled.Add(1);
-	CSV_CUSTOM_STAT_DEFINED(FrameIoRequestsCancelled, int32(GIoRequestsCancelled.Get()), ECsvCustomStatOp::Set);
+	GIoRequestCancelCount.Add(1);
+	CSV_CUSTOM_STAT_DEFINED(FrameIoRequestCancelCount, int32(GIoRequestCancelCount.Get()), ECsvCustomStatOp::Set);
 }
 
-void FOnDemandIoBackendStats::OnIoRequestFail()
+void FOnDemandIoBackendStats::OnIoRequestError()
 {
-	GIoRequestsFailed.Add(1);
-	CSV_CUSTOM_STAT_DEFINED(FrameIoRequestsFailed, int32(GIoRequestsFailed.Get()), ECsvCustomStatOp::Set);
-}
-
-void FOnDemandIoBackendStats::OnChunkRequestCreate()
-{
-	GReadRequestsCreated.Add(1);
-	CSV_CUSTOM_STAT_DEFINED(FrameReadRequestsCreated, int32(GReadRequestsCreated.Get()), ECsvCustomStatOp::Set);
-}
-
-void FOnDemandIoBackendStats::OnChunkRequestRelease()
-{
-	GReadRequestsRemoved.Add(1);
-	CSV_CUSTOM_STAT_DEFINED(FrameReadRequestsRemoved, int32(GReadRequestsRemoved.Get()), ECsvCustomStatOp::Set);
+	GIoRequestErrorCount.Add(1);
+	CSV_CUSTOM_STAT_DEFINED(FrameIoRequestErrorCount, int32(GIoRequestErrorCount.Get()), ECsvCustomStatOp::Set);
 }
 
 void FOnDemandIoBackendStats::OnCacheError()
 {
 	GCacheErrorCount.Add(1);
-	CSV_CUSTOM_STAT_DEFINED(FrameCacheErrorCount, int32(GCacheErrorCount.Get()), ECsvCustomStatOp::Set);
 }
 
 void FOnDemandIoBackendStats::OnCacheGet(uint64 DataSize)
 {
 	GCacheGetCount.Add(1);
 	GCacheReadBytes.Add(DataSize);
-
-	CSV_CUSTOM_STAT_DEFINED(FrameCacheGetCount, int32(GCacheGetCount.Get()), ECsvCustomStatOp::Set);
-	CSV_CUSTOM_STAT_DEFINED(FrameCacheReadBytes, BytesToApproxKB(GCacheReadBytes.Get()), ECsvCustomStatOp::Set);
 }
 
 void FOnDemandIoBackendStats::OnCachePut()
 {
 	GCachePutCount.Add(1);
-	CSV_CUSTOM_STAT_DEFINED(FrameCachePutCount, int32(GCachePutCount.Get()), ECsvCustomStatOp::Set);
 }
 
 void FOnDemandIoBackendStats::OnCachePutExisting(uint64 /*DataSize*/)
 {
 	GCachePutExistingCount.Add(1);
-	CSV_CUSTOM_STAT_DEFINED(FrameCachePutExistingCount, int32(GCachePutExistingCount.Get()), ECsvCustomStatOp::Set);
 }
 
 void FOnDemandIoBackendStats::OnCachePutReject(uint64 DataSize)
 {
 	GCachePutRejectCount.Add(1);
 	GCacheRejectBytes.Add(DataSize);
-
-	CSV_CUSTOM_STAT_DEFINED(FrameCachePutRejectCount, int32(GCachePutRejectCount.Get()), ECsvCustomStatOp::Set);
-	CSV_CUSTOM_STAT_DEFINED(FrameCacheRejectBytes, BytesToApproxKB(GCacheRejectBytes.Get()), ECsvCustomStatOp::Set);
 }
 
 void FOnDemandIoBackendStats::OnCachePendingBytes(uint64 TotalSize)
 {
 	GCachePendingBytes.Set(TotalSize);
-	CSV_CUSTOM_STAT_DEFINED(FrameCachePendingBytes, BytesToApproxKB(GCachePendingBytes.Get()), ECsvCustomStatOp::Set);
 }
 
 void FOnDemandIoBackendStats::OnCachePersistedBytes(uint64 TotalSize)
 {
 	GCacheCachedBytes.Set(TotalSize);
-	CSV_CUSTOM_STAT_DEFINED(FrameCacheCachedBytes, BytesToApproxKB(GCacheCachedBytes.Get()), ECsvCustomStatOp::Set);
 }
 
 void FOnDemandIoBackendStats::OnHttpEnqueue()
 {
 	GHttpPendingCount.Add(1);
-	CSV_CUSTOM_STAT_DEFINED(HttpPendingCount, int32(GHttpPendingCount.Get()), ECsvCustomStatOp::Set);
 }
 
 void FOnDemandIoBackendStats::OnHttpDequeue()
 {
 	GHttpInflightCount.Add(1);
-	CSV_CUSTOM_STAT_DEFINED(HttpInflightCount, int32(GHttpInflightCount.Get()), ECsvCustomStatOp::Set);
 }
 
-void FOnDemandIoBackendStats::OnHttpGet(uint64 InSize)
+void FOnDemandIoBackendStats::OnHttpGet(uint64 Size, uint64 DurationMs)
 {
 	GHttpPendingCount.Add(-1);
 	GHttpInflightCount.Add(-1);
 	GHttpGetCount.Add(1);
-	GHttpDownloadedBytes.Add(InSize);
+	GHttpDownloadedBytes.Add(Size);
+	GHttpDurationMsSum += DurationMs;
+	GHttpDurationMs.Set(DurationMs);
 
-	CSV_CUSTOM_STAT_DEFINED(HttpPendingCount, int32(GHttpPendingCount.Get()), ECsvCustomStatOp::Set);
-	CSV_CUSTOM_STAT_DEFINED(HttpInflightCount, int32(GHttpInflightCount.Get()), ECsvCustomStatOp::Set);
-	CSV_CUSTOM_STAT_DEFINED(HttpGetCount, int32(GHttpGetCount.Get()), ECsvCustomStatOp::Set);
-	CSV_CUSTOM_STAT_DEFINED(HttpDownloadKB, BytesToApproxKB(GHttpDownloadedBytes.Get()), ECsvCustomStatOp::Set);
+	GHttpBandwidthMpbs.Set((GHttpDownloadedBytes.Get()*8)/(GHttpDurationMsSum+1)/1000);
+	GHttpDurationMsAvg.Set(GHttpDurationMsSum/GHttpGetCount.Get());
+	if (GHttpDurationMsMax.Get() < (int64)DurationMs)
+	{
+		GHttpDurationMsMax.Set((int64)DurationMs);
+	}
 }
 
 void FOnDemandIoBackendStats::OnHttpRetry()
@@ -223,10 +227,6 @@ void FOnDemandIoBackendStats::OnHttpRetry()
 	GHttpPendingCount.Add(-1);
 	GHttpInflightCount.Add(-1);
 	GHttpRetryCount.Add(1);
-
-	CSV_CUSTOM_STAT_DEFINED(HttpPendingCount, int32(GHttpPendingCount.Get()), ECsvCustomStatOp::Set);
-	CSV_CUSTOM_STAT_DEFINED(HttpInflightCount, int32(GHttpInflightCount.Get()), ECsvCustomStatOp::Set);
-	CSV_CUSTOM_STAT_DEFINED(HttpRetryCount, int32(GHttpRetryCount.Get()), ECsvCustomStatOp::Set);
 }
 
 void FOnDemandIoBackendStats::OnHttpError()
@@ -234,10 +234,6 @@ void FOnDemandIoBackendStats::OnHttpError()
 	GHttpPendingCount.Add(-1);
 	GHttpInflightCount.Add(-1);
 	GHttpErrorCount.Add(1);
-
-	CSV_CUSTOM_STAT_DEFINED(HttpPendingCount, int32(GHttpPendingCount.Get()), ECsvCustomStatOp::Set);
-	CSV_CUSTOM_STAT_DEFINED(HttpInflightCount, int32(GHttpInflightCount.Get()), ECsvCustomStatOp::Set);
-	CSV_CUSTOM_STAT_DEFINED(HttpErrorCount, int32(GHttpErrorCount.Get()), ECsvCustomStatOp::Set);
 }
 
 } // namespace UE::IO::Private
diff --git a/Engine/Source/Runtime/Experimental/IoStoreOnDemand/Private/Statistics.h b/Engine/Source/Runtime/Experimental/IoStoreOnDemand/Private/Statistics.h
--- a/Engine/Source/Runtime/Experimental/IoStoreOnDemand/Private/Statistics.h
+++ b/Engine/Source/Runtime/Experimental/IoStoreOnDemand/Private/Statistics.h
@@ -29,11 +29,9 @@ public:
 	FOnDemandIoBackendStats() IAS_STATISTICS_IMPL()
 	~FOnDemandIoBackendStats() IAS_STATISTICS_IMPL()
 	void OnIoRequestEnqueue() IAS_STATISTICS_IMPL()
-	void OnIoRequestComplete(uint64 RequestSize) IAS_STATISTICS_IMPL()
+	void OnIoRequestComplete(uint64 Size, uint64 DurationMs) IAS_STATISTICS_IMPL()
 	void OnIoRequestCancel() IAS_STATISTICS_IMPL()
-	void OnIoRequestFail() IAS_STATISTICS_IMPL()
-	void OnChunkRequestCreate() IAS_STATISTICS_IMPL()
-	void OnChunkRequestRelease() IAS_STATISTICS_IMPL()
+	void OnIoRequestError() IAS_STATISTICS_IMPL()
 
 	void OnCacheError() IAS_STATISTICS_IMPL()
 	void OnCacheGet(uint64 DataSize) IAS_STATISTICS_IMPL()
@@ -45,7 +43,7 @@ public:
 
 	void OnHttpEnqueue() IAS_STATISTICS_IMPL()
 	void OnHttpDequeue() IAS_STATISTICS_IMPL()
-	void OnHttpGet(uint64 InSize) IAS_STATISTICS_IMPL()
+	void OnHttpGet(uint64 Size, uint64 DurationMs) IAS_STATISTICS_IMPL()
 	void OnHttpRetry() IAS_STATISTICS_IMPL()
 	void OnHttpError() IAS_STATISTICS_IMPL()
 };
