Ias: Calculate average duration and bandwidth stats as a running window over last 16 requests

#rnx

diff --git a/Engine/Source/Runtime/Experimental/IoStoreOnDemand/Private/Statistics.cpp b/Engine/Source/Runtime/Experimental/IoStoreOnDemand/Private/Statistics.cpp
--- a/Engine/Source/Runtime/Experimental/IoStoreOnDemand/Private/Statistics.cpp
+++ b/Engine/Source/Runtime/Experimental/IoStoreOnDemand/Private/Statistics.cpp
@@ -60,6 +60,12 @@ FCounterInt			GHttpDurationMs(TEXT("Ias/HttpDurationMs"), TraceCounterDisplayHin
 FCounterInt			GHttpDurationMsAvg(TEXT("Ias/HttpDurationMsAvg"), TraceCounterDisplayHint_None);
 FCounterInt			GHttpDurationMsMax(TEXT("Ias/HttpDurationMsMax"), TraceCounterDisplayHint_None);
 int64				GHttpDurationMsSum = 0;
+constexpr int64		GHttpHistoryCount = 16;
+int64				GHttpHistoryDuration[GHttpHistoryCount] = {};
+int64				GHttpHistoryBytes[GHttpHistoryCount] = {};
+int64				GHttpHistoryTotalDuration = 0;
+int64				GHttpHistoryTotalBytes = 0;
+int64 				GHttpHistoryIndex = 0;
 
 ////////////////////////////////////////////////////////////////////////////////
 // CSV STATS
@@ -214,12 +220,26 @@ void FOnDemandIoBackendStats::OnHttpGet(uint64 Size, uint64 DurationMs)
 	GHttpDurationMsSum += DurationMs;
 	GHttpDurationMs.Set(DurationMs);
 
-	GHttpBandwidthMpbs.Set((GHttpDownloadedBytes.Get()*8)/(GHttpDurationMsSum+1)/1000);
-	GHttpDurationMsAvg.Set(GHttpDurationMsSum/GHttpGetCount.Get());
+	int64 OldDuration = GHttpHistoryDuration[GHttpHistoryIndex];
+	int64 NewDuration = (int64)DurationMs;
+	
+	GHttpHistoryTotalDuration -= OldDuration;
+	GHttpHistoryTotalDuration += NewDuration;
+	GHttpHistoryDuration[GHttpHistoryIndex] = NewDuration;
+
+	GHttpHistoryTotalBytes -= GHttpHistoryBytes[GHttpHistoryIndex];
+	GHttpHistoryTotalBytes += Size;
+	GHttpHistoryBytes[GHttpHistoryIndex] = Size;
+
+	GHttpBandwidthMpbs.Set((GHttpHistoryTotalBytes*8)/(GHttpHistoryTotalDuration+1)/1000);
+	GHttpDurationMsAvg.Set(GHttpHistoryTotalDuration/GHttpHistoryCount);
+
 	if (GHttpDurationMsMax.Get() < (int64)DurationMs)
 	{
 		GHttpDurationMsMax.Set((int64)DurationMs);
 	}
+
+	GHttpHistoryIndex = (GHttpHistoryIndex + 1) % GHttpHistoryCount;
 }
 
 void FOnDemandIoBackendStats::OnHttpRetry()
