Wrote up some quick details about configuration values.

#rnx

diff --git a/Engine/Source/Runtime/Experimental/IoStoreOnDemand/Private/FileIoCache.h b/Engine/Source/Runtime/Experimental/IoStoreOnDemand/Private/FileIoCache.h
--- a/Engine/Source/Runtime/Experimental/IoStoreOnDemand/Private/FileIoCache.h
+++ b/Engine/Source/Runtime/Experimental/IoStoreOnDemand/Private/FileIoCache.h
@@ -44,7 +44,7 @@ struct FFileIoCacheConfig
 
 	uint64		DiskQuota = 512ull << 20;
 	uint32		MemoryQuota = 2 << 20;
-	uint32		JournalQuota = 4 << 20;
+	uint32		JournalQuota = 4 << 20; // description in JournalCache.cpp
 	uint32		JournalFlushInterval = 4;
 	uint32		UseLegacy = 0;
 	FRate		WriteRate;
diff --git a/Engine/Source/Runtime/Experimental/IoStoreOnDemand/Private/JournaledCache.cpp b/Engine/Source/Runtime/Experimental/IoStoreOnDemand/Private/JournaledCache.cpp
--- a/Engine/Source/Runtime/Experimental/IoStoreOnDemand/Private/JournaledCache.cpp
+++ b/Engine/Source/Runtime/Experimental/IoStoreOnDemand/Private/JournaledCache.cpp
@@ -27,6 +27,20 @@
 
 #include <atomic>
 
+/*
+DiskQuota is the maximum bytes on disk the cache will use. This includes the
+JournalQuota (available_data_bytes = diskq - jrnq). JournalQuota should be
+chosen such that it holds at least one wrap such that overwrites can be
+detected. For example, assuming an average size of cache items of 2KiB (very
+conservative), a 512MiB cache can hold 256K items. Journal entries are 16
+bytes, so a 256K * 16 is sufficient; 256K * 16 = 4MiB.
+
+JournalFlushInterval dictates how often a write-op allowance is spent on
+writing the journal. For example, a value of 4 would be 3 data writes followed
+by 1 journal write. A value of <=1 will flush immediately (N.B. doing so causes
+two writes per allowed op-throttle).
+*/
+
 namespace UE::IO::Private
 {
 
