IAS - log when backend is enabled/disabled

#rnx

diff --git a/Engine/Source/Runtime/Experimental/IoStoreOnDemand/Private/OnDemandIoDispatcherBackend.cpp b/Engine/Source/Runtime/Experimental/IoStoreOnDemand/Private/OnDemandIoDispatcherBackend.cpp
--- a/Engine/Source/Runtime/Experimental/IoStoreOnDemand/Private/OnDemandIoDispatcherBackend.cpp
+++ b/Engine/Source/Runtime/Experimental/IoStoreOnDemand/Private/OnDemandIoDispatcherBackend.cpp
@@ -1488,11 +1488,13 @@ void FOnDemandIoBackend::Mount(const FOnDemandEndpoint& Endpoint)
 
 void FOnDemandIoBackend::SetBulkOptionalEnabled(bool bInEnabled)
 {
+	UE_LOG(LogIas, Log, TEXT("HTTP optional bulk data streaming '%s'"), bInEnabled ? TEXT("Enabled") : TEXT("Disabled"));
 	bEnableBulkOptional = bInEnabled;
 }
 
 void FOnDemandIoBackend::SetEnabled(bool bInEnabled)
 {
+	UE_LOG(LogIas, Log, TEXT("HTTP streaming '%s'"), bInEnabled ? TEXT("Enabled") : TEXT("Disabled"));
 	bEnabled = bInEnabled;
 }
 
