Added always-flush path if journal's flush interval is 1 (or 0 for that matter)

#rnx

diff --git a/Engine/Source/Runtime/Experimental/IoStoreOnDemand/Private/JournaledCache.cpp b/Engine/Source/Runtime/Experimental/IoStoreOnDemand/Private/JournaledCache.cpp
--- a/Engine/Source/Runtime/Experimental/IoStoreOnDemand/Private/JournaledCache.cpp
+++ b/Engine/Source/Runtime/Experimental/IoStoreOnDemand/Private/JournaledCache.cpp
@@ -991,19 +991,22 @@ bool FCache::Put(uint64 Key, FIoBuffer& Data)
 ////////////////////////////////////////////////////////////////////////////////
 int32 FCache::Flush(int32 Allowance)
 {
-	FlushIndex += 1;
-	if (FlushIndex >= FlushPeriod)
+	if (FlushPeriod > 1)
 	{
-		FlushIndex -= FlushPeriod;
-		FWriteAccess _(Lock);
-		return Allowance - Journal.Flush();
-	}
+		FlushIndex += 1;
+		if (FlushIndex >= FlushPeriod)
+		{
+			FlushIndex -= FlushPeriod;
+			FWriteAccess _(Lock);
+			return Allowance - Journal.Flush();
+		}
 
-	if (Pending.GetUsed() == 0)
-	{
-		FlushIndex -= (FlushIndex == 1);
-		FWriteAccess _(Lock);
-		return Allowance - Journal.Flush();
+		if (Pending.GetUsed() == 0)
+		{
+			FlushIndex -= (FlushIndex == 1);
+			FWriteAccess _(Lock);
+			return Allowance - Journal.Flush();
+		}
 	}
 
 	TRACE_CPUPROFILER_EVENT_SCOPE(IasCache::Flush_Pending);
@@ -1033,6 +1036,11 @@ int32 FCache::Flush(int32 Allowance)
 		 * we can re-add leftover peeled items back to pending? */
 	}
 
+	if (FlushPeriod <= 1)
+	{
+		Allowance -= Journal.Flush();
+	}
+
 	return Allowance - PendingSize;
 }
 
