Fixed uss-after-free

#rnx

diff --git a/Engine/Source/Runtime/Experimental/IoStoreOnDemand/Private/OnDemandIoDispatcherBackend.cpp b/Engine/Source/Runtime/Experimental/IoStoreOnDemand/Private/OnDemandIoDispatcherBackend.cpp
--- a/Engine/Source/Runtime/Experimental/IoStoreOnDemand/Private/OnDemandIoDispatcherBackend.cpp
+++ b/Engine/Source/Runtime/Experimental/IoStoreOnDemand/Private/OnDemandIoDispatcherBackend.cpp
@@ -342,8 +342,10 @@ FHttpClient::FHttpClient(const FString& ServiceUrl, int32 MaxConnectionCount)
 	: SvcsUrl(ServiceUrl)
 	, MaxConnections(MaxConnectionCount)
 {
+	auto ServiceUrlAnsi = StringCast<ANSICHAR>(*ServiceUrl, ServiceUrl.Len());
+
 	HTTP::FConnectionPool::FParams Params;
-	Params.SetHostFromUrl(StringCast<ANSICHAR>(*ServiceUrl).Get());
+	Params.SetHostFromUrl(ServiceUrlAnsi);
 	Params.ConnectionCount = MaxConnectionCount;
 	ConnectionPool = MakeUnique<HTTP::FConnectionPool>(Params);
 }
