More regular IAS cache flush check so we hit the allowance ticks more accurately.

#rnx

diff --git a/Engine/Source/Runtime/Experimental/IoStoreOnDemand/Private/JournaledCache.cpp b/Engine/Source/Runtime/Experimental/IoStoreOnDemand/Private/JournaledCache.cpp
--- a/Engine/Source/Runtime/Experimental/IoStoreOnDemand/Private/JournaledCache.cpp
+++ b/Engine/Source/Runtime/Experimental/IoStoreOnDemand/Private/JournaledCache.cpp
@@ -1146,7 +1146,7 @@ uint32 FJournaledCache::Run()
 	while (Running.load(std::memory_order_relaxed))
 	{
 		Update();
-		WakeEvent->Wait(127);
+		WakeEvent->Wait(37);
 	}
 	return 0;
 }
