Skip content if the request type won't return any

#rnx

diff --git a/Engine/Source/Runtime/Experimental/IoStoreOnDemand/Private/CoreHttp/Http.cpp b/Engine/Source/Runtime/Experimental/IoStoreOnDemand/Private/CoreHttp/Http.cpp
--- a/Engine/Source/Runtime/Experimental/IoStoreOnDemand/Private/CoreHttp/Http.cpp
+++ b/Engine/Source/Runtime/Experimental/IoStoreOnDemand/Private/CoreHttp/Http.cpp
@@ -800,7 +800,8 @@ struct alignas(16) FActivity
 	EState				State;
 	EWait				SocketWait;
 	uint8				IsKeepAlive : 1;
-	uint8				_Unused0 : 7;
+	uint8				NoContent : 1;
+	uint8				_Unused0 : 6;
 	uint32				StateParam = 0;
 
 	FSocketPool*		Pool;
@@ -1656,16 +1657,19 @@ static int32 DoRecvMessage(FActivity* Activity)
 		Activity->Sink(SinkArg);
 	}
 
-	if (Internal.Dest == nullptr)
+	if (Activity->NoContent == 0)
 	{
-		Activity_SetError(Activity, "User did not provide a destination buffer");
-		return -1;
-	}
+		if (Internal.Dest == nullptr)
+		{
+			Activity_SetError(Activity, "User did not provide a destination buffer");
+			return -1;
+		}
 
-	// The user seems to have forgotten something. Let's help them along
-	if (Internal.Dest->GetSize() == 0)
-	{
-		*Internal.Dest = FIoBuffer(ContentLength);
+		// The user seems to have forgotten something. Let's help them along
+		if (Internal.Dest->GetSize() == 0)
+		{
+			*Internal.Dest = FIoBuffer(ContentLength);
+		}
 	}
 
 	bool Streamed = (Internal.Dest->GetSize() < ContentLength);
@@ -1679,6 +1683,17 @@ static int32 DoRecvMessage(FActivity* Activity)
 		return -1;
 	}
 
+	if (Activity->NoContent == 1)
+	{
+		if (AlreadyReceived)
+		{
+			Activity_SetError(Activity, "Received content when none was expected");
+			return -1;
+		}
+		Activity->State = FActivity::EState::RecvDone;
+		return 0;
+	}
+
 	Activity->State = Streamed ? FActivity::EState::RecvStream : FActivity::EState::RecvContent;
 	Activity->StateParam = AlreadyReceived;
 
@@ -1854,6 +1869,8 @@ FRequest FEventLoop::FImpl::Request(
 		Path = "/";
 	}
 
+	Activity->NoContent = (Method == "HEAD");
+
 	FMessageBuilder Builder(Activity->Buffer);
 
 	Builder << Method << " " << Path << " HTTP/1.1\r\n"
@@ -1985,23 +2002,25 @@ uint32 FEventLoop::FImpl::Tick(uint32 PollTimeoutMs)
 			if (Result)
 				break;
 
-		case FActivity::EState::RecvContent:
-		case FActivity::EState::RecvStream: {
-			decltype(DoRecvContent)* Handler;
-			if (Activity->State == FActivity::EState::RecvContent)
+			if (Activity->State != FActivity::EState::RecvDone)
 			{
-				Handler = DoRecvContent;
-			}
-			else
-			{
-				Handler = DoRecvStream;
-			}
+			case FActivity::EState::RecvContent:
+			case FActivity::EState::RecvStream:
+				decltype(DoRecvContent)* Handler;
+				if (Activity->State == FActivity::EState::RecvContent)
+				{
+					Handler = DoRecvContent;
+				}
+				else
+				{
+					Handler = DoRecvStream;
+				}
 
-			auto [ResultInner, RecvSize] = Handler(Activity, ~0u);
-			Result = ResultInner;
-			if (Result)
-				break;
-		}
+				auto [ResultInner, RecvSize] = Handler(Activity, ~0u);
+				Result = ResultInner;
+				if (Result)
+					break;
+			}
 
 		case FActivity::EState::RecvDone:
 			Result = DoRecvDone(Activity);
