Added debug output approximating IAS' latency to content.

#rnx

diff --git a/Engine/Source/Runtime/Experimental/IoStoreOnDemand/Private/OnDemandIoDispatcherBackend.cpp b/Engine/Source/Runtime/Experimental/IoStoreOnDemand/Private/OnDemandIoDispatcherBackend.cpp
--- a/Engine/Source/Runtime/Experimental/IoStoreOnDemand/Private/OnDemandIoDispatcherBackend.cpp
+++ b/Engine/Source/Runtime/Experimental/IoStoreOnDemand/Private/OnDemandIoDispatcherBackend.cpp
@@ -68,6 +68,60 @@ static FAutoConsoleVariableRef CVar_IoDispatcherBulkOptionalEnabled(
 namespace UE::IO::Private
 {
 
+///////////////////////////////////////////////////////////////////////////////
+#if !UE_BUILD_SHIPPING
+static void LatencyTest(FStringView InUrl, FStringView InPath)
+{
+	auto AnsiUrl = StringCast<ANSICHAR>(InUrl.GetData(), InUrl.Len());
+
+	using namespace UE::HTTP;
+
+	FConnectionPool::FParams PoolParams;
+	PoolParams.SetHostFromUrl(AnsiUrl);
+	PoolParams.ConnectionCount = 1;
+	FConnectionPool Pool(PoolParams);
+
+	TAnsiStringBuilder<256> AnsiPath;
+	AnsiPath << "/";
+	AnsiPath << InPath;
+
+	FEventLoop Loop;
+	int32 Results[4] = {};
+	for (uint32 i = 0; i < UE_ARRAY_COUNT(Results); ++i)
+	{
+		bool Ok = false;
+
+		FRequest Request = Loop.Request("HEAD", AnsiPath, Pool);
+		Loop.Send(MoveTemp(Request), [&] (const FTicketStatus& Status)
+		{
+			if (Status.GetId() != FTicketStatus::EId::Response)
+				return;
+
+			const FResponse& Response = Status.GetResponse();
+			Ok = (Response.GetStatus() == EStatusCodeClass::Successful);
+		});
+
+		uint64 Cycles = FPlatformTime::Cycles64();
+		while (Loop.Tick(-1) != 0);
+		Cycles = FPlatformTime::Cycles64() - Cycles;
+
+		Results[i] = Ok ? int32(Cycles) : -1;
+	}
+
+	int64 Freq = int64(1.0 / FPlatformTime::GetSecondsPerCycle());
+	for (int32& Result : Results)
+	{
+		if (Result == -1)
+			continue;
+
+		Result = int32((int64(Result) * 1000) / Freq);
+	}
+
+	UE_LOG(LogIas, VeryVerbose, TEXT("HEAD latencies (ms); %d %d %d %d (%s)"),
+		Results[0], Results[1], Results[2], Results[3], InUrl.GetData());
+}
+#endif // !UE_BUILD_SHIPPING
+
 ///////////////////////////////////////////////////////////////////////////////
 static void LogHttpResult(const TCHAR* Url, uint32 StatusCode, uint32 DurationMs, uint32 Size, uint32 Offset, const char* Memo="ok")
 {
@@ -1606,6 +1660,13 @@ FIoStatus FOnDemandIoBackend::AddToc(const FOnDemandEndpoint& Endpoint)
 		}
 	}
 
+#if !UE_BUILD_SHIPPING
+	Launch(TEXT("IasLatencyTest"), [ServiceUrl=Endpoint.ServiceUrl, TocPath=Endpoint.TocPath] ()
+	{
+		UE::IO::Private::LatencyTest(ServiceUrl, TocPath);
+	});
+#endif // !UE_BUILD_SHIPPING
+
 	IoStore->AddToc(Endpoint, Toc.ConsumeValueOrDie());
 
 	UE_CLOG(Endpoint.ServiceUrl !=  HttpClient->ServiceUrl(),
