Shuffled latency injector's initialisation around a little

#rnx

diff --git a/Engine/Source/Runtime/Experimental/IoStoreOnDemand/Private/LatencyInjector.cpp b/Engine/Source/Runtime/Experimental/IoStoreOnDemand/Private/LatencyInjector.cpp
--- a/Engine/Source/Runtime/Experimental/IoStoreOnDemand/Private/LatencyInjector.cpp
+++ b/Engine/Source/Runtime/Experimental/IoStoreOnDemand/Private/LatencyInjector.cpp
@@ -21,10 +21,8 @@ static int64 CycleFreq;
 static int32 MinMs;
 static int32 MaxMs;
 
-};
-
 ////////////////////////////////////////////////////////////////////////////////
-void FLatencyInjector::Initialize(const TCHAR* CommandLine)
+static void InitializeLatency(const TCHAR* CommandLine)
 {
 	// -Ias.AddLatencyMs=130   : add 130ms of latency
 	// -Ias.AddLatencyMs=40/90 : random latency between 40ms and 90ms
@@ -44,7 +42,18 @@ void FLatencyInjector::Initialize(const TCHAR* CommandLine)
 	});
 
 	for (; Values.Num() < 2; Values.Add(-1));
-	Set(Values[0], Values[1]);
+	FLatencyInjector::Set(Values[0], Values[1]);
+}
+
+}; // namespace Injector
+
+
+
+////////////////////////////////////////////////////////////////////////////////
+void FLatencyInjector::Initialize(const TCHAR* CommandLine)
+{
+	using namespace Injector;
+	InitializeLatency(CommandLine);
 }
 
 ////////////////////////////////////////////////////////////////////////////////
