Coopted the latency injector so it can also be used to inject a percentage of failures

#rnx

diff --git a/Engine/Source/Runtime/Experimental/IoStoreOnDemand/Private/CoreHttp/Http.cpp b/Engine/Source/Runtime/Experimental/IoStoreOnDemand/Private/CoreHttp/Http.cpp
--- a/Engine/Source/Runtime/Experimental/IoStoreOnDemand/Private/CoreHttp/Http.cpp
+++ b/Engine/Source/Runtime/Experimental/IoStoreOnDemand/Private/CoreHttp/Http.cpp
@@ -1777,7 +1777,11 @@ static FHandlerResult DoRecvContent(FActivity* Activity, uint32 MaxRecvSize)
 		RecvSize += Result;
 	}
 
-	FLatencyInjector::Begin(FLatencyInjector::EType::Network, Activity->StateParam);
+	if (!FLatencyInjector::Begin(FLatencyInjector::EType::Network, Activity->StateParam))
+	{
+		Activity_SetError(Activity, "Forced random failure");
+		return { -1 };
+	}
 
 	Activity->State = FActivity::EState::RecvDone;
 	return { 0, RecvSize };
diff --git a/Engine/Source/Runtime/Experimental/IoStoreOnDemand/Private/LatencyInjector.cpp b/Engine/Source/Runtime/Experimental/IoStoreOnDemand/Private/LatencyInjector.cpp
--- a/Engine/Source/Runtime/Experimental/IoStoreOnDemand/Private/LatencyInjector.cpp
+++ b/Engine/Source/Runtime/Experimental/IoStoreOnDemand/Private/LatencyInjector.cpp
@@ -20,7 +20,8 @@ static int64 CycleBase;
 static int64 CycleFreq;
 static int32 MinMs;
 static int32 MaxMs;
-
+static int32 FailPercent;
+static int32 FailState;
 ////////////////////////////////////////////////////////////////////////////////
 static void InitializeLatency(const TCHAR* CommandLine)
 {
@@ -45,6 +46,27 @@ static void InitializeLatency(const TCHAR* CommandLine)
 	FLatencyInjector::Set(Values[0], Values[1]);
 }
 
+////////////////////////////////////////////////////////////////////////////////
+static void InitializeFailure(const TCHAR* CommandLine)
+{
+	// -Ias.ForceFailPercent=40  : 40% chance that a request will fail
+	// -Ias.ForceFailPercent=0   : invalid
+	// -Ias.ForceFailPercent=101 : invalid
+
+	FString Value;
+	if (!FParse::Value(CommandLine, TEXT("Ias.ForceFailPercent="), Value))
+	{
+		return;
+	}
+
+	int32 Percent = -1;
+	LexFromString(Percent, *Value);
+	if (Percent < 1 || Percent > 100)
+		return;
+
+	FLatencyInjector::SetFailureRate(Percent);
+}
+
 }; // namespace Injector
 
 
@@ -54,6 +76,7 @@ void FLatencyInjector::Initialize(const TCHAR* CommandLine)
 {
 	using namespace Injector;
 	InitializeLatency(CommandLine);
+	InitializeFailure(CommandLine);
 }
 
 ////////////////////////////////////////////////////////////////////////////////
@@ -79,14 +102,31 @@ void FLatencyInjector::Set(int32 MinMs, int32 MaxMs)
 }
 
 ////////////////////////////////////////////////////////////////////////////////
-void FLatencyInjector::Begin(EType, uint32& Param)
+void FLatencyInjector::SetFailureRate(int32 Percent)
+{
+	using namespace Injector;
+	FailPercent = Percent;
+	FailState = int32(FPlatformTime::Cycles64()) | 1;
+}
+
+////////////////////////////////////////////////////////////////////////////////
+bool FLatencyInjector::Begin(EType, uint32& Param)
 {
 	using namespace Injector;
 
+	if (FailPercent)
+	{
+		FailState *= 0xac564b05;
+		if (int32 Value = ((FailState & 0xff) * 100) / 0xff; Value < FailPercent)
+		{
+			return false;
+		}
+	}
+
 	if (MaxMs <= 0 || MinMs < 0)
 	{
 		Param = 0;
-		return;
+		return true;
 	}
 
 	int64 Cycles = FPlatformTime::Cycles64();
@@ -107,6 +147,7 @@ void FLatencyInjector::Begin(EType, uint32& Param)
 	Bias /= 1000;
 
 	Param = uint32(Cycles + Bias - CycleBase);
+	return true;
 }
 
 ////////////////////////////////////////////////////////////////////////////////
diff --git a/Engine/Source/Runtime/Experimental/IoStoreOnDemand/Private/LatencyInjector.h b/Engine/Source/Runtime/Experimental/IoStoreOnDemand/Private/LatencyInjector.h
--- a/Engine/Source/Runtime/Experimental/IoStoreOnDemand/Private/LatencyInjector.h
+++ b/Engine/Source/Runtime/Experimental/IoStoreOnDemand/Private/LatencyInjector.h
@@ -21,7 +21,8 @@ struct FLatencyInjector
 	enum class EType { Network, File };
 	static void Initialize(const TCHAR* CommandLine)IAS_DISABLED_IMPL(return);
 	static void Set(int32 MinMs, int32 MaxMs)		IAS_DISABLED_IMPL(return);
-	static void Begin(EType, uint32& Param)			IAS_DISABLED_IMPL(return);
+	static void SetFailureRate(int32 Percent)		IAS_DISABLED_IMPL(return);
+	static bool Begin(EType, uint32& Param)			IAS_DISABLED_IMPL(return true);
 	static bool HasExpired(uint32 Param)			IAS_DISABLED_IMPL(return true);
 };
 
