As IAS receives data in tight loops we can block until there's something to do

#rnx

diff --git a/Engine/Source/Runtime/Experimental/IoStoreOnDemand/Private/OnDemandIoDispatcherBackend.cpp b/Engine/Source/Runtime/Experimental/IoStoreOnDemand/Private/OnDemandIoDispatcherBackend.cpp
--- a/Engine/Source/Runtime/Experimental/IoStoreOnDemand/Private/OnDemandIoDispatcherBackend.cpp
+++ b/Engine/Source/Runtime/Experimental/IoStoreOnDemand/Private/OnDemandIoDispatcherBackend.cpp
@@ -381,7 +381,7 @@ public:
 	void Get(FAnsiStringView Url, const FIoOffsetAndLength& Range, FGetCallback&& Callback);
 
 	/** @return True if the client has pending work otherwise false. */
-	bool Tick();
+	bool Tick(bool Block=false);
 
 private:
 	void Issue(FAnsiStringView Url, FGetCallback&& Callback, FIoOffsetAndLength Range = FIoOffsetAndLength());
@@ -469,9 +469,10 @@ void FHttpClient::Issue(FAnsiStringView Url, FGetCallback&& Callback, FIoOffsetA
 	EventLoop.Send(MoveTemp(Request), MoveTemp(Sink));
 }
 
-bool FHttpClient::Tick()
+bool FHttpClient::Tick(bool Block)
 {
-	return EventLoop.Tick(GIoDispatcherHttpPollTimeoutMs) != 0;
+	int32 TimeoutMs = Block ? -1 : GIoDispatcherHttpPollTimeoutMs; 
+	return EventLoop.Tick(TimeoutMs) != 0;
 }
 
 ///////////////////////////////////////////////////////////////////////////////
@@ -1630,7 +1631,7 @@ TIoStatusOr<FOnDemandToc> FOnDemandIoBackend::GetToc(FHttpClient& HttpClient, co
 			}
 		});
 
-		while (HttpClient.Tick());
+		while (HttpClient.Tick(true));
 
 		if (Toc.IsOk())
 		{
@@ -1742,7 +1743,7 @@ uint32 FOnDemandIoBackend::Run()
 					TRACE_CPUPROFILER_EVENT_SCOPE(FOnDemandIoBackend::TickHttpSaturated);
 					while (NumConcurrentRequests >= MaxConcurrentRequests)
 					{
-						HttpClient->Tick();
+						HttpClient->Tick(true);
 					}
 				}
 
@@ -1755,7 +1756,7 @@ uint32 FOnDemandIoBackend::Run()
 			{
 				// Keep processing pending connections until all requests are completed or a new one is issued
 				TRACE_CPUPROFILER_EVENT_SCOPE(FOnDemandIoBackend::TickHttp);
-				while (HttpClient->Tick() && !NextChunkRequest)
+				while (HttpClient->Tick(true) && !NextChunkRequest)
 				{
 					NextChunkRequest = HttpRequests.Dequeue();
 				}
