Hooking IAS into analytics

#rnx

- We currently log if Ias is enabled, the url of the CDN and then the counter statistics that have been marked as useful for analytics.
- Can be disabled by setting the cvar "ias.ReportAnalytics" to false
- Fixed some whitespace issues and include order.

diff --git a/Engine/Source/Runtime/Experimental/IoStoreOnDemand/IoStoreOnDemand.Build.cs b/Engine/Source/Runtime/Experimental/IoStoreOnDemand/IoStoreOnDemand.Build.cs
--- a/Engine/Source/Runtime/Experimental/IoStoreOnDemand/IoStoreOnDemand.Build.cs
+++ b/Engine/Source/Runtime/Experimental/IoStoreOnDemand/IoStoreOnDemand.Build.cs
@@ -6,8 +6,15 @@ public class IoStoreOnDemand : ModuleRules
 {
 	public IoStoreOnDemand(ReadOnlyTargetRules Target) : base(Target)
 	{
-        PublicDependencyModuleNames.Add("Core");
-        PrivateDependencyModuleNames.AddRange(new string[] { "HTTP", "Json" });
+		PublicDependencyModuleNames.Add("Core");
+		PrivateDependencyModuleNames.AddRange(
+			new string[] {
+				"HTTP",
+				"Json",
+				"Analytics",
+			}
+		);
+
 		bAllowConfidentialPlatformDefines = true;
 
 		if (Target.Platform.IsInGroup(UnrealPlatformGroup.Desktop) &&
@@ -15,5 +22,5 @@ public class IoStoreOnDemand : ModuleRules
 		{
 			PrivateDependencyModuleNames.AddRange(new string[] { "S3Client", "RSA" });
 		}
-    }
+	}
 }
diff --git a/Engine/Source/Runtime/Experimental/IoStoreOnDemand/Private/IoStoreOnDemand.cpp b/Engine/Source/Runtime/Experimental/IoStoreOnDemand/Private/IoStoreOnDemand.cpp
--- a/Engine/Source/Runtime/Experimental/IoStoreOnDemand/Private/IoStoreOnDemand.cpp
+++ b/Engine/Source/Runtime/Experimental/IoStoreOnDemand/Private/IoStoreOnDemand.cpp
@@ -1443,6 +1443,14 @@ void FIoStoreOnDemandModule::SetEnabled(bool bInEnabled)
 	}
 }
 
+void FIoStoreOnDemandModule::ReportAnalytics(TArray<FAnalyticsEventAttribute>& OutAnalyticsArray) const
+{
+	if (Backend.IsValid())
+	{
+		Backend->ReportAnalytics(OutAnalyticsArray);
+	}
+}
+
 void FIoStoreOnDemandModule::StartupModule()
 {
 	LLM_SCOPE(ELLMTag::FileSystem);
diff --git a/Engine/Source/Runtime/Experimental/IoStoreOnDemand/Private/OnDemandIoDispatcherBackend.cpp b/Engine/Source/Runtime/Experimental/IoStoreOnDemand/Private/OnDemandIoDispatcherBackend.cpp
--- a/Engine/Source/Runtime/Experimental/IoStoreOnDemand/Private/OnDemandIoDispatcherBackend.cpp
+++ b/Engine/Source/Runtime/Experimental/IoStoreOnDemand/Private/OnDemandIoDispatcherBackend.cpp
@@ -1,8 +1,8 @@
 // Copyright Epic Games, Inc. All Rights Reserved.
 
 #include "OnDemandIoDispatcherBackend.h"
-#include "Statistics.h"
 
+#include "AnalyticsEventAttribute.h"
 #include "CancellationToken.h"
 #include "Containers/StringView.h"
 #include "CoreHttp/Client.h"
@@ -32,6 +32,7 @@
 #include "Serialization/JsonReader.h"
 #include "Serialization/JsonSerializer.h"
 #include "Serialization/MemoryReader.h"
+#include "Statistics.h"
 #include "Tasks/Task.h"
 
 #include <atomic>
@@ -65,6 +66,11 @@ static FAutoConsoleVariableRef CVar_IoDispatcherBulkOptionalEnabled(
 	TEXT("Enables bulk optional requests.")
 );
 
+static TAutoConsoleVariable<bool> CVar_IoReportAnalytics(
+	TEXT("ias.ReportAnalytics"),
+	true,
+	TEXT("Enables reporting statics to the analytics system"));
+
 namespace UE::IO::Private
 {
 
@@ -1189,6 +1195,7 @@ public:
 	virtual void Mount(const FOnDemandEndpoint& Endpoint) override;
 	virtual void SetBulkOptionalEnabled(bool bInEnabled) override;
 	virtual void SetEnabled(bool bInEnabled) override;
+	virtual void ReportAnalytics(TArray<FAnalyticsEventAttribute>& OutAnalyticsArray) const override;
 
 #if IS_PROGRAM || WITH_EDITOR
 	virtual bool FlushDeferedEndPoints(double TimeOut = 0.0) override;
@@ -1555,6 +1562,32 @@ void FOnDemandIoBackend::SetEnabled(bool bInEnabled)
 	bEnabled = bInEnabled;
 }
 
+void FOnDemandIoBackend::ReportAnalytics(TArray<FAnalyticsEventAttribute>& OutAnalyticsArray) const
+{
+	if (!CVar_IoReportAnalytics.GetValueOnAnyThread())
+	{
+		return;
+	}
+
+	if (bEnabled)
+	{
+		AppendAnalyticsEventAttributeArray(OutAnalyticsArray,
+			TEXT("IasEnabled"), false,
+			TEXT("IasCDNBackend"), HttpClient->ServiceUrl()
+		);
+
+		Stats.ReportAnalytics(OutAnalyticsArray);
+	}
+	else
+	{
+		AppendAnalyticsEventAttributeArray(OutAnalyticsArray,
+			TEXT("IasEnabled"), false
+		);
+	}
+
+	
+}
+
 #if IS_PROGRAM || WITH_EDITOR
 bool FOnDemandIoBackend::FlushDeferedEndPoints(double TimeOut)
 {
diff --git a/Engine/Source/Runtime/Experimental/IoStoreOnDemand/Private/OnDemandIoDispatcherBackend.h b/Engine/Source/Runtime/Experimental/IoStoreOnDemand/Private/OnDemandIoDispatcherBackend.h
--- a/Engine/Source/Runtime/Experimental/IoStoreOnDemand/Private/OnDemandIoDispatcherBackend.h
+++ b/Engine/Source/Runtime/Experimental/IoStoreOnDemand/Private/OnDemandIoDispatcherBackend.h
@@ -8,6 +8,7 @@
 #define UE_API IOSTOREONDEMAND_API
 
 class IIoCache;
+struct FAnalyticsEventAttribute;
 
 namespace UE
 {
@@ -25,7 +26,10 @@ struct FOnDemandEndpoint
 	FString ServiceUrl;
 	FString TocPath;
 
-	bool IsValid() const { return (DistributionUrl.Len() > 0 || ServiceUrl.Len() > 0) && TocPath.Len() > 0; }
+	bool IsValid() const
+	{
+		return (DistributionUrl.Len() > 0 || ServiceUrl.Len() > 0) && TocPath.Len() > 0;
+	}
 };
 
 class IOnDemandIoDispatcherBackend
@@ -33,10 +37,13 @@ class IOnDemandIoDispatcherBackend
 {
 public:
 	virtual ~IOnDemandIoDispatcherBackend() = default;
+
 	virtual void Mount(const FOnDemandEndpoint& Endpoint) = 0;
 	virtual void SetBulkOptionalEnabled(bool bInEnabled) = 0;
 	virtual void SetEnabled(bool bInEnabled) = 0;
 
+	virtual void ReportAnalytics(TArray<FAnalyticsEventAttribute>& OutAnalyticsArray) const = 0;
+
 #if IS_PROGRAM || WITH_EDITOR
 	virtual bool FlushDeferedEndPoints(double TimeOut) = 0;
 	virtual TArray<FIoChunkId> GetAllChunkIds() = 0;
diff --git a/Engine/Source/Runtime/Experimental/IoStoreOnDemand/Private/Statistics.cpp b/Engine/Source/Runtime/Experimental/IoStoreOnDemand/Private/Statistics.cpp
--- a/Engine/Source/Runtime/Experimental/IoStoreOnDemand/Private/Statistics.cpp
+++ b/Engine/Source/Runtime/Experimental/IoStoreOnDemand/Private/Statistics.cpp
@@ -1,6 +1,8 @@
 // Copyright Epic Games, Inc. All Rights Reserved.
 
 #include "Statistics.h"
+
+#include "AnalyticsEventAttribute.h"
 #include "Misc/CoreDelegates.h"
 
 #if IAS_WITH_STATISTICS
@@ -137,6 +139,26 @@ FOnDemandIoBackendStats* FOnDemandIoBackendStats::Get()
 	return GStatistics;
 }
 
+void FOnDemandIoBackendStats::ReportAnalytics(TArray<FAnalyticsEventAttribute>& OutAnalyticsArray) const
+{
+	AppendAnalyticsEventAttributeArray(OutAnalyticsArray,
+		TEXT("IasHttpErrorCount"), GHttpErrorCount.Get(), 
+		TEXT("IasHttpRetryCount"), GHttpRetryCount.Get(),
+		TEXT("IasHttpGetCount"), GHttpGetCount.Get(),
+		TEXT("IasHttpPendingCount"), GHttpPendingCount.Get(),
+		TEXT("IasHttpDownloadedBytes"), GHttpDownloadedBytes.Get(),
+
+		TEXT("IasCacheErrorCount"), GCacheErrorCount.Get(),
+		TEXT("IasCacheGetCount"), GCacheGetCount.Get(),
+		TEXT("IasCachePutCount"), GCachePutCount.Get(),
+		TEXT("IasCachetRejectCount"), GCachePutRejectCount.Get(),
+		
+		TEXT("IasCacheCachedBytes"), GCacheCachedBytes.Get(),
+		TEXT("IasCacheReadBytes"), GCacheReadBytes.Get(),
+		TEXT("IasCacheRejectBytes"), GCacheRejectBytes.Get()
+	);	
+}
+
 void FOnDemandIoBackendStats::OnIoRequestEnqueue()
 {
 	GIoRequestCount.Add(1);
diff --git a/Engine/Source/Runtime/Experimental/IoStoreOnDemand/Private/Statistics.h b/Engine/Source/Runtime/Experimental/IoStoreOnDemand/Private/Statistics.h
--- a/Engine/Source/Runtime/Experimental/IoStoreOnDemand/Private/Statistics.h
+++ b/Engine/Source/Runtime/Experimental/IoStoreOnDemand/Private/Statistics.h
@@ -5,14 +5,14 @@
 #include "ProfilingDebugging/CountersTrace.h"
 #include "ProfilingDebugging/CsvProfiler.h"
 
-#include <atomic>
-
 #if COUNTERSTRACE_ENABLED || CSV_PROFILER
 #	define IAS_WITH_STATISTICS 1
 #else
 #	define IAS_WITH_STATISTICS 0
 #endif
 
+struct FAnalyticsEventAttribute;
+
 namespace UE::IO::Private
 {
 
@@ -26,8 +26,12 @@ class FOnDemandIoBackendStats
 {
 public:
 	static FOnDemandIoBackendStats* Get() IAS_STATISTICS_IMPL(nullptr)
+
 	FOnDemandIoBackendStats() IAS_STATISTICS_IMPL()
 	~FOnDemandIoBackendStats() IAS_STATISTICS_IMPL()
+
+	void ReportAnalytics(TArray<FAnalyticsEventAttribute>& OutAnalyticsArray) const IAS_STATISTICS_IMPL()
+
 	void OnIoRequestEnqueue() IAS_STATISTICS_IMPL()
 	void OnIoRequestComplete(uint64 Size, uint64 DurationMs) IAS_STATISTICS_IMPL()
 	void OnIoRequestCancel() IAS_STATISTICS_IMPL()
diff --git a/Engine/Source/Runtime/Experimental/IoStoreOnDemand/Public/IO/IoStoreOnDemand.h b/Engine/Source/Runtime/Experimental/IoStoreOnDemand/Public/IO/IoStoreOnDemand.h
--- a/Engine/Source/Runtime/Experimental/IoStoreOnDemand/Public/IO/IoStoreOnDemand.h
+++ b/Engine/Source/Runtime/Experimental/IoStoreOnDemand/Public/IO/IoStoreOnDemand.h
@@ -7,21 +7,22 @@
 #include "IO/IoHash.h"
 #include "IO/IoStatus.h"
 #include "Misc/Guid.h"
+#include "Modules/ModuleInterface.h"
 #include "Templates/SharedPointer.h"
 #include "UObject/NameTypes.h"
-#include "Modules/ModuleInterface.h"
 
 #if (IS_PROGRAM || WITH_EDITOR)
-#include "Misc/AES.h"
 #include "Containers/Map.h"
+#include "Misc/AES.h"
 #endif // (IS_PROGRAM || WITH_EDITOR)
 
 #define UE_API IOSTOREONDEMAND_API
 
 class FArchive;
-class FCbWriter;
 class FCbFieldView;
+class FCbWriter;
 class IIoStoreWriter;
+struct FAnalyticsEventAttribute;
 struct FIoContainerSettings;
 struct FIoStoreWriterSettings;
 
@@ -178,16 +179,17 @@ namespace UE
 class FIoStoreOnDemandModule
 	: public IModuleInterface
 {
+private:
 	TSharedPtr<UE::IOnDemandIoDispatcherBackend> Backend;
 
 public:
 	UE_API void SetBulkOptionalEnabled(bool bInEnabled);
 	UE_API void SetEnabled(bool bInEnabled);
 
+	UE_API void ReportAnalytics(TArray<FAnalyticsEventAttribute>& OutAnalyticsArray) const;
+
 	virtual void StartupModule() override;
 	virtual void ShutdownModule() override;
-
-private:
 };
 
 #undef UE_API
