Strain modification : add support for partial destruction
- this is required to support the fatigue based mode

diff --git a/Engine/Source/Runtime/Experimental/Chaos/Private/Chaos/StrainModification.cpp b/Engine/Source/Runtime/Experimental/Chaos/Private/Chaos/StrainModification.cpp
--- a/Engine/Source/Runtime/Experimental/Chaos/Private/Chaos/StrainModification.cpp
+++ b/Engine/Source/Runtime/Experimental/Chaos/Private/Chaos/StrainModification.cpp
@@ -3,6 +3,7 @@
 #include "Chaos/StrainModification.h"
 #include "Chaos/PBDRigidClustering.h"
 #include "PhysicsProxy/GeometryCollectionPhysicsProxy.h"
+#include "PhysicsProxy/ClusterUnionPhysicsProxy.h"
 #include "GeometryCollection/GeometryCollectionSimulationCoreTypes.h"
 #include "GeometryCollection/GeometryCollection.h"
 
@@ -16,14 +17,67 @@ namespace
 	}
 }
 
-const FGeometryCollectionPhysicsProxy* Chaos::FStrainedProxyModifier::GetProxy() const
+FGeometryCollectionPhysicsProxy* Chaos::FStrainedProxyAndRoot::CastToGeometryCollectionProxy() const
 {
-	return Proxy;
+	if (Proxy && Proxy->GetType() == EPhysicsProxyType::GeometryCollectionType)
+	{
+		return static_cast<FGeometryCollectionPhysicsProxy*>(Proxy);
+	}
+	return nullptr;
 }
 
-const Chaos::FPBDRigidParticleHandle* Chaos::FStrainedProxyModifier::GetRootHandle() const
+bool Chaos::FStrainedProxyAndRoot::IsPartialDestruction() const
 {
-	return RootHandle;
+	return bPartialDestruction;
+}
+
+bool Chaos::FStrainedProxyAndRoot::operator==(const  Chaos::FStrainedProxyAndRoot& Other) const
+{
+	// no need to test the partial destruction flag 
+	return ((Proxy == Other.Proxy) && (ParticleHandle == Other.ParticleHandle));
+}
+
+
+const IPhysicsProxyBase* Chaos::FStrainedProxyModifier::GetProxy() const
+{
+	return ProxyAndRoot.Proxy;
+}
+
+const Chaos::FPBDRigidParticleHandle* Chaos::FStrainedProxyModifier::GetOriginalRootHandle() const
+{
+	if (FGeometryCollectionPhysicsProxy* ProxyGC = ProxyAndRoot.CastToGeometryCollectionProxy())
+	{
+		return ProxyGC->GetInitialRootHandle_Internal();
+	}
+	return nullptr;
+}
+
+const Chaos::FPBDRigidParticleHandle* Chaos::FStrainedProxyModifier::GetParticleHandle() const
+{
+	return ProxyAndRoot.ParticleHandle;
+}
+
+template <typename TFunction>
+static void ForEachRootChildParticle(const Chaos::FStrainedProxyAndRoot& ProxyAndRoot, const TSet<int32>* RestChildren, TFunction Func)
+{
+	if (ProxyAndRoot.IsPartialDestruction())
+	{
+		// when partial destruction, only process the particle handle
+		Func(ProxyAndRoot.ParticleHandle);
+	}
+	else if (RestChildren)
+	{
+		if (FGeometryCollectionPhysicsProxy* ProxyGC = ProxyAndRoot.CastToGeometryCollectionProxy())
+		{
+			for (int32 RestChildIdx : *RestChildren)
+			{
+				Chaos::FPBDRigidClusteredParticleHandle* ChildHandle = ProxyGC->GetSolverParticleHandles()[RestChildIdx];
+				if (ChildHandle->Parent() == nullptr) { continue; }
+
+				Func(ChildHandle);
+			}
+		}
+	}
 }
 
 int32 Chaos::FStrainedProxyModifier::GetNumRestBreakables() const
@@ -32,36 +86,28 @@ int32 Chaos::FStrainedProxyModifier::GetNumRestBreakables() const
 	{
 		return RestChildren->Num();
 	}
-	return 0;
+	return 1; // cluser union case
 }
 
 int32 Chaos::FStrainedProxyModifier::GetNumBreakingStrains(const bool bDoubleCount, const uint8 StrainTypes) const
 {
-	// Make sure we have a proxy and rest-children
-	if (Proxy == nullptr) { return 0; }
-	if (RestChildren == nullptr) { return 0; }
-
-	// Loop over each child, checking whether or not it will have been freed
-	// by the strain that it has accumulated
 	int32 NumBreakingStrains = 0;
-	for (int32 RestChildIdx : *RestChildren)
-	{
-		Chaos::FPBDRigidClusteredParticleHandle* ChildHandle = Proxy->GetSolverParticleHandles()[RestChildIdx];
-		if (ChildHandle->Parent() == nullptr) { continue; }
-
-		// Get the applied and internal strain
-		const float InternalStrain = ChildHandle->GetInternalStrains();
-		const float MaxAppliedStrain = GetMaxAppliedStrain(ChildHandle, StrainTypes);
-
-		if (bDoubleCount && InternalStrain > SMALL_NUMBER)
-		{
-			NumBreakingStrains += (int32)(MaxAppliedStrain / InternalStrain);
-		}
-		else if (MaxAppliedStrain >= InternalStrain)
+	ForEachRootChildParticle(ProxyAndRoot, RestChildren,
+		[this, &NumBreakingStrains, &bDoubleCount, &StrainTypes] (Chaos::FPBDRigidClusteredParticleHandle* ChildHandle)
 		{
-			++NumBreakingStrains;
-		}
-	}
+			// Get the applied and internal strain
+			const float InternalStrain = ChildHandle->GetInternalStrains();
+			const float MaxAppliedStrain = GetMaxAppliedStrain(ChildHandle, StrainTypes);
+
+			if (bDoubleCount && InternalStrain > SMALL_NUMBER)
+			{
+				NumBreakingStrains += (int32)(MaxAppliedStrain / InternalStrain);
+			}
+			else if (MaxAppliedStrain >= InternalStrain)
+			{
+				++NumBreakingStrains;
+			}
+		});
 
 	// Return the number of breaking strains
 	return NumBreakingStrains;
@@ -69,64 +115,41 @@ int32 Chaos::FStrainedProxyModifier::GetNumBreakingStrains(const bool bDoubleCou
 
 float Chaos::FStrainedProxyModifier::GetMaxBreakStrainRatio(const float FatigueThreshold, const uint8 StrainTypes) const
 {
-	// Make sure we have a proxy and rest-children
-	if (Proxy == nullptr) { return 0; }
-	if (RestChildren == nullptr) { return 0; }
-
-	// Loop over each child, checking whether or not it will have been freed
-	// by the strain that it has accumulated
 	float MaxBreakStrainRatio = 0.f;
-	for (int32 RestChildIdx : *RestChildren)
-	{
-		Chaos::FPBDRigidClusteredParticleHandle* ChildHandle = Proxy->GetSolverParticleHandles()[RestChildIdx];
-		if (ChildHandle->Parent() == nullptr) { continue; }
-
-		// compute the strain ratio
-		const float InternalStrain = ChildHandle->GetInternalStrains();
-		const float MaxAppliedStrain = GetMaxAppliedStrain(ChildHandle, StrainTypes);
-		const float AdjustedAppliedStrain = (MaxAppliedStrain >= FatigueThreshold) ? MaxAppliedStrain : 0.f;
-		const float StrainRatio = (InternalStrain > SMALL_NUMBER)? (MaxAppliedStrain / InternalStrain) : 1.0f;
+	ForEachRootChildParticle(ProxyAndRoot, RestChildren,
+		[this, &MaxBreakStrainRatio, &FatigueThreshold, &StrainTypes] (Chaos::FPBDRigidClusteredParticleHandle* ChildHandle)
+		{
+			// compute the strain ratio
+			const float InternalStrain = ChildHandle->GetInternalStrains();
+			const float MaxAppliedStrain = GetMaxAppliedStrain(ChildHandle, StrainTypes);
+			const float AdjustedAppliedStrain = (MaxAppliedStrain >= FatigueThreshold) ? MaxAppliedStrain : 0.f;
+			const float StrainRatio = (InternalStrain > SMALL_NUMBER) ? (MaxAppliedStrain / InternalStrain) : 1.0f;
 
-		MaxBreakStrainRatio = FMath::Max(MaxBreakStrainRatio, StrainRatio);
-	}
+			MaxBreakStrainRatio = FMath::Max(MaxBreakStrainRatio, StrainRatio);
+		});
 
 	return MaxBreakStrainRatio;
 }
 
 void Chaos::FStrainedProxyModifier::AdjustStrainForBreak(const float FatigueThreshold, const uint8 StrainTypes)
 {
-	// Make sure we have a proxy and rest-children
-	if (Proxy == nullptr) { return; }
-	if (RestChildren == nullptr) { return; }
-
-	// Loop over each child, checking whether or not it will have been freed
-	// by the strain that it has accumulated
-	float MaxBreakStrainRatio = 0.f;
-	for (int32 RestChildIdx : *RestChildren)
-	{
-		Chaos::FPBDRigidClusteredParticleHandle* ChildHandle = Proxy->GetSolverParticleHandles()[RestChildIdx];
-		if (ChildHandle->Parent() == nullptr) { continue; }
-
-		// compute the strain ratio
-		const float InternalStrain = ChildHandle->GetInternalStrains();
-		const float MaxAppliedStrain = GetMaxAppliedStrain(ChildHandle, StrainTypes);
-		if (MaxAppliedStrain >= FatigueThreshold)
+	ForEachRootChildParticle(ProxyAndRoot, RestChildren,
+		[this, &FatigueThreshold, &StrainTypes] (Chaos::FPBDRigidClusteredParticleHandle* ChildHandle)
 		{
-			RigidClustering.SetExternalStrain(ChildHandle, InternalStrain);
-		}
-	}
+			// compute the strain ratio
+			const float InternalStrain = ChildHandle->GetInternalStrains();
+			const float MaxAppliedStrain = GetMaxAppliedStrain(ChildHandle, StrainTypes);
+			if (MaxAppliedStrain >= FatigueThreshold)
+			{
+				RigidClustering.SetExternalStrain(ChildHandle, InternalStrain);
+			}
+		});
 }
 
 void Chaos::FStrainedProxyModifier::ClearStrains()
 {
-	// Make sure we have a proxy and rest-children
-	if (Proxy == nullptr) { return; }
-	if (RestChildren == nullptr) { return; }
-
-	// Loop over each child, clearing the strain associated with it
-	for (int32 RestChildIdx : *RestChildren)
-	{
-		if (Chaos::FPBDRigidClusteredParticleHandle* ChildHandle = Proxy->GetSolverParticleHandles()[RestChildIdx])
+	ForEachRootChildParticle(ProxyAndRoot, RestChildren,
+		[this] (Chaos::FPBDRigidClusteredParticleHandle* ChildHandle)
 		{
 			// Clear accumulated collision impulses
 			ChildHandle->ClearCollisionImpulse();
@@ -137,31 +160,15 @@ void Chaos::FStrainedProxyModifier::ClearStrains()
 			// ChildHandle->ClearExternalStrain();
 			// Instead we use:
 			RigidClustering.SetExternalStrain(ChildHandle, FReal(0));
-		}
-	}
-}
-
-Chaos::FPBDRigidClusteredParticleHandle* Chaos::FStrainedProxyModifier::InitRootHandle(FGeometryCollectionPhysicsProxy* Proxy)
-{
-	if (Proxy == nullptr) { return nullptr; }
-
-	// Get the root index of the proxy
-	FSimulationParameters& Parameters = Proxy->GetSimParameters();
-	const int32 RootIndex = Parameters.InitialRootIndex;
-	if (RootIndex == INDEX_NONE) { return nullptr; }
-
-	// Return the particle handle corresponding to the root index
-	TArray<Chaos::FPBDRigidClusteredParticleHandle*>& ParticleHandles = Proxy->GetSolverParticleHandles();
-	if (!ParticleHandles.IsValidIndex(RootIndex)) { return nullptr; }
-	return ParticleHandles[RootIndex];
+		});
 }
 
-const TSet<int32>* Chaos::FStrainedProxyModifier::InitRestChildren(FGeometryCollectionPhysicsProxy* Proxy)
+const TSet<int32>* Chaos::FStrainedProxyModifier::InitRestChildren(FGeometryCollectionPhysicsProxy* ProxyGC)
 {
-	if (Proxy == nullptr) { return nullptr; }
+	if (ProxyGC == nullptr) { return nullptr; }
 
 	// Get the root index of the proxy
-	FSimulationParameters& Parameters = Proxy->GetSimParameters();
+	FSimulationParameters& Parameters = ProxyGC->GetSimParameters();
 	const int32 RootIndex = Parameters.InitialRootIndex;
 	if (RootIndex == INDEX_NONE) { return nullptr; }
 
@@ -171,7 +178,7 @@ const TSet<int32>* Chaos::FStrainedProxyModifier::InitRestChildren(FGeometryColl
 
 Chaos::FStrainedProxyModifier Chaos::FStrainedProxyIterator::operator*()
 {
-	return Chaos::FStrainedProxyModifier(RigidClustering, Proxies[Index]);
+	return Chaos::FStrainedProxyModifier(RigidClustering, ProxyAndRoots[Index]);
 }
 
 Chaos::FStrainedProxyIterator& Chaos::FStrainedProxyIterator::operator++()
@@ -183,7 +190,7 @@ Chaos::FStrainedProxyIterator& Chaos::FStrainedProxyIterator::operator++()
 bool Chaos::FStrainedProxyIterator::operator==(const FStrainedProxyIterator& Other) const
 {
 	return
-		Proxies == Other.Proxies &&
+		ProxyAndRoots == Other.ProxyAndRoots &&
 		Index == Other.Index;
 }
 
@@ -191,50 +198,32 @@ Chaos::FStrainedProxyRange::FStrainedProxyRange(Chaos::FRigidClustering& InRigid
 	: RigidClustering(InRigidClustering)
 	, StrainedParticles(InStrainedParticles)
 {
-	const TSet<Chaos::FPBDRigidClusteredParticleHandle*>& StrainedParents = RigidClustering.GetTopLevelClusterParentsStrained();
+	if (!StrainedParticles)
+	{
+		return;
+	}
 
-	auto ForEveryParticle = [this, &StrainedParents]<typename TLambda>(TLambda&& Func)
+	ProxyAndRoots.Reserve(StrainedParticles->Num());
+	for (Chaos::FPBDRigidClusteredParticleHandle* Cluster : *StrainedParticles)
 	{
-		if (StrainedParticles)
-		{
-			for (Chaos::FPBDRigidClusteredParticleHandle* Cluster : *StrainedParticles)
-			{
-				Func(Cluster);
-			}
-		}
-		else
+		// Make sure the cluster's physics proxy exists and is the right type
+		IPhysicsProxyBase* Proxy = Cluster->PhysicsProxy();
+		if (Proxy == nullptr)
 		{
-			for (Chaos::FPBDRigidClusteredParticleHandle* Cluster : StrainedParents)
-			{
-				Func(Cluster);
-			}
+			return;
 		}
-	};
 
-	Proxies.Reserve(StrainedParticles ? StrainedParticles->Num() : StrainedParents.Num());
-	ForEveryParticle(
-		[this, bRootLevelOnly](Chaos::FPBDRigidClusteredParticleHandle* Cluster)
+		const bool bIsGeometryCollectionProxy = (Proxy->GetType() == EPhysicsProxyType::GeometryCollectionType);
+		if (bIsGeometryCollectionProxy)
 		{
-			// Make sure the cluster's physics proxy exists and is the right type
-			IPhysicsProxyBase* ProxyBase = Cluster->PhysicsProxy();
-			if (ProxyBase == nullptr)
-			{
-				return;
-			}
-
-			if (ProxyBase->GetType() != EPhysicsProxyType::GeometryCollectionType)
-			{
-				return;
-			}
-
-			FGeometryCollectionPhysicsProxy* Proxy = static_cast<FGeometryCollectionPhysicsProxy*>(ProxyBase);
-
-			// Make sure the rest collection has a root index
+			// Make sure the rest collection has a root index (If a GC)
 			if (bRootLevelOnly)
 			{
-				FSimulationParameters& Parameters = Proxy->GetSimParameters();
+				FGeometryCollectionPhysicsProxy* ProxyGC = static_cast<FGeometryCollectionPhysicsProxy*>(Proxy);
+
+				FSimulationParameters& Parameters = ProxyGC->GetSimParameters();
 				const int32 RootIndex = Parameters.InitialRootIndex;
-				TArray<Chaos::FPBDRigidClusteredParticleHandle*>& ParticleHandles = Proxy->GetSolverParticleHandles();
+				TArray<Chaos::FPBDRigidClusteredParticleHandle*>& ParticleHandles = ProxyGC->GetSolverParticleHandles();
 				if (ParticleHandles.IsValidIndex(RootIndex))
 				{
 					if (Cluster != ParticleHandles[RootIndex])
@@ -242,20 +231,43 @@ Chaos::FStrainedProxyRange::FStrainedProxyRange(Chaos::FRigidClustering& InRigid
 						return;
 					}
 				}
-			}
 
-			// Only need to use AddUnique if we're not checking for root, since at most
-			// one cluster will have the rest collection's root index.
-			if (bRootLevelOnly)
-			{
-				Proxies.Add(Proxy);
+				// Only need to use AddUnique if we're not checking for root, since at most
+				// one cluster will have the rest collection's root index.
+				ProxyAndRoots.Add({ Proxy, Cluster, /*bPartialDestruction*/ false });
 			}
 			else
 			{
-				Proxies.AddUnique(Proxy);
+
+				ProxyAndRoots.AddUnique({ Proxy, Cluster, /*bPartialDestruction*/ false});
+			}
+		}
+
+		// if the cluster union is the one reported we are in a partial destruction scenario
+		const bool bIsClusterUnionProxy = (Proxy->GetType() == EPhysicsProxyType::ClusterUnionProxy);
+		if (bIsClusterUnionProxy)
+		{
+			// for cluster union we go through the children 
+			Chaos::FClusterUnionPhysicsProxy* ProxyUC = static_cast<FClusterUnionPhysicsProxy*>(Proxy);
+			if (const TArray<Chaos::FPBDRigidParticleHandle*>* ClusterUnionChildren = RigidClustering.GetChildrenMap().Find(ProxyUC->GetParticle_Internal()))
+			{
+				for (Chaos::FPBDRigidParticleHandle* ChildHandle : *ClusterUnionChildren)
+				{
+					if (ChildHandle)
+					{
+						if (Chaos::FPBDRigidClusteredParticleHandle* ClusteredChildHandle = ChildHandle->CastToClustered())
+						{
+							if (ClusteredChildHandle->GetExternalStrain() > 0 || ClusteredChildHandle->CollisionImpulse() > 0)
+							{
+								// we actually add tyhe cluster union proxy and the partial destruction child handle
+								ProxyAndRoots.Add({ ClusteredChildHandle->PhysicsProxy(), ClusteredChildHandle, /*bPartialDestruction*/ true});
+							}
+						}
+					}
+				}
 			}
 		}
-	);
+	}
 }
 
 Chaos::FStrainedProxyRange Chaos::FStrainModifierAccessor::GetStrainedProxies(const bool bRootLevelOnly)
diff --git a/Engine/Source/Runtime/Experimental/Chaos/Private/PhysicsProxy/GeometryCollectionPhysicsProxy.cpp b/Engine/Source/Runtime/Experimental/Chaos/Private/PhysicsProxy/GeometryCollectionPhysicsProxy.cpp
--- a/Engine/Source/Runtime/Experimental/Chaos/Private/PhysicsProxy/GeometryCollectionPhysicsProxy.cpp
+++ b/Engine/Source/Runtime/Experimental/Chaos/Private/PhysicsProxy/GeometryCollectionPhysicsProxy.cpp
@@ -4914,6 +4914,15 @@ void FGeometryCollectionPhysicsProxy::FieldForcesUpdateCallback(Chaos::FPBDRigid
 	}
 }
 
+FGeometryCollectionPhysicsProxy::FClusterHandle* FGeometryCollectionPhysicsProxy::GetInitialRootHandle_Internal() const
+{
+	if (SolverParticleHandles.IsValidIndex(Parameters.InitialRootIndex))
+	{
+		return SolverParticleHandles[Parameters.InitialRootIndex];
+	}
+	return nullptr;
+}
+
 TArray<Chaos::FPhysicsObjectHandle> FGeometryCollectionPhysicsProxy::GetAllPhysicsObjects() const
 {
 	TArray<Chaos::FPhysicsObjectHandle> Handles;
diff --git a/Engine/Source/Runtime/Experimental/Chaos/Public/Chaos/StrainModification.h b/Engine/Source/Runtime/Experimental/Chaos/Public/Chaos/StrainModification.h
--- a/Engine/Source/Runtime/Experimental/Chaos/Public/Chaos/StrainModification.h
+++ b/Engine/Source/Runtime/Experimental/Chaos/Public/Chaos/StrainModification.h
@@ -5,11 +5,13 @@
 #include "Chaos/Core.h"
 #include "ParticleHandleFwd.h"
 
+class IPhysicsProxyBase;
 class FGeometryCollectionPhysicsProxy;
 
 namespace Chaos
 {
 	class FRigidClustering;
+	class FClusterUnionPhysicsProxy;
 	class FStrainModifierAccessor;
 	class FStrainedProxyModifier;
 	class FStrainedProxyIterator;
@@ -23,6 +25,23 @@ namespace Chaos
 		CollisionStrain = 1 << 2
 	};
 
+	struct FStrainedProxyAndRoot
+	{
+		IPhysicsProxyBase* Proxy = nullptr;
+
+		// Can be a root if proxy is a GC or the directly strained particle if the proxy is a cluster union
+		Chaos::FPBDRigidClusteredParticleHandle* ParticleHandle = nullptr;
+
+		// true if the particle handle is actually a partial GC attached toa cluster union
+		bool bPartialDestruction = false;
+
+		FGeometryCollectionPhysicsProxy* CastToGeometryCollectionProxy() const;
+		
+		bool IsPartialDestruction() const;
+
+		bool operator==(const  Chaos::FStrainedProxyAndRoot& Other) const;
+	};
+
 	// FStrainedProxyModifier
 	//
 	// User-facing api for accessing the proxy of a strained cluster. Provides const access to the
@@ -30,24 +49,26 @@ namespace Chaos
 	class FStrainedProxyModifier
 	{
 	public:
-		FStrainedProxyModifier(FRigidClustering& InRigidClustering, FGeometryCollectionPhysicsProxy* InProxy)
+		FStrainedProxyModifier(FRigidClustering& InRigidClustering, FStrainedProxyAndRoot InProxyAndRoot)
 			: RigidClustering(InRigidClustering)
-			, Proxy(InProxy)
-			, RootHandle(InitRootHandle(InProxy))
-			, RestChildren(InitRestChildren(InProxy))
+			, ProxyAndRoot(InProxyAndRoot)
+			, RestChildren(InitRestChildren(InProxyAndRoot.CastToGeometryCollectionProxy()))
 		{ }
 
 		FStrainedProxyModifier(const FStrainedProxyModifier& Other)
 			: RigidClustering(Other.RigidClustering)
-			, Proxy(Other.Proxy)
+			, ProxyAndRoot(Other.ProxyAndRoot)
 			, RestChildren(Other.RestChildren)
 		{ }
 
 		// Get the proxy that owns the strained cluster or clusters
-		CHAOS_API const FGeometryCollectionPhysicsProxy* GetProxy() const;
+		CHAOS_API const IPhysicsProxyBase* GetProxy() const;
 
-		// Get the physics handle for the strained parent cluster
-		CHAOS_API const Chaos::FPBDRigidParticleHandle* GetRootHandle() const;
+		// Get the original root handle from the geometry collection proxy
+		CHAOS_API const Chaos::FPBDRigidParticleHandle* GetOriginalRootHandle() const;
+
+		// Get the physics handle for the strained parent cluster or the strained particle directly if it's part of partial destruction
+		CHAOS_API const Chaos::FPBDRigidParticleHandle* GetParticleHandle() const;
 
 		// Get the number of level-1 strainable entities (number of rest-children in the per-particle
 		// strain model, or number of rest-connections in the edge/area model).
@@ -73,13 +94,10 @@ namespace Chaos
 
 	private:
 
-		static CHAOS_API Chaos::FPBDRigidClusteredParticleHandle* InitRootHandle(FGeometryCollectionPhysicsProxy* Proxy);
-
 		static CHAOS_API const TSet<int32>* InitRestChildren(FGeometryCollectionPhysicsProxy* Proxy);
 
 		FRigidClustering& RigidClustering;
-		FGeometryCollectionPhysicsProxy* Proxy;
-		Chaos::FPBDRigidClusteredParticleHandle* RootHandle;
+		FStrainedProxyAndRoot ProxyAndRoot;
 		const TSet<int32>* RestChildren;
 	};
 
@@ -90,15 +108,15 @@ namespace Chaos
 	class FStrainedProxyIterator
 	{
 	public:
-		FStrainedProxyIterator(FRigidClustering& InRigidClustering, TArray<FGeometryCollectionPhysicsProxy*>& InProxies, int32 InIndex)
+		FStrainedProxyIterator(FRigidClustering& InRigidClustering, TArray<FStrainedProxyAndRoot>& InProxyAndRoots, int32 InIndex)
 			: RigidClustering(InRigidClustering)
-			, Proxies(InProxies)
+			, ProxyAndRoots(InProxyAndRoots)
 			, Index(InIndex)
 		{ }
 
 		FStrainedProxyIterator(const FStrainedProxyIterator& Other)
 			: RigidClustering(Other.RigidClustering)
-			, Proxies(Other.Proxies)
+			, ProxyAndRoots(Other.ProxyAndRoots)
 			, Index(Other.Index)
 		{ }
 
@@ -112,7 +130,7 @@ namespace Chaos
 
 	private:
 		FRigidClustering& RigidClustering;
-		TArray<FGeometryCollectionPhysicsProxy*>& Proxies;
+		TArray<FStrainedProxyAndRoot>& ProxyAndRoots;
 		int32 Index;
 	};
 
@@ -129,22 +147,22 @@ namespace Chaos
 
 		FStrainedProxyRange(const FStrainedProxyRange& Other)
 			: RigidClustering(Other.RigidClustering)
-			, Proxies(Other.Proxies)
+			, ProxyAndRoots(Other.ProxyAndRoots)
 		{ }
 
 		FStrainedProxyIterator begin()
 		{
-			return FStrainedProxyIterator(RigidClustering, Proxies, 0);
+			return FStrainedProxyIterator(RigidClustering, ProxyAndRoots, 0);
 		}
 
 		FStrainedProxyIterator end()
 		{
-			return FStrainedProxyIterator(RigidClustering, Proxies, Proxies.Num());
+			return FStrainedProxyIterator(RigidClustering, ProxyAndRoots, ProxyAndRoots.Num());
 		}
 
 	private:
 		FRigidClustering& RigidClustering;
-		TArray<FGeometryCollectionPhysicsProxy*> Proxies;
+		TArray<FStrainedProxyAndRoot> ProxyAndRoots;
 		const TArray<FPBDRigidClusteredParticleHandle*>* StrainedParticles;
 	};
 
diff --git a/Engine/Source/Runtime/Experimental/Chaos/Public/PhysicsProxy/GeometryCollectionPhysicsProxy.h b/Engine/Source/Runtime/Experimental/Chaos/Public/PhysicsProxy/GeometryCollectionPhysicsProxy.h
--- a/Engine/Source/Runtime/Experimental/Chaos/Public/PhysicsProxy/GeometryCollectionPhysicsProxy.h
+++ b/Engine/Source/Runtime/Experimental/Chaos/Public/PhysicsProxy/GeometryCollectionPhysicsProxy.h
@@ -469,6 +469,8 @@ public:
 		PostPhysicsSyncCallback = Callback;
 	}
 	
+	CHAOS_API FClusterHandle* GetInitialRootHandle_Internal() const;
+
 	CHAOS_API TArray<Chaos::FPhysicsObjectHandle> GetAllPhysicsObjects() const ;
 	CHAOS_API TArray<Chaos::FPhysicsObjectHandle> GetAllPhysicsObjectIncludingNulls() const;
 	CHAOS_API Chaos::FPhysicsObjectHandle GetPhysicsObjectByIndex(int32 Index) const;
