Fix issue where fatigue threshold was not accounted for properly for decrementing health

"#preflight none

diff --git a/Engine/Source/Runtime/Experimental/Chaos/Private/Chaos/StrainModification.cpp b/Engine/Source/Runtime/Experimental/Chaos/Private/Chaos/StrainModification.cpp
--- a/Engine/Source/Runtime/Experimental/Chaos/Private/Chaos/StrainModification.cpp
+++ b/Engine/Source/Runtime/Experimental/Chaos/Private/Chaos/StrainModification.cpp
@@ -122,10 +122,11 @@ float Chaos::FStrainedProxyModifier::GetMaxBreakStrainRatio(const float FatigueT
 			// compute the strain ratio
 			const float InternalStrain = ChildHandle->GetInternalStrains();
 			const float MaxAppliedStrain = GetMaxAppliedStrain(ChildHandle, StrainTypes);
-			const float AdjustedAppliedStrain = (MaxAppliedStrain >= FatigueThreshold) ? MaxAppliedStrain : 0.f;
-			const float StrainRatio = (InternalStrain > SMALL_NUMBER) ? (MaxAppliedStrain / InternalStrain) : 1.0f;
-
-			MaxBreakStrainRatio = FMath::Max(MaxBreakStrainRatio, StrainRatio);
+			if (MaxAppliedStrain >= FatigueThreshold)
+			{
+				const float StrainRatio = (InternalStrain > SMALL_NUMBER) ? (MaxAppliedStrain / InternalStrain) : 1.0f;
+				MaxBreakStrainRatio = FMath::Max(MaxBreakStrainRatio, StrainRatio);
+			}
 		});
 
 	return MaxBreakStrainRatio;
