fix various sources of incorrectness relating to intercluster edges that could cause incorrect behaviors when adding to/removing from/splitting cluster unions

- Minor Perf Improvement: Connectivity edges for cluster union children are only computed on the server and not on the client
- Intercluster edges are now generated by the cluster union (controlled by a cvar p.Chaos.ClusterUnion.GenerateInterclusterEdges) instead of relying on an external GT process to trigger that
- IsInterclusterEdge is now more accurate by checking the parent instead of the proxy. Checking the proxy would miss the case where a GC gets split and one half gets attached to a cluster union. It wouldn't identify the edges between the two halves of the GC as being an intercluster edge.
- Removed some cases where the client would try to add/remove from the cluster union by itself.

diff --git a/Engine/Source/Runtime/Engine/Private/PhysicsEngine/ClusterUnionComponent.cpp b/Engine/Source/Runtime/Engine/Private/PhysicsEngine/ClusterUnionComponent.cpp
--- a/Engine/Source/Runtime/Engine/Private/PhysicsEngine/ClusterUnionComponent.cpp
+++ b/Engine/Source/Runtime/Engine/Private/PhysicsEngine/ClusterUnionComponent.cpp
@@ -588,6 +588,7 @@ void UClusterUnionComponent::OnCreatePhysicsState()
 
 	// Only need to check connectivity on the server and have the client rely on replication to get the memo on when to release from cluster union.
 	InitData.bCheckConnectivity = bHasAuthority;
+	InitData.bGenerateConnectivityEdges = bHasAuthority;
 
 	PhysicsProxy = new Chaos::FClusterUnionPhysicsProxy{ this, Parameters, InitData };
 	PhysicsProxy->Initialize_External();
diff --git a/Engine/Source/Runtime/Experimental/Chaos/Private/Chaos/ClusterUnionManager.cpp b/Engine/Source/Runtime/Experimental/Chaos/Private/Chaos/ClusterUnionManager.cpp
--- a/Engine/Source/Runtime/Experimental/Chaos/Private/Chaos/ClusterUnionManager.cpp
+++ b/Engine/Source/Runtime/Experimental/Chaos/Private/Chaos/ClusterUnionManager.cpp
@@ -16,6 +16,13 @@ namespace Chaos
 {
 	namespace
 	{
+		bool bChaosClusterUnionGenerateInterclusterEdges = true;
+		FAutoConsoleVariableRef CVarChaosClusterUnionGenerateInterclusterEdges(
+			TEXT("p.Chaos.ClusterUnion.GenerateInterclusterEdges"),
+			bChaosClusterUnionGenerateInterclusterEdges,
+			TEXT("Whether to generate intercluster edges automatically when adding to a cluster union (and remove them when removing from the cluster union).")
+		);
+
 		FRigidTransform3 GetParticleRigidFrameInClusterUnion(FPBDRigidParticleHandle* Child, const FRigidTransform3& ClusterWorldTM)
 		{
 			FRigidTransform3 Frame = FRigidTransform3::Identity;
@@ -138,21 +145,100 @@ namespace Chaos
 						}
 						if (bOverlap)
 						{
-							Clustering.CreateNodeConnection(Particle, ClusterUnion.ChildParticles[RootObjectIndex]);
+							FPBDRigidParticleHandle* OtherParticle = ClusterUnion.ChildParticles[RootObjectIndex];
+							Clustering.CreateNodeConnection(Particle, OtherParticle);
+
+							if (bChaosClusterUnionGenerateInterclusterEdges && Particle && OtherParticle)
+							{
+								// Only generate intercluster edges for main particles. Auxiliary particles that are just bits and pieces of geometry collections
+								// shouldn't also generate intercluster edges.
+								if (FClusterUnionParticleProperties* Properties = ClusterUnion.ChildProperties.Find(Particle); !Properties || !Properties->bIsAuxiliaryParticle)
+								{
+									const TArray<FPBDRigidParticleHandle*>& ParticleChildren = Clustering.GetChildrenMap().FindRef(Particle->CastToClustered());
+									const TArray<FPBDRigidParticleHandle*>& OtherChildren = Clustering.GetChildrenMap().FindRef(OtherParticle->CastToClustered());
+
+									TSet<FPBDRigidParticleHandle*> FromSet{ ParticleChildren };
+									TSet<FPBDRigidParticleHandle*> ToSet{ OtherChildren };
+
+									TArray<FPBDRigidParticleHandle*> AllParticles;
+									AllParticles.Reserve(FromSet.Num() + ToSet.Num());
+									AllParticles.Append(ParticleChildren);
+									AllParticles.Append(OtherChildren);
+
+									FClusterCreationParameters Parameters{ 0.3f, 100, false, false };
+									Parameters.ConnectionMethod = FClusterCreationParameters::EConnectionMethod::BoundsOverlapFilteredDelaunayTriangulation;
+									Parameters.ConnectionGraphBoundsFilteringMargin = 1.0;
+
+									Clustering.GenerateConnectionGraph(AllParticles, Parameters, &FromSet, &ToSet);
+								}
+							}
 						}
 					}
 				);
 			}
 		}
 
+		void RemoveClusterUnionEdges(FRigidClustering& Clustering, FPBDRigidParticleHandle* ParticleHandle)
+		{
+			if (!ParticleHandle)
+			{
+				return;
+			}
+
+			Clustering.RemoveNodeConnections(ParticleHandle);
+
+			if (bChaosClusterUnionGenerateInterclusterEdges)
+			{
+				if (FPBDRigidClusteredParticleHandle* ClusterParticle = ParticleHandle->CastToClustered())
+				{
+					if (TArray<FPBDRigidParticleHandle*>* AllChildren = Clustering.GetChildrenMap().Find(ClusterParticle))
+					{
+						for (FPBDRigidParticleHandle* Child : *AllChildren)
+						{
+							if (!Child)
+							{
+								continue;
+							}
+
+							if (FPBDRigidClusteredParticleHandle* ClusterChild = Child->CastToClustered())
+							{
+								Clustering.RemoveFilteredNodeConnections(
+									ClusterChild,
+								    [Child](const FConnectivityEdge& Edge)
+									{
+										return IsInterclusterEdge(*Child, Edge);
+									}
+								);
+							}
+						}
+					}
+				}
+			}
+		}
+
 		void GenerateConnectionGraph(FRigidClustering& Clustering, FClusterUnion& ClusterUnion)
 		{
-			Clustering.ClearConnectionGraph(ClusterUnion.InternalCluster);
 			for (FPBDRigidParticleHandle* ChildParticle : ClusterUnion.ChildParticles)
 			{
+				RemoveClusterUnionEdges(Clustering, ChildParticle);
 				AddParticleToConnectionGraph(Clustering, ClusterUnion, ChildParticle);
 			}
 		}
+
+		FPBDRigidClusteredParticleHandle* GetParentParticleInClusterUnion(const FClusterUnion& ClusterUnion, FPBDRigidClusteredParticleHandle* Particle)
+		{
+			if (!Particle)
+			{
+				return nullptr;
+			}
+
+			if (Particle->Parent() == ClusterUnion.InternalCluster)
+			{
+				return Particle;
+			}
+
+			return GetParentParticleInClusterUnion(ClusterUnion, Particle->Parent());
+		}
 	}
 
 	FClusterUnionManager::FClusterUnionManager(FRigidClustering& InClustering, FPBDRigidsEvolutionGBF& InEvolution)
@@ -493,6 +579,8 @@ namespace Chaos
 				const FRigidTransform3 Frame = GetParticleRigidFrameInClusterUnion(Particle, ClusterWorldTM);
 				ClusterParticle->SetChildToParent(Frame);
 			}
+
+			MEvolution.GetParticles().MarkTransientDirtyParticle(Particle);
 		}
 
 		// For all the particles that have been added to the cluster we need to set their parent proxy to the
@@ -621,6 +709,8 @@ namespace Chaos
 				}
 
 				Cluster->PendingConnectivityOperations.Add({ Handle, EClusterUnionConnectivityOperation::Remove });
+
+				MEvolution.GetParticles().MarkTransientDirtyParticle(Handle);
 			}
 		}
 
@@ -722,14 +812,17 @@ namespace Chaos
 		UpdateShapesDatas(ClusterUnion.ChildParticles, ClusterUnion.InternalCluster->ShapesArray(),
 			ClusterUnion.ClusterUnionParameters.ActorId, ClusterUnion.ClusterUnionParameters.ComponentId, NumSimpleShapes);
 
-		if(EnumHasAnyFlags(Flags, EUpdateClusterUnionPropertiesFlags::ForceGenerateConnectionGraph))
+		if (ClusterUnion.bGenerateConnectivityEdges)
 		{
-			GenerateConnectionGraph(MClustering, ClusterUnion);
-			ClusterUnion.PendingConnectivityOperations.Empty();
-		}
-		else if (EnumHasAnyFlags(Flags, EUpdateClusterUnionPropertiesFlags::IncrementalGenerateConnectionGraph))
-		{
-			FlushIncrementalConnectivityGraphOperations(ClusterUnion);
+			if(EnumHasAnyFlags(Flags, EUpdateClusterUnionPropertiesFlags::ForceGenerateConnectionGraph))
+			{
+				GenerateConnectionGraph(MClustering, ClusterUnion);
+				ClusterUnion.PendingConnectivityOperations.Empty();
+			}
+			else if (EnumHasAnyFlags(Flags, EUpdateClusterUnionPropertiesFlags::IncrementalGenerateConnectionGraph))
+			{
+				FlushIncrementalConnectivityGraphOperations(ClusterUnion);
+			}
 		}
 
 		for (FPBDRigidClusteredParticleHandle* ChildParticle : PendingParticlesToUndoChildToParentLock)
@@ -755,7 +848,7 @@ namespace Chaos
 			}
 			else if (Op.Value == EClusterUnionConnectivityOperation::Remove)
 			{
-				MClustering.RemoveNodeConnections(Op.Key);
+				RemoveClusterUnionEdges(MClustering, Op.Key);
 			}
 		}
 
@@ -933,7 +1026,7 @@ namespace Chaos
 				return !Props->bIsAuxiliaryParticle;
 			}
 
-			return true;
+			return false;
 		};
 
 		if (FPBDRigidClusteredParticleHandle* ClusterParticle = Particle->CastToClustered())
@@ -945,9 +1038,10 @@ namespace Chaos
 
 			for (const FConnectivityEdge& Edge : ClusterParticle->ConnectivityEdges())
 			{
-				if (IsInterclusterEdge(*ClusterParticle, Edge))
+				if (Edge.Sibling && IsInterclusterEdge(*ClusterParticle, Edge))
 				{
-					if (Edge.Sibling && IsParticleMainParticle(Edge.Sibling->CastToClustered()))
+					// We don't want to check the sibling, we want to check the parent of the sibling that's actually in the cluster union
+					if (IsParticleMainParticle(GetParentParticleInClusterUnion(ClusterUnion, Edge.Sibling->CastToClustered())))
 					{
 						return true;
 					}
diff --git a/Engine/Source/Runtime/Experimental/Chaos/Private/Chaos/PBDRigidClustering.cpp b/Engine/Source/Runtime/Experimental/Chaos/Private/Chaos/PBDRigidClustering.cpp
--- a/Engine/Source/Runtime/Experimental/Chaos/Private/Chaos/PBDRigidClustering.cpp
+++ b/Engine/Source/Runtime/Experimental/Chaos/Private/Chaos/PBDRigidClustering.cpp
@@ -715,6 +715,10 @@ namespace Chaos
 		if (ensure(Child != nullptr && ClusteredParent != nullptr))
 		{
 			FPBDRigidClusteredParticleHandle* ClusteredChild = Child->CastToClustered();
+			if (!ClusteredChild || ClusteredChild->Parent() != ClusteredParent)
+			{
+				return;
+			}
 
 			MEvolution.EnableParticle(Child);
 			TopLevelClusterParents.Add(ClusteredChild);
@@ -1208,9 +1212,12 @@ namespace Chaos
 										{
 											if (FClusterUnion* ClusterUnion = ClusterUnionManager.FindClusterUnionFromParticle(GetActiveParentParticle(Edge.Sibling)))
 											{
-												AttachedClusterUnion = ClusterUnion;
+												if (ClusterUnion->bCheckConnectivity)
+												{
+													AttachedClusterUnion = ClusterUnion;
+													break;
+												}
 											}
-											break;
 										}
 									}
 								}
@@ -2748,19 +2755,7 @@ namespace Chaos
 		FPBDRigidClusteredParticleHandle* ClusteredChild)
 	{
 		SCOPE_CYCLE_COUNTER(STAT_RemoveNodeConnections);
-		check(ClusteredChild);
-		TArray<TConnectivityEdge<FReal>>& Edges = ClusteredChild->ConnectivityEdges();
-		for (TConnectivityEdge<FReal>& Edge : Edges)
-		{
-			TArray<TConnectivityEdge<FReal>>& OtherEdges = Edge.Sibling->CastToClustered()->ConnectivityEdges();
-			const int32 Idx = OtherEdges.IndexOfByKey(ClusteredChild);
-			if (Idx != INDEX_NONE)
-				OtherEdges.RemoveAtSwap(Idx);
-			// Make sure there are no duplicates!
-			check(OtherEdges.IndexOfByKey(ClusteredChild) == INDEX_NONE);
-		}
-		Edges.SetNum(0);
+		RemoveFilteredNodeConnections(ClusteredChild, true);
 	}
 
-
 } // namespace Chaos
diff --git a/Engine/Source/Runtime/Experimental/Chaos/Private/PhysicsProxy/ClusterUnionPhysicsProxy.cpp b/Engine/Source/Runtime/Experimental/Chaos/Private/PhysicsProxy/ClusterUnionPhysicsProxy.cpp
--- a/Engine/Source/Runtime/Experimental/Chaos/Private/PhysicsProxy/ClusterUnionPhysicsProxy.cpp
+++ b/Engine/Source/Runtime/Experimental/Chaos/Private/PhysicsProxy/ClusterUnionPhysicsProxy.cpp
@@ -143,6 +143,7 @@ namespace Chaos
 			// be a result of physical simulation.
 			ClusterUnion->bNeedsXRInitialization = false;
 			ClusterUnion->bCheckConnectivity = InitData.bCheckConnectivity;
+			ClusterUnion->bGenerateConnectivityEdges = InitData.bGenerateConnectivityEdges;
 		}
 	}
 
diff --git a/Engine/Source/Runtime/Experimental/Chaos/Public/Chaos/ClusterUnionManager.h b/Engine/Source/Runtime/Experimental/Chaos/Public/Chaos/ClusterUnionManager.h
--- a/Engine/Source/Runtime/Experimental/Chaos/Public/Chaos/ClusterUnionManager.h
+++ b/Engine/Source/Runtime/Experimental/Chaos/Public/Chaos/ClusterUnionManager.h
@@ -99,6 +99,9 @@ namespace Chaos
 		// Whether or not we need to check for connectivity for this cluster union.
 		bool bCheckConnectivity = false;
 
+		// Whether or not to generate connectivity edges (used to differentiate between server graphs vs client graphs).
+		bool bGenerateConnectivityEdges = true;
+
 		// Pending particles that need to be added into the connectivity graph.
 		TArray<TPair<FPBDRigidParticleHandle*, EClusterUnionConnectivityOperation>> PendingConnectivityOperations;
 	};
diff --git a/Engine/Source/Runtime/Experimental/Chaos/Public/Chaos/PBDRigidClusteredParticles.h b/Engine/Source/Runtime/Experimental/Chaos/Public/Chaos/PBDRigidClusteredParticles.h
--- a/Engine/Source/Runtime/Experimental/Chaos/Public/Chaos/PBDRigidClusteredParticles.h
+++ b/Engine/Source/Runtime/Experimental/Chaos/Public/Chaos/PBDRigidClusteredParticles.h
@@ -78,7 +78,16 @@ bool IsInterclusterEdge(const TPBDRigidParticleHandle<T, 3>& Particle, const TCo
 	{
 		return false;
 	}
-	return Particle.PhysicsProxy() != Edge.Sibling->PhysicsProxy();
+
+	const TPBDRigidClusteredParticleHandle<T, 3>* ClusterParticle = Particle.CastToClustered();
+	const TPBDRigidClusteredParticleHandle<T, 3>* SiblingParticle = Edge.Sibling->CastToClustered();
+
+	if (!ClusterParticle || !SiblingParticle)
+	{
+		return false;
+	}
+
+	return ClusterParticle->Parent() != SiblingParticle->Parent();
 }
 
 class FRigidClusteredFlags
diff --git a/Engine/Source/Runtime/Experimental/Chaos/Public/Chaos/PBDRigidClustering.h b/Engine/Source/Runtime/Experimental/Chaos/Public/Chaos/PBDRigidClustering.h
--- a/Engine/Source/Runtime/Experimental/Chaos/Public/Chaos/PBDRigidClustering.h
+++ b/Engine/Source/Runtime/Experimental/Chaos/Public/Chaos/PBDRigidClustering.h
@@ -407,6 +407,49 @@ public:
 	CHAOS_API void RemoveNodeConnections(FPBDRigidParticleHandle* Child);
 	CHAOS_API void RemoveNodeConnections(FPBDRigidClusteredParticleHandle* Child);
 
+	template<typename TFilter>
+	void RemoveFilteredNodeConnections(FPBDRigidClusteredParticleHandle* ClusteredChild, TFilter&& Filter)
+	{
+		check(ClusteredChild != nullptr);
+
+		constexpr bool bHasFilter = std::is_invocable_r_v < bool, TFilter, const TConnectivityEdge<FReal>&>;
+		TArray<TConnectivityEdge<FReal>>& Edges = ClusteredChild->ConnectivityEdges();
+		for (int32 EdgeIndex = Edges.Num() - 1; EdgeIndex >= 0; --EdgeIndex)
+		{
+			const TConnectivityEdge<FReal>& Edge = Edges[EdgeIndex];
+			FPBDRigidParticleHandle* Sibling = Edge.Sibling;
+			if constexpr (bHasFilter) 
+			{
+				if (!Filter(Edge))
+				{
+					continue;
+				}
+
+				Edges.RemoveAtSwap(EdgeIndex, 1, false);
+			}
+
+			check(Sibling != nullptr);
+			TArray<TConnectivityEdge<FReal>>& OtherEdges = Sibling->CastToClustered()->ConnectivityEdges();
+			const int32 Idx = OtherEdges.IndexOfByKey(ClusteredChild);
+			if (Idx != INDEX_NONE)
+			{
+				OtherEdges.RemoveAtSwap(Idx);
+			}
+
+			// Make sure there are no duplicates!
+			check(OtherEdges.IndexOfByKey(ClusteredChild) == INDEX_NONE);
+		}
+
+		if constexpr (!bHasFilter)
+		{
+			Edges.SetNum(0);
+		}
+		else
+		{
+			Edges.Shrink();
+		}
+	}
+
 	template<typename ParticleHandleTypeA, typename ParticleHandleTypeB>
 	void CreateNodeConnection(ParticleHandleTypeA* A, ParticleHandleTypeB* B)
 	{
diff --git a/Engine/Source/Runtime/Experimental/Chaos/Public/PhysicsProxy/ClusterUnionPhysicsProxy.h b/Engine/Source/Runtime/Experimental/Chaos/Public/PhysicsProxy/ClusterUnionPhysicsProxy.h
--- a/Engine/Source/Runtime/Experimental/Chaos/Public/PhysicsProxy/ClusterUnionPhysicsProxy.h
+++ b/Engine/Source/Runtime/Experimental/Chaos/Public/PhysicsProxy/ClusterUnionPhysicsProxy.h
@@ -33,6 +33,7 @@ namespace Chaos
 		uint32 ComponentId = INDEX_NONE;
 		bool bCheckConnectivity = true;
 		bool bUnbreakable = false;
+		bool bGenerateConnectivityEdges = true;
 	};
 
 	/**
