Fix compilation problem happening in other branch because of changed API

diff --git a/Engine/Source/Runtime/Experimental/Chaos/Private/Chaos/StrainModification.cpp b/Engine/Source/Runtime/Experimental/Chaos/Private/Chaos/StrainModification.cpp
--- a/Engine/Source/Runtime/Experimental/Chaos/Private/Chaos/StrainModification.cpp
+++ b/Engine/Source/Runtime/Experimental/Chaos/Private/Chaos/StrainModification.cpp
@@ -217,11 +217,11 @@ Chaos::FStrainedProxyRange::FStrainedProxyRange(Chaos::FRigidClustering& InRigid
 		const bool bIsGeometryCollectionProxy = (Proxy->GetType() == EPhysicsProxyType::GeometryCollectionType);
 		if (bIsGeometryCollectionProxy)
 		{
+			FGeometryCollectionPhysicsProxy* ProxyGC = static_cast<FGeometryCollectionPhysicsProxy*>(Proxy);
+
 			// Make sure the rest collection has a root index (If a GC)
 			if (bRootLevelOnly)
 			{
-				FGeometryCollectionPhysicsProxy* ProxyGC = static_cast<FGeometryCollectionPhysicsProxy*>(Proxy);
-
 				FSimulationParameters& Parameters = ProxyGC->GetSimParameters();
 				const int32 RootIndex = Parameters.InitialRootIndex;
 				TArray<Chaos::FPBDRigidClusteredParticleHandle*>& ParticleHandles = ProxyGC->GetSolverParticleHandles();
@@ -235,12 +235,12 @@ Chaos::FStrainedProxyRange::FStrainedProxyRange(Chaos::FRigidClustering& InRigid
 
 				// Only need to use AddUnique if we're not checking for root, since at most
 				// one cluster will have the rest collection's root index.
-				ProxyAndRoots.Add({ Proxy, Cluster, /*bPartialDestruction*/ false });
+				ProxyAndRoots.Add({ ProxyGC, Cluster, /*bPartialDestruction*/ false });
 			}
 			else
 			{
 
-				ProxyAndRoots.AddUnique({ Proxy, Cluster, /*bPartialDestruction*/ false});
+				ProxyAndRoots.AddUnique({ ProxyGC, Cluster, /*bPartialDestruction*/ false});
 			}
 		}
 
@@ -260,8 +260,13 @@ Chaos::FStrainedProxyRange::FStrainedProxyRange(Chaos::FRigidClustering& InRigid
 						{
 							if (ClusteredChildHandle->GetExternalStrain() > 0 || ClusteredChildHandle->CollisionImpulse() > 0)
 							{
-								// we actually add tyhe cluster union proxy and the partial destruction child handle
-								ProxyAndRoots.Add({ ClusteredChildHandle->PhysicsProxy(), ClusteredChildHandle, /*bPartialDestruction*/ true});
+								IPhysicsProxyBase* ChildProxy = ClusteredChildHandle->PhysicsProxy();
+								if (ChildProxy && ChildProxy->GetType() == EPhysicsProxyType::GeometryCollectionType)
+								{
+									// we actually add tyhe cluster union proxy and the partial destruction child handle
+									FGeometryCollectionPhysicsProxy* ChildProxyGC = static_cast<FGeometryCollectionPhysicsProxy*>(ChildProxy);
+									ProxyAndRoots.Add({ ChildProxyGC, ClusteredChildHandle, /*bPartialDestruction*/ true });
+								}
 							}
 						}
 					}
diff --git a/Engine/Source/Runtime/Experimental/Chaos/Public/Chaos/StrainModification.h b/Engine/Source/Runtime/Experimental/Chaos/Public/Chaos/StrainModification.h
--- a/Engine/Source/Runtime/Experimental/Chaos/Public/Chaos/StrainModification.h
+++ b/Engine/Source/Runtime/Experimental/Chaos/Public/Chaos/StrainModification.h
@@ -27,7 +27,7 @@ namespace Chaos
 
 	struct FStrainedProxyAndRoot
 	{
-		IPhysicsProxyBase* Proxy = nullptr;
+		FGeometryCollectionPhysicsProxy* Proxy = nullptr;
 
 		// Can be a root if proxy is a GC or the directly strained particle if the proxy is a cluster union
 		Chaos::FPBDRigidClusteredParticleHandle* ParticleHandle = nullptr;
@@ -70,6 +70,10 @@ namespace Chaos
 		// Get the physics handle for the strained parent cluster or the strained particle directly if it's part of partial destruction
 		CHAOS_API const Chaos::FPBDRigidParticleHandle* GetParticleHandle() const;
 
+		// Get the physics handle for the strained parent cluster
+		UE_DEPRECATED(5.4, "This has been replaced by GetParticleHandle and GetOriginalRootHandle for finer grain access to the underlying handles")
+		CHAOS_API const Chaos::FPBDRigidParticleHandle* GetRootHandle() const { return GetParticleHandle(); }
+
 		// Get the number of level-1 strainable entities (number of rest-children in the per-particle
 		// strain model, or number of rest-connections in the edge/area model).
 		CHAOS_API int32 GetNumRestBreakables() const;
