Fix for the Ias latency tests crashing on startup

#rnx

- When making a HEAD request we do not expect content to be downloaded so it makes sense for the target buffer (Internal.Dest) to be null.
- However when this was the case we would crash trying to evaluate bStreamed as it accessed the buffer.
- Moved the evaluation until after the early out for no content requests.

diff --git a/Engine/Source/Runtime/Experimental/IoStoreOnDemand/Private/CoreHttp/Http.cpp b/Engine/Source/Runtime/Experimental/IoStoreOnDemand/Private/CoreHttp/Http.cpp
--- a/Engine/Source/Runtime/Experimental/IoStoreOnDemand/Private/CoreHttp/Http.cpp
+++ b/Engine/Source/Runtime/Experimental/IoStoreOnDemand/Private/CoreHttp/Http.cpp
@@ -1672,14 +1672,12 @@ static int32 DoRecvMessage(FActivity* Activity)
 		}
 	}
 
-	bool Streamed = (Internal.Dest->GetSize() < ContentLength);
-
 	// Perhaps we have some of the content already?
 	const char* BufferRight = Buffer.GetData() + Buffer.GetSize();
 	uint32 AlreadyReceived = uint32(ptrdiff_t(BufferRight - MessageRight));
 	if (AlreadyReceived > uint32(ContentLength))
 	{
-		Activity_SetError(Activity, "More data recevied that expected");
+		Activity_SetError(Activity, "More data received that expected");
 		return -1;
 	}
 
@@ -1694,7 +1692,11 @@ static int32 DoRecvMessage(FActivity* Activity)
 		return 0;
 	}
 
-	Activity->State = Streamed ? FActivity::EState::RecvStream : FActivity::EState::RecvContent;
+	check(Internal.Dest != nullptr);
+
+	const bool bStreamed = Internal.Dest->GetSize() < ContentLength;
+
+	Activity->State = bStreamed ? FActivity::EState::RecvStream : FActivity::EState::RecvContent;
 	Activity->StateParam = AlreadyReceived;
 
 	if (AlreadyReceived == 0)
@@ -1704,7 +1706,7 @@ static int32 DoRecvMessage(FActivity* Activity)
 
 	FMutableMemoryView DestView = Internal.Dest->GetMutableView();
 	const char* Cursor = BufferRight - AlreadyReceived;
-	if (!Streamed || AlreadyReceived < DestView.GetSize())
+	if (!bStreamed || AlreadyReceived < DestView.GetSize())
 	{
 		::memcpy(DestView.GetData(), Cursor, AlreadyReceived);
 		return 0;
