Fix for URLs that aren't null terminated.

#rnx

diff --git a/Engine/Source/Runtime/Experimental/IoStoreOnDemand/Private/CoreHttp/Http.cpp b/Engine/Source/Runtime/Experimental/IoStoreOnDemand/Private/CoreHttp/Http.cpp
--- a/Engine/Source/Runtime/Experimental/IoStoreOnDemand/Private/CoreHttp/Http.cpp
+++ b/Engine/Source/Runtime/Experimental/IoStoreOnDemand/Private/CoreHttp/Http.cpp
@@ -247,10 +247,11 @@ struct FUrlOffsets
 {
 	struct Slice
 	{
-		FAnsiStringView	Get(FAnsiStringView Url) const { return Url.Mid(Off, Len); }
-						operator bool () const { return Len != 0; }
-		uint8			Off;
-		uint8			Len;
+		FAnsiStringView	Get(FAnsiStringView Url) const { return Url.Mid(Left, Right - Left); }
+						operator bool () const { return Left > 0; }
+		int32			Len() const { return Right - Left; }
+		uint8			Left;
+		uint8			Right;
 	};
 	Slice				UserInfo;
 	Slice				HostName;
@@ -261,13 +262,15 @@ struct FUrlOffsets
 
 static int32 ParseUrl(FAnsiStringView Url, FUrlOffsets& Out)
 {
-	static const int32 LengthLimit = 127;
+	if (Url.Len() < 5)
+	{
+		return -1;
+	}
 
 	Out = {};
 
 	const char* Start = Url.GetData();
 	const char* Cursor = Start;
-	const char* End = Start + Url.Len();
 
 	// Scheme
 	int32 i = 0;
@@ -291,11 +294,11 @@ static int32 ParseUrl(FAnsiStringView Url, FUrlOffsets& Out)
 	{
 		return -1;
 	}
-	Cursor += i + 3;
+	i += 3;
 
 	struct { int32 c; int32 i; } Seps[2];
 	int32 SepCount = 0;
-	for (i = 0; i < LengthLimit - 8; ++i) // '8' is roughly "http[s]://"
+	for (; i < Url.Len(); ++i)
 	{
 		int32 c = Cursor[i];
 		if (c < '-')							break;
@@ -309,13 +312,14 @@ static int32 ParseUrl(FAnsiStringView Url, FUrlOffsets& Out)
 		Seps[SepCount++] = { c, i };
 	}
 
-	if (int32 c = Cursor[i]; c)
+	if (i > 0xff || i <= Scheme.Len() + 3)
 	{
-		if (c != '/')
-		{
-			return -1;
-		}
-		Out.Path = uint8(ptrdiff_t(Cursor + i - Start));
+		return -1;
+	}
+
+	if (i < Url.Len())
+	{
+		Out.Path = uint8(i);
 	}
 
 	Out.HostName = { uint8(Scheme.Len() + 3), uint8(i) };
@@ -328,14 +332,14 @@ static int32 ParseUrl(FAnsiStringView Url, FUrlOffsets& Out)
 	case 1:
 		if (Seps[0].c == ':')
 		{
-			Out.Port = { uint8(Out.HostName.Off + Seps[0].i + 1), uint8(i - Seps[0].i - 1) };
-			Out.HostName.Len = Seps[0].i;
+			Out.Port = { uint8(Seps[0].i + 1), uint8(i) };
+			Out.HostName.Right = Seps[0].i;
 		}
 		else
 		{
-			Out.UserInfo = { Out.HostName.Off, uint8(Seps[0].i) };
-			Out.HostName.Off += Seps[0].i + 1;
-			Out.HostName.Len -= Seps[0].i + 1;
+			Out.UserInfo = { Out.HostName.Left, uint8(Seps[0].i) };
+			Out.HostName.Left += Seps[0].i + 1;
+			Out.HostName.Right += Seps[0].i + 1;
 		}
 		break;
 
@@ -344,12 +348,11 @@ static int32 ParseUrl(FAnsiStringView Url, FUrlOffsets& Out)
 		{
 			return -1;
 		}
-		Out.UserInfo = { Out.HostName.Off, uint8(Seps[0].i) };
-		Out.Port = Out.HostName;
-		Out.Port.Off += Seps[1].i + 1;
-		Out.Port.Len -= Seps[1].i + 1;
-		Out.HostName.Off += Out.UserInfo.Len + 1;
-		Out.HostName.Len -= Out.UserInfo.Len + Out.Port.Len + 2;
+		Out.UserInfo = { Out.HostName.Left, uint8(Seps[0].i) };
+		Out.Port.Left = Seps[1].i + 1;
+		Out.Port.Right = Out.HostName.Right;
+		Out.HostName.Left = Out.UserInfo.Right + 1;
+		Out.HostName.Right = Out.Port.Left - 1;
 		break;
 
 	default:
@@ -357,15 +360,15 @@ static int32 ParseUrl(FAnsiStringView Url, FUrlOffsets& Out)
 	}
 
 	bool Bad = false;
-	Bad |= (Out.HostName.Len == 0);
-	Bad |= (Out.UserInfo.Off != 0) & (Out.UserInfo.Len == 0);
+	Bad |= (Out.HostName.Len() == 0);
+	Bad |= bool(Out.UserInfo) & (Out.UserInfo.Len() == 0);
 
-	if (Out.Port.Off)
+	if (Out.Port.Left)
 	{
-		Bad |= (Out.Port.Len == 0);
-		for (int32 j = 0, n = Out.Port.Len; j < n; ++j)
+		Bad |= (Out.Port.Len() == 0);
+		for (int32 j = 0, n = Out.Port.Len(); j < n; ++j)
 		{
-			Bad |= (uint32(Start[Out.Port.Off + j] - '0') > 9);
+			Bad |= (uint32(Start[Out.Port.Left + j] - '0') > 9);
 		}
 	}
 
@@ -2267,16 +2270,37 @@ static void MiscTest()
 	check(ParseUrl("http://@:/", UrlOut) == -1);
 	check(ParseUrl("http://foo:ba:r/", UrlOut) == -1);
 	check(ParseUrl("http://foo@ba:r/", UrlOut) == -1);
+	check(ParseUrl("http://foo@ba:r", UrlOut) == -1);
 	check(ParseUrl("http://foo@ba:/", UrlOut) == -1);
 	check(ParseUrl("http://foo@ba@9/", UrlOut) == -1);
 	check(ParseUrl("http://@ba:9/", UrlOut) == -1);
-	check(ParseUrl("http://zzzzzzzzzzzzzzzzzzzzzzzzzzzzzzzzzzzzzzzzzzzzzzzzzzzzzzzzzzzzzzzzzzzzzzzzzzzzzzzzzzzzzzzzzzzzzzzzzzzzzzzzzzzzzzzzzzzzzzzzzzzzzzzzzzz-hello-zz.com/", UrlOut) == -1);
+	check(ParseUrl(
+		"http://zzzzzzzzzzzzzzzzzzzzzzzzzzzzzzzzzzzzzzzzzzzzzzzzzzzzzzzzzzzzzzzzzzzzzzzzzzzzzzzzzzzzzzzzzzzzzzzzzzzzzzzzzzzzzzzzzzzzzzzzzz"
+		"zzzzzzzzzzzzzzzzzzzzzzzzzzzzzzzzzzzzzzzzzzzzzzzzzzzzzzzzzzzzzzzzzzzzzzzzzzzzzzzzzzzzzzzzzzzzzzzzzzzzzzzzzzzzzzzzzzzzzzzzzzzzzzzzz"
+		"zzzzzzzzzzzzzzzzzzzzzzzzzzzzzzzzzzzzzzzzzzzzzzzzzzzzzzzzzzzzzzzzzzzzzzzzzzzzzzzzzzzzzzzzzzzzzzzzzzzzzzzzzzzzzzzzzzzzzzzzzzzzzzzzz"
+		"zzzzzzzzzzzzzzzzzzzzzzzzzzzzzzzzzzzzzzzzzzzzzzzzzzzzzzzzzzzzzzzzzzzzzzzzzzzzzzzzzzzzzzzzzzzzzzzzzzzzzzzzzzzzzzzzzzzzzzzzzzzzzzzzz"
+		"zzzzzzzzzzzzzzzzzzzzzzzzzzzzzzzzzzzzzzzzzzzzzzzzzzzzzzzzzzzzzzzzzzzzzzzzzzzzzzzzzzzzzzzzzzzzzzzzzzzzzzzzzzzzzzzzzzzzzzzzzzzzz.com",
+		UrlOut) == -1);
 
 	check(ParseUrl("http://ab-c.com/", UrlOut) > 0);
 	check(ParseUrl("http:[REDACTED_EMAIL]/", UrlOut) > 0);
 	check(ParseUrl("https://abc.com", UrlOut) > 0);
 	check(ParseUrl("https://abc.com:999", UrlOut) > 0);
 	check(ParseUrl("https://abc.com:999/", UrlOut) > 0);
+	check(ParseUrl("https://foo:[REDACTED_EMAIL]:999", UrlOut) > 0);
+	check(ParseUrl("https://foo:[REDACTED_EMAIL]:999/", UrlOut) > 0);
+	check(ParseUrl("https:[REDACTED_EMAIL]:999", UrlOut) > 0);
+	check(ParseUrl("https:[REDACTED_EMAIL]:999/", UrlOut) > 0);
+
+	for (int32 i : { 0x10, 0x20, 0x40, 0x7f, 0xff })
+	{
+		char Url[] = "http://stockholm.patchercache.epicgames.net:123";
+		char Buffer[512];
+		std::memset(Buffer, i, sizeof(Buffer));
+		std::memcpy(Buffer, Url, sizeof(Url) - 1);
+		check(ParseUrl(FAnsiStringView(Buffer, sizeof(Url) - 1), UrlOut) > 0);
+		check(UrlOut.Port.Get(Url) == "123");
+	}
 
 	FAnsiStringView Url = "http://abc:[REDACTED_EMAIL]:999/";
 	check(ParseUrl(Url, UrlOut) > 0);
