Added Ias LLM scopes.

#rnx

diff --git a/Engine/Source/Runtime/Experimental/IoStoreOnDemand/Private/IoStoreOnDemand.cpp b/Engine/Source/Runtime/Experimental/IoStoreOnDemand/Private/IoStoreOnDemand.cpp
--- a/Engine/Source/Runtime/Experimental/IoStoreOnDemand/Private/IoStoreOnDemand.cpp
+++ b/Engine/Source/Runtime/Experimental/IoStoreOnDemand/Private/IoStoreOnDemand.cpp
@@ -4,6 +4,7 @@
 #include "OnDemandIoDispatcherBackend.h"
 #include "EncryptionKeyManager.h"
 #include "LatencyInjector.h"
+#include "Statistics.h"
 
 #include "FileIoCache.h"
 #include "HAL/LowLevelMemTracker.h"
@@ -1453,7 +1454,7 @@ void FIoStoreOnDemandModule::ReportAnalytics(TArray<FAnalyticsEventAttribute>& O
 
 void FIoStoreOnDemandModule::StartupModule()
 {
-	LLM_SCOPE(ELLMTag::FileSystem);
+	LLM_SCOPE_BYTAG(Ias);
 	using namespace UE::IO::Private;
 
 #if WITH_EDITOR
diff --git a/Engine/Source/Runtime/Experimental/IoStoreOnDemand/Private/JournaledCache.cpp b/Engine/Source/Runtime/Experimental/IoStoreOnDemand/Private/JournaledCache.cpp
--- a/Engine/Source/Runtime/Experimental/IoStoreOnDemand/Private/JournaledCache.cpp
+++ b/Engine/Source/Runtime/Experimental/IoStoreOnDemand/Private/JournaledCache.cpp
@@ -8,6 +8,7 @@
 #include "Containers/Array.h"
 #include "Containers/Map.h"
 #include "Containers/UnrealString.h"
+#include "HAL/LowLevelMemTracker.h"
 #include "HAL/PlatformFile.h"
 #include "HAL/PlatformProcess.h"
 #include "HAL/Runnable.h"
@@ -1127,6 +1128,8 @@ FJournaledCache::~FJournaledCache()
 ////////////////////////////////////////////////////////////////////////////////
 void FJournaledCache::Update()
 {
+	LLM_SCOPE_BYTAG(Ias);
+
 	int32 WriteAllowance = Governor.TickAllowance();
 	if (!WriteAllowance)
 	{
@@ -1173,6 +1176,8 @@ FJournaledCache::GetRetType	FJournaledCache::Get(
 {
 	uint64 InnerKey = ReduceKey(Key);
 	return UE::Tasks::Launch(TEXT("IasCacheGet"), [this, InnerKey, Options, CancelToken] () {
+		LLM_SCOPE_BYTAG(Ias);
+
 		FCacheInner::FEntry Entry = Cache->Get(InnerKey);
 		if (!Entry.IsHit())
 		{
@@ -1483,6 +1488,7 @@ IOSTOREONDEMAND_API void Tests()
 ////////////////////////////////////////////////////////////////////////////////
 TUniquePtr<IIoCache> MakeJournaledCache(const FFileIoCacheConfig& Config)
 {
+	LLM_SCOPE_BYTAG(Ias);
 	return MakeUnique<UE::IO::Private::FJournaledCache>(Config);
 }
 
diff --git a/Engine/Source/Runtime/Experimental/IoStoreOnDemand/Private/OnDemandIoDispatcherBackend.cpp b/Engine/Source/Runtime/Experimental/IoStoreOnDemand/Private/OnDemandIoDispatcherBackend.cpp
--- a/Engine/Source/Runtime/Experimental/IoStoreOnDemand/Private/OnDemandIoDispatcherBackend.cpp
+++ b/Engine/Source/Runtime/Experimental/IoStoreOnDemand/Private/OnDemandIoDispatcherBackend.cpp
@@ -265,7 +265,7 @@ void FDistributionEndpoints::IssueEndpointRequests()
 				[this, &ResolveRequest, MaxAttempts]
 				(FHttpRequestPtr, FHttpResponsePtr Response, bool bOk)
 				{
-					LLM_SCOPE(ELLMTag::FileSystem);
+					LLM_SCOPE_BYTAG(Ias);
 					FHttpRequestPtr Request = MoveTemp(ResolveRequest.HttpRequest);
 					if (Response->GetResponseCode() != 200)
 					{
@@ -706,7 +706,7 @@ void FOnDemandIoStore::AddDeferredContainers()
 
 void FOnDemandIoStore::OnEncryptionKeyAdded(const FGuid& Id, const FAES::FAESKey& Key)
 {
-	LLM_SCOPE(ELLMTag::FileSystem);
+	LLM_SCOPE_BYTAG(Ias);
 	TRACE_CPUPROFILER_EVENT_SCOPE(FOnDemandIoBackend::OnEncryptionKeyAdded);
 	AddDeferredContainers();
 }
@@ -1246,7 +1246,7 @@ FOnDemandIoBackend::~FOnDemandIoBackend()
 void FOnDemandIoBackend::Initialize(TSharedRef<const FIoDispatcherBackendContext> Context)
 {
 	TRACE_CPUPROFILER_EVENT_SCOPE(FOnDemandIoBackend::Initialize);
-	LLM_SCOPE(ELLMTag::FileSystem);
+	LLM_SCOPE_BYTAG(Ias);
 	UE_LOG(LogIas, Log, TEXT("Initializing on demand I/O dispatcher backend"));
 	BackendContext = Context;
 	DistributionEndpoints.ResolveDeferredEndpoints();
@@ -1269,7 +1269,7 @@ void FOnDemandIoBackend::Shutdown()
 
 void FOnDemandIoBackend::CompleteRequest(FChunkRequest* ChunkRequest)
 {
-	LLM_SCOPE(ELLMTag::FileSystem);
+	LLM_SCOPE_BYTAG(Ias);
 	TRACE_CPUPROFILER_EVENT_SCOPE(FOnDemandIoBackend::CompleteRequest);
 	check(ChunkRequest != nullptr);
 	const bool bCancelled = ChunkRequest->CancellationToken.IsCancelled();
@@ -1385,7 +1385,7 @@ bool FOnDemandIoBackend::Resolve(FIoRequestImpl* Request)
 	const ETaskPriority TaskPriority = ChunkRequest->Priority > IoDispatcherPriority_Medium ? ETaskPriority::High : ETaskPriority::Normal;
 	Launch(UE_SOURCE_LOCATION, [this, ChunkRequest]()
 	{
-		LLM_SCOPE(ELLMTag::FileSystem);
+		LLM_SCOPE_BYTAG(Ias);
 		TRACE_CPUPROFILER_EVENT_SCOPE(FOnDemandIoBackend::CompleteOrEnqueueHttpRequest);
 		if (ChunkRequest->CacheTask.IsValid())
 		{
@@ -1512,7 +1512,7 @@ FIoStatus FOnDemandIoBackend::MountDeferredEndpoints(const FString& Distribution
 
 void FOnDemandIoBackend::Mount(const FOnDemandEndpoint& Endpoint)
 {
-	LLM_SCOPE(ELLMTag::FileSystem);
+	LLM_SCOPE_BYTAG(Ias);
 	TRACE_CPUPROFILER_EVENT_SCOPE(FOnDemandIoBackend::Mount);
 
 	if ((Endpoint.DistributionUrl.IsEmpty() && Endpoint.ServiceUrl.IsEmpty()) || Endpoint.TocPath.IsEmpty())
@@ -1711,7 +1711,7 @@ FIoStatus FOnDemandIoBackend::AddToc(const FOnDemandEndpoint& Endpoint)
 
 uint32 FOnDemandIoBackend::Run()
 {
-	LLM_SCOPE(ELLMTag::FileSystem);
+	LLM_SCOPE_BYTAG(Ias);
 
 	const int32 MaxConcurrentRequests = HttpClient->MaxConnectionCount();
 	FChunkRequest* NextChunkRequest = nullptr;
diff --git a/Engine/Source/Runtime/Experimental/IoStoreOnDemand/Private/Statistics.cpp b/Engine/Source/Runtime/Experimental/IoStoreOnDemand/Private/Statistics.cpp
--- a/Engine/Source/Runtime/Experimental/IoStoreOnDemand/Private/Statistics.cpp
+++ b/Engine/Source/Runtime/Experimental/IoStoreOnDemand/Private/Statistics.cpp
@@ -5,6 +5,8 @@
 #include "AnalyticsEventAttribute.h"
 #include "Misc/CoreDelegates.h"
 
+LLM_DEFINE_TAG(Ias);
+
 #if IAS_WITH_STATISTICS
 
 namespace UE::IO::Private
diff --git a/Engine/Source/Runtime/Experimental/IoStoreOnDemand/Private/Statistics.h b/Engine/Source/Runtime/Experimental/IoStoreOnDemand/Private/Statistics.h
--- a/Engine/Source/Runtime/Experimental/IoStoreOnDemand/Private/Statistics.h
+++ b/Engine/Source/Runtime/Experimental/IoStoreOnDemand/Private/Statistics.h
@@ -13,6 +13,8 @@
 
 struct FAnalyticsEventAttribute;
 
+LLM_DECLARE_TAG(Ias);
+
 namespace UE::IO::Private
 {
 
