Non-unity fix

#rnx

diff --git a/Engine/Source/Runtime/Experimental/IoStoreOnDemand/Private/Statistics.h b/Engine/Source/Runtime/Experimental/IoStoreOnDemand/Private/Statistics.h
--- a/Engine/Source/Runtime/Experimental/IoStoreOnDemand/Private/Statistics.h
+++ b/Engine/Source/Runtime/Experimental/IoStoreOnDemand/Private/Statistics.h
@@ -2,6 +2,7 @@
 
 #pragma once
 
+#include "HAL/LowLevelMemTracker.h"
 #include "ProfilingDebugging/CountersTrace.h"
 #include "ProfilingDebugging/CsvProfiler.h"
 
