Geometry Collection : Fix apply breaking velocity to account for leaf nodes

diff --git a/Engine/Source/Runtime/Experimental/Chaos/Private/PhysicsProxy/GeometryCollectionPhysicsProxy.cpp b/Engine/Source/Runtime/Experimental/Chaos/Private/PhysicsProxy/GeometryCollectionPhysicsProxy.cpp
--- a/Engine/Source/Runtime/Experimental/Chaos/Private/PhysicsProxy/GeometryCollectionPhysicsProxy.cpp
+++ b/Engine/Source/Runtime/Experimental/Chaos/Private/PhysicsProxy/GeometryCollectionPhysicsProxy.cpp
@@ -2417,24 +2417,35 @@ static void ApplyToBreakingChildren_Internal(Chaos::FRigidClustering& Clustering
 {
 	const bool bIsCluster = (ClusteredHandle.ClusterIds().NumChildren > 0);
 	const bool bHasClusterUnionParent = Clustering.GetClusterUnionManager().IsClusterUnionParticle(ClusteredHandle.Parent());
-	if (bIsCluster && (!ClusteredHandle.Disabled() || bHasClusterUnionParent))
+	if (!ClusteredHandle.Disabled() || bHasClusterUnionParent)
 	{
-		Chaos::FRigidClustering::FClusterMap& ChildrenMap = Clustering.GetChildrenMap();
-		if (const TArray<Chaos::FPBDRigidParticleHandle*>* ChildrenHandles = ChildrenMap.Find(&ClusteredHandle))
+		if (bIsCluster)
 		{
-			for (Chaos::FPBDRigidParticleHandle* ChildHandle: *ChildrenHandles)
+			Chaos::FRigidClustering::FClusterMap& ChildrenMap = Clustering.GetChildrenMap();
+			if (const TArray<Chaos::FPBDRigidParticleHandle*>* ChildrenHandles = ChildrenMap.Find(&ClusteredHandle))
 			{
-				if (Chaos::FPBDRigidClusteredParticleHandle* ClusteredChildHandle = ChildHandle->CastToClustered())
+				for (Chaos::FPBDRigidParticleHandle* ChildHandle : *ChildrenHandles)
 				{
-					// todo(chaos) : this does not account for the various damage models, we should eventually call an evaluate function to avoid replicating logic from the clustering code
-					// also we cannot account for collision impulses because they are set after the physics callbacks are evaluated
-					if (ClusteredChildHandle->GetExternalStrain() >= ClusteredChildHandle->GetInternalStrains())
+					if (Chaos::FPBDRigidClusteredParticleHandle* ClusteredChildHandle = ChildHandle->CastToClustered())
 					{
-						Action(ClusteredChildHandle);
+						// todo(chaos) : this does not account for the various damage models, we should eventually call an evaluate function to avoid replicating logic from the clustering code
+						// also we cannot account for collision impulses because they are set after the physics callbacks are evaluated
+						if (ClusteredChildHandle->GetExternalStrain() >= ClusteredChildHandle->GetInternalStrains())
+						{
+							Action(ClusteredChildHandle);
+						}
 					}
 				}
 			}
 		}
+		else
+		{
+			// leaf node , let's just apply to it directly
+			if (ClusteredHandle.GetExternalStrain() >= ClusteredHandle.GetInternalStrains())
+			{
+				Action(&ClusteredHandle);
+			}
+		}
 	}
 }
 
