Chaos - support large collision manifolds (previously limited to 4 points, now unlimited)
- fix collision visitor with VisitSleeping flag

diff --git a/Engine/Source/Runtime/Experimental/Chaos/Private/Chaos/Collision/CollisionConstraintAllocator.cpp b/Engine/Source/Runtime/Experimental/Chaos/Private/Chaos/Collision/CollisionConstraintAllocator.cpp
--- a/Engine/Source/Runtime/Experimental/Chaos/Private/Chaos/Collision/CollisionConstraintAllocator.cpp
+++ b/Engine/Source/Runtime/Experimental/Chaos/Private/Chaos/Collision/CollisionConstraintAllocator.cpp
@@ -168,7 +168,7 @@ namespace Chaos
 			{
 				if (Constraint != nullptr)
 				{
-					Constraint->SetDisabled(true);
+					Constraint->BeginTick();
 				}
 			}
 		}
diff --git a/Engine/Source/Runtime/Experimental/Chaos/Private/Chaos/Collision/PBDCollisionConstraint.cpp b/Engine/Source/Runtime/Experimental/Chaos/Private/Chaos/Collision/PBDCollisionConstraint.cpp
--- a/Engine/Source/Runtime/Experimental/Chaos/Private/Chaos/Collision/PBDCollisionConstraint.cpp
+++ b/Engine/Source/Runtime/Experimental/Chaos/Private/Chaos/Collision/PBDCollisionConstraint.cpp
@@ -635,14 +635,6 @@ namespace Chaos
 
 	void FPBDCollisionConstraint::AddIncrementalManifoldContact(const FContactPoint& ContactPoint)
 	{
-		if (ManifoldPoints.IsFull())
-		{
-			// @todo(chaos): we should remove a contact here if we try to add a new point
-			// For now just update the existing ones to select the deepest
-			UpdateManifoldContacts();
-			return;
-		}
-
 		if (Flags.bUseManifold)
 		{
 			// See if the manifold point already exists
@@ -1059,7 +1051,7 @@ namespace Chaos
 		return TNumericLimits<FReal>::Max();
 	}
 
-	int32 FPBDCollisionConstraint::FindSavedManifoldPoint(const int32 ManifoldPointIndex, TCArray<int32, MaxManifoldPoints>& InOutAllowedSavedPointIndices) const
+	int32 FPBDCollisionConstraint::FindSavedManifoldPoint(const int32 ManifoldPointIndex, int32* InOutAllowedSavedPointIndices, int32& InOutNumAllowedSavedPoints) const
 	{
 		int32 MatchIndex = INDEX_NONE;
 
@@ -1072,7 +1064,7 @@ namespace Chaos
 				const FReal DistanceToleranceSq = FMath::Square(Chaos_Manifold_FrictionPositionTolerance);
 				FReal BestScore = DistanceToleranceSq;
 
-				for (int32 AllowedPointIndex = 0; AllowedPointIndex < InOutAllowedSavedPointIndices.Num(); ++AllowedPointIndex)
+				for (int32 AllowedPointIndex = 0; AllowedPointIndex < InOutNumAllowedSavedPoints; ++AllowedPointIndex)
 				{
 					const int32 SavedPointIndex = InOutAllowedSavedPointIndices[AllowedPointIndex];
 					const FSavedManifoldPoint& SavedManifoldPoint = SavedManifoldPoints[SavedPointIndex];
@@ -1083,7 +1075,12 @@ namespace Chaos
 						BestScore = Score;
 						MatchIndex = SavedPointIndex;
 
-						InOutAllowedSavedPointIndices.RemoveAtSwap(AllowedPointIndex);
+						// RemoveAtSwap
+						--InOutNumAllowedSavedPoints;
+						if (AllowedPointIndex < InOutNumAllowedSavedPoints)
+						{
+							InOutAllowedSavedPointIndices[AllowedPointIndex] = InOutAllowedSavedPointIndices[InOutNumAllowedSavedPoints];
+						}
 
 						// Just take the first match we find that meets our tolerance
 						break;
@@ -1097,16 +1094,17 @@ namespace Chaos
 
 	void FPBDCollisionConstraint::AssignSavedManifoldPoints()
 	{
-		TCArray<int32, MaxManifoldPoints> AllowedPointIndices;
+		TManifoldPointArray<int32> AllowedPointIndices;
 		AllowedPointIndices.SetNum(SavedManifoldPoints.Num());
 		for (int32 PointIndex = 0, PointEndIndex = SavedManifoldPoints.Num(); PointIndex < PointEndIndex; ++PointIndex)
 		{
 			AllowedPointIndices[PointIndex] = PointIndex;
 		}
 
+		int32 NumAllowedPointIndices = AllowedPointIndices.Num();
 		for (int32 PointIndex = 0, PointEndIndex = ManifoldPoints.Num(); PointIndex < PointEndIndex; ++PointIndex)
 		{
-			const int32 SavedManifoldPointIndex = FindSavedManifoldPoint(PointIndex, AllowedPointIndices);
+			const int32 SavedManifoldPointIndex = FindSavedManifoldPoint(PointIndex, AllowedPointIndices.GetData(), NumAllowedPointIndices);
 
 			if (SavedManifoldPointIndex != INDEX_NONE)
 			{
diff --git a/Engine/Source/Runtime/Experimental/Chaos/Private/Chaos/Collision/PBDCollisionContainerSolver.cpp b/Engine/Source/Runtime/Experimental/Chaos/Private/Chaos/Collision/PBDCollisionContainerSolver.cpp
--- a/Engine/Source/Runtime/Experimental/Chaos/Private/Chaos/Collision/PBDCollisionContainerSolver.cpp
+++ b/Engine/Source/Runtime/Experimental/Chaos/Private/Chaos/Collision/PBDCollisionContainerSolver.cpp
@@ -173,10 +173,12 @@ namespace Chaos
 			if (!Constraint->GetManifoldPoint(ConstraintManifoldPointIndex).Flags.bDisabled)
 			{
 				const int32 SolverManifoldPointIndex = Solver.AddManifoldPoint();
-
-				// Transform the constraint contact data into world space for use by the solver
-				// We build this data directly into the solver's world-space contact data which looks a bit odd with "Init" called after but there you go
-				UpdateCollisionSolverContactPointFromConstraint(Solver, SolverManifoldPointIndex, Constraint, ConstraintManifoldPointIndex, Dt, Body0, Body1);
+				if (SolverManifoldPointIndex != INDEX_NONE)
+				{
+					// Transform the constraint contact data into world space for use by the solver
+					// We build this data directly into the solver's world-space contact data which looks a bit odd with "Init" called after but there you go
+					UpdateCollisionSolverContactPointFromConstraint(Solver, SolverManifoldPointIndex, Constraint, ConstraintManifoldPointIndex, Dt, Body0, Body1);
+				}
 			}
 		}
 
@@ -249,21 +251,24 @@ namespace Chaos
 
 			if (!Constraint->GetManifoldPoint(ManifoldPointIndex).Flags.bDisabled)
 			{
-				const Private::FPBDCollisionSolverManifoldPoint& SolverManifoldPoint = Solver.GetManifoldPoint(SolverManifoldPointIndex);
+				if (SolverManifoldPointIndex < Solver.NumManifoldPoints())
+				{
+					const Private::FPBDCollisionSolverManifoldPoint& SolverManifoldPoint = Solver.GetManifoldPoint(SolverManifoldPointIndex);
 
-				NetPushOut =
-					SolverManifoldPoint.NetPushOutNormal * SolverManifoldPoint.ContactNormal +
-					SolverManifoldPoint.NetPushOutTangentU * SolverManifoldPoint.ContactTangentU +
-					SolverManifoldPoint.NetPushOutTangentV * SolverManifoldPoint.ContactTangentV;
+					NetPushOut =
+						SolverManifoldPoint.NetPushOutNormal * SolverManifoldPoint.ContactNormal +
+						SolverManifoldPoint.NetPushOutTangentU * SolverManifoldPoint.ContactTangentU +
+						SolverManifoldPoint.NetPushOutTangentV * SolverManifoldPoint.ContactTangentV;
 
-				NetImpulse =
-					SolverManifoldPoint.NetImpulseNormal * SolverManifoldPoint.ContactNormal +
-					SolverManifoldPoint.NetImpulseTangentU * SolverManifoldPoint.ContactTangentU +
-					SolverManifoldPoint.NetImpulseTangentV * SolverManifoldPoint.ContactTangentV;
+					NetImpulse =
+						SolverManifoldPoint.NetImpulseNormal * SolverManifoldPoint.ContactNormal +
+						SolverManifoldPoint.NetImpulseTangentU * SolverManifoldPoint.ContactTangentU +
+						SolverManifoldPoint.NetImpulseTangentV * SolverManifoldPoint.ContactTangentV;
 
-				StaticFrictionRatio = SolverManifoldPoint.StaticFrictionRatio;
+					StaticFrictionRatio = SolverManifoldPoint.StaticFrictionRatio;
 
-				++SolverManifoldPointIndex;
+					++SolverManifoldPointIndex;
+				}
 			}
 
 			// NOTE: We call this even for points we did not run the solver for (but with zero results)
@@ -287,7 +292,10 @@ namespace Chaos
 		, CollisionConstraints()
 		, AppliedShockPropagation(1)
 		, Scratch()
-		, CollisionSolvers()
+		, CollisionSolvers(nullptr)
+		, CollisionSolverManifoldPoints(nullptr)
+		, NumCollisionSolverManifoldPoints(0)
+		, MaxCollisionSolverManifoldPoints(0)
 		, bCollisionConstraintPerIterationCollisionDetection()
 		, bPerIterationCollisionDetection(false)
 	{
@@ -301,30 +309,72 @@ namespace Chaos
 	{
 		AppliedShockPropagation = FSolverReal(1);
 
-		// A over-allocation policy to avoid reallocation every frame in the common case where a pile of objects is dropped
-		// and the number of contacts increases every tick.
-		int CollisionBufferNum = MaxCollisions;
-		if (CollisionBufferNum > CollisionConstraints.Max())
-		{
-			CollisionBufferNum = (5 * MaxCollisions) / 4; // +25%
-		}
-
+		const int CollisionBufferNum = CalculateCollisionBufferNum(MaxCollisions, CollisionConstraints.Num());
 		CollisionConstraints.Reset(CollisionBufferNum);
 		bCollisionConstraintPerIterationCollisionDetection.Reset(CollisionBufferNum);
 
 		// Just set the array size for these right away - all data will be initialized later
 		bCollisionConstraintPerIterationCollisionDetection.SetNumUninitialized(MaxCollisions, false);
 
-		// Prepare the scratch buffer
-		const size_t AlignedSolverSize = Align(sizeof(Private::FPBDCollisionSolver), alignof(Private::FPBDCollisionSolver));
-		const size_t AlignedPointSize = Align(sizeof(Private::FPBDCollisionSolverManifoldPoint), alignof(Private::FPBDCollisionSolverManifoldPoint));
-		const size_t ScratchSize = CollisionBufferNum * (AlignedSolverSize + Private::FPBDCollisionSolver::MaxPointsPerConstraint * AlignedPointSize);
-		Scratch.Reset(ScratchSize);
+		// Reset the solver buffers. We could manifold points as constraints are added,
+		// and re-allocate and assign the scratch buffers after
+		NumCollisionSolverManifoldPoints = 0;
+		MaxCollisionSolverManifoldPoints = 0;
+		CollisionSolvers = nullptr;
+		CollisionSolverManifoldPoints = nullptr;
+	}
+
+	int32 FPBDCollisionContainerSolver::CalculateCollisionBufferNum(const int32 InTightFittingNum, const int32 InCurrentBufferNum) const
+	{
+		// A buffer over-allocation policy to avoid reallocation every frame in the common case where a pile of objects is dropped
+		// and the number of contacts increases every tick. Used for collision solvers and manifold points
+		int CollisionBufferNum = InTightFittingNum;
+		if (CollisionBufferNum > InCurrentBufferNum)
+		{
+			CollisionBufferNum = (5 * InTightFittingNum) / 4; // +25%
+		}
+		return CollisionBufferNum;
+	}
+
+	int32 FPBDCollisionContainerSolver::CalculateConstraintMaxManifoldPoints(const FPBDCollisionConstraint* Constraint) const
+	{
+		// NOTE: we don't know how many points we will create if collision detection is deferred or incrememntal, so we just assume 4
+		// @todo(chaos): see if we can do better here
+		const bool bDeferredCollisionDetection = ConstraintContainer.GetDetectorSettings().bDeferNarrowPhase;
+		const int32 ManifoldPointMax = (bDeferredCollisionDetection || Constraint->GetUseIncrementalManifold()) ? 4 : Constraint->NumManifoldPoints();
+		return ManifoldPointMax;
+	}
 
-		// Allocate scratch space for the collision solvers and manifold points
-		CollisionSolvers = Scratch.AllocArray<Private::FPBDCollisionSolver>(MaxCollisions);
-		CollisionSolverManifoldPoints = Scratch.AllocArray<Private::FPBDCollisionSolverManifoldPoint>(MaxCollisions);
+	void FPBDCollisionContainerSolver::PrepareSolverBuffer()
+	{
 		NumCollisionSolverManifoldPoints = 0;
+		MaxCollisionSolverManifoldPoints = 0;
+		CollisionSolvers = nullptr;
+		CollisionSolverManifoldPoints = nullptr;
+
+		// Count the manifold points
+		// @todo(chaos): can we avoid this?
+		const int32 NumCollisionsConstraints = CollisionConstraints.Num();
+		for (int32 ConstraintIndex = 0; ConstraintIndex < NumCollisionsConstraints; ++ConstraintIndex)
+		{
+			MaxCollisionSolverManifoldPoints += CalculateConstraintMaxManifoldPoints(GetConstraint(ConstraintIndex));
+		}
+
+		// Set up the solver buffers
+		if (NumCollisionsConstraints > 0)
+		{
+			// Resize the scratch buffer (up to 25% slack)
+			constexpr size_t AlignedSolverSize = Align(sizeof(Private::FPBDCollisionSolver), alignof(Private::FPBDCollisionSolver));
+			constexpr size_t AlignedPointSize = Align(sizeof(Private::FPBDCollisionSolverManifoldPoint), alignof(Private::FPBDCollisionSolverManifoldPoint));
+			const size_t ScratchSize = NumCollisionsConstraints * AlignedSolverSize + MaxCollisionSolverManifoldPoints * AlignedPointSize;
+			const size_t ScratchBufferSize = CalculateCollisionBufferNum(ScratchSize, Scratch.BufferSize());
+			Scratch.Reset(ScratchBufferSize);
+
+			// Allocate scratch space for the collision solvers and manifold points
+			// NOTE: scratch return a valid pointer even for 0 size so CollisionSolverManifoldPoints is always a valid pointer in the scratch (but may have no space)
+			CollisionSolvers = Scratch.AllocArray<Private::FPBDCollisionSolver>(NumCollisionsConstraints);
+			CollisionSolverManifoldPoints = Scratch.AllocArray<Private::FPBDCollisionSolverManifoldPoint>(MaxCollisionSolverManifoldPoints);
+		}
 	}
 
 	void FPBDCollisionContainerSolver::AddConstraints()
@@ -344,7 +394,7 @@ namespace Chaos
 	{
 		for (Private::FPBDIslandConstraint* IslandConstraint : IslandConstraints)
 		{
-			// We will only ever be given constraints from our container (asserts in non-shipping)
+			// We will only ever be given constraints from our container (NOTE: AsUnsafe asserts in dev)
 			FPBDCollisionConstraint& Constraint = IslandConstraint->GetConstraint()->AsUnsafe<FPBDCollisionConstraintHandle>()->GetContact();
 
 			AddConstraint(Constraint);
@@ -353,33 +403,22 @@ namespace Chaos
 
 	void FPBDCollisionContainerSolver::AddConstraint(FPBDCollisionConstraint& Constraint)
 	{
-		// NOTE: No need to add to CollisionSolvers or bCollisionConstraintPerIterationCollisionDetection
-		const int32 Index = CollisionConstraints.Add(&Constraint);
-
-		// Allocate the manifold points for this constraints solver
-		// NOTE: we don't know how many points we will create if collision detection is deferred or incrememntal, so just allocate space for the max allowed
-		// @todo(chaos): see if we can do better - we don't want to pay for this logic when it is only used rarely
-		const bool bDeferredCollisionDetection = ConstraintContainer.GetDetectorSettings().bDeferNarrowPhase;
-		const int32 ManifoldPointMax = (bDeferredCollisionDetection || Constraint.GetUseIncrementalManifold()) ? Private::FPBDCollisionSolver::MaxPointsPerConstraint : Constraint.NumManifoldPoints();
-
-		if (ManifoldPointMax > 0)
-		{
-			GetSolver(Index).SetManifoldPointsBuffer(&CollisionSolverManifoldPoints[NumCollisionSolverManifoldPoints], ManifoldPointMax);
-			NumCollisionSolverManifoldPoints += ManifoldPointMax;
-		}
-		else
-		{
-			GetSolver(Index).Reset();
-		}
+		// NOTE: No need to add to CollisionSolvers or bCollisionConstraintPerIterationCollisionDetection - handled later
+		CollisionConstraints.Add(&Constraint);
 	}
 
-
 	void FPBDCollisionContainerSolver::AddBodies(FSolverBodyContainer& SolverBodyContainer)
 	{
-		for (int32 SolverIndex = 0, SolverEndIndex = NumSolvers(); SolverIndex < SolverEndIndex; ++SolverIndex)
+		// All constarints are now added. We can allocate the solver buffers.
+		PrepareSolverBuffer();
+
+		// Make sure have a valid manifold point buffer if we have constraints (it may be zero size, but we want the pointer to be valid)
+		check((CollisionSolverManifoldPoints != nullptr) || (GetNumConstraints() == 0));
+
+		for (int32 ConstraintIndex = 0, ConstraintEndIndex = GetNumConstraints(); ConstraintIndex < ConstraintEndIndex; ++ConstraintIndex)
 		{
-			Private::FPBDCollisionSolver& CollisionSolver = GetSolver(SolverIndex);
-			FPBDCollisionConstraint* Constraint = GetConstraint(SolverIndex);
+			Private::FPBDCollisionSolver& CollisionSolver = GetSolver(ConstraintIndex);
+			FPBDCollisionConstraint* Constraint = GetConstraint(ConstraintIndex);
 			check(Constraint != nullptr);
 
 			// Find the solver bodies for the particles we constrain. This will add them to the container
@@ -389,7 +428,13 @@ namespace Chaos
 			check(Body0 != nullptr);
 			check(Body1 != nullptr);
 
+			// Set up the solver manifold point buffer pointer
+			const int32 ConstraintManifoldPointMax = CalculateConstraintMaxManifoldPoints(GetConstraint(ConstraintIndex));
+
+			CollisionSolver.Reset(&CollisionSolverManifoldPoints[NumCollisionSolverManifoldPoints], ConstraintManifoldPointMax);
 			CollisionSolver.SetSolverBodies(*Body0, *Body1);
+
+			NumCollisionSolverManifoldPoints += ConstraintManifoldPointMax;
 		}
 	}
 
@@ -467,7 +512,7 @@ namespace Chaos
 
 			// Reset the collision solver here as the pointers will be invalid on the next tick
 			// @todo(chaos): Pointers are always reinitalized before use next tick so this isn't strictly necessary
-			CollisionSolver.Reset();
+			CollisionSolver.Reset(nullptr, 0);
 		}
 	}
 
diff --git a/Engine/Source/Runtime/Experimental/Chaos/Private/Chaos/PBDSuspensionConstraints.cpp b/Engine/Source/Runtime/Experimental/Chaos/Private/Chaos/PBDSuspensionConstraints.cpp
--- a/Engine/Source/Runtime/Experimental/Chaos/Private/Chaos/PBDSuspensionConstraints.cpp
+++ b/Engine/Source/Runtime/Experimental/Chaos/Private/Chaos/PBDSuspensionConstraints.cpp
@@ -215,7 +215,7 @@ namespace Chaos
 			check(ConstraintSolverBodies[ConstraintIndex] != nullptr);
 
 			FPBDCollisionSolver* Solver = &CollisionSolvers[ConstraintIndex];
-			Solver->SetManifoldPointsBuffer(&CollisionSolverManifoldPoints[ConstraintIndex], 1);
+			Solver->Reset(&CollisionSolverManifoldPoints[ConstraintIndex], 1);
 			Solver->SetStiffness(1);
 
 			FSolverBody* Body0 = ConstraintSolverBodies[ConstraintIndex];	// vehicle chassis			
@@ -305,7 +305,7 @@ namespace Chaos
 		}
 
 		ConstraintSolverBodies[ConstraintIndex] = nullptr;
-		CollisionSolvers[ConstraintIndex].Reset();
+		CollisionSolvers[ConstraintIndex].Reset(nullptr, 0);
 	}
 
 	void FPBDSuspensionConstraints::ApplyPositionConstraint(const int32 ConstraintIndex, const FReal Dt, const int32 It, const int32 NumIts)
diff --git a/Engine/Source/Runtime/Experimental/Chaos/Private/ChaosVisualDebugger/ChaosVDDataWrapperUtils.cpp b/Engine/Source/Runtime/Experimental/Chaos/Private/ChaosVisualDebugger/ChaosVDDataWrapperUtils.cpp
--- a/Engine/Source/Runtime/Experimental/Chaos/Private/ChaosVisualDebugger/ChaosVDDataWrapperUtils.cpp
+++ b/Engine/Source/Runtime/Experimental/Chaos/Private/ChaosVisualDebugger/ChaosVDDataWrapperUtils.cpp
@@ -106,10 +106,11 @@ FChaosVDConstraint FChaosVDDataWrapperUtils::BuildConstraintDataWrapperFromConst
 	WrappedConstraintData.LastShapeWorldPositionDelta = FVector(InConstraint.LastShapeWorldPositionDelta);
 	WrappedConstraintData.LastShapeWorldRotationDelta = FQuat(InConstraint.LastShapeWorldRotationDelta);
 
-	WrappedConstraintData.ManifoldPoints.Reserve(Chaos::FPBDCollisionConstraint::MaxManifoldPoints);
-	WrappedConstraintData.ManifoldPoints.SetNum(Chaos::FPBDCollisionConstraint::MaxManifoldPoints);
+	const int32 MaxManifoldPoints = InConstraint.ManifoldPoints.Num();
+	WrappedConstraintData.ManifoldPoints.Reserve(MaxManifoldPoints);
+	WrappedConstraintData.ManifoldPoints.SetNum(MaxManifoldPoints);
 
-	for (int32 PointIndex = 0; PointIndex < Chaos::FPBDCollisionConstraint::MaxManifoldPoints; PointIndex++)
+	for (int32 PointIndex = 0; PointIndex < MaxManifoldPoints; PointIndex++)
 	{
 		FChaosVDManifoldPoint& CurrentCVDMainFoldPoint = WrappedConstraintData.ManifoldPoints[PointIndex];
 
diff --git a/Engine/Source/Runtime/Experimental/Chaos/Public/Chaos/Collision/CollisionConstraintAllocator.h b/Engine/Source/Runtime/Experimental/Chaos/Public/Chaos/Collision/CollisionConstraintAllocator.h
--- a/Engine/Source/Runtime/Experimental/Chaos/Public/Chaos/Collision/CollisionConstraintAllocator.h
+++ b/Engine/Source/Runtime/Experimental/Chaos/Public/Chaos/Collision/CollisionConstraintAllocator.h
@@ -603,9 +603,23 @@ namespace Chaos
 
 		if (!(VisitFlags & ECollisionVisitorFlags::VisitDisabled))
 		{
-			if (!Constraint->IsEnabled())
+			// Awkward. Sleeping constraints are considered "disabled" so if we want to
+			// visit sleeping constraints we must ignore the disabled state when asleep
+			if (!(VisitFlags & ECollisionVisitorFlags::VisitSleeping))
 			{
-				return false;
+				// We do NOT want to visit sleeping, so we can just check the enabled flag
+				if (!Constraint->IsEnabled())
+				{
+					return false;
+				}
+			}
+			else
+			{
+				// We do want to visit sleeping, so we only skip disabled-awake constraints
+				if (!Constraint->IsEnabled() && !Constraint->IsSleeping())
+				{
+					return false;
+				}
 			}
 		}
 
diff --git a/Engine/Source/Runtime/Experimental/Chaos/Public/Chaos/Collision/PBDCollisionConstraint.h b/Engine/Source/Runtime/Experimental/Chaos/Public/Chaos/Collision/PBDCollisionConstraint.h
--- a/Engine/Source/Runtime/Experimental/Chaos/Public/Chaos/Collision/PBDCollisionConstraint.h
+++ b/Engine/Source/Runtime/Experimental/Chaos/Public/Chaos/Collision/PBDCollisionConstraint.h
@@ -206,8 +206,6 @@ namespace Chaos
 	public:
 		using FConstraintContainerHandle = TIntrusiveConstraintHandle<FPBDCollisionConstraint>;
 
-		static const int32 MaxManifoldPoints = 4;
-
 		static constexpr FRealSingle MaxTOI = std::numeric_limits<FRealSingle>::max();
 
 		/**
@@ -350,10 +348,23 @@ namespace Chaos
 		void ResetPhi(FReal InPhi) { ClosestManifoldPointIndex = INDEX_NONE; }
 		FReal GetPhi() const { return (ClosestManifoldPointIndex != INDEX_NONE) ? ManifoldPoints[ClosestManifoldPointIndex].ContactPoint.Phi : TNumericLimits<FReal>::Max(); }
 
+		// A tick reset called on all constraints that were Activated last tick
+		void BeginTick()
+		{
+			Flags.bDisabled = true;
+			if (!IsSleeping())
+			{
+				Flags.bIsCurrent = false;
+			}
+			ContainerCookie.ConstraintIndex = INDEX_NONE;
+			ContainerCookie.CCDConstraintIndex = INDEX_NONE;
+		}
+
 		// Was this constraint activated this frame? It will be activated if the shapes are within CullDistance of each other.
-		// NOTE: All Current constraints are in the ActiveConstraints list on the CollisionConstraintAllocator. This remains true 
-		// even if disabled by the user (via SetDisabled()) in a callback. We only care about "not current" constraints
+		// NOTE: All Awake Current constraints are in the ActiveConstraints list on the CollisionConstraintAllocator. This remains true 
+		// even if disabled by the user (via SetDisabled()) in a callback. We usually only care about "not current" constraints
 		// for debug visualization/reporting. Normally you would only need to consider GetDisabled()
+		// NOTE: sleeping constraints are also considered Current, but are not in the ActiveConstraints list
 		bool IsCurrent() const { return Flags.bIsCurrent; }
 
 		// Allow the user to disable this constraint. @see GetActive().
@@ -464,8 +475,8 @@ namespace Chaos
 		CHAOS_API void ResetManifold();
 
 		// @todo(chaos): remove array view and provide per-point accessor
-		TArrayView<FManifoldPoint> GetManifoldPoints() { return MakeArrayView(ManifoldPoints.begin(), ManifoldPoints.Num()); }
-		TArrayView<const FManifoldPoint> GetManifoldPoints() const { return MakeArrayView(ManifoldPoints.begin(), ManifoldPoints.Num()); }
+		TArrayView<FManifoldPoint> GetManifoldPoints() { return MakeArrayView(ManifoldPoints); }
+		TArrayView<const FManifoldPoint> GetManifoldPoints() const { return MakeArrayView(ManifoldPoints); }
 
 		int32 NumManifoldPoints() const { return ManifoldPoints.Num(); }
 		FManifoldPoint& GetManifoldPoint(const int32 PointIndex) { return ManifoldPoints[PointIndex]; }
@@ -518,7 +529,7 @@ namespace Chaos
 		// @todo(chaos): remove this and use SetOneShotManifoldContacts
 		inline void AddOneshotManifoldContact(const FContactPoint& ContactPoint)
 		{
-			if (ContactPoint.IsSet() && !ManifoldPoints.IsFull())
+			if (ContactPoint.IsSet())
 			{
 				int32 ManifoldPointIndex = AddManifoldPoint(ContactPoint);
 				if (ManifoldPoints[ManifoldPointIndex].ContactPoint.Phi < GetPhi())
@@ -539,7 +550,11 @@ namespace Chaos
 			ResetActiveManifoldContacts();
 
 			FReal MinPhi = TNumericLimits<FReal>::Max();
-			const int32 NumContacts = FMath::Min(ContactPoints.Num(), MaxManifoldPoints);
+			
+			const int32 NumContacts = ContactPoints.Num();
+			
+			ManifoldPoints.Reserve(NumContacts);
+
 			for (int32 ContactIndex = 0; ContactIndex < NumContacts; ++ContactIndex)
 			{
 				const FContactPoint& ContactPoint = ContactPoints[ContactIndex];
@@ -600,16 +615,6 @@ namespace Chaos
 		CHAOS_API bool TryAddManifoldContact(const FContactPoint& ContactPoint);
 		CHAOS_API bool TryInsertManifoldContact(const FContactPoint& ContactPoint);
 
-		//@ todo(chaos): These are for the collision forwarding system - this should use the collision modifier system (which should be extended to support adding collisions)
-		void SetManifoldPoints(const TArray<FManifoldPoint>& InManifoldPoints)
-		{ 
-			ManifoldPoints.SetNum(FMath::Min(MaxManifoldPoints, InManifoldPoints.Num()));
-			for (int32 ManifoldPointIndex = 0; ManifoldPoints.Num(); ++ManifoldPointIndex)
-			{
-				ManifoldPoints[ManifoldPointIndex] = InManifoldPoints[ManifoldPointIndex];
-			}
-		}
-
 		// The GJK warm-start data. This is updated directly in the narrow phase
 		FGJKSimplexData& GetGJKWarmStartData() { return GJKWarmStartData; }
 
@@ -783,7 +788,7 @@ namespace Chaos
 		CHAOS_API bool AreMatchingContactPoints(const FContactPoint& A, const FContactPoint& B, FReal& OutScore) const;
 		CHAOS_API int32 FindManifoldPoint(const FContactPoint& ContactPoint) const;
 
-		CHAOS_API int32 FindSavedManifoldPoint(const int32 ManifoldPointIndex, TCArray<int32, MaxManifoldPoints>& InOutAllowedSavedPointIndices) const;
+		CHAOS_API int32 FindSavedManifoldPoint(const int32 ManifoldPointIndex, int32* InOutAllowedSavedPointIndices, int32& InOutNumAllowedSavedPoints) const;
 		CHAOS_API void AssignSavedManifoldPoints();
 
 		inline void InitManifoldPoint(const int32 ManifoldPointIndex, const FContactPoint& ContactPoint)
@@ -800,7 +805,7 @@ namespace Chaos
 
 		inline int32 AddManifoldPoint(const FContactPoint& ContactPoint)
 		{
-			int32 ManifoldPointIndex = ManifoldPoints.Add();	// Note: no initialization (see TCArray)
+			int32 ManifoldPointIndex = ManifoldPoints.AddUninitialized();	// Note: no initialization (see InitManifoldPoint)
 			InitManifoldPoint(ManifoldPointIndex, ContactPoint);
 			return ManifoldPointIndex;
 		}
@@ -891,9 +896,12 @@ namespace Chaos
 		Private::FImplicitBoundsTestFlags BoundsTestFlags;
 		EContactShapesType ShapesType;
 
-		TCArray<FSavedManifoldPoint, MaxManifoldPoints> SavedManifoldPoints;
-		TCArray<FManifoldPoint, MaxManifoldPoints> ManifoldPoints;
-		TCArray<FManifoldPointResult, MaxManifoldPoints> ManifoldPointResults;
+		template<typename T>
+		using TManifoldPointArray = TArray<T, TInlineAllocator<4>>;
+
+		TManifoldPointArray<FSavedManifoldPoint> SavedManifoldPoints;
+		TManifoldPointArray<FManifoldPoint> ManifoldPoints;
+		TManifoldPointArray<FManifoldPointResult> ManifoldPointResults;
 
 		// Value in range [0,1] used to interpolate P between [X,P] that we will rollback to when solving at time of impact.
 		FRealSingle CCDTimeOfImpact;
diff --git a/Engine/Source/Runtime/Experimental/Chaos/Public/Chaos/Collision/PBDCollisionContainerSolver.h b/Engine/Source/Runtime/Experimental/Chaos/Public/Chaos/Collision/PBDCollisionContainerSolver.h
--- a/Engine/Source/Runtime/Experimental/Chaos/Public/Chaos/Collision/PBDCollisionContainerSolver.h
+++ b/Engine/Source/Runtime/Experimental/Chaos/Public/Chaos/Collision/PBDCollisionContainerSolver.h
@@ -58,6 +58,9 @@ namespace Chaos
 
 		void CachePrefetchSolver(const int32 ConstraintIndex) const;
 		void AddConstraint(FPBDCollisionConstraint& Constraint);
+		int32 CalculateCollisionBufferNum(const int32 InTightFittingNum, const int32 InCurrentBufferNum) const;
+		int32 CalculateConstraintMaxManifoldPoints(const FPBDCollisionConstraint* Constraint) const;
+		void PrepareSolverBuffer();
 		void UpdatePositionShockPropagation(const FReal Dt, const int32 It, const int32 NumIts, const int32 BeginIndex, const int32 EndIndex, const FPBDCollisionSolverSettings& SolverSettings);
 		void UpdateVelocityShockPropagation(const FReal Dt, const int32 It, const int32 NumIts, const int32 BeginIndex, const int32 EndIndex, const FPBDCollisionSolverSettings& SolverSettings);
 		void ApplyShockPropagation(const FSolverReal ShockPropagation);
@@ -81,6 +84,7 @@ namespace Chaos
 		// The start of the manifold points array in the scratch buffer
 		Private::FPBDCollisionSolverManifoldPoint* CollisionSolverManifoldPoints;
 		int32 NumCollisionSolverManifoldPoints;
+		int32 MaxCollisionSolverManifoldPoints;
 
 		// Whether we need to run incremental collision for each constraint (LevelSets only now)
 		TArray<bool> bCollisionConstraintPerIterationCollisionDetection;
diff --git a/Engine/Source/Runtime/Experimental/Chaos/Public/Chaos/Collision/PBDCollisionSolver.h b/Engine/Source/Runtime/Experimental/Chaos/Public/Chaos/Collision/PBDCollisionSolver.h
--- a/Engine/Source/Runtime/Experimental/Chaos/Public/Chaos/Collision/PBDCollisionSolver.h
+++ b/Engine/Source/Runtime/Experimental/Chaos/Public/Chaos/Collision/PBDCollisionSolver.h
@@ -75,6 +75,9 @@ namespace Chaos
 			// Equal to (NormalPushOut / TangentialPushOut) before clamping to the friction cone.
 			// Used to move the static friction anchors to the edge of the cone in Scatter.
 			FSolverReal StaticFrictionRatio;
+
+			// Transient - whether to apply friction on the current iteration
+			uint32 bApplyFriction : 1;
 		};
 
 		/**
@@ -87,7 +90,6 @@ namespace Chaos
 		{
 		public:
 			static const int32 MaxConstrainedBodies = 2;
-			static const int32 MaxPointsPerConstraint = 4;
 
 			// Create a solver that is initialized to safe defaults
 			static FPBDCollisionSolver MakeInitialized()
@@ -107,13 +109,13 @@ namespace Chaos
 			FPBDCollisionSolver() {}
 
 			/** Reset the state of the collision solver */
-			void Reset()
+			void Reset(FPBDCollisionSolverManifoldPoint* InManifoldPoints, const int32 InMaxManifoldPoints)
 			{
 				State.SolverBodies[0].Reset();
 				State.SolverBodies[1].Reset();
-				State.ManifoldPoints = nullptr;
+				State.ManifoldPoints = InManifoldPoints;
 				State.NumManifoldPoints = 0;
-				State.MaxManifoldPoints = 0;
+				State.MaxManifoldPoints = InMaxManifoldPoints;
 			}
 
 			void ResetManifold()
@@ -155,20 +157,18 @@ namespace Chaos
 
 			int32 AddManifoldPoint()
 			{
-				check(State.NumManifoldPoints < State.MaxManifoldPoints);
-				return State.NumManifoldPoints++;
-			}
-
-			void SetManifoldPointsBuffer(FPBDCollisionSolverManifoldPoint* InManifoldPoints, const int32 InMaxManifoldPoints)
-			{
-				State.ManifoldPoints = InManifoldPoints;
-				State.MaxManifoldPoints = InMaxManifoldPoints;
-				State.NumManifoldPoints = 0;
+				if (State.NumManifoldPoints < State.MaxManifoldPoints)
+				{
+					return State.NumManifoldPoints++;
+				}
+				return INDEX_NONE;
 			}
 
 			const FPBDCollisionSolverManifoldPoint& GetManifoldPoint(const int32 ManifoldPointIndex) const
 			{
+				check(State.ManifoldPoints != nullptr);
 				check(ManifoldPointIndex < NumManifoldPoints());
+
 				return State.ManifoldPoints[ManifoldPointIndex];
 			}
 
@@ -918,6 +918,8 @@ namespace Chaos
 			const FSolverReal InWorldContactDeltaTangentV,
 			const FSolverReal InWorldContactVelocityTargetNormal)
 		{
+			check(State.ManifoldPoints != nullptr);
+			check(PointIndex < State.NumManifoldPoints);
 			FPBDCollisionSolverManifoldPoint& ManifoldPoint = State.ManifoldPoints[PointIndex];
 
 			ManifoldPoint.RelativeContactPoints[0] = InRelativeContactPosition0;
@@ -958,7 +960,6 @@ namespace Chaos
 		{
 			// Accumulate net pushout for friction limits below
 			bool bApplyFriction = false;
-			bool bApplyPointFriction[MaxPointsPerConstraint] = { false, };
 			FSolverReal NumFrictionContacts = FSolverReal(0);	// NOTE: deliberately not an int
 			FSolverReal TotalPushOutNormal = FSolverReal(0);
 
@@ -982,6 +983,7 @@ namespace Chaos
 
 				// Friction gets updated for any point with a net normal correction or where we have previously had a normal correction and 
 				// already applied friction (in which case we may need to zero it)
+				SolverManifoldPoint.bApplyFriction = false;
 				if (SolverManifoldPoint.NetPushOutNormal > 0)
 				{
 					TotalPushOutNormal += SolverManifoldPoint.NetPushOutNormal;
@@ -990,7 +992,7 @@ namespace Chaos
 				if ((SolverManifoldPoint.NetPushOutNormal > 0) || (SolverManifoldPoint.NetPushOutTangentU != 0) || (SolverManifoldPoint.NetPushOutTangentV != 0))
 				{
 					bApplyFriction = true;
-					bApplyPointFriction[PointIndex] = ((SolverManifoldPoint.NetPushOutNormal > 0) || (SolverManifoldPoint.NetPushOutTangentU != 0) || (SolverManifoldPoint.NetPushOutTangentV != 0));
+					SolverManifoldPoint.bApplyFriction = ((SolverManifoldPoint.NetPushOutNormal > 0) || (SolverManifoldPoint.NetPushOutTangentU != 0) || (SolverManifoldPoint.NetPushOutTangentV != 0));
 				}
 			}
 
@@ -1007,7 +1009,7 @@ namespace Chaos
 				for (int32 PointIndex = 0; PointIndex < NumManifoldPoints(); ++PointIndex)
 				{
 					FPBDCollisionSolverManifoldPoint& SolverManifoldPoint = State.ManifoldPoints[PointIndex];
-					if (bApplyPointFriction[PointIndex])
+					if (SolverManifoldPoint.bApplyFriction)
 					{
 						FSolverReal ContactDeltaTangentU, ContactDeltaTangentV;
 						CalculateContactPositionErrorTangential(SolverManifoldPoint, ContactDeltaTangentU, ContactDeltaTangentV);
