geometry collection particles should not increment sleep/break timers while in a cluster union

diff --git a/Engine/Source/Runtime/Experimental/Chaos/Private/GeometryCollectionProxyData.cpp b/Engine/Source/Runtime/Experimental/Chaos/Private/GeometryCollectionProxyData.cpp
--- a/Engine/Source/Runtime/Experimental/Chaos/Private/GeometryCollectionProxyData.cpp
+++ b/Engine/Source/Runtime/Experimental/Chaos/Private/GeometryCollectionProxyData.cpp
@@ -159,12 +159,18 @@ bool FGeometryCollectionDynamicStateFacade::HasBrokenOff(int32 TransformIndex) c
 
 bool FGeometryCollectionDynamicStateFacade::HasInternalClusterParent(int32 TransformIndex) const
 {
-	const uint8 InternalParenttype = InternalClusterParentTypeAttribute.Get()[TransformIndex];
-	return InternalParenttype != (uint8)Chaos::EInternalClusterType::None;
+	const uint8 InternalParentType = InternalClusterParentTypeAttribute.Get()[TransformIndex];
+	return InternalParentType != (uint8)Chaos::EInternalClusterType::None;
 }
 
 bool FGeometryCollectionDynamicStateFacade::HasDynamicInternalClusterParent(int32 TransformIndex) const
 {
-	const uint8 InternalParenttype = InternalClusterParentTypeAttribute.Get()[TransformIndex];
-	return InternalParenttype == (uint8)Chaos::EInternalClusterType::Dynamic;
+	const uint8 InternalParentType = InternalClusterParentTypeAttribute.Get()[TransformIndex];
+	return InternalParentType == (uint8)Chaos::EInternalClusterType::Dynamic;
+}
+
+bool FGeometryCollectionDynamicStateFacade::HasClusterUnionParent(int32 TransformIndex) const
+{
+	const uint8 InternalParentType = InternalClusterParentTypeAttribute.Get()[TransformIndex];
+	return InternalParentType == (uint8)Chaos::EInternalClusterType::ClusterUnion;
 }
\ No newline at end of file
diff --git a/Engine/Source/Runtime/Experimental/Chaos/Private/PhysicsProxy/GeometryCollectionPhysicsProxy.cpp b/Engine/Source/Runtime/Experimental/Chaos/Private/PhysicsProxy/GeometryCollectionPhysicsProxy.cpp
--- a/Engine/Source/Runtime/Experimental/Chaos/Private/PhysicsProxy/GeometryCollectionPhysicsProxy.cpp
+++ b/Engine/Source/Runtime/Experimental/Chaos/Private/PhysicsProxy/GeometryCollectionPhysicsProxy.cpp
@@ -3094,6 +3094,7 @@ void FGeometryCollectionPhysicsProxy::BufferPhysicsResults_Internal(Chaos::FPBDR
 			StateData.State.DisabledState = Handle->Disabled();
 			StateData.State.HasInternalClusterParent = false;
 			StateData.State.DynamicInternalClusterParent = false;
+			StateData.State.HasClusterUnionParent = false;
 			StateData.State.DynamicState = static_cast<int8>(GetObjectStateFromHandle(Handle));
 
 			bool bHasChanged = false;
@@ -3160,6 +3161,7 @@ void FGeometryCollectionPhysicsProxy::BufferPhysicsResults_Internal(Chaos::FPBDR
 				StateData.State.DisabledState = false;
 				StateData.State.HasInternalClusterParent = true;
 				StateData.State.DynamicInternalClusterParent = (ClusterParent->IsDynamic());
+				StateData.State.HasClusterUnionParent = ClusterParent->PhysicsProxy()->GetType() == EPhysicsProxyType::ClusterUnionProxy;
 
 				bHasChanged = true;
 
@@ -3473,7 +3475,14 @@ bool FGeometryCollectionPhysicsProxy::PullNonInterpolatableDataFromSinglePhysics
 			Chaos::EInternalClusterType ParentType = Chaos::EInternalClusterType::None;
 			if (StateData.State.HasInternalClusterParent != 0)
 			{
-				ParentType = (StateData.State.DynamicInternalClusterParent != 0) ? Chaos::EInternalClusterType::Dynamic : Chaos::EInternalClusterType::KinematicOrStatic;
+				if (StateData.State.HasClusterUnionParent)
+				{
+					ParentType = Chaos::EInternalClusterType::ClusterUnion;
+				}
+				else
+				{
+					ParentType = (StateData.State.DynamicInternalClusterParent != 0) ? Chaos::EInternalClusterType::Dynamic : Chaos::EInternalClusterType::KinematicOrStatic;
+				}
 			}
 			const uint8 ParentTypeUInt8 = static_cast<uint8>(ParentType);
 			if (UpdateValue((*InternalClusterParentTypeArray)[TransformGroupIndex], ParentTypeUInt8))
diff --git a/Engine/Source/Runtime/Experimental/Chaos/Public/Chaos/GeometryParticlesfwd.h b/Engine/Source/Runtime/Experimental/Chaos/Public/Chaos/GeometryParticlesfwd.h
--- a/Engine/Source/Runtime/Experimental/Chaos/Public/Chaos/GeometryParticlesfwd.h
+++ b/Engine/Source/Runtime/Experimental/Chaos/Public/Chaos/GeometryParticlesfwd.h
@@ -105,5 +105,6 @@ namespace Chaos
 		None,
 		KinematicOrStatic,
 		Dynamic,
+		ClusterUnion
 	};
 }
\ No newline at end of file
diff --git a/Engine/Source/Runtime/Experimental/Chaos/Public/GeometryCollectionProxyData.h b/Engine/Source/Runtime/Experimental/Chaos/Public/GeometryCollectionProxyData.h
--- a/Engine/Source/Runtime/Experimental/Chaos/Public/GeometryCollectionProxyData.h
+++ b/Engine/Source/Runtime/Experimental/Chaos/Public/GeometryCollectionProxyData.h
@@ -135,6 +135,9 @@ public:
 
 	/** return true if the transform has an internal cluster parent in a dynamic state */
 	CHAOS_API bool HasDynamicInternalClusterParent(int32 TransformIndex) const;
+
+	/** Return true if the transform has a cluster union parent. */
+	CHAOS_API bool HasClusterUnionParent(int32 TransformIndex) const;
 	
 private:
 	/** Active state, true means that the transform is active or broken off from its parent */
@@ -343,7 +346,8 @@ public:
 		uint16 DisabledState: 1;
 		uint16 HasInternalClusterParent: 1;
 		uint16 DynamicInternalClusterParent: 1;
-		// 5 bits left
+		uint16 HasClusterUnionParent: 1;
+		// 4 bits left
 	};
 
 	struct FStateData
diff --git a/Engine/Source/Runtime/Experimental/GeometryCollectionEngine/Private/GeometryCollection/GeometryCollectionComponent.cpp b/Engine/Source/Runtime/Experimental/GeometryCollectionEngine/Private/GeometryCollection/GeometryCollectionComponent.cpp
--- a/Engine/Source/Runtime/Experimental/GeometryCollectionEngine/Private/GeometryCollection/GeometryCollectionComponent.cpp
+++ b/Engine/Source/Runtime/Experimental/GeometryCollectionEngine/Private/GeometryCollection/GeometryCollectionComponent.cpp
@@ -5043,7 +5043,7 @@ void UGeometryCollectionComponent::IncrementSleepTimer(float DeltaTime)
 				{
 					// this children has an dynamic internal cluster parent so it can't be removed but we need tyo process the internal cluster by looking at the original parent properties
 					const int32 OriginalParentIdx = OriginalParents[TransformIdx];
-					const bool HasDynamicInternalClusterParent = DynamicStateFacade.HasDynamicInternalClusterParent(TransformIdx);
+					const bool HasDynamicInternalClusterParent = DynamicStateFacade.HasDynamicInternalClusterParent(TransformIdx) && DynamicStateFacade.HasClusterUnionParent(TransformIdx);
 					if (OriginalParentIdx > INDEX_NONE && HasDynamicInternalClusterParent && RemoveOnSleepFacade.IsRemovalActive(OriginalParentIdx))
 					{
 						const bool UseClusterCrumbling = true; // with sleep removal : internal clusters always crumble - this will change when we merge the removal feature together
@@ -5113,7 +5113,7 @@ void UGeometryCollectionComponent::IncrementBreakTimer(float DeltaTime)
 				{
 					// this children has an internal cluster parent so it can't be removed but we need tyo process the internal cluster by looking at the original parent properties
 					const int32 OriginalParentIdx = OriginalParents[TransformIdx];
-					const bool HasDynamicInternalClusterParent = DynamicStateFacade.HasDynamicInternalClusterParent(TransformIdx);
+					const bool HasDynamicInternalClusterParent = DynamicStateFacade.HasDynamicInternalClusterParent(TransformIdx) && !DynamicStateFacade.HasClusterUnionParent(TransformIdx);
 
 					if (OriginalParentIdx > INDEX_NONE && HasDynamicInternalClusterParent && RemoveOnBreakFacade.IsRemovalActive(OriginalParentIdx))
 					{
