Chaos - use a large manifold for convex-mesh collisions
- removes duplicate and invalid points but does not attempt to reduce the manifold to 4 contacts
- optimize data layout in FContactTriangleCollector to reduce cost of PruneUnnecessaryContactPoints (which still has an O(N^2) loop in it)

diff --git a/Engine/Source/Runtime/Experimental/Chaos/Private/Chaos/Collision/ContactTriangles.cpp b/Engine/Source/Runtime/Experimental/Chaos/Private/Chaos/Collision/ContactTriangles.cpp
--- a/Engine/Source/Runtime/Experimental/Chaos/Private/Chaos/Collision/ContactTriangles.cpp
+++ b/Engine/Source/Runtime/Experimental/Chaos/Private/Chaos/Collision/ContactTriangles.cpp
@@ -9,16 +9,70 @@ namespace Chaos
 		extern int32 ChaosSolverDebugDrawMeshContacts;
 	}
 	extern bool bChaos_Collision_EnableEdgePrune;
+	extern int32 Chaos_Collision_ReduceMeshManifoldThreshold;
+	extern FRealSingle Chaos_Collision_MeshContactNormalThreshold;
 
 
 	void FContactTriangleCollector::SetNumContacts(const int32 NumContacts)
 	{
 		TriangleContactPoints.SetNum(NumContacts, false);
+		TriangleContactPointDatas.SetNum(NumContacts, false);
 	}
 
-	void FContactTriangleCollector::RemoveContact(const int32 ContactIndex)
+	void FContactTriangleCollector::DisableContact(const int32 ContactIndex)
 	{
-		TriangleContactPoints.RemoveAtSwap(ContactIndex, 1, false);
+		// Flag the contact as disabled. A subsequent call to RemoveDisabledContacts will repack the array
+		// by removing all disabled contacts, without re-ordering
+		TriangleContactPointDatas[ContactIndex].SetDisabled();
+		++NumDisabledTriangleContactPoints;
+	}
+
+	void FContactTriangleCollector::RemoveDisabledContacts()
+	{
+		// Skip if we have not disabled any contacts
+		if (NumDisabledTriangleContactPoints == 0)
+		{
+			return;
+		}
+
+		// Re-pack the contact point array without re-ordering
+		const int32 NumContactPoints = TriangleContactPoints.Num();
+		int32 DestContactIndex = 0;
+		int32 SrcContactIndex = 0;
+		while (SrcContactIndex < NumContactPoints)
+		{
+			if (!TriangleContactPointDatas[SrcContactIndex].IsEnabled())
+			{
+				// Inner loop to find the next enabled point index to use as the copy source
+				while (++SrcContactIndex < NumContactPoints)
+				{
+					if (TriangleContactPointDatas[SrcContactIndex].IsEnabled())
+					{
+						break;
+					}
+				}
+
+				if (SrcContactIndex == NumContactPoints)
+				{
+					break;
+				}
+			}
+
+			// If we have removed elements, copy source to dest
+			if (DestContactIndex != SrcContactIndex)
+			{
+				TriangleContactPointDatas[DestContactIndex] = TriangleContactPointDatas[SrcContactIndex];
+				TriangleContactPoints[DestContactIndex] = TriangleContactPoints[SrcContactIndex];
+			}
+
+			++DestContactIndex;
+			++SrcContactIndex;
+		}
+
+		// Clip the array to the enabled set
+		TriangleContactPointDatas.SetNum(DestContactIndex, false);
+		TriangleContactPoints.SetNum(DestContactIndex, false);
+		NumDisabledTriangleContactPoints = 0;
 	}
 
 	int32 FContactTriangleCollector::AddTriangle(const FTriangle& Triangle, const int32 TriangleIndex, const int32 VertexIndex0, const int32 VertexIndex1, const int32 VertexIndex2)
@@ -92,8 +146,9 @@ namespace Chaos
 
 		for (int32 ContactIndex = 0; ContactIndex < NumContactPoints; ++ContactIndex)
 		{
-			FTriangleContactPoint& ContactPoint = TriangleContactPoints[ContactIndex];
-			const FContactTriangle& ContactTriangle = ContactTriangles[ContactPoint.ContactTriangleIndex];
+			FTriangleContactPointData& ContactPointData = TriangleContactPointDatas[ContactIndex];
+			const FContactPoint& ContactPoint = TriangleContactPoints[ContactIndex];
+			const FContactTriangle& ContactTriangle = ContactTriangles[ContactPointData.GetTriangleIndex()];
 
 			if (ContactPoint.ContactType == EContactPointType::VertexPlane)
 			{
@@ -118,12 +173,12 @@ namespace Chaos
 						ContactVertices.Add(FContactVertexID(VertexIndexB));
 
 						// Store the edge ID so we can see if any other faces add it to the don'tcollide list
-						ContactPoint.EdgeID = FContactEdgeID(VertexIndexA, VertexIndexB);
+						ContactPointData.SetEdgeID(VertexIndexA, VertexIndexB);
 					}
 					else if (VertexIndexA != INDEX_NONE)
 					{
 						// Store the vertex ID so we can see if any other faces or edges add it to the don'tcollide list
-						ContactPoint.VertexID = FContactVertexID(VertexIndexA);
+						ContactPointData.SetVertexID(VertexIndexA);
 					}
 				}
 			}
@@ -132,15 +187,20 @@ namespace Chaos
 
 	void FContactTriangleCollector::ProcessContacts(const FRigidTransform3& MeshToConvexTransform)
 	{
+		// @todo(chaos): Some of the pruning mechanisms don't work very well now that we have support for large manifolds.
+		const bool bEnableFullPruning = ((TriangleContactPoints.Num() - NumDisabledTriangleContactPoints) < Chaos_Collision_ReduceMeshManifoldThreshold);
+
 		DebugDrawContactPoints(FColor::White, 0.2);
 
 		// Build the feature set from the full list of contacts before pruning
 		BuildContactFeatureSets();
 
 		// Remove edge contacts that should not be possible because we are also colliding with a face that used the edge.
-		if (TriangleContactPoints.Num() > 1)
+		if ((TriangleContactPoints.Num() > 1))
 		{
-			PruneEdgeContactPoints();
+			const bool bPruneEdges = true;
+			const bool bPruneVertices = bEnableFullPruning;	// @todo(chaos): fix - when set, removes too many interior contacts
+			PruneEdgeAndVertexContactPoints(bPruneEdges, bPruneVertices);
 		}
 
 		DebugDrawContactPoints(FColor::Cyan, 0.35);
@@ -152,33 +212,32 @@ namespace Chaos
 			PruneInfacingContactPoints();
 		}
 
-		SortContactPointsForPruning();
-
-		// Remove close contacts and those not near the deepest Phi
-		// NOTE: relies on the sort above, and retains order
-		if (TriangleContactPoints.Num() > 4)
+		// Remove close contacts
+		if (TriangleContactPoints.Num() > 1)
 		{
 			PruneUnnecessaryContactPoints();
 		}
 
-		DebugDrawContactPoints(FColor::Yellow, 0.5);
-
 		// Reduce to only 4 contact points from here
-		// NOTE: relies on the sort above and may not retain order
-		if (TriangleContactPoints.Num() > 4)
+		// @todo(chaos): A 4 point manifold is often not sufficient when we have low curvature and we end up removing contacts that
+		// would be required during the solve to prevent jitter. For now we disable the reduction when we are colliding
+		// with "lots" of triangles (large objects relative to triangle size), but keep the original behaviour for 
+		// smaller objects so we don't impact perf for the most common cases until we do this properly.
+		if (bEnableFullPruning)
 		{
+			DebugDrawContactPoints(FColor::Yellow, 0.5);
+
 			ReduceManifoldContactPointsTriangeMesh();
 		}
 
+		DebugDrawContactPoints(FColor::Green, 0.75);
+
 		// Fix contact normals for edge and vertex collisions
-		if (TriangleContactPoints.Num() > 1)
+		if (TriangleContactPoints.Num() > 0)
 		{
 			FixInvalidNormalContactPoints();
 		}
 
-		// Sort the contacts so that faces are processed before edges, edges before vertices, and then deepest are processed first
-		SortContactPointsForSolving();
-
 		DebugDrawContactPoints(FColor::Red, 1.0);
 
 		FinalizeContacts(MeshToConvexTransform);
@@ -186,16 +245,40 @@ namespace Chaos
 
 	void FContactTriangleCollector::SortContactPointsForPruning()
 	{
+		// Repack the array
+		RemoveDisabledContacts();
+
+		// Sort TriangleContactPoints and TriangleContactPointDatas for pruning
 		if (TriangleContactPoints.Num() > 1)
 		{
+			TArray<int32> SortedIndices;
+			SortedIndices.SetNumUninitialized(TriangleContactPoints.Num());
+			for (int32 Index = 0; Index < TriangleContactPoints.Num(); ++Index)
+			{
+				SortedIndices[Index] = Index;
+			}
+
+			// Sort by depth
 			std::sort(
-				&TriangleContactPoints[0],
-				&TriangleContactPoints[0] + TriangleContactPoints.Num(),
-				[](const FTriangleContactPoint& L, const FTriangleContactPoint& R)
+				&SortedIndices[0],
+				&SortedIndices[0] + SortedIndices.Num(),
+				[this](const int32& L, const int32& R)
 				{
-					return L.Phi < R.Phi;
+					return TriangleContactPoints[L].Phi < TriangleContactPoints[R].Phi;
 				}
 			);
+
+			TArray<FContactPoint> SortedContactPoints;
+			TArray<FTriangleContactPointData> SortedContactPointDatas;
+			SortedContactPoints.SetNum(TriangleContactPoints.Num());
+			SortedContactPointDatas.SetNum(TriangleContactPoints.Num());
+			for (int32 Index = 0; Index < TriangleContactPoints.Num(); ++Index)
+			{
+				SortedContactPoints[Index] = TriangleContactPoints[SortedIndices[Index]];
+				SortedContactPointDatas[Index] = TriangleContactPointDatas[SortedIndices[Index]];
+			}
+			Swap(TriangleContactPoints, SortedContactPoints);
+			Swap(TriangleContactPointDatas, SortedContactPointDatas);
 		}
 	}
 
@@ -203,12 +286,22 @@ namespace Chaos
 	// NOTE: This relies on the enum order of EContactPointType.
 	void FContactTriangleCollector::SortContactPointsForSolving()
 	{
+		// Repack the array
+		RemoveDisabledContacts();
+
+		// Sort TriangleContactPoints in solver preferred order, but ignore TriangleContactPointDatas
+		// NOTE: This should only be called at the end of the pruning proxess when we no longer care 
+		// about TriangleContactPointDatas.
 		if (TriangleContactPoints.Num() > 1)
 		{
+			// Reset metadata to ensure we don't access it again
+			TriangleContactPointDatas.Reset();
+
+			// Sort by contact type (in enum order) then phi
 			std::sort(
 				&TriangleContactPoints[0],
 				&TriangleContactPoints[0] + TriangleContactPoints.Num(),
-				[](const FTriangleContactPoint& L, const FTriangleContactPoint& R)
+				[](const FContactPoint& L, const FContactPoint& R)
 				{
 					if (L.ContactType == R.ContactType)
 					{
@@ -220,39 +313,37 @@ namespace Chaos
 		}
 	}
 
-	// For each edge contact, find the valid range of normal and reject the contact if the normal is outside the range.
-	// This rejects collisions against concave edges, as well as edge contacts that should be hidden by face contacts on adjacent triangles.
-	// NOTE: This only works if both triangle sharing the edge are in the ContactPointTriangles list. They usually will be because you won't collide with an edge unless you overlap the adjacent triangle.
-	// NOTE: Assumes the ContactPoints array has both contacts in the same space (Tri Mesh and Heightfields are this way until the final step)
-	void FContactTriangleCollector::PruneEdgeContactPoints()
+	void FContactTriangleCollector::PruneEdgeAndVertexContactPoints(const bool bPruneEdges, const bool bPruneVertices)
 	{
 		if (!bChaos_Collision_EnableEdgePrune)
 		{
 			return;
 		}
 
-		// Remove all vertex collisions when we have also an edge collision that includes the vertex
-		// Remove all edge collisions when we have a face collision that includes the edge
-		int32 NumContactPoints = TriangleContactPoints.Num();
-		for (int32 ContactIndex = 0; ContactIndex < NumContactPoints; ++ContactIndex)
+		// Remove contacts that are invalid edge or vertex contacts
+		// - vertex collisions if we have an edge collision that includes the vertex
+		// - edge collisions if we have a face collision that includes the edge
+		for (int32 ContactIndex = TriangleContactPointDatas.Num() - 1; ContactIndex >= 0; --ContactIndex)
 		{
-			const FTriangleContactPoint& ContactPoint = TriangleContactPoints[ContactIndex];
+			FTriangleContactPointData& ContactPointData = TriangleContactPointDatas[ContactIndex];
+			if (!ContactPointData.IsEnabled())
+			{
+				continue;
+			}
 
 			bool bRemoveContact = false;
-			if (ContactPoint.VertexID != INDEX_NONE)
+			if (bPruneVertices && ContactPointData.IsVertex())
 			{
-				bRemoveContact = ContactVertices.Contains(ContactPoint.VertexID);
+				bRemoveContact = ContactVertices.Contains(ContactPointData.GetVertexID());
 			}
-			else if (ContactPoint.EdgeID.IsValid())
+			else if (bPruneEdges && ContactPointData.IsEdge())
 			{
-				bRemoveContact = ContactEdges.Contains(ContactPoint.EdgeID);
+				bRemoveContact = ContactEdges.Contains(ContactPointData.GetEdgeID());
 			}
 
 			if (bRemoveContact)
 			{
-				RemoveContact(ContactIndex);
-				--ContactIndex;
-				--NumContactPoints;
+				DisableContact(ContactIndex);
 			}
 		}
 	}
@@ -264,14 +355,19 @@ namespace Chaos
 		const int32 NumContactPoints = TriangleContactPoints.Num();
 		for (int32 ContactIndex = 0; ContactIndex < NumContactPoints; ++ContactIndex)
 		{
-			FTriangleContactPoint& ContactPoint = TriangleContactPoints[ContactIndex];
+			FTriangleContactPointData& ContactPointData = TriangleContactPointDatas[ContactIndex];
+			FContactPoint& ContactPoint = TriangleContactPoints[ContactIndex];
+			if (!ContactPointData.IsEnabled())
+			{
+				continue;
+			}
 
 			// Do we have a collision with the triangle edge or vertex?
 			// NOTE: second shape is always the triangle
 			const bool bIsTriangleEdgeOrVertex = (ContactPoint.ContactType == EContactPointType::EdgeEdge) || (ContactPoint.ContactType == EContactPointType::PlaneVertex);
 			if (bIsTriangleEdgeOrVertex)
 			{
-				const int32 ContactTriangleIndex = ContactPoint.ContactTriangleIndex;
+				const int32 ContactTriangleIndex = ContactPointData.GetTriangleIndex();
 				const FContactTriangle& ContactTriangle = ContactTriangles[ContactTriangleIndex];
 
 				// Find the vertex positions and indices for this contact
@@ -298,7 +394,7 @@ namespace Chaos
 								// NOTE: We keep the contact depth as it is because we know that the depth is a lower-bound. 
 								// We have to update the ShapeContactPoint[0] because Phi is actually derived from the positions (the value in ContactPoint is just a cache of current state)
 								// @todo(chaos): face selection logic might be better if we knew the contact velocity
-								const FReal OtherContactDotNormal = FVec3::DotProduct(ContactPoint.ShapeContactNormal, ContactTriangle.FaceNormal);
+								const FReal OtherContactDotNormal = FVec3::DotProduct(ContactPoint.ShapeContactNormal, ContactTriangle.FaceNormal);		// @todo(chaos): OtherContactTriangle??
 								if (ContactDotNormal >= OtherContactDotNormal)
 								{
 									ContactPoint.ShapeContactNormal = (ContactDotNormal > -SMALL_NUMBER) ? ContactTriangle.FaceNormal : -ContactTriangle.FaceNormal;
@@ -308,7 +404,7 @@ namespace Chaos
 								{
 									ContactPoint.ShapeContactNormal = (OtherContactDotNormal > -SMALL_NUMBER) ? OtherContactTriangle.FaceNormal : -OtherContactTriangle.FaceNormal;
 									ContactPoint.ShapeContactPoints[0] = ContactPoint.ShapeContactPoints[1] + ContactPoint.Phi * ContactPoint.ShapeContactNormal;
-									ContactPoint.ContactTriangleIndex = OtherContactTriangleIndex;
+									ContactPointData.SetTriangleIndex(OtherContactTriangleIndex);
 								}
 							}
 						}
@@ -334,7 +430,7 @@ namespace Chaos
 										const FReal NormalDotCentroid = FVec3::DotProduct(ContactPoint.ShapeContactNormal, Centroid - ContactPoint.ShapeContactPoints[1]);
 										if (NormalDotCentroid > 0)
 										{
-											const FReal OtherContactDotNormal = FVec3::DotProduct(ContactPoint.ShapeContactNormal, ContactTriangle.FaceNormal);
+											const FReal OtherContactDotNormal = FVec3::DotProduct(ContactPoint.ShapeContactNormal, ContactTriangle.FaceNormal);	// @todo(chaos): OtherContactTriangle??
 											ContactPoint.ShapeContactNormal = (OtherContactDotNormal > -SMALL_NUMBER) ? OtherContactTriangle.FaceNormal : -OtherContactTriangle.FaceNormal;
 											ContactPoint.ShapeContactPoints[0] = ContactPoint.ShapeContactPoints[1] + ContactPoint.Phi * ContactPoint.ShapeContactNormal;
 										}
@@ -351,71 +447,85 @@ namespace Chaos
 	// Remove contact points that have a normal pointing inthe opposite direction of the face normal
 	void FContactTriangleCollector::PruneInfacingContactPoints()
 	{
-		for (int32 ContactIndex = 0; ContactIndex < TriangleContactPoints.Num(); ++ContactIndex)
+		for (int32 ContactIndex = TriangleContactPointDatas.Num() - 1; ContactIndex >= 0; --ContactIndex)
 		{
-			const FTriangleContactPoint& ContactPoint = TriangleContactPoints[ContactIndex];
-			const FContactTriangle& ContactTriangle = ContactTriangles[ContactPoint.ContactTriangleIndex];
+			FTriangleContactPointData& ContactPointData = TriangleContactPointDatas[ContactIndex];
+			const FContactPoint& ContactPoint = TriangleContactPoints[ContactIndex];
+			const FContactTriangle& ContactTriangle = ContactTriangles[ContactPointData.GetTriangleIndex()];
 
-			const FReal ContactNormalDotFaceNormal = FVec3::DotProduct(ContactPoint.ShapeContactNormal, ContactTriangle.FaceNormal);
-			if (ContactNormalDotFaceNormal < 0)
+			if (ContactPointData.IsEnabled())
 			{
-				RemoveContact(ContactIndex);
-				--ContactIndex;
+				const FReal ContactNormalDotFaceNormal = FVec3::DotProduct(ContactPoint.ShapeContactNormal, ContactTriangle.FaceNormal);
+				if (ContactNormalDotFaceNormal < 0)
+				{
+					DisableContact(ContactIndex);
+				}
 			}
 		}
 	}
 
-	// Remove contact points that are much more shallow that the deepest
-	// Remove contact points that are too close to others.
-	// NOTE: ContactPoints must be sorted on Phi on entry, and will be sorted on exit
-	// @todo(chaos): we should only remove shallow contacts that are hidden by nearby deeper contacts. It should be ok to have shallow contacts
-	// far from the deeper ones. E.g., a box landing on plane at an angle.
 	void FContactTriangleCollector::PruneUnnecessaryContactPoints()
 	{
-		// Remove all points except for the deepest one, and ones with phis similar to it
-		int32 NewContactPointCount = TriangleContactPoints.Num() > 0 ? 1 : 0;
-		for (int32 Index = 1; Index < TriangleContactPoints.Num(); Index++)
+		if (TriangleContactPoints.Num() < 2)
 		{
-			if ((TriangleContactPoints[Index].Phi < 0) || ((TriangleContactPoints[Index].Phi - TriangleContactPoints[0].Phi) < PhiTolerance))
-			{
-				NewContactPointCount++;
-			}
-			else
-			{
-				break;
-			}
+			return;
 		}
-		SetNumContacts(NewContactPointCount);
+
+		// Below is O(N^2) so packing contacts is beneficial when there are lots of contacts
+		RemoveDisabledContacts();
+
+		const int32 NumContactPoints = TriangleContactPointDatas.Num();
 
 		// Remove points that are very close together. Always keep the one that is least likely to be an edge contact
+		// @todo(chaos): Use a grid to remove the O(N^2) loop
 		const FReal DistanceToleranceSq = FMath::Square(DistanceTolerance);
-		for (int32 ContactIndex0 = 0; ContactIndex0 < TriangleContactPoints.Num(); ++ContactIndex0)
+		for (int32 ContactIndex0 = 0; ContactIndex0 < NumContactPoints; ++ContactIndex0)
 		{
-			const FTriangleContactPoint& ContactPoint0 = TriangleContactPoints[ContactIndex0];
-			const FContactTriangle& ContactTriangle0 = ContactTriangles[ContactPoint0.ContactTriangleIndex];
+			FTriangleContactPointData& ContactPointData0 = TriangleContactPointDatas[ContactIndex0];
+			const FContactPoint& ContactPoint0 = TriangleContactPoints[ContactIndex0];
+			const FContactTriangle& ContactTriangle0 = ContactTriangles[ContactPointData0.GetTriangleIndex()];
+			if (!ContactPointData0.IsEnabled())
+			{
+				continue;
+			}
 
-			for (int32 ContactIndex1 = ContactIndex0 + 1; ContactIndex1 < TriangleContactPoints.Num(); ++ContactIndex1)
+			// Remove duplicate contact points (each triangle generates a manifold so the vertex and edge contacts are typically duplicated)
+			for (int32 ContactIndex1 = ContactIndex0 + 1; ContactIndex1 < NumContactPoints; ++ContactIndex1)
 			{
-				const FTriangleContactPoint& ContactPoint1 = TriangleContactPoints[ContactIndex1];
-				const FContactTriangle& ContactTriangle1 = ContactTriangles[ContactPoint1.ContactTriangleIndex];
+				FTriangleContactPointData& ContactPointData1 = TriangleContactPointDatas[ContactIndex1];
+				const FContactPoint& ContactPoint1 = TriangleContactPoints[ContactIndex1];
+				const FContactTriangle& ContactTriangle1 = ContactTriangles[ContactPointData1.GetTriangleIndex()];
+				if (!ContactPointData1.IsEnabled())
+				{
+					continue;
+				}
+
+				// Are these two points at the same position? They must be at the same vertex or edge for this to be true.
+				// If both points are from the same triangle vertex, we know they are coincident (but normals may be different)
+				const bool bIsSameVertex = ContactPointData0.IsVertex() && (ContactPointData0.GetVertexID() == ContactPointData1.GetVertexID());
+				bool bIsSamePosition = bIsSameVertex;
+				if (!bIsSameVertex)
+				{
+					// Check distance and normal
+					// NOTE: ShapeContactPoints[1] is on the triangle
+					const FVec3 ContactDelta = ContactPoint0.ShapeContactPoints[1] - ContactPoint1.ShapeContactPoints[1];
+					const FReal ContactDeltaSq = ContactDelta.SizeSquared();
+					bIsSamePosition = (ContactDeltaSq < DistanceToleranceSq);
+				}
 
-				// NOTE: ShapeContactPoints[1] is on the triangle
-				const FVec3 ContactDelta = ContactPoint0.ShapeContactPoints[1] - ContactPoint1.ShapeContactPoints[1];
-				const FReal ContactDeltaSq = ContactDelta.SizeSquared();
-				if (ContactDeltaSq < DistanceToleranceSq)
+				if (bIsSamePosition)
 				{
+					// Keep the contact with the normal that is most along the triangle face normal
 					const FReal NormalDot0 = FVec3::DotProduct(ContactPoint0.ShapeContactNormal, ContactTriangle0.FaceNormal);
 					const FReal NormalDot1 = FVec3::DotProduct(ContactPoint1.ShapeContactNormal, ContactTriangle1.FaceNormal);
 					if (FMath::Abs(NormalDot0) >= FMath::Abs(NormalDot1))
 					{
-						TriangleContactPoints.RemoveAt(ContactIndex1, 1, false);
-						--ContactIndex1;
+						DisableContact(ContactIndex1);
 						continue;
 					}
 					else
 					{
-						TriangleContactPoints.RemoveAt(ContactIndex0, 1, false);
-						--ContactIndex0;
+						DisableContact(ContactIndex0);
 						break;
 					}
 				}
@@ -424,10 +534,15 @@ namespace Chaos
 	}
 
 	// Reduce the number of contact points (in place)
-	// Prerequisites to calling this function:
-	// ContactPoints are sorted on phi (ascending)
 	void FContactTriangleCollector::ReduceManifoldContactPointsTriangeMesh()
 	{
+		// Sort contact points on phi (ascending)
+		SortContactPointsForPruning();
+
+		// Make sure we packed the contacts
+		check(NumDisabledTriangleContactPoints == 0);
+
+		// NOTE: The above pack and sort may reduce the point count to below our threshold
 		if (TriangleContactPoints.Num() <= 4)
 		{
 			return;
@@ -465,6 +580,7 @@ namespace Chaos
 			}
 			// Farthest point will be added now
 			Swap(TriangleContactPoints[1], TriangleContactPoints[FarthestPointIndex]);
+			Swap(TriangleContactPointDatas[1], TriangleContactPointDatas[FarthestPointIndex]);
 		}
 
 		// Point 3) Largest triangle area
@@ -483,10 +599,13 @@ namespace Chaos
 			}
 			// Point causing the largest triangle will be added now
 			Swap(TriangleContactPoints[2], TriangleContactPoints[LargestTrianglePointIndex]);
+			Swap(TriangleContactPointDatas[2], TriangleContactPointDatas[LargestTrianglePointIndex]);
+
 			// Ensure the winding order is consistent
 			if (LargestTrianglePointSignedArea < 0)
 			{
 				Swap(TriangleContactPoints[0], TriangleContactPoints[1]);
+				Swap(TriangleContactPointDatas[0], TriangleContactPointDatas[1]);
 			}
 		}
 
@@ -508,6 +627,7 @@ namespace Chaos
 			}
 			// Point causing the largest positive triangle area will be added now
 			Swap(TriangleContactPoints[3], TriangleContactPoints[LargestTrianglePointIndex]);
+			Swap(TriangleContactPointDatas[3], TriangleContactPointDatas[LargestTrianglePointIndex]);
 		}
 
 		// Revisit Point 0 - if we have a choice that increases area at similar depth, use it
@@ -518,18 +638,16 @@ namespace Chaos
 
 	void FContactTriangleCollector::FinalizeContacts(const FRigidTransform3& MeshToConvexTransform)
 	{
-		FinalContactPoints.Reset(TriangleContactPoints.Num());
+		// Remove any disabled contacts and sort
+		SortContactPointsForSolving();
 
-		// Transform contact data back into shape-local space and build the output array
+		// Transform contact data back into shape-local space
 		for (int32 ContactIndex = 0; ContactIndex < TriangleContactPoints.Num(); ++ContactIndex)
 		{
-			FTriangleContactPoint& ContactPoint = TriangleContactPoints[ContactIndex];
-			const FContactTriangle& ContactTriangle = ContactTriangles[ContactPoint.ContactTriangleIndex];
+			FContactPoint& ContactPoint = TriangleContactPoints[ContactIndex];
 
 			ContactPoint.ShapeContactPoints[1] = MeshToConvexTransform.InverseTransformPositionNoScale(ContactPoint.ShapeContactPoints[1]);
 			ContactPoint.ShapeContactNormal = MeshToConvexTransform.InverseTransformVectorNoScale(ContactPoint.ShapeContactNormal);
-
-			FinalContactPoints.Add(ContactPoint);
 		}
 	}
 
@@ -543,8 +661,15 @@ namespace Chaos
 
 			const FReal Duration = 0;
 			const int32 DrawPriority = 10;
-			for (const FTriangleContactPoint& ContactPoint : TriangleContactPoints)
+			for (int32 ContactIndex = 0; ContactIndex < TriangleContactPointDatas.Num(); ++ContactIndex)
 			{
+				const FTriangleContactPointData& ContactPointData = TriangleContactPointDatas[ContactIndex];
+				const FContactPoint& ContactPoint = TriangleContactPoints[ContactIndex];
+				if (!ContactPointData.IsEnabled())
+				{
+					continue;
+				}
+
 				const FVec3 P0 = ConvexTransform.TransformPosition(ContactPoint.ShapeContactPoints[0]);
 				const FVec3 P1 = ConvexTransform.TransformPosition(ContactPoint.ShapeContactPoints[1]);
 				const FVec3 N = ConvexTransform.TransformVectorNoScale(ContactPoint.ShapeContactNormal);
@@ -555,16 +680,16 @@ namespace Chaos
 				// Draw a thin black line connecting the two contact points (triangle face to convex surface)
 				//FDebugDrawQueue::GetInstance().DrawDebugLine(P0, P1, FColor::Black, false, FRealSingle(Duration), DrawPriority, 0.5f * FRealSingle(LineThickness));
 
-				if (!bTriangleDrawn[ContactPoint.ContactTriangleIndex])
+				if (!bTriangleDrawn[ContactPointData.GetTriangleIndex()])
 				{
-					FContactTriangle& Triangle = ContactTriangles[ContactPoint.ContactTriangleIndex];
+					FContactTriangle& Triangle = ContactTriangles[ContactPointData.GetTriangleIndex()];
 					const FVec3 V0 = ConvexTransform.TransformPosition(Triangle.Vertices[0]);
 					const FVec3 V1 = ConvexTransform.TransformPosition(Triangle.Vertices[1]);
 					const FVec3 V2 = ConvexTransform.TransformPosition(Triangle.Vertices[2]);
 					FDebugDrawQueue::GetInstance().DrawDebugLine(V0, V1, FColor::Silver, false, FRealSingle(Duration), DrawPriority, FRealSingle(LineScale));
 					FDebugDrawQueue::GetInstance().DrawDebugLine(V1, V2, FColor::Silver, false, FRealSingle(Duration), DrawPriority, FRealSingle(LineScale));
 					FDebugDrawQueue::GetInstance().DrawDebugLine(V2, V0, FColor::Silver, false, FRealSingle(Duration), DrawPriority, FRealSingle(LineScale));
-					bTriangleDrawn[ContactPoint.ContactTriangleIndex] = true;
+					bTriangleDrawn[ContactPointData.GetTriangleIndex()] = true;
 				}
 			}
 		}
diff --git a/Engine/Source/Runtime/Experimental/Chaos/Private/Chaos/CollisionOneShotManifolds.cpp b/Engine/Source/Runtime/Experimental/Chaos/Private/Chaos/CollisionOneShotManifolds.cpp
--- a/Engine/Source/Runtime/Experimental/Chaos/Private/Chaos/CollisionOneShotManifolds.cpp
+++ b/Engine/Source/Runtime/Experimental/Chaos/Private/Chaos/CollisionOneShotManifolds.cpp
@@ -82,6 +82,12 @@ namespace Chaos
 	FAutoConsoleVariableRef CVarChaos_Collision_EnableEdgePrune(TEXT("p.Chaos.Collision.EnableEdgePrune"), bChaos_Collision_EnableEdgePrune, TEXT(""));
 	FAutoConsoleVariableRef CVarChaos_Collision_EdgePrunePlaneDistance(TEXT("p.Chaos.Collision.EdgePrunePlaneDistance"), Chaos_Collision_EdgePrunePlaneDistance, TEXT(""));
 
+	int32 Chaos_Collision_ReduceMeshManifoldThreshold = 1000;	// Effectively Disabled
+	FAutoConsoleVariableRef CVarChaos_Collision_ReduceMeshManifoldThreshold(TEXT("p.Chaos.Collision.ReduceMeshManifoldThreshold"), Chaos_Collision_ReduceMeshManifoldThreshold, TEXT("If an object has more contacts with a mesh than this, do not attempt to reduce the manifold"));
+
+	FRealSingle Chaos_Collision_MeshContactNormalThreshold = 0.98f;	// ~10deg
+	FAutoConsoleVariableRef CVarChaos_Collision_MeshContactNormalThreshold(TEXT("p.Chaos.Collision.MeshContactNormalThreshold"), Chaos_Collision_MeshContactNormalThreshold, TEXT("Treat contact with a dot product between the normal and the triangle face greater than this as face collisions"));
+
 	// Whether to use the new index-less GJK. 
 	// @todo(chaos): This should be removed once soaked for a bit (enabled 7 June 2022)
 	bool bChaos_Collision_UseGJK2 = false;
diff --git a/Engine/Source/Runtime/Experimental/Chaos/Public/Chaos/Collision/ContactTriangles.h b/Engine/Source/Runtime/Experimental/Chaos/Public/Chaos/Collision/ContactTriangles.h
--- a/Engine/Source/Runtime/Experimental/Chaos/Public/Chaos/Collision/ContactTriangles.h
+++ b/Engine/Source/Runtime/Experimental/Chaos/Public/Chaos/Collision/ContactTriangles.h
@@ -118,16 +118,18 @@ namespace Chaos
 		}
 	};
 
+	// ID for a vertex in a trangle mesh, used by FTriangleContactPointData
 	using FContactVertexID = int32;
 
+	// ID for an edge in a trangle mesh, used by FTriangleContactPointData
 	struct FContactEdgeID
 	{
-		FContactEdgeID()
+		FORCEINLINE FContactEdgeID()
 			: VertexIDs{ INDEX_NONE, INDEX_NONE }
 		{
 		}
 
-		FContactEdgeID(const FContactVertexID VertexIndexA, const FContactVertexID VertexIndexB)
+		FORCEINLINE FContactEdgeID(const FContactVertexID VertexIndexA, const FContactVertexID VertexIndexB)
 		{
 			// EdgeID is the same if we swap the vertex indices
 			if (VertexIndexA < VertexIndexB)
@@ -142,27 +144,27 @@ namespace Chaos
 			}
 		}
 
-		bool IsValid() const
+		FORCEINLINE bool IsValid() const
 		{
 			return (VertexIDs[0] != INDEX_NONE) && (VertexIDs[1] != INDEX_NONE);
 		}
 
-		friend bool operator==(const FContactEdgeID& L, const FContactEdgeID& R)
+		FORCEINLINE friend bool operator==(const FContactEdgeID& L, const FContactEdgeID& R)
 		{
 			return L.EdgeID == R.EdgeID;
 		}
 
-		friend bool operator!=(const FContactEdgeID& L, const FContactEdgeID& R)
+		FORCEINLINE friend bool operator!=(const FContactEdgeID& L, const FContactEdgeID& R)
 		{
 			return L.EdgeID != R.EdgeID;
 		}
 
-		friend bool operator<(const FContactEdgeID& L, const FContactEdgeID& R)
+		FORCEINLINE friend bool operator<(const FContactEdgeID& L, const FContactEdgeID& R)
 		{
 			return L.EdgeID < R.EdgeID;
 		}
 
-		friend uint32 GetTypeHash(const FContactEdgeID& V)
+		FORCEINLINE friend uint32 GetTypeHash(const FContactEdgeID& V)
 		{
 			return ::GetTypeHash(V.EdgeID);
 		}
@@ -174,12 +176,8 @@ namespace Chaos
 		};
 	};
 
-	/**
-	 * @brief Extended data used when processing contactpoints on a triangle mesh. 
-	 * Adds information about which triangle we hit (in the local array of FoOntactTriangles)
-	 * as well as which mesh edge or vertex we hit if appropriate (used for pruning).
-	*/
-	class FTriangleContactPoint : public FContactPoint
+	
+	class UE_DEPRECATED(5.4, "No longer used") FTriangleContactPoint : public FContactPoint
 	{
 	public:
 		FTriangleContactPoint()
@@ -204,6 +202,166 @@ namespace Chaos
 		int32 ContactTriangleIndex;
 	};
 
+	/**
+	* An ID for an Edge or Vertex in a triangle mesh.
+	*/
+	struct FContactEdgeOrVertexID
+	{
+	public:
+		FORCEINLINE FContactEdgeOrVertexID()
+			: EdgeID()
+		{
+		}
+
+		FORCEINLINE FContactEdgeOrVertexID(const FContactVertexID VertexIndex)
+		{
+			EdgeID.VertexIDs[0] = VertexIndex;
+			EdgeID.VertexIDs[1] = INDEX_NONE;
+		}
+
+		FORCEINLINE FContactEdgeOrVertexID(const FContactVertexID VertexIndexA, const FContactVertexID VertexIndexB)
+		{
+			EdgeID = FContactEdgeID(VertexIndexA, VertexIndexB);
+		}
+
+		FORCEINLINE bool IsValid() const
+		{
+			return (EdgeID.VertexIDs[0] != INDEX_NONE);
+		}
+
+		FORCEINLINE bool IsVertex() const
+		{
+			return (EdgeID.VertexIDs[0] != INDEX_NONE) && (EdgeID.VertexIDs[1] == INDEX_NONE);
+		}
+
+		FORCEINLINE bool IsEdge() const
+		{
+			return (EdgeID.VertexIDs[0] != INDEX_NONE) && (EdgeID.VertexIDs[1] != INDEX_NONE);
+		}
+
+		FORCEINLINE const FContactVertexID& GetVertexID() const
+		{
+			return EdgeID.VertexIDs[0];
+		}
+
+		FORCEINLINE const FContactEdgeID& GetEdgeID() const
+		{
+			return EdgeID;
+		}
+
+		FORCEINLINE friend bool operator==(const FContactEdgeOrVertexID& L, const FContactEdgeOrVertexID& R)
+		{
+			return L.EdgeID == R.EdgeID;
+		}
+
+		FORCEINLINE friend bool operator!=(const FContactEdgeOrVertexID& L, const FContactEdgeOrVertexID& R)
+		{
+			return L.EdgeID != R.EdgeID;
+		}
+
+		FORCEINLINE friend bool operator<(const FContactEdgeOrVertexID& L, const FContactEdgeOrVertexID& R)
+		{
+			return L.EdgeID < R.EdgeID;
+		}
+
+		FORCEINLINE friend uint32 GetTypeHash(const FContactEdgeOrVertexID& V)
+		{
+			return (V.IsVertex()) ? V.GetVertexID() : GetTypeHash(V.GetEdgeID());
+		}
+
+	private:
+		// We treat the EdgeID as a VertexID when only one vertex ID is set
+		FContactEdgeID EdgeID;
+	};
+
+	/**
+	 * @brief Extended data used when processing contactpoints on a triangle mesh.
+	 * Adds information about which triangle we hit (in the local array of FoOntactTriangles)
+	 * as well as which mesh edge or vertex we hit if appropriate (used for pruning).
+	*/
+	class FTriangleContactPointData
+	{
+	public:
+		FORCEINLINE FTriangleContactPointData()
+			: EdgeOrVertexID()
+			, ContactTriangleIndex(INDEX_NONE)
+			, ContactNormalDotTriangleNormal(0)
+			, bIsEnabled(false)
+		{
+		}
+
+		FORCEINLINE void SetVertexID(const FContactVertexID InVertexID)
+		{
+			EdgeOrVertexID = FContactEdgeOrVertexID(InVertexID);
+		}
+
+		FORCEINLINE void SetEdgeID(const FContactVertexID InVertexIDA, const FContactVertexID InVertexIDB)
+		{
+			EdgeOrVertexID = FContactEdgeOrVertexID(InVertexIDA, InVertexIDB);
+		}
+
+		FORCEINLINE void SetTriangleIndex(const int32 InTriangleIndex)
+		{
+			ContactTriangleIndex = InTriangleIndex;
+		}
+
+		FORCEINLINE void SetContactNormalDotTriangleNormal(const FRealSingle InContactNormalDotTriangleNormal)
+		{
+			ContactNormalDotTriangleNormal = InContactNormalDotTriangleNormal;
+		}
+
+		FORCEINLINE void SetEnabled()
+		{
+			bIsEnabled = true;
+		}
+
+		FORCEINLINE void SetDisabled()
+		{
+			bIsEnabled = false;
+		}
+
+		FORCEINLINE bool IsEnabled() const
+		{
+			return bIsEnabled;
+		}
+
+		FORCEINLINE int32 GetTriangleIndex() const
+		{
+			return ContactTriangleIndex;
+		}
+
+		FORCEINLINE bool IsVertex() const
+		{
+			return EdgeOrVertexID.IsVertex();
+		}
+
+		FORCEINLINE bool IsEdge() const
+		{
+			return EdgeOrVertexID.IsEdge();
+		}
+
+		FORCEINLINE FContactVertexID GetVertexID() const
+		{
+			return EdgeOrVertexID.GetVertexID();
+		}
+
+		FORCEINLINE FContactEdgeID GetEdgeID() const
+		{
+			return EdgeOrVertexID.GetEdgeID();
+		}
+
+		FORCEINLINE FRealSingle GetContactNormalDotTriangleNormal() const
+		{
+			return ContactNormalDotTriangleNormal;
+		}
+
+	private:
+		FContactEdgeOrVertexID EdgeOrVertexID;
+		int32 ContactTriangleIndex;
+		FRealSingle ContactNormalDotTriangleNormal;
+		bool bIsEnabled;
+	};
+
 	/**
 	 * Holds the contacts produced when colliding a convex shape against a non-convex triangular mesh, and provides
 	 * methods to reduce the total set of contact points down to a minimum manifold.
@@ -215,6 +373,7 @@ namespace Chaos
 		FContactTriangleCollector(const bool bInOneSided, const FRigidTransform3& InConvexTransform)
 			: PhiTolerance(0.1)
 			, DistanceTolerance(0.1)
+			, NumDisabledTriangleContactPoints(0)
 			, bOneSidedCollision(bInOneSided)
 			, ConvexTransform(InConvexTransform)
 		{
@@ -227,6 +386,7 @@ namespace Chaos
 			const FRigidTransform3& InConvexTransform)
 			: PhiTolerance(InPhiTolerance)
 			, DistanceTolerance(InDistanceTolerance)
+			, NumDisabledTriangleContactPoints(0)
 			, bOneSidedCollision(bInOneSided)
 			, ConvexTransform(InConvexTransform)
 		{
@@ -234,7 +394,7 @@ namespace Chaos
 
 		inline TArrayView<const FContactPoint> GetContactPoints() const
 		{ 
-			return MakeArrayView(FinalContactPoints);
+			return MakeArrayView(TriangleContactPoints);
 		}
 		
 		// Add contacts between the convex and a single triangle for later processing
@@ -243,16 +403,24 @@ namespace Chaos
 			if (InContactPoints.Num() > 0)
 			{
 				const int32 ContactTriangleIndex = AddTriangle(Triangle, TriangleIndex, VertexIndex0, VertexIndex1, VertexIndex2);
-
-				TriangleContactPoints.Reserve(TriangleContactPoints.Num() + InContactPoints.Num());
+				
+				const int32 NewNumContactPoints = TriangleContactPoints.Num() + InContactPoints.Num();
+				TriangleContactPoints.Reserve(NewNumContactPoints);
+				TriangleContactPointDatas.Reserve(NewNumContactPoints);
+				
 				for (const FContactPoint& ContactPoint : InContactPoints)
 				{
 					if (ContactPoint.Phi < CullDistance)
 					{
+						const FReal ContactNormalDotTriangleNormal = FVec3::DotProduct(ContactPoint.ShapeContactNormal, ContactTriangles[ContactTriangleIndex].FaceNormal);
+
 						const int32 ContactIndex = TriangleContactPoints.Add(ContactPoint);
+						TriangleContactPointDatas.AddDefaulted();
 
 						TriangleContactPoints[ContactIndex].FaceIndex = TriangleIndex;
-						TriangleContactPoints[ContactIndex].ContactTriangleIndex = ContactTriangleIndex;
+						TriangleContactPointDatas[ContactIndex].SetTriangleIndex(ContactTriangleIndex);
+						TriangleContactPointDatas[ContactIndex].SetContactNormalDotTriangleNormal(FRealSingle(ContactNormalDotTriangleNormal));
+						TriangleContactPointDatas[ContactIndex].SetEnabled();
 					}
 				}
 			}
@@ -263,12 +431,13 @@ namespace Chaos
 
 	private:
 		void SetNumContacts(const int32 Num);
-		void RemoveContact(const int32 ContactIndex);
+		void DisableContact(const int32 ContactIndex);
+		void RemoveDisabledContacts();
 
 		void BuildContactFeatureSets();
 		void SortContactPointsForPruning();
 		void SortContactPointsForSolving();
-		void PruneEdgeContactPoints();
+		void PruneEdgeAndVertexContactPoints(const bool bPruneEdges, const bool bPruneVertices);
 		void FixInvalidNormalContactPoints();
 		void PruneInfacingContactPoints();
 		void PruneUnnecessaryContactPoints();
@@ -281,12 +450,9 @@ namespace Chaos
 
 		void DebugDrawContactPoints(const FColor& Color, const FReal LineThickness);
 
-		// The output set of contacts
-		// @todo(chaos): would be nice if we didn't need this and we could provide TriangleContactPoints instead
-		TArray<FContactPoint> FinalContactPoints;
-
 		// The collected set of contacts with additional triangle information
-		TArray<FTriangleContactPoint> TriangleContactPoints;
+		TArray<FContactPoint> TriangleContactPoints;
+		TArray<FTriangleContactPointData> TriangleContactPointDatas;
 
 		// All the triangles referenced by the TriangleContactPoints
 		TArray<FContactTriangle> ContactTriangles;
@@ -297,14 +463,20 @@ namespace Chaos
 		// A list of all the vertices of contacts with faces and edges. Used to prune vertex contacts
 		TSet<FContactVertexID> ContactVertices;
 
-		// We remove contacts that are shallower by this much compared to the deepest contact (negative to disable this functionality)
+		// We remove contacts that are shallower by this much compared to the deepest contact
 		FReal PhiTolerance;
 
 		// We remove contacts that are closer than this to any other contact (negative to disable this functionality)
 		FReal DistanceTolerance;
 
+		// How many disabled contacts are in the TriangleContactPoints and TriangleContactPointDatas arrays (until the next Pack operation)
+		int32 NumDisabledTriangleContactPoints;
+
+		// Whether we want single-sided collision and therefore reject all normal opposing the triangle faces
 		bool bOneSidedCollision;
 
+		// Local space for the triangles
+		// @todo(chaos): rename this
 		FRigidTransform3 ConvexTransform;
 	};
 
