geometry collections can now specify a collision profile per particle/level that loads the data collision profile

diff --git a/Engine/Source/Runtime/Experimental/GeometryCollectionEngine/Private/GeometryCollection/GeometryCollectionComponent.cpp b/Engine/Source/Runtime/Experimental/GeometryCollectionEngine/Private/GeometryCollection/GeometryCollectionComponent.cpp
--- a/Engine/Source/Runtime/Experimental/GeometryCollectionEngine/Private/GeometryCollection/GeometryCollectionComponent.cpp
+++ b/Engine/Source/Runtime/Experimental/GeometryCollectionEngine/Private/GeometryCollection/GeometryCollectionComponent.cpp
@@ -257,6 +257,7 @@ FAutoConsoleVariableRef CVarGeometryCollectionNavigationSizeThreshold(TEXT("p.Ge
 bool bGeometryCollectionSingleThreadedBoundsCalculation = false;
 FAutoConsoleVariableRef CVarGeometryCollectionSingleThreadedBoundsCalculation(TEXT("p.GeometryCollectionSingleThreadedBoundsCalculation"), bGeometryCollectionSingleThreadedBoundsCalculation, TEXT("[Debug Only] Single threaded bounds calculation. [def:false]"));
 
+FName UGeometryCollectionComponent::DefaultCollisionProfileName("InternalGCDefaultCollision");
 
 FGeomComponentCacheParameters::FGeomComponentCacheParameters()
 	: CacheMode(EGeometryCollectionCacheType::None)
@@ -3438,25 +3439,44 @@ void UGeometryCollectionComponent::LoadCollisionProfiles()
 
 		FCollisionProfileDataCache Cache;
 		FCollisionResponseTemplate Template;
-		if (ProfileName == NAME_None || !UCollisionProfile::Get()->GetProfileTemplate(ProfileName, Template))
+		if (ProfileName == NAME_None)
 		{
 			return nullptr;
 		}
 
-		AActor* Owner = GetOwner();
-		const uint32 ActorID = Owner ? Owner->GetUniqueID() : 0;
-		const uint32 CompID = GetUniqueID();
+		if (ProfileName == DefaultCollisionProfileName)
+		{
+			Cache.QueryFilter = InitialQueryFilter;
+			Cache.SimFilter = InitialSimFilter;
+
+			Cache.QueryFilter.Word3 |= (EPDF_SimpleCollision | EPDF_ComplexCollision);
+			Cache.SimFilter.Word3 |= (EPDF_SimpleCollision | EPDF_ComplexCollision);
+		
+			Cache.bQueryEnabled = true;
+			Cache.bSimEnabled = true;
+		}
+		else if (UCollisionProfile::Get()->GetProfileTemplate(ProfileName, Template))
+		{
+			AActor* Owner = GetOwner();
+			const uint32 ActorID = Owner ? Owner->GetUniqueID() : 0;
+			const uint32 CompID = GetUniqueID();
 
-		Cache.QueryFilter = InitialQueryFilter;
-		Cache.SimFilter = InitialSimFilter;
-		CreateShapeFilterData(Template.ObjectType, BodyInstance.GetMaskFilter(), ActorID, Template.ResponseToChannels, CompID, INDEX_NONE, Cache.QueryFilter, Cache.SimFilter, BodyInstance.bUseCCD, bNotifyCollisions || bNotifyGlobalCollisions, false, false);
+			Cache.QueryFilter = InitialQueryFilter;
+			Cache.SimFilter = InitialSimFilter;
+			CreateShapeFilterData(Template.ObjectType, BodyInstance.GetMaskFilter(), ActorID, Template.ResponseToChannels, CompID, INDEX_NONE, Cache.QueryFilter, Cache.SimFilter, BodyInstance.bUseCCD, bNotifyCollisions || bNotifyGlobalCollisions, false, false);
 
-		// Maintain parity with the rest of the geometry collection filters.
-		Cache.QueryFilter.Word3 |= (EPDF_SimpleCollision | EPDF_ComplexCollision);
-		Cache.SimFilter.Word3 |= (EPDF_SimpleCollision | EPDF_ComplexCollision);
+			// Maintain parity with the rest of the geometry collection filters.
+			Cache.QueryFilter.Word3 |= (EPDF_SimpleCollision | EPDF_ComplexCollision);
+			Cache.SimFilter.Word3 |= (EPDF_SimpleCollision | EPDF_ComplexCollision);
 		
-		Cache.bQueryEnabled = CollisionEnabledHasQuery(Template.CollisionEnabled);
-		Cache.bSimEnabled = CollisionEnabledHasPhysics(Template.CollisionEnabled);
+			Cache.bQueryEnabled = CollisionEnabledHasQuery(Template.CollisionEnabled);
+			Cache.bSimEnabled = CollisionEnabledHasPhysics(Template.CollisionEnabled);
+		}
+		else
+		{
+			return nullptr;
+		}
+
 		Cache.bIsValid = true;
 		return &CachedData.Add(ProfileName, Cache);
 	};
diff --git a/Engine/Source/Runtime/Experimental/GeometryCollectionEngine/Public/GeometryCollection/GeometryCollectionComponent.h b/Engine/Source/Runtime/Experimental/GeometryCollectionEngine/Public/GeometryCollection/GeometryCollectionComponent.h
--- a/Engine/Source/Runtime/Experimental/GeometryCollectionEngine/Public/GeometryCollection/GeometryCollectionComponent.h
+++ b/Engine/Source/Runtime/Experimental/GeometryCollectionEngine/Public/GeometryCollection/GeometryCollectionComponent.h
@@ -452,6 +452,9 @@ class UGeometryCollectionComponent : public UMeshComponent, public IChaosNotifyH
 
 public:
 
+	// Collision profile name that indicates we should use the geometry collection's default collision profile.
+	GEOMETRYCOLLECTIONENGINE_API static FName DefaultCollisionProfileName;
+
 	//~ Begin UActorComponent Interface.
 	GEOMETRYCOLLECTIONENGINE_API virtual void CreateRenderState_Concurrent(FRegisterComponentContext* Context) override;
 	GEOMETRYCOLLECTIONENGINE_API virtual void SendRenderDynamicData_Concurrent() override;
