fix errors where rigid clustering would erroneously consider an island in a cluster union to still be considered attached to a main particle

diff --git a/Engine/Source/Runtime/Experimental/Chaos/Private/Chaos/PBDRigidClustering.cpp b/Engine/Source/Runtime/Experimental/Chaos/Private/Chaos/PBDRigidClustering.cpp
--- a/Engine/Source/Runtime/Experimental/Chaos/Private/Chaos/PBDRigidClustering.cpp
+++ b/Engine/Source/Runtime/Experimental/Chaos/Private/Chaos/PBDRigidClustering.cpp
@@ -1264,6 +1264,8 @@ namespace Chaos
 						{
 							for (FPBDRigidParticleHandle* ChildParticle : ParticlesToRelease)
 							{
+								// Need to remove node connections here. Otherwise it may be possible for the cluster union to have erroneous intercluster edges that connect it to another cluster union.
+								RemoveNodeConnections(ChildParticle);
 								RemoveChildFromParent(ChildParticle, ClusteredParticle);
 							}
 							ActivatedChildren.Append(ParticlesToRelease);
@@ -1311,10 +1313,6 @@ namespace Chaos
 									bHasMainParticle = true;
 								}
 							}
-							else
-							{
-								bHasMainParticle = true;
-							}
 						}
 						
 						if (bOnlyUseInterclusterEdgesAttachedToMainParticles)
