Log an error when creating a FHttpClient if the host cannot be set from the provided url

#rnx

- Logging this failure should make url issues easier to diagnose in the future.

diff --git a/Engine/Source/Runtime/Experimental/IoStoreOnDemand/Private/OnDemandIoDispatcherBackend.cpp b/Engine/Source/Runtime/Experimental/IoStoreOnDemand/Private/OnDemandIoDispatcherBackend.cpp
--- a/Engine/Source/Runtime/Experimental/IoStoreOnDemand/Private/OnDemandIoDispatcherBackend.cpp
+++ b/Engine/Source/Runtime/Experimental/IoStoreOnDemand/Private/OnDemandIoDispatcherBackend.cpp
@@ -405,7 +405,11 @@ FHttpClient::FHttpClient(const FString& ServiceUrl, int32 MaxConnectionCount)
 	auto ServiceUrlAnsi = StringCast<ANSICHAR>(*ServiceUrl, ServiceUrl.Len());
 
 	HTTP::FConnectionPool::FParams Params;
-	Params.SetHostFromUrl(ServiceUrlAnsi);
+	if (Params.SetHostFromUrl(ServiceUrlAnsi) < 0)
+	{
+		UE_LOG(LogIas, Error, TEXT("Failed to set host from '%s'"), *ServiceUrl);
+	}
+
 	Params.ConnectionCount = MaxConnectionCount;
 	ConnectionPool = MakeUnique<HTTP::FConnectionPool>(Params);
 }
