IAS - try read config file from platform package before using root folder

#rnx

diff --git a/Engine/Source/Runtime/Experimental/IoStoreOnDemand/Private/IoStoreOnDemand.cpp b/Engine/Source/Runtime/Experimental/IoStoreOnDemand/Private/IoStoreOnDemand.cpp
--- a/Engine/Source/Runtime/Experimental/IoStoreOnDemand/Private/IoStoreOnDemand.cpp
+++ b/Engine/Source/Runtime/Experimental/IoStoreOnDemand/Private/IoStoreOnDemand.cpp
@@ -109,48 +109,66 @@ static bool ApplyEncryptionKeyFromString(const FString& GuidKeyPair)
 }
 
 ////////////////////////////////////////////////////////////////////////////////
-static bool ParseConfigFile(const FString& ConfigPath, UE::FOnDemandEndpoint& OutEndpoint)
+static bool TryParseConfigContent(const FString& ConfigContent, const FString& ConfigFileName, UE::FOnDemandEndpoint& OutEndpoint)
 {
-	FString ConfigContent;
-	if (!FFileHelper::LoadFileToString(ConfigContent, &IPlatformFile::GetPlatformPhysical(), *ConfigPath))
+	if (ConfigContent.IsEmpty())
 	{
 		return false;
 	}
 
-	if (!ConfigContent.IsEmpty())
+	FConfigFile Config;
+	Config.ProcessInputFileContents(ConfigContent, ConfigFileName);
+
+	Config.GetString(TEXT("Endpoint"), TEXT("DistributionUrl"), OutEndpoint.DistributionUrl);
+	Config.GetString(TEXT("Endpoint"), TEXT("ServiceUrl"), OutEndpoint.ServiceUrl);
+	Config.GetString(TEXT("Endpoint"), TEXT("TocPath"), OutEndpoint.TocPath);
+
+	if (OutEndpoint.DistributionUrl.EndsWith(TEXT("/")))
 	{
-		const FString ConfigFileName = FPaths::GetCleanFilename(ConfigPath);
+		OutEndpoint.DistributionUrl = OutEndpoint.DistributionUrl.Left(OutEndpoint.DistributionUrl.Len() - 1);
+	}
 
-		FConfigFile Config;
-		Config.ProcessInputFileContents(ConfigContent, ConfigFileName);
+	if (OutEndpoint.ServiceUrl.EndsWith(TEXT("/")))
+	{
+		OutEndpoint.ServiceUrl = OutEndpoint.DistributionUrl.Left(OutEndpoint.ServiceUrl.Len() - 1);
+	}
 
-		Config.GetString(TEXT("Endpoint"), TEXT("DistributionUrl"), OutEndpoint.DistributionUrl);
-		Config.GetString(TEXT("Endpoint"), TEXT("ServiceUrl"), OutEndpoint.ServiceUrl);
-		Config.GetString(TEXT("Endpoint"), TEXT("TocPath"), OutEndpoint.TocPath);
+	if (OutEndpoint.TocPath.StartsWith(TEXT("/")))
+	{
+		OutEndpoint.TocPath.RightChopInline(1);
+	}
 
-		if (OutEndpoint.DistributionUrl.EndsWith(TEXT("/")))
-		{
-			OutEndpoint.DistributionUrl = OutEndpoint.DistributionUrl.Left(OutEndpoint.DistributionUrl.Len() - 1);
-		}
+	FString ContentKey;
+	if (Config.GetString(TEXT("Endpoint"), TEXT("ContentKey"), ContentKey))
+	{
+		ApplyEncryptionKeyFromString(ContentKey);
+	}
 
-		if (OutEndpoint.ServiceUrl.EndsWith(TEXT("/")))
-		{
-			OutEndpoint.ServiceUrl = OutEndpoint.DistributionUrl.Left(OutEndpoint.ServiceUrl.Len() - 1);
-		}
+	return OutEndpoint.IsValid();
+}
 
-		if (OutEndpoint.TocPath.StartsWith(TEXT("/")))
-		{
-			OutEndpoint.TocPath.RightChopInline(1);
-		}
+////////////////////////////////////////////////////////////////////////////////
+static bool TryParseConfigFileFromPlatformPackage(UE::FOnDemandEndpoint& OutEndpoint)
+{
+	const FString ConfigFileName = TEXT("IoStoreOnDemand.ini");
+	const FString ConfigPath = FPaths::Combine(TEXT("Cloud"), ConfigFileName);
+	const FString ConfigContent = FPlatformMisc::LoadTextFileFromPlatformPackage(ConfigPath);
 
-		FString ContentKey;
-		if (Config.GetString(TEXT("Endpoint"), TEXT("ContentKey"), ContentKey))
-		{
-			ApplyEncryptionKeyFromString(ContentKey);
-		}
+	return TryParseConfigContent(ConfigContent, ConfigFileName, OutEndpoint);
+}
+
+////////////////////////////////////////////////////////////////////////////////
+static bool TryParseConfigFile(const FString& ConfigPath, UE::FOnDemandEndpoint& OutEndpoint)
+{
+	const FString ConfigFileName = TEXT("IoStoreOnDemand.ini");
+	FString ConfigContent; 
+
+	if (!FFileHelper::LoadFileToString(ConfigContent, &IPlatformFile::GetPlatformPhysical(), *ConfigPath))
+	{
+		return false;
 	}
 
-	return OutEndpoint.IsValid();
+	return TryParseConfigContent(ConfigContent, ConfigFileName, OutEndpoint);
 }
 
 ////////////////////////////////////////////////////////////////////////////////
@@ -1363,7 +1381,7 @@ FIoStatus PrimeEndPoint(FStringView IoStoreOnDemandIniPath)
 	using namespace UE::IO::Private;
 
 	UE::FOnDemandEndpoint EndPoint;
-	if (!ParseConfigFile(FString(IoStoreOnDemandIniPath), EndPoint))
+	if (!TryParseConfigFile(FString(IoStoreOnDemandIniPath), EndPoint))
 	{
 		return FIoStatus(EIoErrorCode::Unknown, TEXT("Failed to parse config file"));
 	}
@@ -1497,15 +1515,10 @@ void FIoStoreOnDemandModule::StartupModule()
 	if (!Endpoint.IsValid())
 	{
 		Endpoint = UE::FOnDemandEndpoint();
-		FString ConfigFileName = TEXT("IoStoreOnDemand.ini");
-		FString ConfigPath = FPaths::Combine(FPaths::RootDir(), TEXT("Cloud"), ConfigFileName);
-
-		ParseConfigFile(ConfigPath, Endpoint);
-	}
-
-	if (!Endpoint.IsValid())
-	{
-		return;
+		if (!TryParseConfigFileFromPlatformPackage(Endpoint))
+		{
+			return;
+		}
 	}
 
 #if !UE_BUILD_SHIPPING
