Add tracking of download times/rates for IAS to analytics.

#rnx

- Add tracking of the http duration (in milliseconds) and data transfer rate (MiB/s) to the analytics integration.
-- Note that we send both the mean average and the standard deviation.
- Duplicated FIncrementalVariance from SummarizeTraceCommandlet.cpp as an easy way to keep a running average/variance/deviation value. Moving the class to a common location will be done as a future work item at which point the code duplication will be removed.
- Update FOnDemandIoBackendStats::OnHttpGet to make it clear which unit the size parameter is in (bytes)

diff --git a/Engine/Source/Runtime/Experimental/IoStoreOnDemand/Private/Statistics.cpp b/Engine/Source/Runtime/Experimental/IoStoreOnDemand/Private/Statistics.cpp
--- a/Engine/Source/Runtime/Experimental/IoStoreOnDemand/Private/Statistics.cpp
+++ b/Engine/Source/Runtime/Experimental/IoStoreOnDemand/Private/Statistics.cpp
@@ -16,6 +16,137 @@ namespace UE::IO::Private
 static int32 BytesToApproxMB(uint64 Bytes) { return int32(Bytes >> 20); }
 static int32 BytesToApproxKB(uint64 Bytes) { return int32(Bytes >> 10); }
 
+/**
+ * Code taken from SummarizeTraceCommandlet.cpp pending discussion on moving it
+ * somewhere for general use.
+ * Currently not thread safe!
+ */
+class FIncrementalVariance
+{
+public:
+	FIncrementalVariance()
+		: Count(0)
+		, Mean(0.0)
+		, VarianceAccumulator(0.0)
+	{
+
+	}
+
+	uint64 GetCount() const
+	{
+		return Count;
+	}
+
+	double GetMean() const
+	{
+		return Mean;
+	}
+
+	/**
+	* Compute the variance given Welford's accumulator and the overall count
+	*
+	* @return The variance in sample units squared
+	*/
+	double GetVariance() const
+	{
+		double Result = 0.0;
+
+		if (Count > 1)
+		{
+			// Welford's final step, dependent on sample count
+			Result = VarianceAccumulator / double(Count - 1);
+		}
+
+		return Result;
+	}
+
+	/**
+	* Compute the standard deviation given Welford's accumulator and the overall count
+	*
+	* @return The standard deviation in sample units
+	*/
+	double GetDeviation() const
+	{
+		double Result = 0.0;
+
+		if (Count > 1)
+		{
+			// Welford's final step, dependent on sample count
+			double DeviationSqrd = VarianceAccumulator / double(Count - 1);
+
+			// stddev is sqrt of variance, to restore to units (vs. units squared)
+			Result = sqrt(DeviationSqrd);
+		}
+
+		return Result;
+	}
+
+	/**
+	* Perform an increment of work for Welford's variance, from which we can compute variation and standard deviation
+	*
+	* @param InSample	The new sample value to operate on
+	*/
+	void Increment(const double InSample)
+	{
+		Count++;
+		const double OldMean = Mean;
+		Mean += ((InSample - Mean) / double(Count));
+		VarianceAccumulator += ((InSample - Mean) * (InSample - OldMean));
+	}
+
+	/**
+	* Merge with another IncrementalVariance series in progress
+	*
+	* @param Other	The other variance incremented from another mutually exclusive population of analogous data.
+	*/
+	void Merge(const FIncrementalVariance& Other)
+	{
+		// empty other, nothing to do
+		if (Other.Count == 0)
+		{
+			return;
+		}
+
+		// empty this, just copy other
+		if (Count == 0)
+		{
+			Count = Other.Count;
+			Mean = Other.Mean;
+			VarianceAccumulator = Other.VarianceAccumulator;
+			return;
+		}
+
+		const double TotalPopulation = static_cast<double>(Count + Other.Count);
+		const double MeanDifference = Mean - Other.Mean;
+		const double A = ((Count - 1) * GetVariance()) + ((Other.Count - 1) * Other.GetVariance());
+		const double B = (MeanDifference) * (MeanDifference) * (Count * Other.Count / TotalPopulation);
+		const double MergedVariance = (A + B) / (TotalPopulation - 1);
+
+		const uint64 NewCount = Count + Other.Count;
+		const double NewMean = ((Mean * double(Count)) + (Other.Mean * double(Other.Count))) / double(NewCount);
+		const double NewVarianceAccumulator = MergedVariance * (NewCount - 1);
+
+		Count = NewCount;
+		Mean = NewMean;
+		VarianceAccumulator = NewVarianceAccumulator;
+	}
+
+	/**
+	* Reset state back to initialized.
+	*/
+	void Reset()
+	{
+		Count = 0;
+		Mean = 0.0;
+		VarianceAccumulator = 0.0;
+	}
+
+private:
+	uint64 Count;
+	double Mean;
+	double VarianceAccumulator;
+};
+
 ////////////////////////////////////////////////////////////////////////////////
 // TRACE STATS
 
@@ -37,39 +168,41 @@ static int32 BytesToApproxKB(uint64 Bytes) { return int32(Bytes >> 10); }
 #endif 
 
 // iorequest stats
-FCounterInt			GIoRequestCount(TEXT("Ias/IoRequestCount"), TraceCounterDisplayHint_None);
-FCounterAtomicInt	GIoRequestReadCount(TEXT("Ias/IoRequestReadCount"), TraceCounterDisplayHint_None);
-FCounterAtomicInt	GIoRequestReadBytes(TEXT("Ias/IoRequestReadBytes"), TraceCounterDisplayHint_Memory);
-FCounterInt			GIoRequestCancelCount(TEXT("Ias/IoRequestCancelCount"), TraceCounterDisplayHint_None);
-FCounterAtomicInt	GIoRequestErrorCount(TEXT("Ias/IoRequestErrorCount"), TraceCounterDisplayHint_None);
+FCounterInt				GIoRequestCount(TEXT("Ias/IoRequestCount"), TraceCounterDisplayHint_None);
+FCounterAtomicInt		GIoRequestReadCount(TEXT("Ias/IoRequestReadCount"), TraceCounterDisplayHint_None);
+FCounterAtomicInt		GIoRequestReadBytes(TEXT("Ias/IoRequestReadBytes"), TraceCounterDisplayHint_Memory);
+FCounterInt				GIoRequestCancelCount(TEXT("Ias/IoRequestCancelCount"), TraceCounterDisplayHint_None);
+FCounterAtomicInt		GIoRequestErrorCount(TEXT("Ias/IoRequestErrorCount"), TraceCounterDisplayHint_None);
 // cache stats
-FCounterAtomicInt	GCacheErrorCount(TEXT("Ias/CacheErrorCount"), TraceCounterDisplayHint_None);
-FCounterAtomicInt	GCacheGetCount(TEXT("Ias/CacheGetCount"), TraceCounterDisplayHint_None);
-FCounterAtomicInt	GCachePutCount(TEXT("Ias/CachePutCount"), TraceCounterDisplayHint_None);
-FCounterAtomicInt	GCachePutExistingCount(TEXT("Ias/CachePutExistingCount"), TraceCounterDisplayHint_None);
-FCounterAtomicInt	GCachePutRejectCount(TEXT("Ias/CachePutRejectCount"), TraceCounterDisplayHint_None);
-FCounterAtomicInt	GCacheCachedBytes(TEXT("Ias/CacheCachedBytes"), TraceCounterDisplayHint_Memory);
-FCounterAtomicInt	GCachePendingBytes(TEXT("Ias/CachePendingBytes"), TraceCounterDisplayHint_Memory);
-FCounterAtomicInt	GCacheReadBytes(TEXT("Ias/CacheReadBytes"), TraceCounterDisplayHint_Memory);
-FCounterAtomicInt	GCacheRejectBytes(TEXT("Ias/CachePutRejectBytes"), TraceCounterDisplayHint_Memory);
+FCounterAtomicInt		GCacheErrorCount(TEXT("Ias/CacheErrorCount"), TraceCounterDisplayHint_None);
+FCounterAtomicInt		GCacheGetCount(TEXT("Ias/CacheGetCount"), TraceCounterDisplayHint_None);
+FCounterAtomicInt		GCachePutCount(TEXT("Ias/CachePutCount"), TraceCounterDisplayHint_None);
+FCounterAtomicInt		GCachePutExistingCount(TEXT("Ias/CachePutExistingCount"), TraceCounterDisplayHint_None);
+FCounterAtomicInt		GCachePutRejectCount(TEXT("Ias/CachePutRejectCount"), TraceCounterDisplayHint_None);
+FCounterAtomicInt		GCacheCachedBytes(TEXT("Ias/CacheCachedBytes"), TraceCounterDisplayHint_Memory);
+FCounterAtomicInt		GCachePendingBytes(TEXT("Ias/CachePendingBytes"), TraceCounterDisplayHint_Memory);
+FCounterAtomicInt		GCacheReadBytes(TEXT("Ias/CacheReadBytes"), TraceCounterDisplayHint_Memory);
+FCounterAtomicInt		GCacheRejectBytes(TEXT("Ias/CachePutRejectBytes"), TraceCounterDisplayHint_Memory);
 // http stats
-FCounterInt			GHttpGetCount(TEXT("Ias/HttpGetCount"), TraceCounterDisplayHint_None);
-FCounterInt			GHttpErrorCount(TEXT("Ias/HttpErrorCount"), TraceCounterDisplayHint_None);
-FCounterInt			GHttpRetryCount(TEXT("Ias/HttpRetryCount"), TraceCounterDisplayHint_None);
-FCounterAtomicInt	GHttpPendingCount(TEXT("Ias/HttpPendingCount"), TraceCounterDisplayHint_None);
-FCounterInt			GHttpInflightCount(TEXT("Ias/HttpInflightCount"), TraceCounterDisplayHint_None);
-FCounterInt			GHttpDownloadedBytes(TEXT("Ias/HttpDownloadedBytes"), TraceCounterDisplayHint_Memory);
-FCounterInt			GHttpBandwidthMpbs(TEXT("Ias/HttpBandwidthMbps"), TraceCounterDisplayHint_None);
-FCounterInt			GHttpDurationMs(TEXT("Ias/HttpDurationMs"), TraceCounterDisplayHint_None);
-FCounterInt			GHttpDurationMsAvg(TEXT("Ias/HttpDurationMsAvg"), TraceCounterDisplayHint_None);
-FCounterInt			GHttpDurationMsMax(TEXT("Ias/HttpDurationMsMax"), TraceCounterDisplayHint_None);
-int64				GHttpDurationMsSum = 0;
-constexpr int64		GHttpHistoryCount = 16;
-int64				GHttpHistoryDuration[GHttpHistoryCount] = {};
-int64				GHttpHistoryBytes[GHttpHistoryCount] = {};
-int64				GHttpHistoryTotalDuration = 0;
-int64				GHttpHistoryTotalBytes = 0;
-int64 				GHttpHistoryIndex = 0;
+FCounterInt				GHttpGetCount(TEXT("Ias/HttpGetCount"), TraceCounterDisplayHint_None);
+FCounterInt				GHttpErrorCount(TEXT("Ias/HttpErrorCount"), TraceCounterDisplayHint_None);
+FCounterInt				GHttpRetryCount(TEXT("Ias/HttpRetryCount"), TraceCounterDisplayHint_None);
+FCounterAtomicInt		GHttpPendingCount(TEXT("Ias/HttpPendingCount"), TraceCounterDisplayHint_None);
+FCounterInt				GHttpInflightCount(TEXT("Ias/HttpInflightCount"), TraceCounterDisplayHint_None);
+FCounterInt				GHttpDownloadedBytes(TEXT("Ias/HttpDownloadedBytes"), TraceCounterDisplayHint_Memory);
+FCounterInt				GHttpBandwidthMpbs(TEXT("Ias/HttpBandwidthMbps"), TraceCounterDisplayHint_None);
+FCounterInt				GHttpDurationMs(TEXT("Ias/HttpDurationMs"), TraceCounterDisplayHint_None);
+FCounterInt				GHttpDurationMsAvg(TEXT("Ias/HttpDurationMsAvg"), TraceCounterDisplayHint_None);
+FCounterInt				GHttpDurationMsMax(TEXT("Ias/HttpDurationMsMax"), TraceCounterDisplayHint_None);
+int64					GHttpDurationMsSum = 0;
+constexpr int64			GHttpHistoryCount = 16;
+int64					GHttpHistoryDuration[GHttpHistoryCount] = {};
+int64					GHttpHistoryBytes[GHttpHistoryCount] = {};
+int64					GHttpHistoryTotalDuration = 0;
+int64					GHttpHistoryTotalBytes = 0;
+int64 					GHttpHistoryIndex = 0;
+FIncrementalVariance	GHttpAvgDuration; // Duration of the http requests, in milliseconds
+FIncrementalVariance	GHttpAvgRate; // The download rate of the http requests, in MiB/s
 
 ////////////////////////////////////////////////////////////////////////////////
 // CSV STATS
@@ -149,6 +282,10 @@ void FOnDemandIoBackendStats::ReportAnalytics(TArray<FAnalyticsEventAttribute>&
 		TEXT("IasHttpGetCount"), GHttpGetCount.Get(),
 		TEXT("IasHttpPendingCount"), GHttpPendingCount.Get(),
 		TEXT("IasHttpDownloadedBytes"), GHttpDownloadedBytes.Get(),
+		TEXT("IasHttpDurationMeanAvg"), GHttpAvgDuration.GetMean(),
+		TEXT("IasHttpDurationStdDev"), GHttpAvgDuration.GetDeviation(),
+		TEXT("IasHttpRateMeanAvg"), GHttpAvgRate.GetMean(),
+		TEXT("IasHttpRateStdDev"), GHttpAvgRate.GetDeviation(),
 
 		TEXT("IasCacheErrorCount"), GCacheErrorCount.Get(),
 		TEXT("IasCacheGetCount"), GCacheGetCount.Get(),
@@ -235,12 +372,12 @@ void FOnDemandIoBackendStats::OnHttpDequeue()
 	GHttpInflightCount.Add(1);
 }
 
-void FOnDemandIoBackendStats::OnHttpGet(uint64 Size, uint64 DurationMs)
+void FOnDemandIoBackendStats::OnHttpGet(uint64 SizeBytes, uint64 DurationMs)
 {
 	GHttpPendingCount.Add(-1);
 	GHttpInflightCount.Add(-1);
 	GHttpGetCount.Add(1);
-	GHttpDownloadedBytes.Add(Size);
+	GHttpDownloadedBytes.Add(SizeBytes);
 	GHttpDurationMsSum += DurationMs;
 	GHttpDurationMs.Set(DurationMs);
 
@@ -252,8 +389,8 @@ void FOnDemandIoBackendStats::OnHttpGet(uint64 Size, uint64 DurationMs)
 	GHttpHistoryDuration[GHttpHistoryIndex] = NewDuration;
 
 	GHttpHistoryTotalBytes -= GHttpHistoryBytes[GHttpHistoryIndex];
-	GHttpHistoryTotalBytes += Size;
-	GHttpHistoryBytes[GHttpHistoryIndex] = Size;
+	GHttpHistoryTotalBytes += SizeBytes;
+	GHttpHistoryBytes[GHttpHistoryIndex] = SizeBytes;
 
 	GHttpBandwidthMpbs.Set((GHttpHistoryTotalBytes*8)/(GHttpHistoryTotalDuration+1)/1000);
 	GHttpDurationMsAvg.Set(GHttpHistoryTotalDuration/GHttpHistoryCount);
@@ -264,6 +401,13 @@ void FOnDemandIoBackendStats::OnHttpGet(uint64 Size, uint64 DurationMs)
 	}
 
 	GHttpHistoryIndex = (GHttpHistoryIndex + 1) % GHttpHistoryCount;
+
+	GHttpAvgDuration.Increment(static_cast<double>(DurationMs));
+
+	const double SizeMiB = SizeBytes / (1024.0 * 1024.0);
+	const double DurationSeconds = DurationMs / 1000.0;
+
+	GHttpAvgRate.Increment(SizeMiB / DurationSeconds);
 }
 
 void FOnDemandIoBackendStats::OnHttpRetry()
diff --git a/Engine/Source/Runtime/Experimental/IoStoreOnDemand/Private/Statistics.h b/Engine/Source/Runtime/Experimental/IoStoreOnDemand/Private/Statistics.h
--- a/Engine/Source/Runtime/Experimental/IoStoreOnDemand/Private/Statistics.h
+++ b/Engine/Source/Runtime/Experimental/IoStoreOnDemand/Private/Statistics.h
@@ -50,7 +50,7 @@ public:
 
 	void OnHttpEnqueue() IAS_STATISTICS_IMPL()
 	void OnHttpDequeue() IAS_STATISTICS_IMPL()
-	void OnHttpGet(uint64 Size, uint64 DurationMs) IAS_STATISTICS_IMPL()
+	void OnHttpGet(uint64 SizeBytes, uint64 DurationMs) IAS_STATISTICS_IMPL()
 	void OnHttpRetry() IAS_STATISTICS_IMPL()
 	void OnHttpError() IAS_STATISTICS_IMPL()
 };
