Chaos - fix capsule-convex contact type for end-cap collisions

diff --git a/Engine/Source/Runtime/Experimental/Chaos/Private/Chaos/CollisionOneShotManifolds.cpp b/Engine/Source/Runtime/Experimental/Chaos/Private/Chaos/CollisionOneShotManifolds.cpp
--- a/Engine/Source/Runtime/Experimental/Chaos/Private/Chaos/CollisionOneShotManifolds.cpp
+++ b/Engine/Source/Runtime/Experimental/Chaos/Private/Chaos/CollisionOneShotManifolds.cpp
@@ -1209,7 +1209,7 @@ namespace Chaos
 				ContactPoint.ShapeContactNormal = -GJKContactPoint.ShapeContactNormal;
 				ContactPoint.Phi = GJKContactPoint.Phi;
 				ContactPoint.FaceIndex = INDEX_NONE;
-				ContactPoint.ContactType = EContactPointType::EdgeEdge;
+				ContactPoint.ContactType = bIsConvexPlaneContact ? EContactPointType::VertexPlane : EContactPointType::EdgeEdge;
 				OutContactPoints.Add(ContactPoint);
 				return;
 			}
