Chaos: Clustering : when throttling breaks , allow the remaining to be rescheduled for the next frame

diff --git a/Engine/Source/Runtime/Experimental/Chaos/Private/Chaos/PBDRigidClustering.cpp b/Engine/Source/Runtime/Experimental/Chaos/Private/Chaos/PBDRigidClustering.cpp
--- a/Engine/Source/Runtime/Experimental/Chaos/Private/Chaos/PBDRigidClustering.cpp
+++ b/Engine/Source/Runtime/Experimental/Chaos/Private/Chaos/PBDRigidClustering.cpp
@@ -863,15 +863,14 @@ namespace Chaos
 
 	void FRigidClustering::UpdateTopLevelParticle(FPBDRigidClusteredParticleHandle* Particle)
 	{
+		FPBDRigidClusteredParticleHandle* ParticleToAdd = Particle;
 		FPBDRigidClusteredParticleHandle* Parent = Particle->Parent();
 		if (Parent != nullptr)
 		{
-			TopLevelClusterParentsStrained.Add(Parent);
-		}
-		else
-		{
-			TopLevelClusterParentsStrained.Add(Particle);
+			ParticleToAdd = Parent;
 		}
+		// make sure we only update the timestamp if it is not already in the map
+		TopLevelClusterParentsStrained.FindOrAdd(ParticleToAdd, FPlatformTime::Cycles());
 	}
 	
 	DECLARE_CYCLE_STAT(TEXT("TPBDRigidClustering<>::ReleaseClusterParticles(STRAIN)"), STAT_ReleaseClusterParticles_STRAIN, STATGROUP_Chaos);
@@ -1021,7 +1020,6 @@ namespace Chaos
 					// 1 entry at a time.
 					Children.RemoveAtSwap(ChildIdx, 1, false);
 				}
-				
 				ActivatedChildren.Add(Child);
 				SendBreakingEvent(Child, bParentCrumbled);
 			}
@@ -1428,6 +1426,10 @@ namespace Chaos
 	FAutoConsoleVariableRef CVarPerAdvanceBreaksAllowed(TEXT("p.Chaos.Clustering.PerAdvanceBreaksAllowed"), GPerAdvanceBreaksAllowed,
 		TEXT("Number of breaks allowed to occur for each invokation of AdvanceClustering"));
 
+	static int32 GPerAdvanceBreaksRescheduleLimit = TNumericLimits<int32>::Max();
+	FAutoConsoleVariableRef CVarPerAdvanceBreaksRescheduleLimit(TEXT("p.Chaos.Clustering.PerAdvanceBreaksRescheduleLimit"), GPerAdvanceBreaksRescheduleLimit,
+		TEXT("Number of breaks allowed to be rescheduled for next frame if any "));
+
 	static int32 GDumpClusterAndReleaseStats = 0;
 	FAutoConsoleVariableRef CVarDumpClusterAndReleaseStats(TEXT("p.Chaos.Clustering.DumpClusterAndReleaseStats"), GDumpClusterAndReleaseStats,
 		TEXT("Report the number of cluster processes and released particles per frame, on/off 1/0"));
@@ -1519,8 +1521,15 @@ namespace Chaos
 
 			{
 				SCOPE_CYCLE_COUNTER(STAT_UpdateDirtyImpulses);
-				for (const auto& ActiveCluster : TopLevelClusterParentsStrained)
+
+				// sort by incrementing timestamp
+				TopLevelClusterParentsStrained.ValueSort([](const int64 A, const int64 B) { return A < B; });
+
+				// no process the strained parent and fill ParticlesToProcess 
+				for (const auto& StrainedParentEntry: TopLevelClusterParentsStrained)
 				{
+					Chaos::FPBDRigidClusteredParticleHandle* ActiveCluster = StrainedParentEntry.Key;
+
 					bool bIgnoreDisabledCheck = false;
 					FClusterUnion* ClusterUnion = ClusterUnionManager.FindClusterUnionFromParticle(ActiveCluster);
 					if (ClusterUnion && ClusterUnion->InternalCluster != ActiveCluster)
@@ -1595,14 +1604,42 @@ namespace Chaos
 				// #TODO convert to visitor pattern to avoid TArray allocations above.
 				if(GClusterBreakOnlyStrained == 1)
 				{
-					// Restrict particle view to allowed breaks. We still build the whole array to allow the strain
-					// modifies (executed above) to all strained particles, even though they may not release
-					TArrayView<FPBDRigidClusteredParticleHandle*> ToProcessView{ 
-						ParticlesToProcess.GetData(), 
-						FMath::Min(ParticlesToProcess.Num(), GPerAdvanceBreaksAllowed) 
-					};
-
-					BreakingModel(ToProcessView);
+					// call break model for each particle and only count the ones we breaks
+					// some may strained parent may result in non breaking clusters if strain modifier has changed the strain values
+					// todo(chaos): we should certainly try to have the strain modifier providing a list of those instead of preemptively process thenm to realise that there's nothoing to break
+					int32 NumBreaks = 0;
+					int32 LastProcessedIndex = 0;
+					for (int32 Index = 0; Index < ParticlesToProcess.Num(); Index++)
+					{
+						FPBDRigidClusteredParticleHandle* ParticleToProcess = ParticlesToProcess[Index];
+
+						if (BreakingModel({ &ParticleToProcess , 1 }))
+						{
+							NumBreaks++;
+							if (NumBreaks >= GPerAdvanceBreaksAllowed)
+							{
+								LastProcessedIndex = Index;
+								break;
+							}
+						}
+					}
+
+					// Add back the rest of the particles to process back in the strained array for later processing
+					if (GPerAdvanceBreaksRescheduleLimit > 0)
+					{
+						// Add back the non processed parent clusters for next tick
+						for (int32 Index = 0; Index < GPerAdvanceBreaksRescheduleLimit; Index++)
+						{
+							const int32 ParticleToProcessIndex = Index + LastProcessedIndex + 1;
+							if (ParticleToProcessIndex >= ParticlesToProcess.Num())
+							{
+								break;
+							}
+							FPBDRigidClusteredParticleHandle* ParticleToProcessNextTick = ParticlesToProcess[ParticleToProcessIndex];
+							// since ParticlesToProcess is sorted by time stamp, the new time stamp will make sure they are in the same order 
+							TopLevelClusterParentsStrained.FindOrAdd(ParticleToProcessNextTick, FPlatformTime::Cycles());
+						}
+					}
 				}
 				else
 				{
@@ -1697,7 +1734,7 @@ namespace Chaos
 		BreakingModel(MakeArrayView(InParticles));
 	}
 
-	void FRigidClustering::BreakingModel(TArrayView<FPBDRigidClusteredParticleHandle*> InParticles)
+	bool FRigidClustering::BreakingModel(TArrayView<FPBDRigidClusteredParticleHandle*> InParticles)
 	{
 		SCOPE_CYCLE_COUNTER(STAT_BreakingModel);
 
@@ -1706,22 +1743,28 @@ namespace Chaos
 		// Clear the set tracking breaking collisions
 		BreakingCollisions.Empty();
 
+		bool bHasReleasedParticles = false;
 		for(FPBDRigidClusteredParticleHandle* ClusteredParticle : InParticles)
 		{
 			if(ClusteredParticle->ClusterIds().NumChildren)
 			{
-				ReleaseClusterParticles(ClusteredParticle);
+				TSet<FPBDRigidParticleHandle*> ActivatedParticles = ReleaseClusterParticles(ClusteredParticle);
+				bHasReleasedParticles |= (ActivatedParticles.Num() > 0);
 			}
 		}
 
-		// This way if we break apart a large cluster union here (i.e. many of its children want to be released from ReleaseClusterParticles due to strain)
-		// we'll only update the cluster properties once here (connection graph, geometry, etc.).
-		ClusterUnionManager.HandleDeferredClusterUnionUpdateProperties();
-		// Restore some of the momentum of objects that were touching rigid clusters that broke
-		if(RestoreBreakingMomentumPercent > 0.f)
+		if (bHasReleasedParticles)
 		{
-			RestoreBreakingMomentum();
+			// This way if we break apart a large cluster union here (i.e. many of its children want to be released from ReleaseClusterParticles due to strain)
+			// we'll only update the cluster properties once here (connection graph, geometry, etc.).
+			ClusterUnionManager.HandleDeferredClusterUnionUpdateProperties();
+			// Restore some of the momentum of objects that were touching rigid clusters that broke
+			if (RestoreBreakingMomentumPercent > 0.f)
+			{
+				RestoreBreakingMomentum();
+			}
 		}
+		return bHasReleasedParticles;
 	}
 
 	DECLARE_CYCLE_STAT(TEXT("FRigidClustering::Visitor"), STAT_ClusterVisitor, STATGROUP_Chaos);
diff --git a/Engine/Source/Runtime/Experimental/Chaos/Public/Chaos/PBDRigidClustering.h b/Engine/Source/Runtime/Experimental/Chaos/Public/Chaos/PBDRigidClustering.h
--- a/Engine/Source/Runtime/Experimental/Chaos/Public/Chaos/PBDRigidClustering.h
+++ b/Engine/Source/Runtime/Experimental/Chaos/Public/Chaos/PBDRigidClustering.h
@@ -223,7 +223,7 @@ public:
 	*/
 	CHAOS_API void BreakingModel();
 	CHAOS_API void BreakingModel(TArray<FPBDRigidClusteredParticleHandle*>& InParticles);
-	CHAOS_API void BreakingModel(TArrayView<FPBDRigidClusteredParticleHandle*> InParticles);
+	CHAOS_API bool BreakingModel(TArrayView<FPBDRigidClusteredParticleHandle*> InParticles);
 	
 	//
 	// Access
@@ -401,7 +401,12 @@ public:
 	FClusterUnionManager& GetClusterUnionManager() { return ClusterUnionManager; }
 	const FClusterUnionManager& GetClusterUnionManager() const { return ClusterUnionManager; }
 
-	const TSet<Chaos::FPBDRigidClusteredParticleHandle*>& GetTopLevelClusterParentsStrained() const { return TopLevelClusterParentsStrained; }
+	UE_DEPRECATED(5.4, "No longer expose publicly - now return empty set")
+	const TSet<Chaos::FPBDRigidClusteredParticleHandle*>& GetTopLevelClusterParentsStrained() const
+	{ 
+		static const TSet<Chaos::FPBDRigidClusteredParticleHandle*> ConstEmptySet;
+		return ConstEmptySet;
+	}
 
 	// Remove connectivity edges for specified particles
 	CHAOS_API void RemoveNodeConnections(FPBDRigidParticleHandle* Child);
@@ -570,7 +575,8 @@ private:
 	FRigidEvolution& MEvolution;
 	FPBDRigidClusteredParticles& MParticles;
 	TSet<Chaos::FPBDRigidClusteredParticleHandle*> TopLevelClusterParents;
-	TSet<Chaos::FPBDRigidClusteredParticleHandle*> TopLevelClusterParentsStrained;
+
+	TMap<Chaos::FPBDRigidClusteredParticleHandle*, int64> TopLevelClusterParentsStrained;
 
 	// Cluster data
 	FClusterMap MChildren;
