Chaos - particles keep track of the SOA lists they are in.

cvar-enabled validation can check for duplicate particles in views that would cause race conditions in collision detection.

diff --git a/Engine/Source/Runtime/Experimental/Chaos/Private/Chaos/PBDRigidsEvolutionGBF.cpp b/Engine/Source/Runtime/Experimental/Chaos/Private/Chaos/PBDRigidsEvolutionGBF.cpp
--- a/Engine/Source/Runtime/Experimental/Chaos/Private/Chaos/PBDRigidsEvolutionGBF.cpp
+++ b/Engine/Source/Runtime/Experimental/Chaos/Private/Chaos/PBDRigidsEvolutionGBF.cpp
@@ -247,6 +247,12 @@ void CheckParticleViewsForDupes(FPBDRigidsSOAs& Particles)
 	{
 		QUICK_SCOPE_CYCLE_COUNTER(STAT_Evolution_CheckParticleViewsForDupes);
 
+		// Check that all particles know what lists they are in
+		Particles.CheckListMasks();
+
+		// Check that we have no particles in multiple views
+		Particles.CheckViewMasks();
+
 		// A particle appearing twice in either of these results in a race condition because the 
 		// collision detection loop will visit the same particle pair twice on different threads.
 		CheckParticleViewForDupes(TEXT("NonDisabledDynamicView"), Particles.GetNonDisabledDynamicView());
diff --git a/Engine/Source/Runtime/Experimental/Chaos/Private/Chaos/PBDRigidsSOAs.cpp b/Engine/Source/Runtime/Experimental/Chaos/Private/Chaos/PBDRigidsSOAs.cpp
new file mode 100644
--- /dev/null
+++ b/Engine/Source/Runtime/Experimental/Chaos/Private/Chaos/PBDRigidsSOAs.cpp
@@ -0,0 +1,263 @@
+// Copyright Epic Games, Inc. All Rights Reserved.
+#include "Chaos/PBDRigidsSOAs.h"
+
+namespace Chaos
+{
+	void FPBDRigidsSOAs::UpdateViews()
+	{
+		//build various views. Group SOA types together for better branch prediction
+		{
+			auto TmpArray = bResimulating
+				? TArray<TSOAView<FGeometryParticles>>
+			{
+				// re-sim only works with a reduced set of particle types, i.e. no DynamicClusteredMapArray (DynamicGeometryCollectionArray??) BH
+				{&ResimStaticParticles.GetArray()},
+				{&ResimKinematicParticles.GetArray()},
+				{&ResimDynamicParticles.GetArray()},
+				{&ResimDynamicKinematicParticles.GetArray()}
+			}
+			: TArray<TSOAView<FGeometryParticles>>
+			{
+				StaticParticles.Get(),
+				KinematicParticles.Get(),
+				DynamicParticles.Get(),
+				DynamicKinematicParticles.Get(),
+				{&StaticClusteredMapArray.GetArray()},
+				{&KinematicClusteredMapArray.GetArray()},
+				{&DynamicClusteredMapArray.GetArray()},
+				{&StaticGeometryCollectionArray.GetArray()},
+				{&KinematicGeometryCollectionArray.GetArray()},
+				{&SleepingGeometryCollectionArray.GetArray()},
+				{&DynamicGeometryCollectionArray.GetArray()}
+			};
+			NonDisabledView = MakeParticleView(MoveTemp(TmpArray));
+		}
+		{
+			auto TmpArray = bResimulating
+				? TArray<TSOAView<FPBDRigidParticles>>
+			{
+				{&ResimDynamicParticles.GetArray() }
+			}
+			: TArray<TSOAView<FPBDRigidParticles>>
+			{
+				DynamicParticles.Get(),
+				{&DynamicClusteredMapArray.GetArray() },
+				{&SleepingGeometryCollectionArray.GetArray() },
+				{&DynamicGeometryCollectionArray.GetArray() }
+			};
+			NonDisabledDynamicView = MakeParticleView(MoveTemp(TmpArray));
+		}
+		{
+			TArray<TSOAView<FPBDRigidParticles>> TmpArray =
+			{
+				{&ActiveParticlesMapArray.GetArray()},
+				{&MovingKinematicsMapArray.GetArray()},
+				//{&KinematicClusteredMapArray.GetArray()},	// @todo(chaos): we should include moving clustered kinematics in this view
+			};
+			ActiveDynamicMovingKinematicParticlesView = MakeParticleView(MoveTemp(TmpArray));
+		}
+		{
+			auto TmpArray = bResimulating
+				? TArray<TSOAView<FPBDRigidParticles>>
+			{
+				{&ResimActiveParticlesMapArray.GetArray()}
+			}
+			: TArray<TSOAView<FPBDRigidParticles>>
+			{
+				{&ActiveParticlesMapArray.GetArray()},
+				{&StaticGeometryCollectionArray.GetArray()},
+				{&KinematicGeometryCollectionArray.GetArray()},
+			};
+			ActiveParticlesView = MakeParticleView(MoveTemp(TmpArray));
+		}
+		{
+			auto TmpArray = bResimulating
+				? TArray<TSOAView<FPBDRigidParticles>>
+			{
+				{&ResimActiveParticlesMapArray.GetArray()},
+				{&TransientDirtyMapArray.GetArray()}
+			}
+			: TArray<TSOAView<FPBDRigidParticles>>
+			{
+				{&ActiveParticlesMapArray.GetArray()},
+				{&MovingKinematicsMapArray.GetArray()},
+				{&SleepingGeometryCollectionArray.GetArray()},
+				{&TransientDirtyMapArray.GetArray()}
+			};
+			DirtyParticlesView = MakeParticleView(MoveTemp(TmpArray));
+		}
+		{
+			TArray<TSOAView<FGeometryParticles>> TmpArray =
+			{
+				StaticParticles.Get(),
+				StaticDisabledParticles.Get(),
+				KinematicParticles.Get(),
+				KinematicDisabledParticles.Get(),
+				DynamicParticles.Get(),
+				DynamicDisabledParticles.Get(),
+				DynamicKinematicParticles.Get(),
+				ClusteredParticles.Get(),
+				GeometryCollectionParticles.Get()
+			};
+			AllParticlesView = MakeParticleView(MoveTemp(TmpArray));
+		}
+		{
+			auto TmpArray = bResimulating
+				? TArray<TSOAView<FKinematicGeometryParticles>>
+			{
+				{&ResimKinematicParticles.GetArray()},
+				{ &ResimDynamicKinematicParticles.GetArray() }
+			}
+			: TArray<TSOAView<FKinematicGeometryParticles>>
+			{
+				KinematicParticles.Get(),
+				DynamicKinematicParticles.Get(),
+				{&KinematicGeometryCollectionArray.GetArray()},
+				{&KinematicClusteredMapArray.GetArray()}
+			};
+			ActiveKinematicParticlesView = MakeParticleView(MoveTemp(TmpArray));
+		}
+		{
+			// todo(chaos) handle resim ?
+			TArray<TSOAView<FPBDRigidParticles>> TmpArray =
+			{
+				{&MovingKinematicsMapArray.GetArray()},
+			};
+			ActiveMovingKinematicParticlesView = MakeParticleView(MoveTemp(TmpArray));
+		}
+		{
+			auto TmpArray = bResimulating
+				? TArray<TSOAView<FGeometryParticles>>
+			{
+				{&ResimStaticParticles.GetArray()}
+			}
+			: TArray<TSOAView<FGeometryParticles>>
+			{
+				StaticParticles.Get(),
+				{&StaticClusteredMapArray.GetArray()}
+			};
+			ActiveStaticParticlesView = MakeParticleView(MoveTemp(TmpArray));
+		}
+		{
+			auto TmpArray = bResimulating
+				? TArray<TSOAView<TPBDGeometryCollectionParticles<FReal, 3>>> {}	//no geometry collection during resim
+			: TArray<TSOAView<TPBDGeometryCollectionParticles<FReal, 3>>>
+			{
+				{&StaticGeometryCollectionArray.GetArray()},
+				{&KinematicGeometryCollectionArray.GetArray()},
+				{&DynamicGeometryCollectionArray.GetArray()}
+			};
+			ActiveGeometryCollectionParticlesView = MakeParticleView(MoveTemp(TmpArray));
+		}
+
+		{
+			auto TmpArray = bResimulating
+				? TArray<TSOAView<FPBDRigidClusteredParticles>> {}	//no clusters during resim
+			: TArray<TSOAView<FPBDRigidClusteredParticles>>
+			{
+				{&StaticClusteredMapArray.GetArray()},
+				{&KinematicClusteredMapArray.GetArray()},
+				{&DynamicClusteredMapArray.GetArray()}
+			};
+			NonDisabledClusteredView = MakeParticleView(MoveTemp(TmpArray));
+		}
+	}
+
+	void FPBDRigidsSOAs::SetContainerListMasks()
+	{
+		StaticParticles->SetContainerListMask(EGeometryParticleListMask::StaticParticles);
+		StaticDisabledParticles->SetContainerListMask(EGeometryParticleListMask::StaticDisabledParticles);
+		KinematicParticles->SetContainerListMask(EGeometryParticleListMask::KinematicParticles);
+		KinematicDisabledParticles->SetContainerListMask(EGeometryParticleListMask::KinematicDisabledParticles);
+		DynamicDisabledParticles->SetContainerListMask(EGeometryParticleListMask::DynamicDisabledParticles);
+		DynamicParticles->SetContainerListMask(EGeometryParticleListMask::DynamicParticles);
+		DynamicKinematicParticles->SetContainerListMask(EGeometryParticleListMask::DynamicKinematicParticles);
+		ClusteredParticles->SetContainerListMask(EGeometryParticleListMask::ClusteredParticles);
+		GeometryCollectionParticles->SetContainerListMask(EGeometryParticleListMask::GeometryCollectionParticles);
+
+		StaticGeometryCollectionArray.SetContainerListMask(EGeometryParticleListMask::StaticGeometryCollectionArray);
+		KinematicGeometryCollectionArray.SetContainerListMask(EGeometryParticleListMask::KinematicGeometryCollectionArray);
+		SleepingGeometryCollectionArray.SetContainerListMask(EGeometryParticleListMask::SleepingGeometryCollectionArray);
+		DynamicGeometryCollectionArray.SetContainerListMask(EGeometryParticleListMask::DynamicGeometryCollectionArray);
+		ActiveParticlesMapArray.SetContainerListMask(EGeometryParticleListMask::ActiveParticlesMapArray);
+		TransientDirtyMapArray.SetContainerListMask(EGeometryParticleListMask::TransientDirtyMapArray);
+		MovingKinematicsMapArray.SetContainerListMask(EGeometryParticleListMask::MovingKinematicsMapArray);
+		ResimActiveParticlesMapArray.SetContainerListMask(EGeometryParticleListMask::ResimActiveParticlesMapArray);
+		ResimDynamicParticles.SetContainerListMask(EGeometryParticleListMask::ResimDynamicParticles);
+		ResimDynamicKinematicParticles.SetContainerListMask(EGeometryParticleListMask::ResimDynamicKinematicParticles);
+		ResimStaticParticles.SetContainerListMask(EGeometryParticleListMask::ResimStaticParticles);
+		ResimKinematicParticles.SetContainerListMask(EGeometryParticleListMask::ResimKinematicParticles);
+		StaticClusteredMapArray.SetContainerListMask(EGeometryParticleListMask::StaticClusteredMapArray);
+		KinematicClusteredMapArray.SetContainerListMask(EGeometryParticleListMask::KinematicClusteredMapArray);
+		DynamicClusteredMapArray.SetContainerListMask(EGeometryParticleListMask::DynamicClusteredMapArray);
+	}
+
+	void FPBDRigidsSOAs::CheckListMasks()
+	{
+		CheckSOAMasks(StaticParticles.Get());
+		CheckSOAMasks(StaticDisabledParticles.Get());
+		CheckSOAMasks(KinematicParticles.Get());
+		CheckSOAMasks(KinematicDisabledParticles.Get());
+		CheckSOAMasks(DynamicDisabledParticles.Get());
+		CheckSOAMasks(DynamicParticles.Get());
+		CheckSOAMasks(DynamicKinematicParticles.Get());
+		CheckSOAMasks(ClusteredParticles.Get());
+		CheckSOAMasks(GeometryCollectionParticles.Get());
+
+		CheckListMasks(StaticGeometryCollectionArray);
+		CheckListMasks(KinematicGeometryCollectionArray);
+		CheckListMasks(SleepingGeometryCollectionArray);
+		CheckListMasks(DynamicGeometryCollectionArray);
+		CheckListMasks(ActiveParticlesMapArray);
+		CheckListMasks(TransientDirtyMapArray);
+		CheckListMasks(MovingKinematicsMapArray);
+		CheckListMasks(ResimActiveParticlesMapArray);
+		CheckListMasks(ResimDynamicParticles);
+		CheckListMasks(ResimDynamicKinematicParticles);
+		CheckListMasks(ResimStaticParticles);
+		CheckListMasks(ResimKinematicParticles);
+		CheckListMasks(StaticClusteredMapArray);
+		CheckListMasks(KinematicClusteredMapArray);
+		CheckListMasks(DynamicClusteredMapArray);
+	}
+
+	void FPBDRigidsSOAs::CheckViewMasks()
+	{
+		for (auto& Particle : GetAllParticlesView())
+		{
+			CheckParticleViewMaskImpl(Particle.Handle());
+		}
+	}
+
+	// Check that the particle's list mask contains only one of the bits from a view's list mask
+	bool ZeroOrOneBitSet(const EGeometryParticleListMask ParticleListMask, const EGeometryParticleListMask ViewListMask)
+	{
+		const EGeometryParticleListMask ParticleViewListMask = ParticleListMask & ViewListMask;
+		return FMath::IsPowerOfTwo(uint32(ParticleViewListMask));
+	}
+
+	void FPBDRigidsSOAs::CheckParticleViewMaskImpl(const FGeometryParticleHandle* Particle) const
+	{
+		const EGeometryParticleListMask ParticleListMask = Particle->ListMask();
+
+		// All particles must be in a list or SOA
+		ensureAlwaysMsgf(!(ParticleListMask == EGeometryParticleListMask::None), TEXT("Particle %s is not in any lists"), *GetParticleDebugName(*Particle));
+
+		// Cannot be in more than one of these lists at a time
+		// @todo(chaos): we could automate ViewListMask creation if the Views knew what ParticleMapArrays and 
+		// ParticleArrays they contained but currently they just holds a pointer to the underlying TArray.
+
+		// NonDisabledDynamicView is used by collision detection. Duplicates in this view can lead to race conditions in the midphase
+		const EGeometryParticleListMask NonDisabledDynamicViewListMask = 
+			EGeometryParticleListMask::DynamicParticles 
+			| EGeometryParticleListMask::DynamicClusteredMapArray 
+			| EGeometryParticleListMask::SleepingGeometryCollectionArray 
+			| EGeometryParticleListMask::DynamicGeometryCollectionArray;
+		ensureAlwaysMsgf(ZeroOrOneBitSet(ParticleListMask, NonDisabledDynamicViewListMask), TEXT("Particle %s duplicated in NonDisabledDynamicView (ListMask 0x%08x)"), *GetParticleDebugName(*Particle), ParticleListMask);
+	
+		const EGeometryParticleListMask ActiveDynamicMovingKinematicParticlesViewListMask = 
+			EGeometryParticleListMask::ActiveParticlesMapArray
+			| EGeometryParticleListMask::MovingKinematicsMapArray;
+		ensureAlwaysMsgf(ZeroOrOneBitSet(ParticleListMask, ActiveDynamicMovingKinematicParticlesViewListMask), TEXT("Particle %s duplicated in ActiveDynamicMovingKinematicParticlesView (ListMask 0x%08x)"), *GetParticleDebugName(*Particle), ParticleListMask);
+	}
+}
\ No newline at end of file
diff --git a/Engine/Source/Runtime/Experimental/Chaos/Public/Chaos/GeometryParticles.h b/Engine/Source/Runtime/Experimental/Chaos/Public/Chaos/GeometryParticles.h
--- a/Engine/Source/Runtime/Experimental/Chaos/Public/Chaos/GeometryParticles.h
+++ b/Engine/Source/Runtime/Experimental/Chaos/Public/Chaos/GeometryParticles.h
@@ -161,46 +161,15 @@
 		
 		TGeometryParticlesImp()
 		    : TParticles<T, d>()
+			, MContainerListMask(EGeometryParticleListMask::None)
 		{
 			MParticleType = EParticleType::Static;
-			TArrayCollection::AddArray(&MUniqueIdx);
-			TArrayCollection::AddArray(&MR);
-			TArrayCollection::AddArray(&MGeometry);
-#if CHAOS_DETERMINISTIC
-			TArrayCollection::AddArray(&MParticleIDs);
-#endif
-			TArrayCollection::AddArray(&MHasCollision);
-			TArrayCollection::AddArray(&MShapesArray);
-			TArrayCollection::AddArray(&MLocalBounds);
-			TArrayCollection::AddArray(&MCCDAxisThreshold);
-			TArrayCollection::AddArray(&MWorldSpaceInflatedBounds);
-			TArrayCollection::AddArray(&MHasBounds);
-			TArrayCollection::AddArray(&MSpatialIdx);
-			TArrayCollection::AddArray(&MSyncState);
-			TArrayCollection::AddArray(&MWeakParticleHandle);
-			TArrayCollection::AddArray(&MParticleConstraints);
-			TArrayCollection::AddArray(&MParticleCollisions);
-			TArrayCollection::AddArray(&MGraphNode);
-			TArrayCollection::AddArray(&MResimType);
-			TArrayCollection::AddArray(&MEnabledDuringResim);
-			TArrayCollection::AddArray(&MLightWeightDisabled);
-
-
-#if CHAOS_DEBUG_NAME
-			TArrayCollection::AddArray(&MDebugName);
-#endif
-
-			if (IsRigidBodySim())
-			{
-				TArrayCollection::AddArray(&MGeometryParticleHandle);
-				TArrayCollection::AddArray(&MGeometryParticle);
-				TArrayCollection::AddArray(&MPhysicsProxy);
-			}
-
+			RegisterArrays();
 		}
 		TGeometryParticlesImp(const TGeometryParticlesImp<T, d, SimType>& Other) = delete;
 		TGeometryParticlesImp(TGeometryParticlesImp<T, d, SimType>&& Other)
 		    : TParticles<T, d>(MoveTemp(Other))
+			, MContainerListMask(Other.MContainerListMask)
 			, MUniqueIdx(MoveTemp(Other.MUniqueIdx))
 			, MR(MoveTemp(Other.MR))
 			, MGeometry(MoveTemp(Other.MGeometry))
@@ -231,79 +200,17 @@
 
 		{
 			MParticleType = EParticleType::Static;
-			TArrayCollection::AddArray(&MUniqueIdx);
-			TArrayCollection::AddArray(&MR);
-			TArrayCollection::AddArray(&MGeometry);
-#if CHAOS_DETERMINISTIC
-			TArrayCollection::AddArray(&MParticleIDs);
-#endif
-			TArrayCollection::AddArray(&MHasCollision);
-			TArrayCollection::AddArray(&MShapesArray);
-			TArrayCollection::AddArray(&MLocalBounds);
-			TArrayCollection::AddArray(&MCCDAxisThreshold);
-			TArrayCollection::AddArray(&MWorldSpaceInflatedBounds);
-			TArrayCollection::AddArray(&MHasBounds);
-			TArrayCollection::AddArray(&MSpatialIdx);
-			TArrayCollection::AddArray(&MSyncState);
-			TArrayCollection::AddArray(&MWeakParticleHandle);
-			TArrayCollection::AddArray(&MParticleConstraints);
-			TArrayCollection::AddArray(&MParticleCollisions);
-			TArrayCollection::AddArray(&MGraphNode);
-			TArrayCollection::AddArray(&MResimType);
-			TArrayCollection::AddArray(&MEnabledDuringResim);
-			TArrayCollection::AddArray(&MLightWeightDisabled);
-
-#if CHAOS_DEBUG_NAME
-			TArrayCollection::AddArray(&MDebugName);
-#endif
-
-			if (IsRigidBodySim())
-			{
-				TArrayCollection::AddArray(&MGeometryParticleHandle);
-				TArrayCollection::AddArray(&MGeometryParticle);
-				TArrayCollection::AddArray(&MPhysicsProxy);
-			}
+			RegisterArrays();
 		}
 
 		static constexpr bool IsRigidBodySim() { return SimType == EGeometryParticlesSimType::RigidBodySim; }
 
 		TGeometryParticlesImp(TParticles<T, d>&& Other)
 		    : TParticles<T, d>(MoveTemp(Other))
+			, MContainerListMask(EGeometryParticleListMask::None)
 		{
 			MParticleType = EParticleType::Static;
-			TArrayCollection::AddArray(&MUniqueIdx);
-			TArrayCollection::AddArray(&MR);
-			TArrayCollection::AddArray(&MGeometry);
-#if CHAOS_DETERMINISTIC
-			TArrayCollection::AddArray(&MParticleIDs);
-#endif
-			TArrayCollection::AddArray(&MHasCollision);
-			TArrayCollection::AddArray(&MShapesArray);
-			TArrayCollection::AddArray(&MLocalBounds);
-			TArrayCollection::AddArray(&MCCDAxisThreshold);
-			TArrayCollection::AddArray(&MWorldSpaceInflatedBounds);
-			TArrayCollection::AddArray(&MHasBounds);
-			TArrayCollection::AddArray(&MSpatialIdx);
-			TArrayCollection::AddArray(&MSyncState);
-			TArrayCollection::AddArray(&MWeakParticleHandle);
-			TArrayCollection::AddArray(&MParticleConstraints);
-			TArrayCollection::AddArray(&MParticleCollisions);
-			TArrayCollection::AddArray(&MGraphNode);
-			TArrayCollection::AddArray(&MResimType);
-			TArrayCollection::AddArray(&MEnabledDuringResim);
-			TArrayCollection::AddArray(&MLightWeightDisabled);
-
-
-#if CHAOS_DEBUG_NAME
-			TArrayCollection::AddArray(&MDebugName);
-#endif
-
-			if (IsRigidBodySim())
-			{
-				TArrayCollection::AddArray(&MGeometryParticleHandle);
-				TArrayCollection::AddArray(&MGeometryParticle);
-				TArrayCollection::AddArray(&MPhysicsProxy);
-			}
+			RegisterArrays();
 		}
 
 		virtual ~TGeometryParticlesImp()
@@ -363,6 +270,43 @@
 		void SetSharedGeometry(const int32 Index, TSharedPtr<const FImplicitObject, ESPMode::ThreadSafe> InShared) { check(false); }
 
 	private:
+		void RegisterArrays()
+		{
+			TArrayCollection::AddArray(&MUniqueIdx);
+			TArrayCollection::AddArray(&MR);
+			TArrayCollection::AddArray(&MGeometry);
+#if CHAOS_DETERMINISTIC
+			TArrayCollection::AddArray(&MParticleIDs);
+#endif
+			TArrayCollection::AddArray(&MHasCollision);
+			TArrayCollection::AddArray(&MShapesArray);
+			TArrayCollection::AddArray(&MLocalBounds);
+			TArrayCollection::AddArray(&MCCDAxisThreshold);
+			TArrayCollection::AddArray(&MWorldSpaceInflatedBounds);
+			TArrayCollection::AddArray(&MHasBounds);
+			TArrayCollection::AddArray(&MSpatialIdx);
+			TArrayCollection::AddArray(&MSyncState);
+			TArrayCollection::AddArray(&MWeakParticleHandle);
+			TArrayCollection::AddArray(&MParticleConstraints);
+			TArrayCollection::AddArray(&MParticleCollisions);
+			TArrayCollection::AddArray(&MGraphNode);
+			TArrayCollection::AddArray(&MResimType);
+			TArrayCollection::AddArray(&MEnabledDuringResim);
+			TArrayCollection::AddArray(&MLightWeightDisabled);
+			TArrayCollection::AddArray(&MParticleListMask);
+
+#if CHAOS_DEBUG_NAME
+			TArrayCollection::AddArray(&MDebugName);
+#endif
+
+			if (IsRigidBodySim())
+			{
+				TArrayCollection::AddArray(&MGeometryParticleHandle);
+				TArrayCollection::AddArray(&MGeometryParticle);
+				TArrayCollection::AddArray(&MPhysicsProxy);
+			}
+		}
+
 		void SetGeometryImpl(const int32 Index, const FImplicitObjectPtr& InGeometry)
 		{
 			MGeometry[Index] = InGeometry;
@@ -586,6 +530,8 @@
 		FORCEINLINE bool LightWeightDisabled(const int32 Index) const { return MLightWeightDisabled[Index]; }
 		FORCEINLINE bool& LightWeightDisabled(const int32 Index) { return MLightWeightDisabled[Index]; }
 
+		FORCEINLINE EGeometryParticleListMask ListMask(const int32 Index) const { return MParticleListMask[Index]; }
+		FORCEINLINE EGeometryParticleListMask& ListMask(const int32 Index) { return MParticleListMask[Index]; }
 
 		// Deprecated API
 		UE_DEPRECATED(5.3, "Use ConstraintGraphNode") const int32 ConstraintGraphIndex(const int32 Index) const { return INDEX_NONE; }
@@ -695,6 +641,9 @@
 
 		FORCEINLINE EParticleType ParticleType() const { return MParticleType; }
 
+		FORCEINLINE EGeometryParticleListMask GetContainerListMask() const { return MContainerListMask; }
+		void SetContainerListMask(const EGeometryParticleListMask InMask) { MContainerListMask = InMask; }
+
 		FORCEINLINE TArray<TRotation<T, d>>& AllR() { return MR; }
 		FORCEINLINE TArray<TAABB<T, d>>& AllLocalBounds() { return MLocalBounds; }
 		FORCEINLINE TArray<TAABB<T, d>>& AllWorldSpaceInflatedBounds() { return MWorldSpaceInflatedBounds; }
@@ -702,6 +651,7 @@
 
 	protected:
 		EParticleType MParticleType;
+		EGeometryParticleListMask MContainerListMask;
 
 	private:
 		TArrayCollectionArray<FUniqueIdx> MUniqueIdx;
@@ -727,6 +677,7 @@
 		TArrayCollectionArray<EResimType> MResimType;
 		TArrayCollectionArray<bool> MEnabledDuringResim;
 		TArrayCollectionArray<bool> MLightWeightDisabled;
+		TArrayCollectionArray<EGeometryParticleListMask> MParticleListMask;
 
 		CHAOS_API void UpdateShapesArray(const int32 Index);
 
diff --git a/Engine/Source/Runtime/Experimental/Chaos/Public/Chaos/GeometryParticlesfwd.h b/Engine/Source/Runtime/Experimental/Chaos/Public/Chaos/GeometryParticlesfwd.h
--- a/Engine/Source/Runtime/Experimental/Chaos/Public/Chaos/GeometryParticlesfwd.h
+++ b/Engine/Source/Runtime/Experimental/Chaos/Public/Chaos/GeometryParticlesfwd.h
@@ -107,4 +107,45 @@ namespace Chaos
 		Dynamic,
 		ClusterUnion
 	};
+
+	/**
+	* Used to keep track of which Lists a particle is in
+	*/
+	enum class EGeometryParticleListMask
+	{
+		None								= 0,
+
+		//
+		// NOTE: The contents of this enum match the members of FPBDRigidsSOA
+		//
+
+		// Particle containers (each particle can only be in one of these by definition)
+		StaticParticles						= 1 << 0,
+		StaticDisabledParticles				= 1 << 1,
+		KinematicParticles					= 1 << 2,
+		KinematicDisabledParticles			= 1 << 3,
+		DynamicDisabledParticles			= 1 << 4,
+		DynamicParticles					= 1 << 5,
+		DynamicKinematicParticles			= 1 << 6,
+		ClusteredParticles					= 1 << 7,
+		GeometryCollectionParticles			= 1 << 8,
+
+		// Particle lists (each particle can be zero or more of these, depending on state
+		StaticGeometryCollectionArray		= 1 << 9,
+		KinematicGeometryCollectionArray	= 1 << 10,
+		SleepingGeometryCollectionArray		= 1 << 11,
+		DynamicGeometryCollectionArray		= 1 << 12,
+		ActiveParticlesMapArray				= 1 << 13,
+		TransientDirtyMapArray				= 1 << 14,
+		MovingKinematicsMapArray			= 1 << 15,
+		StaticClusteredMapArray				= 1 << 16,
+		KinematicClusteredMapArray			= 1 << 18,
+		DynamicClusteredMapArray			= 1 << 19,
+		ResimActiveParticlesMapArray		= 1 << 20,
+		ResimDynamicParticles				= 1 << 21,
+		ResimDynamicKinematicParticles		= 1 << 22,
+		ResimStaticParticles				= 1 << 23,
+		ResimKinematicParticles				= 1 << 24,
+	};
+	ENUM_CLASS_FLAGS(EGeometryParticleListMask);
 }
\ No newline at end of file
diff --git a/Engine/Source/Runtime/Experimental/Chaos/Public/Chaos/PBDRigidsSOAs.h b/Engine/Source/Runtime/Experimental/Chaos/Public/Chaos/PBDRigidsSOAs.h
--- a/Engine/Source/Runtime/Experimental/Chaos/Public/Chaos/PBDRigidsSOAs.h
+++ b/Engine/Source/Runtime/Experimental/Chaos/Public/Chaos/PBDRigidsSOAs.h
@@ -10,6 +10,19 @@ namespace Chaos
 namespace CVars
 {
 	extern CHAOS_API bool bRemoveParticleFromMovingKinematicsOnDisable;
+	extern CHAOS_API bool bChaosSolverCheckParticleViews;
+}
+
+// A helper function to get the debug name for a particle or "UNKNOWN" when debug names are unavailable
+template<typename TParticleType>
+const FString& GetParticleDebugName(const TParticleType& Particle)
+{
+#if CHAOS_DEBUG_NAME
+	return *Particle.DebugName();
+#else
+	static const FString SNoDebugNames = TEXT("<UNKNOWN>");
+	return SNoDebugNames;
+#endif
 }
 
 class IParticleUniqueIndices
@@ -106,10 +119,20 @@ private:
 
 using FParticleUniqueIndices = FParticleUniqueIndicesMultithreaded;
 
+/**
+* A list of particles held in a array.
+* Also contains a particle-index map for faster find/removal with large lists.
+* @see TParticleArray
+*/
 template <typename TParticleType>
 class TParticleMapArray
 {
 public:
+	TParticleMapArray()
+		: ContainerListMask(EGeometryParticleListMask::None)
+	{
+	}
+
 	void Reset()
 	{
 		ParticleToIndex.Reset();
@@ -188,9 +211,58 @@ public:
 	const TArray<TParticleType*>& GetArray() const { return ParticleArray;	}
 	TArray<TParticleType*>& GetArray() { return ParticleArray; }
 
+	void SetContainerListMask(const EGeometryParticleListMask InMask) { ContainerListMask = InMask; }
+	EGeometryParticleListMask GetContainerListMask() const { return ContainerListMask; }
+
 private:
 	TMap<TParticleType*, int32> ParticleToIndex;
 	TArray<TParticleType*> ParticleArray;
+	EGeometryParticleListMask ContainerListMask;
+};
+
+/**
+* A list of particles held in a array.
+* O(N) find/remove so only for use on small lists. 
+* @see TParticleArrayMap
+*/
+template <typename TParticleType>
+class TParticleArray
+{
+public:
+	TParticleArray()
+		: ContainerListMask(EGeometryParticleListMask::None)
+	{
+	}
+
+	void Reset()
+	{
+		ParticleArray.Reset();
+	}
+
+	bool Contains(const TParticleType* Particle) const
+	{
+		return ParticleArray.Contains(Particle);
+	}
+
+	void Insert(TParticleType* Particle)
+	{
+		ParticleArray.Add(Particle);
+	}
+
+	void Remove(TParticleType* Particle)
+	{
+		ParticleArray.Remove(Particle);
+	}
+
+	const TArray<TParticleType*>& GetArray() const { return ParticleArray; }
+	TArray<TParticleType*>& GetArray() { return ParticleArray; }
+
+	void SetContainerListMask(const EGeometryParticleListMask InMask) { ContainerListMask = InMask; }
+	EGeometryParticleListMask GetContainerListMask() const { return ContainerListMask; }
+
+private:
+	TArray<TParticleType*> ParticleArray;
+	EGeometryParticleListMask ContainerListMask;
 };
 
 class FPBDRigidsSOAs
@@ -204,19 +276,21 @@ public:
 #endif
 
 		StaticParticles = MakeUnique<FGeometryParticles>();
-		StaticDisabledParticles = MakeUnique <FGeometryParticles>();
+		StaticDisabledParticles = MakeUnique<FGeometryParticles>();
 
-		KinematicParticles = MakeUnique < FKinematicGeometryParticles>();
-		KinematicDisabledParticles = MakeUnique < FKinematicGeometryParticles>();
+		KinematicParticles = MakeUnique<FKinematicGeometryParticles>();
+		KinematicDisabledParticles = MakeUnique<FKinematicGeometryParticles>();
 
 		DynamicDisabledParticles = MakeUnique<FPBDRigidParticles>();
 		DynamicParticles = MakeUnique<FPBDRigidParticles>();
 		DynamicKinematicParticles = MakeUnique<FPBDRigidParticles>();
 
-		ClusteredParticles = MakeUnique< FPBDRigidClusteredParticles>();
+		ClusteredParticles = MakeUnique<FPBDRigidClusteredParticles>();
 
 		GeometryCollectionParticles = MakeUnique<TPBDGeometryCollectionParticles<FReal, 3>>();
 
+		SetContainerListMasks();
+
 		UpdateViews();
 	}
 
@@ -351,7 +425,10 @@ public:
 			//if active it's already in the dirty view
 			if (!ActiveParticlesMapArray.Contains(Rigid) && !MovingKinematicsMapArray.Contains(Rigid))
 			{
-				TransientDirtyMapArray.Insert(Rigid);
+				AddToList(Rigid, TransientDirtyMapArray);
+		
+				CheckParticleViewMask(Particle);
+
 				if (bUpdateViews)
 				{
 					UpdateViews();
@@ -364,11 +441,11 @@ public:
 	{
 		if (bResimulating)
 		{
-			ResimStaticParticles.Remove(Particle);
+			RemoveFromList(Particle, ResimStaticParticles);
 			FKinematicGeometryParticleHandle* Kinematic = Particle->CastToKinematicParticle();
 			if (Kinematic)
 			{
-				ResimKinematicParticles.Remove(Kinematic);
+				RemoveFromList(Kinematic, ResimKinematicParticles);
 			}
 		}
 
@@ -379,12 +456,12 @@ public:
 		{
 			if (bResimulating)
 			{
-				ResimDynamicParticles.Remove(PBDRigid);
-				ResimDynamicKinematicParticles.Remove(PBDRigid);
+				RemoveFromList(PBDRigid, ResimDynamicParticles);
+				RemoveFromList(PBDRigid, ResimDynamicKinematicParticles);
 			}
 
 			RemoveFromActiveArray(PBDRigid, /*bStillDirty=*/ false);
-			MovingKinematicsMapArray.Remove(PBDRigid);
+			RemoveFromList(PBDRigid, MovingKinematicsMapArray);
 
 			if (TPBDGeometryCollectionParticleHandle<FReal, 3>* GCHandle = Particle->CastToGeometryCollection())
 			{
@@ -445,7 +522,7 @@ public:
 				RemoveFromActiveArray(PBDRigid, /*bStillDirty=*/false);
 				if (CVars::bRemoveParticleFromMovingKinematicsOnDisable)
 				{
-					MovingKinematicsMapArray.Remove(PBDRigid);
+					RemoveFromList(PBDRigid, MovingKinematicsMapArray);
 				}
 			}
 		}
@@ -457,6 +534,9 @@ public:
 		{
 			Particle->MoveToSOA(*StaticDisabledParticles);
 		}
+
+		CheckParticleViewMask(Particle);
+
 		UpdateViews();
 	}
 
@@ -498,6 +578,9 @@ public:
 		{
 			Particle->MoveToSOA(*StaticParticles);
 		}
+
+		CheckParticleViewMask(Particle);
+
 		UpdateViews();
 	}
 
@@ -525,6 +608,8 @@ public:
 					check(PBDRigid->ObjectState() == EObjectStateType::Dynamic);
 					AddToActiveArray(PBDRigid);
 
+					CheckParticleViewMask(Particle);
+
 					if(!DeferUpdateViews)
 					{
 						UpdateViews();
@@ -575,6 +660,8 @@ public:
 					}
 					RemoveFromActiveArray(PBDRigid, /*bStillDirty=*/true);
 
+					CheckParticleViewMask(Particle);
+
 					if (!DeferUpdateViews)
 					{
 						UpdateViews();
@@ -614,7 +701,7 @@ public:
 			RemoveFromActiveArray(Particle->CastToRigidParticle(), /*bStillDirty=*/ false);
 			if (CVars::bRemoveParticleFromMovingKinematicsOnDisable)
 			{
-				MovingKinematicsMapArray.Remove(Particle);
+				RemoveFromList(Particle, MovingKinematicsMapArray);
 			}
 		}
 		else
@@ -630,7 +717,7 @@ public:
 
 			if (Particle->ObjectState() != EObjectStateType::Kinematic)
 			{
-				MovingKinematicsMapArray.Remove(Particle);
+				RemoveFromList(Particle, MovingKinematicsMapArray);
 			}
 
 			if (Particle->Type == EParticleType::Clustered)
@@ -638,8 +725,8 @@ public:
 				Particle->MoveToSOA(*ClusteredParticles);
 				if (bResimulating)
 				{
-					ResimDynamicKinematicParticles.Remove(Particle);
-					ResimDynamicParticles.Insert(Particle);
+					RemoveFromList(Particle, ResimDynamicKinematicParticles);
+					AddToList(Particle, ResimDynamicParticles);
 				}
 			}
 			else if (Particle->Type == EParticleType::GeometryCollection)
@@ -647,8 +734,8 @@ public:
 				Particle->MoveToSOA(*GeometryCollectionParticles);
 				if (bResimulating)
 				{
-					ResimDynamicKinematicParticles.Remove(Particle);
-					ResimDynamicParticles.Insert(Particle);
+					RemoveFromList(Particle, ResimDynamicKinematicParticles);
+					AddToList(Particle, ResimDynamicParticles);
 				}
 			}
 			else
@@ -660,8 +747,8 @@ public:
 					Particle->MoveToSOA(*DynamicKinematicParticles);
 					if (bResimulating)
 					{
-						ResimDynamicParticles.Remove(Particle);
-						ResimDynamicKinematicParticles.Insert(Particle);
+						RemoveFromList(Particle, ResimDynamicParticles);
+						AddToList(Particle, ResimDynamicKinematicParticles);
 					}
 					break;
 
@@ -669,8 +756,8 @@ public:
 					Particle->MoveToSOA(*DynamicParticles);
 					if (bResimulating)
 					{
-						ResimDynamicKinematicParticles.Remove(Particle);
-						ResimDynamicParticles.Insert(Particle);
+						RemoveFromList(Particle, ResimDynamicKinematicParticles);
+						AddToList(Particle, ResimDynamicParticles);
 					}
 					break;
 
@@ -679,13 +766,16 @@ public:
 					Particle->MoveToSOA(*DynamicParticles);
 					if (bResimulating)
 					{
-						ResimDynamicKinematicParticles.Remove(Particle);
-						ResimDynamicParticles.Insert(Particle);
+						RemoveFromList(Particle, ResimDynamicKinematicParticles);
+						AddToList(Particle, ResimDynamicParticles);
 					}
 					break;
 				}
 			}
 		}
+
+		CheckParticleViewMask(Particle);
+
 		UpdateViews();
 	}
 
@@ -693,7 +783,7 @@ public:
 	{
 		if (ClusteredParticle->ObjectState() != EObjectStateType::Kinematic)
 		{
-			MovingKinematicsMapArray.Remove(ClusteredParticle);
+			RemoveFromList(ClusteredParticle, MovingKinematicsMapArray);
 		}
 
 		if (TPBDGeometryCollectionParticleHandle<FReal, 3>* GCParticle = ClusteredParticle->CastToGeometryCollection())
@@ -718,6 +808,8 @@ public:
 				AddToActiveArray(ClusteredParticle->CastToRigidParticle());
 			}
 		}
+
+		CheckParticleViewMask(ClusteredParticle);
 		
 		UpdateViews();
 	}
@@ -726,9 +818,11 @@ public:
 	{
 		if (FPBDRigidParticleHandle* Rigid = Particle->CastToRigidParticle())
 		{
-			MovingKinematicsMapArray.Insert(Rigid);
+			AddToList(Rigid, MovingKinematicsMapArray);
 			// Remove from TransientDirtyMapArray to be sure not have have particle in both lists
-			TransientDirtyMapArray.Remove(Rigid);
+			RemoveFromList(Rigid, TransientDirtyMapArray);
+
+			CheckParticleViewMask(Particle);
 		}
 	}
 
@@ -744,12 +838,14 @@ public:
 			{
 				// remove and do not increment Index
 				FPBDRigidParticleHandle* Rigid = MovingKinematicsMapArray.GetArray()[Index];
-				MovingKinematicsMapArray.Remove(Rigid);
+				RemoveFromList(Rigid, MovingKinematicsMapArray);
 
 				// Particle may not be moving, but its velocity was just changed (to zero), so it
 				// needs to be in a dirty list. Since it are no longer in MovingKinematicsMapArray
 				// we must put it the transient dirty list
 				MarkTransientDirtyParticle(Rigid);
+
+				CheckParticleViewMask(Rigid);
 			}
 			else
 			{
@@ -812,6 +908,12 @@ public:
 
 		//todo: update deterministic ID
 
+		// We do not serialize the list masks to simplify maintenance
+		if (Ar.IsLoading())
+		{
+			SetContainerListMasks();
+		}
+
 		//if (!GeometryCollectionParticles || !GeometryCollectionParticles->Size())
 			UpdateViews();
 		//else
@@ -878,16 +980,16 @@ public:
 				ensure(false); // we should probably not be here 
 				break;
 			case EObjectStateType::Static:
-				StaticGeometryCollectionArray.Insert(GCParticle);
+				AddToList(GCParticle, StaticGeometryCollectionArray);
 				break;
 			case EObjectStateType::Kinematic:
-				KinematicGeometryCollectionArray.Insert(GCParticle);
+				AddToList(GCParticle, KinematicGeometryCollectionArray);
 				break;
 			case EObjectStateType::Dynamic:
-				DynamicGeometryCollectionArray.Insert(GCParticle);
+				AddToList(GCParticle, DynamicGeometryCollectionArray);
 				break;
 			case EObjectStateType::Sleeping:
-				SleepingGeometryCollectionArray.Insert(GCParticle);
+				AddToList(GCParticle, SleepingGeometryCollectionArray);
 				break;
 			}
 		}
@@ -895,10 +997,10 @@ public:
 
 	void RemoveGeometryCollectionParticle(TPBDGeometryCollectionParticleHandle<FReal, 3>* GCParticle)
 	{
-		StaticGeometryCollectionArray.Remove(GCParticle);
-		KinematicGeometryCollectionArray.Remove(GCParticle);
-		SleepingGeometryCollectionArray.Remove(GCParticle);
-		DynamicGeometryCollectionArray.Remove(GCParticle);
+		RemoveFromList(GCParticle, StaticGeometryCollectionArray);
+		RemoveFromList(GCParticle, KinematicGeometryCollectionArray);
+		RemoveFromList(GCParticle, SleepingGeometryCollectionArray);
+		RemoveFromList(GCParticle, DynamicGeometryCollectionArray);
 	}
 
 	const auto& GetClusteredParticles() const { return *ClusteredParticles; }
@@ -906,7 +1008,104 @@ public:
 
 	auto& GetUniqueIndices() { return UniqueIndices; }
 
+	// Check that the particles in each of the particle lists have their ListMask sett appropriately
+	CHAOS_API void CheckListMasks();
+	
+	// Check that no particles are in multiple lists that contribute to a single view
+	CHAOS_API void CheckViewMasks();
+
 private:
+
+	//
+	// List Tracking System: 
+	// every particle stores the set of SOAs/ParticleMapArrays/ParticleArrays 
+	// that the particle is in. This information is used to check for invalid 
+	// combinations (some list memberships are mutually exclusive) and to 
+	// reduce the cost of checking whether a particle is in a particular list.
+	//
+
+	// Initialize the list mask for all particle containers and lists
+	CHAOS_API void SetContainerListMasks();
+
+	// Add a particle to the specified SOA/List and update its ListMask
+	template<typename TListType, typename TParticleHandleType>
+	void AddToList(TParticleHandleType* Particle, TListType& List)
+	{
+		//if (!Particle->IsInAnyList(List.GetContainerListMask()))
+		{
+			Particle->AddToLists(List.GetContainerListMask());
+			List.Insert(Particle);
+		}
+	}
+
+	// Add a set of particles to the specified SOA/List and update their ListMasks
+	template<typename TListType, typename TParticleHandleType>
+	void AddToList(const TArray<TParticleHandleType*> Particles, TListType& List)
+	{
+		check(List.GetContainerListMask() != EGeometryParticleListMask::None);
+
+		for (TParticleHandleType* Particle : Particles)
+		{
+			Particle->AddToLists(List.GetContainerListMask());
+		}
+		List.Insert(Particles);
+	}
+
+	// Remove a particle from the specified SOA/List and update its ListMask
+	template<typename TListType, typename TParticleHandleType>
+	void RemoveFromList(TParticleHandleType* Particle, TListType& List)
+	{
+		check(List.GetContainerListMask() != EGeometryParticleListMask::None);
+
+		//if (Particle->IsInAnyList(List.GetContainerListMask()))
+		{
+			Particle->RemoveFromLists(List.GetContainerListMask());
+			List.Remove(Particle);
+		}
+	}
+
+	// Check that the particles in an SOA have the appropriate bit set in their ListMask
+	template<typename TSOAType>
+	void CheckSOAMasks(TSOAType* SOA)
+	{
+		TArray<TSOAView<FGeometryParticles>> SOAViewArray = { SOA };
+		auto View = MakeConstParticleView(MoveTemp(SOAViewArray));
+		for (const auto& Particle : View)
+		{
+			ensureAlwaysMsgf(Particle.IsInAnyList(SOA->GetContainerListMask()), TEXT("Particle %s ListMask missing expected SOA bit 0x%x"), *GetParticleDebugName(Particle), SOA->GetContainerListMask());
+		}
+	}
+
+	// Check that the particles in a List have the appropriate bit set in their ListMask
+	template<typename TListType>
+	void CheckListMasks(TListType& List)
+	{
+		TSOAView<FGeometryParticles> SOAView(&List.GetArray());
+		TArray<TSOAView<FGeometryParticles>> SOAViewArray = { SOAView };
+		auto View = MakeConstParticleView(MoveTemp(SOAViewArray));
+		for (const auto& Particle : View)
+		{
+			ensureAlwaysMsgf(Particle.IsInAnyList(List.GetContainerListMask()), TEXT("Particle %s ListMask missing expected List bit 0x%x"), *GetParticleDebugName(Particle), List.GetContainerListMask());
+		}
+	}
+
+	// Make sure the particle is not in a set of lists that would cause it to appear in any single view multiple times
+	void CheckParticleViewMask(const FGeometryParticleHandle* Particle) const
+	{
+		// This is called after every change to a particle's SOA/List for validation. 
+		// It should be cheap when the cvar is disabled but we fully remove it in shipping anyway
+#if !UE_BUILD_SHIPPING
+		if (CVars::bChaosSolverCheckParticleViews)
+		{
+			CheckParticleViewMaskImpl(Particle);
+		}
+#endif
+	}
+
+	// Make sure the particle is not in a set of lists that would cause it to appear in any single view multiple times
+	CHAOS_API void CheckParticleViewMaskImpl(const FGeometryParticleHandle* Particle) const;
+
+
 	template <typename TParticleHandleType, typename TParticles>
 	TArray<TParticleHandleType*> CreateParticlesHelper(int32 NumParticles, const FUniqueIdx* ExistingIndices,  TUniquePtr<TParticles>& Particles, const FGeometryParticleParameters& Params)
 	{
@@ -925,6 +1124,7 @@ private:
 
 			TUniquePtr<TParticleHandleType> NewParticleHandle = TParticleHandleType::CreateParticleHandle(MakeSerializable(Particles), ParticleIdx, HandleIdx);
 			NewParticleHandle->SetParticleID(FParticleID{ INDEX_NONE, BiggestParticleID++ });
+			NewParticleHandle->AddToLists(Particles->GetContainerListMask());
 			ReturnHandles[Count] = NewParticleHandle.Get();
 			//If unique indices are null it means there is no GT particle that already registered an ID, so create one
 			if(ExistingIndices)
@@ -944,40 +1144,40 @@ private:
 	
 	void AddToActiveArray(const TArray<FPBDRigidParticleHandle*>& Particles)
 	{
-		ActiveParticlesMapArray.Insert(Particles);
+		AddToList(Particles, ActiveParticlesMapArray);
 
 		if(bResimulating)
 		{
-			ResimActiveParticlesMapArray.Insert(Particles);
+			AddToList(Particles, ResimActiveParticlesMapArray);
 		}
 		
 		//dirty contains Active so make sure no duplicates
 		for(FPBDRigidParticleHandle* Particle : Particles)
 		{
-			TransientDirtyMapArray.Remove(Particle);
+			RemoveFromList(Particle, TransientDirtyMapArray);
 		}
 	}
 
 	void AddToActiveArray(FPBDRigidParticleHandle* Particle)
 	{
-		ActiveParticlesMapArray.Insert(Particle);
+		AddToList(Particle, ActiveParticlesMapArray);
 
 		if (bResimulating)
 		{
-			ResimActiveParticlesMapArray.Insert(Particle);
+			AddToList(Particle, ResimActiveParticlesMapArray);
 		}
 
 		//dirty contains Active so make sure no duplicates
-		TransientDirtyMapArray.Remove(Particle);
+		RemoveFromList(Particle, TransientDirtyMapArray);
 	}
 
 	void RemoveFromActiveArray(FPBDRigidParticleHandle* Particle, bool bStillDirty)
 	{
-		ActiveParticlesMapArray.Remove(Particle);
+		RemoveFromList(Particle, ActiveParticlesMapArray);
 
 		if (bResimulating)
 		{
-			ResimActiveParticlesMapArray.Remove(Particle);
+			RemoveFromList(Particle, ResimActiveParticlesMapArray);
 		}
 
 		if(bStillDirty)
@@ -985,178 +1185,19 @@ private:
 			if(!MovingKinematicsMapArray.Contains(Particle))
 			{ 
 				//no longer active, but still dirty
-				TransientDirtyMapArray.Insert(Particle);
+				AddToList(Particle, TransientDirtyMapArray);
 			}
 		}
 		else
 		{
 			//might have already been removed from active from a previous call
 			//but now removing and don't want it dirty either
-			TransientDirtyMapArray.Remove(Particle);
+			RemoveFromList(Particle, TransientDirtyMapArray);
 		}
 	}
 	
 	//should be called whenever particles are added / removed / reordered
-	void UpdateViews()
-	{
-		//build various views. Group SOA types together for better branch prediction
-		{
-			auto TmpArray = bResimulating
-			? TArray<TSOAView<FGeometryParticles>>
-			{
-				// re-sim only works with a reduced set of particle types, i.e. no DynamicClusteredMapArray (DynamicGeometryCollectionArray??) BH
-				{&ResimStaticParticles},
-				{&ResimKinematicParticles},
-				{&ResimDynamicParticles.GetArray() },
-				{&ResimDynamicKinematicParticles.GetArray() }
-			}
-			: TArray<TSOAView<FGeometryParticles>>
-			{ 
-				StaticParticles.Get(), 
-				KinematicParticles.Get(), 
-				DynamicParticles.Get(),
-				DynamicKinematicParticles.Get(),
-				{&StaticClusteredMapArray.GetArray()},
-				{&KinematicClusteredMapArray.GetArray()},
-				{&DynamicClusteredMapArray.GetArray()},
-				{&StaticGeometryCollectionArray.GetArray()},
-				{&KinematicGeometryCollectionArray.GetArray()},
-				{&SleepingGeometryCollectionArray.GetArray()},
-				{&DynamicGeometryCollectionArray.GetArray()}
-			};
-			NonDisabledView = MakeParticleView(MoveTemp(TmpArray));
-		}
-		{
-			auto TmpArray = bResimulating
-			? TArray<TSOAView<FPBDRigidParticles>>
-			{ 
-				{&ResimDynamicParticles.GetArray() }
-			}
-			: TArray<TSOAView<FPBDRigidParticles>>
-			{
-				DynamicParticles.Get(),
-				{&DynamicClusteredMapArray.GetArray() },
-				{&SleepingGeometryCollectionArray.GetArray() },
-				{&DynamicGeometryCollectionArray.GetArray() }
-			};
-			NonDisabledDynamicView = MakeParticleView(MoveTemp(TmpArray));
-		}
-		{
-			TArray<TSOAView<FPBDRigidParticles>> TmpArray =
-			{ 
-				{&ActiveParticlesMapArray.GetArray()},
-				{&MovingKinematicsMapArray.GetArray()},
-				//{&KinematicClusteredMapArray.GetArray()},	// @todo(chaos): we should include moving clustered kinematics in this view
-			};
-			ActiveDynamicMovingKinematicParticlesView = MakeParticleView(MoveTemp(TmpArray));
-		}
-		{
-			auto TmpArray = bResimulating
-			? TArray<TSOAView<FPBDRigidParticles>>
-			{
-				{&ResimActiveParticlesMapArray.GetArray()}
-			}
-			: TArray<TSOAView<FPBDRigidParticles>>
-			{ 
-				{&ActiveParticlesMapArray.GetArray()},
-				{&StaticGeometryCollectionArray.GetArray()},
-				{&KinematicGeometryCollectionArray.GetArray()},
-			};
-			ActiveParticlesView = MakeParticleView(MoveTemp(TmpArray));
-		}
-		{
-			auto TmpArray = bResimulating
-			? TArray<TSOAView<FPBDRigidParticles>>
-			{
-				{&ResimActiveParticlesMapArray.GetArray()},
-				{&TransientDirtyMapArray.GetArray()}
-			}
-			: TArray<TSOAView<FPBDRigidParticles>>
-			{
-				{&ActiveParticlesMapArray.GetArray()},
-				{&MovingKinematicsMapArray.GetArray()},
-				{&SleepingGeometryCollectionArray.GetArray()},
-				{&TransientDirtyMapArray.GetArray()}
-			};
-			DirtyParticlesView = MakeParticleView(MoveTemp(TmpArray));
-		}
-		{
-			TArray<TSOAView<FGeometryParticles>> TmpArray = 
-			{ 
-				StaticParticles.Get(), 
-				StaticDisabledParticles.Get(), 
-				KinematicParticles.Get(), 
-				KinematicDisabledParticles.Get(),
-				DynamicParticles.Get(), 
-				DynamicDisabledParticles.Get(), 
-				DynamicKinematicParticles.Get(), 
-				ClusteredParticles.Get(), 
-				GeometryCollectionParticles.Get() 
-			};
-			AllParticlesView = MakeParticleView(MoveTemp(TmpArray));
-		}
-		{
-			auto TmpArray = bResimulating
-			? TArray<TSOAView<FKinematicGeometryParticles>>
-			{
-				{&ResimKinematicParticles},
-				{&ResimDynamicKinematicParticles.GetArray()}
-			}
-			: TArray<TSOAView<FKinematicGeometryParticles>>
-			{ 
-				KinematicParticles.Get(),
-				DynamicKinematicParticles.Get(),
-				{&KinematicGeometryCollectionArray.GetArray()},
-				{&KinematicClusteredMapArray.GetArray()}
-			};
-			ActiveKinematicParticlesView = MakeParticleView(MoveTemp(TmpArray));
-		}
-		{
-			// todo(chaos) handle resim ?
-			TArray<TSOAView<FPBDRigidParticles>> TmpArray =
-			{
-				{&MovingKinematicsMapArray.GetArray()},
-			};
-			ActiveMovingKinematicParticlesView = MakeParticleView(MoveTemp(TmpArray));
-		}
-		{
-			auto TmpArray = bResimulating
-			? TArray<TSOAView<FGeometryParticles>>
-			{
-				{&ResimStaticParticles}
-			}
-			: TArray<TSOAView<FGeometryParticles>>
-			{
-				StaticParticles.Get(),
-				{&StaticClusteredMapArray.GetArray()}
-			};
-			ActiveStaticParticlesView = MakeParticleView(MoveTemp(TmpArray));
-		}
-		{
-			auto TmpArray = bResimulating
-			? TArray<TSOAView<TPBDGeometryCollectionParticles<FReal, 3>>> {}	//no geometry collection during resim
-			: TArray<TSOAView<TPBDGeometryCollectionParticles<FReal, 3>>>
-			{ 
-				{&StaticGeometryCollectionArray.GetArray()},
-				{&KinematicGeometryCollectionArray.GetArray()},
-				{&DynamicGeometryCollectionArray.GetArray()}
-			};
-			ActiveGeometryCollectionParticlesView = MakeParticleView(MoveTemp(TmpArray));
-		}
-
-		{
-			auto TmpArray = bResimulating
-			? TArray<TSOAView<FPBDRigidClusteredParticles>> {}	//no clusters during resim
-			: TArray<TSOAView<FPBDRigidClusteredParticles>>
-			{
-				{&StaticClusteredMapArray.GetArray()},
-				{&KinematicClusteredMapArray.GetArray()},
-				{&DynamicClusteredMapArray.GetArray()}
-			};
-			NonDisabledClusteredView = MakeParticleView(MoveTemp(TmpArray));
-		}
-		
-	}
+	CHAOS_API void UpdateViews();
 
 	void InsertClusteredParticle(FPBDRigidClusteredParticleHandle* ClusteredParticle)
 	{
@@ -1168,14 +1209,14 @@ private:
 				ensure(false); // we should probably not be here 
 				break;
 			case EObjectStateType::Static:
-				StaticClusteredMapArray.Insert(ClusteredParticle);
+				AddToList(ClusteredParticle, StaticClusteredMapArray);
 				break;
 			case EObjectStateType::Kinematic:
-				KinematicClusteredMapArray.Insert(ClusteredParticle);
+				AddToList(ClusteredParticle, KinematicClusteredMapArray);
 				break;
 			case EObjectStateType::Dynamic:
 			case EObjectStateType::Sleeping:
-				DynamicClusteredMapArray.Insert(ClusteredParticle);
+				AddToList(ClusteredParticle, DynamicClusteredMapArray);
 				break;
 			}
 		}
@@ -1191,9 +1232,9 @@ private:
 
 	void RemoveClusteredParticle(FPBDRigidClusteredParticleHandle* ClusteredParticle)
 	{
-		StaticClusteredMapArray.Remove(ClusteredParticle);
-		KinematicClusteredMapArray.Remove(ClusteredParticle);
-		DynamicClusteredMapArray.Remove(ClusteredParticle);
+		RemoveFromList(ClusteredParticle, StaticClusteredMapArray);
+		RemoveFromList(ClusteredParticle, KinematicClusteredMapArray);
+		RemoveFromList(ClusteredParticle, DynamicClusteredMapArray);
 	}
 
 	//Organized by SOA type
@@ -1211,6 +1252,11 @@ private:
 
 	TUniquePtr<TPBDGeometryCollectionParticles<FReal, 3>> GeometryCollectionParticles;
 
+	//
+	// NOTE: The member here are enumerated in EGeometryParticleListMask which must
+	// be kept up to date with changes here. Also see SetContainerListMasks() and CheckListMasks()
+	//
+
 	// Geometry collection particle state is controlled via their disabled state and assigned 
 	// EObjectStateType, and are shuffled into these corresponding arrays in UpdateGeometryCollectionViews().
 	TParticleMapArray<FPBDGeometryCollectionParticleHandle> StaticGeometryCollectionArray;
@@ -1218,19 +1264,19 @@ private:
 	TParticleMapArray<FPBDGeometryCollectionParticleHandle> SleepingGeometryCollectionArray;
 	TParticleMapArray<FPBDGeometryCollectionParticleHandle> DynamicGeometryCollectionArray;
 
-	//Utility structures for maintaining an Active particles view
+	// Utility structures for maintaining an Active particles view
 	TParticleMapArray<FPBDRigidParticleHandle> ActiveParticlesMapArray;
 	TParticleMapArray<FPBDRigidParticleHandle> TransientDirtyMapArray;
 	
-	// keep track of kinematic that have their kinematic target set for this current frame
+	// Keep track of kinematic that have their kinematic target set for this current frame
 	TParticleMapArray<FPBDRigidParticleHandle> MovingKinematicsMapArray;
 
 	// Structures for maintaining a subset view during a resim (TParticleMapArray used when we need to dynamically add/remove during resim)
 	TParticleMapArray<FPBDRigidParticleHandle> ResimActiveParticlesMapArray;
 	TParticleMapArray<FPBDRigidParticleHandle> ResimDynamicParticles;
 	TParticleMapArray<FPBDRigidParticleHandle> ResimDynamicKinematicParticles;
-	TArray<FGeometryParticleHandle*> ResimStaticParticles;
-	TArray<FKinematicGeometryParticleHandle*> ResimKinematicParticles;
+	TParticleArray<FGeometryParticleHandle> ResimStaticParticles;
+	TParticleArray<FKinematicGeometryParticleHandle> ResimKinematicParticles;
 	bool bResimulating = false;
 
 	// NonDisabled clustered particle arrays
@@ -1238,7 +1284,7 @@ private:
 	TParticleMapArray<FPBDRigidClusteredParticleHandle> KinematicClusteredMapArray;
 	TParticleMapArray<FPBDRigidClusteredParticleHandle> DynamicClusteredMapArray;
 
-	//Particle Views
+	// Particle Views
 	TParticleView<FGeometryParticles> NonDisabledView;								//all particles that are not disabled
 	TParticleView<FPBDRigidParticles> NonDisabledDynamicView;						//all dynamic particles that are not disabled
 	TParticleView<FPBDRigidClusteredParticles> NonDisabledClusteredView;			//all clustered particles that are not disabled
@@ -1251,7 +1297,7 @@ private:
 	TParticleView<FGeometryParticles> ActiveStaticParticlesView;					//all static particles that are not disabled
 	TParticleView<TPBDGeometryCollectionParticles<FReal, 3>> ActiveGeometryCollectionParticlesView; // all geom collection particles that are not disabled
 
-	//Auxiliary data synced with particle handles
+	// Auxiliary data synced with particle handles
 	TGeometryParticleHandles<FReal, 3> ParticleHandles;
 
 	IParticleUniqueIndices& UniqueIndices;
diff --git a/Engine/Source/Runtime/Experimental/Chaos/Public/Chaos/ParticleHandle.h b/Engine/Source/Runtime/Experimental/Chaos/Public/Chaos/ParticleHandle.h
--- a/Engine/Source/Runtime/Experimental/Chaos/Public/Chaos/ParticleHandle.h
+++ b/Engine/Source/Runtime/Experimental/Chaos/Public/Chaos/ParticleHandle.h
@@ -674,6 +674,31 @@ public:
 	bool LightWeightDisabled() const { return GeometryParticles->LightWeightDisabled(ParticleIdx); }
 	void SetLightWeightDisabled(bool bLightWeightDisabled) { GeometryParticles->LightWeightDisabled(ParticleIdx) = bLightWeightDisabled; }
 
+	EGeometryParticleListMask ListMask() const
+	{
+		return GeometryParticles->ListMask(ParticleIdx);
+	}
+
+	bool IsInAnyList(const EGeometryParticleListMask ListMask) const
+	{
+		return (GeometryParticles->ListMask(ParticleIdx) & ListMask) != EGeometryParticleListMask::None;
+	}
+
+	bool IsInAllLists(const EGeometryParticleListMask ListMask) const
+	{
+		return (GeometryParticles->ListMask(ParticleIdx) & ListMask) == ListMask;
+	}
+
+	void AddToLists(const EGeometryParticleListMask ListMask)
+	{
+		GeometryParticles->ListMask(ParticleIdx) |= ListMask;
+	}
+
+	void RemoveFromLists(const EGeometryParticleListMask ListMask)
+	{
+		GeometryParticles->ListMask(ParticleIdx) &= ~ListMask;
+	}
+
 #if CHAOS_DETERMINISTIC
 	FParticleID ParticleID() const { return GeometryParticles->ParticleID(ParticleIdx); }
 	void SetParticleID(const FParticleID& ParticleID)
@@ -704,8 +729,16 @@ public:
 			{
 				GeometryParticles->Handle(ParticleIdx)->ParticleIdx = ParticleIdx;
 			}
+
 			const int32 NewParticleIdx = ToSOA.Size() - 1;
 			ParticleIdx = NewParticleIdx;
+
+			// Update the particles' list mask to match its new container
+			check(GeometryParticles->GetContainerListMask() != EGeometryParticleListMask::None);
+			check(ToSOA.GetContainerListMask() != EGeometryParticleListMask::None);
+			ToSOA.ListMask(ParticleIdx) &= ~GeometryParticles->GetContainerListMask();
+			ToSOA.ListMask(ParticleIdx) |= ToSOA.GetContainerListMask();
+
 			GeometryParticles = &ToSOA;
 		}
 	}
