cluster union child to parent update needs to store which cluster union the update is for - otherwise may apply an update which forces the particle to go to the wrong location since the child to parent update was for a cluster union it is no longer in

diff --git a/Engine/Source/Runtime/Experimental/Chaos/Private/Chaos/ClusterUnionManager.cpp b/Engine/Source/Runtime/Experimental/Chaos/Private/Chaos/ClusterUnionManager.cpp
--- a/Engine/Source/Runtime/Experimental/Chaos/Private/Chaos/ClusterUnionManager.cpp
+++ b/Engine/Source/Runtime/Experimental/Chaos/Private/Chaos/ClusterUnionManager.cpp
@@ -946,6 +946,7 @@ namespace Chaos
 				}
 				Update.ChildToParent = ChildToParent[InputIndex];
 				Update.bLock = bLock;
+				Update.ClusterUnionIndex = Index;
 			}
 		}
 	}
@@ -965,6 +966,8 @@ namespace Chaos
 		// may try to recompute the ChildToParent using the position of the particle. To counteract this, we will force the ChildToParent
 		// to be *temporarily* locked for the duration of the *next* UpdateAllClusterUnionProperties. However, unless bLock is true, we will restore the
 		// lock state of the particle to what it was previously.
+
+		bool bMadeChanges = false;
 		for (FPBDRigidParticleHandle* Particle : Particles)
 		{
 			if (!ensure(Particle))
@@ -979,6 +982,11 @@ namespace Chaos
 				{
 					if (const FClusterUnionChildToParentUpdate* Update = PendingChildToParentUpdates.Find(ChildHandle))
 					{
+						if (Update->ClusterUnionIndex != ClusterIndex)
+						{
+							continue;
+						}
+
 						const FRigidTransform3 ChildToParent = Update->ChildToParent;
 						ChildHandle->SetChildToParent(ChildToParent);
 
@@ -1004,12 +1012,16 @@ namespace Chaos
 						// the child to parent update might move the node so far away as to make the old connectivity edges incorrect.
 						ClusterUnion->PendingConnectivityOperations.Add({ Particle, EClusterUnionConnectivityOperation::Remove });
 						ClusterUnion->PendingConnectivityOperations.Add({ Particle, EClusterUnionConnectivityOperation::Add });
+						bMadeChanges = true;
 					}
 				}
 			}
 		}
 
-		RequestDeferredClusterPropertiesUpdate(ClusterIndex, EUpdateClusterUnionPropertiesFlags::IncrementalGenerateConnectionGraph);
+		if (bMadeChanges)
+		{
+			RequestDeferredClusterPropertiesUpdate(ClusterIndex, EUpdateClusterUnionPropertiesFlags::IncrementalGenerateConnectionGraph);
+		}
 	}
 
 	bool FClusterUnionManager::IsDirectlyConnectedToMainParticleInClusterUnion(const FClusterUnion& ClusterUnion, FPBDRigidParticleHandle* Particle) const
diff --git a/Engine/Source/Runtime/Experimental/Chaos/Public/Chaos/ClusterUnionManager.h b/Engine/Source/Runtime/Experimental/Chaos/Public/Chaos/ClusterUnionManager.h
--- a/Engine/Source/Runtime/Experimental/Chaos/Public/Chaos/ClusterUnionManager.h
+++ b/Engine/Source/Runtime/Experimental/Chaos/Public/Chaos/ClusterUnionManager.h
@@ -110,6 +110,7 @@ namespace Chaos
 	{
 		FTransform ChildToParent;
 		bool bLock = false;
+		FClusterUnionIndex ClusterUnionIndex = INDEX_NONE;
 	};
 
 	/**
