Fix potential truncation when logging IAS http/io operations

#rnx

diff --git a/Engine/Source/Runtime/Experimental/IoStoreOnDemand/Private/OnDemandIoDispatcherBackend.cpp b/Engine/Source/Runtime/Experimental/IoStoreOnDemand/Private/OnDemandIoDispatcherBackend.cpp
--- a/Engine/Source/Runtime/Experimental/IoStoreOnDemand/Private/OnDemandIoDispatcherBackend.cpp
+++ b/Engine/Source/Runtime/Experimental/IoStoreOnDemand/Private/OnDemandIoDispatcherBackend.cpp
@@ -129,10 +129,10 @@ static void LatencyTest(FStringView InUrl, FStringView InPath)
 #endif // !UE_BUILD_SHIPPING
 
 ///////////////////////////////////////////////////////////////////////////////
-static void LogHttpResult(const TCHAR* Url, uint32 StatusCode, uint32 DurationMs, uint32 Size, uint32 Offset, const char* Memo="ok")
+static void LogHttpResult(const TCHAR* Url, uint32 StatusCode, uint64 DurationMs, uint64 Size, uint64 Offset, const char* Memo="ok")
 {
 	Size >>= 10;
-	UE_LOG(LogIas, VeryVerbose, TEXT("http-%3u: %5ums %5uKiB [%7u] '%S' %s"), StatusCode, DurationMs, Size, Offset, Memo, Url);
+	UE_LOG(LogIas, VeryVerbose, TEXT("http-%3u: %5" UINT64_FMT "ms %5" UINT64_FMT "KiB[%7" UINT64_FMT "] '%S' %s"), StatusCode, DurationMs, Size, Offset, Memo, Url);
 };
 
 using namespace UE::Tasks;
@@ -1011,10 +1011,10 @@ struct FChunkRequest
 static void LogIoResult(
 	const FIoChunkId& ChunkId,
 	const FIoHash& UrlHash,
-	uint32 DurationMs,
-	uint32 UncompressedSize,
-	uint32 UncompressedOffset,
-	uint32 CompressedOffset,
+	uint64 DurationMs,
+	uint64 UncompressedSize,
+	uint64 UncompressedOffset,
+	uint64 CompressedOffset,
 	int32 Priority,
 	bool bCached)
 {
@@ -1026,7 +1026,8 @@ static void LogIoResult(
 		}
 		return bCached ? TEXT("io-cache") : TEXT("io-http ");
 	}();
-	UE_LOG(LogIas, VeryVerbose, TEXT("%s: %5ums %5uKiB [%7u] %s:%s|%u (%d)"),
+
+	UE_LOG(LogIas, VeryVerbose, TEXT("%s: %5" UINT64_FMT "ms %5" UINT64_FMT "KiB[%7" UINT64_FMT "] % s: % s | %" UINT64_FMT "(%d)"),
 		Prefix,
 		DurationMs,
 		UncompressedSize >> 10,
