Remove pending count from the IAS analytics

#rnx

- Not really useful info for analytics

diff --git a/Engine/Source/Runtime/Experimental/IoStoreOnDemand/Private/Statistics.cpp b/Engine/Source/Runtime/Experimental/IoStoreOnDemand/Private/Statistics.cpp
--- a/Engine/Source/Runtime/Experimental/IoStoreOnDemand/Private/Statistics.cpp
+++ b/Engine/Source/Runtime/Experimental/IoStoreOnDemand/Private/Statistics.cpp
@@ -280,7 +280,6 @@ void FOnDemandIoBackendStats::ReportAnalytics(TArray<FAnalyticsEventAttribute>&
 		TEXT("IasHttpErrorCount"), GHttpErrorCount.Get(), 
 		TEXT("IasHttpRetryCount"), GHttpRetryCount.Get(),
 		TEXT("IasHttpGetCount"), GHttpGetCount.Get(),
-		TEXT("IasHttpPendingCount"), GHttpPendingCount.Get(),
 		TEXT("IasHttpDownloadedBytes"), GHttpDownloadedBytes.Get(),
 		TEXT("IasHttpDurationMeanAvg"), GHttpAvgDuration.GetMean(),
 		TEXT("IasHttpDurationStdDev"), GHttpAvgDuration.GetDeviation(),
