IAS - recreate the HTTP client between attempts when downloading the TCO

- if the client fails to resolve the endpoint the socket pool is set to an error state. The next attempt will skip resolving the host since the socket pool is in an error state and immediately fail trying to connect. Resetting the HTTP client between attempts fixes this issue.

#rnx

diff --git a/Engine/Source/Runtime/Experimental/IoStoreOnDemand/Private/OnDemandIoDispatcherBackend.cpp b/Engine/Source/Runtime/Experimental/IoStoreOnDemand/Private/OnDemandIoDispatcherBackend.cpp
--- a/Engine/Source/Runtime/Experimental/IoStoreOnDemand/Private/OnDemandIoDispatcherBackend.cpp
+++ b/Engine/Source/Runtime/Experimental/IoStoreOnDemand/Private/OnDemandIoDispatcherBackend.cpp
@@ -1215,7 +1215,7 @@ public:
 private:
 	void CompleteRequest(FChunkRequest* ChunkRequest);
 	FIoStatus MountDeferredEndpoints(const FString& DistributionUrl, const TConstArrayView<FString>& ServiceUrls);
-	static TIoStatusOr<FOnDemandToc> GetToc(FHttpClient& HttpClient, const FString& TocPath);
+	static TIoStatusOr<FOnDemandToc> GetToc(const FOnDemandEndpoint& Endpoint);
 	FIoStatus AddToc(const FOnDemandEndpoint& Endpoint);
 
 	TSharedPtr<IIoCache> Cache;
@@ -1623,19 +1623,20 @@ TArray<FIoChunkId> FOnDemandIoBackend::GetAllChunkIds()
 
 #endif // IS_PROGRAM || WITH_EDITOR
 
-TIoStatusOr<FOnDemandToc> FOnDemandIoBackend::GetToc(FHttpClient& HttpClient, const FString& TocPath)
+TIoStatusOr<FOnDemandToc> FOnDemandIoBackend::GetToc(const FOnDemandEndpoint& Endpoint)
 {
 	TRACE_CPUPROFILER_EVENT_SCOPE(FOnDemandIoBackend::GetToc);
 
-	TAnsiStringBuilder<256> Url;
-	Url << "/" << TocPath;
-
 	for (int32 Attempt = 0, MaxAttempts = GIoDispatcherMaxHttpRetryCount; Attempt <= MaxAttempts; ++Attempt)
 	{
-		UE_LOG(LogIas, Log, TEXT("Fetching TOC '%s/%s' (#%d/%d)"), *HttpClient.ServiceUrl(), *TocPath, Attempt + 1, MaxAttempts);
+		TUniquePtr<FHttpClient> HttpClient = MakeUnique<FHttpClient>(Endpoint.ServiceUrl, GIoDispatcherMaxHttpConnectionCount);
+		TAnsiStringBuilder<256> Url;
+
+		Url << "/" << Endpoint.TocPath;
+		UE_LOG(LogIas, Log, TEXT("Fetching TOC '%s/%s' (#%d/%d)"), *HttpClient->ServiceUrl(), *Endpoint.TocPath, Attempt + 1, MaxAttempts);
 		
 		TIoStatusOr<FOnDemandToc> Toc;
-		HttpClient.Get(Url.ToView(), [&Toc](TIoStatusOr<FIoBuffer> Response, uint64 DurationMs)
+		HttpClient->Get(Url.ToView(), [&Toc](TIoStatusOr<FIoBuffer> Response, uint64 DurationMs)
 		{
 			if (Response.IsOk())
 			{
@@ -1660,7 +1661,7 @@ TIoStatusOr<FOnDemandToc> FOnDemandIoBackend::GetToc(FHttpClient& HttpClient, co
 			}
 		});
 
-		while (HttpClient.Tick(true));
+		while (HttpClient->Tick());
 
 		if (Toc.IsOk())
 		{
@@ -1673,9 +1674,7 @@ TIoStatusOr<FOnDemandToc> FOnDemandIoBackend::GetToc(FHttpClient& HttpClient, co
 
 FIoStatus FOnDemandIoBackend::AddToc(const FOnDemandEndpoint& Endpoint)
 {
-	TUniquePtr<FHttpClient> Client = MakeUnique<FHttpClient>(Endpoint.ServiceUrl, GIoDispatcherMaxHttpConnectionCount);
-
-	TIoStatusOr<FOnDemandToc> Toc = GetToc(*Client, Endpoint.TocPath);
+	TIoStatusOr<FOnDemandToc> Toc = GetToc(Endpoint);
 	if (!Toc.IsOk())
 	{
 		return FIoStatus(Toc.Status());
@@ -1685,7 +1684,7 @@ FIoStatus FOnDemandIoBackend::AddToc(const FOnDemandEndpoint& Endpoint)
 		FWriteScopeLock _(Lock);
 		if (!HttpClient.IsValid())
 		{
-			HttpClient = MoveTemp(Client);
+			HttpClient = MakeUnique<FHttpClient>(Endpoint.ServiceUrl, GIoDispatcherMaxHttpConnectionCount);
 			BackendThread.Reset(FRunnableThread::Create(this, TEXT("IoStoreOnDemand"), 0, TPri_AboveNormal));
 		}
 	}
