per-particle collision profiles need to load correctly on the PT even when there are nullptr particles

diff --git a/Engine/Source/Runtime/Experimental/Chaos/Private/PhysicsProxy/GeometryCollectionPhysicsProxy.cpp b/Engine/Source/Runtime/Experimental/Chaos/Private/PhysicsProxy/GeometryCollectionPhysicsProxy.cpp
--- a/Engine/Source/Runtime/Experimental/Chaos/Private/PhysicsProxy/GeometryCollectionPhysicsProxy.cpp
+++ b/Engine/Source/Runtime/Experimental/Chaos/Private/PhysicsProxy/GeometryCollectionPhysicsProxy.cpp
@@ -3696,18 +3696,17 @@ void FGeometryCollectionPhysicsProxy::UpdatePerParticleFilterData_External(const
 		{
 			Chaos::FWritePhysicsObjectInterface_Internal Interface= Chaos::FPhysicsObjectInternalInterface::GetWrite();
 
-			int32 ParticleIndex = 0;
 			for (const FParticleCollisionFilterData& Data : PerParticleData)
 			{
-				Chaos::FPhysicsObjectHandle Object = PhysicsObjects[ParticleIndex].Get();
-				if (Data.bIsValid && Object)
+				if (Data.bIsValid && Data.ParticleIndex != INDEX_NONE)
 				{
-					TArrayView<Chaos::FPhysicsObjectHandle> ParticleView{ &Object, 1 };
-					Interface.UpdateShapeCollisionFlags(ParticleView, Data.bSimEnabled, Data.bQueryEnabled);
-					Interface.UpdateShapeFilterData(ParticleView, Data.QueryFilter, Data.SimFilter);
+					if (Chaos::FPhysicsObjectHandle Object = PhysicsObjects[Data.ParticleIndex].Get())
+					{
+						TArrayView<Chaos::FPhysicsObjectHandle> ParticleView{ &Object, 1 };
+						Interface.UpdateShapeCollisionFlags(ParticleView, Data.bSimEnabled, Data.bQueryEnabled);
+						Interface.UpdateShapeFilterData(ParticleView, Data.QueryFilter, Data.SimFilter);
+					}
 				}
-
-				++ParticleIndex;
 			}
 		});
 	}
diff --git a/Engine/Source/Runtime/Experimental/Chaos/Public/PhysicsProxy/GeometryCollectionPhysicsProxy.h b/Engine/Source/Runtime/Experimental/Chaos/Public/PhysicsProxy/GeometryCollectionPhysicsProxy.h
--- a/Engine/Source/Runtime/Experimental/Chaos/Public/PhysicsProxy/GeometryCollectionPhysicsProxy.h
+++ b/Engine/Source/Runtime/Experimental/Chaos/Public/PhysicsProxy/GeometryCollectionPhysicsProxy.h
@@ -447,6 +447,7 @@ public:
 
 	struct FParticleCollisionFilterData
 	{
+		int32 ParticleIndex = INDEX_NONE;
 		bool bIsValid = false;
 		bool bQueryEnabled = false;
 		bool bSimEnabled = false;
diff --git a/Engine/Source/Runtime/Experimental/GeometryCollectionEngine/Private/GeometryCollection/GeometryCollectionComponent.cpp b/Engine/Source/Runtime/Experimental/GeometryCollectionEngine/Private/GeometryCollection/GeometryCollectionComponent.cpp
--- a/Engine/Source/Runtime/Experimental/GeometryCollectionEngine/Private/GeometryCollection/GeometryCollectionComponent.cpp
+++ b/Engine/Source/Runtime/Experimental/GeometryCollectionEngine/Private/GeometryCollection/GeometryCollectionComponent.cpp
@@ -3514,13 +3514,14 @@ void UGeometryCollectionComponent::LoadCollisionProfiles()
 
 			if (Data)
 			{
-				PerParticleData.Add(*Data);
-				Interface->UpdateShapeCollisionFlags(ParticleView, Data->bSimEnabled, Data->bQueryEnabled);
-				Interface->UpdateShapeFilterData(ParticleView, Data->QueryFilter, Data->SimFilter);
-			}
-			else
-			{
-				PerParticleData.Add({});
+				FCollisionProfileDataCache ParticleData = *Data;
+				ParticleData.ParticleIndex = ParticleIndex;
+
+				PerParticleData.Emplace(MoveTemp(ParticleData));
+
+				// Need to update on the GT too (physics proxy will just enqueue to do it on the PT) or else the changes won't take effect for GT SQ queries.
+				Interface->UpdateShapeCollisionFlags(ParticleView, ParticleData.bSimEnabled, ParticleData.bQueryEnabled);
+				Interface->UpdateShapeFilterData(ParticleView, ParticleData.QueryFilter, ParticleData.SimFilter);
 			}
 		}
 	}
