ChaosFlesh : Add Gauss Seidel master constraint

diff --git a/Engine/Source/Runtime/Experimental/Chaos/Private/Chaos/Deformable/ChaosDeformableSolver.cpp b/Engine/Source/Runtime/Experimental/Chaos/Private/Chaos/Deformable/ChaosDeformableSolver.cpp
--- a/Engine/Source/Runtime/Experimental/Chaos/Private/Chaos/Deformable/ChaosDeformableSolver.cpp
+++ b/Engine/Source/Runtime/Experimental/Chaos/Private/Chaos/Deformable/ChaosDeformableSolver.cpp
@@ -91,6 +91,9 @@ namespace Chaos::Softs
 	FCriticalSection FDeformableSolver::PackageInputMutex;
 	FCriticalSection FDeformableSolver::SolverEnabledMutex;
 
+	int32 GSParallelMax = 100;
+	FAutoConsoleVariableRef CVarDeformableGSParrelMax(TEXT("p.Chaos.Deformable.GSParallelMax"), GSParallelMax, TEXT("Minimal number of particles to process in parallel for Gauss Seidel constraints. [def: 100]"));
+
 	FDeformableSolver::FDeformableSolver(FDeformableSolverProperties InProp)
 		: CurrentInputPackage(TUniquePtr < FDeformablePackage >(nullptr))
 		, PreviousInputPackage(TUniquePtr < FDeformablePackage >(nullptr))
@@ -1198,37 +1201,58 @@ namespace Chaos::Softs
 	{
 		PERF_SCOPE(STAT_ChaosDeformableSolver_InitializeGaussSeidelConstraintVariables);
 
+		GSMasterConstraint.Reset(new Chaos::Softs::FGaussSeidelMasterConstraint<FSolverReal, FSolverParticles>(Evolution->Particles(), Property.bDoQuasistatics, Property.bUseSOR, Property.OmegaSOR, GSParallelMax));
+
 		if (Property.bUseGSNeohookean)
 		{
 			GSNeohookeanConstraints.Reset(new Chaos::Softs::FGaussSeidelNeohookeanConstraints<FSolverReal, FSolverParticles>(
 				Evolution->Particles(), *AllElements, *AllTetEMeshArray, *AllTetNuMeshArray, MoveTemp(*AllTetAlphaJArray), MoveTemp(*AllIncidentElements), MoveTemp(*AllIncidentElementsLocal), 0, Evolution->Particles().Size(), Property.bDoQuasistatics, Property.bUseSOR, Property.OmegaSOR, GDeformableXPBDCorotatedParams));
 			Evolution->ResetConstraintRules();
+			
+			GSMasterConstraint->AddStaticConstraints(GSNeohookeanConstraints->GetMeshArray(), GSNeohookeanConstraints->GetIncidentElements(), GSNeohookeanConstraints->GetIncidentElementsLocal());
+
 			int32 InitIndex1 = Evolution->AddConstraintInitRange(1, true);
 			Evolution->ConstraintInits()[InitIndex1] =
 				[this](FSolverParticles& InParticles, const FSolverReal Dt)
 			{
-				this->GSNeohookeanConstraints->Init(Dt, InParticles);
+				this->GSMasterConstraint->Init(Dt, InParticles);
 			};
+
 			int32 ConstraintIndex1 = Evolution->AddConstraintRuleRange(1, true);
 			Evolution->ConstraintRules()[ConstraintIndex1] =
 				[this](FSolverParticles& InParticles, const FSolverReal Dt)
 			{
-				this->GSNeohookeanConstraints->Apply(InParticles, Dt);
+				this->GSMasterConstraint->Apply(InParticles, Dt);
 			};
 
+
+			int32 StaticIndex = GSMasterConstraint->AddStaticConstraintResidualAndHessianRange(1);
+
+			GSMasterConstraint->StaticConstraintResidualAndHessian()[StaticIndex] = [this](const FSolverParticles& Particles, const int32 ElementIndex, const int32 ElementIndexLocal, const FSolverReal Dt, TVec3<FSolverReal>& ParticleResidual, Chaos::PMatrix<FSolverReal, 3, 3>& ParticleHessian)
+			{
+				this->GSNeohookeanConstraints->AddHyperelasticResidualAndHessian(Particles, ElementIndex, ElementIndexLocal, Dt, ParticleResidual, ParticleHessian);
+			};
+			
+			
 			if (Property.bEnablePositionTargets)
 			{
 				TArray<TArray<int32>> ParticlesPerColor;
 				GSWeakConstraints->ComputeInitialWCData(Evolution->Particles(), GSNeohookeanConstraints->GetMeshArray(), GSNeohookeanConstraints->GetIncidentElements(), GSNeohookeanConstraints->GetIncidentElementsLocal(), ParticlesPerColor);
-				GSNeohookeanConstraints->SetParticlesPerColor(MoveTemp(ParticlesPerColor));
 
-				GSNeohookeanConstraints->AddAdditionalRes = [this](const FSolverParticles& InParticles, const int32 p, const FSolverReal Dt, TVec3<FSolverReal>& res)
+				TArray<TArray<int32>> StaticIncidentElements, StaticIncidentElementsLocal;
+				TArray<TArray<int32>> StaticConstraints = GSWeakConstraints->GetStaticConstraintArrays(StaticIncidentElements, StaticIncidentElementsLocal);
+				GSMasterConstraint->AddStaticConstraints(StaticConstraints, StaticIncidentElements, StaticIncidentElementsLocal);
+
+				int32 StaticIndex1 = GSMasterConstraint->AddStaticConstraintResidualAndHessianRange(1);
+				GSMasterConstraint->StaticConstraintResidualAndHessian()[StaticIndex1] = [this](const FSolverParticles& Particles, const int32 ConstraintIndex, const int32 ConstraintIndexLocal, const FSolverReal Dt, TVec3<FSolverReal>& ParticleResidual, Chaos::PMatrix<FSolverReal, 3, 3>& ParticleHessian)
 				{
-					this->GSWeakConstraints->AddWCResidual(InParticles, p, Dt, res);
+					this->GSWeakConstraints->AddWCResidualAndHessian(Particles, ConstraintIndex, ConstraintIndexLocal, Dt, ParticleResidual, ParticleHessian);
 				};
-				GSNeohookeanConstraints->AddAdditionalHessian = [this](const FSolverParticles& InParticles, const int32 p, const FSolverReal Dt, Chaos::PMatrix<FSolverReal, 3, 3>& hessian)
+
+				int32 PerNodeIndex = GSMasterConstraint->AddPerNodeHessianRange(1);
+				GSMasterConstraint->PerNodeHessian()[PerNodeIndex] = [this](const int32 p, const FSolverReal Dt, Chaos::PMatrix<FSolverReal, 3, 3>& ParticleHessian)
 				{
-					this->GSWeakConstraints->AddWCHessian(p, Dt, hessian);
+					this->GSWeakConstraints->AddWCHessian(p, Dt, ParticleHessian);
 				};
 			}
 		} 
@@ -1237,35 +1261,58 @@ namespace Chaos::Softs
 			GSCorotatedConstraints.Reset(new Chaos::Softs::FGaussSeidelCorotatedConstraints<FSolverReal, FSolverParticles>(
 				Evolution->Particles(), *AllElements, *AllTetEMeshArray, *AllTetNuMeshArray, MoveTemp(*AllTetAlphaJArray), MoveTemp(*AllIncidentElements), MoveTemp(*AllIncidentElementsLocal), 0, Evolution->Particles().Size(), Property.bDoQuasistatics, Property.bUseSOR, Property.OmegaSOR, GDeformableXPBDCorotatedParams));
 			Evolution->ResetConstraintRules();
+
+			GSMasterConstraint->AddStaticConstraints(GSCorotatedConstraints->GetMeshArray(), GSCorotatedConstraints->GetIncidentElements(), GSCorotatedConstraints->GetIncidentElementsLocal());
+			
 			int32 InitIndex1 = Evolution->AddConstraintInitRange(1, true);
 			Evolution->ConstraintInits()[InitIndex1] =
 				[this](FSolverParticles& InParticles, const FSolverReal Dt)
 			{
-				this->GSCorotatedConstraints->Init(Dt, InParticles);
+				this->GSMasterConstraint->Init(Dt, InParticles);
 			};
+
 			int32 ConstraintIndex1 = Evolution->AddConstraintRuleRange(1, true);
 			Evolution->ConstraintRules()[ConstraintIndex1] =
 				[this](FSolverParticles& InParticles, const FSolverReal Dt)
 			{
-				this->GSCorotatedConstraints->Apply(InParticles, Dt);
+				this->GSMasterConstraint->Apply(InParticles, Dt);
+			};
+
+
+			int32 StaticIndex = GSMasterConstraint->AddStaticConstraintResidualAndHessianRange(1);
+
+			GSMasterConstraint->StaticConstraintResidualAndHessian()[StaticIndex] = [this](const FSolverParticles& Particles, const int32 ElementIndex, const int32 ElementIndexLocal, const FSolverReal Dt, TVec3<FSolverReal>& ParticleResidual, Chaos::PMatrix<FSolverReal, 3, 3>& ParticleHessian)
+			{
+				this->GSCorotatedConstraints->AddHyperelasticResidualAndHessian(Particles, ElementIndex, ElementIndexLocal, Dt, ParticleResidual, ParticleHessian);
 			};
+				
 
 			if (Property.bEnablePositionTargets)
 			{
 				TArray<TArray<int32>> ParticlesPerColor;
 				GSWeakConstraints->ComputeInitialWCData(Evolution->Particles(), GSCorotatedConstraints->GetMeshArray(), GSCorotatedConstraints->GetIncidentElements(), GSCorotatedConstraints->GetIncidentElementsLocal(), ParticlesPerColor);
-				GSCorotatedConstraints->SetParticlesPerColor(MoveTemp(ParticlesPerColor));
 
-				GSCorotatedConstraints->AddAdditionalRes = [this](const FSolverParticles& InParticles, const int32 p, const FSolverReal Dt, TVec3<FSolverReal>& res)
+				TArray<TArray<int32>> StaticIncidentElements, StaticIncidentElementsLocal;
+				TArray<TArray<int32>> StaticConstraints = GSWeakConstraints->GetStaticConstraintArrays(StaticIncidentElements, StaticIncidentElementsLocal);
+				GSMasterConstraint->AddStaticConstraints(StaticConstraints, StaticIncidentElements, StaticIncidentElementsLocal);
+
+				int32 StaticIndex1 = GSMasterConstraint->AddStaticConstraintResidualAndHessianRange(1);
+				GSMasterConstraint->StaticConstraintResidualAndHessian()[StaticIndex1] = [this](const FSolverParticles& Particles, const int32 ConstraintIndex, const int32 ConstraintIndexLocal, const FSolverReal Dt, TVec3<FSolverReal>& ParticleResidual, Chaos::PMatrix<FSolverReal, 3, 3>& ParticleHessian)
 				{
-					this->GSWeakConstraints->AddWCResidual(InParticles, p, Dt, res);
+					this->GSWeakConstraints->AddWCResidualAndHessian(Particles, ConstraintIndex, ConstraintIndexLocal, Dt, ParticleResidual, ParticleHessian);
 				};
-				GSCorotatedConstraints->AddAdditionalHessian = [this](const FSolverParticles& InParticles, const int32 p, const FSolverReal Dt, Chaos::PMatrix<FSolverReal, 3, 3>& hessian)
+
+				int32 PerNodeIndex = GSMasterConstraint->AddPerNodeHessianRange(1);
+				GSMasterConstraint->PerNodeHessian()[PerNodeIndex] = [this](const int32 p, const FSolverReal Dt, Chaos::PMatrix<FSolverReal, 3, 3>& ParticleHessian)
 				{
-					this->GSWeakConstraints->AddWCHessian(p, Dt, hessian);
+					this->GSWeakConstraints->AddWCHessian(p, Dt, ParticleHessian);
 				};
+
 			}
 		}
+
+		GSMasterConstraint->InitStaticColor(Evolution->Particles());
+
 		if (Property.bEnablePositionTargets)
 		{
 			int32 InitIndex1 = Evolution->AddConstraintInitRange(1, true);
@@ -1278,6 +1325,21 @@ namespace Chaos::Softs
 
 		if (Property.bDoSelfCollision)
 		{
+			int32 DynamicIndex = GSMasterConstraint->AddDynamicConstraintResidualAndHessianRange(1);
+			GSMasterConstraint->DynamicConstraintResidualAndHessian()[DynamicIndex] = [this](const FSolverParticles& Particles, const int32 ConstraintIndex, const int32 ConstraintIndexLocal, const FSolverReal Dt, TVec3<FSolverReal>& ParticleResidual, Chaos::PMatrix<FSolverReal, 3, 3>& ParticleHessian)
+			{
+				this->GSWeakConstraints->AddWCResidualAndHessian(Particles, ConstraintIndex + this->GSWeakConstraints->InitialWCSize, ConstraintIndexLocal, Dt, ParticleResidual, ParticleHessian);
+			};
+
+			if (!Property.bEnablePositionTargets)
+			{
+				int32 PerNodeIndex = GSMasterConstraint->AddPerNodeHessianRange(1);
+				GSMasterConstraint->PerNodeHessian()[PerNodeIndex] = [this](const int32 p, const FSolverReal Dt, Chaos::PMatrix<FSolverReal, 3, 3>& ParticleHessian)
+				{
+					this->GSWeakConstraints->AddWCHessian(p, Dt, ParticleHessian);
+				};
+			}
+
 			this->GSWeakConstraints->UpdateBoundaryVertices(*SurfaceElements);
 			
 			int32 InitIndex = Evolution->AddConstraintInitRange(1, true);
@@ -1294,26 +1356,17 @@ namespace Chaos::Softs
 			};
 
 			int32 InitIndex1 = Evolution->AddConstraintInitRange(1, true);
-			if (Property.bUseGSNeohookean)
-			{
-				Evolution->ConstraintInits()[InitIndex1] =
-					[this](FSolverParticles& InParticles, const FSolverReal Dt)
-				{
-					TArray<TArray<int32>> ParticlesPerColor;
-					this->GSWeakConstraints->ComputeCollisionWCData(this->Evolution->Particles(), this->GSNeohookeanConstraints->GetMeshArray(), this->GSNeohookeanConstraints->GetIncidentElements(), this->GSNeohookeanConstraints->GetIncidentElementsLocal(), ParticlesPerColor);
-					this->GSNeohookeanConstraints->SetParticlesPerColor(MoveTemp(ParticlesPerColor));
-				};
-			}
-			else
+
+
+			Evolution->ConstraintInits()[InitIndex1] =
+				[this](FSolverParticles& InParticles, const FSolverReal Dt)
 			{
-				Evolution->ConstraintInits()[InitIndex1] =
-					[this](FSolverParticles& InParticles, const FSolverReal Dt)
-				{
-					TArray<TArray<int32>> ParticlesPerColor;
-					this->GSWeakConstraints->ComputeCollisionWCData(this->Evolution->Particles(), this->GSCorotatedConstraints->GetMeshArray(), this->GSCorotatedConstraints->GetIncidentElements(), this->GSCorotatedConstraints->GetIncidentElementsLocal(), ParticlesPerColor);
-					this->GSCorotatedConstraints->SetParticlesPerColor(MoveTemp(ParticlesPerColor));
-				};
-			}
+				TArray<TArray<int32>> WCCollisionConstraints, WCCollisionIncidentElements, WCCollisionIncidentElementsLocal;
+				this->GSWeakConstraints->ComputeCollisionWCDataSimplified(WCCollisionConstraints, WCCollisionIncidentElements, WCCollisionIncidentElementsLocal);
+				this->GSMasterConstraint->AddDynamicConstraints(WCCollisionConstraints, WCCollisionIncidentElements, WCCollisionIncidentElementsLocal);
+				this->GSMasterConstraint->InitDynamicColor(InParticles);
+			}; 
+			
 		}
 	}
 
diff --git a/Engine/Source/Runtime/Experimental/Chaos/Private/Chaos/Deformable/GaussSeidelMasterConstraint.cpp b/Engine/Source/Runtime/Experimental/Chaos/Private/Chaos/Deformable/GaussSeidelMasterConstraint.cpp
new file mode 100644
--- /dev/null
+++ b/Engine/Source/Runtime/Experimental/Chaos/Private/Chaos/Deformable/GaussSeidelMasterConstraint.cpp
@@ -0,0 +1,81 @@
+// Copyright Epic Games, Inc. All Rights Reserved.
+#include "Chaos/Deformable/GaussSeidelMasterConstraint.h"
+
+namespace Chaos::Softs
+{
+	template <typename T, typename ParticleType>
+	void FGaussSeidelMasterConstraint<T, ParticleType>::AddStaticConstraints(const TArray<TArray<int32>>& ExtraConstraints, TArray<TArray<int32>>& ExtraIncidentElements, TArray<TArray<int32>>& ExtraIncidentElementsLocal)
+	{	
+		if (!IsClean(ExtraConstraints, ExtraIncidentElements, ExtraIncidentElementsLocal))
+		{
+			ExtraIncidentElements = Chaos::Utilities::ComputeIncidentElements(ExtraConstraints, &ExtraIncidentElementsLocal);
+		}
+
+		int32 Offset = StaticConstraints.Num();
+		StaticConstraints += ExtraConstraints;
+		for (int32 i = 0; i < ExtraIncidentElements.Num(); i++)
+		{
+			if (ExtraIncidentElements[i].Num() > 0)
+			{
+				TArray<int32> ExtraIncidentElementsWithOffset = ExtraIncidentElements[i];
+				for (int32 j = 0; j < ExtraIncidentElementsWithOffset.Num(); j++)
+				{
+					ExtraIncidentElementsWithOffset[j] += Offset;
+				}
+				StaticIncidentElements[i] += ExtraIncidentElementsWithOffset;
+				StaticIncidentElementsLocal[i] += ExtraIncidentElementsLocal[i];
+			}
+		}
+		if (StaticIncidentElementsOffsets.Num() > 0)
+		{
+			StaticIncidentElementsOffsets.RemoveAt(StaticIncidentElementsOffsets.Num() - 1);
+		}
+		StaticIncidentElementsOffsets.Add(Offset);
+		StaticIncidentElementsOffsets.Add(StaticConstraints.Num());
+	}
+
+	template <typename T, typename ParticleType>
+	void FGaussSeidelMasterConstraint<T, ParticleType>::AddDynamicConstraints(const TArray<TArray<int32>>& ExtraConstraints, TArray<TArray<int32>>& ExtraIncidentElements, TArray<TArray<int32>>& ExtraIncidentElementsLocal, bool CheckIncidentElements)
+	{
+		if (CheckIncidentElements)
+		{
+			if (!IsClean(ExtraConstraints, ExtraIncidentElements, ExtraIncidentElementsLocal))
+			{
+				ExtraIncidentElements = Chaos::Utilities::ComputeIncidentElements(ExtraConstraints, &ExtraIncidentElementsLocal);
+			}
+		}
+
+		int32 Offset = DynamicConstraints.Num();
+		DynamicConstraints += ExtraConstraints;
+		for (int32 i = 0; i < ExtraIncidentElements.Num(); i++)
+		{
+			if (ExtraIncidentElements[i].Num() > 0)
+			{
+				TArray<int32> ExtraIncidentElementsWithOffset = ExtraIncidentElements[i];
+				for (int32 j = 0; j < ExtraIncidentElementsWithOffset.Num(); j++)
+				{
+					ExtraIncidentElementsWithOffset[j] += Offset;
+				}
+				DynamicIncidentElements[i] += ExtraIncidentElementsWithOffset;
+				DynamicIncidentElementsLocal[i] += ExtraIncidentElementsLocal[i];
+			}
+		}
+		if (DynamicIncidentElementsOffsets.Num() > 0)
+		{
+			DynamicIncidentElementsOffsets.RemoveAt(StaticIncidentElementsOffsets.Num() - 1);
+		}
+		DynamicIncidentElementsOffsets.Add(Offset);
+		DynamicIncidentElementsOffsets.Add(DynamicConstraints.Num());
+	
+	}
+
+}
+
+
+template CHAOS_API void Chaos::Softs::FGaussSeidelMasterConstraint<Chaos::FRealSingle, Chaos::Softs::FSolverParticles>::AddStaticConstraints(const TArray<TArray<int32>>& ExtraConstraints, TArray<TArray<int32>>& ExtraIncidentElements, TArray<TArray<int32>>& ExtraIncidentElementsLocal);
+
+template CHAOS_API void Chaos::Softs::FGaussSeidelMasterConstraint<Chaos::FRealDouble, Chaos::TDynamicParticles<Chaos::FRealDouble, 3>>::AddStaticConstraints(const TArray<TArray<int32>>& ExtraConstraints, TArray<TArray<int32>>& ExtraIncidentElements, TArray<TArray<int32>>& ExtraIncidentElementsLocal);
+
+template CHAOS_API void Chaos::Softs::FGaussSeidelMasterConstraint<Chaos::FRealSingle, Chaos::Softs::FSolverParticles>::AddDynamicConstraints(const TArray<TArray<int32>>& ExtraConstraints, TArray<TArray<int32>>& ExtraIncidentElements, TArray<TArray<int32>>& ExtraIncidentElementsLocal, bool CheckIncidentElements);
+
+template CHAOS_API void Chaos::Softs::FGaussSeidelMasterConstraint<Chaos::FRealDouble, Chaos::TDynamicParticles<Chaos::FRealDouble, 3>>::AddDynamicConstraints(const TArray<TArray<int32>>& ExtraConstraints, TArray<TArray<int32>>& ExtraIncidentElements, TArray<TArray<int32>>& ExtraIncidentElementsLocal, bool CheckIncidentElements);
diff --git a/Engine/Source/Runtime/Experimental/Chaos/Private/Chaos/GraphColoring.cpp b/Engine/Source/Runtime/Experimental/Chaos/Private/Chaos/GraphColoring.cpp
--- a/Engine/Source/Runtime/Experimental/Chaos/Private/Chaos/GraphColoring.cpp
+++ b/Engine/Source/Runtime/Experimental/Chaos/Private/Chaos/GraphColoring.cpp
@@ -346,7 +346,7 @@ static bool VerifyNodalColoring(const TArray<TArray<int32>>& Graph, const Chaos:
 	for (int32 p = 0; p < GraphParticlesEnd - GraphParticlesStart; p++)
 	{
 		int32 ParticleIndex = p + GraphParticlesStart;
-		if (InParticles.InvM(ParticleIndex) != (T)0.)
+		if (InParticles.InvM(ParticleIndex) != (T)0. && IncidentElements[ParticleIndex].Num() > 0 )
 		{
 			if (!ParticleIsIncluded[ParticleIndex])
 			{
@@ -411,7 +411,7 @@ static bool VerifyExtraNodalColoring(const Chaos::TDynamicParticles<T, 3>& InPar
 
 	for (int32 ParticleIndex = 0; ParticleIndex < (int32)InParticles.Size(); ParticleIndex++)
 	{
-		if (InParticles.InvM(ParticleIndex) != (T)0.)
+		if (InParticles.InvM(ParticleIndex) != (T)0. && (IncidentElements[ParticleIndex].Num() > 0 || ExtraIncidentElements[ParticleIndex].Num() > 0))
 		{
 			if (!ParticleIsIncluded[ParticleIndex])
 			{
@@ -1778,7 +1778,7 @@ void Chaos::ComputeExtraNodalColoring(const TArray<TArray<int32>>& Graph, const
 	ParticlesPerColor.SetNum(ParticleColors.Max() + 1);
 	for (int32 i = 0; i < ExtraIncidentElements.Num(); i++)
 	{
-		if (ParticleIsAffected[i])
+		if (ParticleIsAffected[i] && ParticleColors[i] != INDEX_NONE)
 		{
 			ParticlesPerColor[ParticleColors[i]].Emplace(i);
 		}
diff --git a/Engine/Source/Runtime/Experimental/Chaos/Public/Chaos/Deformable/ChaosDeformableSolver.h b/Engine/Source/Runtime/Experimental/Chaos/Public/Chaos/Deformable/ChaosDeformableSolver.h
--- a/Engine/Source/Runtime/Experimental/Chaos/Public/Chaos/Deformable/ChaosDeformableSolver.h
+++ b/Engine/Source/Runtime/Experimental/Chaos/Public/Chaos/Deformable/ChaosDeformableSolver.h
@@ -14,6 +14,7 @@
 #include "Chaos/Deformable/GaussSeidelCorotatedConstraints.h"
 #include "Chaos/Deformable/GaussSeidelNeohookeanConstraints.h"
 #include "Chaos/Deformable/GaussSeidelWeakConstraints.h"
+#include "Chaos/Deformable/GaussSeidelMasterConstraint.h"
 #include "Chaos/XPBDWeakConstraints.h"
 #include "Chaos/BlendedXPBDCorotatedConstraints.h"
 #include "Chaos/XPBDGridBasedCorotatedConstraints.h"
@@ -183,6 +184,7 @@ namespace Chaos::Softs
 		TArray<TUniquePtr<Softs::FXPBDWeakConstraints<Softs::FSolverReal, Softs::FSolverParticles>>> WeakConstraints;
 		TArray<TUniquePtr<Softs::FBlendedXPBDCorotatedConstraints<Softs::FSolverReal, Softs::FSolverParticles>>> BlendedCorotatedConstraints;
 		TUniquePtr<Softs::FXPBDGridBasedCorotatedConstraints<Softs::FSolverReal, Softs::FSolverParticles>> GridBasedCorotatedConstraint;
+		TUniquePtr<Softs::FGaussSeidelMasterConstraint<Softs::FSolverReal, Softs::FSolverParticles>> GSMasterConstraint;
 		TUniquePtr<Softs::FPBDCollisionSpringConstraints> CollisionSpringConstraint;
 		TUniquePtr<Softs::FPBDTriangleMeshCollisions> TriangleMeshCollisions;
 		TArrayCollectionArray<const UObject*> MObjects;
diff --git a/Engine/Source/Runtime/Experimental/Chaos/Public/Chaos/Deformable/GaussSeidelCorotatedConstraints.h b/Engine/Source/Runtime/Experimental/Chaos/Public/Chaos/Deformable/GaussSeidelCorotatedConstraints.h
--- a/Engine/Source/Runtime/Experimental/Chaos/Public/Chaos/Deformable/GaussSeidelCorotatedConstraints.h
+++ b/Engine/Source/Runtime/Experimental/Chaos/Public/Chaos/Deformable/GaussSeidelCorotatedConstraints.h
@@ -25,6 +25,7 @@ namespace Chaos::Softs
 		using Base::MuElementArray;
 		using Base::LambdaElementArray;
 		using Base::CorotatedParams;
+		using Base::F;
 
 	public:
 		FGaussSeidelCorotatedConstraints(
@@ -48,13 +49,10 @@ namespace Chaos::Softs
 			ParticleStartIndex(ParticleStartIndexIn), ParticleEndIndex(ParticleEndIndexIn), bDoQuasistatics(bDoQuasistaticsIn), bDoSOR(bDoSORIn), OmegaSOR(InOmegaSOR)
 		{
 			Base::LambdaArray.SetNum(0);
-			Particle2Incident.Init(-1, ParticleEndIndexIn - ParticleStartIndexIn);
+			Particle2Incident.Init(INDEX_NONE, ParticleEndIndexIn - ParticleStartIndexIn);
 			for (int32 i = 0; i < IncidentElements.Num(); i++)
 			{
-				if (IncidentElements[i].Num() > 0)
-				{
-					Particle2Incident[InMesh[IncidentElements[i][0]][IncidentElementsLocal[i][0]] - ParticleStartIndex] = i;
-				}
+				Particle2Incident[i] = i;
 			}
 			for (int32 e = 0; e < InMesh.Num(); e++)
 			{
@@ -75,15 +73,6 @@ namespace Chaos::Softs
 			}
 			InitColor(InParticles);
 			InitializeLambdas();
-			MeshArray.SetNum(MeshConstraints.Num());
-			for (int32 i = 0; i < MeshConstraints.Num(); i++)
-			{
-				MeshArray[i].SetNum(4);
-				for (int32 ie = 0; ie < 4; ie++)
-				{
-					MeshArray[i][ie] = MeshConstraints[i][ie];
-				}
-			}
 		}
 
 		virtual ~FGaussSeidelCorotatedConstraints() {}
@@ -133,13 +122,27 @@ namespace Chaos::Softs
 
 		void Init() const override {}
 
-		TArray<TArray<int32>> GetIncidentElements() const { return IncidentElements; }
-		TArray<TArray<int32>> GetIncidentElementsLocal() const { return IncidentElementsLocal; }
+		TArray<TArray<int32>>& GetIncidentElements() { return IncidentElements; }
+		TArray<TArray<int32>>& GetIncidentElementsLocal() { return IncidentElementsLocal; }
 		TArray<TVector<int32, 4>> GetMeshConstraints() const { return MeshConstraints; }
-		TArray<TArray<int32>> GetMeshArray() const { return MeshArray; }
 		void SetParticlesPerColor(TArray<TArray<int32>>&& InParticlesPerColor) {ParticlesPerColor = MoveTemp(InParticlesPerColor); }
 
+		TArray<TArray<int32>> GetMeshArray() const 
+		{
+			TArray<TArray<int32>> MeshArray;
+			MeshArray.SetNum(MeshConstraints.Num());
+			for (int32 i = 0; i < MeshConstraints.Num(); i++)
+			{
+				MeshArray[i].SetNum(4);
+				for (int32 ie = 0; ie < 4; ie++)
+				{
+					MeshArray[i][ie] = MeshConstraints[i][ie];
+				}
+			}
+			return MeshArray;
+		}
 
+		// Apply Successive Over Relaxation (SOR) to accelerate convergence rate
 		void ApplySOR(ParticleType& Particles, const T Dt) const
 		{
 			if (bDoSOR)
@@ -232,6 +235,47 @@ namespace Chaos::Softs
 			return final_hessian;
 		}
 
+
+		void AddHyperelasticResidualAndHessian (const ParticleType& Particles, const int32 ElementIndex, const int32 ElementIndexLocal, const T Dt, TVec3<T>& ParticleResidual, Chaos::PMatrix<T, 3, 3>& ParticleHessian)
+		{
+			Chaos::PMatrix<T, 3, 3> DmInvT = Base::ElementDmInv(ElementIndex).GetTransposed();
+			Chaos::PMatrix<T, 3, 3> Fe = F(ElementIndex, Particles);
+			Chaos::PMatrix<T, 3, 3> P((T)0.);
+
+			this->ComputeStress(Fe, MuElementArray[ElementIndex], LambdaElementArray[ElementIndex], P);
+
+			Chaos::PMatrix<T, 3, 3> force_term = -Measure[ElementIndex] * DmInvT * P;
+			Chaos::TVector<T, 3> dx((T)0.);
+			if (ElementIndexLocal > 0)
+			{
+				for (int32 c = 0; c < 3; c++)
+				{
+					dx[c] += force_term.GetAt(c, ElementIndexLocal - 1);
+				}
+			}
+			else
+			{
+				for (int32 c = 0; c < 3; c++)
+				{
+					for (int32 h = 0; h < 3; h++)
+					{
+						dx[c] -= force_term.GetAt(c, h);
+					}
+				}
+			}
+
+			for (int32 ll = 0; ll < 3; ll++) {
+				dx[ll] *= Dt * Dt;
+			}
+
+			for (int32 alpha = 0; alpha < 3; alpha++) {
+				ParticleResidual[alpha] -= dx[alpha];
+			}
+
+			ComputeHessianHelper(Fe, Base::ElementDmInv(ElementIndex), MuElementArray[ElementIndex], LambdaElementArray[ElementIndex], ElementIndexLocal, Dt * Dt * Measure[ElementIndex], ParticleHessian);
+
+		}
+
 	protected:
 		 void InitColor(const ParticleType& Particles)
 		{
@@ -341,7 +385,6 @@ namespace Chaos::Softs
 
 	protected:	
 
-		TArray<TArray<int32>> MeshArray;
 		TArray<int32> Particle2Incident;
 		TArray<TArray<int32>> IncidentElements;
 		TArray<TArray<int32>> IncidentElementsLocal;
@@ -354,6 +397,7 @@ namespace Chaos::Softs
 		mutable TArray<Chaos::TVector<T, 3>> X_k;
 		bool bDoQuasistatics = false;
 		bool bDoSOR = true;
+		//Parameter to control the behavior of SOR. 1-1.7 is the usual range. Larger paramater leads to potentially more convergence but less stable behavior. 
 		T OmegaSOR = T(1.6);
 		mutable int32 CurrentIt = 0;
 		TFunction<void(const Chaos::PMatrix<T, 3, 3>&, const T, const T, Chaos::PMatrix<T, 3, 3>& )> ComputeStress;
diff --git a/Engine/Source/Runtime/Experimental/Chaos/Public/Chaos/Deformable/GaussSeidelMasterConstraint.h b/Engine/Source/Runtime/Experimental/Chaos/Public/Chaos/Deformable/GaussSeidelMasterConstraint.h
new file mode 100644
--- /dev/null
+++ b/Engine/Source/Runtime/Experimental/Chaos/Public/Chaos/Deformable/GaussSeidelMasterConstraint.h
@@ -0,0 +1,309 @@
+// Copyright Epic Games, Inc. All Rights Reserved.
+#pragma once
+
+#include "Chaos/PBDSpringConstraintsBase.h"
+#include "Chaos/XPBDCorotatedConstraints.h"
+#include "ChaosStats.h"
+#include "Chaos/ImplicitQRSVD.h"
+#include "Chaos/GraphColoring.h"
+#include "Chaos/NewtonCorotatedCache.h"
+#include "Chaos/Framework/Parallel.h"
+#include "Chaos/Deformable/GaussSeidelWeakConstraints.h"
+#include "Chaos/Deformable/GaussSeidelCorotatedConstraints.h"
+
+
+#define PERF_SCOPE(X) SCOPE_CYCLE_COUNTER(X); TRACE_CPUPROFILER_EVENT_SCOPE(X);
+DECLARE_CYCLE_STAT(TEXT("Chaos.Deformable.GSMasterConstraint.Apply"), STAT_ChaosGSMasterConstraint_Apply, STATGROUP_Chaos);
+DECLARE_CYCLE_STAT(TEXT("Chaos.Deformable.GSMasterConstraint.Acceleration"), STAT_ChaosGSMasterConstraint_Acceleration, STATGROUP_Chaos);
+DECLARE_CYCLE_STAT(TEXT("Chaos.Deformable.GSMasterConstraint.Init"), STAT_ChaosGSMasterConstraint_Init, STATGROUP_Chaos);
+DECLARE_CYCLE_STAT(TEXT("Chaos.Deformable.GSMasterConstraint.InitDynamicColor"), STAT_ChaosGSMasterConstraint_InitDynamicColor, STATGROUP_Chaos);
+
+namespace Chaos::Softs
+{
+	template <typename T, typename ParticleType>
+	class FGaussSeidelMasterConstraint
+	{
+
+	public:
+		FGaussSeidelMasterConstraint(
+			const ParticleType& InParticles, 
+			const bool bDoQuasistaticsIn = false,
+			const bool bDoSORIn = true,
+			const T InOmegaSOR = (T)1.6,
+			const int32 ParallelMaxIn = 1000, 
+			const FDeformableXPBDCorotatedParams& InParams = FDeformableXPBDCorotatedParams())
+			: bDoQuasistatics(bDoQuasistaticsIn)
+			, bDoAcceleration(bDoSORIn)
+			, OmegaSOR(InOmegaSOR)
+			, ParallelMax(ParallelMaxIn)
+			, CorotatedParams(InParams)
+		{
+			Resize((int32)InParticles.Size());
+
+			InitializeLambdas();
+		}
+
+		virtual ~FGaussSeidelMasterConstraint() {}
+
+		void Resize(const int32 NewSize)
+		{
+			xtilde.SetNum(NewSize);
+			StaticIncidentElements.SetNum(NewSize);
+			StaticIncidentElementsLocal.SetNum(NewSize);
+			DynamicIncidentElements.SetNum(NewSize);
+			DynamicIncidentElementsLocal.SetNum(NewSize);
+			X_k_1.Init(TVector<T, 3>(T(0.)), NewSize);
+			X_k.Init(TVector<T, 3>(T(0.)), NewSize);
+		}
+
+		const TArray<TFunction<void(const ParticleType&, const int32, const int32, const T, TVec3<T>&, Chaos::PMatrix<T, 3, 3>&)>>& StaticConstraintResidualAndHessian() const { return AddStaticConstraintResidualAndHessian; }
+		TArray<TFunction<void(const ParticleType&, const int32, const int32, const T, TVec3<T>&, Chaos::PMatrix<T, 3, 3>&)>>& StaticConstraintResidualAndHessian() { return AddStaticConstraintResidualAndHessian; }
+		const TArray<TFunction<void(const ParticleType&, const int32, const int32, const T, TVec3<T>&, Chaos::PMatrix<T, 3, 3>&)>>& DynamicConstraintResidualAndHessian() const { return AddDynamicConstraintResidualAndHessian; }
+		TArray<TFunction<void(const ParticleType&, const int32, const int32, const T, TVec3<T>&, Chaos::PMatrix<T, 3, 3>&)>>& DynamicConstraintResidualAndHessian() { return AddDynamicConstraintResidualAndHessian; }
+		const TArray<TFunction<void(const int32, const T, Chaos::PMatrix<T, 3, 3>&)>>& PerNodeHessian() const { return AddPerNodeHessian; }
+		TArray<TFunction<void(const int32, const T, Chaos::PMatrix<T, 3, 3>&)>>& PerNodeHessian() { return AddPerNodeHessian; }
+
+		int32 AddStaticConstraintResidualAndHessianRange(int32 NumConstraints)
+		{
+			int32 CurrentSize = AddStaticConstraintResidualAndHessian.Num();
+			AddStaticConstraintResidualAndHessian.AddDefaulted(NumConstraints);
+			return CurrentSize;
+		}
+
+		int32 AddDynamicConstraintResidualAndHessianRange(int32 NumConstraints)
+		{
+			int32 CurrentSize = AddDynamicConstraintResidualAndHessian.Num();
+			AddDynamicConstraintResidualAndHessian.AddDefaulted(NumConstraints);
+			return CurrentSize;
+		}
+
+		int32 AddPerNodeHessianRange(int32 NumConstraints)
+		{
+			int32 CurrentSize = AddPerNodeHessian.Num();
+			AddPerNodeHessian.AddDefaulted(NumConstraints);
+			return CurrentSize;
+		}
+
+		void AddStaticConstraints(const TArray<TArray<int32>>& ExtraConstraints, TArray<TArray<int32>>& ExtraIncidentElements, TArray<TArray<int32>>& ExtraIncidentElementsLocal);
+
+		void AddDynamicConstraints(const TArray<TArray<int32>>& ExtraConstraints, TArray<TArray<int32>>& ExtraIncidentElements, TArray<TArray<int32>>& ExtraIncidentElementsLocal, bool CheckIncidentElements = false);
+
+		void Apply(ParticleType& Particles, const T Dt)
+		{
+			PERF_SCOPE(STAT_ChaosGSMasterConstraint_Apply);
+
+			for (int32 i = 0; i < ParticlesPerColor.Num(); i++)
+			{
+				int32 NumBatch = ParticlesPerColor[i].Num() / CorotatedParams.XPBDCorotatedBatchSize;
+				if (ParticlesPerColor[i].Num() % CorotatedParams.XPBDCorotatedBatchSize != 0)
+				{
+					NumBatch ++;
+				}
+
+				PhysicsParallelFor(NumBatch, [&](const int32 BatchIndex)
+					{
+						for (int32 BatchSubIndex = 0; BatchSubIndex < CorotatedParams.XPBDCorotatedBatchSize; BatchSubIndex++) {
+							int32 TaskIndex = CorotatedParams.XPBDCorotatedBatchSize * BatchIndex + BatchSubIndex;
+							if (TaskIndex < ParticlesPerColor[i].Num())
+							{
+								int32 ParticleIndex = ParticlesPerColor[i][TaskIndex];
+								ApplySingleParticle(ParticleIndex, Dt, Particles);
+							}
+						}
+					}, NumBatch < CorotatedParams.XPBDCorotatedBatchThreshold);
+
+			}
+
+			if (bDoAcceleration)
+			{
+				AccelerationTechnique(Particles);
+			}
+
+			CurrentIt ++;
+		}
+
+		void InitStaticColor(const ParticleType& Particles)
+		{
+			StaticParticlesPerColor = ComputeNodalColoring(StaticConstraints, Particles, 0, Particles.Size(), StaticIncidentElements, StaticIncidentElementsLocal, &StaticParticleColors);
+			ParticleColors = StaticParticleColors;
+			ParticlesPerColor = StaticParticlesPerColor;
+		}
+
+		void InitDynamicColor(const ParticleType& Particles)
+		{
+			PERF_SCOPE(STAT_ChaosGSMasterConstraint_InitDynamicColor);
+			ParticleColors = StaticParticleColors;
+			ParticlesPerColor = StaticParticlesPerColor;
+			Chaos::ComputeExtraNodalColoring(StaticConstraints, DynamicConstraints, Particles, StaticIncidentElements, DynamicIncidentElements, ParticleColors, ParticlesPerColor);
+		}
+
+		void Init(const T Dt, const ParticleType& Particles)
+		{
+			PERF_SCOPE(STAT_ChaosGSMasterConstraint_Init);
+			DynamicConstraints.SetNum(0);
+			for (int32 p = 0; p < DynamicIncidentElements.Num(); p++)
+			{
+				DynamicIncidentElements[p].SetNum(0);
+				DynamicIncidentElementsLocal[p].SetNum(0);
+			}
+			DynamicIncidentElementsOffsets.SetNum(0);
+			if (!bDoQuasistatics)
+			{
+				for (int32 i = 0; i < (int32)Particles.Size(); i++)
+				{
+					xtilde[i] = Particles.X(i) + Dt * Particles.V(i);
+				}
+			}
+			CurrentIt = 0;
+		}
+
+		static bool IsClean(const TArray<TArray<int32>>& ConstraintsIn, const TArray<TArray<int32>>& IncidentElementsIn, const TArray<TArray<int32>>& IncidentElementsLocalIn)
+		{
+			if (IncidentElementsIn.Num() == IncidentElementsLocalIn.Num())
+			{
+				int32 TotalEntries = 0;
+				for (int32 i = 0; i < IncidentElementsIn.Num(); i++)
+				{
+					TotalEntries += IncidentElementsIn[i].Num();
+					if (IncidentElementsIn[i].Num() != IncidentElementsLocalIn[i].Num())
+					{
+						return false;
+					}
+					for (int32 j = 0; j < IncidentElementsIn[i].Num(); j++)
+					{
+						if (IncidentElementsIn[i][j] > ConstraintsIn.Num() || IncidentElementsLocalIn[i][j] > ConstraintsIn[IncidentElementsIn[i][j]].Num())
+						{
+							return false;
+						}
+					}
+				}
+				if (TotalEntries > 0)
+				{
+					return true;
+				}
+			}
+			return false;
+		}
+
+		
+	private:
+
+		void ApplySingleParticle(const int32 p, const T Dt, ParticleType& Particles)
+		{
+			int32 ConstraintIndex = 0;
+			Chaos::TVector<T, 3> ParticleResidual((T)0.);
+			Chaos::PMatrix<T, 3, 3> ParticleHessian = Chaos::PMatrix<T, 3, 3>::Zero;
+
+			ComputeInitialResidualAndHessian(Particles, p, ParticleResidual, ParticleHessian);
+
+			for (int32 i = 0; i < StaticIncidentElements[p].Num(); i++)
+			{
+				while (StaticIncidentElements[p][i] >= StaticIncidentElementsOffsets[ConstraintIndex + 1] && ConstraintIndex < StaticIncidentElementsOffsets.Num() - 1)
+				{
+					ConstraintIndex ++;
+				}
+
+				AddStaticConstraintResidualAndHessian[ConstraintIndex](Particles, StaticIncidentElements[p][i] - StaticIncidentElementsOffsets[ConstraintIndex], StaticIncidentElementsLocal[p][i], Dt, ParticleResidual, ParticleHessian);
+
+			}
+
+			ConstraintIndex = 0;
+
+			for (int32 i = 0; i < DynamicIncidentElements[p].Num(); i++)
+			{
+				while (DynamicIncidentElements[p][i] >= DynamicIncidentElementsOffsets[ConstraintIndex + 1] && ConstraintIndex < DynamicIncidentElementsOffsets.Num() - 1)
+				{
+					ConstraintIndex ++;
+				}
+
+				AddDynamicConstraintResidualAndHessian[ConstraintIndex](Particles, DynamicIncidentElements[p][i] - DynamicIncidentElementsOffsets[ConstraintIndex], DynamicIncidentElementsLocal[p][i], Dt, ParticleResidual, ParticleHessian);
+
+			}
+
+			for (int32 i = 0; i < AddPerNodeHessian.Num(); i++)
+			{
+				AddPerNodeHessian[i](p, Dt, ParticleHessian);
+			}
+
+
+			T HessianDet = ParticleHessian.Determinant();
+			if (HessianDet > UE_SMALL_NUMBER)
+			{
+				Chaos::PMatrix<T, 3, 3> HessianInv = ParticleHessian.SymmetricCofactorMatrix();
+				HessianInv *= T(1) / HessianDet;
+				Chaos::TVector<T, 3> Dx = HessianInv.GetTransposed() * (-ParticleResidual);
+				Particles.P(p) += Dx;
+			}
+		}
+		
+		void InitializeLambdas()
+		{
+			ComputeInitialResidualAndHessian = [this](const ParticleType& Particles, const int32 p, TVec3<T>& ParticleResidual, Chaos::PMatrix<T, 3, 3>& ParticleHessian)
+			{
+				if (!this->bDoQuasistatics) {
+					for (int32 alpha = 0; alpha < 3; alpha++) {
+						ParticleResidual[alpha] = Particles.M(p) * (Particles.P(p)[alpha] - this->xtilde[p][alpha]);
+					}
+					for (int32 alpha = 0; alpha < 3; alpha++) {
+						ParticleHessian.SetAt(alpha, alpha, Particles.M(p));
+					}
+				}
+			};
+
+			AccelerationTechnique = [this](ParticleType& Particles)
+			{
+				PERF_SCOPE(STAT_ChaosGSMasterConstraint_Acceleration);
+				PhysicsParallelFor(Particles.Size(), [&](const int32 ParticleIndex)
+					{
+						if (Particles.InvM(ParticleIndex) != T(0) && CurrentIt > SORStart)
+						{
+							Particles.P(ParticleIndex) = OmegaSOR * (Particles.P(ParticleIndex) - this->X_k_1[ParticleIndex]) + this->X_k_1[ParticleIndex];
+						}
+						this->X_k_1[ParticleIndex] = this->X_k[ParticleIndex];
+						this->X_k[ParticleIndex] = Particles.P(ParticleIndex);
+					}, Particles.Size() < 1000);
+			};
+		}
+
+		//Constraints storage:
+		TArray<TArray<int32>> StaticConstraints;
+		TArray<TArray<int32>> StaticIncidentElements;
+		TArray<TArray<int32>> StaticIncidentElementsLocal;
+		TArray<TArray<int32>> DynamicConstraints;
+		TArray<TArray<int32>> DynamicIncidentElements;
+		TArray<TArray<int32>> DynamicIncidentElementsLocal;
+
+		//Lambds for spcifying residual/hessian computations:
+		TFunction<void(const ParticleType&, const int32, TVec3<T>&, Chaos::PMatrix<T, 3, 3>&)> ComputeInitialResidualAndHessian;
+		TArray<TFunction<void(const ParticleType&, const int32, const int32, const T, TVec3<T>&, Chaos::PMatrix<T, 3, 3>&)>> AddStaticConstraintResidualAndHessian;
+		TArray<TFunction<void(const ParticleType&, const int32, const int32, const T, TVec3<T>&, Chaos::PMatrix<T, 3, 3>&)>> AddDynamicConstraintResidualAndHessian;
+		TArray<TFunction<void(const int32, const T, Chaos::PMatrix<T, 3, 3>&)>> AddPerNodeHessian;
+
+		//Coloring information:
+		TArray<int32> StaticParticleColors;
+		TArray<TArray<int32>> StaticParticlesPerColor;
+
+		TArray<int32> ParticleColors;
+		TArray<TArray<int32>> ParticlesPerColor;
+
+		TArray<int32> StaticIncidentElementsOffsets;
+		TArray<int32> DynamicIncidentElementsOffsets;
+
+		bool bDoQuasistatics = false;
+		mutable TArray<Chaos::TVector<T, 3>> xtilde;
+
+		//SOR variables:
+		TFunction<void(ParticleType&)> AccelerationTechnique;
+		mutable TArray<Chaos::TVector<T, 3>> X_k_1;
+		mutable TArray<Chaos::TVector<T, 3>> X_k;
+		int32 CurrentIt = 0;
+		bool bDoAcceleration = true;
+		T OmegaSOR = T(1.6);
+		int32 SORStart = 1;
+
+		int32 ParallelMax = 1000;
+
+		FDeformableXPBDCorotatedParams CorotatedParams;
+
+	};
+}
diff --git a/Engine/Source/Runtime/Experimental/Chaos/Public/Chaos/Deformable/GaussSeidelWeakConstraints.h b/Engine/Source/Runtime/Experimental/Chaos/Public/Chaos/Deformable/GaussSeidelWeakConstraints.h
--- a/Engine/Source/Runtime/Experimental/Chaos/Public/Chaos/Deformable/GaussSeidelWeakConstraints.h
+++ b/Engine/Source/Runtime/Experimental/Chaos/Public/Chaos/Deformable/GaussSeidelWeakConstraints.h
@@ -66,24 +66,7 @@ namespace Chaos::Softs
 					ExtraConstraints[i][j+Indices[i].Num()] = SecondIndices[i][j];
 				}
 			}
-			TArray<TArray<int32>> TotalConstraints = MeshConstraints;
-			TotalConstraints += ExtraConstraints;
-			TArray<TArray<int32>> TotalIncidentElements = MeshIncidentElements, TotalIncidentElementsLocal = MeshIncidentElementsLocal;
 			WCIncidentElements = Chaos::Utilities::ComputeIncidentElements(ExtraConstraints, &WCIncidentElementsLocal);
-			
-			//TODO (Yizhou): Directly computing incident elements is inefficient. Optimize the total incident elements in future.
-			TotalIncidentElements = Chaos::Utilities::ComputeIncidentElements(TotalConstraints, &TotalIncidentElementsLocal);
-			Particle2WCIndices.Init(INDEX_NONE, InParticles.Size());
-			for (int32 i = 0; i < WCIncidentElements.Num(); i++)
-			{
-				if (WCIncidentElements[i].Num() > 0)
-				{
-					int32 p = ExtraConstraints[WCIncidentElements[i][0]][WCIncidentElementsLocal[i][0]];
-					Particle2WCIndices[p] = i;
-				}
-			}
-
-			ParticlesPerColor = ComputeNodalColoring(TotalConstraints, InParticles, 0, InParticles.Size(), TotalIncidentElements, TotalIncidentElementsLocal, &ParticleColors);
 
 			NodalWeights.Init({}, InParticles.Size());
 
@@ -709,6 +692,147 @@ namespace Chaos::Softs
 		}
 
 
+		void ComputeCollisionWCDataSimplified(TArray<TArray<int32>>& ExtraConstraints, TArray<TArray<int32>>& ExtraWCIncidentElements, TArray<TArray<int32>>& ExtraWCIncidentElementsLocal)
+		{
+			ensureMsgf(Indices.Num() >= InitialWCSize, TEXT("The size of Indices is smaller than InitialWCSize"));
+
+			ExtraConstraints.Init(TArray<int32>(), Indices.Num() - InitialWCSize);
+			for (int32 i = InitialWCSize; i < Indices.Num(); i++)
+			{
+				ExtraConstraints[i - InitialWCSize].SetNum(Indices[i].Num() + SecondIndices[i].Num());
+				for (int32 j = 0; j < Indices[i].Num(); j++)
+				{
+					ExtraConstraints[i - InitialWCSize][j] = Indices[i][j];
+				}
+				for (int32 j = 0; j < SecondIndices[i].Num(); j++)
+				{
+					ExtraConstraints[i - InitialWCSize][j + Indices[i].Num()] = SecondIndices[i][j];
+				}
+			}
+
+			ExtraWCIncidentElements = Chaos::Utilities::ComputeIncidentElements(ExtraConstraints, &ExtraWCIncidentElementsLocal);
+
+			NodalWeights = NoCollisionNodalWeights;
+			for (int32 i = 0; i < ExtraWCIncidentElements.Num(); i++)
+			{
+				if (ExtraWCIncidentElements[i].Num() > 0)
+				{
+					int32 p = ExtraConstraints[ExtraWCIncidentElements[i][0]][ExtraWCIncidentElementsLocal[i][0]];
+					if (NodalWeights[p].Num() == 0)
+					{
+						NodalWeights[p].Init(T(0), 6);
+					}
+					for (int32 j = 0; j < ExtraWCIncidentElements[i].Num(); j++)
+					{
+						int32 LocalIndex = ExtraWCIncidentElementsLocal[i][j];
+						int32 ConstraintIndex = ExtraWCIncidentElements[i][j] + InitialWCSize;
+						T weight = T(0);
+						if (LocalIndex >= Indices[ConstraintIndex].Num())
+						{
+							weight = SecondWeights[ConstraintIndex][LocalIndex - Indices[ConstraintIndex].Num()];
+						}
+						else
+						{
+							weight = Weights[ConstraintIndex][LocalIndex];
+						}
+						if (IsAnisotropic[ConstraintIndex])
+						{
+							for (int32 alpha = 0; alpha < 3; alpha++)
+							{
+								NodalWeights[p][alpha] += Normals[ConstraintIndex][alpha] * Normals[ConstraintIndex][alpha] * weight * weight * Stiffness[ConstraintIndex];
+							}
+
+							NodalWeights[p][3] += Normals[ConstraintIndex][0] * Normals[ConstraintIndex][1] * weight * weight * Stiffness[ConstraintIndex];
+							NodalWeights[p][4] += Normals[ConstraintIndex][0] * Normals[ConstraintIndex][2] * weight * weight * Stiffness[ConstraintIndex];
+							NodalWeights[p][5] += Normals[ConstraintIndex][1] * Normals[ConstraintIndex][2] * weight * weight * Stiffness[ConstraintIndex];
+						}
+						else
+						{
+							for (int32 alpha = 0; alpha < 3; alpha++)
+							{
+								NodalWeights[p][alpha] += weight * weight * Stiffness[ConstraintIndex];
+							}
+						}
+					}
+				}
+			}
+		}
+
+
+		TArray<TArray<int32>> GetStaticConstraintArrays(TArray<TArray<int32>>& IncidentElements, TArray<TArray<int32>> IncidentElementsLocal)
+		{
+			IncidentElements = NoCollisionWCIncidentElements;
+			IncidentElementsLocal = NoCollisionWCIncidentElementsLocal;
+			return NoCollisionConstraints;
+		}
+
+		TArray<TArray<int32>> GetDynamicConstraintArrays(TArray<TArray<int32>>& IncidentElements, TArray<TArray<int32>> IncidentElementsLocal)
+		{
+			TArray<TArray<int32>> ExtraConstraints;
+			ExtraConstraints.Init(TArray<int32>(), Indices.Num());
+			for (int32 i = InitialWCSize; i < Indices.Num(); i++)
+			{
+				ExtraConstraints[i - InitialWCSize].SetNum(Indices[i].Num() + SecondIndices[i].Num());
+				for (int32 j = 0; j < Indices[i].Num(); j++)
+				{
+					ExtraConstraints[i - InitialWCSize][j] = Indices[i][j];
+				}
+				for (int32 j = 0; j < SecondIndices[i].Num(); j++)
+				{
+					ExtraConstraints[i - InitialWCSize][j + Indices[i].Num()] = SecondIndices[i][j];
+				}
+			}
+
+			IncidentElements = Chaos::Utilities::ComputeIncidentElements(ExtraConstraints, &IncidentElementsLocal);
+
+			return ExtraConstraints;
+		}
+
+		void AddWCResidualAndHessian(const ParticleType& InParticles, const int32 ConstraintIndex, const int32 LocalIndex, const T Dt, TVec3<T>& ParticleResidual, Chaos::PMatrix<T, 3, 3>& ParticleHessian)
+		{
+			TVec3<T> SpringEdge((T)0.);
+			for (int32 l = 0; l < Weights[ConstraintIndex].Num(); l++)
+			{
+				for (int32 beta = 0; beta < 3; beta++)
+				{
+					SpringEdge[beta] += Weights[ConstraintIndex][l] * InParticles.P(Indices[ConstraintIndex][l])[beta];
+				}
+			}
+			for (int32 l = 0; l < SecondWeights[ConstraintIndex].Num(); l++)
+			{
+				for (int32 beta = 0; beta < 3; beta++)
+				{
+					SpringEdge[beta] -= SecondWeights[ConstraintIndex][l] * InParticles.P(SecondIndices[ConstraintIndex][l])[beta];
+				}
+			}
+			T weight = T(0);
+			if (LocalIndex >= Indices[ConstraintIndex].Num())
+			{
+				weight = -SecondWeights[ConstraintIndex][LocalIndex - Indices[ConstraintIndex].Num()];
+			}
+			else
+			{
+				weight = Weights[ConstraintIndex][LocalIndex];
+			}
+			if (IsAnisotropic[ConstraintIndex])
+			{
+				T comp = TVec3<T>::DotProduct(SpringEdge, Normals[ConstraintIndex]);
+				TVec3<T> proj = Normals[ConstraintIndex] * comp;
+				for (int32 alpha = 0; alpha < 3; alpha++)
+				{
+					ParticleResidual[alpha] += Dt * Dt * Stiffness[ConstraintIndex] * proj[alpha] * weight;
+				}
+			}
+			else
+			{
+				for (int32 alpha = 0; alpha < 3; alpha++)
+				{
+					ParticleResidual[alpha] += Dt * Dt * Stiffness[ConstraintIndex] * SpringEdge[alpha] * weight;
+				}
+			}
+
+		}
+
 		TArray<TArray<int32>> Indices;
 		TArray<TArray<T>> Weights;
 		TArray<TVector<T, 3>> Constraints;
