Improve IAS http analytics

#rnx

- Add a new set of http analytics that sorts each http request into a different bucket based on the overall duration and then reports the count for each bucket.
-- This new set of stats should make it easier to get an overview of how well the backends are responding and what sort of experience different users are getting.
- Add two new cvars:
-- "ias.ReportHttpAnalytics" enables the reporting of http analytics (true by default)
-- "ias.ReportCacheAnalytics" enables the reporting of cache analytics (false by default)

diff --git a/Engine/Source/Runtime/Experimental/IoStoreOnDemand/Private/Statistics.cpp b/Engine/Source/Runtime/Experimental/IoStoreOnDemand/Private/Statistics.cpp
--- a/Engine/Source/Runtime/Experimental/IoStoreOnDemand/Private/Statistics.cpp
+++ b/Engine/Source/Runtime/Experimental/IoStoreOnDemand/Private/Statistics.cpp
@@ -3,6 +3,7 @@
 #include "Statistics.h"
 
 #include "AnalyticsEventAttribute.h"
+#include "HAL/IConsoleManager.h"
 #include "Misc/CoreDelegates.h"
 
 LLM_DEFINE_TAG(Ias);
@@ -12,6 +13,16 @@ LLM_DEFINE_TAG(Ias);
 namespace UE::IO::Private
 {
 
+static TAutoConsoleVariable<bool> CVar_ReportHttpAnalytics(
+	TEXT("ias.ReportHttpAnalytics"),
+	true,
+	TEXT("Enables reporting statics on our http traffic to the analytics system"));
+
+static TAutoConsoleVariable<bool> CVar_ReportCacheAnalytics(
+	TEXT("ias.ReportCacheAnalytics"),
+	false,
+	TEXT("Enables reporting statics on our file cache usage to the analytics system"));
+
 ////////////////////////////////////////////////////////////////////////////////
 static int32 BytesToApproxMB(uint64 Bytes) { return int32(Bytes >> 20); }
 static int32 BytesToApproxKB(uint64 Bytes) { return int32(Bytes >> 10); }
@@ -204,6 +215,27 @@ int64 					GHttpHistoryIndex = 0;
 FIncrementalVariance	GHttpAvgDuration; // Duration of the http requests, in milliseconds
 FIncrementalVariance	GHttpAvgRate; // The download rate of the http requests, in MiB/s
 
+////////////////////////////////////////////////////////////////////////////////
+// 
+
+static const uint64 DurationBuckets[] = { 30, 150, 400, 1000 }; // Time boundaries (in ms) for each bucket
+
+int32 FindDurationBucket(uint64 DurationMs)
+{
+	const uint32 NumBoundries = UE_ARRAY_COUNT(DurationBuckets);
+	for (int32 Index = 0; Index < NumBoundries; ++Index)
+	{
+		if (DurationMs <= DurationBuckets[Index])
+		{
+			return Index;
+		}
+	}
+
+	return NumBoundries;
+}
+
+uint32 GHttpDurationBuckets[UE_ARRAY_COUNT(DurationBuckets) +1] = { 0 };
+
 ////////////////////////////////////////////////////////////////////////////////
 // CSV STATS
 CSV_DEFINE_CATEGORY(Ias, true);
@@ -276,25 +308,39 @@ FOnDemandIoBackendStats* FOnDemandIoBackendStats::Get()
 
 void FOnDemandIoBackendStats::ReportAnalytics(TArray<FAnalyticsEventAttribute>& OutAnalyticsArray) const
 {
-	AppendAnalyticsEventAttributeArray(OutAnalyticsArray,
-		TEXT("IasHttpErrorCount"), GHttpErrorCount.Get(), 
-		TEXT("IasHttpRetryCount"), GHttpRetryCount.Get(),
-		TEXT("IasHttpGetCount"), GHttpGetCount.Get(),
-		TEXT("IasHttpDownloadedBytes"), GHttpDownloadedBytes.Get(),
-		TEXT("IasHttpDurationMeanAvg"), GHttpAvgDuration.GetMean(),
-		TEXT("IasHttpDurationStdDev"), GHttpAvgDuration.GetDeviation(),
-		TEXT("IasHttpRateMeanAvg"), GHttpAvgRate.GetMean(),
-		TEXT("IasHttpRateStdDev"), GHttpAvgRate.GetDeviation(),
-
-		TEXT("IasCacheErrorCount"), GCacheErrorCount.Get(),
-		TEXT("IasCacheGetCount"), GCacheGetCount.Get(),
-		TEXT("IasCachePutCount"), GCachePutCount.Get(),
-		TEXT("IasCachetRejectCount"), GCachePutRejectCount.Get(),
-		
-		TEXT("IasCacheCachedBytes"), GCacheCachedBytes.Get(),
-		TEXT("IasCacheReadBytes"), GCacheReadBytes.Get(),
-		TEXT("IasCacheRejectBytes"), GCacheRejectBytes.Get()
-	);	
+	if (CVar_ReportHttpAnalytics.GetValueOnAnyThread())
+	{
+		AppendAnalyticsEventAttributeArray(OutAnalyticsArray,
+			TEXT("IasHttpErrorCount"), GHttpErrorCount.Get(),
+			TEXT("IasHttpRetryCount"), GHttpRetryCount.Get(),
+			TEXT("IasHttpGetCount"), GHttpGetCount.Get(),
+			TEXT("IasHttpDownloadedBytes"), GHttpDownloadedBytes.Get(),
+			TEXT("IasHttpDurationMeanAvg"), GHttpAvgDuration.GetMean(),
+			TEXT("IasHttpDurationStdDev"), GHttpAvgDuration.GetDeviation(),
+			TEXT("IasHttpRateMeanAvg"), GHttpAvgRate.GetMean(),
+			TEXT("IasHttpRateStdDev"), GHttpAvgRate.GetDeviation(),
+
+			TEXT("IasHttpDuration0"), GHttpDurationBuckets[0],
+			TEXT("IasHttpDuration1"), GHttpDurationBuckets[1],
+			TEXT("IasHttpDuration2"), GHttpDurationBuckets[2],
+			TEXT("IasHttpDuration3"), GHttpDurationBuckets[3],
+			TEXT("IasHttpDuration4"), GHttpDurationBuckets[4]
+		);
+	}
+
+	if (CVar_ReportCacheAnalytics.GetValueOnAnyThread())
+	{
+		AppendAnalyticsEventAttributeArray(OutAnalyticsArray,
+			TEXT("IasCacheErrorCount"), GCacheErrorCount.Get(),
+			TEXT("IasCacheGetCount"), GCacheGetCount.Get(),
+			TEXT("IasCachePutCount"), GCachePutCount.Get(),
+			TEXT("IasCachetRejectCount"), GCachePutRejectCount.Get(),
+
+			TEXT("IasCacheCachedBytes"), GCacheCachedBytes.Get(),
+			TEXT("IasCacheReadBytes"), GCacheReadBytes.Get(),
+			TEXT("IasCacheRejectBytes"), GCacheRejectBytes.Get()
+		);
+	}
 }
 
 void FOnDemandIoBackendStats::OnIoRequestEnqueue()
@@ -407,6 +453,8 @@ void FOnDemandIoBackendStats::OnHttpGet(uint64 SizeBytes, uint64 DurationMs)
 	const double DurationSeconds = DurationMs / 1000.0;
 
 	GHttpAvgRate.Increment(SizeMiB / DurationSeconds);
+
+	GHttpDurationBuckets[FindDurationBucket(DurationMs)]++;
 }
 
 void FOnDemandIoBackendStats::OnHttpRetry()
