Chaos - Ensure that collisions disabled using the FMidPhaseModifier are reset to enabled the next frame

diff --git a/Engine/Source/Runtime/Experimental/Chaos/Public/Chaos/Collision/ParticlePairMidPhase.h b/Engine/Source/Runtime/Experimental/Chaos/Public/Chaos/Collision/ParticlePairMidPhase.h
--- a/Engine/Source/Runtime/Experimental/Chaos/Public/Chaos/Collision/ParticlePairMidPhase.h
+++ b/Engine/Source/Runtime/Experimental/Chaos/Public/Chaos/Collision/ParticlePairMidPhase.h
@@ -317,6 +317,7 @@ namespace Chaos
 		 */
 		void SetIsActive(bool bIsActive)
 		{
+			Flags.bIsModified = true;
 			Flags.bIsActive = bIsActive;
 		}
 
