Removed explicit FlushFile calls as they cound towards one write op.

#rnx

diff --git a/Engine/Source/Runtime/Experimental/IoStoreOnDemand/Private/JournaledCache.cpp b/Engine/Source/Runtime/Experimental/IoStoreOnDemand/Private/JournaledCache.cpp
--- a/Engine/Source/Runtime/Experimental/IoStoreOnDemand/Private/JournaledCache.cpp
+++ b/Engine/Source/Runtime/Experimental/IoStoreOnDemand/Private/JournaledCache.cpp
@@ -440,7 +440,6 @@ void FJournal::ClosePhrase(FPhrase& Phrase)
 
 		File->Seek(DataCursor);
 		File->Write(Buffer, WriteSize);
-		File->Flush();
 		File.Reset();
 
 		Prune(DataCursor, WriteSize);
@@ -584,7 +583,6 @@ int32 FJournal::Flush()
 
 		File->Seek(JournalCursor);
 		File->Write((uint8*)(Entries.GetData()), Size);
-		File->Flush();
 		JournalCursor += Size;
 	}
 
