updates to the geometry collection's particles' shapes collision data on the PT need to propagate the information to the cluster union

diff --git a/Engine/Source/Runtime/Experimental/Chaos/Private/PhysicsProxy/GeometryCollectionPhysicsProxy.cpp b/Engine/Source/Runtime/Experimental/Chaos/Private/PhysicsProxy/GeometryCollectionPhysicsProxy.cpp
--- a/Engine/Source/Runtime/Experimental/Chaos/Private/PhysicsProxy/GeometryCollectionPhysicsProxy.cpp
+++ b/Engine/Source/Runtime/Experimental/Chaos/Private/PhysicsProxy/GeometryCollectionPhysicsProxy.cpp
@@ -3692,7 +3692,7 @@ void FGeometryCollectionPhysicsProxy::UpdatePerParticleFilterData_External(const
 	// TODO: Need to figure out how the per-particle collision filter data works with the global one. The per-particle one should probably replace the global one entirely...
 	if (Chaos::FPhysicsSolver* RBDSolver = GetSolver<Chaos::FPhysicsSolver>())
 	{
-		RBDSolver->EnqueueCommandImmediate([this, PerParticleData]()
+		RBDSolver->EnqueueCommandImmediate([this, RBDSolver, PerParticleData]()
 		{
 			Chaos::FWritePhysicsObjectInterface_Internal Interface= Chaos::FPhysicsObjectInternalInterface::GetWrite();
 
@@ -3706,6 +3706,21 @@ void FGeometryCollectionPhysicsProxy::UpdatePerParticleFilterData_External(const
 						Interface.UpdateShapeCollisionFlags(ParticleView, Data.bSimEnabled, Data.bQueryEnabled);
 						Interface.UpdateShapeFilterData(ParticleView, Data.QueryFilter, Data.SimFilter);
 					}
+
+					// todo(chaos): It's not ideal but the geometry collection needs to request the cluster union update its
+					// cached shape data if the particle is in a cluster union. This is because we don't share shape
+					// data between the GC shapes and the cluster union shapes.
+					Chaos::FRigidClustering& Clustering = RBDSolver->GetEvolution()->GetRigidClustering();
+					Chaos::FClusterUnionManager& ClusterUnionManager = Clustering.GetClusterUnionManager();
+
+					if (Chaos::FPBDRigidClusteredParticleHandle* Particle = SolverParticleHandles[Data.ParticleIndex])
+					{
+						Chaos::FClusterUnionIndex ClusterUnionIndex = ClusterUnionManager.FindClusterUnionIndexFromParticle(Particle);
+						if (ClusterUnionIndex != INDEX_NONE)
+						{
+							ClusterUnionManager.RequestDeferredClusterPropertiesUpdate(ClusterUnionIndex, Chaos::EUpdateClusterUnionPropertiesFlags::None);
+						}
+					}
 				}
 			}
 		});
diff --git a/Engine/Source/Runtime/Experimental/GeometryCollectionEngine/Private/GeometryCollection/GeometryCollectionComponent.cpp b/Engine/Source/Runtime/Experimental/GeometryCollectionEngine/Private/GeometryCollection/GeometryCollectionComponent.cpp
--- a/Engine/Source/Runtime/Experimental/GeometryCollectionEngine/Private/GeometryCollection/GeometryCollectionComponent.cpp
+++ b/Engine/Source/Runtime/Experimental/GeometryCollectionEngine/Private/GeometryCollection/GeometryCollectionComponent.cpp
@@ -3449,9 +3449,6 @@ void UGeometryCollectionComponent::LoadCollisionProfiles()
 			Cache.QueryFilter = InitialQueryFilter;
 			Cache.SimFilter = InitialSimFilter;
 
-			Cache.QueryFilter.Word3 |= (EPDF_SimpleCollision | EPDF_ComplexCollision);
-			Cache.SimFilter.Word3 |= (EPDF_SimpleCollision | EPDF_ComplexCollision);
-		
 			Cache.bQueryEnabled = true;
 			Cache.bSimEnabled = true;
 		}
