Chaos: Added mitigation for cases where geometry collections receive out of bounds replication requests for the particle array

diff --git a/Engine/Source/Runtime/Experimental/GeometryCollectionEngine/Private/GeometryCollection/GeometryCollectionComponent.cpp b/Engine/Source/Runtime/Experimental/GeometryCollectionEngine/Private/GeometryCollection/GeometryCollectionComponent.cpp
--- a/Engine/Source/Runtime/Experimental/GeometryCollectionEngine/Private/GeometryCollection/GeometryCollectionComponent.cpp
+++ b/Engine/Source/Runtime/Experimental/GeometryCollectionEngine/Private/GeometryCollection/GeometryCollectionComponent.cpp
@@ -2052,6 +2052,15 @@ namespace {
 		for (; OneOffActivatedProcessed < RepData.OneOffActivated.Num(); ++OneOffActivatedProcessed)
 		{
 			const FGeometryCollectionActivatedCluster& ActivatedCluster = RepData.OneOffActivated[OneOffActivatedProcessed];
+
+#if !UE_BUILD_SHIPPING
+			if(!PhysicsProxy->GetParticles().IsValidIndex(ActivatedCluster.ActivatedIndex))
+			{
+				ensureMsgf(false, TEXT("Invalid activated cluster index processing replication data."));
+				continue;
+			}
+#endif
+
 			FPBDRigidParticleHandle* OneOff = PhysicsProxy->GetParticles()[ActivatedCluster.ActivatedIndex];
 
 			if (ensure(OneOff))
