Chaos - make SetKinematicTarget call SetTransform when called on a dynamic particle

diff --git a/Engine/Source/Runtime/Experimental/Chaos/Private/Chaos/PBDRigidsEvolutionGBF.cpp b/Engine/Source/Runtime/Experimental/Chaos/Private/Chaos/PBDRigidsEvolutionGBF.cpp
--- a/Engine/Source/Runtime/Experimental/Chaos/Private/Chaos/PBDRigidsEvolutionGBF.cpp
+++ b/Engine/Source/Runtime/Experimental/Chaos/Private/Chaos/PBDRigidsEvolutionGBF.cpp
@@ -1051,6 +1051,33 @@ void FPBDRigidsEvolutionGBF::SetParticleTransformSwept(FGeometryParticleHandle*
 	}
 }
 
+void FPBDRigidsEvolutionGBF::SetParticleKinematicTarget(FGeometryParticleHandle* ParticleHandle, const FKinematicTarget& NewKinematicTarget)
+{
+	FGenericParticleHandle Particle = ParticleHandle;
+
+	// NOTE: If called on a dynamic body we just move the body
+	// @todo(chaos): maybe we should ensure that this is not called for dynamics
+	if (Particle->IsKinematic())
+	{
+		FKinematicGeometryParticleHandle* Kinematic = Particle->CastToKinematicParticle();
+
+		// optimization : we keep track of moving kinematic targets ( list gets clear every frame )
+		if (NewKinematicTarget.GetMode() != EKinematicTargetMode::None)
+		{
+			// move particle from "non-moving" kinematics to "moving" kinematics
+			Particles.MarkMovingKinematic(Kinematic);
+		}
+		Kinematic->SetKinematicTarget(NewKinematicTarget);
+	}
+	else if (Particle->IsDynamic())
+	{
+		if (NewKinematicTarget.GetMode() == EKinematicTargetMode::Position)
+		{
+			SetParticleTransform(ParticleHandle, NewKinematicTarget.GetTargetPosition(), NewKinematicTarget.GetTargetRotation(), false);
+		}
+	}
+}
+
 void FPBDRigidsEvolutionGBF::OnParticleMoved(FGeometryParticleHandle* InParticle, const FVec3& PrevX, const FRotation3& PrevR, const bool bIsTeleport)
 {
 	// When a particle is moved, we need to tell the collisions because they cache friction state and 
diff --git a/Engine/Source/Runtime/Experimental/Chaos/Private/PBDRigidsSolver.cpp b/Engine/Source/Runtime/Experimental/Chaos/Private/PBDRigidsSolver.cpp
--- a/Engine/Source/Runtime/Experimental/Chaos/Private/PBDRigidsSolver.cpp
+++ b/Engine/Source/Runtime/Experimental/Chaos/Private/PBDRigidsSolver.cpp
@@ -363,9 +363,7 @@ namespace Chaos
 
 		// Enable a couple bug fixes with temporary roll-back just in case
 		bool bRemoveParticleFromMovingKinematicsOnDisable = true;
-		bool bDisallowSetKinematicTargetOnDynamics = true;
 		FAutoConsoleVariableRef CVarChaosRemoveParticleFromMovingKinematicsOnDisable(TEXT("p.Chaos.RemoveParticleFromMovingKinematicsOnDisable"), bRemoveParticleFromMovingKinematicsOnDisable, TEXT(""));
-		FAutoConsoleVariableRef CVarChaosDisallowSetKinematicTargetOnDynamics(TEXT("p.Chaos.DisallowSetKinematicTargetOnDynamics"), bDisallowSetKinematicTargetOnDynamics, TEXT(""));
 	}
 }
 
diff --git a/Engine/Source/Runtime/Experimental/Chaos/Public/Chaos/PBDRigidsEvolution.h b/Engine/Source/Runtime/Experimental/Chaos/Public/Chaos/PBDRigidsEvolution.h
--- a/Engine/Source/Runtime/Experimental/Chaos/Public/Chaos/PBDRigidsEvolution.h
+++ b/Engine/Source/Runtime/Experimental/Chaos/Public/Chaos/PBDRigidsEvolution.h
@@ -66,11 +66,6 @@ struct FBroadPhaseConfig
 
 extern CHAOS_API FBroadPhaseConfig BroadPhaseConfig;
 
-namespace CVars
-{
-	extern CHAOS_API bool bDisallowSetKinematicTargetOnDynamics;
-}
-
 extern CHAOS_API int32 FixBadAccelerationStructureRemoval;
 
 class FChaosArchive;
@@ -407,32 +402,6 @@ public:
 		return IslandGroupManager.GetIterationSettings().GetNumProjectionIterations();
 	}
 
-	/**
-	* Set the kinematic target for a particle. This will exist for only one tick - a new target must be set for the next tick if required.
-	*/
-	void SetParticleKinematicTarget(FKinematicGeometryParticleHandle* KinematicHandle, const FKinematicTarget& NewKinematicTarget)
-	{
-		if (KinematicHandle == nullptr)
-		{
-			return;
-		}
-
-		// NOTE: We ignore SetKinematicTarget if called on a dynamic body since this is not supported
-		// and if we callMarkMovingKinematic(KinematicHandle) on a dynamic the particle will end up
-		// in two mutally exlcusivbe particle lists, leading to a collision detection race condition.
-		// @todo(chaos): maybe we should ensure that this is not called for dynamics
-		if ((KinematicHandle->ObjectState() == EObjectStateType::Kinematic) || !CVars::bDisallowSetKinematicTargetOnDynamics)
-		{
-			// optimization : we keep track of moving kinematic targets ( list gets clear every frame )
-			if (NewKinematicTarget.GetMode() != EKinematicTargetMode::None)
-			{
-				// move particle from "non-moving" kinematics to "moving" kinematics
-				Particles.MarkMovingKinematic(KinematicHandle);
-			}
-			KinematicHandle->SetKinematicTarget(NewKinematicTarget);
-		}
-	}
-	
 	/**
 	* To be called after creating a particle in the Particles container
 	* @todo(chaos): We should add a particle creation API to the evolution
diff --git a/Engine/Source/Runtime/Experimental/Chaos/Public/Chaos/PBDRigidsEvolutionGBF.h b/Engine/Source/Runtime/Experimental/Chaos/Public/Chaos/PBDRigidsEvolutionGBF.h
--- a/Engine/Source/Runtime/Experimental/Chaos/Public/Chaos/PBDRigidsEvolutionGBF.h
+++ b/Engine/Source/Runtime/Experimental/Chaos/Public/Chaos/PBDRigidsEvolutionGBF.h
@@ -151,6 +151,12 @@ namespace Chaos
 		 */
 		CHAOS_API virtual void SetParticleTransformSwept(FGeometryParticleHandle* InParticle, const FVec3& InPos, const FRotation3& InRot, const bool bIsTeleport);
 
+		/**
+		* Set the kinematic target for a particle. This will exist for only one tick - a new target must be set for the next tick if required.
+		* If called on a dynamic object, is equivalent to SetParticleTransform with bIsTeleport=false
+		*/
+		CHAOS_API void SetParticleKinematicTarget(FGeometryParticleHandle* ParticleHandle, const FKinematicTarget& NewKinematicTarget);
+
 		/**
 		 * Called when a particle is moved. We need to reset some friction properties, sleeping properties, etc
 		 */
