Removed the now-unused FileIoCache

#rnx

diff --git a/Engine/Source/Runtime/Experimental/IoStoreOnDemand/Private/FileIoCache.cpp b/Engine/Source/Runtime/Experimental/IoStoreOnDemand/Private/FileIoCache.cpp
deleted file mode 100644
diff --git a/Engine/Source/Runtime/Experimental/IoStoreOnDemand/Private/FileIoCache.h b/Engine/Source/Runtime/Experimental/IoStoreOnDemand/Private/FileIoCache.h
--- a/Engine/Source/Runtime/Experimental/IoStoreOnDemand/Private/FileIoCache.h
+++ b/Engine/Source/Runtime/Experimental/IoStoreOnDemand/Private/FileIoCache.h
@@ -12,8 +12,6 @@ class FIoBuffer;
 class FIoReadOptions;
 struct FIoHash;
 
-DECLARE_LOG_CATEGORY_EXTERN(LogIasCache, Log, All);
-
 /** Cache for binary blobs with a 20 byte cache key. */
 class IIoCache
 {
@@ -46,7 +44,6 @@ struct FFileIoCacheConfig
 	uint32		MemoryQuota = 2 << 20;
 	uint32		JournalQuota = 4 << 20; // description in JournalCache.cpp
 	uint32		JournalFlushInterval = 4;
-	uint32		UseLegacy = 0;
 	FRate		WriteRate;
 	bool		DropCache = false;
 };
diff --git a/Engine/Source/Runtime/Experimental/IoStoreOnDemand/Private/JournaledCache.cpp b/Engine/Source/Runtime/Experimental/IoStoreOnDemand/Private/JournaledCache.cpp
--- a/Engine/Source/Runtime/Experimental/IoStoreOnDemand/Private/JournaledCache.cpp
+++ b/Engine/Source/Runtime/Experimental/IoStoreOnDemand/Private/JournaledCache.cpp
@@ -1484,7 +1484,7 @@ IOSTOREONDEMAND_API void Tests()
 // }}}
 
 ////////////////////////////////////////////////////////////////////////////////
-TUniquePtr<IIoCache> MakeJournaledCache(const FFileIoCacheConfig& Config)
+TUniquePtr<IIoCache> MakeFileIoCache(const FFileIoCacheConfig& Config)
 {
 	LLM_SCOPE_BYTAG(Ias);
 	return MakeUnique<UE::IO::Private::FJournaledCache>(Config);
