Renamed FileIoCache-related things to use an Ias prefix.

#rnx

diff --git a/Engine/Source/Runtime/Experimental/IoStoreOnDemand/Private/FileIoCache.h b/Engine/Source/Runtime/Experimental/IoStoreOnDemand/Private/FileIoCache.h
deleted file mode 100644
diff --git a/Engine/Source/Runtime/Experimental/IoStoreOnDemand/Private/IasCache.h b/Engine/Source/Runtime/Experimental/IoStoreOnDemand/Private/IasCache.h
new file mode 100644
--- /dev/null
+++ b/Engine/Source/Runtime/Experimental/IoStoreOnDemand/Private/IasCache.h
@@ -0,0 +1,51 @@
+// Copyright Epic Games, Inc. All Rights Reserved.
+
+#pragma once
+
+#include "CancellationToken.h"
+#include "IO/IoStatus.h"
+#include "Memory/MemoryFwd.h"
+#include "Tasks/Task.h"
+#include "Templates/UniquePtr.h"
+
+class FIoBuffer;
+class FIoReadOptions;
+struct FIoHash;
+
+/** Cache for binary blobs with a 20 byte cache key. */
+class IIasCache
+{
+public:
+	virtual ~IIasCache() = default;
+
+	/** Returns whether the specified cache key is present in the cache. */
+	virtual bool ContainsChunk(const FIoHash& Key) const = 0;
+
+	/** Get the chunk associated with the specified cache key. */
+	virtual UE::Tasks::TTask<TIoStatusOr<FIoBuffer>> Get(
+		const FIoHash& Key,
+		const FIoReadOptions& Options,
+		const FIoCancellationToken* CancellationToken) = 0;
+
+	/** Insert a new chunk into the cache. */
+	virtual FIoStatus Put(const FIoHash& Key, FIoBuffer& Data) = 0;
+};
+
+struct FIasCacheConfig
+{
+	struct FRate
+	{
+		uint32	Allowance = 16 << 20;
+		uint32	Ops = 32;
+		uint32	Seconds = 60;
+	};
+
+	uint64		DiskQuota = 512ull << 20;
+	uint32		MemoryQuota = 2 << 20;
+	uint32		JournalQuota = 4 << 20; // description in JournalCache.cpp
+	uint32		JournalFlushInterval = 4;
+	FRate		WriteRate;
+	bool		DropCache = false;
+};
+
+TUniquePtr<IIasCache> MakeIasCache(const FIasCacheConfig& Config);
diff --git a/Engine/Source/Runtime/Experimental/IoStoreOnDemand/Private/IoStoreOnDemand.cpp b/Engine/Source/Runtime/Experimental/IoStoreOnDemand/Private/IoStoreOnDemand.cpp
--- a/Engine/Source/Runtime/Experimental/IoStoreOnDemand/Private/IoStoreOnDemand.cpp
+++ b/Engine/Source/Runtime/Experimental/IoStoreOnDemand/Private/IoStoreOnDemand.cpp
@@ -1,12 +1,13 @@
 // Copyright Epic Games, Inc. All Rights Reserved.
 
 #include "IO/IoStoreOnDemand.h"
-#include "OnDemandIoDispatcherBackend.h"
+
 #include "EncryptionKeyManager.h"
+#include "IasCache.h"
 #include "LatencyInjector.h"
+#include "OnDemandIoDispatcherBackend.h"
 #include "Statistics.h"
 
-#include "FileIoCache.h"
 #include "HAL/LowLevelMemTracker.h"
 #include "HAL/PlatformMisc.h"
 #include "HAL/PlatformTime.h"
@@ -172,9 +173,9 @@ static bool TryParseConfigFile(const FString& ConfigPath, UE::FOnDemandEndpoint&
 }
 
 ////////////////////////////////////////////////////////////////////////////////
-static FFileIoCacheConfig GetFileIoCacheConfig(const TCHAR* CommandLine)
+static FIasCacheConfig GetIasCacheConfig(const TCHAR* CommandLine)
 {
-	FFileIoCacheConfig Ret;
+	FIasCacheConfig Ret;
 
 	// Fetch values from .ini files
 	auto GetConfigIntImpl = [CommandLine] (const TCHAR* ConfigKey, const TCHAR* ParamName, auto& Out)
@@ -1532,10 +1533,10 @@ void FIoStoreOnDemandModule::StartupModule()
 
 	Endpoint.EndpointType = UE::EOnDemandEndpointType::CDN;
 
-	TSharedPtr<IIoCache> Cache;
-	if (FFileIoCacheConfig Config = GetFileIoCacheConfig(CommandLine); Config.DiskQuota > 0)
+	TSharedPtr<IIasCache> Cache;
+	if (FIasCacheConfig Config = GetIasCacheConfig(CommandLine); Config.DiskQuota > 0)
 	{
-		Cache = MakeShareable(MakeFileIoCache(Config).Release());
+		Cache = MakeShareable(MakeIasCache(Config).Release());
 	}
 	else
 	{
diff --git a/Engine/Source/Runtime/Experimental/IoStoreOnDemand/Private/JournaledCache.cpp b/Engine/Source/Runtime/Experimental/IoStoreOnDemand/Private/JournaledCache.cpp
--- a/Engine/Source/Runtime/Experimental/IoStoreOnDemand/Private/JournaledCache.cpp
+++ b/Engine/Source/Runtime/Experimental/IoStoreOnDemand/Private/JournaledCache.cpp
@@ -1,6 +1,6 @@
 // Copyright Epic Games, Inc. All Rights Reserved.
 
-#include "FileIoCache.h"
+#include "IasCache.h"
 #include "IO/IoStoreOnDemand.h"
 #include "Statistics.h"
 
@@ -865,7 +865,7 @@ class FCache
 {
 public:
 	struct FConfig
-		: public FFileIoCacheConfig
+		: public FIasCacheConfig
 	{
 		FString	Path;
 	};
@@ -1063,13 +1063,13 @@ uint32 FCache::DebugVisit(void* Param, FDebugCacheEntry::Callback* Callback)
 
 ////////////////////////////////////////////////////////////////////////////////
 class FJournaledCache
-	: public IIoCache
+	: public IIasCache
 	, public FRunnable
 {
 public:
 	using GetRetType = UE::Tasks::TTask<TIoStatusOr<FIoBuffer>>;
 
-								FJournaledCache(const FFileIoCacheConfig& Config);
+								FJournaledCache(const FIasCacheConfig& Config);
 								~FJournaledCache();
 	virtual bool				ContainsChunk(const FIoHash& Key) const override;
 	virtual GetRetType			Get(const FIoHash& Key, const FIoReadOptions& Options, const FIoCancellationToken* CancellationToken) override;
@@ -1093,19 +1093,19 @@ private:
 };
 
 ////////////////////////////////////////////////////////////////////////////////
-FJournaledCache::FJournaledCache(const FFileIoCacheConfig& Config)
+FJournaledCache::FJournaledCache(const FIasCacheConfig& Config)
 {
 	TStringBuilder<256> CachePath;
 	CachePath << FPaths::ProjectPersistentDownloadDir();
 	CachePath << TEXT("ias.cache.0");
 
 	FCacheInner::FConfig EventualConfig;
-	static_cast<FFileIoCacheConfig&>(EventualConfig) = Config;
+	static_cast<FIasCacheConfig&>(EventualConfig) = Config;
 	EventualConfig.Path = CachePath;
 	Cache = MakeUnique<FCacheInner>(MoveTemp(EventualConfig));
 	Cache->Load();
 
-	const FFileIoCacheConfig::FRate& WriteRate = Config.WriteRate;
+	const FIasCacheConfig::FRate& WriteRate = Config.WriteRate;
 	Governor.Set(WriteRate.Allowance, WriteRate.Ops, WriteRate.Seconds);
 
 	StartThread();
@@ -1484,7 +1484,7 @@ IOSTOREONDEMAND_API void Tests()
 // }}}
 
 ////////////////////////////////////////////////////////////////////////////////
-TUniquePtr<IIoCache> MakeFileIoCache(const FFileIoCacheConfig& Config)
+TUniquePtr<IIasCache> MakeIasCache(const FIasCacheConfig& Config)
 {
 	LLM_SCOPE_BYTAG(Ias);
 	return MakeUnique<UE::IO::Private::FJournaledCache>(Config);
diff --git a/Engine/Source/Runtime/Experimental/IoStoreOnDemand/Private/OnDemandIoDispatcherBackend.cpp b/Engine/Source/Runtime/Experimental/IoStoreOnDemand/Private/OnDemandIoDispatcherBackend.cpp
--- a/Engine/Source/Runtime/Experimental/IoStoreOnDemand/Private/OnDemandIoDispatcherBackend.cpp
+++ b/Engine/Source/Runtime/Experimental/IoStoreOnDemand/Private/OnDemandIoDispatcherBackend.cpp
@@ -7,7 +7,7 @@
 #include "Containers/StringView.h"
 #include "CoreHttp/Client.h"
 #include "EncryptionKeyManager.h"
-#include "FileIoCache.h"
+#include "IasCache.h"
 #include "HAL/Event.h"
 #include "HAL/LowLevelMemTracker.h"
 #include "HAL/Platform.h"
@@ -1182,7 +1182,7 @@ class FOnDemandIoBackend final
 	};
 public:
 
-	FOnDemandIoBackend(TSharedPtr<IIoCache> Cache);
+	FOnDemandIoBackend(TSharedPtr<IIasCache> Cache);
 	virtual ~FOnDemandIoBackend();
 
 	// I/O dispatcher backend
@@ -1218,7 +1218,7 @@ private:
 	static TIoStatusOr<FOnDemandToc> GetToc(const FOnDemandEndpoint& Endpoint);
 	FIoStatus AddToc(const FOnDemandEndpoint& Endpoint);
 
-	TSharedPtr<IIoCache> Cache;
+	TSharedPtr<IIasCache> Cache;
 	TUniquePtr<FOnDemandIoStore> IoStore;
 	TSharedPtr<const FIoDispatcherBackendContext> BackendContext;
 	TUniquePtr<FRunnableThread> BackendThread;
@@ -1237,7 +1237,7 @@ private:
 };
 
 ///////////////////////////////////////////////////////////////////////////////
-FOnDemandIoBackend::FOnDemandIoBackend(TSharedPtr<IIoCache> InCache)
+FOnDemandIoBackend::FOnDemandIoBackend(TSharedPtr<IIasCache> InCache)
 	: Cache(InCache)
 {
 	IoStore = MakeUnique<FOnDemandIoStore>();
@@ -1811,7 +1811,7 @@ uint32 FOnDemandIoBackend::Run()
 namespace UE
 {
 
-TSharedPtr<IOnDemandIoDispatcherBackend> MakeOnDemandIoDispatcherBackend(TSharedPtr<IIoCache> Cache)
+TSharedPtr<IOnDemandIoDispatcherBackend> MakeOnDemandIoDispatcherBackend(TSharedPtr<IIasCache> Cache)
 {
 	return MakeShareable<IOnDemandIoDispatcherBackend>(new UE::IO::Private::FOnDemandIoBackend(Cache));
 }
diff --git a/Engine/Source/Runtime/Experimental/IoStoreOnDemand/Private/OnDemandIoDispatcherBackend.h b/Engine/Source/Runtime/Experimental/IoStoreOnDemand/Private/OnDemandIoDispatcherBackend.h
--- a/Engine/Source/Runtime/Experimental/IoStoreOnDemand/Private/OnDemandIoDispatcherBackend.h
+++ b/Engine/Source/Runtime/Experimental/IoStoreOnDemand/Private/OnDemandIoDispatcherBackend.h
@@ -7,7 +7,7 @@
 
 #define UE_API IOSTOREONDEMAND_API
 
-class IIoCache;
+class IIasCache;
 struct FAnalyticsEventAttribute;
 
 namespace UE
@@ -50,7 +50,7 @@ public:
 #endif // IS_PROGRAM || WITH_EDITOR
 };
 
-UE_API TSharedPtr<IOnDemandIoDispatcherBackend> MakeOnDemandIoDispatcherBackend(TSharedPtr<IIoCache> Cache);
+UE_API TSharedPtr<IOnDemandIoDispatcherBackend> MakeOnDemandIoDispatcherBackend(TSharedPtr<IIasCache> Cache);
 
 } // namespace UE
 
