Add statistic when we dropping a Physics tick.

diff --git a/Engine/Source/Runtime/Experimental/Chaos/Private/Chaos/Framework/PhysicsSolverBase.cpp b/Engine/Source/Runtime/Experimental/Chaos/Private/Chaos/Framework/PhysicsSolverBase.cpp
--- a/Engine/Source/Runtime/Experimental/Chaos/Private/Chaos/Framework/PhysicsSolverBase.cpp
+++ b/Engine/Source/Runtime/Experimental/Chaos/Private/Chaos/Framework/PhysicsSolverBase.cpp
@@ -22,6 +22,8 @@ DEFINE_STAT(STAT_ProcessGCProxy);
 DEFINE_STAT(STAT_ProcessClusterUnionProxy);
 DEFINE_STAT(STAT_PullConstraints);
 
+CSV_DEFINE_CATEGORY(ChaosPhysicsSolver, true);
+
 namespace Chaos
 {	
 	extern int GSingleThreadedPhysics;
@@ -422,10 +424,12 @@ namespace Chaos
 		{
 			// Make sure not to accumulate too many physics solver tasks.
 			constexpr int32 MaxPhysicsStepToKeep = 3;
-			if (NumSteps + NumPendingSolverAdvanceTasks > MaxPhysicsStepToKeep)
+			const int32 MaxNumSteps = MaxPhysicsStepToKeep - NumPendingSolverAdvanceTasks;
+			if (NumSteps > MaxNumSteps)
 			{
+				CSV_CUSTOM_STAT(ChaosPhysicsSolver, PhysicsFrameDropped, NumSteps - MaxNumSteps, ECsvCustomStatOp::Accumulate);
 				// NumSteps + NumPendingSolverAdvanceTasks shouldn't be bigger than MaxPhysicsStepToKeep
-				NumSteps = FMath::Clamp<int32>(NumSteps, 0, MaxPhysicsStepToKeep - NumPendingSolverAdvanceTasks);
+				NumSteps = FMath::Min<int32>(NumSteps, MaxNumSteps);
 			}
 		}
 			
