Geometry collection : allow SetDamageThreshold to be set without recreating the physics state
- Consolidate the damage thresholds calculation into one method
- Do not recreate proxy when setting damage thresholds from the component

diff --git a/Engine/Source/Runtime/Experimental/Chaos/Private/PhysicsProxy/GeometryCollectionPhysicsProxy.cpp b/Engine/Source/Runtime/Experimental/Chaos/Private/PhysicsProxy/GeometryCollectionPhysicsProxy.cpp
--- a/Engine/Source/Runtime/Experimental/Chaos/Private/PhysicsProxy/GeometryCollectionPhysicsProxy.cpp
+++ b/Engine/Source/Runtime/Experimental/Chaos/Private/PhysicsProxy/GeometryCollectionPhysicsProxy.cpp
@@ -452,6 +452,7 @@ FGeometryCollectionPhysicsProxy::FGeometryCollectionPhysicsProxy(
 	, GameThreadPerFrameData(SimulationParameters)
 	, bIsPhysicsThreadWorldTransformDirty(false)
 	, bIsCollisionFilterDataDirty(false)
+	, bIsDamageThresholdDataDirty(false)
 	, CollectorGuid(InCollectorGuid)
 {
 	// We rely on a guarded buffer.
@@ -956,6 +957,62 @@ void FGeometryCollectionPhysicsProxy::DirtyAllParticles(const Chaos::FPBDRigidsS
 	}
 }
 
+void FGeometryCollectionPhysicsProxy::UpdateDamageThreshold_Internal()
+{
+	const int32 NumTransforms = SolverParticleHandles.Num();
+	ensure(SolverParticleHandles.Num() == SolverClusterHandles.Num());
+
+	if (Chaos::FPhysicsSolver* RBDSolver = GetSolver<Chaos::FPhysicsSolver>())
+	{
+		Chaos::FRigidClustering& Clustering = RBDSolver->GetEvolution()->GetRigidClustering();
+
+		const float StrainDefault = Parameters.DamageThreshold.Num() ? Parameters.DamageThreshold[0] : 0;
+
+		for (int32 TransformIndex = 0; TransformIndex < NumTransforms; TransformIndex++)
+		{
+			if (Chaos::FPBDRigidClusteredParticleHandle* Handle = SolverParticleHandles[TransformIndex])
+			{
+				float DamageThreshold = StrainDefault;
+
+				switch (Parameters.DamageModel)
+				{
+				case EDamageModelTypeEnum::Chaos_Damage_Model_UserDefined_Damage_Threshold:
+					if (!Parameters.bUsePerClusterOnlyDamageThreshold)
+					{
+						const bool bIsACluster = (Handle->ClusterIds().NumChildren > 0);
+						if (Parameters.EnableClustering || bIsACluster)
+						{
+							DamageThreshold = ComputeUserDefinedDamageThreshold_Internal(TransformIndex);
+						}
+					}
+					break;
+				case EDamageModelTypeEnum::Chaos_Damage_Model_Material_Strength_And_Connectivity_DamageThreshold:
+					DamageThreshold = ComputeMaterialBasedDamageThreshold_Internal(TransformIndex);
+					break;
+				}
+				Clustering.SetInternalStrain(Handle, DamageThreshold);
+			}
+		}
+
+		// user defined legacy mode: propagate to the children form the cluster values
+		if (Parameters.bUsePerClusterOnlyDamageThreshold && 
+			Parameters.DamageModel == EDamageModelTypeEnum::Chaos_Damage_Model_UserDefined_Damage_Threshold)
+		{
+			for (int32 TransformIndex = 0; TransformIndex < NumTransforms; TransformIndex++)
+			{
+				if (Chaos::FPBDRigidClusteredParticleHandle* Handle = SolverParticleHandles[TransformIndex])
+				{
+					if (Chaos::FPBDRigidClusteredParticleHandle* ParentHandle = SolverClusterHandles[TransformIndex])
+					{
+						const float DamageThreshold = ParentHandle->GetInternalStrains();
+						Clustering.SetInternalStrain(Handle, DamageThreshold);
+					}
+				}
+			}
+		}
+	}
+}
+
 void FGeometryCollectionPhysicsProxy::InitializeBodiesPT(Chaos::FPBDRigidsSolver* RigidsSolver, typename Chaos::FPBDRigidsSolver::FParticlesType& Particles)
 {
 	const FGeometryCollection* RestCollection = Parameters.RestCollection;
@@ -1056,15 +1113,6 @@ void FGeometryCollectionPhysicsProxy::InitializeBodiesPT(Chaos::FPBDRigidsSolver
 				if (Parameters.EnableClustering)
 				{
 					Handle->SetClusterGroupIndex(Parameters.ClusterGroupIndex);
-
-					float DamageThreshold = StrainDefault;
-					if (!Parameters.bUsePerClusterOnlyDamageThreshold)
-					{
-						DamageThreshold = ComputeUserDefinedDamageThreshold_Internal(TransformGroupIndex);
-					}
-					Chaos::FPhysicsSolver* RBDSolver = GetSolver<Chaos::FPhysicsSolver>();
-					RBDSolver->GetEvolution()->GetRigidClustering().SetInternalStrain(Handle, DamageThreshold);
-
 				}
 
 				// #BGTODO - non-updating parameters - remove lin/ang drag arrays and always query material if this stays a material parameter
@@ -1408,19 +1456,8 @@ void FGeometryCollectionPhysicsProxy::InitializeBodiesPT(Chaos::FPBDRigidsSolver
 #endif
 		}
 
-		// assign the material based damage thresholds 
-		if (Parameters.DamageModel != EDamageModelTypeEnum::Chaos_Damage_Model_UserDefined_Damage_Threshold)
-		{
-			Chaos::FRigidClustering& Clustering = RigidsSolver->GetEvolution()->GetRigidClustering();
-			for (int32 Index = 0; Index < NumTransforms; Index++)
-			{
-				if (Chaos::FPBDRigidClusteredParticleHandle* ClusteredParticle = SolverParticleHandles[Index])
-				{
-					const float DamageThreshold = ComputeMaterialBasedDamageThreshold_Internal(Index);
-					Clustering.SetInternalStrain(ClusteredParticle, DamageThreshold);
-				}
-			}
-		}
+		// set the damage thresholds
+		UpdateDamageThreshold_Internal();
 
 		// call DirtyParticle to make sure the acceleration structure is up to date with all the changes happening here
 		DirtyAllParticles(*RigidsSolver);
@@ -1638,16 +1675,6 @@ Chaos::TPBDGeometryCollectionParticleHandle<Chaos::FReal, 3>* FGeometryCollectio
 	if (Parameters.EnableClustering)
 	{
 		Handle->SetClusterGroupIndex(Parameters.ClusterGroupIndex);
-
-		const float StrainDefault = Parameters.DamageThreshold.Num() ? Parameters.DamageThreshold[0] : 0;
-		float DamageThreshold = StrainDefault;
-		if (!Parameters.bUsePerClusterOnlyDamageThreshold)
-		{
-			DamageThreshold = ComputeUserDefinedDamageThreshold_Internal(CollectionClusterIndex);
-		}
-		Chaos::FPhysicsSolver* RBDSolver = GetSolver<Chaos::FPhysicsSolver>();
-		RBDSolver->GetEvolution()->GetRigidClustering().SetInternalStrain(Handle, DamageThreshold);
-
 	}
 
 	// #BGTODO This will not automatically update - material properties should only ever exist in the material, not in other arrays
@@ -1786,11 +1813,6 @@ FGeometryCollectionPhysicsProxy::BuildClusters_Internal(
 	// two-way mapping
 	SolverClusterHandles[CollectionClusterIndex] = Parent;
 
-	const float DamageThreshold = ComputeUserDefinedDamageThreshold_Internal(CollectionClusterIndex);
-
-	Chaos::FRigidClustering& RigidClustering = static_cast<Chaos::FPBDRigidsSolver*>(Solver)->GetEvolution()->GetRigidClustering();
-	RigidClustering.SetInternalStrain(Parent, DamageThreshold);
-	
 	// #BGTODO This will not automatically update - material properties should only ever exist in the material, not in other arrays
 	const Chaos::FChaosPhysicsMaterial* CurMaterial = static_cast<Chaos::FPBDRigidsSolver*>(Solver)->GetSimMaterials().Get(Parameters.PhysicalMaterialHandle.InnerHandle);
 	if(CurMaterial)
@@ -1812,13 +1834,6 @@ FGeometryCollectionPhysicsProxy::BuildClusters_Internal(
 	{
 		// set the damage threshold on children as they are the one where the strain is tested when breaking 
 		Chaos::FPBDRigidParticleHandle* Child = ChildHandles[Idx];
-		if (Parameters.bUsePerClusterOnlyDamageThreshold)
-		{
-			if (Chaos::FPBDRigidClusteredParticleHandle* ClusteredChild = Child->CastToClustered())
-			{
-				RigidClustering.SetInternalStrain(ClusteredChild, DamageThreshold);
-			}
-		}
 
 		const int32 ChildTransformGroupIndex = ChildTransformGroupIndices[Idx];
 		SolverClusterHandles[ChildTransformGroupIndex] = Parent;
@@ -2769,6 +2784,13 @@ void FGeometryCollectionPhysicsProxy::SetWorldTransform_External(const FTransfor
 	SetProxyDirty_External();
 }
 
+void FGeometryCollectionPhysicsProxy::SetDamageThresholds_External(const TArray<float>& DamageThresholds)
+{
+	check(IsInGameThread());
+	GameThreadPerFrameData.SetDamageThresholds(DamageThresholds);
+	SetProxyDirty_External();
+}
+
 void FGeometryCollectionPhysicsProxy::PushStateOnGameThread(Chaos::FPBDRigidsSolver* InSolver)
 {
 	// CONTEXT: GAMETHREAD
@@ -2805,6 +2827,13 @@ void FGeometryCollectionPhysicsProxy::PushStateOnGameThread(Chaos::FPBDRigidsSol
 		Parameters.bEnableStrainOnCollision = GameThreadPerFrameData.GetEnableStrainOnCollision();
 		GameThreadPerFrameData.ResetIsDamageSettingsDataDirty();
 	}
+
+	bIsDamageThresholdDataDirty = GameThreadPerFrameData.GetIsDamageThresholdDataDirty();
+	if (bIsDamageThresholdDataDirty)
+	{
+		Parameters.DamageThreshold = GameThreadPerFrameData.GetDamageThresholds();
+		GameThreadPerFrameData.ResetIsDamageThresholdDataDirty();
+	}
 }
 
 void FGeometryCollectionPhysicsProxy::PushToPhysicsState()
@@ -2929,6 +2958,11 @@ void FGeometryCollectionPhysicsProxy::PushToPhysicsState()
 			}
 		}
 	}
+
+	if (bIsDamageThresholdDataDirty)
+	{
+		UpdateDamageThreshold_Internal();
+	}
 }
 
 void FGeometryCollectionPhysicsProxy::SetClusteredParticleKinematicTarget_Internal(Chaos::FPBDRigidClusteredParticleHandle* Handle, const FTransform& NewWorldTransform)
diff --git a/Engine/Source/Runtime/Experimental/Chaos/Public/GeometryCollectionProxyData.h b/Engine/Source/Runtime/Experimental/Chaos/Public/GeometryCollectionProxyData.h
--- a/Engine/Source/Runtime/Experimental/Chaos/Public/GeometryCollectionProxyData.h
+++ b/Engine/Source/Runtime/Experimental/Chaos/Public/GeometryCollectionProxyData.h
@@ -164,6 +164,7 @@ public:
 		, bIsCollisionFilterDataDirty(false)
 		, bIsNotificationDataDirty(false) 
 		, bIsDamageSettingsDataDirty(false)
+		, bIsDamageThresholdDataDirty(false)
 		, bNotifyBreakings(false)
 		, bNotifyRemovals(false)
 		, bNotifyCrumblings(false)
@@ -180,6 +181,7 @@ public:
 		, bIsCollisionFilterDataDirty(false)
 		, bIsNotificationDataDirty(false)
 		, bIsDamageSettingsDataDirty(false)
+		, bIsDamageThresholdDataDirty(false)
 		, bNotifyBreakings(Parameters.bGenerateBreakingData)
 		, bNotifyRemovals(false)
 		, bNotifyCrumblings(Parameters.bGenerateCrumblingData)
@@ -189,6 +191,7 @@ public:
 		, bNotifyGlobalCrumblings(Parameters.bGenerateGlobalCrumblingData)
 		, bGlobalCrumblingEventIncludesChildren(Parameters.bGenerateGlobalCrumblingChildrenData)
 		, bEnableStrainOnCollision(false)
+		, DamageThresholds(Parameters.DamageThreshold)
 	{}
 
 	const FTransform& GetWorldTransform() const { return WorldTransform; }
@@ -286,11 +289,24 @@ public:
 	bool GetIsDamageSettingsDataDirty() const { return bIsDamageSettingsDataDirty; }
 	void ResetIsDamageSettingsDataDirty() { bIsDamageSettingsDataDirty = false; }
 
+	void SetDamageThresholds(const TArray<float>& InDamageThresholds)
+	{
+		DamageThresholds = InDamageThresholds;
+		bIsDamageThresholdDataDirty = true;
+	}
+
+	const TArray<float>& GetDamageThresholds() const { return DamageThresholds; }
+
+	bool GetIsDamageThresholdDataDirty() const { return bIsDamageThresholdDataDirty; }
+	void ResetIsDamageThresholdDataDirty() { bIsDamageThresholdDataDirty = false; }
+
+
 private:
 	uint16 bIsWorldTransformDirty : 1;
 	uint16 bIsCollisionFilterDataDirty : 1;
 	uint16 bIsNotificationDataDirty: 1;
 	uint16 bIsDamageSettingsDataDirty : 1;
+	uint16 bIsDamageThresholdDataDirty : 1;
 
 	/** updated when bNotificationDataDirty is set */
 	uint16 bNotifyBreakings : 1;
@@ -302,9 +318,12 @@ private:
 	uint16 bNotifyGlobalCrumblings : 1;
 	uint16 bGlobalCrumblingEventIncludesChildren : 1;
 
-	/** updated when bDamageSettingsDataDirty is set */
+	/** updated when bIsDamageSettingsDataDirty is set */
 	uint16 bEnableStrainOnCollision : 1;
 
+	/** updated when bIsDamageThresholdDataDirty is set */
+	TArray<float> DamageThresholds;
+
 	/** updated when bIsWorldTransformDirty is set */
 	FTransform WorldTransform; 
 
diff --git a/Engine/Source/Runtime/Experimental/Chaos/Public/PhysicsProxy/GeometryCollectionPhysicsProxy.h b/Engine/Source/Runtime/Experimental/Chaos/Public/PhysicsProxy/GeometryCollectionPhysicsProxy.h
--- a/Engine/Source/Runtime/Experimental/Chaos/Public/PhysicsProxy/GeometryCollectionPhysicsProxy.h
+++ b/Engine/Source/Runtime/Experimental/Chaos/Public/PhysicsProxy/GeometryCollectionPhysicsProxy.h
@@ -457,6 +457,8 @@ public:
 
 	CHAOS_API void UpdatePerParticleFilterData_External(const TArray<FParticleCollisionFilterData>& Data);
 	
+	CHAOS_API void SetDamageThresholds_External(const TArray<float>& DamageThresholds);
+
 	/** 
 	 * Traverses the parents of TransformGroupIdx counting number of levels,
 	 * and sets levels array value for TransformGroupIdx and its parents if not yet initialized.
@@ -530,6 +532,9 @@ protected:
 	CHAOS_API void CreateNonClusteredParticles(Chaos::FPBDRigidsSolver* RigidsSolver,	const FGeometryCollection& RestCollection, const FGeometryDynamicCollection& DynamicCollection, const TBitArray<>& EffectiveParticles);
 
 	CHAOS_API Chaos::FPBDRigidClusteredParticleHandle* FindClusteredParticleHandleByItemIndex_Internal(FGeometryCollectionItemIndex ItemIndex) const;
+
+	CHAOS_API void UpdateDamageThreshold_Internal();
+
 private:
 
 	/**
@@ -628,6 +633,7 @@ private:
 	FGeometryCollectioPerFrameData GameThreadPerFrameData;
 	bool bIsPhysicsThreadWorldTransformDirty;
 	bool bIsCollisionFilterDataDirty;
+	bool bIsDamageThresholdDataDirty;
 
 	// Currently this is using triple buffers for game-physics and 
 	// physics-game thread communication, but not for any reason other than this 
diff --git a/Engine/Source/Runtime/Experimental/GeometryCollectionEngine/Private/GeometryCollection/GeometryCollectionComponent.cpp b/Engine/Source/Runtime/Experimental/GeometryCollectionEngine/Private/GeometryCollection/GeometryCollectionComponent.cpp
--- a/Engine/Source/Runtime/Experimental/GeometryCollectionEngine/Private/GeometryCollection/GeometryCollectionComponent.cpp
+++ b/Engine/Source/Runtime/Experimental/GeometryCollectionEngine/Private/GeometryCollection/GeometryCollectionComponent.cpp
@@ -723,10 +723,9 @@ void UGeometryCollectionComponent::SetDamageThreshold(const TArray<float>& InDam
 
 	DamageThreshold = InDamageThreshold;
 
-	if (IsPhysicsStateCreated())
+	if (PhysicsProxy)
 	{
-		// NOTE: Probably should use a deferred version of this.
-		RecreatePhysicsState();
+		PhysicsProxy->SetDamageThresholds_External(DamageThreshold);
 	}
 }
 
