Moved FDistributionEndpoints to it's own code files

#rnx

- Added to the UE::IO namespace since the header is private.
- Fixed a crash in FDistributionEndpoints::IssueEndpointRequests as if the request timed out the returned response pointer will be null and needs to be handled as a failure.
- No longer load the http module directly as calling FHttpModule::Get() will do that anyway.
- Moved the code from FOnDemandIoBackend::FlushDeferedEndPoints to FDistributionEndpoints::Flush.
- Improved the error message if flushing the endpoints times out.

diff --git a/Engine/Source/Runtime/Experimental/IoStoreOnDemand/Private/DistributionEndpoints.cpp b/Engine/Source/Runtime/Experimental/IoStoreOnDemand/Private/DistributionEndpoints.cpp
new file mode 100644
--- /dev/null
+++ b/Engine/Source/Runtime/Experimental/IoStoreOnDemand/Private/DistributionEndpoints.cpp
@@ -0,0 +1,238 @@
+// Copyright Epic Games, Inc. All Rights Reserved.
+
+#include "DistributionEndpoints.h"
+
+#include "Dom/JsonValue.h"
+#include "HAL/PlatformTime.h"
+#include "HttpManager.h"
+#include "HttpModule.h"
+#include "IO/IoStoreOnDemand.h"
+#include "Interfaces/IHttpResponse.h"
+#include "Misc/ScopeRWLock.h"
+#include "Serialization/JsonReader.h"
+#include "Serialization/JsonSerializer.h"
+#include "Statistics.h"
+
+extern int32 GIoDispatcherMaxHttpRetryCount;
+
+namespace UE::IO
+{
+
+FDistributionEndpoints::~FDistributionEndpoints()
+{
+	CancelEndpointRequests();
+}
+
+#if IS_PROGRAM || WITH_EDITOR
+
+bool FDistributionEndpoints::Flush(double TimeOut)
+{
+	TRACE_CPUPROFILER_EVENT_SCOPE(FDistributionEndpoints::Flush);
+
+	FHttpManager& HttpManager = FHttpModule::Get().GetHttpManager();
+
+	const double StartTime = FPlatformTime::Seconds();
+
+	while (!PendingRequests.IsEmpty())
+	{
+		HttpManager.Tick(0.0);
+
+		FPlatformProcess::SleepNoStats(0.0f);
+
+		if (TimeOut > 0.0 && (FPlatformTime::Seconds() - StartTime) > TimeOut)
+		{
+			return false;
+		}
+	}
+
+	return true;
+}
+
+#endif // IS_PROGRAM || WITH_EDITOR
+
+void FDistributionEndpoints::ResolveEndpoints(const FString& DistributionUrl, FOnEndpointResolved&& OnResolved)
+{
+	TRACE_CPUPROFILER_EVENT_SCOPE(FOnDemandIoBackend::ResolveEndpoints);
+	const FResolvedEndpoint* Ep = nullptr;
+	{
+		FReadScopeLock _(Lock);
+		if (TUniquePtr<FResolvedEndpoint>* Entry = ResolvedEndpoints.Find(DistributionUrl))
+		{
+			Ep = Entry->Get();
+		}
+	}
+
+	if (Ep != nullptr)
+	{
+		return OnResolved(DistributionUrl, Ep->ServiceUrls);
+	}
+
+	bool bIssueRequest = false;
+	{
+		FWriteScopeLock _(Lock);
+		TUniquePtr<FResolveRequest>& Request = PendingRequests.FindOrAdd(DistributionUrl);
+		if (!Request.IsValid())
+		{
+			Request.Reset(new FResolveRequest{ DistributionUrl });
+			bIssueRequest = bInitialized;
+		}
+		Request->Callbacks.Add(MoveTemp(OnResolved));
+	}
+
+	if (bIssueRequest)
+	{
+		IssueEndpointRequests();
+	}
+}
+
+void FDistributionEndpoints::ResolveDeferredEndpoints()
+{
+	TRACE_CPUPROFILER_EVENT_SCOPE(FOnDemandIoBackend::ResolveDeferredEndpoints);
+	{
+		FWriteScopeLock _(Lock);
+		bInitialized = true;
+	}
+
+	IssueEndpointRequests();
+}
+
+void FDistributionEndpoints::IssueEndpointRequests()
+{
+	TRACE_CPUPROFILER_EVENT_SCOPE(FOnDemandIoBackend::IssueEndpointRequests);
+	// Currently we need to use the HTTP module in order to resolve service endpoints due to HTTPS
+	FHttpModule& HttpModule = FHttpModule::Get();
+	const int32 MaxAttempts = GIoDispatcherMaxHttpRetryCount;
+
+	TArray<FHttpRequestPtr, TInlineAllocator<2>> HttpRequests;
+	{
+		FWriteScopeLock _(Lock);
+		check(bInitialized);
+
+		for (const TPair<FString, TUniquePtr<FResolveRequest>>& Kv : PendingRequests)
+		{
+			if (Kv.Value->HttpRequest.IsValid())
+			{
+				continue;
+			}
+
+			FResolveRequest& ResolveRequest = *Kv.Value.Get();
+			UE_LOG(LogIas, Log, TEXT("Resolving '%s' (#%d/%d)"), *ResolveRequest.DistributionUrl, ResolveRequest.RetryCount + 1, MaxAttempts);
+
+			FHttpRequestPtr HttpRequest = HttpModule.CreateRequest();
+			HttpRequest->SetTimeout(3.0f);
+			HttpRequest->SetURL(Kv.Key);
+			HttpRequest->SetVerb(TEXT("GET"));
+			HttpRequest->SetHeader(TEXT("Accept"), TEXT("application/json"));
+			HttpRequest->OnProcessRequestComplete().BindLambda(
+				[this, &ResolveRequest, MaxAttempts](FHttpRequestPtr, FHttpResponsePtr Response, bool bOk)
+				{
+					LLM_SCOPE_BYTAG(Ias);
+					FHttpRequestPtr Request = MoveTemp(ResolveRequest.HttpRequest);
+
+					// Response will be null if the connection timed out
+					if (Response == nullptr || Response->GetResponseCode() != 200)
+					{
+						if (++ResolveRequest.RetryCount <= MaxAttempts)
+						{
+							Request->OnProcessRequestComplete().Unbind();
+							return IssueEndpointRequests();
+						}
+					}
+
+					CompleteEndpointRequest(ResolveRequest, Response);
+				});
+
+			ResolveRequest.HttpRequest = HttpRequest;
+			HttpRequests.Add(HttpRequest);
+		}
+	}
+
+	for (FHttpRequestPtr& Request : HttpRequests)
+	{
+		Request->ProcessRequest();
+	}
+}
+
+void FDistributionEndpoints::CancelEndpointRequests()
+{
+	TArray<FHttpRequestPtr, TInlineAllocator<2>> HttpRequests;
+	{
+		FWriteScopeLock _(Lock);
+		for (const TPair<FString, TUniquePtr<FResolveRequest>>& Kv : PendingRequests)
+		{
+			if (Kv.Value->HttpRequest.IsValid())
+			{
+				HttpRequests.Add(Kv.Value->HttpRequest);
+			}
+		}
+	}
+
+	if (!HttpRequests.IsEmpty())
+	{
+		FHttpModule& HttpModule = FHttpModule::Get();
+		for (FHttpRequestPtr& Request : HttpRequests)
+		{
+			HttpModule.GetHttpManager().RemoveRequest(Request.ToSharedRef());
+			//TODO: Flush?
+		}
+	}
+}
+
+void FDistributionEndpoints::CompleteEndpointRequest(FResolveRequest& ResolveRequest, FHttpResponsePtr HttpResponse)
+{
+	TRACE_CPUPROFILER_EVENT_SCOPE(FOnDemandIoBackend::CompleteEndpointRequest);
+
+	using FJsonValuePtr = TSharedPtr<FJsonValue>;
+	using FJsonObjPtr = TSharedPtr<FJsonObject>;
+	using FJsonReader = TJsonReader<TCHAR>;
+	using FJsonReaderPtr = TSharedRef<FJsonReader>;
+
+	TArray<FString> ServiceUrls;
+	if (HttpResponse->GetResponseCode() == 200)
+	{
+		FString Json = HttpResponse->GetContentAsString();
+		FJsonReaderPtr JsonReader = TJsonReaderFactory<TCHAR>::Create(Json);
+
+		FJsonObjPtr JsonObj;
+		if (FJsonSerializer::Deserialize(JsonReader, JsonObj))
+		{
+			TArray<FJsonValuePtr> JsonValues = JsonObj->GetArrayField(TEXT("distributions"));
+			for (const FJsonValuePtr& JsonValue : JsonValues)
+			{
+				FString ServiceUrl = JsonValue->AsString();
+				if (ServiceUrl.EndsWith(TEXT("/")))
+				{
+					ServiceUrl.LeftInline(ServiceUrl.Len() - 1);
+				}
+				ServiceUrls.Add(MoveTemp(ServiceUrl));
+			}
+		}
+	}
+
+	const FResolvedEndpoint* ResolvedEndpoint = nullptr;
+	FString DistributionUrl;
+	TArray<FOnEndpointResolved> Callbacks;
+
+	{
+		FWriteScopeLock _(Lock);
+		if (!ServiceUrls.IsEmpty())
+		{
+			ResolvedEndpoint = ResolvedEndpoints.Emplace(
+				ResolveRequest.DistributionUrl,
+				new FResolvedEndpoint{ MoveTemp(ServiceUrls) })
+				.Get();
+		}
+
+		Callbacks = MoveTemp(ResolveRequest.Callbacks);
+		DistributionUrl = MoveTemp(ResolveRequest.DistributionUrl);
+		PendingRequests.Remove(DistributionUrl);
+	}
+
+	TConstArrayView<FString> Urls = ResolvedEndpoint ? ResolvedEndpoint->ServiceUrls : TConstArrayView<FString>();
+	for (FOnEndpointResolved& Callback : Callbacks)
+	{
+		Callback(DistributionUrl, Urls);
+	}
+}
+
+} // namespace UE::IO
diff --git a/Engine/Source/Runtime/Experimental/IoStoreOnDemand/Private/DistributionEndpoints.h b/Engine/Source/Runtime/Experimental/IoStoreOnDemand/Private/DistributionEndpoints.h
new file mode 100644
--- /dev/null
+++ b/Engine/Source/Runtime/Experimental/IoStoreOnDemand/Private/DistributionEndpoints.h
@@ -0,0 +1,54 @@
+// Copyright Epic Games, Inc. All Rights Reserved.
+
+#pragma once
+
+#include "Containers/Array.h"
+#include "Containers/Map.h"
+#include "Containers/UnrealString.h"
+#include "HAL/CriticalSection.h"
+#include "HttpFwd.h"
+
+namespace UE::IO
+{
+
+class FDistributionEndpoints
+{
+public:
+	using FOnEndpointResolved = TFunction<void(const FString&, TConstArrayView<FString>)>;
+
+	FDistributionEndpoints() = default;
+	~FDistributionEndpoints();
+
+#if IS_PROGRAM || WITH_EDITOR
+	bool Flush(double TimeOut);
+#endif //IS_PROGRAM || WITH_EDITOR
+
+	void ResolveEndpoints(const FString& DistributionUrl, FOnEndpointResolved&& OnResolved);
+	void ResolveDeferredEndpoints();
+
+private:
+	struct FResolvedEndpoint
+	{
+		TArray<FString> ServiceUrls;
+	};
+
+	struct FResolveRequest
+	{
+		FString DistributionUrl;
+		FHttpRequestPtr HttpRequest;
+		TArray<FOnEndpointResolved> Callbacks;
+		int32 RetryCount = 0;
+	};
+
+	void IssueEndpointRequests();
+	void CancelEndpointRequests();
+	void CompleteEndpointRequest(FResolveRequest& ResolveRequest, FHttpResponsePtr HttpResponse);
+
+	TMap<FString, TUniquePtr<FResolvedEndpoint>> ResolvedEndpoints;
+	TMap<FString, TUniquePtr<FResolveRequest>> PendingRequests;
+
+	FRWLock Lock;
+	bool bInitialized = false;
+};
+
+} // namespace UE::IO
diff --git a/Engine/Source/Runtime/Experimental/IoStoreOnDemand/Private/IoStoreOnDemand.cpp b/Engine/Source/Runtime/Experimental/IoStoreOnDemand/Private/IoStoreOnDemand.cpp
--- a/Engine/Source/Runtime/Experimental/IoStoreOnDemand/Private/IoStoreOnDemand.cpp
+++ b/Engine/Source/Runtime/Experimental/IoStoreOnDemand/Private/IoStoreOnDemand.cpp
@@ -1406,7 +1406,7 @@ FIoStatus PrimeEndPoint(FStringView IoStoreOnDemandIniPath)
 	if (!Backend->FlushDeferedEndPoints(30.0))
 	{
 		// TODO: Real error value
-		return FIoStatus(EIoErrorCode::Unknown, TEXT("Unable to connect to endpoint"));
+		return FIoStatus(EIoErrorCode::Unknown, TEXT("Unable to connect to endpoint, timed out..."));
 	}
 
 	UE_LOG(LogIas, Display, TEXT("Finding on demand chunks..."));
diff --git a/Engine/Source/Runtime/Experimental/IoStoreOnDemand/Private/OnDemandIoDispatcherBackend.cpp b/Engine/Source/Runtime/Experimental/IoStoreOnDemand/Private/OnDemandIoDispatcherBackend.cpp
--- a/Engine/Source/Runtime/Experimental/IoStoreOnDemand/Private/OnDemandIoDispatcherBackend.cpp
+++ b/Engine/Source/Runtime/Experimental/IoStoreOnDemand/Private/OnDemandIoDispatcherBackend.cpp
@@ -6,6 +6,7 @@
 #include "CancellationToken.h"
 #include "Containers/StringView.h"
 #include "CoreHttp/Client.h"
+#include "DistributionEndpoints.h"
 #include "EncryptionKeyManager.h"
 #include "IasCache.h"
 #include "HAL/Event.h"
@@ -15,7 +16,6 @@
 #include "HAL/PreprocessorHelpers.h"
 #include "HAL/Runnable.h"
 #include "HAL/RunnableThread.h"
-#include "Http.h"
 #include "HttpManager.h"
 #include "IO/IoAllocators.h"
 #include "IO/IoChunkEncoding.h"
@@ -26,11 +26,8 @@
 #include "IO/IoStoreOnDemand.h"
 #include "Misc/ScopeLock.h"
 #include "Misc/ScopeRWLock.h"
-#include "Modules/ModuleManager.h"
 #include "Serialization/CompactBinary.h"
 #include "Serialization/CompactBinarySerialization.h"
-#include "Serialization/JsonReader.h"
-#include "Serialization/JsonSerializer.h"
 #include "Serialization/MemoryReader.h"
 #include "Statistics.h"
 #include "Tasks/Task.h"
@@ -147,231 +144,6 @@ FIoHash GetChunkKey(const FIoHash& ChunkHash, const FIoOffsetAndLength& Range)
 	return HashBuilder.Finalize();
 }
 
-///////////////////////////////////////////////////////////////////////////////
-class FDistributionEndpoints
-{
-public:
-	using FOnEndpointResolved = TFunction<void(const FString&, TConstArrayView<FString>)>;
-	
-	FDistributionEndpoints() = default;
-	~FDistributionEndpoints();
-
-	void ResolveEndpoints(const FString& DistributionUrl, FOnEndpointResolved&& OnResolved);
-	void ResolveDeferredEndpoints();
-
-private:
-	struct FResolvedEndpoint
-	{
-		TArray<FString> ServiceUrls;
-	};
-
-	struct FResolveRequest
-	{
-		FString DistributionUrl;
-		FHttpRequestPtr HttpRequest;
-		TArray<FOnEndpointResolved> Callbacks;
-		int32 RetryCount = 0;
-	};
-
-	void IssueEndpointRequests();
-	void CancelEndpointRequests();
-	void CompleteEndpointRequest(FResolveRequest& ResolveRequest, FHttpResponsePtr HttpResponse);
-
-	TMap<FString, TUniquePtr<FResolvedEndpoint>> ResolvedEndpoints;
-	TMap<FString, TUniquePtr<FResolveRequest>> PendingRequests;
-	FRWLock Lock;
-	bool bInitialized = false;
-};
-
-FDistributionEndpoints::~FDistributionEndpoints()
-{
-	CancelEndpointRequests();
-}
-
-void FDistributionEndpoints::ResolveEndpoints(const FString& DistributionUrl, FOnEndpointResolved&& OnResolved)
-{
-	TRACE_CPUPROFILER_EVENT_SCOPE(FOnDemandIoBackend::ResolveEndpoints);
-	const FResolvedEndpoint* Ep = nullptr;
-	{
-		FReadScopeLock _(Lock);
-		if (TUniquePtr<FResolvedEndpoint>* Entry = ResolvedEndpoints.Find(DistributionUrl))
-		{
-			Ep = Entry->Get();
-		}
-	}
-
-	if (Ep)
-	{
-		return OnResolved(DistributionUrl, Ep->ServiceUrls);
-	}
-
-	bool bIssueRequest = false;
-	{
-		FWriteScopeLock _(Lock);
-		TUniquePtr<FResolveRequest>& Request = PendingRequests.FindOrAdd(DistributionUrl);
-		if (!Request.IsValid())
-		{
-			Request.Reset(new FResolveRequest{DistributionUrl});
-			bIssueRequest = bInitialized;
-		}
-		Request->Callbacks.Add(MoveTemp(OnResolved));
-	}
-
-	if (bIssueRequest)
-	{
-		IssueEndpointRequests();
-	}
-}
-
-void FDistributionEndpoints::ResolveDeferredEndpoints()
-{
-	TRACE_CPUPROFILER_EVENT_SCOPE(FOnDemandIoBackend::ResolveDeferredEndpoints);
-	{
-		FWriteScopeLock _(Lock);
-		bInitialized = true;
-	}
-
-	IssueEndpointRequests();
-}
-
-void FDistributionEndpoints::IssueEndpointRequests()
-{
-	TRACE_CPUPROFILER_EVENT_SCOPE(FOnDemandIoBackend::IssueEndpointRequests);
-	// Currently we need to use the HTTP module in order to resolve service endpoints due to HTTPS
-	FHttpModule& HttpModule = FModuleManager::LoadModuleChecked<FHttpModule>("HTTP");
-	const int32 MaxAttempts = GIoDispatcherMaxHttpRetryCount;
-
-	TArray<FHttpRequestPtr, TInlineAllocator<2>> HttpRequests;
-	{
-		FWriteScopeLock _(Lock);
-		check(bInitialized);
-
-		for (auto& Kv : PendingRequests)
-		{
-			if (Kv.Value->HttpRequest.IsValid())
-			{
-				continue;
-			}
-
-			FResolveRequest& ResolveRequest = *Kv.Value.Get();
-			UE_LOG(LogIas, Log, TEXT("Resolving '%s' (#%d/%d)"), *ResolveRequest.DistributionUrl, ResolveRequest.RetryCount + 1, MaxAttempts);
-
-			FHttpRequestPtr HttpRequest = HttpModule.Get().CreateRequest();
-			HttpRequest->SetTimeout(3.0f);
-			HttpRequest->SetURL(Kv.Key);
-			HttpRequest->SetVerb(TEXT("GET"));
-			HttpRequest->SetHeader(TEXT("Accept"), TEXT("application/json"));
-			HttpRequest->OnProcessRequestComplete().BindLambda(
-				[this, &ResolveRequest, MaxAttempts]
-				(FHttpRequestPtr, FHttpResponsePtr Response, bool bOk)
-				{
-					LLM_SCOPE_BYTAG(Ias);
-					FHttpRequestPtr Request = MoveTemp(ResolveRequest.HttpRequest);
-					if (Response->GetResponseCode() != 200)
-					{
-						if (++ResolveRequest.RetryCount <= MaxAttempts)
-						{
-							Request->OnProcessRequestComplete().Unbind();
-							return IssueEndpointRequests();
-						}
-					}
-
-					CompleteEndpointRequest(ResolveRequest, Response);
-				});
-
-			ResolveRequest.HttpRequest = HttpRequest;
-			HttpRequests.Add(HttpRequest);
-		}
-	}
-
-	for (FHttpRequestPtr& Request : HttpRequests)
-	{
-		Request->ProcessRequest();
-	}
-}
-
-void FDistributionEndpoints::CancelEndpointRequests()
-{
-	TArray<FHttpRequestPtr, TInlineAllocator<2>> HttpRequests;
-	{
-		FWriteScopeLock _(Lock);
-		for (auto& Kv : PendingRequests)
-		{
-			if (Kv.Value->HttpRequest.IsValid())
-			{
-				HttpRequests.Add(Kv.Value->HttpRequest);
-			}
-		}
-	}
-
-	if (!HttpRequests.IsEmpty())
-	{
-		FHttpModule& HttpModule = FModuleManager::LoadModuleChecked<FHttpModule>("HTTP");
-		for (FHttpRequestPtr& Request : HttpRequests)
-		{
-			HttpModule.GetHttpManager().RemoveRequest(Request.ToSharedRef());
-			//TODO: Flush?
-		}
-	}
-}
-
-void FDistributionEndpoints::CompleteEndpointRequest(FResolveRequest& ResolveRequest, FHttpResponsePtr HttpResponse)
-{
-	TRACE_CPUPROFILER_EVENT_SCOPE(FOnDemandIoBackend::CompleteEndpointRequest);
-
-	using FJsonValuePtr = TSharedPtr<FJsonValue>;
-	using FJsonObjPtr = TSharedPtr<FJsonObject>;
-	using FJsonReader = TJsonReader<TCHAR>;
-	using FJsonReaderPtr = TSharedRef<FJsonReader>;
-
-	TArray<FString> ServiceUrls;
-	if (HttpResponse->GetResponseCode() == 200)
-	{
-		FString Json = HttpResponse->GetContentAsString();
-		FJsonReaderPtr JsonReader = TJsonReaderFactory<TCHAR>::Create(Json);
-
-		FJsonObjPtr JsonObj;
-		if (FJsonSerializer::Deserialize(JsonReader, JsonObj))
-		{
-			TArray<FJsonValuePtr> JsonValues = JsonObj->GetArrayField(TEXT("distributions"));
-			for (const FJsonValuePtr& JsonValue : JsonValues)
-			{
-				FString ServiceUrl = JsonValue->AsString();
-				if (ServiceUrl.EndsWith(TEXT("/")))
-				{
-					ServiceUrl.LeftInline(ServiceUrl.Len() - 1);
-				}
-				ServiceUrls.Add(MoveTemp(ServiceUrl));
-			}
-		}
-	}
-
-	const FResolvedEndpoint* ResolvedEndpoint = nullptr;
-	FString DistributionUrl;
-	TArray<FOnEndpointResolved> Callbacks;
-
-	{
-		FWriteScopeLock _(Lock);
-		if (!ServiceUrls.IsEmpty())
-		{
-			ResolvedEndpoint = ResolvedEndpoints.Emplace(
-				ResolveRequest.DistributionUrl,
-				new FResolvedEndpoint{MoveTemp(ServiceUrls)})
-			.Get();
-		}
-
-		Callbacks = MoveTemp(ResolveRequest.Callbacks);
-		DistributionUrl = MoveTemp(ResolveRequest.DistributionUrl);
-		PendingRequests.Remove(DistributionUrl);
-	}
-
-	TConstArrayView<FString> Urls = ResolvedEndpoint ? ResolvedEndpoint->ServiceUrls : TConstArrayView<FString>();
-	for (FOnEndpointResolved& Callback : Callbacks)
-	{
-		Callback(DistributionUrl, Urls);
-	}
-}
-
 ///////////////////////////////////////////////////////////////////////////////
 class FHttpClient
 {
@@ -1497,7 +1269,7 @@ FIoStatus FOnDemandIoBackend::MountDeferredEndpoints(const FString& Distribution
 		FIoStatus Status;
 		for (const FString& SerivceUrl : ServiceUrls)
 		{
-			// Currenlty we don't need use secure sockets to fetch on demand content
+			// Currently we don't need use secure sockets to fetch on demand content
 			FString UnsecureUrl = SerivceUrl.Replace(TEXT("https"), TEXT("http"));
 			Status = AddToc(FOnDemandEndpoint{Ep.EndpointType, DistributionUrl, UnsecureUrl, Ep.TocPath});
 			if (Status.IsOk())
@@ -1576,9 +1348,7 @@ void FOnDemandIoBackend::ReportAnalytics(TArray<FAnalyticsEventAttribute>& OutAn
 
 	if (bEnabled)
 	{
-		AppendAnalyticsEventAttributeArray(OutAnalyticsArray,
-			TEXT("IasCDNBackend"), HttpClient->ServiceUrl()
-		);
+		AppendAnalyticsEventAttributeArray(OutAnalyticsArray, TEXT("IasCDNBackend"), HttpClient->ServiceUrl());
 
 		Stats.ReportAnalytics(OutAnalyticsArray);
 	}
@@ -1587,24 +1357,7 @@ void FOnDemandIoBackend::ReportAnalytics(TArray<FAnalyticsEventAttribute>& OutAn
 #if IS_PROGRAM || WITH_EDITOR
 bool FOnDemandIoBackend::FlushDeferedEndPoints(double TimeOut)
 {
-	TRACE_CPUPROFILER_EVENT_SCOPE(FOnDemandIoBackend::FlushDeferedEndPoints);
-
-	FHttpManager& HttpManager = FHttpModule::Get().GetHttpManager();
-
-	const double StartTime = FPlatformTime::Seconds();
-
-	while (!DeferredEndpoints.IsEmpty())
-	{
-		HttpManager.Tick(0.0);
-		FPlatformProcess::SleepNoStats(0.0f);
-
-		if (TimeOut > 0.0 && (FPlatformTime::Seconds() - StartTime) > TimeOut)
-		{
-			return false;
-		}
-	}
-
-	return true;
+	return DistributionEndpoints.Flush(TimeOut);
 }
 
 TArray<FIoChunkId> FOnDemandIoBackend::GetAllChunkIds()
