Iris
* Fix flushed objects still getting copied after destruction

#rnx

diff --git a/Engine/Source/Runtime/Experimental/Iris/Core/Private/Iris/ReplicationSystem/ReplicationBridge.cpp b/Engine/Source/Runtime/Experimental/Iris/Core/Private/Iris/ReplicationSystem/ReplicationBridge.cpp
--- a/Engine/Source/Runtime/Experimental/Iris/Core/Private/Iris/ReplicationSystem/ReplicationBridge.cpp
+++ b/Engine/Source/Runtime/Experimental/Iris/Core/Private/Iris/ReplicationSystem/ReplicationBridge.cpp
@@ -685,6 +685,11 @@ void UReplicationBridge::InternalFlushStateData(UE::Net::FNetSerializationContex
 	} 
 
 	FReplicationInstanceOperationsInternal::CopyObjectStateData(ChangeMaskWriter, ChangeMaskCache, *NetRefHandleManager, SerializationContext, InternalObjectIndex);
+
+	// Clear the copy flag since it was done directly here.
+	NetRefHandleManager->GetDirtyObjectsToCopy().ClearBit(InternalObjectIndex);
+
+	// $IRIS TODO:  Should we also clear the DirtyTracker flags for this flushed object ?
 }
 
 void UReplicationBridge::InternalFlushStateData(FNetRefHandle Handle)
