Adding IAS enabled state to the crash report context.

diff --git a/Engine/Source/Runtime/Experimental/IoStoreOnDemand/Private/OnDemandIoDispatcherBackend.cpp b/Engine/Source/Runtime/Experimental/IoStoreOnDemand/Private/OnDemandIoDispatcherBackend.cpp
--- a/Engine/Source/Runtime/Experimental/IoStoreOnDemand/Private/OnDemandIoDispatcherBackend.cpp
+++ b/Engine/Source/Runtime/Experimental/IoStoreOnDemand/Private/OnDemandIoDispatcherBackend.cpp
@@ -9,6 +9,7 @@
 #include "DistributionEndpoints.h"
 #include "EncryptionKeyManager.h"
 #include "IasCache.h"
+#include "GenericPlatform/GenericPlatformCrashContext.h"
 #include "HAL/Event.h"
 #include "HAL/LowLevelMemTracker.h"
 #include "HAL/Platform.h"
@@ -1025,6 +1026,7 @@ void FOnDemandIoBackend::Initialize(TSharedRef<const FIoDispatcherBackendContext
 	TRACE_CPUPROFILER_EVENT_SCOPE(FOnDemandIoBackend::Initialize);
 	LLM_SCOPE_BYTAG(Ias);
 	UE_LOG(LogIas, Log, TEXT("Initializing on demand I/O dispatcher backend"));
+	FGenericCrashContext::SetEngineData(TEXT("IAS.Enabled"), bEnabled ? TEXT("true") : TEXT("false"));
 	BackendContext = Context;
 	DistributionEndpoints.ResolveDeferredEndpoints();
 }
@@ -1337,6 +1339,7 @@ void FOnDemandIoBackend::SetEnabled(bool bInEnabled)
 {
 	UE_LOG(LogIas, Log, TEXT("HTTP streaming '%s'"), bInEnabled ? TEXT("Enabled") : TEXT("Disabled"));
 	bEnabled = bInEnabled;
+	FGenericCrashContext::SetEngineData(TEXT("IAS.Enabled"), bEnabled ? TEXT("true") : TEXT("false"));
 }
 
 void FOnDemandIoBackend::ReportAnalytics(TArray<FAnalyticsEventAttribute>& OutAnalyticsArray) const
