Chaos - fix collision test mode to work in PIE
- add post-integrate debug draw of shapes for better collision debugging

diff --git a/Engine/Source/Runtime/Experimental/Chaos/Private/Chaos/PBDRigidsEvolution.cpp b/Engine/Source/Runtime/Experimental/Chaos/Private/Chaos/PBDRigidsEvolution.cpp
--- a/Engine/Source/Runtime/Experimental/Chaos/Private/Chaos/PBDRigidsEvolution.cpp
+++ b/Engine/Source/Runtime/Experimental/Chaos/Private/Chaos/PBDRigidsEvolution.cpp
@@ -24,6 +24,11 @@ FAutoConsoleVariableRef CVarChaosNumContactIterationsOverride(TEXT("p.ChaosNumCo
 
 namespace Chaos
 {
+	namespace CVars
+	{
+		extern bool bChaos_Solver_TestMode;
+	}
+
 	CHAOS_API int32 FixBadAccelerationStructureRemoval = 1;
 	FAutoConsoleVariableRef CVarFixBadAccelerationStructureRemoval(TEXT("p.FixBadAccelerationStructureRemoval"), FixBadAccelerationStructureRemoval, TEXT(""));
 
@@ -1146,4 +1151,95 @@ namespace Chaos
 	{
 		return Private::GetFirstPhysicsMaterial(Particle, &PhysicsMaterials, &PerParticlePhysicsMaterials, &SolverPhysicsMaterials);
 	}
+
+#if CHAOS_EVOLUTION_COLLISION_TESTMODE
+	void FPBDRigidsEvolutionBase::TestModeParticleDisabled(FGeometryParticleHandle* Particle)
+	{
+		if (FPBDRigidParticleHandle* Rigid = Particle->CastToRigidParticle())
+		{
+			TestModeData.Remove(Rigid);
+		}
+	}
+
+	void FPBDRigidsEvolutionBase::TestModeSaveParticles()
+	{
+		if (!CVars::bChaos_Solver_TestMode)
+		{
+			return;
+		}
+
+		for (auto& Rigid : Particles.GetNonDisabledDynamicView())
+		{
+			TestModeSaveParticle(Rigid.Handle());
+		}
+	}
+
+	void FPBDRigidsEvolutionBase::TestModeSaveParticle(FGeometryParticleHandle* Particle)
+	{
+		if (!CVars::bChaos_Solver_TestMode)
+		{
+			return;
+		}
+
+		if (FPBDRigidParticleHandle* Rigid = Particle->CastToRigidParticle())
+		{
+			FTestModeParticleData* Data = TestModeData.Find(Rigid);
+			if (Data == nullptr)
+			{
+				Data = &TestModeData.Add(Rigid);
+			}
+			Data->X = Rigid->X();
+			Data->P = Rigid->P();
+			Data->R = Rigid->R();
+			Data->Q = Rigid->Q();
+			Data->V = Rigid->V();
+			Data->W = Rigid->W();
+		}
+	}
+
+	void FPBDRigidsEvolutionBase::TestModeRestoreParticles()
+	{
+		if (!CVars::bChaos_Solver_TestMode)
+		{
+			return;
+		}
+
+		for (auto& Rigid : Particles.GetNonDisabledDynamicView())
+		{
+			// If this is the first time we have seen this particle, save its data, otherwise restore it from the cache
+			FTestModeParticleData* Data = TestModeData.Find(Rigid.Handle());
+			if (Data == nullptr)
+			{
+				TestModeSaveParticle(Rigid.Handle());
+			}
+			else
+			{
+				TestModeRestoreParticle(Rigid.Handle());
+			}
+		}
+	}
+
+	void FPBDRigidsEvolutionBase::TestModeRestoreParticle(FGeometryParticleHandle* Particle)
+	{
+		if (!CVars::bChaos_Solver_TestMode)
+		{
+			return;
+		}
+
+		if (FPBDRigidParticleHandle* Rigid = Particle->CastToRigidParticle())
+		{
+			FTestModeParticleData* Data = TestModeData.Find(Rigid);
+			if (Data != nullptr)
+			{
+				Rigid->X() = Data->X;
+				Rigid->P() = Data->P;
+				Rigid->R() = Data->R;
+				Rigid->Q() = Data->Q;
+				Rigid->V() = Data->V;
+				Rigid->W() = Data->W;
+			}
+		}
+	}
+#endif
+
 }
diff --git a/Engine/Source/Runtime/Experimental/Chaos/Private/Chaos/PBDRigidsEvolutionGBF.cpp b/Engine/Source/Runtime/Experimental/Chaos/Private/Chaos/PBDRigidsEvolutionGBF.cpp
--- a/Engine/Source/Runtime/Experimental/Chaos/Private/Chaos/PBDRigidsEvolutionGBF.cpp
+++ b/Engine/Source/Runtime/Experimental/Chaos/Private/Chaos/PBDRigidsEvolutionGBF.cpp
@@ -159,7 +159,8 @@ namespace Chaos
 
 		// Put the solver into a mode where it reset particles to their initial positions each frame.
 		// This is used to test collision detection and - it will be removed
-		// @chaos(todo): remove this when no longer needed
+		// NOTE: You should also set the following for dragging in PIE to work while test mode is active:
+		// 		p.DisableEditorPhysicsHandle 1
 		bool bChaos_Solver_TestMode  = false;
 		FAutoConsoleVariableRef CVarChaosSolverTestMode(TEXT("p.Chaos.Solver.TestMode"), bChaos_Solver_TestMode, TEXT(""));
 
@@ -400,6 +401,12 @@ void FPBDRigidsEvolutionGBF::AdvanceOneTimeStepImpl(const FReal Dt, const FSubSt
 	// Update the collision solver type (used to support runtime comparisons of solver types for debugging/testing)
 	UpdateCollisionSolverType();
 
+#if CHAOS_EVOLUTION_COLLISION_TESTMODE
+	{
+		TestModeSaveParticles();
+	}
+#endif
+
 	{
 		CVD_SCOPE_TRACE_SOLVER_STEP(TEXT("Integrate"))
 		SCOPE_CYCLE_COUNTER(STAT_Evolution_Integrate);
@@ -407,11 +414,12 @@ void FPBDRigidsEvolutionGBF::AdvanceOneTimeStepImpl(const FReal Dt, const FSubSt
 		Integrate(Particles.GetActiveParticlesView(), Dt);
 	}
 
-	if (bChaos_Solver_TestMode)
+#if CHAOS_EVOLUTION_COLLISION_TESTMODE
 	{
-		TestModeResetParticles();
+		TestModeRestoreParticles();
 		TestModeResetCollisions();
 	}
+#endif
 
 	{
 		CVD_SCOPE_TRACE_SOLVER_STEP(TEXT("ApplyKinematicTargets"))
@@ -621,10 +629,11 @@ void FPBDRigidsEvolutionGBF::AdvanceOneTimeStepImpl(const FReal Dt, const FSubSt
 		GetIslandManager().EndTick();
 	}
 
-	if (bChaos_Solver_TestMode)
+#if CHAOS_EVOLUTION_COLLISION_TESTMODE
 	{
-		TestModeResetParticles();
+		TestModeRestoreParticles();
 	}
+#endif
 
 	if (CVars::DoFinalProbeNarrowPhase)
 	{
@@ -677,45 +686,6 @@ void FPBDRigidsEvolutionGBF::SetShockPropagationIterations(const int32 InPositio
 	IslandManager.SetAssignLevels(CollisionConstraints.IsShockPropagationEnabled());
 }
 
-void FPBDRigidsEvolutionGBF::TestModeResetParticles()
-{
-	for (auto& Rigid : Particles.GetNonDisabledDynamicView())
-	{
-		FTestModeParticleData* Data = TestModeData.Find(Rigid.Handle());
-		if (Data != nullptr)
-		{
-			Rigid.X() = Data->X;
-			Rigid.P() = Data->P;
-			Rigid.R() = Data->R;
-			Rigid.Q() = Data->Q;
-			Rigid.V() = Data->V;
-			Rigid.W() = Data->W;
-		}
-		if (Data == nullptr)
-		{
-			Data = &TestModeData.Add(Rigid.Handle());
-			Data->X = Rigid.X();
-			Data->P = Rigid.P();
-			Data->R = Rigid.R();
-			Data->Q = Rigid.Q();
-			Data->V = Rigid.V();
-			Data->W = Rigid.W();
-		}
-	}
-}
-
-void FPBDRigidsEvolutionGBF::TestModeResetCollisions()
-{
-	for (FPBDCollisionConstraintHandle* Collision : CollisionConstraints.GetConstraintHandles())
-	{
-		if (Collision != nullptr)
-		{
-			Collision->GetContact().ResetManifold();
-			Collision->GetContact().ResetModifications();
-			Collision->GetContact().GetGJKWarmStartData().Reset();
-		}
-	}
-}
 
 void FPBDRigidsEvolutionGBF::ResetCollisions()
 {
@@ -1019,6 +989,13 @@ CHAOS_API void FPBDRigidsEvolutionGBF::SetParticleTransform(FGeometryParticleHan
 	FGenericParticleHandle(InParticle)->SetTransform(InPos, InRot);
 
 	OnParticleMoved(InParticle, PrevX, PrevR, bIsTeleport);
+
+#if CHAOS_EVOLUTION_COLLISION_TESTMODE
+	{
+		// Update the test mode cache so we can move particles in PIE to test collisions
+		TestModeSaveParticle(InParticle);
+	}
+#endif
 }
 
 
@@ -1192,5 +1169,19 @@ void FPBDRigidsEvolutionGBF::UpdateInertiaConditioning()
 	}
 }
 
+#if CHAOS_EVOLUTION_COLLISION_TESTMODE
+void FPBDRigidsEvolutionGBF::TestModeResetCollisions()
+{
+	for (FPBDCollisionConstraintHandle* Collision : CollisionConstraints.GetConstraintHandles())
+	{
+		if (Collision != nullptr)
+		{
+			Collision->GetContact().ResetManifold();
+			Collision->GetContact().ResetModifications();
+			Collision->GetContact().GetGJKWarmStartData().Reset();
+		}
+	}
+}
+#endif
 }
 
diff --git a/Engine/Source/Runtime/Experimental/Chaos/Private/PBDRigidsSolver.cpp b/Engine/Source/Runtime/Experimental/Chaos/Private/PBDRigidsSolver.cpp
--- a/Engine/Source/Runtime/Experimental/Chaos/Private/PBDRigidsSolver.cpp
+++ b/Engine/Source/Runtime/Experimental/Chaos/Private/PBDRigidsSolver.cpp
@@ -168,6 +168,11 @@ namespace Chaos
 		FAutoConsoleVariableRef CVarChaosSolverDebugDrawShowServer(TEXT("p.Chaos.Solver.DebugDraw.ShowServer"), ChaosSolverDebugDrawShowServer, TEXT("Draw server related debug data"));
 		FAutoConsoleVariableRef CVarChaosSolverDebugDrawShowClient(TEXT("p.Chaos.Solver.DebugDraw.ShowClient"), ChaosSolverDebugDrawShowClient, TEXT("Draw client related debug data"));
 
+		int32 ChaosSolverDebugDrawPostIntegrationShapes = 0;
+		int32 ChaosSolverDebugDrawPostIntegrationCollisions = 0;
+		FAutoConsoleVariableRef CVarChaosSolverDrawPostIntegrationShapes(TEXT("p.Chaos.Solver.DebugDrawPostIntegrationShapes"), ChaosSolverDebugDrawPostIntegrationShapes, TEXT("Draw Shapes prior to constraint solve phase."));
+		FAutoConsoleVariableRef CVarChaosSolverDrawPostIntegrationCollisions(TEXT("p.Chaos.Solver.DebugDrawPostIntegrationCollisions"), ChaosSolverDebugDrawPostIntegrationCollisions, TEXT("Draw Collisions prior to constraint solve phase."));
+
 
 		DebugDraw::FChaosDebugDrawSettings ChaosSolverDebugDebugDrawSettings(
 			/* ArrowSize =					*/ 10.0f,
@@ -300,27 +305,27 @@ namespace Chaos
 		FAutoConsoleVariableRef CVarChaosSolverDeterministic(TEXT("p.Chaos.Solver.Deterministic"), ChaosSolverDeterministic, TEXT("Override determinism. 0: disabled; 1: enabled; -1: use config"));
 
 		// Copied from RBAN
-		Chaos::FRealSingle ChaosSolverJointPositionTolerance = 0.025f;
+		FRealSingle ChaosSolverJointPositionTolerance = 0.025f;
 		FAutoConsoleVariableRef CVarChaosSolverJointPositionTolerance(TEXT("p.Chaos.Solver.Joint.PositionTolerance"), ChaosSolverJointPositionTolerance, TEXT("PositionTolerance."));
-		Chaos::FRealSingle ChaosSolverJointAngleTolerance = 0.001f;
+		FRealSingle ChaosSolverJointAngleTolerance = 0.001f;
 		FAutoConsoleVariableRef CVarChaosSolverJointAngleTolerance(TEXT("p.Chaos.Solver.Joint.AngleTolerance"), ChaosSolverJointAngleTolerance, TEXT("AngleTolerance."));
-		Chaos::FRealSingle ChaosSolverJointMinParentMassRatio = 0.2f;
+		FRealSingle ChaosSolverJointMinParentMassRatio = 0.2f;
 		FAutoConsoleVariableRef CVarChaosSolverJointMinParentMassRatio(TEXT("p.Chaos.Solver.Joint.MinParentMassRatio"), ChaosSolverJointMinParentMassRatio, TEXT("6Dof joint MinParentMassRatio (if > 0)"));
-		Chaos::FRealSingle ChaosSolverJointMaxInertiaRatio = 5.0f;
+		FRealSingle ChaosSolverJointMaxInertiaRatio = 5.0f;
 		FAutoConsoleVariableRef CVarChaosSolverJointMaxInertiaRatio(TEXT("p.Chaos.Solver.Joint.MaxInertiaRatio"), ChaosSolverJointMaxInertiaRatio, TEXT("6Dof joint MaxInertiaRatio (if > 0)"));
 
 		// Collision detection cvars
 		// These override the engine config if >= 0
-		Chaos::FRealSingle ChaosSolverCullDistance = -1.0f;
+		FRealSingle ChaosSolverCullDistance = -1.0f;
 		FAutoConsoleVariableRef CVarChaosSolverCullDistance(TEXT("p.Chaos.Solver.Collision.CullDistance"), ChaosSolverCullDistance, TEXT("Override cull distance (if >= 0)"));
 
 		// @todo(chaos): move to physics project settings and set these to -1 when we are settled on values...
-		Chaos::FRealSingle ChaosSolverVelocityBoundsMultiplier = 1.0f;
-		Chaos::FRealSingle ChaosSolverMaxVelocityBoundsExpansion = 3.0f;	// This should probably be a fraction of object size (see FParticlePairMidPhase::GenerateCollisions)
+		FRealSingle ChaosSolverVelocityBoundsMultiplier = 1.0f;
+		FRealSingle ChaosSolverMaxVelocityBoundsExpansion = 3.0f;	// This should probably be a fraction of object size (see FParticlePairMidPhase::GenerateCollisions)
 		FAutoConsoleVariableRef CVarChaosSolverVelocityBoundsMultiplier(TEXT("p.Chaos.Solver.Collision.VelocityBoundsMultiplier"), ChaosSolverVelocityBoundsMultiplier, TEXT("Override velocity bounds multiplier (if >= 0)"));
 		FAutoConsoleVariableRef CVarChaosSolverMaxVelocityBoundsExpansion(TEXT("p.Chaos.Solver.Collision.MaxVelocityBoundsExpansion"), ChaosSolverMaxVelocityBoundsExpansion, TEXT("Override max velocity bounds expansion (if >= 0)"));
 
-		Chaos::FRealSingle ChaosSolverMaxPushOutVelocity = -1.0f;
+		FRealSingle ChaosSolverMaxPushOutVelocity = -1.0f;
 		FAutoConsoleVariableRef CVarChaosSolverMaxPushOutVelocity(TEXT("p.Chaos.Solver.Collision.MaxPushOutVelocity"), ChaosSolverMaxPushOutVelocity, TEXT("Override max pushout velocity (if >= 0)"));
 
 		int32 ChaosSolverCleanupCommandsOnDestruction = 1;
@@ -404,7 +409,7 @@ namespace Chaos
 			
 			{
 				SCOPE_CYCLE_COUNTER(STAT_UpdateParams);
-				Chaos::FPBDPositionConstraints PositionTarget; // Dummy for now
+				FPBDPositionConstraints PositionTarget; // Dummy for now
 				TMap<int32, int32> TargetedParticles;
 				{
 					MSolver->FieldParameterUpdateCallback(PositionTarget, TargetedParticles);
@@ -577,10 +582,10 @@ namespace Chaos
 		, FloorHeight(0.f)
 		, bIsDeterministic(false)
 		, Particles(UniqueIndices)
-		, MEvolution(new FPBDRigidsEvolution(Particles, SimMaterials, &MidPhaseModifiers, &CCDModifiers, &StrainModifiers, &ContactModifiers, BufferingModeIn == Chaos::EMultiBufferMode::Single))
+		, MEvolution(new FPBDRigidsEvolution(Particles, SimMaterials, &MidPhaseModifiers, &CCDModifiers, &StrainModifiers, &ContactModifiers, BufferingModeIn == EMultiBufferMode::Single))
 		, MEventManager(new FEventManager(BufferingModeIn))
 		, MSolverEventFilters(new FSolverEventFilters())
-		, MDirtyParticlesBuffer(new FDirtyParticlesBuffer(BufferingModeIn, BufferingModeIn == Chaos::EMultiBufferMode::Single))
+		, MDirtyParticlesBuffer(new FDirtyParticlesBuffer(BufferingModeIn, BufferingModeIn == EMultiBufferMode::Single))
 		, MCurrentLock(new FCriticalSection())
 
 		, PerSolverField(nullptr)
@@ -589,15 +594,21 @@ namespace Chaos
 		Reset();
 
 		MEvolution->SetInternalParticleInitilizationFunction(
-			[this](const Chaos::FGeometryParticleHandle* OldParticle, Chaos::FGeometryParticleHandle* NewParticle) 
-				{
+			[this](const FGeometryParticleHandle* OldParticle, FGeometryParticleHandle* NewParticle) 
+			{
 				IPhysicsProxyBase* Proxy = const_cast<IPhysicsProxyBase*>(OldParticle->PhysicsProxy());
-				if (Chaos::FPBDRigidClusteredParticleHandle* NewClusteredParticle = NewParticle->CastToClustered())
-					{
+				if (FPBDRigidClusteredParticleHandle* NewClusteredParticle = NewParticle->CastToClustered())
+				{
 					NewClusteredParticle->AddPhysicsProxy(Proxy);
 				}
 				NewParticle->SetPhysicsProxy(Proxy);
 			});
+
+		MEvolution->SetPreApplyCallback(
+			[this]()
+			{
+				PreSolveDebugDraw();
+			});
 	}
 
 	FRealSingle MaxBoundsForTree = (FRealSingle)10000;
@@ -636,8 +647,8 @@ namespace Chaos
 
 		RigidBody_External.SetUniqueIdx(GetEvolution()->GenerateUniqueIdx());
 		TrackGTParticle_External(*Proxy->GetParticle_LowLevel());	//todo: remove this
-		//Chaos::FParticlePropertiesData& RemoteParticleData = *DirtyPropertiesManager->AccessProducerBuffer()->NewRemoteParticleProperties();
-		//Chaos::FShapeRemoteDataContainer& RemoteShapeContainer = *DirtyPropertiesManager->AccessProducerBuffer()->NewRemoteShapeContainer();
+		//FParticlePropertiesData& RemoteParticleData = *DirtyPropertiesManager->AccessProducerBuffer()->NewRemoteParticleProperties();
+		//FShapeRemoteDataContainer& RemoteShapeContainer = *DirtyPropertiesManager->AccessProducerBuffer()->NewRemoteShapeContainer();
 
 		Proxy->SetSolver(this);
 		Proxy->GetParticle_LowLevel()->SetProxy(Proxy);
@@ -684,10 +695,10 @@ namespace Chaos
 		// Remove the proxy from the GT proxy map
 		FUniqueIdx UniqueIdx = Proxy->GetGameThreadAPI().UniqueIdx();
 
-		Chaos::FIgnoreCollisionManager& CollisionManager = GetEvolution()->GetBroadPhase().GetIgnoreCollisionManager();
+		FIgnoreCollisionManager& CollisionManager = GetEvolution()->GetBroadPhase().GetIgnoreCollisionManager();
 		{
 			int32 ExternalTimestamp = GetMarshallingManager().GetExternalTimestamp_External();
-			Chaos::FIgnoreCollisionManager::FDeactivationSet& PendingMap = CollisionManager.GetPendingDeactivationsForGameThread(ExternalTimestamp);
+			FIgnoreCollisionManager::FDeactivationSet& PendingMap = CollisionManager.GetPendingDeactivationsForGameThread(ExternalTimestamp);
 			PendingMap.Add(UniqueIdx);
 		}
 
@@ -700,7 +711,7 @@ namespace Chaos
 				GetEventManager()->template ClearEvents<FCollisionEventData>(EEventType::Collision, [Proxy]
 				(FCollisionEventData& EventDataInOut)
 				{
-					Chaos::FCollisionDataArray const& CollisionData = EventDataInOut.CollisionData.AllCollisionsArray;
+					FCollisionDataArray const& CollisionData = EventDataInOut.CollisionData.AllCollisionsArray;
 					if (CollisionData.Num() > 0)
 					{
 						check(Proxy);
@@ -857,7 +868,7 @@ namespace Chaos
 		);
 	}
 
-	void FPBDRigidsSolver::RegisterObject(Chaos::FJointConstraint* GTConstraint)
+	void FPBDRigidsSolver::RegisterObject(FJointConstraint* GTConstraint)
 	{
 		LLM_SCOPE(ELLMTag::ChaosConstraint);
 		FJointConstraintPhysicsProxy* JointProxy = new FJointConstraintPhysicsProxy(GTConstraint, nullptr);
@@ -866,7 +877,7 @@ namespace Chaos
 		AddDirtyProxy(JointProxy);
 	}
 
-	void FPBDRigidsSolver::UnregisterObject(Chaos::FJointConstraint* GTConstraint)
+	void FPBDRigidsSolver::UnregisterObject(FJointConstraint* GTConstraint)
 	{
 		FJointConstraintPhysicsProxy* JointProxy = GTConstraint->GetProxy<FJointConstraintPhysicsProxy>();
 		check(JointProxy);
@@ -899,7 +910,7 @@ namespace Chaos
 		});
 	}
 
-	void FPBDRigidsSolver::RegisterObject(Chaos::FCharacterGroundConstraint* GTConstraint)
+	void FPBDRigidsSolver::RegisterObject(FCharacterGroundConstraint* GTConstraint)
 	{
 		LLM_SCOPE(ELLMTag::ChaosConstraint);
 		FCharacterGroundConstraintProxy* ConstraintProxy = new FCharacterGroundConstraintProxy(GTConstraint);
@@ -908,7 +919,7 @@ namespace Chaos
 		AddDirtyProxy(ConstraintProxy);
 	}
 
-	void FPBDRigidsSolver::UnregisterObject(Chaos::FCharacterGroundConstraint* GTConstraint)
+	void FPBDRigidsSolver::UnregisterObject(FCharacterGroundConstraint* GTConstraint)
 	{
 		FCharacterGroundConstraintProxy* ConstraintProxy = GTConstraint->GetProxy<FCharacterGroundConstraintProxy>();
 		check(ConstraintProxy);
@@ -936,7 +947,7 @@ namespace Chaos
 			});
 	}
 
-	void FPBDRigidsSolver::RegisterObject(Chaos::FSuspensionConstraint* GTConstraint)
+	void FPBDRigidsSolver::RegisterObject(FSuspensionConstraint* GTConstraint)
 	{
 		LLM_SCOPE(ELLMTag::ChaosConstraint);
 		FSuspensionConstraintPhysicsProxy* SuspensionProxy = new FSuspensionConstraintPhysicsProxy(GTConstraint, nullptr);
@@ -945,7 +956,7 @@ namespace Chaos
 		AddDirtyProxy(SuspensionProxy);
 	}
 
-	void FPBDRigidsSolver::UnregisterObject(Chaos::FSuspensionConstraint* GTConstraint)
+	void FPBDRigidsSolver::UnregisterObject(FSuspensionConstraint* GTConstraint)
 	{
 		FSuspensionConstraintPhysicsProxy* SuspensionProxy = GTConstraint->GetProxy<FSuspensionConstraintPhysicsProxy>();
 		check(SuspensionProxy);
@@ -968,7 +979,7 @@ namespace Chaos
 			});
 	}
 
-	void FPBDRigidsSolver::SetSuspensionTarget(Chaos::FSuspensionConstraint* GTConstraint, const FVector& TargetPos, const FVector& Normal, bool Enabled)
+	void FPBDRigidsSolver::SetSuspensionTarget(FSuspensionConstraint* GTConstraint, const FVector& TargetPos, const FVector& Normal, bool Enabled)
 	{
 		EnsureIsInPhysicsThreadContext();
 		FSuspensionConstraintPhysicsProxy* SuspensionProxy = GTConstraint->GetProxy<FSuspensionConstraintPhysicsProxy>();
@@ -1033,7 +1044,7 @@ namespace Chaos
 		FEventDefaults::RegisterSystemEvents(*GetEventManager());
 	}
 
-	void FPBDRigidsSolver::ChangeBufferMode(Chaos::EMultiBufferMode InBufferMode)
+	void FPBDRigidsSolver::ChangeBufferMode(EMultiBufferMode InBufferMode)
 	{
 		// This seems unused inside the solver? #BH
 		BufferMode = InBufferMode;
@@ -1915,31 +1926,31 @@ namespace Chaos
 			ensure(PullData->DirtyRigids.Num() == 0);	//we only fill this once per frame
 			int32 BufferIdx = 0;
 
-			for (Chaos::TPBDRigidParticleHandleImp<FReal, 3, false>& DirtyParticle : DirtyParticles)
+			for (TPBDRigidParticleHandleImp<FReal, 3, false>& DirtyParticle : DirtyParticles)
 			{
 				if(IPhysicsProxyBase* Proxy = DirtyParticle.Handle()->PhysicsProxy())
 				{
 					switch(DirtyParticle.GetParticleType())
 					{
-						case Chaos::EParticleType::Rigid:
+						case EParticleType::Rigid:
 						{
 							if(!bIsResim || DirtyParticle.SyncState() == ESyncState::HardDesync)
 							{
-								if (!(Chaos::SyncKinematicOnGameThread == 0 && DirtyParticle.ObjectState() == EObjectStateType::Kinematic))
+								if (!(SyncKinematicOnGameThread == 0 && DirtyParticle.ObjectState() == EObjectStateType::Kinematic))
 								{
 									ActiveRigid.AddUnique((FSingleParticlePhysicsProxy*)Proxy);
 								}
 							}
 							break;
 						}
-						case Chaos::EParticleType::Kinematic:
-						case Chaos::EParticleType::Static:
+						case EParticleType::Kinematic:
+						case EParticleType::Static:
 							ensure(false);
 							break;
-						case Chaos::EParticleType::GeometryCollection:
+						case EParticleType::GeometryCollection:
 							ActiveGC.AddUnique((FGeometryCollectionPhysicsProxy*)(Proxy));
 							break;
-						case Chaos::EParticleType::Clustered:
+						case EParticleType::Clustered:
 							if (auto ClusterParticle = DirtyParticle.CastToClustered())
 							{
 								if (ClusterParticle->InternalCluster())
@@ -2236,6 +2247,63 @@ TRACE_COUNTER_SET(ChaosTraceCounter_##Name, Value)
 		CHAOS_COUNTER_STAT(NumDynamicShapes, NumDynamicShapes);
 	}
 
+	void FPBDRigidsSolver::PreSolveDebugDraw() const
+	{
+#if CHAOS_DEBUG_DRAW
+		QUICK_SCOPE_CYCLE_COUNTER(SolverDebugDraw);
+
+#if !UE_BUILD_SHIPPING && !UE_BUILD_TEST
+		const bool bIsServer = DebugName.ToString().StartsWith(TEXT("Server"));
+		if (bIsServer && !ChaosSolverDebugDrawShowServer)
+		{
+			return;
+		}
+		if (!bIsServer && !ChaosSolverDebugDrawShowClient)
+		{
+			return;
+		}
+
+		if (ChaosSolverDebugDrawColorShapeByClientServer)
+		{
+			if (bIsServer)
+			{
+				ChaosSolverDebugDebugDrawSettings.ShapesColorsPerState = GetSolverShapesColorsByState_Server();
+			}
+			else
+			{
+				ChaosSolverDebugDebugDrawSettings.ShapesColorsPerState = GetSolverShapesColorsByState_Client();
+			}
+		}
+		else
+		{
+			ChaosSolverDebugDebugDrawSettings.ShapesColorsPerState = DebugDraw::GetDefaultShapesColorsByState();
+		}
+#endif
+
+		if (ChaosSolverDebugDrawPostIntegrationShapes == 1)
+		{
+			if (ChaosSolverDrawShapesShowStatic)
+			{
+				DebugDraw::DrawParticleShapes(FRigidTransform3(), Particles.GetActiveStaticParticlesView(), 1.0f, &ChaosSolverDebugDebugDrawSettings);
+			}
+			if (ChaosSolverDrawShapesShowKinematic)
+			{
+				DebugDraw::DrawParticleShapes(FRigidTransform3(), Particles.GetActiveKinematicParticlesView(), 1.0f, &ChaosSolverDebugDebugDrawSettings);
+			}
+			if (ChaosSolverDrawShapesShowDynamic)
+			{
+				DebugDraw::DrawParticleShapes(FRigidTransform3(), Particles.GetNonDisabledDynamicView(), 1.0f, &ChaosSolverDebugDebugDrawSettings);
+			}
+		}
+
+		if (ChaosSolverDebugDrawPostIntegrationCollisions == 1)
+		{
+			DebugDraw::DrawCollisions(FRigidTransform3(), GetEvolution()->GetCollisionConstraints().GetConstraintAllocator(), 1.f, &ChaosSolverDebugDebugDrawSettings);
+		}
+
+#endif
+	}
+
 	void FPBDRigidsSolver::PostTickDebugDraw(FReal Dt) const
 	{
 #if CHAOS_DEBUG_DRAW
@@ -2364,37 +2432,37 @@ TRACE_COUNTER_SET(ChaosTraceCounter_##Name, Value)
 		return GetParticleProxy_PT(Handle.UniqueIdx());
 	}
 
-	void FPBDRigidsSolver::UpdateMaterial(Chaos::FMaterialHandle InHandle, const Chaos::FChaosPhysicsMaterial& InNewData)
+	void FPBDRigidsSolver::UpdateMaterial(FMaterialHandle InHandle, const FChaosPhysicsMaterial& InNewData)
 	{
 		TSolverSimMaterialScope<ELockType::Write> Scope(this);
 		*SimMaterials.Get(InHandle.InnerHandle) = InNewData;
 	}
 
-	void FPBDRigidsSolver::CreateMaterial(Chaos::FMaterialHandle InHandle, const Chaos::FChaosPhysicsMaterial& InNewData)
+	void FPBDRigidsSolver::CreateMaterial(FMaterialHandle InHandle, const FChaosPhysicsMaterial& InNewData)
 	{
 		TSolverSimMaterialScope<ELockType::Write> Scope(this);
 		ensure(SimMaterials.Create(InNewData) == InHandle.InnerHandle);
 	}
 
-	void FPBDRigidsSolver::DestroyMaterial(Chaos::FMaterialHandle InHandle)
+	void FPBDRigidsSolver::DestroyMaterial(FMaterialHandle InHandle)
 	{
 		TSolverSimMaterialScope<ELockType::Write> Scope(this);
 		SimMaterials.Destroy(InHandle.InnerHandle);
 	}
 
-	void FPBDRigidsSolver::UpdateMaterialMask(Chaos::FMaterialMaskHandle InHandle, const Chaos::FChaosPhysicsMaterialMask& InNewData)
+	void FPBDRigidsSolver::UpdateMaterialMask(FMaterialMaskHandle InHandle, const FChaosPhysicsMaterialMask& InNewData)
 	{
 		TSolverSimMaterialScope<ELockType::Write> Scope(this);
 		*SimMaterialMasks.Get(InHandle.InnerHandle) = InNewData;
 	}
 
-	void FPBDRigidsSolver::CreateMaterialMask(Chaos::FMaterialMaskHandle InHandle, const Chaos::FChaosPhysicsMaterialMask& InNewData)
+	void FPBDRigidsSolver::CreateMaterialMask(FMaterialMaskHandle InHandle, const FChaosPhysicsMaterialMask& InNewData)
 	{
 		TSolverSimMaterialScope<ELockType::Write> Scope(this);
 		ensure(SimMaterialMasks.Create(InNewData) == InHandle.InnerHandle);
 	}
 
-	void FPBDRigidsSolver::DestroyMaterialMask(Chaos::FMaterialMaskHandle InHandle)
+	void FPBDRigidsSolver::DestroyMaterialMask(FMaterialMaskHandle InHandle)
 	{
 		TSolverSimMaterialScope<ELockType::Write> Scope(this);
 		SimMaterialMasks.Destroy(InHandle.InnerHandle);
@@ -2495,9 +2563,9 @@ TRACE_COUNTER_SET(ChaosTraceCounter_##Name, Value)
 	}
 
 
-	Chaos::FClusterCreationParameters::EConnectionMethod ToInternalConnectionMethod(EClusterUnionMethod InMethod)
+	FClusterCreationParameters::EConnectionMethod ToInternalConnectionMethod(EClusterUnionMethod InMethod)
 	{
-		using ETargetEnum = Chaos::FClusterCreationParameters::EConnectionMethod;
+		using ETargetEnum = FClusterCreationParameters::EConnectionMethod;
 		switch(InMethod)
 		{
 		case EClusterUnionMethod::PointImplicit:
@@ -2515,7 +2583,7 @@ TRACE_COUNTER_SET(ChaosTraceCounter_##Name, Value)
 		return ETargetEnum::None;
 	}
 
-	void Chaos::FPBDRigidsSolver::ApplyConfig(const FChaosSolverConfiguration& InConfig)
+	void FPBDRigidsSolver::ApplyConfig(const FChaosSolverConfiguration& InConfig)
 	{
 		GetEvolution()->GetRigidClustering().SetClusterConnectionFactor(InConfig.ClusterConnectionFactor);
 		GetEvolution()->GetRigidClustering().SetClusterUnionConnectionType(ToInternalConnectionMethod(InConfig.ClusterUnionConnectionType));
@@ -2537,14 +2605,14 @@ TRACE_COUNTER_SET(ChaosTraceCounter_##Name, Value)
 		EventTeardown.Broadcast();
 	}
 
-	void Chaos::FPBDRigidsSolver::FieldParameterUpdateCallback(
-		Chaos::FPBDPositionConstraints& PositionTarget,
+	void FPBDRigidsSolver::FieldParameterUpdateCallback(
+		FPBDPositionConstraints& PositionTarget,
 		TMap<int32, int32>& TargetedParticles)
 	{
 		GetPerSolverField().FieldParameterUpdateCallback(this, PositionTarget, TargetedParticles);
 	}
 
-	void Chaos::FPBDRigidsSolver::FieldForcesUpdateCallback()
+	void FPBDRigidsSolver::FieldForcesUpdateCallback()
 	{
 		GetPerSolverField().FieldForcesUpdateCallback(this);
 	}
diff --git a/Engine/Source/Runtime/Experimental/Chaos/Public/Chaos/PBDRigidsEvolution.h b/Engine/Source/Runtime/Experimental/Chaos/Public/Chaos/PBDRigidsEvolution.h
--- a/Engine/Source/Runtime/Experimental/Chaos/Public/Chaos/PBDRigidsEvolution.h
+++ b/Engine/Source/Runtime/Experimental/Chaos/Public/Chaos/PBDRigidsEvolution.h
@@ -21,6 +21,9 @@
 #include "RewindData.h"
 #include "ChaosVisualDebugger/ChaosVDContextProvider.h"
 
+// Enable support for Collision Test Mode (reset particle positions and constraints every tick for debugging)
+#define CHAOS_EVOLUTION_COLLISION_TESTMODE (!UE_BUILD_TEST && !UE_BUILD_SHIPPING)
+
 extern int32 ChaosRigidsEvolutionApplyAllowEarlyOutCVar;
 extern int32 ChaosRigidsEvolutionApplyPushoutAllowEarlyOutCVar;
 extern int32 ChaosNumPushOutIterationsOverride;
@@ -440,6 +443,10 @@ public:
 	*/
 	void DisableParticle(FGeometryParticleHandle* Particle)
 	{
+#if CHAOS_EVOLUTION_COLLISION_TESTMODE
+		TestModeParticleDisabled(Particle);
+#endif
+
 		RemoveParticleFromAccelerationStructure(*Particle);
 		Particles.DisableParticle(Particle);
 		DisableConstraints(Particle);
@@ -543,6 +550,10 @@ public:
 
 	void DestroyParticle(FGeometryParticleHandle* Particle)
 	{
+#if CHAOS_EVOLUTION_COLLISION_TESTMODE
+		TestModeParticleDisabled(Particle);
+#endif
+
 		if (MRewindData)
 		{
 			MRewindData->RemoveObject(Particle);
@@ -1101,6 +1112,22 @@ protected:
 	TArray<FUniqueIdx> PendingReleaseIndices;	//for now just assume a one frame delay, but may need something more general
 	bool bIsResim = false; 
 	bool bIsReset = false;
+
+#if CHAOS_EVOLUTION_COLLISION_TESTMODE
+	// Test Mode for Collision issues (resets particle positions every tick for repeatable testing)
+	CHAOS_API void TestModeParticleDisabled(FGeometryParticleHandle* Particle);
+	CHAOS_API void TestModeSaveParticles();
+	CHAOS_API void TestModeSaveParticle(FGeometryParticleHandle* Particle);
+	CHAOS_API void TestModeRestoreParticles();
+	CHAOS_API void TestModeRestoreParticle(FGeometryParticleHandle* Particle);
+
+	struct FTestModeParticleData
+	{
+		FVec3 X, P, V, W;
+		FRotation3 R, Q;
+	};
+	TMap<FPBDRigidParticleHandle*, FTestModeParticleData> TestModeData;
+#endif
 };
 
 
diff --git a/Engine/Source/Runtime/Experimental/Chaos/Public/Chaos/PBDRigidsEvolutionGBF.h b/Engine/Source/Runtime/Experimental/Chaos/Public/Chaos/PBDRigidsEvolutionGBF.h
--- a/Engine/Source/Runtime/Experimental/Chaos/Public/Chaos/PBDRigidsEvolutionGBF.h
+++ b/Engine/Source/Runtime/Experimental/Chaos/Public/Chaos/PBDRigidsEvolutionGBF.h
@@ -412,18 +412,9 @@ namespace Chaos
 
 		bool bIsDeterministic;
 
-		// Test Mode for Collision issues
-		// @todo(chaos): remove this when no longer needed
-		void TestModeResetParticles();
+#if CHAOS_EVOLUTION_COLLISION_TESTMODE
 		void TestModeResetCollisions();
-
-		struct FTestModeParticleData
-		{
-			FVec3 X, P, V, W;
-			FRotation3 R, Q;
-		};
-		TMap<FPBDRigidParticleHandle*, FTestModeParticleData> TestModeData;
-		// End Test Mode
+#endif
 	};
 
 }
diff --git a/Engine/Source/Runtime/Experimental/Chaos/Public/PBDRigidsSolver.h b/Engine/Source/Runtime/Experimental/Chaos/Public/PBDRigidsSolver.h
--- a/Engine/Source/Runtime/Experimental/Chaos/Public/PBDRigidsSolver.h
+++ b/Engine/Source/Runtime/Experimental/Chaos/Public/PBDRigidsSolver.h
@@ -242,6 +242,7 @@ namespace Chaos
 		CHAOS_API void SyncEvents_GameThread();
 
 		/**/
+		CHAOS_API void PreSolveDebugDraw() const;
 		CHAOS_API void PostTickDebugDraw(FReal Dt) const;
 
 		// Visual debugger (VDB) push methods
