Fix MTD result initialization in PhysicsObjectOverlapWithMTD

diff --git a/Engine/Source/Runtime/Experimental/Chaos/Private/Chaos/PhysicsObjectCollisionInterface.cpp b/Engine/Source/Runtime/Experimental/Chaos/Private/Chaos/PhysicsObjectCollisionInterface.cpp
--- a/Engine/Source/Runtime/Experimental/Chaos/Private/Chaos/PhysicsObjectCollisionInterface.cpp
+++ b/Engine/Source/Runtime/Experimental/Chaos/Private/Chaos/PhysicsObjectCollisionInterface.cpp
@@ -25,7 +25,10 @@ namespace Chaos
 			FVector::Zero(),
 			[&OutMTD](const FShapeOverlapData&, const FShapeOverlapData&, const FMTDInfo& MTDInfo)
 			{
-				OutMTD.Penetration = FMath::Max(OutMTD.Penetration, MTDInfo.Penetration);
+				if (MTDInfo.Penetration > OutMTD.Penetration)
+				{
+					OutMTD = MTDInfo;
+				}
 				return true;
 			}
 		);
