geometry collection and cluster unions now initialize their kinematic state on the GT if needed

diff --git a/Engine/Source/Runtime/Engine/Private/PhysicsEngine/ClusterUnionComponent.cpp b/Engine/Source/Runtime/Engine/Private/PhysicsEngine/ClusterUnionComponent.cpp
--- a/Engine/Source/Runtime/Engine/Private/PhysicsEngine/ClusterUnionComponent.cpp
+++ b/Engine/Source/Runtime/Engine/Private/PhysicsEngine/ClusterUnionComponent.cpp
@@ -95,7 +95,7 @@ void UClusterUnionComponent::AddComponentToCluster(UPrimitiveComponent* InCompon
 	PendingComponentsToAdd.Remove(InComponent);
 
 	TArray<Chaos::FPhysicsObjectHandle> AllObjects = InComponent->GetAllPhysicsObjects();
-	FLockedReadPhysicsObjectExternalInterface Interface = FPhysicsObjectExternalInterface::LockRead(AllObjects);
+	FLockedWritePhysicsObjectExternalInterface Interface = FPhysicsObjectExternalInterface::LockWrite(AllObjects);
 
 	TArray<Chaos::FPhysicsObjectHandle> Objects;
 	if (BoneIds.IsEmpty())
@@ -138,6 +138,8 @@ void UClusterUnionComponent::AddComponentToCluster(UPrimitiveComponent* InCompon
 	ChildToParents.Reserve(Objects.Num());
 
 	const FTransform CurrentComponentTransform = GetComponentTransform();
+
+	const bool bHasKinematicObject = !Interface->AreAllDynamicOrSleeping(Objects);
 	for (Chaos::FPhysicsObjectHandle Object : Objects)
 	{
 		if (Object)
@@ -162,6 +164,13 @@ void UClusterUnionComponent::AddComponentToCluster(UPrimitiveComponent* InCompon
 		}
 	}
 
+	// The first time we add a component, we need to manage the GT-side object state and set it to kinematic immediately.
+	if (bHasKinematicObject && PendingComponentSync.IsEmpty() && PerComponentData.IsEmpty())
+	{
+		Chaos::FPhysicsObjectHandle SelfHandle = PhysicsProxy->GetPhysicsObjectHandle();
+		Interface->ForceKinematic( { &SelfHandle, 1 });
+	}
+
 	// Need to listen to changes in the component's physics state. If it gets destroyed it should be removed from the cluster union as well.
 	// This technically is giving us a false sense of security because when we get this notification, the primitive component will already have no physics proxy on it so we can't actually grab physics objects/particles from it.
 	InComponent->OnComponentPhysicsStateChanged.AddUniqueDynamic(this, &UClusterUnionComponent::HandleComponentPhysicsStateChangePostAddIntoClusterUnion);
diff --git a/Engine/Source/Runtime/Experimental/Chaos/Private/Chaos/PhysicsObjectInterface.cpp b/Engine/Source/Runtime/Experimental/Chaos/Private/Chaos/PhysicsObjectInterface.cpp
--- a/Engine/Source/Runtime/Experimental/Chaos/Private/Chaos/PhysicsObjectInterface.cpp
+++ b/Engine/Source/Runtime/Experimental/Chaos/Private/Chaos/PhysicsObjectInterface.cpp
@@ -767,6 +767,27 @@ namespace Chaos
 		}
 	}
 
+	template<EThreadContext Id>
+	void FWritePhysicsObjectInterface<Id>::ForceKinematic(TArrayView<const FPhysicsObjectHandle> InObjects)
+	{
+		for (const FPhysicsObjectHandle Object : InObjects)
+		{
+			if (!Object)
+			{
+				continue;
+			}
+
+			if (TThreadParticle<Id>* Particle = Object->GetParticle<Id>())
+			{
+				EObjectStateType State = Object->ObjectState<Id>();
+				if (State != EObjectStateType::Kinematic)
+				{
+					SetParticleStateHelper<Id>(Object, EObjectStateType::Kinematic);
+				}
+			}
+		}
+	}
+
 	template<EThreadContext Id>
 	void FWritePhysicsObjectInterface<Id>::AddForce(TArrayView<const FPhysicsObjectHandle> InObjects, const FVector& Force, bool bInvalidate)
 	{
diff --git a/Engine/Source/Runtime/Experimental/Chaos/Private/PhysicsProxy/GeometryCollectionPhysicsProxy.cpp b/Engine/Source/Runtime/Experimental/Chaos/Private/PhysicsProxy/GeometryCollectionPhysicsProxy.cpp
--- a/Engine/Source/Runtime/Experimental/Chaos/Private/PhysicsProxy/GeometryCollectionPhysicsProxy.cpp
+++ b/Engine/Source/Runtime/Experimental/Chaos/Private/PhysicsProxy/GeometryCollectionPhysicsProxy.cpp
@@ -541,6 +541,11 @@ void FGeometryCollectionPhysicsProxy::Initialize(Chaos::FPBDRigidsEvolutionBase
 	}
 		
 	GameThreadCollection.AddExternalAttribute<TUniquePtr<FParticle>>(FGeometryCollection::ParticlesAttribute, FTransformCollection::TransformGroup, GTParticles);
+
+	// make sure we copy the anchored information over to the physics thread collection
+	const Chaos::Facades::FCollectionAnchoringFacade DynamicCollectionAnchoringFacade(DynamicCollection);
+	Chaos::Facades::FCollectionAnchoringFacade PhysicsThreadCollectionAnchoringFacade(PhysicsThreadCollection);
+	PhysicsThreadCollectionAnchoringFacade.CopyAnchoredAttribute(DynamicCollectionAnchoringFacade);
 	
 	const FVector Scale = Parameters.WorldTransform.GetScale3D();
 	const TManagedArray<float>& Mass = Parameters.RestCollection->GetAttribute<float>("Mass", FTransformCollection::TransformGroup);
@@ -579,6 +584,10 @@ void FGeometryCollectionPhysicsProxy::Initialize(Chaos::FPBDRigidsEvolutionBase
 				}
 				P->SetGeometry(ImplicitGeometry);
 
+				if (DynamicCollectionAnchoringFacade.IsAnchored(Index))
+				{
+					P->SetObjectState(Chaos::EObjectStateType::Kinematic, false, false);
+				}
 
 				// IMPORTANT: we need to set the right spatial index because GT particle is static and PT particle is rigid
 				// this is causing a mismatch when using the separate acceleration structures optimization which can cause crashes when destroying the particle while async tracing 
@@ -715,11 +724,6 @@ void FGeometryCollectionPhysicsProxy::Initialize(Chaos::FPBDRigidsEvolutionBase
 	PhysicsThreadCollection.CopyMatchingAttributesFrom(DynamicCollection, &SkipList);
 	PhysicsThreadCollection.CopyInitialVelocityAttributesFrom(DynamicCollection);
 
-	// make sure we copy the anchored information over to the physics thread collection
-	const Chaos::Facades::FCollectionAnchoringFacade DynamicCollectionAnchoringFacade(DynamicCollection);
-	Chaos::Facades::FCollectionAnchoringFacade PhysicsThreadCollectionAnchoringFacade(PhysicsThreadCollection);
-	PhysicsThreadCollectionAnchoringFacade.CopyAnchoredAttribute(DynamicCollectionAnchoringFacade);
-	
 	// Copy simplicials.
 	// TODO: Ryan - Should we just transfer ownership of the SimplicialsAttribute from the DynamicCollection to
 	// the PhysicsThreadCollection?
diff --git a/Engine/Source/Runtime/Experimental/Chaos/Public/Chaos/PhysicsObjectInterface.h b/Engine/Source/Runtime/Experimental/Chaos/Public/Chaos/PhysicsObjectInterface.h
--- a/Engine/Source/Runtime/Experimental/Chaos/Public/Chaos/PhysicsObjectInterface.h
+++ b/Engine/Source/Runtime/Experimental/Chaos/Public/Chaos/PhysicsObjectInterface.h
@@ -121,6 +121,7 @@ namespace Chaos
 	public:
 		CHAOS_API void PutToSleep(TArrayView<const FPhysicsObjectHandle> InObjects);
 		CHAOS_API void WakeUp(TArrayView<const FPhysicsObjectHandle> InObjects);
+		CHAOS_API void ForceKinematic(TArrayView<const FPhysicsObjectHandle> InObjects);
 		CHAOS_API void AddForce(TArrayView<const FPhysicsObjectHandle> InObjects, const FVector& Force, bool bInvalidate);
 		CHAOS_API void AddTorque(TArrayView<const FPhysicsObjectHandle> InObjects, const FVector& Torque, bool bInvalidate);
 		CHAOS_API void AddRadialImpulse(TArrayView<const FPhysicsObjectHandle> InObjects, FVector Origin, float Radius, float Strength, enum ERadialImpulseFalloff Falloff, bool bApplyStrain, bool bInvalidate);
