Moved everything Ias-related into its own consistent UE::IO::IAS namespace

#rnx

diff --git a/Engine/Source/Developer/IoStoreUtilities/Private/IoStoreUtilities.cpp b/Engine/Source/Developer/IoStoreUtilities/Private/IoStoreUtilities.cpp
--- a/Engine/Source/Developer/IoStoreUtilities/Private/IoStoreUtilities.cpp
+++ b/Engine/Source/Developer/IoStoreUtilities/Private/IoStoreUtilities.cpp
@@ -99,7 +99,7 @@ static const uint64 DefaultCompressionBlockAlignment = 64 << 10;
 static const uint64 DefaultMemoryMappingAlignment = 16 << 10;
 
 static TUniquePtr<FIoStoreReader> CreateIoStoreReader(const TCHAR* Path, const FKeyChain& KeyChain);
-bool UploadIoStoreContainerFiles(const UE::FIoStoreUploadParams& UploadParams, TConstArrayView<FString> ContainerFiles, const FKeyChain& KeyChain);
+bool UploadIoStoreContainerFiles(const UE::IO::IAS::FIoStoreUploadParams& UploadParams, TConstArrayView<FString> ContainerFiles, const FKeyChain& KeyChain);
 
 class FIoStoreChunkDatabase : public IIoStoreWriterReferenceChunkDatabase
 {
@@ -5070,7 +5070,7 @@ int32 CreateTarget(const FIoStoreArguments& Arguments, const FIoStoreWriterSetti
 
 	if (Arguments.bUpload && OnDemandContainers.IsEmpty() == false)
 	{
-		TIoStatusOr<UE::FIoStoreUploadParams> UploadParams = UE::FIoStoreUploadParams::Parse(FCommandLine::Get());
+		TIoStatusOr<UE::IO::IAS::FIoStoreUploadParams> UploadParams = UE::IO::IAS::FIoStoreUploadParams::Parse(FCommandLine::Get());
 		if (UploadParams.IsOk() == false)
 		{
 			UE_LOG(LogIoStore, Warning, TEXT("Skipping upload of container file(s), reason '%s'"), *UploadParams.Status().ToString());
@@ -8622,7 +8622,7 @@ int32 CreateIoStoreContainerFiles(const TCHAR* CmdLine)
 	return CreateTarget(Arguments, WriterSettings);
 }
 
-bool UploadIoStoreContainerFiles(const UE::FIoStoreUploadParams& UploadParams, TConstArrayView<FString> ContainerFiles, const FKeyChain& KeyChain)
+bool UploadIoStoreContainerFiles(const UE::IO::IAS::FIoStoreUploadParams& UploadParams, TConstArrayView<FString> ContainerFiles, const FKeyChain& KeyChain)
 {
 	TMap<FGuid, FAES::FAESKey> EncryptionKeys;
 	for (const TPair<FGuid, FNamedAESKey>& KeyPair: KeyChain.GetEncryptionKeys())
@@ -8630,14 +8630,14 @@ bool UploadIoStoreContainerFiles(const UE::FIoStoreUploadParams& UploadParams, T
 		EncryptionKeys.Add(KeyPair.Key, KeyPair.Value.Key);
 	}
 
-	TIoStatusOr<UE::FIoStoreUploadResult> Result = UE::UploadContainerFiles(UploadParams, ContainerFiles, EncryptionKeys);
+	TIoStatusOr<UE::IO::IAS::FIoStoreUploadResult> Result = UE::IO::IAS::UploadContainerFiles(UploadParams, ContainerFiles, EncryptionKeys);
 	if (Result.IsOk() == false)
 	{
 		UE_LOG(LogIoStore, Error, TEXT("Failed to upload container file(s), reason '%s'"), *Result.Status().ToString());
 		return false;
 	}
 
-	UE::FIoStoreUploadResult UploadResult = Result.ConsumeValueOrDie();
+	UE::IO::IAS::FIoStoreUploadResult UploadResult = Result.ConsumeValueOrDie();
 
 	FString ConfigFilePath;
 	if (FParse::Value(FCommandLine::Get(), TEXT("ConfigFilePath="), ConfigFilePath))
@@ -8708,7 +8708,7 @@ bool UploadIoStoreContainerFiles(const TCHAR* ContainerPathOrWildcard)
 	FKeyChain KeyChain;
 	LoadKeyChain(FCommandLine::Get(), KeyChain);
 
-	TIoStatusOr<UE::FIoStoreUploadParams> UploadParams = UE::FIoStoreUploadParams::Parse(FCommandLine::Get());
+	TIoStatusOr<UE::IO::IAS::FIoStoreUploadParams> UploadParams = UE::IO::IAS::FIoStoreUploadParams::Parse(FCommandLine::Get());
 	if (UploadParams.IsOk() == false)
 	{
 		UE_LOG(LogIoStore, Error, TEXT("Failed to upload container file(s), reason '%s'"), *UploadParams.Status().ToString());
@@ -8762,14 +8762,14 @@ bool DownloadIoStoreContainerFiles(const TCHAR* TocPath)
 {
 	check(TocPath);
 
-	TIoStatusOr<UE::FIoStoreDownloadParams> Params = UE::FIoStoreDownloadParams::Parse(FCommandLine::Get());
+	TIoStatusOr<UE::IO::IAS::FIoStoreDownloadParams> Params = UE::IO::IAS::FIoStoreDownloadParams::Parse(FCommandLine::Get());
 	if (Params.IsOk() == false)
 	{
 		UE_LOG(LogIoStore, Error, TEXT("Failed to download container file(s), reason '%s'"), *Params.Status().ToString());
 		return false;
 	}
 
-	FIoStatus Status = UE::DownloadContainerFiles(Params.ConsumeValueOrDie(), TocPath);
+	FIoStatus Status = UE::IO::IAS::DownloadContainerFiles(Params.ConsumeValueOrDie(), TocPath);
 	if (Status.IsOk() == false)
 	{
 		UE_LOG(LogIoStore, Error, TEXT("Failed to download container file(s), reason '%s'"), *Status.ToString());
@@ -8780,7 +8780,7 @@ bool DownloadIoStoreContainerFiles(const TCHAR* TocPath)
 
 bool PrimeEndPoint(FStringView IoStoreOnDemandIniPath)
 {
-	FIoStatus Status = UE::PrimeEndPoint(IoStoreOnDemandIniPath);
+	FIoStatus Status = UE::IO::IAS::PrimeEndPoint(IoStoreOnDemandIniPath);
 	if (Status.IsOk() == false)
 	{
 		UE_LOG(LogIoStore, Error, TEXT("Failed to prime end points, reason '%s'"), *Status.ToString());
diff --git a/Engine/Source/Runtime/Experimental/IoStoreOnDemand/Private/CancellationToken.h b/Engine/Source/Runtime/Experimental/IoStoreOnDemand/Private/CancellationToken.h
--- a/Engine/Source/Runtime/Experimental/IoStoreOnDemand/Private/CancellationToken.h
+++ b/Engine/Source/Runtime/Experimental/IoStoreOnDemand/Private/CancellationToken.h
@@ -4,6 +4,8 @@
 
 #include <atomic>
 
+namespace UE::IO::IAS {
+
 /**
  * Experimental type for signaling cancellation of asynchronouse I/O operations.
  */
@@ -31,3 +33,5 @@ struct FIoCancellationToken
 private:
 	std::atomic_bool bCancellationRequested{false};
 };
+
+} // namespace UE::IO::IAS
diff --git a/Engine/Source/Runtime/Experimental/IoStoreOnDemand/Private/CoreHttp/Client.h b/Engine/Source/Runtime/Experimental/IoStoreOnDemand/Private/CoreHttp/Client.h
--- a/Engine/Source/Runtime/Experimental/IoStoreOnDemand/Private/CoreHttp/Client.h
+++ b/Engine/Source/Runtime/Experimental/IoStoreOnDemand/Private/CoreHttp/Client.h
@@ -13,7 +13,7 @@
 ////////////////////////////////////////////////////////////////////////////////
 class FIoBuffer;
 
-namespace UE::HTTP
+namespace UE::IO::IAS::HTTP
 {
 
 ////////////////////////////////////////////////////////////////////////////////
@@ -186,6 +186,6 @@ private:
 	FEventLoop&				operator = (FEventLoop&&)		= delete;
 };
 
-} // namespace UE::HTTP
+} // namespace UE::IO::IAS::HTTP
 
 /* vim: set noet : */
diff --git a/Engine/Source/Runtime/Experimental/IoStoreOnDemand/Private/CoreHttp/Http.cpp b/Engine/Source/Runtime/Experimental/IoStoreOnDemand/Private/CoreHttp/Http.cpp
--- a/Engine/Source/Runtime/Experimental/IoStoreOnDemand/Private/CoreHttp/Http.cpp
+++ b/Engine/Source/Runtime/Experimental/IoStoreOnDemand/Private/CoreHttp/Http.cpp
@@ -77,11 +77,9 @@ enum : SocketType { InvalidSocket = -1 };
 
 // }}}
 
-namespace UE::HTTP
+namespace UE::IO::IAS::HTTP
 {
 
-using FLatencyInjector = UE::IO::Private::FLatencyInjector;
-
 // {{{1 misc ...................................................................
 
 ////////////////////////////////////////////////////////////////////////////////
@@ -2334,7 +2332,7 @@ COREHTTP_API void CoreHttpTest()
 		uint64 Hash = 0;
 	} Content[64];
 
-	auto HashSink = [&] (const UE::HTTP::FTicketStatus& Status) -> FIoBuffer*
+	auto HashSink = [&] (const FTicketStatus& Status) -> FIoBuffer*
 	{
 		check(Status.GetId() != FTicketStatus::EId::Error);
 
@@ -2476,7 +2474,7 @@ COREHTTP_API void CoreHttpTest()
 		};
 
 		uint64 Errors = 0;
-		auto ErrorSink = [&] (const UE::HTTP::FTicketStatus& Status)
+		auto ErrorSink = [&] (const FTicketStatus& Status)
 		{
 			FTicket Ticket = Status.GetTicket();
 			uint32 Index = 63 - FMath::CountLeadingZeros64(uint64(Ticket));
@@ -2589,6 +2587,6 @@ COREHTTP_API void CoreHttpTest()
 
 // }}}
 
-} // namespace UE::HTTP
+} // namespace UE::IO::IAS
 
 /* vim: set noet : */
diff --git a/Engine/Source/Runtime/Experimental/IoStoreOnDemand/Private/DistributionEndpoints.cpp b/Engine/Source/Runtime/Experimental/IoStoreOnDemand/Private/DistributionEndpoints.cpp
--- a/Engine/Source/Runtime/Experimental/IoStoreOnDemand/Private/DistributionEndpoints.cpp
+++ b/Engine/Source/Runtime/Experimental/IoStoreOnDemand/Private/DistributionEndpoints.cpp
@@ -13,11 +13,11 @@
 #include "Serialization/JsonSerializer.h"
 #include "Statistics.h"
 
-extern int32 GIoDispatcherMaxHttpRetryCount;
-
-namespace UE::IO
+namespace UE::IO::IAS
 {
 
+extern int32 GIoDispatcherMaxHttpRetryCount;
+
 FDistributionEndpoints::~FDistributionEndpoints()
 {
 	CancelEndpointRequests();
@@ -235,4 +235,4 @@ void FDistributionEndpoints::CompleteEndpointRequest(FResolveRequest& ResolveReq
 	}
 }
 
-} // namespace UE::IO
+} // namespace UE::IO::IAS
diff --git a/Engine/Source/Runtime/Experimental/IoStoreOnDemand/Private/DistributionEndpoints.h b/Engine/Source/Runtime/Experimental/IoStoreOnDemand/Private/DistributionEndpoints.h
--- a/Engine/Source/Runtime/Experimental/IoStoreOnDemand/Private/DistributionEndpoints.h
+++ b/Engine/Source/Runtime/Experimental/IoStoreOnDemand/Private/DistributionEndpoints.h
@@ -8,7 +8,7 @@
 #include "HAL/CriticalSection.h"
 #include "HttpFwd.h"
 
-namespace UE::IO
+namespace UE::IO::IAS
 {
 
 class FDistributionEndpoints
@@ -51,4 +51,4 @@ private:
 	bool bInitialized = false;
 };
 
-} // namespace UE::IO
+} // namespace UE::IO::IAS
diff --git a/Engine/Source/Runtime/Experimental/IoStoreOnDemand/Private/EncryptionKeyManager.cpp b/Engine/Source/Runtime/Experimental/IoStoreOnDemand/Private/EncryptionKeyManager.cpp
--- a/Engine/Source/Runtime/Experimental/IoStoreOnDemand/Private/EncryptionKeyManager.cpp
+++ b/Engine/Source/Runtime/Experimental/IoStoreOnDemand/Private/EncryptionKeyManager.cpp
@@ -3,7 +3,7 @@
 #include "EncryptionKeyManager.h"
 #include "Misc/CoreDelegates.h"
 
-namespace UE
+namespace UE::IO::IAS
 {
 
 FEncryptionKeyManager::FEncryptionKeyManager()
@@ -74,4 +74,4 @@ FEncryptionKeyManager& FEncryptionKeyManager::Get()
 	return Mgr;
 }
 
-} // namespace UE
+} // namespace UE::IO::IAS
diff --git a/Engine/Source/Runtime/Experimental/IoStoreOnDemand/Private/EncryptionKeyManager.h b/Engine/Source/Runtime/Experimental/IoStoreOnDemand/Private/EncryptionKeyManager.h
--- a/Engine/Source/Runtime/Experimental/IoStoreOnDemand/Private/EncryptionKeyManager.h
+++ b/Engine/Source/Runtime/Experimental/IoStoreOnDemand/Private/EncryptionKeyManager.h
@@ -7,7 +7,7 @@
 #include "Misc/AES.h"
 #include "Misc/Guid.h"
 
-namespace UE
+namespace UE::IO::IAS
 {
 
 /** Manages a set of registered encryption key(s). */
@@ -41,4 +41,4 @@ private:
 	FEncryptionKeyAddedDelegate KeyAdded;
 };
 
-} // namespace UE
+} // namespace UE::IO::IAS
diff --git a/Engine/Source/Runtime/Experimental/IoStoreOnDemand/Private/Governor.inl b/Engine/Source/Runtime/Experimental/IoStoreOnDemand/Private/Governor.inl
--- a/Engine/Source/Runtime/Experimental/IoStoreOnDemand/Private/Governor.inl
+++ b/Engine/Source/Runtime/Experimental/IoStoreOnDemand/Private/Governor.inl
@@ -4,7 +4,7 @@
 
 #include "HAL/PlatformTime.h"
 
-namespace UE::IO::Private {
+namespace UE::IO::IAS {
 
 ////////////////////////////////////////////////////////////////////////////////
 class FGovernorExternal
@@ -123,4 +123,4 @@ inline void FGovernorInternal::Return(int32 LeftOver)
 	using FGovernor = FGovernorInternal;
 #endif
 
-} // namespace UE::IO::Private
+} // namespace UE::IO::IAS
diff --git a/Engine/Source/Runtime/Experimental/IoStoreOnDemand/Private/IasCache.h b/Engine/Source/Runtime/Experimental/IoStoreOnDemand/Private/IasCache.h
--- a/Engine/Source/Runtime/Experimental/IoStoreOnDemand/Private/IasCache.h
+++ b/Engine/Source/Runtime/Experimental/IoStoreOnDemand/Private/IasCache.h
@@ -12,6 +12,8 @@ class FIoBuffer;
 class FIoReadOptions;
 struct FIoHash;
 
+namespace UE::IO::IAS {
+
 /** Cache for binary blobs with a 20 byte cache key. */
 class IIasCache
 {
@@ -49,3 +51,5 @@ struct FIasCacheConfig
 };
 
 TUniquePtr<IIasCache> MakeIasCache(const FIasCacheConfig& Config);
+
+} // namespace UE::IO::IAS
diff --git a/Engine/Source/Runtime/Experimental/IoStoreOnDemand/Private/IoStoreOnDemand.cpp b/Engine/Source/Runtime/Experimental/IoStoreOnDemand/Private/IoStoreOnDemand.cpp
--- a/Engine/Source/Runtime/Experimental/IoStoreOnDemand/Private/IoStoreOnDemand.cpp
+++ b/Engine/Source/Runtime/Experimental/IoStoreOnDemand/Private/IoStoreOnDemand.cpp
@@ -36,7 +36,7 @@
 
 DEFINE_LOG_CATEGORY(LogIas);
 
-namespace UE::IO::Private
+namespace UE::IO::IAS
 {
 
 ////////////////////////////////////////////////////////////////////////////////
@@ -98,7 +98,7 @@ static bool ApplyEncryptionKeyFromString(const FString& GuidKeyPair)
 	if (ParseEncryptionKeyParam(GuidKeyPair, KeyGuid, Key))
 	{
 		// TODO: PAK and I/O store should share key manager
-		UE::FEncryptionKeyManager::Get().AddKey(KeyGuid, Key);
+		FEncryptionKeyManager::Get().AddKey(KeyGuid, Key);
 		FCoreDelegates::GetRegisterEncryptionKeyMulticastDelegate().Broadcast(KeyGuid, Key);
 
 		return true;
@@ -110,7 +110,7 @@ static bool ApplyEncryptionKeyFromString(const FString& GuidKeyPair)
 }
 
 ////////////////////////////////////////////////////////////////////////////////
-static bool TryParseConfigContent(const FString& ConfigContent, const FString& ConfigFileName, UE::FOnDemandEndpoint& OutEndpoint)
+static bool TryParseConfigContent(const FString& ConfigContent, const FString& ConfigFileName, FOnDemandEndpoint& OutEndpoint)
 {
 	if (ConfigContent.IsEmpty())
 	{
@@ -149,7 +149,7 @@ static bool TryParseConfigContent(const FString& ConfigContent, const FString& C
 }
 
 ////////////////////////////////////////////////////////////////////////////////
-static bool TryParseConfigFileFromPlatformPackage(UE::FOnDemandEndpoint& OutEndpoint)
+static bool TryParseConfigFileFromPlatformPackage(FOnDemandEndpoint& OutEndpoint)
 {
 	const FString ConfigFileName = TEXT("IoStoreOnDemand.ini");
 	const FString ConfigPath = FPaths::Combine(TEXT("Cloud"), ConfigFileName);
@@ -159,7 +159,7 @@ static bool TryParseConfigFileFromPlatformPackage(UE::FOnDemandEndpoint& OutEndp
 }
 
 ////////////////////////////////////////////////////////////////////////////////
-static bool TryParseConfigFile(const FString& ConfigPath, UE::FOnDemandEndpoint& OutEndpoint)
+static bool TryParseConfigFile(const FString& ConfigPath, FOnDemandEndpoint& OutEndpoint)
 {
 	const FString ConfigFileName = TEXT("IoStoreOnDemand.ini");
 	FString ConfigContent; 
@@ -225,13 +225,6 @@ static FIasCacheConfig GetIasCacheConfig(const TCHAR* CommandLine)
 	return Ret;
 }
 
-} // namespace UE::IO::Private
-
-
-
-namespace UE
-{
-
 ////////////////////////////////////////////////////////////////////////////////
 FArchive& operator<<(FArchive& Ar, FOnDemandTocHeader& Header)
 {
@@ -768,7 +761,7 @@ TIoStatusOr<FIoStoreUploadResult> UploadContainerFiles(
 			if (Ar.IsError()) 
 			{
 				Toc = FOnDemandToc{};
-				if (UE::LoadFromCompactBinary(FCbFieldView(TocResponse.Body.GetData()), Toc) == false)
+				if (LoadFromCompactBinary(FCbFieldView(TocResponse.Body.GetData()), Toc) == false)
 				{
 					UE_LOG(LogIas, Warning, TEXT("Failed to load TOC '%s/%s/%s'"), *Client.GetConfig().ServiceUrl, *UploadParams.Bucket, *TocInfo.Key);
 					continue;
@@ -1153,7 +1146,7 @@ FIoStatus DownloadContainerFiles(const FIoStoreDownloadParams& DownloadParams, c
 	}
 
 	FOnDemandToc OnDemandToc;
-	if (UE::LoadFromCompactBinary(FCbFieldView(TocResponse.Body.GetData()), OnDemandToc) == false)
+	if (LoadFromCompactBinary(FCbFieldView(TocResponse.Body.GetData()), OnDemandToc) == false)
 	{
 		return FIoStatus(EIoErrorCode::ReadError, TEXT("Failed to load on demand TOC"));
 	}
@@ -1365,31 +1358,19 @@ FIoStatus DownloadContainerFiles(const FIoStoreDownloadParams& DownloadParams, c
 	return FIoStatus::Ok;
 }
 
-class FMyRequest : public FIoRequestImpl
-{
-public:
-
-	static FMyRequest* CreateRequest(const FIoChunkId& ChunkId)
-	{
-		return nullptr;
-	}
-};
-
 FIoStatus PrimeEndPoint(FStringView IoStoreOnDemandIniPath)
 {
 	TRACE_CPUPROFILER_EVENT_SCOPE(PrimeEndPoint);
 
-	using namespace UE::IO::Private;
-
-	UE::FOnDemandEndpoint EndPoint;
+	FOnDemandEndpoint EndPoint;
 	if (!TryParseConfigFile(FString(IoStoreOnDemandIniPath), EndPoint))
 	{
 		return FIoStatus(EIoErrorCode::Unknown, TEXT("Failed to parse config file"));
 	}
 
-	EndPoint.EndpointType = UE::EOnDemandEndpointType::CDN;
+	EndPoint.EndpointType = EOnDemandEndpointType::CDN;
 
-	TSharedPtr <UE::IOnDemandIoDispatcherBackend> Backend = UE::MakeOnDemandIoDispatcherBackend(nullptr);
+	TSharedPtr <IOnDemandIoDispatcherBackend> Backend = MakeOnDemandIoDispatcherBackend(nullptr);
 	Backend->Mount(EndPoint);
 
 	check(!FIoDispatcher::IsInitialized()); // Assume this is only run in standalone programs that have not
@@ -1441,12 +1422,9 @@ FIoStatus PrimeEndPoint(FStringView IoStoreOnDemandIniPath)
 }
 #endif // (PLATFORM_DESKTOP && (IS_PROGRAM || WITH_EDITOR))
 
-} // namespace UE
-
-////////////////////////////////////////////////////////////////////////////////
 
-IMPLEMENT_MODULE(FIoStoreOnDemandModule, IoStoreOnDemand);
 
+////////////////////////////////////////////////////////////////////////////////
 void FIoStoreOnDemandModule::SetBulkOptionalEnabled(bool bInEnabled)
 {
 	if (Backend.IsValid())
@@ -1474,7 +1452,6 @@ void FIoStoreOnDemandModule::ReportAnalytics(TArray<FAnalyticsEventAttribute>& O
 void FIoStoreOnDemandModule::StartupModule()
 {
 	LLM_SCOPE_BYTAG(Ias);
-	using namespace UE::IO::Private;
 
 #if WITH_EDITOR
 	bool bEnabledInEditor = false;
@@ -1488,7 +1465,7 @@ void FIoStoreOnDemandModule::StartupModule()
 
 	const TCHAR* CommandLine = FCommandLine::Get();
 
-	UE::FOnDemandEndpoint Endpoint;
+	FOnDemandEndpoint Endpoint;
 	
 	FString UrlParam;
 	if (FParse::Value(CommandLine, TEXT("Ias.TocUrl="), UrlParam))
@@ -1515,7 +1492,7 @@ void FIoStoreOnDemandModule::StartupModule()
 
 	if (!Endpoint.IsValid())
 	{
-		Endpoint = UE::FOnDemandEndpoint();
+		Endpoint = FOnDemandEndpoint();
 		if (!TryParseConfigFileFromPlatformPackage(Endpoint))
 		{
 			return;
@@ -1531,7 +1508,7 @@ void FIoStoreOnDemandModule::StartupModule()
 
 	FLatencyInjector::Initialize(CommandLine);
 
-	Endpoint.EndpointType = UE::EOnDemandEndpointType::CDN;
+	Endpoint.EndpointType = EOnDemandEndpointType::CDN;
 
 	TSharedPtr<IIasCache> Cache;
 	if (FIasCacheConfig Config = GetIasCacheConfig(CommandLine); Config.DiskQuota > 0)
@@ -1543,7 +1520,7 @@ void FIoStoreOnDemandModule::StartupModule()
 		UE_LOG(LogIas, Log, TEXT("File cache disabled. Streaming only."));
 	}
 
-	Backend = UE::MakeOnDemandIoDispatcherBackend(Cache);
+	Backend = MakeOnDemandIoDispatcherBackend(Cache);
 	Backend->Mount(Endpoint);
 	int32 BackendPriority = -10;
 #if !UE_BUILD_SHIPPING
@@ -1559,3 +1536,9 @@ void FIoStoreOnDemandModule::StartupModule()
 void FIoStoreOnDemandModule::ShutdownModule()
 {
 }
+
+} // namespace UE::IO::IAS
+
+////////////////////////////////////////////////////////////////////////////////
+
+IMPLEMENT_MODULE(UE::IO::IAS::FIoStoreOnDemandModule, IoStoreOnDemand);
diff --git a/Engine/Source/Runtime/Experimental/IoStoreOnDemand/Private/JournaledCache.cpp b/Engine/Source/Runtime/Experimental/IoStoreOnDemand/Private/JournaledCache.cpp
--- a/Engine/Source/Runtime/Experimental/IoStoreOnDemand/Private/JournaledCache.cpp
+++ b/Engine/Source/Runtime/Experimental/IoStoreOnDemand/Private/JournaledCache.cpp
@@ -42,10 +42,7 @@ by 1 journal write. A value of <=1 will flush immediately (N.B. doing so causes
 two writes per allowed op-throttle).
 */
 
-namespace UE::IO::Private
-{
-
-namespace JournaledCache
+namespace UE::IO::IAS::JournaledCache
 {
 
 // {{{1 structs ................................................................
@@ -1055,12 +1052,14 @@ uint32 FCache::DebugVisit(void* Param, FDebugCacheEntry::Callback* Callback)
 
 // }}}
 
-} // namespace JournaledCache
+} // namespace UE::IO::IAS::JournaledCache
 
 
+namespace UE::IO::IAS {
 
 // {{{1 journaled-cache ........................................................
 
+
 ////////////////////////////////////////////////////////////////////////////////
 class FJournaledCache
 	: public IIasCache
@@ -1231,8 +1230,6 @@ uint64 FJournaledCache::ReduceKey(const FIoHash& Key)
 
 // }}}
 
-} // namespace UE::IO::Private
-
 
 
 // {{{1 test ...................................................................
@@ -1262,7 +1259,7 @@ static uint64 KeyGen(const FIoBuffer& Data)
 ////////////////////////////////////////////////////////////////////////////////
 IOSTOREONDEMAND_API void Tests()
 {
-	using namespace UE::IO::Private::JournaledCache;
+	using namespace JournaledCache;
 
 	// Some randomness
 	uint64 Th = 0x0'a9e0'493; // prime!
@@ -1487,7 +1484,9 @@ IOSTOREONDEMAND_API void Tests()
 TUniquePtr<IIasCache> MakeIasCache(const FIasCacheConfig& Config)
 {
 	LLM_SCOPE_BYTAG(Ias);
-	return MakeUnique<UE::IO::Private::FJournaledCache>(Config);
+	return MakeUnique<FJournaledCache>(Config);
 }
 
+} // namespace UE::IO::IAS
+
 /* vim: set noet foldlevel=1 : */
diff --git a/Engine/Source/Runtime/Experimental/IoStoreOnDemand/Private/LatencyInjector.cpp b/Engine/Source/Runtime/Experimental/IoStoreOnDemand/Private/LatencyInjector.cpp
--- a/Engine/Source/Runtime/Experimental/IoStoreOnDemand/Private/LatencyInjector.cpp
+++ b/Engine/Source/Runtime/Experimental/IoStoreOnDemand/Private/LatencyInjector.cpp
@@ -9,7 +9,7 @@
 
 #if IAS_WITH_LATENCY_INJECTOR
 
-namespace UE::IO::Private
+namespace UE::IO::IAS
 {
 
 ////////////////////////////////////////////////////////////////////////////////
@@ -162,6 +162,6 @@ bool FLatencyInjector::HasExpired(uint32 Param)
 	return Cycles >= (int64(Param) + Injector::CycleBase);
 }
 
-} // namespace UE::IO::Private
+} // namespace UE::IO::IAS
 
 #endif // IAS_WITH_LATENCY_INJECTOR
diff --git a/Engine/Source/Runtime/Experimental/IoStoreOnDemand/Private/LatencyInjector.h b/Engine/Source/Runtime/Experimental/IoStoreOnDemand/Private/LatencyInjector.h
--- a/Engine/Source/Runtime/Experimental/IoStoreOnDemand/Private/LatencyInjector.h
+++ b/Engine/Source/Runtime/Experimental/IoStoreOnDemand/Private/LatencyInjector.h
@@ -5,7 +5,7 @@
 #include "HAL/Platform.h"
 #include "Misc/Build.h"
 
-namespace UE::IO::Private
+namespace UE::IO::IAS
 {
 
 #if !UE_BUILD_SHIPPING
@@ -26,4 +26,4 @@ struct FLatencyInjector
 	static bool HasExpired(uint32 Param)			IAS_DISABLED_IMPL(return true);
 };
 
-} // namespace UE::IO::Private
+} // namespace UE::IO::IAS
diff --git a/Engine/Source/Runtime/Experimental/IoStoreOnDemand/Private/OnDemandIoDispatcherBackend.cpp b/Engine/Source/Runtime/Experimental/IoStoreOnDemand/Private/OnDemandIoDispatcherBackend.cpp
--- a/Engine/Source/Runtime/Experimental/IoStoreOnDemand/Private/OnDemandIoDispatcherBackend.cpp
+++ b/Engine/Source/Runtime/Experimental/IoStoreOnDemand/Private/OnDemandIoDispatcherBackend.cpp
@@ -35,6 +35,9 @@
 
 #include <atomic>
 
+namespace UE::IO::IAS
+{
+
 ///////////////////////////////////////////////////////////////////////////////
 int32 GIoDispatcherMaxHttpConnectionCount = 8;
 static FAutoConsoleVariableRef CVar_IoDispatcherMaxHttpConnectionCount (
@@ -69,16 +72,13 @@ static TAutoConsoleVariable<bool> CVar_IoReportAnalytics(
 	true,
 	TEXT("Enables reporting statics to the analytics system"));
 
-namespace UE::IO::Private
-{
-
 ///////////////////////////////////////////////////////////////////////////////
 #if !UE_BUILD_SHIPPING
 static void LatencyTest(FStringView InUrl, FStringView InPath)
 {
 	auto AnsiUrl = StringCast<ANSICHAR>(InUrl.GetData(), InUrl.Len());
 
-	using namespace UE::HTTP;
+	using namespace UE::IO::IAS::HTTP;
 
 	FConnectionPool::FParams PoolParams;
 	PoolParams.SetHostFromUrl(AnsiUrl);
@@ -199,7 +199,7 @@ void FHttpClient::Get(FAnsiStringView Url, FGetCallback&& Callback)
 
 void FHttpClient::Issue(FAnsiStringView Url, FGetCallback&& Callback, FIoOffsetAndLength Range)
 {
-	using namespace UE::HTTP;
+	using namespace UE::IO::IAS::HTTP;
 
 	auto Sink = [
 		Buffer = FIoBuffer(),
@@ -241,7 +241,7 @@ void FHttpClient::Issue(FAnsiStringView Url, FGetCallback&& Callback, FIoOffsetA
 			}
 		};
 
-	UE::HTTP::FRequest Request = EventLoop.Get(Url, *ConnectionPool);
+	UE::IO::IAS::HTTP::FRequest Request = EventLoop.Get(Url, *ConnectionPool);
 	const uint64 RangeStart = Range.GetOffset();
 	const uint64 RangeEnd = Range.GetOffset() + Range.GetLength();
 	if (RangeStart > 0 || RangeEnd > 0)
@@ -814,7 +814,7 @@ static void LogIoResult(
 ///////////////////////////////////////////////////////////////////////////////
 class FOnDemandIoBackend final
 	: public FRunnable
-	, public UE::IOnDemandIoDispatcherBackend
+	, public IOnDemandIoDispatcherBackend
 {
 	using FIoRequestQueue = TThreadSafeIntrusiveQueue<FIoRequestImpl>;
 	using FChunkRequestQueue = TThreadSafeIntrusiveQueue<FChunkRequest>;
@@ -1448,7 +1448,7 @@ FIoStatus FOnDemandIoBackend::AddToc(const FOnDemandEndpoint& Endpoint)
 #if !UE_BUILD_SHIPPING
 	Launch(TEXT("IasLatencyTest"), [ServiceUrl=Endpoint.ServiceUrl, TocPath=Endpoint.TocPath] ()
 	{
-		UE::IO::Private::LatencyTest(ServiceUrl, TocPath);
+		LatencyTest(ServiceUrl, TocPath);
 	});
 #endif // !UE_BUILD_SHIPPING
 
@@ -1562,14 +1562,9 @@ uint32 FOnDemandIoBackend::Run()
 	return 0;
 }
 
-} // namespace UE::IO::Private
-
-namespace UE
-{
-
 TSharedPtr<IOnDemandIoDispatcherBackend> MakeOnDemandIoDispatcherBackend(TSharedPtr<IIasCache> Cache)
 {
-	return MakeShareable<IOnDemandIoDispatcherBackend>(new UE::IO::Private::FOnDemandIoBackend(Cache));
+	return MakeShareable<IOnDemandIoDispatcherBackend>(new FOnDemandIoBackend(Cache));
 }
 
-} // namespace UE
+} // namespace UE::IO::IAS
diff --git a/Engine/Source/Runtime/Experimental/IoStoreOnDemand/Private/OnDemandIoDispatcherBackend.h b/Engine/Source/Runtime/Experimental/IoStoreOnDemand/Private/OnDemandIoDispatcherBackend.h
--- a/Engine/Source/Runtime/Experimental/IoStoreOnDemand/Private/OnDemandIoDispatcherBackend.h
+++ b/Engine/Source/Runtime/Experimental/IoStoreOnDemand/Private/OnDemandIoDispatcherBackend.h
@@ -5,12 +5,10 @@
 #include "IO/IoDispatcherBackend.h"
 #include "Templates/SharedPointer.h"
 
-#define UE_API IOSTOREONDEMAND_API
-
 class IIasCache;
 struct FAnalyticsEventAttribute;
 
-namespace UE
+namespace UE::IO::IAS
 {
 
 enum class EOnDemandEndpointType
@@ -50,8 +48,6 @@ public:
 #endif // IS_PROGRAM || WITH_EDITOR
 };
 
-UE_API TSharedPtr<IOnDemandIoDispatcherBackend> MakeOnDemandIoDispatcherBackend(TSharedPtr<IIasCache> Cache);
-
-} // namespace UE
+TSharedPtr<IOnDemandIoDispatcherBackend> MakeOnDemandIoDispatcherBackend(TSharedPtr<IIasCache> Cache);
 
-#undef UE_API
+} // namespace UE::IO::IAS
diff --git a/Engine/Source/Runtime/Experimental/IoStoreOnDemand/Private/Statistics.cpp b/Engine/Source/Runtime/Experimental/IoStoreOnDemand/Private/Statistics.cpp
--- a/Engine/Source/Runtime/Experimental/IoStoreOnDemand/Private/Statistics.cpp
+++ b/Engine/Source/Runtime/Experimental/IoStoreOnDemand/Private/Statistics.cpp
@@ -10,7 +10,7 @@ LLM_DEFINE_TAG(Ias);
 
 #if IAS_WITH_STATISTICS
 
-namespace UE::IO::Private
+namespace UE::IO::IAS
 {
 
 static TAutoConsoleVariable<bool> CVar_ReportHttpAnalytics(
@@ -471,6 +471,6 @@ void FOnDemandIoBackendStats::OnHttpError()
 	GHttpErrorCount.Add(1);
 }
 
-} // namespace UE::IO::Private
+} // namespace UE::IO::IAS
 
 #endif // IAS_WITH_STATISTICS
diff --git a/Engine/Source/Runtime/Experimental/IoStoreOnDemand/Private/Statistics.h b/Engine/Source/Runtime/Experimental/IoStoreOnDemand/Private/Statistics.h
--- a/Engine/Source/Runtime/Experimental/IoStoreOnDemand/Private/Statistics.h
+++ b/Engine/Source/Runtime/Experimental/IoStoreOnDemand/Private/Statistics.h
@@ -16,7 +16,7 @@ struct FAnalyticsEventAttribute;
 
 LLM_DECLARE_TAG(Ias);
 
-namespace UE::IO::Private
+namespace UE::IO::IAS
 {
 
 #if IAS_WITH_STATISTICS
@@ -57,4 +57,4 @@ public:
 
 #undef IAS_STATISTICS_IMPL
 
-} // namespace UE::IO::Private
+} // namespace UE::IO::IAS
diff --git a/Engine/Source/Runtime/Experimental/IoStoreOnDemand/Public/IO/IoStoreOnDemand.h b/Engine/Source/Runtime/Experimental/IoStoreOnDemand/Public/IO/IoStoreOnDemand.h
--- a/Engine/Source/Runtime/Experimental/IoStoreOnDemand/Public/IO/IoStoreOnDemand.h
+++ b/Engine/Source/Runtime/Experimental/IoStoreOnDemand/Public/IO/IoStoreOnDemand.h
@@ -28,7 +28,7 @@ struct FIoStoreWriterSettings;
 
 UE_API DECLARE_LOG_CATEGORY_EXTERN(LogIas, VeryVerbose, All);
 
-namespace UE
+namespace UE::IO::IAS
 {
 
 ////////////////////////////////////////////////////////////////////////////////
@@ -169,18 +169,13 @@ UE_API FIoStatus PrimeEndPoint(FStringView IoStoreOnDemandIniPath);
 
 #endif // (IS_PROGRAM || WITH_EDITOR)
 
-} // namespace UE
-
-namespace UE
-{
-	class IOnDemandIoDispatcherBackend;
-}
+class IOnDemandIoDispatcherBackend;
 
 class FIoStoreOnDemandModule
 	: public IModuleInterface
 {
 private:
-	TSharedPtr<UE::IOnDemandIoDispatcherBackend> Backend;
+	TSharedPtr<IOnDemandIoDispatcherBackend> Backend;
 
 public:
 	UE_API void SetBulkOptionalEnabled(bool bInEnabled);
@@ -192,4 +187,6 @@ public:
 	virtual void ShutdownModule() override;
 };
 
+} // namespace UE::IO::IAS
+
 #undef UE_API
