Move FHttpClient to it's own code file and rename to FOnDemandHttpClient

#rnx

diff --git a/Engine/Source/Runtime/Experimental/IoStoreOnDemand/Private/OnDemandHttpClient.cpp b/Engine/Source/Runtime/Experimental/IoStoreOnDemand/Private/OnDemandHttpClient.cpp
new file mode 100644
--- /dev/null
+++ b/Engine/Source/Runtime/Experimental/IoStoreOnDemand/Private/OnDemandHttpClient.cpp
@@ -0,0 +1,114 @@
+// Copyright Epic Games, Inc. All Rights Reserved.
+
+#include "OnDemandHttpClient.h"
+
+#include "Containers/StringConv.h"
+#include "HAL/IConsoleManager.h"
+#include "HAL/PlatformTime.h"
+#include "IO/IoBuffer.h"
+#include "IO/IoStoreOnDemand.h"
+
+namespace UE::IO::IAS
+{
+
+int32 GIoDispatcherHttpPollTimeoutMs = 0;
+static FAutoConsoleVariableRef CVar_IoDispatcherMaxHttpPollTimeoutMs(
+	TEXT("ias.HttpPollTimeout"),
+	GIoDispatcherHttpPollTimeoutMs,
+	TEXT("FHttpClient::Tick poll timeout in milliseconds")
+);
+
+static void LogHttpResult(const TCHAR* Url, uint32 StatusCode, uint64 DurationMs, uint64 Size, uint64 Offset, const char* Memo = "ok")
+{
+	Size >>= 10;
+	UE_LOG(LogIas, VeryVerbose, TEXT("http-%3u: %5" UINT64_FMT "ms %5" UINT64_FMT "KiB[%7" UINT64_FMT "] '%S' %s"), StatusCode, DurationMs, Size, Offset, Memo, Url);
+};
+
+FOnDemandHttpClient::FOnDemandHttpClient(const FString& ServiceUrl, int32 MaxConnectionCount)
+	: SvcsUrl(ServiceUrl)
+	, MaxConnections(MaxConnectionCount)
+{
+	auto ServiceUrlAnsi = StringCast<ANSICHAR>(*ServiceUrl, ServiceUrl.Len());
+
+	HTTP::FConnectionPool::FParams Params;
+	if (Params.SetHostFromUrl(ServiceUrlAnsi) < 0)
+	{
+		UE_LOG(LogIas, Error, TEXT("Failed to set host from '%s'"), *ServiceUrl);
+	}
+
+	Params.ConnectionCount = MaxConnectionCount;
+	ConnectionPool = MakeUnique<HTTP::FConnectionPool>(Params);
+}
+
+void FOnDemandHttpClient::Get(FAnsiStringView Url, const FIoOffsetAndLength& Range, FGetCallback&& Callback)
+{
+	Issue(Url, MoveTemp(Callback), Range);
+}
+
+void FOnDemandHttpClient::Get(FAnsiStringView Url, FGetCallback&& Callback)
+{
+	Issue(Url, MoveTemp(Callback));
+}
+
+void FOnDemandHttpClient::Issue(FAnsiStringView Url, FGetCallback&& Callback, FIoOffsetAndLength Range)
+{
+	using namespace UE::IO::IAS::HTTP;
+
+	auto Sink = [
+		Buffer = FIoBuffer(),
+			Callback = MoveTemp(Callback),
+			Url = FString(Url),
+			Offset = Range.GetOffset(),
+			StartTime = FPlatformTime::Cycles64(),
+			StatusCode = uint32(0)]
+			(const FTicketStatus& Status) mutable
+	{
+		if (FTicketStatus::EId::Response == Status.GetId())
+		{
+			FResponse& Response = Status.GetResponse();
+			StatusCode = Response.GetStatusCode();
+			Response.SetDestination(&Buffer);
+		}
+		else if (FTicketStatus::EId::Content == Status.GetId())
+		{
+			const uint64 DurationMs = (uint64)FPlatformTime::ToMilliseconds64(FPlatformTime::Cycles64() - StartTime);
+			const FIoBuffer& Content = Status.GetContent();
+
+			LogHttpResult(*Url, StatusCode, DurationMs, Content.GetSize(), Offset);
+
+			const bool bSuccessful = StatusCode > 199 && StatusCode < 300;
+			if (bSuccessful && Content.GetSize() > 0)
+			{
+				Callback(Content, DurationMs);
+			}
+			else
+			{
+				Callback(FIoStatus(EIoErrorCode::NotFound, TEXTVIEW("Invalid Content")), DurationMs);
+			}
+		}
+		else if (FTicketStatus::EId::Error == Status.GetId())
+		{
+			const uint64 DurationMs = (uint64)FPlatformTime::ToMilliseconds64(FPlatformTime::Cycles64() - StartTime);
+			LogHttpResult(*Url, StatusCode, DurationMs, 0, Offset, Status.GetErrorReason());
+			Callback(FIoStatus(EIoErrorCode::ReadError, FString(Status.GetErrorReason())), DurationMs);
+		}
+	};
+
+		UE::IO::IAS::HTTP::FRequest Request = EventLoop.Get(Url, *ConnectionPool);
+		const uint64 RangeStart = Range.GetOffset();
+		const uint64 RangeEnd = Range.GetOffset() + Range.GetLength();
+		if (RangeStart > 0 || RangeEnd > 0)
+		{
+			Request.Header(ANSITEXTVIEW("Range"), WriteToAnsiString<64>(ANSITEXTVIEW("bytes="), RangeStart, ANSITEXTVIEW("-"), RangeEnd));
+		}
+
+		EventLoop.Send(MoveTemp(Request), MoveTemp(Sink));
+}
+
+bool FOnDemandHttpClient::Tick(bool Block)
+{
+	int32 TimeoutMs = Block ? -1 : GIoDispatcherHttpPollTimeoutMs;
+	return EventLoop.Tick(TimeoutMs) != 0;
+}
+
+} //namespace UE::IO::IAS
diff --git a/Engine/Source/Runtime/Experimental/IoStoreOnDemand/Private/OnDemandHttpClient.h b/Engine/Source/Runtime/Experimental/IoStoreOnDemand/Private/OnDemandHttpClient.h
new file mode 100644
--- /dev/null
+++ b/Engine/Source/Runtime/Experimental/IoStoreOnDemand/Private/OnDemandHttpClient.h
@@ -0,0 +1,51 @@
+// Copyright Epic Games, Inc. All Rights Reserved.
+
+#pragma once
+
+#include "CoreHttp/Client.h"
+#include "Containers/StringView.h"
+#include "Containers/UnrealString.h"
+#include "IO/IoOffsetLength.h"
+#include "IO/IoStatus.h"
+#include "Templates/Function.h"
+#include "Templates/UniquePtr.h"
+
+class FConnectionPool;
+
+namespace UE::IO::IAS
+{
+
+class FOnDemandHttpClient
+{
+public:
+	using FGetCallback = TFunction<void(TIoStatusOr<FIoBuffer>, uint64 DurationMs)>;
+
+	FOnDemandHttpClient(const FString& ServiceUrl, int32 MaxConnectionCount = 8);
+	~FOnDemandHttpClient() = default;
+
+	const FString& ServiceUrl() const 
+	{
+		return SvcsUrl;
+	}
+
+	int32 MaxConnectionCount() const
+	{
+		return MaxConnections;
+	}
+
+	void Get(FAnsiStringView Url, FGetCallback&& Callback);
+	void Get(FAnsiStringView Url, const FIoOffsetAndLength& Range, FGetCallback&& Callback);
+
+	/** @return True if the client has pending work otherwise false. */
+	bool Tick(bool Block = false);
+
+private:
+	void Issue(FAnsiStringView Url, FGetCallback&& Callback, FIoOffsetAndLength Range = FIoOffsetAndLength());
+
+	FString SvcsUrl;
+	int32 MaxConnections;
+	HTTP::FEventLoop EventLoop;
+	TUniquePtr<HTTP::FConnectionPool> ConnectionPool;
+};
+
+} // namespace UE::IO::IAS
diff --git a/Engine/Source/Runtime/Experimental/IoStoreOnDemand/Private/OnDemandIoDispatcherBackend.cpp b/Engine/Source/Runtime/Experimental/IoStoreOnDemand/Private/OnDemandIoDispatcherBackend.cpp
--- a/Engine/Source/Runtime/Experimental/IoStoreOnDemand/Private/OnDemandIoDispatcherBackend.cpp
+++ b/Engine/Source/Runtime/Experimental/IoStoreOnDemand/Private/OnDemandIoDispatcherBackend.cpp
@@ -5,10 +5,8 @@
 #include "AnalyticsEventAttribute.h"
 #include "CancellationToken.h"
 #include "Containers/StringView.h"
-#include "CoreHttp/Client.h"
 #include "DistributionEndpoints.h"
 #include "EncryptionKeyManager.h"
-#include "IasCache.h"
 #include "GenericPlatform/GenericPlatformCrashContext.h"
 #include "HAL/Event.h"
 #include "HAL/LowLevelMemTracker.h"
@@ -25,8 +23,10 @@
 #include "IO/IoStatus.h"
 #include "IO/IoStore.h"
 #include "IO/IoStoreOnDemand.h"
+#include "IasCache.h"
 #include "Misc/ScopeLock.h"
 #include "Misc/ScopeRWLock.h"
+#include "OnDemandHttpClient.h"
 #include "Serialization/CompactBinary.h"
 #include "Serialization/CompactBinarySerialization.h"
 #include "Serialization/MemoryReader.h"
@@ -53,13 +53,6 @@ static FAutoConsoleVariableRef CVar_IoDispatcherMaxHttpRetryCount (
 	TEXT("Max number of HTTP request retries before failing the I/O request.")
 );
 
-int32 GIoDispatcherHttpPollTimeoutMs = 0;
-static FAutoConsoleVariableRef CVar_IoDispatcherMaxHttpPollTimeoutMs (
-	TEXT("ias.HttpPollTimeout"),
-	GIoDispatcherHttpPollTimeoutMs,
-	TEXT("Tick() poll timeout in milliseconds")
-);
-
 bool GIoDispatcherBulkOptionalEnabled = true;
 static FAutoConsoleVariableRef CVar_IoDispatcherBulkOptionalEnabled(
 	TEXT("ias.BulkOptionalEnabled"),
@@ -126,15 +119,6 @@ static void LatencyTest(FStringView InUrl, FStringView InPath)
 }
 #endif // !UE_BUILD_SHIPPING
 
-///////////////////////////////////////////////////////////////////////////////
-static void LogHttpResult(const TCHAR* Url, uint32 StatusCode, uint64 DurationMs, uint64 Size, uint64 Offset, const char* Memo="ok")
-{
-	Size >>= 10;
-	UE_LOG(LogIas, VeryVerbose, TEXT("http-%3u: %5" UINT64_FMT "ms %5" UINT64_FMT "KiB[%7" UINT64_FMT "] '%S' %s"), StatusCode, DurationMs, Size, Offset, Memo, Url);
-};
-
-using namespace UE::Tasks;
-
 ///////////////////////////////////////////////////////////////////////////////
 FIoHash GetChunkKey(const FIoHash& ChunkHash, const FIoOffsetAndLength& Range)
 {
@@ -146,119 +130,7 @@ FIoHash GetChunkKey(const FIoHash& ChunkHash, const FIoOffsetAndLength& Range)
 }
 
 ///////////////////////////////////////////////////////////////////////////////
-class FHttpClient
-{
-public:
-	using FGetCallback = TFunction<void(TIoStatusOr<FIoBuffer>, uint64 DurationMs)>;
-
-	FHttpClient(const FString& ServiceUrl, int32 MaxConnectionCount = 8);
-	~FHttpClient() = default;
 
-	const FString& ServiceUrl() const { return SvcsUrl; }
-	int32 MaxConnectionCount() const { return MaxConnections;}
-	void Get(FAnsiStringView Url, FGetCallback&& Callback);
-	void Get(FAnsiStringView Url, const FIoOffsetAndLength& Range, FGetCallback&& Callback);
-
-	/** @return True if the client has pending work otherwise false. */
-	bool Tick(bool Block=false);
-
-private:
-	void Issue(FAnsiStringView Url, FGetCallback&& Callback, FIoOffsetAndLength Range = FIoOffsetAndLength());
-
-	FString SvcsUrl;
-	int32 MaxConnections;
-	HTTP::FEventLoop EventLoop;
-	TUniquePtr<HTTP::FConnectionPool> ConnectionPool;
-};
-
-FHttpClient::FHttpClient(const FString& ServiceUrl, int32 MaxConnectionCount)
-	: SvcsUrl(ServiceUrl)
-	, MaxConnections(MaxConnectionCount)
-{
-	auto ServiceUrlAnsi = StringCast<ANSICHAR>(*ServiceUrl, ServiceUrl.Len());
-
-	HTTP::FConnectionPool::FParams Params;
-	if (Params.SetHostFromUrl(ServiceUrlAnsi) < 0)
-	{
-		UE_LOG(LogIas, Error, TEXT("Failed to set host from '%s'"), *ServiceUrl);
-	}
-
-	Params.ConnectionCount = MaxConnectionCount;
-	ConnectionPool = MakeUnique<HTTP::FConnectionPool>(Params);
-}
-
-void FHttpClient::Get(FAnsiStringView Url, const FIoOffsetAndLength& Range, FGetCallback&& Callback)
-{
-	Issue(Url, MoveTemp(Callback), Range);
-}
-
-void FHttpClient::Get(FAnsiStringView Url, FGetCallback&& Callback)
-{
-	Issue(Url, MoveTemp(Callback));
-}
-
-void FHttpClient::Issue(FAnsiStringView Url, FGetCallback&& Callback, FIoOffsetAndLength Range)
-{
-	using namespace UE::IO::IAS::HTTP;
-
-	auto Sink = [
-		Buffer = FIoBuffer(),
-		Callback = MoveTemp(Callback),
-		Url = FString(Url),
-		Offset = Range.GetOffset(),
-		StartTime = FPlatformTime::Cycles64(),
-		StatusCode = uint32(0)]
-		(const FTicketStatus& Status) mutable
-		{ 
-			if (FTicketStatus::EId::Response == Status.GetId())
-			{
-				FResponse& Response = Status.GetResponse();
-				StatusCode = Response.GetStatusCode();
-				Response.SetDestination(&Buffer);
-			}
-			else if (FTicketStatus::EId::Content == Status.GetId())
-			{
-				const uint64 DurationMs = (uint64)FPlatformTime::ToMilliseconds64(FPlatformTime::Cycles64() - StartTime);
-				const FIoBuffer& Content = Status.GetContent(); 
-
-				LogHttpResult(*Url, StatusCode, DurationMs, Content.GetSize(), Offset);
-
-				const bool bSuccessful = StatusCode > 199 && StatusCode < 300;
-				if (bSuccessful && Content.GetSize() > 0)
-				{
-					Callback(Content, DurationMs);
-				}
-				else
-				{
-					Callback(FIoStatus(EIoErrorCode::NotFound, TEXTVIEW("Invalid Content")), DurationMs);
-				}
-			}
-			else if (FTicketStatus::EId::Error == Status.GetId())
-			{
-				const uint64 DurationMs = (uint64)FPlatformTime::ToMilliseconds64(FPlatformTime::Cycles64() - StartTime);
-				LogHttpResult(*Url, StatusCode, DurationMs, 0, Offset, Status.GetErrorReason());
-				Callback(FIoStatus(EIoErrorCode::ReadError, FString(Status.GetErrorReason())), DurationMs);
-			}
-		};
-
-	UE::IO::IAS::HTTP::FRequest Request = EventLoop.Get(Url, *ConnectionPool);
-	const uint64 RangeStart = Range.GetOffset();
-	const uint64 RangeEnd = Range.GetOffset() + Range.GetLength();
-	if (RangeStart > 0 || RangeEnd > 0)
-	{
-		Request.Header(ANSITEXTVIEW("Range"), WriteToAnsiString<64>(ANSITEXTVIEW("bytes="), RangeStart, ANSITEXTVIEW("-"), RangeEnd));
-	}
-	
-	EventLoop.Send(MoveTemp(Request), MoveTemp(Sink));
-}
-
-bool FHttpClient::Tick(bool Block)
-{
-	int32 TimeoutMs = Block ? -1 : GIoDispatcherHttpPollTimeoutMs; 
-	return EventLoop.Tick(TimeoutMs) != 0;
-}
-
-///////////////////////////////////////////////////////////////////////////////
 class FOnDemandIoStore
 {
 public:
@@ -770,8 +642,8 @@ struct FChunkRequest
 	FIoRequestImpl* RequestHead;
 	FIoRequestImpl* RequestTail;
 	FIoBuffer Chunk;
-	TTask<TIoStatusOr<FIoBuffer>> CacheTask;
-	FTask DecodeTask;
+	UE::Tasks::TTask<TIoStatusOr<FIoBuffer>> CacheTask;
+	UE::Tasks::FTask DecodeTask;
 	FIoCancellationToken CancellationToken;
 	uint64 StartTime;
 	int32 Priority;
@@ -1001,7 +873,7 @@ private:
 	FChunkRequests ChunkRequests;
 	FIoRequestQueue CompletedRequests;
 	FChunkRequestQueue HttpRequests;
-	TUniquePtr<FHttpClient> HttpClient;
+	TUniquePtr<FOnDemandHttpClient> HttpClient;
 	FOnDemandIoBackendStats Stats;
 	FRWLock Lock;
 	std::atomic_bool bStopRequested{false};
@@ -1385,7 +1257,7 @@ TIoStatusOr<FOnDemandToc> FOnDemandIoBackend::GetToc(const FOnDemandEndpoint& En
 
 	for (int32 Attempt = 0, MaxAttempts = GIoDispatcherMaxHttpRetryCount; Attempt <= MaxAttempts; ++Attempt)
 	{
-		TUniquePtr<FHttpClient> HttpClient = MakeUnique<FHttpClient>(Endpoint.ServiceUrl, GIoDispatcherMaxHttpConnectionCount);
+		TUniquePtr<FOnDemandHttpClient> HttpClient = MakeUnique<FOnDemandHttpClient>(Endpoint.ServiceUrl, GIoDispatcherMaxHttpConnectionCount);
 		TAnsiStringBuilder<256> Url;
 
 		Url << "/" << Endpoint.TocPath;
@@ -1440,13 +1312,13 @@ FIoStatus FOnDemandIoBackend::AddToc(const FOnDemandEndpoint& Endpoint)
 		FWriteScopeLock _(Lock);
 		if (!HttpClient.IsValid())
 		{
-			HttpClient = MakeUnique<FHttpClient>(Endpoint.ServiceUrl, GIoDispatcherMaxHttpConnectionCount);
+			HttpClient = MakeUnique<FOnDemandHttpClient>(Endpoint.ServiceUrl, GIoDispatcherMaxHttpConnectionCount);
 			BackendThread.Reset(FRunnableThread::Create(this, TEXT("IoStoreOnDemand"), 0, TPri_AboveNormal));
 		}
 	}
 
 #if !UE_BUILD_SHIPPING
-	Launch(TEXT("IasLatencyTest"), [ServiceUrl=Endpoint.ServiceUrl, TocPath=Endpoint.TocPath] ()
+	UE::Tasks::Launch(TEXT("IasLatencyTest"), [ServiceUrl=Endpoint.ServiceUrl, TocPath=Endpoint.TocPath] ()
 	{
 		LatencyTest(ServiceUrl, TocPath);
 	});
@@ -1515,7 +1387,7 @@ uint32 FOnDemandIoBackend::Run()
 								Stats.OnHttpError();
 							}
 
-							Launch(UE_SOURCE_LOCATION, [this, ChunkRequest]()
+							UE::Tasks::Launch(UE_SOURCE_LOCATION, [this, ChunkRequest]()
 							{
 								CompleteRequest(ChunkRequest);
 							});
