Geometry collection : add bDensityFromPhysicsMaterial to the geometry collection component
- This is set from the geometry collection asset when calling SetRestCollection and can be changed after
- Also fix the way we apply the mass multiplier to inertia

diff --git a/Engine/Source/Runtime/Experimental/Chaos/Private/PhysicsProxy/GeometryCollectionPhysicsProxy.cpp b/Engine/Source/Runtime/Experimental/Chaos/Private/PhysicsProxy/GeometryCollectionPhysicsProxy.cpp
--- a/Engine/Source/Runtime/Experimental/Chaos/Private/PhysicsProxy/GeometryCollectionPhysicsProxy.cpp
+++ b/Engine/Source/Runtime/Experimental/Chaos/Private/PhysicsProxy/GeometryCollectionPhysicsProxy.cpp
@@ -1599,8 +1599,8 @@ float FGeometryCollectionPhysicsProxy::AdjustMassForScale(float Mass) const
 Chaos::FVec3f FGeometryCollectionPhysicsProxy::AdjustInertiaForScale(const Chaos::FVec3f& Inertia) const
 {
 	const FVector3f WorldScale(Parameters.WorldTransform.GetScale3D());
-	const FVector3f InertiaScale = WorldScale * Parameters.MaterialOverrideMassScaleMultiplier;
-	return Chaos::Utilities::ScaleInertia<float>(Inertia, InertiaScale, true);
+	const FVector3f WorldScaledInertia = Chaos::Utilities::ScaleInertia<float>(Inertia, WorldScale, true);
+	return WorldScaledInertia * Parameters.MaterialOverrideMassScaleMultiplier;
 }
 
 Chaos::TPBDGeometryCollectionParticleHandle<Chaos::FReal, 3>* FGeometryCollectionPhysicsProxy::BuildNonClusters_Internal(const uint32 CollectionClusterIndex, Chaos::FPBDRigidsSolver* RigidsSolver, float Mass, Chaos::FVec3f Inertia)
diff --git a/Engine/Source/Runtime/Experimental/GeometryCollectionEngine/Private/GeometryCollection/GeometryCollectionComponent.cpp b/Engine/Source/Runtime/Experimental/GeometryCollectionEngine/Private/GeometryCollection/GeometryCollectionComponent.cpp
--- a/Engine/Source/Runtime/Experimental/GeometryCollectionEngine/Private/GeometryCollection/GeometryCollectionComponent.cpp
+++ b/Engine/Source/Runtime/Experimental/GeometryCollectionEngine/Private/GeometryCollection/GeometryCollectionComponent.cpp
@@ -323,6 +323,7 @@ UGeometryCollectionComponent::UGeometryCollectionComponent(const FObjectInitiali
 	, InitializationState(ESimulationInitializationState::Unintialized)
 	, ObjectType(EObjectStateTypeEnum::Chaos_Object_Dynamic)
 	, GravityGroupIndex(0)
+	, bDensityFromPhysicsMaterial(false)
 	, bForceMotionBlur()
 	, EnableClustering(true)
 	, ClusterGroupIndex(0)
@@ -3700,6 +3701,8 @@ void UGeometryCollectionComponent::ApplyAssetDefaults()
 		{
 			BodyInstance.SetPhysMaterialOverride(RestCollection->PhysicsMaterial);
 		}
+
+		bDensityFromPhysicsMaterial = RestCollection->bDensityFromPhysicsMaterial;
 	}
 }
 
@@ -5502,7 +5505,7 @@ float UGeometryCollectionComponent::ComputeMassScaleRelativeToAsset() const
 	UPhysicalMaterial* EnginePhysicalMaterial = GetPhysicalMaterial();
 	if (ensure(EnginePhysicalMaterial))
 	{
-		if (RestCollection && RestCollection->bDensityFromPhysicsMaterial)
+		if (RestCollection && bDensityFromPhysicsMaterial)
 		{
 			bool bMassAsDensity = false;
 			const float AssetMassOrDensity = RestCollection->GetMassOrDensity(bMassAsDensity);
diff --git a/Engine/Source/Runtime/Experimental/GeometryCollectionEngine/Public/GeometryCollection/GeometryCollectionComponent.h b/Engine/Source/Runtime/Experimental/GeometryCollectionEngine/Public/GeometryCollection/GeometryCollectionComponent.h
--- a/Engine/Source/Runtime/Experimental/GeometryCollectionEngine/Public/GeometryCollection/GeometryCollectionComponent.h
+++ b/Engine/Source/Runtime/Experimental/GeometryCollectionEngine/Public/GeometryCollection/GeometryCollectionComponent.h
@@ -780,6 +780,10 @@ public:
 	UPROPERTY(EditAnywhere, BlueprintReadWrite, Category = "ChaosPhysics|General")
 	int32 GravityGroupIndex;
 
+	/** when true, density will be used to compute mass using the assigned physics material */
+	UPROPERTY(EditAnywhere, BlueprintReadWrite, Category = "ChaosPhysics|General")
+	bool bDensityFromPhysicsMaterial;
+
 	/** If ForceMotionBlur is on, motion blur will always be active, even if the GeometryCollection is at rest. */
 	UPROPERTY(EditAnywhere, BlueprintReadWrite, Category = "ChaosPhysics|General")
 	bool bForceMotionBlur;
