Chaos - improved edge filtering for contacts on mesh boundary edges
- disabled edge contact filtering between non-mesh objects due to bugs

diff --git a/Engine/Source/Runtime/Experimental/Chaos/Private/Chaos/ChaosDebugDraw.cpp b/Engine/Source/Runtime/Experimental/Chaos/Private/Chaos/ChaosDebugDraw.cpp
--- a/Engine/Source/Runtime/Experimental/Chaos/Private/Chaos/ChaosDebugDraw.cpp
+++ b/Engine/Source/Runtime/Experimental/Chaos/Private/Chaos/ChaosDebugDraw.cpp
@@ -1011,13 +1011,12 @@ namespace Chaos
 
 					const bool bIsProbe = Contact.GetIsProbe();
 					const bool bIsActive = ManifoldPointResult.bIsValid && (!ManifoldPointResult.NetPushOut.IsNearlyZero() || !ManifoldPointResult.NetImpulse.IsNearlyZero() || (!Contact.GetUseManifold() && !Contact.AccumulatedImpulse.IsNearlyZero()));
-					if (!bIsActive && !bChaosDebugDebugDrawInactiveContacts)
+					const bool bPruned = ManifoldPoint.Flags.bDisabled;
+					if ((bPruned || !bIsActive) && !bChaosDebugDebugDrawInactiveContacts)
 					{
 						continue;
 					}
 
-					const bool bPruned = ManifoldPoint.Flags.bDisabled;
-
 					const int32 ContactPlaneOwner = 1;
 					const int32 ContactPointOwner = 1 - ContactPlaneOwner;
 					const FRigidTransform3& PlaneTransform = (ContactPlaneOwner == 0) ? Contact.GetShapeRelativeTransform0() * WorldActorTransform0 : Contact.GetShapeRelativeTransform1() * WorldActorTransform1;
@@ -1031,32 +1030,41 @@ namespace Chaos
 					// Dynamic friction, restitution = red
 					// Static friction, no restitution = green
 					// Inactive = gray
-					FColor DiscColor = FColor(200, 0, 0);
-					FColor PlaneNormalColor = FColor(200, 0, 0);
-					FColor EdgeNormalColor = FColor(200, 150, 0);
-					FColor ImpulseColor = FColor(0, 0, 200);
-					FColor PushOutColor = FColor(0, 200, 0);
-					FColor PushOutImpusleColor = FColor(0, 200, 200);
+					FColor DiscColor = FColor(250, 0, 0);
+					FColor PhiDiscColor = FColor(128, 128, 0);
+					FColor PlaneNormalColor = FColor(250, 0, 0);
+					FColor EdgeNormalColor = FColor(250, 150, 0);
+					FColor ImpulseColor = FColor(0, 0, 250);
+					FColor PushOutColor = FColor(0, 250, 0);
+					FColor PushOutImpusleColor = FColor(0, 250, 250);
+					FReal ContactLenScale = FReal(1);
 					if (ManifoldPointResult.bInsideStaticFrictionCone)
 					{
 						DiscColor = FColor(150, 200, 0);
+						PhiDiscColor = FColor(150, 200, 0);
 					}
 					if (bIsProbe)
 					{
 						DiscColor = FColor(50, 180, 180);
+						PhiDiscColor = FColor(50, 180, 180);
 						PlaneNormalColor = FColor(50, 180, 180);
-						EdgeNormalColor = FColor(50, 180, 130);
+						EdgeNormalColor = FColor(50, 220, 180);
 					}
 					else if (!bIsActive)
 					{
-						DiscColor = FColor(100, 100, 100);
-						PlaneNormalColor = FColor(100, 0, 0);
-						EdgeNormalColor = FColor(100, 80, 0);
+						DiscColor = FColor(150, 150, 150);
+						PhiDiscColor = FColor(150, 150, 150);
+						PlaneNormalColor = FColor(150, 150, 150);
+						EdgeNormalColor = FColor(150, 180, 150);
+						ContactLenScale = 0.75;
 					}
 					if (bPruned)
 					{
-						PlaneNormalColor = FColor(200, 0, 200);
-						EdgeNormalColor = FColor(200, 0, 200);
+						DiscColor = FColor(50, 50, 50);
+						PhiDiscColor = FColor(50, 50, 50);
+						PlaneNormalColor = FColor(50, 50, 50);
+						EdgeNormalColor = FColor(50, 80, 50);
+						ContactLenScale = 0.5;
 					}
 
 					const FVec3 WorldPointLocation = SpaceTransform.TransformPosition(PointLocation);
@@ -1068,30 +1076,30 @@ namespace Chaos
 					// Pushout
 					if ((Settings.PushOutScale > 0) && ManifoldPointResult.bIsValid && !ManifoldPointResult.NetPushOut.IsNearlyZero())
 					{
-						FColor Color = (ColorScale * PushOutImpusleColor).ToFColor(false);
+						FColor Color = PushOutImpusleColor;
 						FDebugDrawQueue::GetInstance().DrawDebugLine(WorldPointPlaneLocation, WorldPointPlaneLocation + Settings.DrawScale * Settings.PushOutScale * SpaceTransform.TransformVectorNoScale(FVec3(ManifoldPointResult.NetPushOut)), Color, false, Duration, uint8(Settings.DrawPriority), Settings.LineThickness);
 					}
 					if ((Settings.ImpulseScale > 0) && ManifoldPointResult.bIsValid && !ManifoldPointResult.NetImpulse.IsNearlyZero())
 					{
-						FColor Color = (ColorScale * ImpulseColor).ToFColor(false);
+						FColor Color = ImpulseColor;
 						FDebugDrawQueue::GetInstance().DrawDebugLine(WorldPointPlaneLocation, WorldPointPlaneLocation + Settings.DrawScale * Settings.ImpulseScale * SpaceTransform.TransformVectorNoScale(FVec3(ManifoldPointResult.NetImpulse)), Color, false, Duration, uint8(Settings.DrawPriority), Settings.LineThickness);
 					}
 
 					// Manifold plane and normal
 					if (Settings.ContactWidth > 0)
 					{
-						FColor C0 = (ColorScale * DiscColor).ToFColor(false);
+						FColor C0 = DiscColor;
 						FDebugDrawQueue::GetInstance().DrawDebugCircle(WorldPlaneLocation, Settings.DrawScale * Settings.ContactWidth, 12, C0, false, Duration, uint8(Settings.DrawPriority), Settings.LineThickness, Axes.GetUnitAxis(EAxis::Y), Axes.GetUnitAxis(EAxis::Z), false);
 					}
 					if (Settings.ContactLen > 0)
 					{
 						FColor NormalColor = ((ManifoldPoint.ContactPoint.ContactType != EContactPointType::EdgeEdge) ? PlaneNormalColor : EdgeNormalColor);
-						FColor C1 = (ColorScale * NormalColor).ToFColor(false);
-						FDebugDrawQueue::GetInstance().DrawDebugLine(WorldPlaneLocation, WorldPlaneLocation + Settings.DrawScale * Settings.ContactLen * WorldPlaneNormal, C1, false, Duration, uint8(Settings.DrawPriority), Settings.LineThickness);
+						FColor C1 = NormalColor;
+						FDebugDrawQueue::GetInstance().DrawDebugLine(WorldPlaneLocation, WorldPlaneLocation + Settings.DrawScale * Settings.ContactLen * ContactLenScale * WorldPlaneNormal, C1, false, Duration, uint8(Settings.DrawPriority), Settings.LineThickness);
 					}
 					if (Settings.ContactPhiWidth > 0 && (ManifoldPoint.ContactPoint.Phi < FLT_MAX))
 					{
-						FColor C2 = (ColorScale * FColor(128, 128, 0)).ToFColor(false);
+						FColor C2 = PhiDiscColor;
 						FDebugDrawQueue::GetInstance().DrawDebugCircle(WorldPlaneLocation - ManifoldPoint.ContactPoint.Phi * WorldPlaneNormal, Settings.DrawScale * Settings.ContactPhiWidth, 12, C2, false, Duration, uint8(Settings.DrawPriority), Settings.LineThickness, Axes.GetUnitAxis(EAxis::Y), Axes.GetUnitAxis(EAxis::Z), false);
 					}
 
@@ -1132,7 +1140,7 @@ namespace Chaos
 				// AccumulatedImpulse
 				if ((Settings.ImpulseScale > 0) && !Contact.GetAccumulatedImpulse().IsNearlyZero())
 				{
-					FColor Color = (ColorScale * FColor::White).ToFColor(false);
+					FColor Color = FColor::White;
 					const FVec3 ImpulsePos = SpaceTransform.TransformPosition(WorldActorTransform0.GetTranslation());
 					FDebugDrawQueue::GetInstance().DrawDebugLine(ImpulsePos, ImpulsePos + Settings.DrawScale * Settings.ImpulseScale * SpaceTransform.TransformVectorNoScale(Contact.GetAccumulatedImpulse()), Color, false, Duration, uint8(Settings.DrawPriority), Settings.LineThickness);
 				}
@@ -1143,7 +1151,7 @@ namespace Chaos
 				const FVec3 Location = SpaceTransform.TransformPosition(Contact.CalculateWorldContactLocation());
 				const FVec3 Normal = SpaceTransform.TransformVector(Contact.CalculateWorldContactNormal());
 
-				const FColor C3 = (ColorScale * FColor(128, 128, 128)).ToFColor(false);
+				const FColor C3 = FColor(128, 128, 128);
 				const FMatrix Axes = FRotationMatrix::MakeFromX(Normal);
 				const FVec3 P0 = SpaceTransform.TransformPosition(Contact.GetParticle0()->X());
 				const FVec3 P1 = SpaceTransform.TransformPosition(Contact.GetParticle1()->X());
@@ -1938,7 +1946,7 @@ namespace Chaos
 					{
 						DrawCollisionImpl(SpaceTransform, &Collision, ColorScale, ChaosDebugDrawCollisionDuration, GetChaosDebugDrawSettings(Settings));
 						return ECollisionVisitorResult::Continue;
-					});
+					}, ECollisionVisitorFlags::VisitAllCurrent);
 			}
 		}
 
diff --git a/Engine/Source/Runtime/Experimental/Chaos/Private/Chaos/Collision/CollisionPruning.cpp b/Engine/Source/Runtime/Experimental/Chaos/Private/Chaos/Collision/CollisionPruning.cpp
--- a/Engine/Source/Runtime/Experimental/Chaos/Private/Chaos/Collision/CollisionPruning.cpp
+++ b/Engine/Source/Runtime/Experimental/Chaos/Private/Chaos/Collision/CollisionPruning.cpp
@@ -2,14 +2,112 @@
 #include "Chaos/Collision/CollisionPruning.h"
 
 #include "Chaos/Collision/CollisionConstraintAllocator.h"
+#include "Chaos/Collision/ContactTriangles.h"
 #include "Chaos/Collision/ParticleCollisions.h"
 #include "Chaos/Collision/PBDCollisionConstraint.h"
+#include "Chaos/HeightField.h"
+#include "Chaos/ImplicitObject.h"
+#include "Chaos/ImplicitObjectScaled.h"
 #include "Chaos/ParticleHandle.h"
+#include "Chaos/TriangleMeshImplicitObject.h"
 
 namespace Chaos
 {
 	extern FRealSingle Chaos_Collision_EdgePrunePlaneDistance;
 
+	inline FReal CalculateEdgePlaneTolarance(const FGeometryParticleHandle* Particle)
+	{
+		// Chaos_Collision_EdgePrunePlaneDistance is prune distance for unit sized object (100cm) so scale to the size of this object
+		const FReal EdgePlaneTolerancePerCm = FReal(Chaos_Collision_EdgePrunePlaneDistance) * FReal(0.01);
+		const FReal ParticleSize = Particle->HasBounds() ? Particle->LocalBounds().Extents().Min() : FReal(1);
+		const FReal EdgePlaneTolerance = ParticleSize * EdgePlaneTolerancePerCm;
+		return EdgePlaneTolerance;
+	}
+
+	inline FReal CalculateEdgePointBarycentricTolarance(const FGeometryParticleHandle* Particle)
+	{
+		return FReal(0.001);
+	}
+
+	inline bool CollisionsHasTriangleEdgeContacts(const FPBDCollisionConstraint& Collision, const int32 MeshParticleIndex)
+	{
+		const EContactPointType TriangleFaceContactPointType = (MeshParticleIndex == 0) ? EContactPointType::PlaneVertex : EContactPointType::VertexPlane;
+		for (const FManifoldPoint& ManifoldPoint : Collision.GetManifoldPoints())
+		{
+			if (ManifoldPoint.ContactPoint.ContactType != TriangleFaceContactPointType)
+			{
+				return true;
+			}
+		}
+		return false;
+	}
+
+	inline FVec3 CalculateTriangleNormal(const FVec3 Vertices[])
+	{
+		return FVec3::CrossProduct(Vertices[1] - Vertices[0], Vertices[2] - Vertices[0]).GetSafeNormal();
+	}
+
+	// Casting utility for mesh types (tri mesh and heightfield)
+	// TFunctor: void(TMeshImplicitType* MeshImplicit, const FVec3& MeshScale)
+	template <typename TFunctor>
+	void MeshCastHelper(const FImplicitObject& Geom, const TFunctor& Func)
+	{
+		const EImplicitObjectType Type = Geom.GetType();
+		switch (Type)
+		{
+		case ImplicitObjectType::TriangleMesh:
+		{
+			const FTriangleMeshImplicitObject* Mesh = Geom.AsAChecked<FTriangleMeshImplicitObject>();
+			return Func(Mesh, FVec3(1));
+		}
+		case ImplicitObjectType::HeightField:
+		{
+			const FHeightField* Mesh = Geom.AsAChecked<FHeightField>();
+			return Func(Mesh, FVec3(1));
+		}
+		case ImplicitObjectType::IsScaled | ImplicitObjectType::TriangleMesh:
+		{
+			const TImplicitObjectScaled<FTriangleMeshImplicitObject>* ScaledMesh = Geom.AsA<TImplicitObjectScaled<FTriangleMeshImplicitObject>>();
+			return Func(ScaledMesh->GetUnscaledObject(), ScaledMesh->GetScale());
+		}
+		case ImplicitObjectType::IsScaled | ImplicitObjectType::HeightField:
+		{
+			const TImplicitObjectScaled<FHeightField>* ScaledMesh = Geom.AsA<TImplicitObjectScaled<FHeightField>>();
+			return Func(ScaledMesh->GetUnscaledObject(), ScaledMesh->GetScale());
+		}
+		case ImplicitObjectType::IsInstanced | ImplicitObjectType::TriangleMesh:
+		{
+			const TImplicitObjectInstanced<FTriangleMeshImplicitObject>* ScaledMesh = Geom.AsA<TImplicitObjectInstanced<FTriangleMeshImplicitObject>>();
+			return Func(ScaledMesh->GetInstancedObject(), FVec3(1));
+		}
+		case ImplicitObjectType::IsInstanced | ImplicitObjectType::HeightField:
+		{
+			const TImplicitObjectInstanced<FHeightField>* ScaledMesh = Geom.AsA<TImplicitObjectInstanced<FHeightField>>();
+			return Func(ScaledMesh->GetInstancedObject(), FVec3(1));
+		}
+		}
+
+		// Should only be called with a known mesh
+		check(false);
+	}
+
+	// Extract the triangle with the specified index (from a collision result) from a mesh or heightfield
+	void GetTransformedMeshTriangle(
+		const FImplicitObject* InImplicit, const int32 TriangleIndex, const FRigidTransform3& InTransform,
+		FVec3 OutVertices[],
+		int32 OutVertexIndices[])
+	{
+		MeshCastHelper(*InImplicit,
+			[&InTransform, TriangleIndex, &OutVertices, OutVertexIndices](const auto& MeshImplicit, const FVec3& MeshScale) -> void
+			{
+				const FTransform Transform = FTransform(InTransform.GetRotation(), InTransform.GetLocation(), MeshScale);
+				FTriangle Triangle;
+				MeshImplicit->GetTransformedTriangle(TriangleIndex, Transform, Triangle, OutVertexIndices[0], OutVertexIndices[1], OutVertexIndices[2]);
+				OutVertices[0] = Triangle.GetVertex(0);
+				OutVertices[1] = Triangle.GetVertex(1);
+				OutVertices[2] = Triangle.GetVertex(2);
+			});
+	}
 
 	// Loop over edge collisions, then all plane collisions and remove the edge collision if it is hidden by a plane collision
 	// NOTE: We only look at plane collisions where the other shape owns the plane.
@@ -19,10 +117,7 @@ namespace Chaos
 	// @todo(chaos): perf issue: this processes contacts in world space, but we don't calculated that data until Gather. Fix this.
 	void FParticleEdgeCollisionPruner::Prune()
 	{
-		// Chaos_Collision_EdgePrunePlaneDistance is prune distance for unit sized object (100cm) so scale to the size of this object
-		const FReal EdgePlaneTolerancePerCm = FReal(Chaos_Collision_EdgePrunePlaneDistance) * FReal(0.01);
-		const FReal EdgePlaneToleranceSize = Particle->HasBounds() ? Particle->LocalBounds().Extents().Min() : FReal(1);
-		const FReal EdgePlaneTolerance = EdgePlaneToleranceSize * EdgePlaneTolerancePerCm;
+		const FReal EdgePlaneTolerance = CalculateEdgePlaneTolarance(Particle);
 
 		FParticleCollisions& ParticleCollisions = Particle->ParticleCollisions();
 
@@ -87,6 +182,160 @@ namespace Chaos
 			});
 	}
 
+	void FParticleMeshCollisionPruner::Prune()
+	{
+		const FReal EdgePointBarycentricTolerance = CalculateEdgePointBarycentricTolarance(Particle);
+
+		// NOTE: We are assuming that our particle does not have mesh collisions and we are just checking for collisions
+		// with other meshes. This is ok since we don't support meshes on dynamic particles.
+		FParticleCollisions& ParticleCollisions = Particle->ParticleCollisions();
+
+		// Visit all pairs of collisions where we are colliding with a mesh.
+		// Eliminate edge collisions against one of the meshes if the edge is shared with the other mesh
+		// and the edge is in the plane of a face that we are colliding with on that mesh.
+		// Fix normals if we discover than an edge collision is on an edge that is shared with a triangle on another mesh
+		ParticleCollisions.VisitCollisions(
+			[this, &ParticleCollisions, EdgePointBarycentricTolerance](FPBDCollisionConstraint& MeshCollisionA)
+			{
+				// Skip non-mesh collisions
+				const int32 MeshParticleIndexA = 1;	// Meshes are always the second shape
+				const bool bIsMeshA = MeshCollisionA.GetImplicit(MeshParticleIndexA)->IsUnderlyingMesh();
+				if (!bIsMeshA)
+				{
+					return ECollisionVisitorResult::Continue;
+				}
+
+				// For the outer loop, we only care about collisions with triangle edges (i.e., not the triangle face).
+				if (!CollisionsHasTriangleEdgeContacts(MeshCollisionA, MeshParticleIndexA))
+				{
+					return ECollisionVisitorResult::Continue;
+				}
+
+				// Collision A is with a mesh and has edges. Now see if any other mesh collision invalidate the contact or require the normal to be fixed
+				ParticleCollisions.VisitCollisions(
+					[this, &MeshCollisionA, MeshParticleIndexA, EdgePointBarycentricTolerance](FPBDCollisionConstraint& MeshCollisionB)
+					{
+						const int32 MeshParticleIndexB = 1; // Meshes are always the second shape
+
+						// Skip collisions pairs on the same mesh (this implicitly skips the A-A case)
+						if (MeshCollisionA.GetImplicit(MeshParticleIndexA) == MeshCollisionB.GetImplicit(MeshParticleIndexB))
+						{
+							return ECollisionVisitorResult::Continue;
+						}
+
+						// Skip non-mesh collisions
+						const bool bIsMeshB = MeshCollisionB.GetImplicit(MeshParticleIndexB)->IsUnderlyingMesh();
+						if (!bIsMeshB)
+						{
+							return ECollisionVisitorResult::Continue;
+						}
+
+						const EContactPointType TriangleFaceContactPointTypeB = (MeshParticleIndexB == 0) ? EContactPointType::PlaneVertex : EContactPointType::VertexPlane;
+
+						// Visit the edge contact points on A...
+						for (int32 ManifoldPointIndexA = 0; ManifoldPointIndexA < MeshCollisionA.NumManifoldPoints(); ++ManifoldPointIndexA)
+						{
+							FManifoldPoint& ManifoldPointA = MeshCollisionA.GetManifoldPoint(ManifoldPointIndexA);
+							const EContactPointType TriangleFaceContactPointTypeA = (MeshParticleIndexA == 0) ? EContactPointType::PlaneVertex : EContactPointType::VertexPlane;
+							if (ManifoldPointA.ContactPoint.ContactType == TriangleFaceContactPointTypeA)
+							{
+								// Not an edge
+								continue;
+							}
+
+							// Visit all the face contacts on B...
+							for (const FManifoldPoint& ManifoldPointB : MeshCollisionB.GetManifoldPoints())
+							{
+								// The transforms of the two meshes
+								const FRigidTransform3& MeshTransformA = MeshCollisionA.GetShapeWorldTransform(MeshParticleIndexA);
+								const FRigidTransform3& MeshTransformB = MeshCollisionB.GetShapeWorldTransform(MeshParticleIndexB);
+
+								// Get the world-space triangle for collision B
+								FVec3 VerticesB[3];
+								int32 VertexIndicesB[3];
+								GetTransformedMeshTriangle(MeshCollisionB.GetImplicit(MeshParticleIndexB), ManifoldPointB.ContactPoint.FaceIndex, MeshTransformB, VerticesB, VertexIndicesB);
+
+								// World-space position of the edge contact A
+								const FVec3 EdgePosA = MeshTransformA.TransformPositionNoScale(FVec3(ManifoldPointA.ContactPoint.ShapeContactPoints[MeshParticleIndexA]));
+
+								// Is the edge contact on A on one of the edges of B?
+								int32 EdgeVertexIndicesB[2];
+								const bool bIsOnEdgeB = GetTriangleEdgeVerticesAtPosition(EdgePosA, VerticesB, EdgeVertexIndicesB[0], EdgeVertexIndicesB[1], EdgePointBarycentricTolerance);
+								
+								// If ManifoldPointA is not on and edge/vertex of triangleB then we won't be disabling A here
+								if (!bIsOnEdgeB)
+								{
+									continue;
+								}
+
+								// Get the world-space triangle for collision A
+								FVec3 VerticesA[3];
+								int32 VertexIndicesA[3];
+								GetTransformedMeshTriangle(MeshCollisionA.GetImplicit(MeshParticleIndexA), ManifoldPointA.ContactPoint.FaceIndex, MeshTransformA, VerticesA, VertexIndicesA);
+
+								// If ManifoldPoint A is on Triangle B's edge, and we collided with B's face, fix the normal on A
+								const FVec3 FaceNormalA = CalculateTriangleNormal(VerticesA);
+								if (bIsOnEdgeB && (ManifoldPointB.ContactPoint.ContactType == TriangleFaceContactPointTypeB))
+								{
+									ManifoldPointA.ContactPoint.ShapeContactNormal = MeshTransformA.InverseTransformVectorNoScale(FaceNormalA);
+									break;
+								}
+
+								// @todo(chaos): We are only handling edge contacts here, not vertex contacts which are much harder to deal with
+								// because we need to know all triangles that share the vertex in order to correct the normal
+								if ((EdgeVertexIndicesB[0] != INDEX_NONE) && (EdgeVertexIndicesB[1] != INDEX_NONE))
+								{
+									// It is common for the normals to be almost the same (near flat surfaces) and if that's the case just clamp the contact normal to the face normal
+									const FVec3 FaceNormalB = CalculateTriangleNormal(VerticesB);
+									const FReal FaceNormalADotB = FVec3::DotProduct(FaceNormalA, FaceNormalB);
+									const FReal FlatNormalThreshold = FReal(0.999);	// ~2deg
+									if (FaceNormalADotB > FlatNormalThreshold)
+									{
+										ManifoldPointA.ContactPoint.ShapeContactNormal = MeshTransformA.InverseTransformVectorNoScale(FaceNormalA);
+										break;
+									}
+
+									// Calculate the edge plane on triangle B
+									const FVec3 EdgeVectorB = VerticesB[EdgeVertexIndicesB[1]] - VerticesB[EdgeVertexIndicesB[0]];
+									const FVec3 EdgePlaneNormalB = FVec3::CrossProduct(EdgeVectorB, FaceNormalB);	// NOTE: Not normalized!!
+									const FReal EdgePlaneNormalLenSqB = EdgePlaneNormalB.SizeSquared();
+
+									// If the Face Normal on Triangle A points into the Edge Plane Normal from Triangle B, we are at a concave edge
+									// in which case we set the normal to the face normal
+									const FVec3 ContactNormalA = MeshTransformA.TransformVectorNoScale(FVec3(ManifoldPointA.ContactPoint.ShapeContactNormal));
+									const FReal FaceNormalADotEdgePlaneNormalB = FVec3::DotProduct(FaceNormalA, EdgePlaneNormalB);	// Scaled by EdgePlaneNormal.Size()
+									const FReal ConcaveNormalThresholdSq = -FMath::Square(0.01f);	// 0.5 deg
+									if (FaceNormalADotEdgePlaneNormalB * FMath::Abs(FaceNormalADotEdgePlaneNormalB) < ConcaveNormalThresholdSq * EdgePlaneNormalLenSqB)			// Accounts for non-normalized NormalADotEdgePlaneNormalB
+									{
+										ManifoldPointA.ContactPoint.ShapeContactNormal = MeshTransformA.InverseTransformVectorNoScale(FaceNormalA);
+										break;
+									}
+
+									// If the faces are not concave, force the normal to the face normal if it points into the edge plane
+									// NOTE: We assume that the contact normal does not point into the edge plane on A because collision detection should ensure that
+									const FReal ContactNormalADotEdgePlaneNormalB = FVec3::DotProduct(ContactNormalA, EdgePlaneNormalB);
+									if (ContactNormalADotEdgePlaneNormalB < FReal(0))
+									{
+										ManifoldPointA.ContactPoint.ShapeContactNormal = MeshTransformA.InverseTransformVectorNoScale(FaceNormalA);
+										break;
+									}
+								}
+							}
+						}
+
+						// If we disable all of A's edges, we don't need to go on
+						if (!CollisionsHasTriangleEdgeContacts(MeshCollisionA, MeshParticleIndexA))
+						{
+							return ECollisionVisitorResult::Stop;
+						}
+						return ECollisionVisitorResult::Continue;
+					});
+
+				return ECollisionVisitorResult::Continue;
+			});
+	}
+
+
 	/**
 	 * @brief Information about a contact point held in a ContactGroup (points in the same plane)
 	*/
diff --git a/Engine/Source/Runtime/Experimental/Chaos/Private/Chaos/ImplicitObject.cpp b/Engine/Source/Runtime/Experimental/Chaos/Private/Chaos/ImplicitObject.cpp
--- a/Engine/Source/Runtime/Experimental/Chaos/Private/Chaos/ImplicitObject.cpp
+++ b/Engine/Source/Runtime/Experimental/Chaos/Private/Chaos/ImplicitObject.cpp
@@ -92,6 +92,12 @@ namespace Chaos
 		return (Type == ImplicitObjectType::Union) || (Type == ImplicitObjectType::UnionClustered);
 	}
 
+	bool FImplicitObject::IsUnderlyingMesh() const
+	{
+		const EImplicitObjectType InnerType = GetInnerType(Type);
+		return (InnerType == ImplicitObjectType::TriangleMesh) || (InnerType == ImplicitObjectType::HeightField);
+	}
+
 	FReal FImplicitObject::SignedDistance(const FVec3& x) const
 	{
 		FVec3 Normal;
diff --git a/Engine/Source/Runtime/Experimental/Chaos/Private/Chaos/PBDCollisionConstraints.cpp b/Engine/Source/Runtime/Experimental/Chaos/Private/Chaos/PBDCollisionConstraints.cpp
--- a/Engine/Source/Runtime/Experimental/Chaos/Private/Chaos/PBDCollisionConstraints.cpp
+++ b/Engine/Source/Runtime/Experimental/Chaos/Private/Chaos/PBDCollisionConstraints.cpp
@@ -74,8 +74,13 @@ namespace Chaos
 
 	bool CollisionsAllowParticleTracking = true;
 	FAutoConsoleVariableRef CVarCollisionsAllowParticleTracking(TEXT("p.Chaos.Collision.AllowParticleTracking"), CollisionsAllowParticleTracking, TEXT("Allow particles to track their collisions constraints when their DoBufferCollisions flag is enable [def:true]"));
-	
+
+	//	Which edge pruning features to enable (for particle swith EdgeSmoothing enabled)
+	bool bCollisionsEnableEdgeCollisionPruning = false;
+	bool bCollisionsEnableMeshCollisionPruning = true;
 	bool bCollisionsEnableSubSurfaceCollisionPruning = false;
+	FAutoConsoleVariableRef CVarCollisionsEnableEdgeCollisionPruning(TEXT("p.Chaos.Collision.EnableEdgeCollisionPruning"), bCollisionsEnableEdgeCollisionPruning, TEXT(""));
+	FAutoConsoleVariableRef CVarCollisionsEnableMeshCollisionPruning(TEXT("p.Chaos.Collision.EnableMeshCollisionPruning"), bCollisionsEnableMeshCollisionPruning, TEXT(""));
 	FAutoConsoleVariableRef CVarCollisionsEnableSubSurfaceCollisionPruning(TEXT("p.Chaos.Collision.EnableSubSurfaceCollisionPruning"), bCollisionsEnableSubSurfaceCollisionPruning, TEXT(""));
 
 	bool DebugDrawProbeDetection = false;
@@ -700,8 +705,17 @@ namespace Chaos
 			{
 				if ((ParticleHandle.CollisionConstraintFlags() & (uint32)ECollisionConstraintFlags::CCF_SmoothEdgeCollisions) != 0)
 				{
-					FParticleEdgeCollisionPruner EdgePruner(ParticleHandle.Handle());
-					EdgePruner.Prune();
+					if (bCollisionsEnableMeshCollisionPruning)
+					{
+						FParticleMeshCollisionPruner MeshPruner(ParticleHandle.Handle());
+						MeshPruner.Prune();
+					}
+
+					if (bCollisionsEnableEdgeCollisionPruning)
+					{
+						FParticleEdgeCollisionPruner EdgePruner(ParticleHandle.Handle());
+						EdgePruner.Prune();
+					}
 
 					if (bCollisionsEnableSubSurfaceCollisionPruning)
 					{
diff --git a/Engine/Source/Runtime/Experimental/Chaos/Public/Chaos/Collision/CollisionPruning.h b/Engine/Source/Runtime/Experimental/Chaos/Public/Chaos/Collision/CollisionPruning.h
--- a/Engine/Source/Runtime/Experimental/Chaos/Public/Chaos/Collision/CollisionPruning.h
+++ b/Engine/Source/Runtime/Experimental/Chaos/Public/Chaos/Collision/CollisionPruning.h
@@ -27,6 +27,24 @@ namespace Chaos
 		FGeometryParticleHandle* Particle;
 	};
 
+	/**
+	 * \brief Utility to remove or correct edge contacts on the boundary between tweo separate meshes which coincident edges/vertices
+	 * \note This only look sta collisions against triangle mesh and heightfield.
+	 */
+	class FParticleMeshCollisionPruner
+	{
+	public:
+		FParticleMeshCollisionPruner(FGeometryParticleHandle* InParticle)
+			: Particle(InParticle)
+		{
+		}
+
+		void Prune();
+
+	private:
+		FGeometryParticleHandle* Particle;
+	};
+
 	/**
 	 * \brief Utility to attempt to remove sub-surface contacts on a particle when those edge contacts are "hidden" by other contacts
 	 * This is not 100% effective and may occasionally remove contacts that it shouldn't! Ideally you would just use
diff --git a/Engine/Source/Runtime/Experimental/Chaos/Public/Chaos/Collision/ContactTriangles.h b/Engine/Source/Runtime/Experimental/Chaos/Public/Chaos/Collision/ContactTriangles.h
--- a/Engine/Source/Runtime/Experimental/Chaos/Public/Chaos/Collision/ContactTriangles.h
+++ b/Engine/Source/Runtime/Experimental/Chaos/Public/Chaos/Collision/ContactTriangles.h
@@ -9,6 +9,62 @@
 
 namespace Chaos
 {
+	// Given a point and a triangle, check if the point is onan edge or vertex and return vertex indices if so.
+	// If Position is at one of the vertices, OutEdgeVertexIndexA will an index into Vertices
+	// If Position is one one of the edges, both OutEdgeVertexIndexA and OutEdgeVertexIndexB will be an index into Vertices
+	// If both OutEdgeVertexIndexA and OutEdgeVertexIndexB are INDEX_NONE, Position is not on the triangle edges
+	inline bool GetTriangleEdgeVerticesAtPosition(
+		const FVec3& Position, 
+		const FVec3 Vertices[],
+		int32& OutEdgeVertexIndexA, int32& OutEdgeVertexIndexB, 
+		const FReal BaryCentricTolerance = UE_KINDA_SMALL_NUMBER)
+	{
+		OutEdgeVertexIndexA = INDEX_NONE;
+		OutEdgeVertexIndexB = INDEX_NONE;
+
+		const FVec3 BaryCentric = ToBarycentric(Position, Vertices[0], Vertices[1], Vertices[2]);
+
+		// Is it a vertex contact?
+		if (FMath::IsNearlyEqual(BaryCentric.X, FReal(1), BaryCentricTolerance))
+		{
+			OutEdgeVertexIndexA = 0;
+			return true;
+		}
+		if (FMath::IsNearlyEqual(BaryCentric.Y, FReal(1), BaryCentricTolerance))
+		{
+			OutEdgeVertexIndexA = 1;
+			return true;
+		}
+		if (FMath::IsNearlyEqual(BaryCentric.Z, FReal(1), BaryCentricTolerance))
+		{
+			OutEdgeVertexIndexA = 2;
+			return true;
+		}
+
+		// Is it an edge contact?
+		if (FMath::IsNearlyEqual(BaryCentric.X, FReal(0), BaryCentricTolerance))
+		{
+			OutEdgeVertexIndexA = 1;
+			OutEdgeVertexIndexB = 2;
+			return true;
+		}
+		if (FMath::IsNearlyEqual(BaryCentric.Y, FReal(0), BaryCentricTolerance))
+		{
+			OutEdgeVertexIndexA = 2;
+			OutEdgeVertexIndexB = 0;
+			return true;
+		}
+		if (FMath::IsNearlyEqual(BaryCentric.Z, FReal(0), BaryCentricTolerance))
+		{
+			OutEdgeVertexIndexA = 0;
+			OutEdgeVertexIndexB = 1;
+			return true;
+		}
+
+		return false;
+	}
+
+
 	/**
 	 * @brief Data held alongside contact points when generating contacts against a (likely non-convex) mesh of triangles
 	*/
@@ -60,57 +116,25 @@ namespace Chaos
 		}
 
 		// Given a contact point on a triangle, check if it is an edge contact or vertex contact and return vertex indices if so
-		// @todo(chaos): this function should be unnecessary. Collision detection should return the vertex indices for the vertex/edge/face that we collide with
-		inline bool GetEdgeVerticesAtPosition(const FVec3& ContactPosition, FVec3& OutEdgeVertexA, FVec3& OutEdgeVertexB, int32& OutVertexIndexA, int32& OutVertexIndexB, const FReal BaryCentricTolerance = UE_KINDA_SMALL_NUMBER) const
+		// @todo(chaos): this function not should be unnecessary. Collision detection should return the vertex indices for the vertex/edge/face that we collide with
+		inline bool GetEdgeVerticesAtPosition(const FVec3& ContactPosition, FVec3& OutEdgeVertexA, FVec3& OutEdgeVertexB, int32& OutEdgeVertexIndexA, int32& OutEdgeVertexIndexB, const FReal BaryCentricTolerance = UE_KINDA_SMALL_NUMBER) const
 		{
-			OutVertexIndexA = INDEX_NONE;
-			OutVertexIndexB = INDEX_NONE;
-
-			const FVec3 BaryCentric = ToBarycentric(ContactPosition, Vertices[0], Vertices[1], Vertices[2]);
+			OutEdgeVertexIndexA = INDEX_NONE;
+			OutEdgeVertexIndexB = INDEX_NONE;
 
-			// Is it a vertex contact?
-			if (FMath::IsNearlyEqual(BaryCentric.X, FReal(1), BaryCentricTolerance))
+			int32 TriVertexIndexA, TriVertexIndexB;
+			if (GetTriangleEdgeVerticesAtPosition(ContactPosition, Vertices, TriVertexIndexA, TriVertexIndexB, BaryCentricTolerance))
 			{
-				OutVertexIndexA = VertexIndices[0];
-				OutEdgeVertexA = Vertices[0];
-				return true;
-			}
-			if (FMath::IsNearlyEqual(BaryCentric.Y, FReal(1), BaryCentricTolerance))
-			{
-				OutVertexIndexA = VertexIndices[1];
-				OutEdgeVertexA = Vertices[1];
-				return true;
-			}
-			if (FMath::IsNearlyEqual(BaryCentric.Z, FReal(1), BaryCentricTolerance))
-			{
-				OutVertexIndexA = VertexIndices[2];
-				OutEdgeVertexA = Vertices[2];
-				return true;
-			}
-
-			// Is it an edge contact?
-			if (FMath::IsNearlyEqual(BaryCentric.X, FReal(0), BaryCentricTolerance))
-			{
-				OutVertexIndexA = VertexIndices[1];
-				OutVertexIndexB = VertexIndices[2];
-				OutEdgeVertexA = Vertices[1];
-				OutEdgeVertexB = Vertices[2];
-				return true;
-			}
-			if (FMath::IsNearlyEqual(BaryCentric.Y, FReal(0), BaryCentricTolerance))
-			{
-				OutVertexIndexA = VertexIndices[2];
-				OutVertexIndexB = VertexIndices[0];
-				OutEdgeVertexA = Vertices[2];
-				OutEdgeVertexB = Vertices[0];
-				return true;
-			}
-			if (FMath::IsNearlyEqual(BaryCentric.Z, FReal(0), BaryCentricTolerance))
-			{
-				OutVertexIndexA = VertexIndices[0];
-				OutVertexIndexB = VertexIndices[1];
-				OutEdgeVertexA = Vertices[0];
-				OutEdgeVertexB = Vertices[1];
+				if (TriVertexIndexA != INDEX_NONE)
+				{
+					OutEdgeVertexIndexA = VertexIndices[TriVertexIndexA];
+					OutEdgeVertexA = Vertices[TriVertexIndexA];
+				}
+				if (TriVertexIndexB != INDEX_NONE)
+				{
+					OutEdgeVertexIndexB = VertexIndices[TriVertexIndexB];
+					OutEdgeVertexB = Vertices[TriVertexIndexB];
+				}
 				return true;
 			}
 
diff --git a/Engine/Source/Runtime/Experimental/Chaos/Public/Chaos/Collision/PBDCollisionConstraint.h b/Engine/Source/Runtime/Experimental/Chaos/Public/Chaos/Collision/PBDCollisionConstraint.h
--- a/Engine/Source/Runtime/Experimental/Chaos/Public/Chaos/Collision/PBDCollisionConstraint.h
+++ b/Engine/Source/Runtime/Experimental/Chaos/Public/Chaos/Collision/PBDCollisionConstraint.h
@@ -591,6 +591,7 @@ namespace Chaos
 
 		const FRigidTransform3& GetShapeWorldTransform0() const { return ShapeWorldTransforms[0]; }
 		const FRigidTransform3& GetShapeWorldTransform1() const { return ShapeWorldTransforms[1]; }
+		const FRigidTransform3& GetShapeWorldTransform(const int32 ParticleIndex) const { check((ParticleIndex >= 0) && (ParticleIndex < 2)); return ShapeWorldTransforms[ParticleIndex]; }
 
 		void SetShapeWorldTransforms(const FRigidTransform3& InShapeWorldTransform0, const FRigidTransform3& InShapeWorldTransform1)
 		{
diff --git a/Engine/Source/Runtime/Experimental/Chaos/Public/Chaos/HeightField.h b/Engine/Source/Runtime/Experimental/Chaos/Public/Chaos/HeightField.h
--- a/Engine/Source/Runtime/Experimental/Chaos/Public/Chaos/HeightField.h
+++ b/Engine/Source/Runtime/Experimental/Chaos/Public/Chaos/HeightField.h
@@ -99,6 +99,70 @@ namespace Chaos
 		CHAOS_API virtual int32 FindMostOpposingFace(const FVec3& Position, const FVec3& UnitDir, int32 HintFaceIndex, FReal SearchDist) const override;
 		CHAOS_API virtual FVec3 FindGeometryOpposingNormal(const FVec3& DenormDir, int32 FaceIndex, const FVec3& OriginalNormal) const override;
 
+		/**
+		 * @brief Generate the triangle at the specified index with the specified transform (including scale)
+		 * @note does not correct winding for negative scales
+		*/
+		void GetTransformedTriangle(const int32 TriangleIndex, const FRigidTransform3& Transform, FTriangle& OutTriangle, int32& OutVertexIndex0, int32& OutVertexIndex1, int32& OutVertexIndex2) const
+		{
+			// Convert the triangle index into a cell index and extract the cell row/column
+			// Reverse of FaceIndex0 = CellIndex * 2 + 0|1;
+			const int32 CellIndex = TriangleIndex / 2;
+			const int32 CellRowIndex = CellIndex / (GeomData.NumCols - 1);
+			const int32 CellColIndex = CellIndex % (GeomData.NumCols - 1);
+
+			// Calculate the vertex indices for the 4 cell corners
+			const int32 CellVertexIndex = CellRowIndex * GeomData.NumCols + CellColIndex;
+			const int32 VertexIndex0 = CellVertexIndex;
+			const int32 VertexIndex1 = CellVertexIndex + 1;
+			const int32 VertexIndex2 = CellVertexIndex + GeomData.NumCols;
+			const int32 VertexIndex3 = CellVertexIndex + GeomData.NumCols + 1;
+
+			// Get the vertices
+			const FVec3 Vertex0 = Transform.TransformPosition(GeomData.GetPointScaled(VertexIndex0));
+			const FVec3 Vertex1 = Transform.TransformPosition(GeomData.GetPointScaled(VertexIndex1));
+			const FVec3 Vertex2 = Transform.TransformPosition(GeomData.GetPointScaled(VertexIndex2));
+			const FVec3 Vertex3 = Transform.TransformPosition(GeomData.GetPointScaled(VertexIndex3));
+
+			// Set the output triangle which depends on winding (begative scales) and which of the two triangles in the cell we want
+			const bool bStandardWinding = ((GeomData.Scale.X * GeomData.Scale.Y * GeomData.Scale.Z) >= FReal(0));
+			const bool bIsFirstTriangle = ((TriangleIndex & 1) == 0);
+			if (bIsFirstTriangle)
+			{
+				if (bStandardWinding)
+				{
+					OutTriangle = FTriangle(Vertex0, Vertex1, Vertex3);
+					OutVertexIndex0 = VertexIndex0;
+					OutVertexIndex1 = VertexIndex1;
+					OutVertexIndex2 = VertexIndex3;
+				}
+				else
+				{
+					OutTriangle = FTriangle(Vertex0, Vertex3, Vertex1);
+					OutVertexIndex0 = VertexIndex0;
+					OutVertexIndex1 = VertexIndex3;
+					OutVertexIndex2 = VertexIndex1;
+				}
+			}
+			else
+			{
+				if (bStandardWinding)
+				{
+					OutTriangle = FTriangle(Vertex0, Vertex3, Vertex2);
+					OutVertexIndex0 = VertexIndex0;
+					OutVertexIndex1 = VertexIndex3;
+					OutVertexIndex2 = VertexIndex2;
+				}
+				else
+				{
+					OutTriangle = FTriangle(Vertex0, Vertex2, Vertex3);
+					OutVertexIndex0 = VertexIndex0;
+					OutVertexIndex1 = VertexIndex2;
+					OutVertexIndex2 = VertexIndex3;
+				}
+			}
+		}
+
 		/**
 		 * Visit all triangles transformed into the specified space. Triangles will have standard winding, regardless of heightfield scale.
 		 * 
diff --git a/Engine/Source/Runtime/Experimental/Chaos/Public/Chaos/ImplicitObject.h b/Engine/Source/Runtime/Experimental/Chaos/Public/Chaos/ImplicitObject.h
--- a/Engine/Source/Runtime/Experimental/Chaos/Public/Chaos/ImplicitObject.h
+++ b/Engine/Source/Runtime/Experimental/Chaos/Public/Chaos/ImplicitObject.h
@@ -273,6 +273,9 @@ public:
 	//This is strictly used for optimization purposes
 	CHAOS_API bool IsUnderlyingUnion() const;
 
+	//This is strictly used for optimization purposes
+	CHAOS_API bool IsUnderlyingMesh() const;
+
 	// Explicitly non-virtual.  Must cast to derived types to target their implementation.
 	CHAOS_API FReal SignedDistance(const FVec3& x) const;
 
diff --git a/Engine/Source/Runtime/Experimental/Chaos/Public/Chaos/ImplicitObjectScaled.h b/Engine/Source/Runtime/Experimental/Chaos/Public/Chaos/ImplicitObjectScaled.h
--- a/Engine/Source/Runtime/Experimental/Chaos/Public/Chaos/ImplicitObjectScaled.h
+++ b/Engine/Source/Runtime/Experimental/Chaos/Public/Chaos/ImplicitObjectScaled.h
@@ -1070,6 +1070,27 @@ Chaos::FImplicitObjectPtr TImplicitObjectInstanced<TConcrete>::CopyGeometryWithS
 	return Chaos::FImplicitObjectPtr(new TImplicitObjectScaled<TConcrete, true>(MObject, Scale, 0.0));
 }
 
+template<>
+struct TImplicitTypeInfo<FImplicitObjectInstanced>
+{
+	// Is InType derived from FImplicitObjectInstanced
+	static bool IsBaseOf(const EImplicitObjectType InType)
+	{
+		return !!(InType & ImplicitObjectType::IsInstanced);
+	}
+};
+
+template<>
+struct TImplicitTypeInfo<FImplicitObjectScaled>
+{
+	// Is InType derived from FImplicitObjectScaled
+	static bool IsBaseOf(const EImplicitObjectType InType)
+	{
+		return !!(InType & ImplicitObjectType::IsScaled);
+	}
+};
+
+
 /**
  * @brief Remove the Instanced or Scaled wrapper from an ImplicitObject of a known inner type and extract the instance properties
  * @return the inner implicit object or null if the wrong type
