No need to DLL export these as they're entirely internal to the module (but we'll keep the annotation for future reference)

#rnx

diff --git a/Engine/Source/Runtime/Experimental/IoStoreOnDemand/Private/CoreHttp/Client.h b/Engine/Source/Runtime/Experimental/IoStoreOnDemand/Private/CoreHttp/Client.h
--- a/Engine/Source/Runtime/Experimental/IoStoreOnDemand/Private/CoreHttp/Client.h
+++ b/Engine/Source/Runtime/Experimental/IoStoreOnDemand/Private/CoreHttp/Client.h
@@ -6,9 +6,7 @@
 #include "Containers/StringView.h"
 #endif
 
-#if !defined(COREHTTP_API)
-#	define COREHTTP_API IOSTOREONDEMAND_API
-#endif
+#define UE_API
 
 ////////////////////////////////////////////////////////////////////////////////
 class FIoBuffer;
@@ -48,7 +46,7 @@ using	FTicket = uint64;
 struct	FActivity;
 
 ////////////////////////////////////////////////////////////////////////////////
-class COREHTTP_API FConnectionPool
+class UE_API FConnectionPool
 {
 public:
 	struct FParams
@@ -84,7 +82,7 @@ private:
 };
 
 ////////////////////////////////////////////////////////////////////////////////
-class COREHTTP_API FRequest
+class UE_API FRequest
 {
 public:
 						~FRequest();
@@ -108,7 +106,7 @@ private:
 };
 
 ////////////////////////////////////////////////////////////////////////////////
-class COREHTTP_API FResponse
+class UE_API FResponse
 {
 public:
 	EStatusCodeClass	GetStatus() const;
@@ -129,7 +127,7 @@ private:
 };
 
 ////////////////////////////////////////////////////////////////////////////////
-class COREHTTP_API FTicketStatus
+class UE_API FTicketStatus
 {
 public:
 	enum class EId : uint8 { Response, Content, Cancelled, Error };
@@ -153,7 +151,7 @@ private:
 using FTicketSink = TFunction<void (const FTicketStatus&)>;
 
 ////////////////////////////////////////////////////////////////////////////////
-class COREHTTP_API FEventLoop
+class UE_API FEventLoop
 {
 	class FImpl;
 
@@ -188,4 +186,6 @@ private:
 
 } // namespace UE::IO::IAS::HTTP
 
+#undef UE_API
+
 /* vim: set noet : */
diff --git a/Engine/Source/Runtime/Experimental/IoStoreOnDemand/Private/CoreHttp/Http.cpp b/Engine/Source/Runtime/Experimental/IoStoreOnDemand/Private/CoreHttp/Http.cpp
--- a/Engine/Source/Runtime/Experimental/IoStoreOnDemand/Private/CoreHttp/Http.cpp
+++ b/Engine/Source/Runtime/Experimental/IoStoreOnDemand/Private/CoreHttp/Http.cpp
@@ -2311,7 +2311,7 @@ static void MiscTest()
 }
 
 ////////////////////////////////////////////////////////////////////////////////
-COREHTTP_API void CoreHttpTest()
+IOSTOREONDEMAND_API void CoreHttpTest()
 {
 	MiscTest();
 
