Constraints: Fix crash with nullptrs

diff --git a/Engine/Source/Runtime/Experimental/Animation/Constraints/Private/ConstraintsManager.cpp b/Engine/Source/Runtime/Experimental/Animation/Constraints/Private/ConstraintsManager.cpp
--- a/Engine/Source/Runtime/Experimental/Animation/Constraints/Private/ConstraintsManager.cpp
+++ b/Engine/Source/Runtime/Experimental/Animation/Constraints/Private/ConstraintsManager.cpp
@@ -694,6 +694,13 @@ TArray< TObjectPtr<UTickableConstraint> > FConstraintsManagerController::GetAllC
 		return Empty;
 	}
 
+	// Remove stale constraints. Stale constraints may be caused to due unexpected unloading
+	// of constributing objects, such as a level sequence.
+	Manager->Constraints.RemoveAll([](const TObjectPtr<UTickableConstraint>& ExistingConstraint) -> bool
+	{
+		return !ExistingConstraint;
+	});
+
 	if (!bSorted)
 	{
 		return Manager->Constraints;
