IAS - allow reading from file cache when HTTP streaming is disabled

#rnx

## Background
The on demand I/O backend doesn't store complete/hole I/O chunks in the local file cache, only the parts that has been requested by the game. There are currently
two ways for querying the I/O backend(s) whether a specific chunk exists or not. The primary one used by the texture streaming system is DoesChunkExist(ChunkId). This
method is called via FBulkData::DoesExist(). In order to be able to fufill this query for partially cached I/O chunks, two additional methods has been
added to the I/O dispatcher API. Currently private, only accessible by FBulkData.
 - DoesChunkExist(ChunkId, OffsetAndLength)
 - GetSizeForChunk(ChunkId, OffsetAndLength)

The method GetSizeForChunk is only ever used by the IAsyncFileHandle wrapper for bulk data and is a tricky compromise. The file handle
represents the hole chunk/file where the bulk data instance resides, i.e. .ubulk, .uptnl etc. Calling get file size on this file handle should return
the size for the entire chunk/file. It is technically possible to read outside of the bulk data range and this would result in the read request
coming back as failed, since the resolve function in the I/O backend wouldn't be found. The file handle wrapper for bulk data is currently only used
by virtual textures and should be deprecated in the long run.

diff --git a/Engine/Source/Runtime/Core/Private/IO/IoDispatcher.cpp b/Engine/Source/Runtime/Core/Private/IO/IoDispatcher.cpp
--- a/Engine/Source/Runtime/Core/Private/IO/IoDispatcher.cpp
+++ b/Engine/Source/Runtime/Core/Private/IO/IoDispatcher.cpp
@@ -3,6 +3,7 @@
 #include "IO/IoDispatcher.h"
 #include "IO/IoDispatcherPrivate.h"
 #include "IO/IoStore.h"
+#include "IO/IoOffsetLength.h"
 #include "Misc/ScopeRWLock.h"
 #include "Misc/CommandLine.h"
 #include "Misc/CoreDelegates.h"
@@ -468,6 +469,18 @@ public:
 		return false;
 	}
 
+	bool DoesChunkExist(const FIoChunkId& ChunkId, const FIoOffsetAndLength& ChunkRange) const
+	{
+		for (const FBackendAndPriority& Backend : Backends)
+		{
+			if (Backend.Value->DoesChunkExist(ChunkId, ChunkRange))
+			{
+				return true;
+			}
+		}
+		return false;
+	}
+
 	TIoStatusOr<uint64> GetSizeForChunk(const FIoChunkId& ChunkId) const
 	{
 		// Only attempt to find the size if the FIoChunkId is valid
@@ -488,6 +501,27 @@ public:
 			return FIoStatus(EIoErrorCode::InvalidParameter, TEXT("FIoChunkId is not valid"));
 		}	
 	}
+	
+	TIoStatusOr<uint64> GetSizeForChunk(const FIoChunkId& ChunkId, const FIoOffsetAndLength& ChunkRange, uint64& OutAvailable) const
+	{
+		// Only attempt to find the size if the FIoChunkId is valid
+		if (ChunkId.IsValid())
+		{
+			for (const FBackendAndPriority& Backend : Backends)
+			{
+				TIoStatusOr<uint64> Result = Backend.Value->GetSizeForChunk(ChunkId, ChunkRange, OutAvailable);
+				if (Result.IsOk())
+				{
+					return Result;
+				}
+			}
+			return FIoStatus(EIoErrorCode::NotFound);
+		}
+		else
+		{
+			return FIoStatus(EIoErrorCode::InvalidParameter, TEXT("FIoChunkId is not valid"));
+		}	
+	}
 
 	FIoSignatureErrorDelegate& OnSignatureError()
 	{
@@ -902,12 +936,24 @@ FIoDispatcher::DoesChunkExist(const FIoChunkId& ChunkId) const
 	return Impl->DoesChunkExist(ChunkId);
 }
 
+bool
+FIoDispatcher::DoesChunkExist(const FIoChunkId& ChunkId, const FIoOffsetAndLength& ChunkRange) const
+{
+	return Impl->DoesChunkExist(ChunkId, ChunkRange);
+}
+
 TIoStatusOr<uint64>
 FIoDispatcher::GetSizeForChunk(const FIoChunkId& ChunkId) const
 {
 	return Impl->GetSizeForChunk(ChunkId);
 }
 
+TIoStatusOr<uint64>
+FIoDispatcher::GetSizeForChunk(const FIoChunkId& ChunkId, const FIoOffsetAndLength& ChunkRange, uint64& OutAvailable) const
+{
+	return Impl->GetSizeForChunk(ChunkId, ChunkRange, OutAvailable);
+}
+
 int64
 FIoDispatcher::GetTotalLoaded() const
 {
diff --git a/Engine/Source/Runtime/Core/Public/IO/IoDispatcher.h b/Engine/Source/Runtime/Core/Public/IO/IoDispatcher.h
--- a/Engine/Source/Runtime/Core/Public/IO/IoDispatcher.h
+++ b/Engine/Source/Runtime/Core/Public/IO/IoDispatcher.h
@@ -64,6 +64,7 @@ class IMappedFileHandle;
 class IMappedFileRegion;
 struct FFileRegion;
 struct IIoDispatcherBackend;
+struct FIoOffsetAndLength;
 template <typename CharType> class TStringBuilderBase;
 template <typename OptionalType> struct TOptional;
 
@@ -350,11 +351,15 @@ public:
 	static CORE_API FIoDispatcher& Get();
 
 private:
+	CORE_API bool					DoesChunkExist(const FIoChunkId& ChunkId, const FIoOffsetAndLength& ChunkRange) const;
+	CORE_API TIoStatusOr<uint64>	GetSizeForChunk(const FIoChunkId& ChunkId, const FIoOffsetAndLength& ChunkRange, uint64& OutAvailable) const;
+
 	FIoDispatcherImpl* Impl = nullptr;
 
 	friend class FIoRequest;
 	friend class FIoBatch;
 	friend class FIoQueue;
+	friend class FBulkData;
 };
 
 //////////////////////////////////////////////////////////////////////////
diff --git a/Engine/Source/Runtime/Core/Public/IO/IoDispatcherBackend.h b/Engine/Source/Runtime/Core/Public/IO/IoDispatcherBackend.h
--- a/Engine/Source/Runtime/Core/Public/IO/IoDispatcherBackend.h
+++ b/Engine/Source/Runtime/Core/Public/IO/IoDispatcherBackend.h
@@ -2,6 +2,7 @@
 
 #pragma once
 
+#include "IO/IoOffsetLength.h"
 #include "IoDispatcher.h"
 #include "Async/InheritedContext.h"
 #include "Misc/Optional.h"
@@ -115,7 +116,14 @@ struct IIoDispatcherBackend
 	virtual void CancelIoRequest(FIoRequestImpl* Request) = 0;
 	virtual void UpdatePriorityForIoRequest(FIoRequestImpl* Request) = 0;
 	virtual bool DoesChunkExist(const FIoChunkId& ChunkId) const = 0;
+	virtual bool DoesChunkExist(const FIoChunkId& ChunkId, const FIoOffsetAndLength& ChunkRange) const { return DoesChunkExist(ChunkId); }
 	virtual TIoStatusOr<uint64> GetSizeForChunk(const FIoChunkId& ChunkId) const = 0;
+	virtual TIoStatusOr<uint64> GetSizeForChunk(const FIoChunkId& ChunkId, const FIoOffsetAndLength& ChunkRange, uint64& OutAvailable) const
+	{
+		TIoStatusOr<uint64> ChunkSize = GetSizeForChunk(ChunkId);
+		OutAvailable = ChunkSize.IsOk() ? ChunkSize.ValueOrDie() : 0;
+		return ChunkSize;
+	}
 	virtual FIoRequestImpl* GetCompletedRequests() = 0;
 	virtual TIoStatusOr<FIoMappedRegion> OpenMapped(const FIoChunkId& ChunkId, const FIoReadOptions& Options) = 0;
 };
diff --git a/Engine/Source/Runtime/CoreUObject/Private/Serialization/BulkData.cpp b/Engine/Source/Runtime/CoreUObject/Private/Serialization/BulkData.cpp
--- a/Engine/Source/Runtime/CoreUObject/Private/Serialization/BulkData.cpp
+++ b/Engine/Source/Runtime/CoreUObject/Private/Serialization/BulkData.cpp
@@ -4,6 +4,7 @@
 #include "Async/MappedFileHandle.h"
 #include "HAL/IConsoleManager.h"
 #include "IO/IoDispatcher.h"
+#include "IO/IoOffsetLength.h"
 #include "Math/GuardedInt.h"
 #include "Misc/ConfigCacheIni.h"
 #include "Misc/Optional.h"
@@ -87,7 +88,11 @@ bool OpenReadBulkData(
 	TFunction<void(FArchive& Ar)>&& Read);
 
 /** Open async read file handle for the specified bulk data chunk ID. */
-TUniquePtr<IAsyncReadFileHandle> OpenAsyncReadBulkData(const FBulkMetaData& BulkMeta, const FIoChunkId& BulkChunkId);
+TUniquePtr<IAsyncReadFileHandle> OpenAsyncReadBulkData(
+	const FBulkMetaData& BulkMeta,
+	const FIoChunkId& BulkChunkId,
+	uint64 ChunkSize,
+	uint64 AvailableChunkSize);
 
 /** Create bulk data streaming request. */
 TUniquePtr<IBulkDataIORequest> CreateStreamingRequest(
@@ -99,9 +104,6 @@ TUniquePtr<IBulkDataIORequest> CreateStreamingRequest(
 	FBulkDataIORequestCallBack* CompleteCallback,
 	uint8* UserSuppliedMemory);
 
-/** Returns whether the bulk data chunk exist or not. */
-bool DoesBulkDataExist(const FIoChunkId& BulkChunkId);
-
 /** Try memory map the chunk specified by the bulk data ID. */
 bool TryMemoryMapBulkData(
 	const FBulkMetaData& BulkMeta,
@@ -286,6 +288,16 @@ bool FBulkMetaData::FromSerialized(FArchive& Ar, int64 ElementSize, FBulkMetaDat
 	return true;
 }
 
+FIoOffsetAndLength FBulkMetaData::GetOffsetAndLength() const
+{
+	const uint64 Offset = GetOffset();
+	const uint64 Size = GetSize();
+
+	return Offset >= 0 && Size > 0
+		? FIoOffsetAndLength(static_cast<int64>(Offset), static_cast<int64>(Size))
+		: FIoOffsetAndLength();
+}
+
 /**
  * Bulkdata payload lengths are stored as type int64 but FMemory::Malloc takes SIZE_T which could potentially
  * be 32bit. This utility helps convert between the two and makes sure that if SizeInBytes exceeds the max size
@@ -761,8 +773,7 @@ bool FBulkData::DoesExist() const
 		return false;
 	}
 #endif
-
-	return UE::BulkData::Private::DoesBulkDataExist(BulkChunkId);
+	return FIoDispatcher::Get().DoesChunkExist(BulkChunkId, BulkMeta.GetOffsetAndLength());
 }
 
 /**
@@ -1585,7 +1596,14 @@ int64 FBulkData::SerializePayload(FArchive& Ar, EBulkDataFlags SerializationFlag
 
 IAsyncReadFileHandle* FBulkData::OpenAsyncReadHandle() const
 {
-	return UE::BulkData::Private::OpenAsyncReadBulkData(BulkMeta, BulkChunkId).Release();
+	uint64 AvailableChunkSize = 0;
+	TIoStatusOr<uint64> ChunkSize = FIoDispatcher::Get().GetSizeForChunk(BulkChunkId, BulkMeta.GetOffsetAndLength(), AvailableChunkSize);
+
+	return UE::BulkData::Private::OpenAsyncReadBulkData(
+		BulkMeta,
+		BulkChunkId,
+		ChunkSize.IsOk() ? ChunkSize.ValueOrDie() : 0,
+		AvailableChunkSize).Release();
 }
 
 IBulkDataIORequest* FBulkData::CreateStreamingRequest(EAsyncIOPriorityAndFlags Priority, FBulkDataIORequestCallBack* CompleteCallback, uint8* UserSuppliedMemory) const
diff --git a/Engine/Source/Runtime/CoreUObject/Private/Serialization/BulkDataStreaming.cpp b/Engine/Source/Runtime/CoreUObject/Private/Serialization/BulkDataStreaming.cpp
--- a/Engine/Source/Runtime/CoreUObject/Private/Serialization/BulkDataStreaming.cpp
+++ b/Engine/Source/Runtime/CoreUObject/Private/Serialization/BulkDataStreaming.cpp
@@ -5,6 +5,7 @@
 #include "Containers/ChunkedArray.h"
 #include "HAL/CriticalSection.h"
 #include "IO/IoDispatcher.h"
+#include "IO/IoOffsetLength.h"
 #include "Misc/ConfigCacheIni.h"
 #include "Misc/Timespan.h"
 #include "ProfilingDebugging/CountersTrace.h"
@@ -258,15 +259,13 @@ void FChunkReadFileRequest::HandleChunkResult(TIoStatusOr<FIoBuffer>&& Result)
 class FChunkFileSizeRequest : public IAsyncReadRequest
 {
 public:
-	FChunkFileSizeRequest(const FIoChunkId& ChunkId, FAsyncFileCallBack* Callback)
+	FChunkFileSizeRequest(const FIoChunkId& ChunkId, uint64 ChunkSize, FAsyncFileCallBack* Callback)
 		: IAsyncReadRequest(Callback, true, nullptr)
 	{
-		TIoStatusOr<uint64> Result = FIoDispatcher::Get().GetSizeForChunk(ChunkId);
-		if (Result.IsOk())
+		if (ChunkSize > 0)
 		{
-			Size = Result.ValueOrDie();
+			Size = static_cast<int64>(ChunkSize);
 		}
-
 		SetComplete();
 	}
 
@@ -295,8 +294,11 @@ private:
 class FChunkReadFileHandle : public IAsyncReadFileHandle
 {
 public:
-	FChunkReadFileHandle(const FIoChunkId& InChunkId) 
+	FChunkReadFileHandle(const FIoChunkId& InChunkId, const FIoOffsetAndLength& InChunkRange, uint64 InChunkSize, uint64 InAvailableChunkSize) 
 		: ChunkId(InChunkId)
+		, ChunkRange(InChunkRange)
+		, ChunkSize(InChunkSize)
+		, AvailableChunkSize(InAvailableChunkSize)
 	{
 	}
 
@@ -313,11 +315,14 @@ public:
 
 private:
 	FIoChunkId ChunkId;
+	FIoOffsetAndLength ChunkRange;
+	uint64 ChunkSize;
+	uint64 AvailableChunkSize;
 };
 
 IAsyncReadRequest* FChunkReadFileHandle::SizeRequest(FAsyncFileCallBack* CompleteCallback)
 {
-	return new FChunkFileSizeRequest(ChunkId, CompleteCallback);
+	return new FChunkFileSizeRequest(ChunkId, ChunkSize, CompleteCallback);
 }
 
 IAsyncReadRequest* FChunkReadFileHandle::ReadRequest(
@@ -327,6 +332,14 @@ IAsyncReadRequest* FChunkReadFileHandle::ReadRequest(
 	FAsyncFileCallBack* CompleteCallback,
 	uint8* UserSuppliedMemory)
 {
+	const bool bIsOutsideBulkDataRange =
+		(Offset < static_cast<int64>(ChunkRange.GetOffset())) ||
+		((Offset + BytesToRead) > static_cast<int64>(ChunkRange.GetOffset() + AvailableChunkSize));
+
+	UE_CLOG(bIsOutsideBulkDataRange, LogSerialization, Warning,
+		TEXT("Reading outside of bulk data range, RequestRange='%lld, %lld', BulkDataRange='%llu, %llu', ChunkId='%s'"),
+		Offset, BytesToRead, ChunkRange.GetOffset(), ChunkRange.GetLength(), *LexToString(ChunkId));
+
 	FIoBuffer Buffer = UserSuppliedMemory ? FIoBuffer(FIoBuffer::Wrap, UserSuppliedMemory, BytesToRead) : FIoBuffer(BytesToRead);
 	FChunkReadFileRequest* Request = new FChunkReadFileRequest(CompleteCallback, MoveTemp(Buffer));
 
@@ -463,6 +476,20 @@ bool OpenReadBulkData(
 
 //////////////////////////////////////////////////////////////////////////////
 
+TUniquePtr<IAsyncReadFileHandle> OpenAsyncReadBulkData(
+	const FBulkMetaData& BulkMeta,
+	const FIoChunkId& BulkChunkId,
+	uint64 ChunkSize,
+	uint64 AvailableChunkSize)
+{
+	if (BulkChunkId.IsValid() == false)
+	{
+		return TUniquePtr<IAsyncReadFileHandle>();
+	}
+
+	return MakeUnique<FChunkReadFileHandle>(BulkChunkId, BulkMeta.GetOffsetAndLength(), ChunkSize, AvailableChunkSize);
+}
+
 TUniquePtr<IAsyncReadFileHandle> OpenAsyncReadBulkData(const FBulkMetaData& BulkMeta, const FIoChunkId& BulkChunkId)
 {
 	if (BulkChunkId.IsValid() == false)
@@ -470,7 +497,9 @@ TUniquePtr<IAsyncReadFileHandle> OpenAsyncReadBulkData(const FBulkMetaData& Bulk
 		return TUniquePtr<IAsyncReadFileHandle>();
 	}
 
-	return MakeUnique<FChunkReadFileHandle>(BulkChunkId);
+	TIoStatusOr<uint64> Status = FIoDispatcher::Get().GetSizeForChunk(BulkChunkId); 
+	const uint64 ChunkSize = Status.IsOk() ? Status.ValueOrDie() : 0;
+	return MakeUnique<FChunkReadFileHandle>(BulkChunkId, BulkMeta.GetOffsetAndLength(), ChunkSize, ChunkSize);
 }
 
 //////////////////////////////////////////////////////////////////////////////
@@ -498,13 +527,6 @@ TUniquePtr<IBulkDataIORequest> CreateStreamingRequest(
 
 //////////////////////////////////////////////////////////////////////////////
 
-bool DoesBulkDataExist(const FIoChunkId& BulkChunkId)
-{
-	return FIoDispatcher::Get().DoesChunkExist(BulkChunkId);
-}
-
-//////////////////////////////////////////////////////////////////////////////
-
 bool TryMemoryMapBulkData(
 	const FBulkMetaData& BulkMeta,
 	const FIoChunkId& BulkChunkId,
diff --git a/Engine/Source/Runtime/CoreUObject/Public/Serialization/BulkData.h b/Engine/Source/Runtime/CoreUObject/Public/Serialization/BulkData.h
--- a/Engine/Source/Runtime/CoreUObject/Public/Serialization/BulkData.h
+++ b/Engine/Source/Runtime/CoreUObject/Public/Serialization/BulkData.h
@@ -46,6 +46,7 @@ class IAsyncReadFileHandle;
 class IAsyncReadRequest;
 class UObject;
 struct FTimespan;
+struct FIoOffsetAndLength;
 namespace UE { namespace Serialization { class FEditorBulkData; } }
 enum class EFileRegionType : uint8;
 
@@ -403,6 +404,8 @@ public:
 		WriteUInt40(Data + 5, uint64(Offset));
 	}
 
+	FIoOffsetAndLength GetOffsetAndLength() const;
+
 	EBulkDataLockStatus GetLockStatus() const
 	{
 		return static_cast<EBulkDataLockStatus>(Data[11]);
diff --git a/Engine/Source/Runtime/Experimental/IoStoreOnDemand/Private/DistributionEndpoints.cpp b/Engine/Source/Runtime/Experimental/IoStoreOnDemand/Private/DistributionEndpoints.cpp
--- a/Engine/Source/Runtime/Experimental/IoStoreOnDemand/Private/DistributionEndpoints.cpp
+++ b/Engine/Source/Runtime/Experimental/IoStoreOnDemand/Private/DistributionEndpoints.cpp
@@ -16,7 +16,7 @@
 namespace UE::IO::IAS
 {
 
-extern int32 GIoDispatcherMaxHttpRetryCount;
+extern int32 GIasMaxHttpRetryCount;
 
 FDistributionEndpoints::~FDistributionEndpoints()
 {
@@ -101,7 +101,7 @@ void FDistributionEndpoints::IssueEndpointRequests()
 	TRACE_CPUPROFILER_EVENT_SCOPE(FOnDemandIoBackend::IssueEndpointRequests);
 	// Currently we need to use the HTTP module in order to resolve service endpoints due to HTTPS
 	FHttpModule& HttpModule = FHttpModule::Get();
-	const int32 MaxAttempts = GIoDispatcherMaxHttpRetryCount;
+	const int32 MaxAttempts = GIasMaxHttpRetryCount;
 
 	TArray<FHttpRequestPtr, TInlineAllocator<2>> HttpRequests;
 	{
diff --git a/Engine/Source/Runtime/Experimental/IoStoreOnDemand/Private/OnDemandIoDispatcherBackend.cpp b/Engine/Source/Runtime/Experimental/IoStoreOnDemand/Private/OnDemandIoDispatcherBackend.cpp
--- a/Engine/Source/Runtime/Experimental/IoStoreOnDemand/Private/OnDemandIoDispatcherBackend.cpp
+++ b/Engine/Source/Runtime/Experimental/IoStoreOnDemand/Private/OnDemandIoDispatcherBackend.cpp
@@ -24,6 +24,7 @@
 #include "IO/IoStore.h"
 #include "IO/IoStoreOnDemand.h"
 #include "IasCache.h"
+#include "Math/NumericLimits.h"
 #include "Misc/ScopeLock.h"
 #include "Misc/ScopeRWLock.h"
 #include "OnDemandHttpClient.h"
@@ -39,25 +40,39 @@ namespace UE::IO::IAS
 {
 
 ///////////////////////////////////////////////////////////////////////////////
-int32 GIoDispatcherMaxHttpConnectionCount = 8;
-static FAutoConsoleVariableRef CVar_IoDispatcherMaxHttpConnectionCount (
+int32 GIasMaxHttpConnectionCount = 8;
+static FAutoConsoleVariableRef CVar_IasMaxHttpConnectionCount(
 	TEXT("ias.MaxHttpConnectionCount"),
-	GIoDispatcherMaxHttpConnectionCount,
+	GIasMaxHttpConnectionCount,
 	TEXT("Max number of open HTTP connections to the on demand endpoint(s).")
 );
 
-int32 GIoDispatcherMaxHttpRetryCount = 2;
-static FAutoConsoleVariableRef CVar_IoDispatcherMaxHttpRetryCount (
+int32 GIasMaxHttpRetryCount = 2;
+static FAutoConsoleVariableRef CVar_IasMaxHttpRetryCount(
 	TEXT("ias.MaxHttpRetryCount"),
-	GIoDispatcherMaxHttpRetryCount,
+	GIasMaxHttpRetryCount,
 	TEXT("Max number of HTTP request retries before failing the I/O request.")
 );
 
-bool GIoDispatcherBulkOptionalEnabled = true;
-static FAutoConsoleVariableRef CVar_IoDispatcherBulkOptionalEnabled(
-	TEXT("ias.BulkOptionalEnabled"),
-	GIoDispatcherBulkOptionalEnabled,
-	TEXT("Enables bulk optional requests.")
+int32 GIasHttpPollTimeoutMs = 0;
+static FAutoConsoleVariableRef CVar_IasMaxHttpPollTimeoutMs(
+	TEXT("ias.HttpPollTimeout"),
+	GIasHttpPollTimeoutMs,
+	TEXT("Tick() poll timeout in milliseconds")
+);
+
+bool GIasHttpEnabled = true;
+static FAutoConsoleVariableRef CVar_IasHttpEnabled(
+	TEXT("ias.HttpEnabled"),
+	GIasHttpEnabled,
+	TEXT("Enables individual asset streaming via HTTP")
+);
+
+bool GIasHttpOptionalBulkDataEnabled = true;
+static FAutoConsoleVariableRef CVar_IasHttpOptionalBulkDataEnabled(
+	TEXT("ias.HttpOptionalBulkDataEnabled"),
+	GIasHttpOptionalBulkDataEnabled,
+	TEXT("Enables optional bulk data via HTTP")
 );
 
 static TAutoConsoleVariable<bool> CVar_IoReportAnalytics(
@@ -494,22 +509,27 @@ private:
 ///////////////////////////////////////////////////////////////////////////////
 struct FChunkRequestParams
 {
-	static FChunkRequestParams Create(FIoRequestImpl* Request, FOnDemandIoStore::FChunkInfo ChunkInfo)
+	static FChunkRequestParams Create(const FIoOffsetAndLength& OffsetLength, FOnDemandIoStore::FChunkInfo ChunkInfo)
 	{
-		check(Request);
-		check(Request->NextRequest == nullptr);
-		const uint64 RawSize = FMath::Min(Request->Options.GetSize(), ChunkInfo.Entry->RawSize);
+		const uint64 RawSize = FMath::Min(OffsetLength.GetLength(), ChunkInfo.Entry->RawSize);
 		
 		const FIoOffsetAndLength ChunkRange = FIoChunkEncoding::GetChunkRange(
 			ChunkInfo.Entry->RawSize,
 			ChunkInfo.Container->BlockSize,
 			ChunkInfo.GetBlocks(),
-			Request->Options.GetOffset(),
+			OffsetLength.GetOffset(),
 			RawSize).ConsumeValueOrDie();
 
 		return FChunkRequestParams{GetChunkKey(ChunkInfo.Entry->Hash, ChunkRange), ChunkRange, ChunkInfo};
 	}
 
+	static FChunkRequestParams Create(FIoRequestImpl* Request, FOnDemandIoStore::FChunkInfo ChunkInfo)
+	{
+		check(Request);
+		check(Request->NextRequest == nullptr);
+		return Create(FIoOffsetAndLength(Request->Options.GetOffset(), Request->Options.GetSize()), ChunkInfo);
+	}
+
 	const FIoHash& GetUrlHash() const
 	{
 		return ChunkInfo.Entry->Hash;
@@ -632,9 +652,9 @@ struct FChunkRequest
 		return Head;
 	}
 
-	double DurationInSeconds() const
+	const FIoChunkId& GetChunkId()
 	{
-		return FPlatformTime::ToSeconds64(FPlatformTime::Cycles64() - StartTime);
+		return RequestHead->ChunkId;
 	}
 
 	FChunkRequest* NextRequest;
@@ -837,6 +857,7 @@ public:
 	virtual void CancelIoRequest(FIoRequestImpl* Request) override;
 	virtual void UpdatePriorityForIoRequest(FIoRequestImpl* Request) override;
 	virtual bool DoesChunkExist(const FIoChunkId& ChunkId) const override;
+	virtual bool DoesChunkExist(const FIoChunkId& ChunkId, const FIoOffsetAndLength& ChunkRange) const override;
 	virtual TIoStatusOr<uint64> GetSizeForChunk(const FIoChunkId& ChunkId) const override;
 	virtual FIoRequestImpl* GetCompletedRequests() override;
 	virtual TIoStatusOr<FIoMappedRegion> OpenMapped(const FIoChunkId& ChunkId, const FIoReadOptions& Options) override;
@@ -862,6 +883,12 @@ private:
 	FIoStatus MountDeferredEndpoints(const FString& DistributionUrl, const TConstArrayView<FString>& ServiceUrls);
 	static TIoStatusOr<FOnDemandToc> GetToc(const FOnDemandEndpoint& Endpoint);
 	FIoStatus AddToc(const FOnDemandEndpoint& Endpoint);
+	bool IsHttpEnabled() const { return bHttpEnabled && GIasHttpEnabled; }
+	bool IsHttpEnabled(const FIoChunkId& ChunkId) const
+	{ 
+		return (ChunkId.GetChunkType() != EIoChunkType::OptionalBulkData || (bHttpOptionalBulkDataEnabled && GIasHttpOptionalBulkDataEnabled));
+	}
+	virtual TIoStatusOr<uint64> GetSizeForChunk(const FIoChunkId& ChunkId, const FIoOffsetAndLength& ChunkRange, uint64& OutAvailable) const;
 
 	TSharedPtr<IIasCache> Cache;
 	TUniquePtr<FOnDemandIoStore> IoStore;
@@ -877,8 +904,8 @@ private:
 	FOnDemandIoBackendStats Stats;
 	FRWLock Lock;
 	std::atomic_bool bStopRequested{false};
-	std::atomic_bool bEnableBulkOptional{true};
-	std::atomic_bool bEnabled{true};
+	std::atomic_bool bHttpOptionalBulkDataEnabled{true};
+	std::atomic_bool bHttpEnabled{true};
 };
 
 ///////////////////////////////////////////////////////////////////////////////
@@ -898,7 +925,7 @@ void FOnDemandIoBackend::Initialize(TSharedRef<const FIoDispatcherBackendContext
 	TRACE_CPUPROFILER_EVENT_SCOPE(FOnDemandIoBackend::Initialize);
 	LLM_SCOPE_BYTAG(Ias);
 	UE_LOG(LogIas, Log, TEXT("Initializing on demand I/O dispatcher backend"));
-	FGenericCrashContext::SetEngineData(TEXT("IAS.Enabled"), bEnabled ? TEXT("true") : TEXT("false"));
+	FGenericCrashContext::SetEngineData(TEXT("IAS.Enabled"), bHttpEnabled.load(std::memory_order_relaxed) ? TEXT("true") : TEXT("false"));
 	BackendContext = Context;
 	DistributionEndpoints.ResolveDeferredEndpoints();
 }
@@ -937,7 +964,8 @@ void FOnDemandIoBackend::CompleteRequest(FChunkRequest* ChunkRequest)
 	FIoBuffer Chunk = MoveTemp(ChunkRequest->Chunk);
 	FIoChunkDecodingParams DecodingParams = ChunkRequest->Params.GetDecodingParams();
 
-	bool bCanCache = Cache.IsValid();
+	// Only cache chunks if HTTP streaming is enabled
+	bool bCanCache = Cache.IsValid() && IsHttpEnabled();
 	FIoRequestImpl* NextRequest = ChunkRequest->DeqeueDispatcherRequests();
 	while (NextRequest)
 	{
@@ -993,25 +1021,25 @@ bool FOnDemandIoBackend::Resolve(FIoRequestImpl* Request)
 {
 	using namespace UE::Tasks;
 
-	if (!bEnabled)
-	{
-		return false;
-	}
-	
 	FOnDemandIoStore::FChunkInfo ChunkInfo = IoStore->GetChunkInfo(Request->ChunkId);
 	if (!ChunkInfo.IsValid())
 	{
 		return false;
 	}
 
-	if ((!GIoDispatcherBulkOptionalEnabled || !bEnableBulkOptional) && Request->ChunkId.GetChunkType() == EIoChunkType::OptionalBulkData)
-	{
-		return false;
+	FChunkRequestParams RequestParams = FChunkRequestParams::Create(Request, ChunkInfo);
+
+	if (IsHttpEnabled() == false || IsHttpEnabled(Request->ChunkId) == false)
+	{ 
+		// Allow reading from the cache only when HTTP streaming is disabled otherwise the chunk may be evicted
+		// before trying to read the cache entry.
+		if (IsHttpEnabled() || Cache.IsValid() == false || Cache->ContainsChunk(RequestParams.ChunkKey) == false)
+		{
+			return false;
+		}
 	}
 
 	Stats.OnIoRequestEnqueue();
-	FChunkRequestParams RequestParams = FChunkRequestParams::Create(Request, ChunkInfo);
-
 	bool bPending = false;
 	bool bUpdatePriority = false;
 	FChunkRequest* ChunkRequest = ChunkRequests.Create(Request, RequestParams, bPending, bUpdatePriority);
@@ -1048,8 +1076,10 @@ bool FOnDemandIoBackend::Resolve(FIoRequestImpl* Request)
 			}
 		}
 
-		if (ChunkRequest->CancellationToken.IsCancelled())
+		const bool bHttpDisabled = IsHttpEnabled() == false || IsHttpEnabled(ChunkRequest->GetChunkId()) == false;
+		if (bHttpDisabled || ChunkRequest->CancellationToken.IsCancelled())
 		{
+			UE_CLOG(bHttpDisabled, LogIas, Log, TEXT("Chunk was not found in the cache and HTTP is disabled"));
 			return CompleteRequest(ChunkRequest);
 		}
 
@@ -1081,28 +1111,56 @@ void FOnDemandIoBackend::UpdatePriorityForIoRequest(FIoRequestImpl* Request)
 
 bool FOnDemandIoBackend::DoesChunkExist(const FIoChunkId& ChunkId) const
 {
-	TIoStatusOr<uint64> ChunkSize = GetSizeForChunk(ChunkId);
+	const TIoStatusOr<uint64> ChunkSize = GetSizeForChunk(ChunkId);
+	return ChunkSize.IsOk();
+}
+
+bool FOnDemandIoBackend::DoesChunkExist(const FIoChunkId& ChunkId, const FIoOffsetAndLength& ChunkRange) const
+{
+	uint64 Unused = 0;
+	const TIoStatusOr<uint64> ChunkSize = GetSizeForChunk(ChunkId, ChunkRange, Unused);
 	return ChunkSize.IsOk();
 }
 
 TIoStatusOr<uint64> FOnDemandIoBackend::GetSizeForChunk(const FIoChunkId& ChunkId) const
 {
-	if (IoStore.IsValid())
+	uint64 Unused = 0;
+	const FIoOffsetAndLength ChunkRange(0, MAX_uint64);
+	return GetSizeForChunk(ChunkId, ChunkRange, Unused);
+}
+
+TIoStatusOr<uint64> FOnDemandIoBackend::GetSizeForChunk(const FIoChunkId& ChunkId, const FIoOffsetAndLength& ChunkRange, uint64& OutAvailable) const
+{
+	OutAvailable = 0;
+
+	const FOnDemandIoStore::FChunkInfo ChunkInfo = IoStore.IsValid() ? IoStore->GetChunkInfo(ChunkId) : FOnDemandIoStore::FChunkInfo();
+	if (ChunkInfo.IsValid() == false)
 	{
-		if (!bEnabled)
+		return FIoStatus(EIoErrorCode::UnknownChunkID);
+	}
+
+	FIoOffsetAndLength RequestedRange(ChunkRange.GetOffset(), FMath::Min(ChunkInfo.Entry->RawSize, ChunkRange.GetLength()));
+	OutAvailable = ChunkInfo.Entry->RawSize;
+
+	if (IsHttpEnabled() == false || IsHttpEnabled(ChunkId) == false)
+	{
+		if (IsHttpEnabled() || Cache.IsValid() == false)
 		{
 			return FIoStatus(EIoErrorCode::UnknownChunkID);
 		}
 
-		if ((!GIoDispatcherBulkOptionalEnabled || !bEnableBulkOptional) && ChunkId.GetChunkType() == EIoChunkType::OptionalBulkData)
+		// When HTTP streaming is disabled, no cache entries will be evicted. 
+		const FChunkRequestParams RequestParams = FChunkRequestParams::Create(RequestedRange, ChunkInfo);
+		if (Cache->ContainsChunk(RequestParams.ChunkKey) == false)
 		{
 			return FIoStatus(EIoErrorCode::UnknownChunkID);
 		}
 
-		return IoStore->GetChunkSize(ChunkId);
+		// Only the specified chunk range is available 
+		OutAvailable = RequestedRange.GetLength();
 	}
 
-	return FIoStatus(EIoErrorCode::UnknownChunkID);
+	return TIoStatusOr<uint64>(ChunkInfo.Entry->RawSize);
 }
 
 FIoRequestImpl* FOnDemandIoBackend::GetCompletedRequests()
@@ -1204,14 +1262,14 @@ void FOnDemandIoBackend::Mount(const FOnDemandEndpoint& Endpoint)
 void FOnDemandIoBackend::SetBulkOptionalEnabled(bool bInEnabled)
 {
 	UE_LOG(LogIas, Log, TEXT("HTTP optional bulk data streaming '%s'"), bInEnabled ? TEXT("Enabled") : TEXT("Disabled"));
-	bEnableBulkOptional = bInEnabled;
+	bHttpOptionalBulkDataEnabled = bInEnabled;
 }
 
 void FOnDemandIoBackend::SetEnabled(bool bInEnabled)
 {
 	UE_LOG(LogIas, Log, TEXT("HTTP streaming '%s'"), bInEnabled ? TEXT("Enabled") : TEXT("Disabled"));
-	bEnabled = bInEnabled;
-	FGenericCrashContext::SetEngineData(TEXT("IAS.Enabled"), bEnabled ? TEXT("true") : TEXT("false"));
+	bHttpEnabled = bInEnabled;
+	FGenericCrashContext::SetEngineData(TEXT("IAS.Enabled"), bHttpEnabled.load(std::memory_order_acquire) ? TEXT("true") : TEXT("false"));
 }
 
 void FOnDemandIoBackend::ReportAnalytics(TArray<FAnalyticsEventAttribute>& OutAnalyticsArray) const
@@ -1221,7 +1279,7 @@ void FOnDemandIoBackend::ReportAnalytics(TArray<FAnalyticsEventAttribute>& OutAn
 		return;
 	}
 
-	if (bEnabled)
+	if (HttpClient.IsValid())
 	{
 		AppendAnalyticsEventAttributeArray(OutAnalyticsArray, TEXT("IasCDNBackend"), HttpClient->ServiceUrl());
 
@@ -1239,12 +1297,12 @@ TArray<FIoChunkId> FOnDemandIoBackend::GetAllChunkIds()
 {
 	TRACE_CPUPROFILER_EVENT_SCOPE(FOnDemandIoBackend::GetAllChunkIds);
 
-	if (!IoStore.IsValid() || !bEnabled)
+	if (!IoStore.IsValid() || !IsHttpEnabled())
 	{
 		return TArray<FIoChunkId>();
 	}
 
-	const bool bAllowOptional = GIoDispatcherBulkOptionalEnabled && bEnableBulkOptional;
+	const bool bAllowOptional = GIasHttpOptionalBulkDataEnabled && bHttpOptionalBulkDataEnabled;
 
 	return IoStore->GetAllChunkIds(bAllowOptional);
 }
@@ -1255,9 +1313,9 @@ TIoStatusOr<FOnDemandToc> FOnDemandIoBackend::GetToc(const FOnDemandEndpoint& En
 {
 	TRACE_CPUPROFILER_EVENT_SCOPE(FOnDemandIoBackend::GetToc);
 
-	for (int32 Attempt = 0, MaxAttempts = GIoDispatcherMaxHttpRetryCount; Attempt <= MaxAttempts; ++Attempt)
+	for (int32 Attempt = 0, MaxAttempts = GIasMaxHttpRetryCount; Attempt <= MaxAttempts; ++Attempt)
 	{
-		TUniquePtr<FOnDemandHttpClient> HttpClient = MakeUnique<FOnDemandHttpClient>(Endpoint.ServiceUrl, GIoDispatcherMaxHttpConnectionCount);
+		TUniquePtr<FOnDemandHttpClient> HttpClient = MakeUnique<FOnDemandHttpClient>(Endpoint.ServiceUrl, GIasMaxHttpConnectionCount);
 		TAnsiStringBuilder<256> Url;
 
 		Url << "/" << Endpoint.TocPath;
@@ -1289,7 +1347,8 @@ TIoStatusOr<FOnDemandToc> FOnDemandIoBackend::GetToc(const FOnDemandEndpoint& En
 			}
 		});
 
-		while (HttpClient->Tick());
+		const bool bBlock = true;
+		while (HttpClient->Tick(bBlock));
 
 		if (Toc.IsOk())
 		{
@@ -1312,7 +1371,7 @@ FIoStatus FOnDemandIoBackend::AddToc(const FOnDemandEndpoint& Endpoint)
 		FWriteScopeLock _(Lock);
 		if (!HttpClient.IsValid())
 		{
-			HttpClient = MakeUnique<FOnDemandHttpClient>(Endpoint.ServiceUrl, GIoDispatcherMaxHttpConnectionCount);
+			HttpClient = MakeUnique<FOnDemandHttpClient>(Endpoint.ServiceUrl, GIasMaxHttpConnectionCount);
 			BackendThread.Reset(FRunnableThread::Create(this, TEXT("IoStoreOnDemand"), 0, TPri_AboveNormal));
 		}
 	}
@@ -1366,7 +1425,7 @@ uint32 FOnDemandIoBackend::Run()
 
 							if (Status.Status().GetErrorCode() == EIoErrorCode::ReadError)
 							{
-								if (++ChunkRequest->HttpRetryCount <= GIoDispatcherMaxHttpRetryCount)
+								if (++ChunkRequest->HttpRetryCount <= GIasMaxHttpRetryCount)
 								{
 									Stats.OnHttpRetry();
 									Stats.OnHttpEnqueue();
