Build fix - multiple console variables with the same name

#preflght 64e760d745616e2a9f72348f
#rnx

diff --git a/Engine/Source/Runtime/Experimental/IoStoreOnDemand/Private/OnDemandHttpClient.cpp b/Engine/Source/Runtime/Experimental/IoStoreOnDemand/Private/OnDemandHttpClient.cpp
--- a/Engine/Source/Runtime/Experimental/IoStoreOnDemand/Private/OnDemandHttpClient.cpp
+++ b/Engine/Source/Runtime/Experimental/IoStoreOnDemand/Private/OnDemandHttpClient.cpp
@@ -11,10 +11,10 @@
 namespace UE::IO::IAS
 {
 
-int32 GIoDispatcherHttpPollTimeoutMs = 0;
-static FAutoConsoleVariableRef CVar_IoDispatcherMaxHttpPollTimeoutMs(
+int32 GIasHttpPollTimeoutMs = 0;
+static FAutoConsoleVariableRef CVar_GIasHttpPollTimeoutMs(
 	TEXT("ias.HttpPollTimeout"),
-	GIoDispatcherHttpPollTimeoutMs,
+	GIasHttpPollTimeoutMs,
 	TEXT("FHttpClient::Tick poll timeout in milliseconds")
 );
 
@@ -107,7 +107,7 @@ void FOnDemandHttpClient::Issue(FAnsiStringView Url, FGetCallback&& Callback, FI
 
 bool FOnDemandHttpClient::Tick(bool Block)
 {
-	int32 TimeoutMs = Block ? -1 : GIoDispatcherHttpPollTimeoutMs;
+	int32 TimeoutMs = Block ? -1 : GIasHttpPollTimeoutMs;
 	return EventLoop.Tick(TimeoutMs) != 0;
 }
 
diff --git a/Engine/Source/Runtime/Experimental/IoStoreOnDemand/Private/OnDemandIoDispatcherBackend.cpp b/Engine/Source/Runtime/Experimental/IoStoreOnDemand/Private/OnDemandIoDispatcherBackend.cpp
--- a/Engine/Source/Runtime/Experimental/IoStoreOnDemand/Private/OnDemandIoDispatcherBackend.cpp
+++ b/Engine/Source/Runtime/Experimental/IoStoreOnDemand/Private/OnDemandIoDispatcherBackend.cpp
@@ -54,13 +54,6 @@ static FAutoConsoleVariableRef CVar_IasMaxHttpRetryCount(
 	TEXT("Max number of HTTP request retries before failing the I/O request.")
 );
 
-int32 GIasHttpPollTimeoutMs = 0;
-static FAutoConsoleVariableRef CVar_IasMaxHttpPollTimeoutMs(
-	TEXT("ias.HttpPollTimeout"),
-	GIasHttpPollTimeoutMs,
-	TEXT("Tick() poll timeout in milliseconds")
-);
-
 bool GIasHttpEnabled = true;
 static FAutoConsoleVariableRef CVar_IasHttpEnabled(
 	TEXT("ias.HttpEnabled"),
