Chaos
- fix logic error in convex-triangle collision against edges
- fix mesh contact reduction removing points that should be kept (FContactTriangleCollector was incorreectly using contacts with a positive separation to cull other points with penetration)

diff --git a/Engine/Source/Runtime/Experimental/Chaos/Private/Chaos/Collision/ContactTriangles.cpp b/Engine/Source/Runtime/Experimental/Chaos/Private/Chaos/Collision/ContactTriangles.cpp
--- a/Engine/Source/Runtime/Experimental/Chaos/Private/Chaos/Collision/ContactTriangles.cpp
+++ b/Engine/Source/Runtime/Experimental/Chaos/Private/Chaos/Collision/ContactTriangles.cpp
@@ -11,6 +11,7 @@ namespace Chaos
 	extern bool bChaos_Collision_EnableEdgePrune;
 	extern int32 Chaos_Collision_ReduceMeshManifoldThreshold;
 	extern FRealSingle Chaos_Collision_MeshContactNormalThreshold;
+	extern bool bChaos_Collision_MeshManifoldSortByDistance;
 
 
 	void FContactTriangleCollector::SetNumContacts(const int32 NumContacts)
@@ -150,15 +151,22 @@ namespace Chaos
 			const FContactPoint& ContactPoint = TriangleContactPoints[ContactIndex];
 			const FContactTriangle& ContactTriangle = ContactTriangles[ContactPointData.GetTriangleIndex()];
 
+			// We only use actually penetrating contacts when determining which faces and edges we hit, otherwise
+			// we may reject contacts that are actually required based on non-contacts.
+			const bool bUseForPruning = (ContactPoint.Phi < 0);
+
 			if (ContactPoint.ContactType == EContactPointType::VertexPlane)
 			{
-				// Since we collided with this face, we cannot collide with these edges or vertices
-				ContactEdges.Add(FContactEdgeID(ContactTriangle.VertexIndices[0], ContactTriangle.VertexIndices[1]));
-				ContactEdges.Add(FContactEdgeID(ContactTriangle.VertexIndices[1], ContactTriangle.VertexIndices[2]));
-				ContactEdges.Add(FContactEdgeID(ContactTriangle.VertexIndices[2], ContactTriangle.VertexIndices[0]));
-				ContactVertices.Add(FContactVertexID(ContactTriangle.VertexIndices[0]));
-				ContactVertices.Add(FContactVertexID(ContactTriangle.VertexIndices[1]));
-				ContactVertices.Add(FContactVertexID(ContactTriangle.VertexIndices[2]));
+				if (bUseForPruning)
+				{
+					// Since we collided with this face, we cannot collide with these edges or vertices
+					ContactEdges.Add(FContactEdgeID(ContactTriangle.VertexIndices[0], ContactTriangle.VertexIndices[1]));
+					ContactEdges.Add(FContactEdgeID(ContactTriangle.VertexIndices[1], ContactTriangle.VertexIndices[2]));
+					ContactEdges.Add(FContactEdgeID(ContactTriangle.VertexIndices[2], ContactTriangle.VertexIndices[0]));
+					ContactVertices.Add(FContactVertexID(ContactTriangle.VertexIndices[0]));
+					ContactVertices.Add(FContactVertexID(ContactTriangle.VertexIndices[1]));
+					ContactVertices.Add(FContactVertexID(ContactTriangle.VertexIndices[2]));
+				}
 			}
 			else if ((ContactPoint.ContactType == EContactPointType::EdgeEdge) || (ContactPoint.ContactType == EContactPointType::PlaneVertex))
 			{
@@ -168,9 +176,12 @@ namespace Chaos
 				{
 					if ((VertexIndexA != INDEX_NONE) && (VertexIndexB != INDEX_NONE))
 					{
-						// Since we collided with this edge, we cannot also collide with these verts
-						ContactVertices.Add(FContactVertexID(VertexIndexA));
-						ContactVertices.Add(FContactVertexID(VertexIndexB));
+						if (bUseForPruning)
+						{
+							// Since we collided with this edge, we cannot also collide with these verts
+							ContactVertices.Add(FContactVertexID(VertexIndexA));
+							ContactVertices.Add(FContactVertexID(VertexIndexB));
+						}
 
 						// Store the edge ID so we can see if any other faces add it to the don'tcollide list
 						ContactPointData.SetEdgeID(VertexIndexA, VertexIndexB);
@@ -190,11 +201,11 @@ namespace Chaos
 		// @todo(chaos): Some of the pruning mechanisms don't work very well now that we have support for large manifolds.
 		const bool bEnableFullPruning = ((TriangleContactPoints.Num() - NumDisabledTriangleContactPoints) < Chaos_Collision_ReduceMeshManifoldThreshold);
 
-		DebugDrawContactPoints(FColor::White, 0.2);
-
 		// Build the feature set from the full list of contacts before pruning
 		BuildContactFeatureSets();
 
+		DebugDrawContactPoints(FColor::White, 0.2);
+
 		// Remove edge contacts that should not be possible because we are also colliding with a face that used the edge.
 		if ((TriangleContactPoints.Num() > 1))
 		{
@@ -203,7 +214,7 @@ namespace Chaos
 			PruneEdgeAndVertexContactPoints(bPruneEdges, bPruneVertices);
 		}
 
-		DebugDrawContactPoints(FColor::Cyan, 0.35);
+		DebugDrawContactPoints(FColor::Cyan, 0.3);
 
 		// Remove contacts that should not be able to occur if we are a one-sided mesh
 		// NOTE: does not require sorted contacts, and does not maintain array order
@@ -212,12 +223,16 @@ namespace Chaos
 			PruneInfacingContactPoints();
 		}
 
+		DebugDrawContactPoints(FColor::Purple, 0.5);
+
 		// Remove close contacts
 		if (TriangleContactPoints.Num() > 1)
 		{
 			PruneUnnecessaryContactPoints();
 		}
 
+		DebugDrawContactPoints(FColor::Yellow, 0.7);
+
 		// Reduce to only 4 contact points from here
 		// @todo(chaos): A 4 point manifold is often not sufficient when we have low curvature and we end up removing contacts that
 		// would be required during the solve to prevent jitter. For now we disable the reduction when we are colliding
@@ -225,12 +240,10 @@ namespace Chaos
 		// smaller objects so we don't impact perf for the most common cases until we do this properly.
 		if (bEnableFullPruning)
 		{
-			DebugDrawContactPoints(FColor::Yellow, 0.5);
-
 			ReduceManifoldContactPointsTriangeMesh();
 		}
 
-		DebugDrawContactPoints(FColor::Green, 0.75);
+		DebugDrawContactPoints(FColor::Green, 0.9);
 
 		// Fix contact normals for edge and vertex collisions
 		if (TriangleContactPoints.Num() > 0)
@@ -297,19 +310,65 @@ namespace Chaos
 			// Reset metadata to ensure we don't access it again
 			TriangleContactPointDatas.Reset();
 
-			// Sort by contact type (in enum order) then phi
-			std::sort(
-				&TriangleContactPoints[0],
-				&TriangleContactPoints[0] + TriangleContactPoints.Num(),
-				[](const FContactPoint& L, const FContactPoint& R)
+			if (bChaos_Collision_MeshManifoldSortByDistance)
+			{
+				// Sort contact points by distance from the center of mass (RxN) so that points closer to the center of
+				// mass are solved first. This produces better solver results for low iterations because, if we were to 
+				// solve the distant points first, we would get extra rotation applied. 
+				//
+				// E.g., consider a box landing on an inclined plane with 5 contact points biassed toward one side. 
+				//
+				// -------------------
+				// |                 |
+				// |                 |
+				// *-*-*-*-*----------
+				//
+				// Solving this left to right would result in extra clockwise rotation after 1 iteration. A subsequent
+				// iteration would partially correct the problem.
+				//
+				TArray<TPair<FReal, int32>> SortKeyValues;
+				SortKeyValues.SetNumUninitialized(TriangleContactPoints.Num());
+				for (int32 ContactIndex = 0; ContactIndex < TriangleContactPoints.Num(); ++ContactIndex)
 				{
-					if (L.ContactType == R.ContactType)
+					const FContactPoint& ContactPoint = TriangleContactPoints[ContactIndex];
+					const FVec3 DeltaTangent = ContactPoint.ShapeContactPoints[1] - FVec3::DotProduct(ContactPoint.ShapeContactPoints[1], ContactPoint.ShapeContactNormal) * ContactPoint.ShapeContactNormal;
+					const FReal DeltaTangentLenSq = DeltaTangent.SizeSquared();
+					SortKeyValues[ContactIndex] = { DeltaTangentLenSq, ContactIndex };
+				}
+
+				std::sort(
+					&SortKeyValues[0],
+					&SortKeyValues[0] + SortKeyValues.Num(),
+					[](const TPair<FReal, int32>& L, const TPair<FReal, int32>& R)
 					{
-						return L.Phi < R.Phi;
+						return L.Key < R.Key;
 					}
-					return L.ContactType < R.ContactType;
+				);
+
+				TArray<FContactPoint> SortedContactPoints;
+				SortedContactPoints.SetNumUninitialized(TriangleContactPoints.Num());
+				for (int32 ContactIndex = 0; ContactIndex < TriangleContactPoints.Num(); ++ContactIndex)
+				{
+					SortedContactPoints[ContactIndex] = TriangleContactPoints[SortKeyValues[ContactIndex].Value];
 				}
-			);
+				Swap(SortedContactPoints, TriangleContactPoints);
+			}
+			else
+			{
+				// Sort by contact type (in enum order) then depth (deepest first)
+				std::sort(
+					&TriangleContactPoints[0],
+					&TriangleContactPoints[0] + TriangleContactPoints.Num(),
+					[](const FContactPoint& L, const FContactPoint& R)
+					{
+						if (L.ContactType == R.ContactType)
+						{
+							return L.Phi < R.Phi;
+						}
+						return L.ContactType < R.ContactType;
+					}
+				);
+			}
 		}
 	}
 
@@ -389,28 +448,39 @@ namespace Chaos
 							const FReal ContactDotNormal = FVec3::DotProduct(ContactPoint.ShapeContactNormal, ContactTriangle.FaceNormal);
 							if (ContactDotNormal < MinContactDotNormal)
 							{
-								// We are outside the valid normal range for this edge
-								// Convert the edge collision to a face collision on one of the faces, selected to get the smallest depth
-								// NOTE: We keep the contact depth as it is because we know that the depth is a lower-bound. 
-								// We have to update the ShapeContactPoint[0] because Phi is actually derived from the positions (the value in ContactPoint is just a cache of current state)
-								// @todo(chaos): face selection logic might be better if we knew the contact velocity
-								const FReal OtherContactDotNormal = FVec3::DotProduct(ContactPoint.ShapeContactNormal, ContactTriangle.FaceNormal);		// @todo(chaos): OtherContactTriangle??
-								if (ContactDotNormal >= OtherContactDotNormal)
+								// If we are separated, just ignore this contact
+								if (ContactPoint.Phi > 0)
 								{
-									ContactPoint.ShapeContactNormal = (ContactDotNormal > -SMALL_NUMBER) ? ContactTriangle.FaceNormal : -ContactTriangle.FaceNormal;
-									ContactPoint.ShapeContactPoints[0] = ContactPoint.ShapeContactPoints[1] + ContactPoint.Phi * ContactPoint.ShapeContactNormal;
+									DisableContact(ContactIndex);
 								}
 								else
 								{
-									ContactPoint.ShapeContactNormal = (OtherContactDotNormal > -SMALL_NUMBER) ? OtherContactTriangle.FaceNormal : -OtherContactTriangle.FaceNormal;
-									ContactPoint.ShapeContactPoints[0] = ContactPoint.ShapeContactPoints[1] + ContactPoint.Phi * ContactPoint.ShapeContactNormal;
-									ContactPointData.SetTriangleIndex(OtherContactTriangleIndex);
+									// We are outside the valid normal range for this edge
+									// Convert the edge collision to a face collision on one of the faces, selected to get the smallest depth
+									// NOTE: We keep the contact depth as it is because we know that the depth is a lower-bound. 
+									// We have to update the ShapeContactPoint[0] because Phi is actually derived from the positions (the value in ContactPoint is just a cache of current state)
+									// @todo(chaos): face selection logic might be better if we knew the contact velocity
+									const FReal OtherContactDotNormal = FVec3::DotProduct(ContactPoint.ShapeContactNormal, OtherContactTriangle.FaceNormal);
+									if (ContactDotNormal >= OtherContactDotNormal)
+									{
+										ContactPoint.ShapeContactNormal = (ContactDotNormal > -SMALL_NUMBER) ? ContactTriangle.FaceNormal : -ContactTriangle.FaceNormal;
+										ContactPoint.ShapeContactPoints[0] = ContactPoint.ShapeContactPoints[1] + ContactPoint.Phi * ContactPoint.ShapeContactNormal;
+									}
+									else
+									{
+										ContactPoint.ShapeContactNormal = (OtherContactDotNormal > -SMALL_NUMBER) ? OtherContactTriangle.FaceNormal : -OtherContactTriangle.FaceNormal;
+										ContactPoint.ShapeContactPoints[0] = ContactPoint.ShapeContactPoints[1] + ContactPoint.Phi * ContactPoint.ShapeContactNormal;
+										ContactPointData.SetTriangleIndex(OtherContactTriangleIndex);
+									}
 								}
 							}
 						}
 					}
 					else if (VertexIndexA != INDEX_NONE)
 					{
+						// If VertexIndexA is valid (VertexIndexB is invalid) we have a vertex collision.
+						// Ensure that the contact normal is in a valid range for the vertex based on the
+						// triangles that share the vertex.
 						for (int32 TriangleIndex = 0; TriangleIndex < ContactTriangles.Num(); ++TriangleIndex)
 						{
 							const FContactTriangle& OtherContactTriangle = ContactTriangles[TriangleIndex];
@@ -430,7 +500,14 @@ namespace Chaos
 										const FReal NormalDotCentroid = FVec3::DotProduct(ContactPoint.ShapeContactNormal, Centroid - ContactPoint.ShapeContactPoints[1]);
 										if (NormalDotCentroid > 0)
 										{
-											const FReal OtherContactDotNormal = FVec3::DotProduct(ContactPoint.ShapeContactNormal, ContactTriangle.FaceNormal);	// @todo(chaos): OtherContactTriangle??
+											// If we are separated, just ignore this contact
+											if (ContactPoint.Phi > 0)
+											{
+												DisableContact(ContactIndex);
+												break;
+											}
+
+											const FReal OtherContactDotNormal = FVec3::DotProduct(ContactPoint.ShapeContactNormal, OtherContactTriangle.FaceNormal);
 											ContactPoint.ShapeContactNormal = (OtherContactDotNormal > -SMALL_NUMBER) ? OtherContactTriangle.FaceNormal : -OtherContactTriangle.FaceNormal;
 											ContactPoint.ShapeContactPoints[0] = ContactPoint.ShapeContactPoints[1] + ContactPoint.Phi * ContactPoint.ShapeContactNormal;
 										}
@@ -502,7 +579,7 @@ namespace Chaos
 
 				// Are these two points at the same position? They must be at the same vertex or edge for this to be true.
 				// If both points are from the same triangle vertex, we know they are coincident (but normals may be different)
-				const bool bIsSameVertex = ContactPointData0.IsVertex() && (ContactPointData0.GetVertexID() == ContactPointData1.GetVertexID());
+				const bool bIsSameVertex = (ContactPointData0.IsVertex() && ContactPointData1.IsVertex()) && (ContactPointData0.GetVertexID() == ContactPointData1.GetVertexID());
 				bool bIsSamePosition = bIsSameVertex;
 				if (!bIsSameVertex)
 				{
@@ -515,10 +592,12 @@ namespace Chaos
 
 				if (bIsSamePosition)
 				{
-					// Keep the contact with the normal that is most along the triangle face normal
+					// Keep the contact with the normal that is most along the triangle face normal. If normals are similar, keep the deepest.
 					const FReal NormalDot0 = FVec3::DotProduct(ContactPoint0.ShapeContactNormal, ContactTriangle0.FaceNormal);
 					const FReal NormalDot1 = FVec3::DotProduct(ContactPoint1.ShapeContactNormal, ContactTriangle1.FaceNormal);
-					if (FMath::Abs(NormalDot0) >= FMath::Abs(NormalDot1))
+					const bool bSameNormal = (FMath::Abs(NormalDot0 - NormalDot1) < 0.01);
+					const bool bTake0 = (bSameNormal && (ContactPoint0.Phi < ContactPoint1.Phi)) || (!bSameNormal && (NormalDot0 > NormalDot1));
+					if (bTake0)
 					{
 						DisableContact(ContactIndex1);
 						continue;
diff --git a/Engine/Source/Runtime/Experimental/Chaos/Private/Chaos/Collision/ConvexTriangleContactPoint.cpp b/Engine/Source/Runtime/Experimental/Chaos/Private/Chaos/Collision/ConvexTriangleContactPoint.cpp
--- a/Engine/Source/Runtime/Experimental/Chaos/Private/Chaos/Collision/ConvexTriangleContactPoint.cpp
+++ b/Engine/Source/Runtime/Experimental/Chaos/Private/Chaos/Collision/ConvexTriangleContactPoint.cpp
@@ -133,6 +133,7 @@ namespace Chaos
 	{
 		const FReal NormalTolerance = FReal(1.e-8);
 		const FReal NormalToleranceSq = NormalTolerance * NormalTolerance;
+		const FReal InvalidPhi = std::numeric_limits<FReal>::lowest();
 
 		// Triangle (same space as convex)
 		const FVec3 TriN = Triangle.GetNormal();
@@ -173,7 +174,7 @@ namespace Chaos
 		// and reject if the separation is more than cull distance.
 		FVec3 ConvexPlaneN = FVec3(0);
 		FVec3 ConvexPlaneX = FVec3(0);
-		FReal ConvexPlaneDMin = std::numeric_limits<FReal>::lowest();
+		FReal ConvexPlaneDMin = InvalidPhi;
 		int32 ConvexPlaneIndexMin = INDEX_NONE;
 		for (int32 PlaneIndex = 0; PlaneIndex < Convex.NumPlanes(); ++PlaneIndex)
 		{
@@ -212,13 +213,13 @@ namespace Chaos
 		FReal TriVertexConvexDMin2 = FVec3::DotProduct(Triangle.GetVertex(2) - ConvexPlaneX, ConvexPlaneN);
 		FReal TriEdgeConvexDMin[3] =
 		{
-			FMath::Max(TriVertexConvexDMin2, TriVertexConvexDMin0),
-			FMath::Max(TriVertexConvexDMin0, TriVertexConvexDMin1),
-			FMath::Max(TriVertexConvexDMin1, TriVertexConvexDMin2),
+			FMath::Min(TriVertexConvexDMin2, TriVertexConvexDMin0),
+			FMath::Min(TriVertexConvexDMin0, TriVertexConvexDMin1),
+			FMath::Min(TriVertexConvexDMin1, TriVertexConvexDMin2),
 		};
 
 		FVec3 EdgeEdgeN = FVec3(0);
-		FReal EdgeEdgeDMin = std::numeric_limits<FReal>::lowest();
+		FReal EdgeEdgeDMin = InvalidPhi;
 		int32 ConvexEdgeIndexMin = INDEX_NONE;
 		int32 TriEdgeIndexMin = INDEX_NONE;
 		for (int32 ConvexEdgeIndex = 0; ConvexEdgeIndex < Convex.NumEdges(); ++ConvexEdgeIndex)
@@ -302,16 +303,39 @@ namespace Chaos
 			}
 		}
 
+		// Determine which of the features we want to use
+		// NOTE: we rely on the fact that all valid Phi values are greater than InvalidPhi here
+		const FReal TriFaceBias = FReal(1.e-2);	// Prevent flip=flip on near parallel cases
+		EContactPointType ContactType = EContactPointType::Unknown;
+		if ((TriPlaneDMin != InvalidPhi) && (TriPlaneDMin + TriFaceBias > ConvexPlaneDMin) && (TriPlaneDMin + TriFaceBias > EdgeEdgeDMin))
+		{
+			// Tri plane is the shallowest penetration
+			ContactType = EContactPointType::VertexPlane;
+		}
+		else if ((ConvexPlaneDMin != InvalidPhi) && (ConvexPlaneDMin > EdgeEdgeDMin))
+		{
+			// Convex plane is the shallowest penetration
+			ContactType = EContactPointType::PlaneVertex;
+		}
+		else if (EdgeEdgeDMin != InvalidPhi)
+		{
+			// Edge-edge is the shallowest penetration
+			ContactType = EContactPointType::EdgeEdge;
+		}
+		else
+		{
+			// No valid features (should not happen - TriPlaneDMin should always be valid)
+			return;
+		}
+
 		// Determine the best features to use for this collision
 		FVec3 SeparatingAxis, ClipAxis;
-		EContactPointType ContactType;
 		bool bClipConvexToTriangle;
 		bool bClipToFaceNormal;
-		const FReal TriFaceBias = FReal(1.e-2);	// Prevent flip=flip on near parallel cases
-		if ((TriPlaneDMin + TriFaceBias > ConvexPlaneDMin) && (TriPlaneDMin + TriFaceBias > EdgeEdgeDMin))
+		if (ContactType == EContactPointType::VertexPlane)
 		{
 			// Triangle face contact - clip the convex vertices to the triangle
-			ContactType = EContactPointType::VertexPlane;
+			
 			bClipConvexToTriangle = true;
 			bClipToFaceNormal = true;
 
@@ -329,10 +353,9 @@ namespace Chaos
 			Convex.GetPlaneNX(ConvexPlaneIndexMin, ConvexPlaneN, ConvexPlaneX);
 
 		}
-		else if (ConvexPlaneDMin > EdgeEdgeDMin)
+		else if (ContactType == EContactPointType::PlaneVertex)
 		{
 			// Convex face contact - clip the triangle to the convex face
-			ContactType = EContactPointType::PlaneVertex;
 			bClipConvexToTriangle = false;
 			bClipToFaceNormal = true;
 
@@ -345,10 +368,9 @@ namespace Chaos
 
 			ClipAxis = ConvexPlaneN;
 		}
-		else
+		else if (ContactType == EContactPointType::EdgeEdge)
 		{
 			// Edge-edge contact - clip triangle vs convex or vice-versa based on most opposing normals
-			ContactType = EContactPointType::EdgeEdge;
 
 			// The separating axis must point from the triangle to the convex
 			SeparatingAxis = EdgeEdgeN;
@@ -383,6 +405,10 @@ namespace Chaos
 			}
 			bClipToFaceNormal = false;
 		}
+		else
+		{
+			return;
+		}
 
 		// @todo(chaos): scratch or stack allocation of clipped vertex buffers
 		const int32 MaxContactPointCount = 32;
diff --git a/Engine/Source/Runtime/Experimental/Chaos/Private/Chaos/CollisionOneShotManifolds.cpp b/Engine/Source/Runtime/Experimental/Chaos/Private/Chaos/CollisionOneShotManifolds.cpp
--- a/Engine/Source/Runtime/Experimental/Chaos/Private/Chaos/CollisionOneShotManifolds.cpp
+++ b/Engine/Source/Runtime/Experimental/Chaos/Private/Chaos/CollisionOneShotManifolds.cpp
@@ -88,6 +88,9 @@ namespace Chaos
 	FRealSingle Chaos_Collision_MeshContactNormalThreshold = 0.98f;	// ~10deg
 	FAutoConsoleVariableRef CVarChaos_Collision_MeshContactNormalThreshold(TEXT("p.Chaos.Collision.MeshContactNormalThreshold"), Chaos_Collision_MeshContactNormalThreshold, TEXT("Treat contact with a dot product between the normal and the triangle face greater than this as face collisions"));
 
+	bool bChaos_Collision_MeshManifoldSortByDistance = false;
+	FAutoConsoleVariableRef CVarChaos_Collision_LargeMeshManifoldSortByDistance(TEXT("p.Chaos.Collision.SortMeshManifoldByDistance"), bChaos_Collision_MeshManifoldSortByDistance, TEXT("Sort large mesh manifold points by |RxN| for improved solver stability (less rotation in first iteration)"));
+
 	// Whether to use the new index-less GJK. 
 	// @todo(chaos): This should be removed once soaked for a bit (enabled 7 June 2022)
 	bool bChaos_Collision_UseGJK2 = false;
