Geometry collection optimization : Deprecated collision related attributes in dynamic collection
- Saves 3 x uint32 per bone per dynamic collection ( x2 )
- CollisionMask and CollisionStructureID were not used at all
- CollisionGroup was used but set to the same value for all bones
- also deprecated fuinction and enums in GeometryCollectionEngineTypes

diff --git a/Engine/Source/Runtime/Experimental/Chaos/Private/GeometryCollectionProxyData.cpp b/Engine/Source/Runtime/Experimental/Chaos/Private/GeometryCollectionProxyData.cpp
--- a/Engine/Source/Runtime/Experimental/Chaos/Private/GeometryCollectionProxyData.cpp
+++ b/Engine/Source/Runtime/Experimental/Chaos/Private/GeometryCollectionProxyData.cpp
@@ -36,8 +36,6 @@ void FTransformDynamicCollection::Construct()
 */
 
 const FName FGeometryDynamicCollection::ActiveAttribute("Active");
-const FName FGeometryDynamicCollection::CollisionGroupAttribute("CollisionGroup");
-const FName FGeometryDynamicCollection::CollisionMaskAttribute("CollisionMask");
 const FName FGeometryDynamicCollection::DynamicStateAttribute("DynamicState");
 const FName FGeometryDynamicCollection::ImplicitsAttribute("Implicits");
 const FName FGeometryDynamicCollection::ShapesQueryDataAttribute("ShapesQueryData");
@@ -46,14 +44,15 @@ const FName FGeometryDynamicCollection::SimplicialsAttribute("CollisionParticles
 const FName FGeometryDynamicCollection::SimulatableParticlesAttribute("SimulatableParticlesAttribute");
 const FName FGeometryDynamicCollection::SharedImplicitsAttribute("SharedImplicits");
 
+// Deprecated
+const FName FGeometryDynamicCollection::CollisionMaskAttribute("CollisionMask");
+const FName FGeometryDynamicCollection::CollisionGroupAttribute("CollisionGroup");
+
 FGeometryDynamicCollection::FGeometryDynamicCollection()
 	: FTransformDynamicCollection()
 {
 	// Transform Group
 	AddExternalAttribute<bool>(FGeometryDynamicCollection::ActiveAttribute, FTransformCollection::TransformGroup, Active);
-	AddExternalAttribute<int32>(FGeometryDynamicCollection::CollisionGroupAttribute, FTransformCollection::TransformGroup, CollisionGroup);
-	AddExternalAttribute<int32>(FGeometryDynamicCollection::CollisionMaskAttribute, FTransformCollection::TransformGroup, CollisionMask);
-	AddExternalAttribute("CollisionStructureID", FTransformCollection::TransformGroup, CollisionStructureID);
 	AddExternalAttribute<int32>(FGeometryDynamicCollection::DynamicStateAttribute, FTransformCollection::TransformGroup, DynamicState);
 	AddExternalAttribute(ImplicitsAttribute, FTransformCollection::TransformGroup, Implicits);
 	AddExternalAttribute("MassToLocal", FTransformCollection::TransformGroup, MassToLocal);
diff --git a/Engine/Source/Runtime/Experimental/Chaos/Private/PhysicsProxy/GeometryCollectionPhysicsProxy.cpp b/Engine/Source/Runtime/Experimental/Chaos/Private/PhysicsProxy/GeometryCollectionPhysicsProxy.cpp
--- a/Engine/Source/Runtime/Experimental/Chaos/Private/PhysicsProxy/GeometryCollectionPhysicsProxy.cpp
+++ b/Engine/Source/Runtime/Experimental/Chaos/Private/PhysicsProxy/GeometryCollectionPhysicsProxy.cpp
@@ -789,7 +789,6 @@ void FGeometryCollectionPhysicsProxy::InitializeDynamicCollection(FGeometryDynam
 	KeepFromDynamicCollection.Add(FGeometryCollection::SimulationTypeAttribute);
 	KeepFromDynamicCollection.Add(DynamicCollection.SimplicialsAttribute);
 	KeepFromDynamicCollection.Add(DynamicCollection.ActiveAttribute);
-	KeepFromDynamicCollection.Add(DynamicCollection.CollisionGroupAttribute);
 	DynamicCollection.CopyMatchingAttributesFrom(RestCollection, &SkipList);
 
 
@@ -1034,7 +1033,6 @@ void FGeometryCollectionPhysicsProxy::InitializeBodiesPT(Chaos::FPBDRigidsSolver
 
 		const int32 NumTransforms = DynamicCollection.NumElements(FTransformCollection::TransformGroup);
 		const TManagedArray<int32>& DynamicState = DynamicCollection.DynamicState;
-		const TManagedArray<int32>& CollisionGroup = DynamicCollection.CollisionGroup;
 		const TManagedArray<bool>& SimulatableParticles = DynamicCollection.SimulatableParticles;
 		const TManagedArray<FTransform>& MassToLocal = DynamicCollection.MassToLocal;
 		const TManagedArray<Chaos::FImplicitObjectPtr>& Implicits = DynamicCollection.Implicits;
@@ -1108,7 +1106,7 @@ void FGeometryCollectionPhysicsProxy::InitializeBodiesPT(Chaos::FPBDRigidsSolver
 					ScaledInertia,
 					Transforms[TransformGroupIndex],
 					AdjustedDynamicState,
-					static_cast<int16>(CollisionGroup[TransformGroupIndex]),
+					static_cast<int16>(Parameters.CollisionGroup),
 					CollisionParticlesPerObjectFraction);
 
 				// initialize anchoring information if available 
@@ -1607,7 +1605,6 @@ Chaos::TPBDGeometryCollectionParticleHandle<Chaos::FReal, 3>* FGeometryCollectio
 {
 	FGeometryDynamicCollection& DynamicCollection = PhysicsThreadCollection;
 	TManagedArray<int32>& DynamicState = DynamicCollection.DynamicState;
-	const TManagedArray<int32>& CollisionGroup = DynamicCollection.CollisionGroup;
 	TManagedArray<int32>& ParentIndex = DynamicCollection.Parent;
 	TManagedArray<TSet<int32>>& Children = DynamicCollection.Children;
 	TManagedArray<FTransform>& Transform = DynamicCollection.Transform;
@@ -1669,8 +1666,8 @@ Chaos::TPBDGeometryCollectionParticleHandle<Chaos::FReal, 3>* FGeometryCollectio
 		Inertia,
 		ParticleTM,
 		static_cast<uint8>(DynamicState[CollectionClusterIndex]),
-		static_cast<int16>(CollisionGroup[CollectionClusterIndex]),
-		CollisionParticlesPerObjectFraction); // CollisionGroup
+		static_cast<int16>(Parameters.CollisionGroup),
+		CollisionParticlesPerObjectFraction);
 
 
 	// initialize anchoring information if available 
diff --git a/Engine/Source/Runtime/Experimental/Chaos/Public/GeometryCollectionProxyData.h b/Engine/Source/Runtime/Experimental/Chaos/Public/GeometryCollectionProxyData.h
--- a/Engine/Source/Runtime/Experimental/Chaos/Public/GeometryCollectionProxyData.h
+++ b/Engine/Source/Runtime/Experimental/Chaos/Public/GeometryCollectionProxyData.h
@@ -60,8 +60,6 @@ public:
 	typedef FTransformDynamicCollection Super;
 
 	static CHAOS_API const FName ActiveAttribute;
-	static CHAOS_API const FName CollisionGroupAttribute;
-	static CHAOS_API const FName CollisionMaskAttribute;
 	static CHAOS_API const FName DynamicStateAttribute;
 	static CHAOS_API const FName ImplicitsAttribute;
 	static CHAOS_API const FName ShapesQueryDataAttribute;
@@ -70,17 +68,30 @@ public:
 	static CHAOS_API const FName SimplicialsAttribute;
 	static CHAOS_API const FName SimulatableParticlesAttribute;
 
+	UE_DEPRECATED(5.4, "CollisionMaskAttribute is no longer supported")
+	static CHAOS_API const FName CollisionMaskAttribute;
+
+	UE_DEPRECATED(5.4, "CollisionGroupAttribute is no longer supported")
+	static CHAOS_API const FName CollisionGroupAttribute;
+
 	// Transform Group
 	TManagedArray<bool> Active;
-	TManagedArray<int32> CollisionGroup;
-	TManagedArray<int32> CollisionMask;
-	TManagedArray<int32> CollisionStructureID;
+	
 	TManagedArray<int32> DynamicState;
 	TManagedArray<Chaos::FImplicitObjectPtr> Implicits;
 	TManagedArray<FTransform> MassToLocal;
 	TManagedArray<TUniquePtr<FCollisionStructureManager::FSimplicial>> Simplicials;
 	TManagedArray<bool> SimulatableParticles;
 
+	UE_DEPRECATED(5.4, "CollisionStructureID attribute is no longer supported")
+	TManagedArray<int32> CollisionStructureID;
+
+	UE_DEPRECATED(5.4, "CollisionMask attribute is no longer supported")
+	TManagedArray<int32> CollisionMask;
+
+	UE_DEPRECATED(5.4, "CollisionGroup attribute is no longer supported - you can still set the collision group on the geometry collection")
+	TManagedArray<int32> CollisionGroup;
+
 public:
 	struct FInitialVelocityFacade
 	{
diff --git a/Engine/Source/Runtime/Experimental/GeometryCollectionEngine/Private/GeometryCollection/GeometryCollectionComponent.cpp b/Engine/Source/Runtime/Experimental/GeometryCollectionEngine/Private/GeometryCollection/GeometryCollectionComponent.cpp
--- a/Engine/Source/Runtime/Experimental/GeometryCollectionEngine/Private/GeometryCollection/GeometryCollectionComponent.cpp
+++ b/Engine/Source/Runtime/Experimental/GeometryCollectionEngine/Private/GeometryCollection/GeometryCollectionComponent.cpp
@@ -3001,14 +3001,6 @@ void UGeometryCollectionComponent::OnCreatePhysicsState()
 				}
 			}
 			
-			TManagedArray<int32> & CollisionGroupArray = DynamicCollection->CollisionGroup;
-			{
-				for (int i = 0; i < CollisionGroupArray.Num(); i++)
-				{
-					CollisionGroupArray[i] = CollisionGroup;
-				}
-			}
-
 			// there's a code path where Level is not serialized and InitializeSharedCollisionStructures is not being called,
 			// resulting in the attribute missing and causing a crash in CopyAttribute calls later in FGeometryCollectionPhysicsProxy::Initialize
 			// @todo(chaos) we should better handle computation of dependent attribute like level
diff --git a/Engine/Source/Runtime/Experimental/GeometryCollectionEngine/Private/GeometryCollection/GeometryCollectionEngineTypes.cpp b/Engine/Source/Runtime/Experimental/GeometryCollectionEngine/Private/GeometryCollection/GeometryCollectionEngineTypes.cpp
--- a/Engine/Source/Runtime/Experimental/GeometryCollectionEngine/Private/GeometryCollection/GeometryCollectionEngineTypes.cpp
+++ b/Engine/Source/Runtime/Experimental/GeometryCollectionEngine/Private/GeometryCollection/GeometryCollectionEngineTypes.cpp
@@ -2,7 +2,7 @@
 
 #include "GeometryCollection/GeometryCollectionEngineTypes.h"
 
-
+PRAGMA_DISABLE_DEPRECATION_WARNINGS
 FName GetCollectionAttributeName(ECollectionAttributeEnum TypeIn)
 {
 	switch (TypeIn)
@@ -18,7 +18,6 @@ FName GetCollectionAttributeName(ECollectionAttributeEnum TypeIn)
 	return FName("");
 }
 
-
 FName GetCollectionGroupName(ECollectionGroupEnum TypeIn)
 {
 	switch (TypeIn)
@@ -29,3 +28,4 @@ FName GetCollectionGroupName(ECollectionGroupEnum TypeIn)
 	check(false);
 	return FName("");
 }
+PRAGMA_ENABLE_DEPRECATION_WARNINGS
\ No newline at end of file
diff --git a/Engine/Source/Runtime/Experimental/GeometryCollectionEngine/Public/GeometryCollection/GeometryCollectionEngineTypes.h b/Engine/Source/Runtime/Experimental/GeometryCollectionEngine/Public/GeometryCollection/GeometryCollectionEngineTypes.h
--- a/Engine/Source/Runtime/Experimental/GeometryCollectionEngine/Public/GeometryCollection/GeometryCollectionEngineTypes.h
+++ b/Engine/Source/Runtime/Experimental/GeometryCollectionEngine/Public/GeometryCollection/GeometryCollectionEngineTypes.h
@@ -10,7 +10,7 @@
 #include "GeometryCollectionEngineTypes.generated.h"
 
 UENUM()
-enum class ECollectionAttributeEnum : uint8
+enum class UE_DEPRECATED(5.4, "ECollectionAttributeEnum is no longer supported") ECollectionAttributeEnum : uint8
 {
 	Chaos_Active           UMETA(DisplayName = "Active"),
 	Chaos_DynamicState     UMETA(DisplayName = "DynamicState"),
@@ -19,15 +19,20 @@ enum class ECollectionAttributeEnum : uint8
 	Chaos_Max              UMETA(Hidden)
 };
 
+PRAGMA_DISABLE_DEPRECATION_WARNINGS
+UE_DEPRECATED(5.4, "GetCollectionAttributeName is no longer supported")
 FName GetCollectionAttributeName(ECollectionAttributeEnum TypeIn);
+PRAGMA_ENABLE_DEPRECATION_WARNINGS
 
 UENUM()
-enum class ECollectionGroupEnum : uint8
+enum class UE_DEPRECATED(5.4, "ECollectionGroupEnum is no longer supported") ECollectionGroupEnum : uint8
 {
 	Chaos_Traansform        UMETA(DisplayName = "Transform"),
 	//
 	Chaos_Max                UMETA(Hidden)
 };
 
-
+PRAGMA_DISABLE_DEPRECATION_WARNINGS
+UE_DEPRECATED(5.4, "GetCollectionGroupName is no longer supported")
 FName GetCollectionGroupName(ECollectionGroupEnum TypeIn);
+PRAGMA_ENABLE_DEPRECATION_WARNINGS
\ No newline at end of file
