Chaos - debug draw improvement for collision debugging (better identification of shapes involved in collision)

diff --git a/Engine/Source/Runtime/Experimental/Chaos/Private/Chaos/ChaosDebugDraw.cpp b/Engine/Source/Runtime/Experimental/Chaos/Private/Chaos/ChaosDebugDraw.cpp
--- a/Engine/Source/Runtime/Experimental/Chaos/Private/Chaos/ChaosDebugDraw.cpp
+++ b/Engine/Source/Runtime/Experimental/Chaos/Private/Chaos/ChaosDebugDraw.cpp
@@ -987,6 +987,11 @@ namespace Chaos
 
 		void DrawCollisionImpl(const FRigidTransform3& SpaceTransform, const FPBDCollisionConstraint& Contact, FRealSingle ColorScale, const FRealSingle Duration, const FChaosDebugDrawSettings& Settings)
 		{
+			if (Contact.NumManifoldPoints() == 0)
+			{
+				return;
+			}
+
 			if ((Settings.ContactWidth > 0) || (Settings.ContactLen > 0) || (Settings.ContactInfoWidth > 0) || (Settings.ImpulseScale > 0.0f))
 			{
 				const FConstGenericParticleHandle Particle0 = Contact.GetParticle0();
@@ -1151,10 +1156,10 @@ namespace Chaos
 				const FVec3 Location = SpaceTransform.TransformPosition(Contact.CalculateWorldContactLocation());
 				const FVec3 Normal = SpaceTransform.TransformVector(Contact.CalculateWorldContactNormal());
 
-				const FColor C3 = FColor(128, 128, 128);
+				const FColor C3 = FColor::White;
 				const FMatrix Axes = FRotationMatrix::MakeFromX(Normal);
-				const FVec3 P0 = SpaceTransform.TransformPosition(Contact.GetParticle0()->X());
-				const FVec3 P1 = SpaceTransform.TransformPosition(Contact.GetParticle1()->X());
+				const FVec3 P0 = SpaceTransform.TransformPosition(Contact.GetShape0()->GetWorldSpaceInflatedShapeBounds().GetCenter());
+				const FVec3 P1 = SpaceTransform.TransformPosition(Contact.GetShape1()->GetWorldSpaceInflatedShapeBounds().GetCenter());
 				FDebugDrawQueue::GetInstance().DrawDebugLine(Location, P0, C3, false, UE_KINDA_SMALL_NUMBER, uint8(Settings.DrawPriority), Settings.LineThickness * 0.5f);
 				FDebugDrawQueue::GetInstance().DrawDebugLine(Location, P1, C3, false, UE_KINDA_SMALL_NUMBER, uint8(Settings.DrawPriority), Settings.LineThickness * 0.5f);
 			}
