Chaos - re-enable non-mesh edge smoothing for back compat

diff --git a/Engine/Source/Runtime/Experimental/Chaos/Private/Chaos/Collision/CollisionPruning.cpp b/Engine/Source/Runtime/Experimental/Chaos/Private/Chaos/Collision/CollisionPruning.cpp
--- a/Engine/Source/Runtime/Experimental/Chaos/Private/Chaos/Collision/CollisionPruning.cpp
+++ b/Engine/Source/Runtime/Experimental/Chaos/Private/Chaos/Collision/CollisionPruning.cpp
@@ -129,6 +129,12 @@ namespace Chaos
 					const int32 EdgeOtherShapeIndex = (EdgeCollision.GetParticle0() == Particle) ? 1 : 0;
 					const EContactPointType VertexContactType = (EdgeOtherShapeIndex == 0) ? EContactPointType::VertexPlane : EContactPointType::PlaneVertex;
 
+					// We skip meshes - they are handled by FParticleMeshCollisionPruner
+					if (EdgeCollision.GetImplicit(EdgeOtherShapeIndex)->IsUnderlyingMesh())
+					{
+						return ECollisionVisitorResult::Continue;
+					}
+
 					for (const FManifoldPoint& EdgeManifoldPoint : EdgeCollision.GetManifoldPoints())
 					{
 						const bool bIsEdgeContact = (EdgeManifoldPoint.ContactPoint.ContactType == EContactPointType::EdgeEdge);
diff --git a/Engine/Source/Runtime/Experimental/Chaos/Private/Chaos/PBDCollisionConstraints.cpp b/Engine/Source/Runtime/Experimental/Chaos/Private/Chaos/PBDCollisionConstraints.cpp
--- a/Engine/Source/Runtime/Experimental/Chaos/Private/Chaos/PBDCollisionConstraints.cpp
+++ b/Engine/Source/Runtime/Experimental/Chaos/Private/Chaos/PBDCollisionConstraints.cpp
@@ -76,7 +76,7 @@ namespace Chaos
 	FAutoConsoleVariableRef CVarCollisionsAllowParticleTracking(TEXT("p.Chaos.Collision.AllowParticleTracking"), CollisionsAllowParticleTracking, TEXT("Allow particles to track their collisions constraints when their DoBufferCollisions flag is enable [def:true]"));
 
 	//	Which edge pruning features to enable (for particle swith EdgeSmoothing enabled)
-	bool bCollisionsEnableEdgeCollisionPruning = false;
+	bool bCollisionsEnableEdgeCollisionPruning = true;
 	bool bCollisionsEnableMeshCollisionPruning = true;
 	bool bCollisionsEnableSubSurfaceCollisionPruning = false;
 	FAutoConsoleVariableRef CVarCollisionsEnableEdgeCollisionPruning(TEXT("p.Chaos.Collision.EnableEdgeCollisionPruning"), bCollisionsEnableEdgeCollisionPruning, TEXT(""));
