Chaos - fix capsule-convex collision edge case

diff --git a/Engine/Source/Runtime/Experimental/Chaos/Private/Chaos/CollisionOneShotManifolds.cpp b/Engine/Source/Runtime/Experimental/Chaos/Private/Chaos/CollisionOneShotManifolds.cpp
--- a/Engine/Source/Runtime/Experimental/Chaos/Private/Chaos/CollisionOneShotManifolds.cpp
+++ b/Engine/Source/Runtime/Experimental/Chaos/Private/Chaos/CollisionOneShotManifolds.cpp
@@ -448,13 +448,13 @@ namespace Chaos
 					ContactPointCount = 0;
 					break;
 				}
-				else if ((Dist0 > PlaneTolerance) && (Dist1 < 0))
+				else if ((Dist0 > PlaneTolerance) && (Dist1 < PlaneTolerance))
 				{
 					// First point is outside, second inside
 					const FReal Alpha = Dist1 / (Dist1 - Dist0);
 					VertexBuffer[0] = VertexBuffer[1] + Alpha * (VertexBuffer[0] - VertexBuffer[1]);
 				}
-				else if ((Dist1 > PlaneTolerance) && (Dist0 < 0))
+				else if ((Dist1 > PlaneTolerance) && (Dist0 < PlaneTolerance))
 				{
 					// Second point is outside, first inside
 					const FReal Alpha = Dist0 / (Dist0 - Dist1);
