Chaos - don't activate swept constraint if they did not hit something

diff --git a/Engine/Source/Runtime/Experimental/Chaos/Private/Chaos/Collision/ParticlePairMidPhase.cpp b/Engine/Source/Runtime/Experimental/Chaos/Private/Chaos/Collision/ParticlePairMidPhase.cpp
--- a/Engine/Source/Runtime/Experimental/Chaos/Private/Chaos/Collision/ParticlePairMidPhase.cpp
+++ b/Engine/Source/Runtime/Experimental/Chaos/Private/Chaos/Collision/ParticlePairMidPhase.cpp
@@ -565,6 +565,10 @@ namespace Chaos
 			}
 #endif
 
+			// @todo(chaos): we always activate swept constraints because the CCD roll-back loops over the active constraints.
+			// Fix this - we should only activate constraints when we have a Phi within CullDistance
+			bool bShouldActivate = bDidSweep && (Constraint->GetCCDTimeOfImpact() < FReal(1));
+
 			// If we did not get a sweep hit (TOI > 1) or did not sweep (bDidSweep = false), we need to run standard collision detection at T=1.
 			// Likewise, if we did get a sweep hit but it's at TOI = 1, treat this constraint as a regular non-swept constraint and skip the rewind.
 			// NOTE: The sweep will report TOI==1 for "shallow" sweep hits below the CCD thresholds in the constraint.
@@ -573,10 +577,16 @@ namespace Chaos
 				// @todo(chaos): should we use a reduced cull distance if we get here? The cull distance will have been set based on movement speed...
 				Collisions::UpdateConstraint(*Constraint.Get(), Constraint->GetShapeWorldTransform0(), Constraint->GetShapeWorldTransform1(), Dt);
 				Constraint->SetCCDSweepEnabled(false);
+				bShouldActivate = (Constraint->GetPhi() <= CullDistance);
 			}
 
-			Context.GetAllocator()->ActivateConstraint(Constraint.Get());
-			LastUsedEpoch = Context.GetAllocator()->GetCurrentEpoch();
+			if (bShouldActivate)
+			{
+				if (Context.GetAllocator()->ActivateConstraint(Constraint.Get()))
+				{
+					LastUsedEpoch = Context.GetAllocator()->GetCurrentEpoch();
+				}
+			}
 
 			return 1;
 		}
@@ -1786,6 +1796,8 @@ namespace Chaos
 		}
 #endif
 
+		bool bShouldActivate = bDidSweep && (Constraint->GetCCDTimeOfImpact() < FReal(1));
+
 		// If we did not get a sweep hit (TOI > 1) or did not sweep (bDidSweep = false), we need to run standard collision detection at T=1.
 		// Likewise, if we did get a sweep hit but it's at TOI = 1, treat this constraint as a regular non-swept constraint and skip the rewind.
 		// NOTE: The sweep will report TOI==1 for "shallow" sweep hits below the CCD thresholds in the constraint.
@@ -1793,9 +1805,10 @@ namespace Chaos
 		{
 			Collisions::UpdateConstraint(*Constraint, Constraint->GetShapeWorldTransform0(), Constraint->GetShapeWorldTransform1(), Dt);
 			Constraint->SetCCDSweepEnabled(false);
+			bShouldActivate = Constraint->GetPhi() < CullDistance;
 		}
 
-		return true;
+		return bShouldActivate;
 	}
 
 	void FGenericParticlePairMidPhase::PruneConstraints(const int32 CurrentEpoch)
