Chaos - fix to mesh collision pruning

diff --git a/Engine/Source/Runtime/Experimental/Chaos/Private/Chaos/Collision/ContactTriangles.cpp b/Engine/Source/Runtime/Experimental/Chaos/Private/Chaos/Collision/ContactTriangles.cpp
--- a/Engine/Source/Runtime/Experimental/Chaos/Private/Chaos/Collision/ContactTriangles.cpp
+++ b/Engine/Source/Runtime/Experimental/Chaos/Private/Chaos/Collision/ContactTriangles.cpp
@@ -448,31 +448,40 @@ namespace Chaos
 							const FReal ContactDotNormal = FVec3::DotProduct(ContactPoint.ShapeContactNormal, ContactTriangle.FaceNormal);
 							if (ContactDotNormal < MinContactDotNormal)
 							{
-								// If we are separated, just ignore this contact
-								if (ContactPoint.Phi > 0)
+								// We are outside the valid normal range for this edge
+								// Convert the edge collision to a face collision on one of the faces, selected to get the smallest depth
+								FVec3 CorrectedContactNormal;
+								int32 CorrectedTriangleIndex;
+								const FReal OtherContactDotNormal = FVec3::DotProduct(ContactPoint.ShapeContactNormal, OtherContactTriangle.FaceNormal);
+								if (ContactDotNormal >= OtherContactDotNormal)
 								{
-									DisableContact(ContactIndex);
+									CorrectedContactNormal = (ContactDotNormal > -SMALL_NUMBER) ? ContactTriangle.FaceNormal : -ContactTriangle.FaceNormal;
+									CorrectedTriangleIndex = ContactTriangleIndex;
 								}
 								else
 								{
-									// We are outside the valid normal range for this edge
-									// Convert the edge collision to a face collision on one of the faces, selected to get the smallest depth
-									// NOTE: We keep the contact depth as it is because we know that the depth is a lower-bound. 
-									// We have to update the ShapeContactPoint[0] because Phi is actually derived from the positions (the value in ContactPoint is just a cache of current state)
-									// @todo(chaos): face selection logic might be better if we knew the contact velocity
-									const FReal OtherContactDotNormal = FVec3::DotProduct(ContactPoint.ShapeContactNormal, OtherContactTriangle.FaceNormal);
-									if (ContactDotNormal >= OtherContactDotNormal)
-									{
-										ContactPoint.ShapeContactNormal = (ContactDotNormal > -SMALL_NUMBER) ? ContactTriangle.FaceNormal : -ContactTriangle.FaceNormal;
-										ContactPoint.ShapeContactPoints[0] = ContactPoint.ShapeContactPoints[1] + ContactPoint.Phi * ContactPoint.ShapeContactNormal;
-									}
-									else
-									{
-										ContactPoint.ShapeContactNormal = (OtherContactDotNormal > -SMALL_NUMBER) ? OtherContactTriangle.FaceNormal : -OtherContactTriangle.FaceNormal;
-										ContactPoint.ShapeContactPoints[0] = ContactPoint.ShapeContactPoints[1] + ContactPoint.Phi * ContactPoint.ShapeContactNormal;
-										ContactPointData.SetTriangleIndex(OtherContactTriangleIndex);
-									}
+									CorrectedContactNormal = (OtherContactDotNormal > -SMALL_NUMBER) ? OtherContactTriangle.FaceNormal : -OtherContactTriangle.FaceNormal;
+									CorrectedTriangleIndex = OtherContactTriangleIndex;
 								}
+
+								// If we had to change the normal by a lot, disable the contacts
+								//if (ContactPoint.Phi > 0)
+								//{
+								//	const FReal CorrectedNormalThreshold = -KINDA_SMALL_NUMBER;	// 90 deg
+								//	const FReal CorrectedContactNormalDotContactNormal = FVec3::DotProduct(CorrectedContactNormal, ContactPoint.ShapeContactNormal);
+								//	if (CorrectedContactNormalDotContactNormal < CorrectedNormalThreshold)
+								//	{
+								//		DisableContact(ContactIndex);
+								//		continue;
+								//	}
+								//}
+
+								// NOTE: We keep the contact depth as it is because we know that the depth is a lower-bound. 
+								// We have to update the ShapeContactPoint[0] because Phi is actually derived from the positions (the value in ContactPoint is just a cache of current state)
+								// @todo(chaos): face selection logic might be better if we knew the contact velocity
+								ContactPoint.ShapeContactNormal = CorrectedContactNormal;
+								ContactPoint.ShapeContactPoints[0] = ContactPoint.ShapeContactPoints[1] + ContactPoint.Phi * CorrectedContactNormal;
+								ContactPointData.SetTriangleIndex(CorrectedTriangleIndex);
 							}
 						}
 					}
@@ -500,16 +509,23 @@ namespace Chaos
 										const FReal NormalDotCentroid = FVec3::DotProduct(ContactPoint.ShapeContactNormal, Centroid - ContactPoint.ShapeContactPoints[1]);
 										if (NormalDotCentroid > 0)
 										{
-											// If we are separated, just ignore this contact
-											if (ContactPoint.Phi > 0)
-											{
-												DisableContact(ContactIndex);
-												break;
-											}
-
 											const FReal OtherContactDotNormal = FVec3::DotProduct(ContactPoint.ShapeContactNormal, OtherContactTriangle.FaceNormal);
-											ContactPoint.ShapeContactNormal = (OtherContactDotNormal > -SMALL_NUMBER) ? OtherContactTriangle.FaceNormal : -OtherContactTriangle.FaceNormal;
-											ContactPoint.ShapeContactPoints[0] = ContactPoint.ShapeContactPoints[1] + ContactPoint.Phi * ContactPoint.ShapeContactNormal;
+											const FVec3 CorrectedContactNormal = (OtherContactDotNormal > -SMALL_NUMBER) ? OtherContactTriangle.FaceNormal : -OtherContactTriangle.FaceNormal;
+
+											// If we had to change the normal by a lot, disable the contacts
+											//if (ContactPoint.Phi > 0)
+											//{
+											//	const FReal CorrectedNormalThreshold = -KINDA_SMALL_NUMBER;	// 90 deg
+											//	const FReal CorrectedContactNormalDotContactNormal = FVec3::DotProduct(CorrectedContactNormal, ContactPoint.ShapeContactNormal);
+											//	if (CorrectedContactNormalDotContactNormal < CorrectedNormalThreshold)
+											//	{
+											//		DisableContact(ContactIndex);
+											//		break;
+											//	}
+											//}
+
+											ContactPoint.ShapeContactNormal = CorrectedContactNormal;
+											ContactPoint.ShapeContactPoints[0] = ContactPoint.ShapeContactPoints[1] + ContactPoint.Phi * CorrectedContactNormal;
 										}
 									}
 								}
