Reduce IAS table entries memory use by 25%

#rnx

diff --git a/Engine/Source/Runtime/Experimental/IoStoreOnDemand/Private/OnDemandIoDispatcherBackend.cpp b/Engine/Source/Runtime/Experimental/IoStoreOnDemand/Private/OnDemandIoDispatcherBackend.cpp
--- a/Engine/Source/Runtime/Experimental/IoStoreOnDemand/Private/OnDemandIoDispatcherBackend.cpp
+++ b/Engine/Source/Runtime/Experimental/IoStoreOnDemand/Private/OnDemandIoDispatcherBackend.cpp
@@ -138,17 +138,16 @@ FIoHash GetChunkKey(const FIoHash& ChunkHash, const FIoOffsetAndLength& Range)
 }
 
 ///////////////////////////////////////////////////////////////////////////////
-
 class FOnDemandIoStore
 {
 public:
 	struct FTocEntry
 	{
-		FIoHash Hash;
-		uint64 RawSize = 0;
-		uint64 EncodedSize = 0;
+		uint32 RawSize = 0;
+		uint32 EncodedSize = 0;
 		uint32 BlockOffset = ~uint32(0);
 		uint32 BlockCount = 0; 
+		FIoHash Hash;
 	};
 
 	struct FToc;
@@ -258,13 +257,15 @@ void FOnDemandIoStore::AddToc(const FOnDemandEndpoint& Ep, FOnDemandToc&& Toc)
 			NewContainer->TocEntries.Reserve(Container.Entries.Num());
 			for (const FOnDemandTocEntry& TocEntry : Container.Entries)
 			{
+				check(TocEntry.RawSize <= 0xffff'ffffull);
+				check(TocEntry.EncodedSize <= 0xffff'ffffull);
 				NewContainer->TocEntries.Add(TocEntry.ChunkId, FTocEntry
 				{
-					TocEntry.Hash,
-					TocEntry.RawSize,
-					TocEntry.EncodedSize,
+					uint32(TocEntry.RawSize),
+					uint32(TocEntry.EncodedSize),
 					TocEntry.BlockOffset,
-					TocEntry.BlockCount
+					TocEntry.BlockCount,
+					TocEntry.Hash,
 				});
 			}
 
@@ -338,7 +339,7 @@ void FOnDemandIoStore::AddDeferredContainers()
 		FContainer* Container = *It;
 		if (Container->EncryptionKeyGuid.IsEmpty())
 		{
-			UE_LOG(LogIas, Log, TEXT("Mounting container '%s'"), *Container->Name);
+			UE_LOG(LogIas, Log, TEXT("Mounting container '%s' (%d entries)"), *Container->Name, Container->TocEntries.Num());
 			RegisteredContainers.Add(Container);
 			It.RemoveCurrent();
 		}
@@ -348,7 +349,7 @@ void FOnDemandIoStore::AddDeferredContainers()
 			ensure(FGuid::Parse(Container->EncryptionKeyGuid, KeyGuid));
 			if (const FAES::FAESKey* Key = FEncryptionKeyManager::Get().GetKey(KeyGuid))
 			{
-				UE_LOG(LogIas, Log, TEXT("Mounting container '%s'"), *Container->Name);
+				UE_LOG(LogIas, Log, TEXT("Mounting container '%s' (%d entries)"), *Container->Name, Container->TocEntries.Num());
 				Container->EncryptionKey = *Key;
 				RegisteredContainers.Add(Container);
 				It.RemoveCurrent();
@@ -504,7 +505,7 @@ struct FChunkRequestParams
 {
 	static FChunkRequestParams Create(const FIoOffsetAndLength& OffsetLength, FOnDemandIoStore::FChunkInfo ChunkInfo)
 	{
-		const uint64 RawSize = FMath::Min(OffsetLength.GetLength(), ChunkInfo.Entry->RawSize);
+		const uint64 RawSize = FMath::Min<uint64>(OffsetLength.GetLength(), ChunkInfo.Entry->RawSize);
 		
 		const FIoOffsetAndLength ChunkRange = FIoChunkEncoding::GetChunkRange(
 			ChunkInfo.Entry->RawSize,
@@ -970,7 +971,7 @@ void FOnDemandIoBackend::CompleteRequest(FChunkRequest* ChunkRequest)
 		if (Chunk.GetSize() > 0)
 		{
 			TRACE_CPUPROFILER_EVENT_SCOPE(FOnDemandIoBackend::DecodeBlocks);
-			const uint64 RawSize = FMath::Min(Request->Options.GetSize(), ChunkRequest->Params.ChunkInfo.Entry->RawSize);
+			const uint64 RawSize = FMath::Min<uint64>(Request->Options.GetSize(), ChunkRequest->Params.ChunkInfo.Entry->RawSize);
 			Request->CreateBuffer(RawSize);
 			DecodingParams.RawOffset = Request->Options.GetOffset(); 
 			bDecoded = FIoChunkEncoding::Decode(DecodingParams, Chunk.GetView(), Request->GetBuffer().GetMutableView());
@@ -1132,7 +1133,7 @@ TIoStatusOr<uint64> FOnDemandIoBackend::GetSizeForChunk(const FIoChunkId& ChunkI
 		return FIoStatus(EIoErrorCode::UnknownChunkID);
 	}
 
-	FIoOffsetAndLength RequestedRange(ChunkRange.GetOffset(), FMath::Min(ChunkInfo.Entry->RawSize, ChunkRange.GetLength()));
+	FIoOffsetAndLength RequestedRange(ChunkRange.GetOffset(), FMath::Min<uint64>(ChunkInfo.Entry->RawSize, ChunkRange.GetLength()));
 	OutAvailable = ChunkInfo.Entry->RawSize;
 
 	if (IsHttpEnabled() == false || IsHttpEnabled(ChunkId) == false)
