Hold the IAS cache's .jrn file open between flushes, closing and reopening it when flushing.

UE's platform layer will commit cache mounts when the open file count reaches zero. However this causes too frequent write accesses making it difficult to get close to TRC write op and size limits. By holding open file count high we can exert some control over the commits

#rnx

diff --git a/Engine/Source/Runtime/Experimental/IoStoreOnDemand/Private/JournaledCache.cpp b/Engine/Source/Runtime/Experimental/IoStoreOnDemand/Private/JournaledCache.cpp
--- a/Engine/Source/Runtime/Experimental/IoStoreOnDemand/Private/JournaledCache.cpp
+++ b/Engine/Source/Runtime/Experimental/IoStoreOnDemand/Private/JournaledCache.cpp
@@ -329,6 +329,7 @@ private:
 	static_assert(sizeof(FMapEntry) == sizeof(uint64));
 
 	using					FDataMap = TMap<uint64, FMapEntry>;
+	void					OpenJrnFile();
 	void					Spam();
 	void					PhraseReset();
 	uint64					Insert(uint64 DataBase, const FDataEntry& Entry);
@@ -337,6 +338,7 @@ private:
 	TArray<FDataEntry>		Entries;
 	FString					BinPath;
 	FDataMap				DataMap;
+	TUniquePtr<IFileHandle>	JrnHandle;
 	uint64					MappedBytes;
 	uint64					MaxDataSize;
 	uint64					DataCursor;
@@ -369,6 +371,19 @@ FJournal::FJournal(FString&& Path, uint64 InMaxDataSize, uint32 InJournalSize)
 	MaxDataSize = (MaxDataSize - JournalSize) & ~((1ull << 20) - 1);
 
 	Reset();
+	OpenJrnFile();
+}
+
+////////////////////////////////////////////////////////////////////////////////
+void FJournal::OpenJrnFile()
+{
+	TStringBuilder<265> JrnPath;
+	JrnPath << BinPath;
+	JrnPath << TEXT(".jrn");
+
+	IPlatformFile& Ipf = IPlatformFile::GetPlatformPhysical();
+	IFileHandle* Handle = Ipf.OpenWrite(*JrnPath, true, false);
+	JrnHandle.Reset(Handle);
 }
 
 ////////////////////////////////////////////////////////////////////////////////
@@ -569,21 +584,18 @@ int32 FJournal::Flush()
 		JournalCursor = 0;
 	}
 
-	IPlatformFile& Ipf = IPlatformFile::GetPlatformPhysical();
-
-	TStringBuilder<265> JrnPath;
-	JrnPath << BinPath;
-	JrnPath << TEXT(".jrn");
-	if (TUniquePtr<IFileHandle> File(Ipf.OpenWrite(*JrnPath, true, false)); File.IsValid())
+	if (JrnHandle.IsValid())
 	{
 		TRACE_CPUPROFILER_EVENT_SCOPE(IasCache::JournalWrite);
 
-		File->Seek(JournalCursor);
-		File->Write((uint8*)(Entries.GetData()), Size);
+		JrnHandle->Seek(JournalCursor);
+		JrnHandle->Write((uint8*)(Entries.GetData()), Size);
+		JrnHandle.Reset();
 		JournalCursor += Size;
 	}
 
 	PhraseReset();
+	OpenJrnFile();
 
 	Spam();
 	return Size;
