Rudimentary nameable IAS cache so there facility to create more than one should the future

#rnx

diff --git a/Engine/Source/Runtime/Experimental/IoStoreOnDemand/Private/IasCache.h b/Engine/Source/Runtime/Experimental/IoStoreOnDemand/Private/IasCache.h
--- a/Engine/Source/Runtime/Experimental/IoStoreOnDemand/Private/IasCache.h
+++ b/Engine/Source/Runtime/Experimental/IoStoreOnDemand/Private/IasCache.h
@@ -42,6 +42,7 @@ struct FIasCacheConfig
 		uint32	Seconds = 60;
 	};
 
+	const TCHAR*Name = TEXT("ias");
 	uint64		DiskQuota = 512ull << 20;
 	uint32		MemoryQuota = 2 << 20;
 	uint32		JournalQuota = 4 << 20; // description in JournalCache.cpp
diff --git a/Engine/Source/Runtime/Experimental/IoStoreOnDemand/Private/JournaledCache.cpp b/Engine/Source/Runtime/Experimental/IoStoreOnDemand/Private/JournaledCache.cpp
--- a/Engine/Source/Runtime/Experimental/IoStoreOnDemand/Private/JournaledCache.cpp
+++ b/Engine/Source/Runtime/Experimental/IoStoreOnDemand/Private/JournaledCache.cpp
@@ -1108,7 +1108,9 @@ FJournaledCache::FJournaledCache(const FIasCacheConfig& Config)
 {
 	TStringBuilder<256> CachePath;
 	CachePath << FPaths::ProjectPersistentDownloadDir();
-	CachePath << TEXT("ias.cache.0");
+	CachePath << TEXT("ias/");
+	CachePath << Config.Name;
+	CachePath << TEXT(".cache.0");
 
 	FCacheInner::FConfig EventualConfig;
 	static_cast<FIasCacheConfig&>(EventualConfig) = Config;
