Poll could get stuck if a hang-up was returned

The changes to the select path are benign as errors are just treated as ready. Now it matches what poll does. However as revents wasn't zeroed out it could return false positives.

#rnx

diff --git a/Engine/Source/Runtime/Experimental/IoStoreOnDemand/Private/CoreHttp/Http.cpp b/Engine/Source/Runtime/Experimental/IoStoreOnDemand/Private/CoreHttp/Http.cpp
--- a/Engine/Source/Runtime/Experimental/IoStoreOnDemand/Private/CoreHttp/Http.cpp
+++ b/Engine/Source/Runtime/Experimental/IoStoreOnDemand/Private/CoreHttp/Http.cpp
@@ -1126,7 +1126,9 @@ struct FSelect
 };
 static const int32 POLLIN  = 1 << 0;
 static const int32 POLLOUT = 1 << 1;
-static const int32 POLLERR = 1 << 1;
+static const int32 POLLERR = 1 << 2;
+static const int32 POLLHUP = POLLERR;
+static const int32 POLLNVAL= POLLERR;
 #else
 using FSelect = pollfd;
 #endif
@@ -1167,8 +1169,15 @@ static bool DoSelect(FSelect* Selects, uint32 SelectNum, int32 TimeoutMs)
 	for (uint32 i = 0; i < SelectNum; ++i)
 	{
 		FSelect& Select = Selects[i];
+
+		if (FD_ISSET(Select.fd, &FdSetExcept))
+		{
+			Select.revents = POLLERR;
+			continue;
+		}
+
 		fd_set* RwSet = (Select.events & POLLIN) ? &FdSetRead : &FdSetWrite;
-		if (FD_ISSET(Select.fd, RwSet) || FD_ISSET(Select.fd, &FdSetExcept))
+		if (FD_ISSET(Select.fd, RwSet))
 		{
 			Select.revents = Select.events;
 		}
@@ -1210,8 +1219,11 @@ static uint64 ReadyCheck(FActivity** Activities, uint32 Num, uint32 TimeoutMs)
 		case EWait::Read:
 		case EWait::Write: {
 			FSelect& Select = Selects[SelectNum];
-			Select.fd = Activity->Socket;
-			Select.events = (Activity->SocketWait == EWait::Read) ? POLLIN : POLLOUT;
+			Select = {
+				Activity->Socket,
+				(Activity->SocketWait == EWait::Read) ? POLLIN : POLLOUT,
+				/* 0 */ // list-init so this value zeros with out per-platform hoops
+			};
 			++SelectNum;
 			} break;
 		}
@@ -1237,7 +1249,7 @@ static uint64 ReadyCheck(FActivity** Activities, uint32 Num, uint32 TimeoutMs)
 			continue;
 		}
 
-		if (int32(SelectCursor->revents & (POLLIN|POLLOUT|POLLERR)))
+		if (int32(SelectCursor->revents & (POLLIN|POLLOUT|POLLERR|POLLHUP|POLLNVAL)))
 		{
 			Activities[i]->SocketWait = EWait::None;
 			Ret |= (1ull << Activities[i]->Slot);
