Use IAS prefix for identifiers/scopes/names/etc.

#rnx

diff --git a/Engine/Source/Runtime/Experimental/IoStoreOnDemand/Private/CoreHttp/Http.cpp b/Engine/Source/Runtime/Experimental/IoStoreOnDemand/Private/CoreHttp/Http.cpp
--- a/Engine/Source/Runtime/Experimental/IoStoreOnDemand/Private/CoreHttp/Http.cpp
+++ b/Engine/Source/Runtime/Experimental/IoStoreOnDemand/Private/CoreHttp/Http.cpp
@@ -1,6 +1,6 @@
 // Copyright Epic Games, Inc. All Rights Reserved.
 
-#include "CoreHttp/Client.h"
+#include "Client.h"
 
 #if !defined(NO_UE_INCLUDES)
 #include "Containers/Array.h"
@@ -1190,7 +1190,7 @@ static bool DoSelect(FSelect* Selects, uint32 SelectNum, int32 TimeoutMs)
 ////////////////////////////////////////////////////////////////////////////////
 static uint64 ReadyCheck(FActivity** Activities, uint32 Num, uint32 TimeoutMs)
 {
-	TRACE_CPUPROFILER_EVENT_SCOPE(CoreHttp::ReadyCheck);
+	TRACE_CPUPROFILER_EVENT_SCOPE(IasHttp::ReadyCheck);
 
 	using EWait = FActivity::EWait;
 
@@ -1274,7 +1274,7 @@ struct FHandlerResult
 ////////////////////////////////////////////////////////////////////////////////
 static int32 DoResolve(FActivity* Activity)
 {
-	TRACE_CPUPROFILER_EVENT_SCOPE(CoreHttp::DoResolve);
+	TRACE_CPUPROFILER_EVENT_SCOPE(IasHttp::DoResolve);
 
 	// todo: GetAddrInfoW() for async resolve on Windows
 
@@ -1358,7 +1358,7 @@ static int32 DoResolve(FActivity* Activity)
 ////////////////////////////////////////////////////////////////////////////////
 static int32 DoConnect(FActivity* Activity)
 {
-	TRACE_CPUPROFILER_EVENT_SCOPE(CoreHttp::DoConnect);
+	TRACE_CPUPROFILER_EVENT_SCOPE(IasHttp::DoConnect);
 
 	FSocketPool* Pool = Activity->Pool;
 	if (Pool->GetState() == FSocketPool::EState::Error)
@@ -1465,7 +1465,7 @@ static int32 DoConnect(FActivity* Activity)
 ////////////////////////////////////////////////////////////////////////////////
 static int32 DoSend(FActivity* Activity)
 {
-	TRACE_CPUPROFILER_EVENT_SCOPE(CoreHttp::DoSend);
+	TRACE_CPUPROFILER_EVENT_SCOPE(IasHttp::DoSend);
 
 	enum { PackBits = 4 };
 	uint32 Index = Activity->StateParam & ((1 << PackBits) - 1);
@@ -1548,7 +1548,7 @@ static int32 DoSend(FActivity* Activity)
 ////////////////////////////////////////////////////////////////////////////////
 static int32 DoRecvMessage(FActivity* Activity)
 {
-	TRACE_CPUPROFILER_EVENT_SCOPE(CoreHttp::DoRecvMessage);
+	TRACE_CPUPROFILER_EVENT_SCOPE(IasHttp::DoRecvMessage);
 
 	static const uint32 PageSize = 2048;
 
@@ -1747,7 +1747,7 @@ static int32 DoRecvMessage(FActivity* Activity)
 ////////////////////////////////////////////////////////////////////////////////
 static FHandlerResult DoRecvContent(FActivity* Activity, uint32 MaxRecvSize)
 {
-	TRACE_CPUPROFILER_EVENT_SCOPE(CoreHttp::DoRecvContent);
+	TRACE_CPUPROFILER_EVENT_SCOPE(IasHttp::DoRecvContent);
 
 	auto& Response = *(FResponseInternal*)(Activity->Buffer.GetData());
 
@@ -1821,7 +1821,7 @@ static int32 DoRecvDone(FActivity* Activity)
 ////////////////////////////////////////////////////////////////////////////////
 static FHandlerResult DoRecvStream(FActivity* Activity, uint32 RecvSize)
 {
-	TRACE_CPUPROFILER_EVENT_SCOPE(CoreHttp::DoRecvStream);
+	TRACE_CPUPROFILER_EVENT_SCOPE(IasHttp::DoRecvStream);
 
 	check(false); // not implemented yet
 	return { -1 };
@@ -1951,7 +1951,7 @@ void FEventLoop::FImpl::Cancel(FTicket Ticket)
 ////////////////////////////////////////////////////////////////////////////////
 uint32 FEventLoop::FImpl::Tick(uint32 PollTimeoutMs)
 {
-	TRACE_CPUPROFILER_EVENT_SCOPE(CoreHttp::Tick);
+	TRACE_CPUPROFILER_EVENT_SCOPE(IasHttp::Tick);
 
     // Collect activity changes
 	uint64 FreeSlotsLoad = FreeSlots.load(std::memory_order_relaxed);
@@ -2323,7 +2323,7 @@ static void MiscTest()
 }
 
 ////////////////////////////////////////////////////////////////////////////////
-IOSTOREONDEMAND_API void CoreHttpTest()
+IOSTOREONDEMAND_API void IasHttpTest()
 {
 	MiscTest();
 
@@ -2399,7 +2399,7 @@ IOSTOREONDEMAND_API void CoreHttpTest()
 	FEventLoop Loop;
 	volatile bool LoopStop = false;
 	volatile bool LoopTickDelay = false;
-	auto LoopTask = UE::Tasks::Launch(TEXT("CoreHttpTest.Loop"), [&] () {
+	auto LoopTask = UE::Tasks::Launch(TEXT("IasHttpTest.Loop"), [&] () {
 		uint32 DelaySeed = 493;
 		while (!LoopStop)
 		{
