Speed throttling for IAS' download client

#rnx

diff --git a/Engine/Source/Runtime/Experimental/IoStoreOnDemand/Private/CoreHttp/Client.h b/Engine/Source/Runtime/Experimental/IoStoreOnDemand/Private/CoreHttp/Client.h
--- a/Engine/Source/Runtime/Experimental/IoStoreOnDemand/Private/CoreHttp/Client.h
+++ b/Engine/Source/Runtime/Experimental/IoStoreOnDemand/Private/CoreHttp/Client.h
@@ -168,6 +168,7 @@ public:
 							FEventLoop();
 							~FEventLoop();
 	uint32					Tick(uint32 PollTimeoutMs=0);
+	void					Throttle(uint32 KiBPerSec);
 	bool					IsIdle() const;
 	void					Cancel(FTicket Ticket);
 	[[nodiscard]] FRequest	Request(FAnsiStringView Method, FAnsiStringView Url, const FRequestParams* Params=nullptr);
diff --git a/Engine/Source/Runtime/Experimental/IoStoreOnDemand/Private/CoreHttp/Http.cpp b/Engine/Source/Runtime/Experimental/IoStoreOnDemand/Private/CoreHttp/Http.cpp
--- a/Engine/Source/Runtime/Experimental/IoStoreOnDemand/Private/CoreHttp/Http.cpp
+++ b/Engine/Source/Runtime/Experimental/IoStoreOnDemand/Private/CoreHttp/Http.cpp
@@ -1843,6 +1843,113 @@ static void DoCancel(FActivity* Activity)
 
 
 
+// {{{1 throttler ..............................................................
+
+////////////////////////////////////////////////////////////////////////////////
+static void ThrottleTest();
+
+////////////////////////////////////////////////////////////////////////////////
+class FThrottler
+{
+public:
+			FThrottler();
+	void	SetLimit(uint32 KiBPerSec);
+	int32	GetAllowance();
+	void	ReturnUnused(uint32 Unused);
+
+private:
+	friend	void ThrottleTest();
+	int32	GetAllowance(uint64 CycleDelta);
+	uint64	CycleFreq;
+	uint64	CycleLast;
+	uint64	CycleIdle;
+	uint32	Limit = 0;
+	int32	Available = 0;
+
+	enum {
+		LIMITLESS	= 0x7fff'ffff,
+		THRESHOLD	= 2 << 10,
+	};
+	
+};
+
+////////////////////////////////////////////////////////////////////////////////
+FThrottler::FThrottler()
+{
+	CycleFreq = uint64(1.0 / FPlatformTime::GetSecondsPerCycle());
+	CycleLast = FPlatformTime::Cycles64() - CycleFreq;
+	CycleIdle = CycleFreq * 8;
+}
+
+////////////////////////////////////////////////////////////////////////////////
+void FThrottler::SetLimit(uint32 KiBPerSec)
+{
+	// 512MiB/s might as well be limitless.
+	KiBPerSec = (KiBPerSec < (512 << 10)) ? KiBPerSec : 0;
+	Limit = KiBPerSec << 10;
+}
+
+////////////////////////////////////////////////////////////////////////////////
+int32 FThrottler::GetAllowance()
+{
+	int64 Cycle = FPlatformTime::Cycles64();
+	int64 CycleDelta = Cycle - CycleLast;
+	if (CycleDelta < 0)
+	{
+		return Limit;
+	}
+	CycleLast = Cycle;
+	return GetAllowance(CycleDelta);
+}
+
+////////////////////////////////////////////////////////////////////////////////
+int32 FThrottler::GetAllowance(uint64 CycleDelta)
+{
+	if (Limit == 0)
+	{
+		return LIMITLESS;
+	}
+
+	// If we're idle for too long then reset the throttling
+	if (CycleDelta >= CycleIdle)
+	{
+		Available = 0;
+		return Limit;
+	}
+
+	uint64 Delta = (uint64(Limit) * CycleDelta) / CycleFreq;
+
+	// A gate against lost precision
+	if (Delta == 0)
+	{
+		CycleLast -= CycleDelta;
+		return 0;
+	}
+
+	// Don't let available run away
+	uint64 Next = FMath::Min<uint64>(uint64(Available) + Delta, Limit * 4);
+
+	Available = uint32(Next);
+
+	// Doesn't make sense to trickle out tiny allowances
+	if (Available < THRESHOLD)
+	{
+		return 0;
+	}
+
+	int32 Released = Available;
+	Available = 0;
+	return Released;
+}
+
+////////////////////////////////////////////////////////////////////////////////
+void FThrottler::ReturnUnused(uint32 Unused)
+{
+	Available += Unused;
+}
+
+
+
 // {{{1 event-loop .............................................................
 
 ////////////////////////////////////////////////////////////////////////////////
@@ -1855,6 +1962,7 @@ public:
 							~FImpl();
 	uint32					Tick(uint32 PollTimeoutMs=0);
 	bool					IsIdle() const;
+	void					Throttle(uint32 KiBPerSec);
 	void					Cancel(FTicket Ticket);
 	FRequest				Request(FAnsiStringView Method, FAnsiStringView Path, FActivity* Activity);
 	FTicket					Send(FActivity* Activity);
@@ -1866,6 +1974,7 @@ private:
 	uint64					PrevFreeSlots	= ~0ull;
 	TArray<FActivity*>		Pending;
 	TArray<FActivity*>		Active;
+	FThrottler				Throttler;
 };
 
 ////////////////////////////////////////////////////////////////////////////////
@@ -1942,6 +2051,12 @@ bool FEventLoop::FImpl::IsIdle() const
 	return FreeSlots.load(std::memory_order_relaxed) == ~0ull;
 }
 
+////////////////////////////////////////////////////////////////////////////////
+void FEventLoop::FImpl::Throttle(uint32 KiBPerSec)
+{
+	Throttler.SetLimit(KiBPerSec);
+}
+
 ////////////////////////////////////////////////////////////////////////////////
 void FEventLoop::FImpl::Cancel(FTicket Ticket)
 {
@@ -1981,6 +2096,7 @@ uint32 FEventLoop::FImpl::Tick(uint32 PollTimeoutMs)
 	}
 
 	// Tick activities
+	int32 RecvAllowance = Throttler.GetAllowance();
 	uint64 CancelsLoad = Cancels.load(std::memory_order_relaxed);
 	for (FActivity* Activity : Active)
 	{
@@ -2035,7 +2151,14 @@ uint32 FEventLoop::FImpl::Tick(uint32 PollTimeoutMs)
 					Handler = DoRecvStream;
 				}
 
-				auto [ResultInner, RecvSize] = Handler(Activity, ~0u);
+				if (RecvAllowance <= 0)
+				{
+					Result = 0;
+					break;
+				}
+
+				auto [ResultInner, RecvSize] = Handler(Activity, RecvAllowance);
+				RecvAllowance -= RecvSize;
 				Result = ResultInner;
 				if (Result)
 					break;
@@ -2061,6 +2184,7 @@ uint32 FEventLoop::FImpl::Tick(uint32 PollTimeoutMs)
 			--BusyCount;
 		}
 	}
+	Throttler.ReturnUnused(RecvAllowance);
 
 	// Reap done activities
 	uint64 ReturnedSlots = 0;
@@ -2100,12 +2224,15 @@ uint32 FEventLoop::FImpl::Tick(uint32 PollTimeoutMs)
 	return BusyCount;
 }
 
+
+
 ////////////////////////////////////////////////////////////////////////////////
 FEventLoop::FEventLoop()						{ Impl = new FEventLoop::FImpl(); }
 FEventLoop::~FEventLoop()						{ delete Impl; }
 uint32 FEventLoop::Tick(uint32 PollTimeoutMs)	{ return Impl->Tick(PollTimeoutMs); }
 bool FEventLoop::IsIdle() const					{ return Impl->IsIdle(); }
 void FEventLoop::Cancel(FTicket Ticket)			{ return Impl->Cancel(Ticket); }
+void FEventLoop::Throttle(uint32 KiBPerSec)		{ return Impl->Throttle(KiBPerSec); }
 
 ////////////////////////////////////////////////////////////////////////////////
 FRequest FEventLoop::Request(
@@ -2323,10 +2450,103 @@ static void MiscTest()
 }
 
 ////////////////////////////////////////////////////////////////////////////////
-IOSTOREONDEMAND_API void IasHttpTest()
+static void ThrottleTest()
 {
-	MiscTest();
+	enum { TheMax = 0x7fff'fffful };
+
+	check(FThrottler().GetAllowance() >= TheMax);
+
+	FThrottler Throttler;
+	uint64 OneSecond = Throttler.CycleFreq;
+
+	Throttler.SetLimit(0);
+	check(Throttler.GetAllowance(0) >= TheMax);
+	check(Throttler.GetAllowance()  >= TheMax);
+	check(Throttler.GetAllowance()  >= TheMax);
+
+	for (uint32 i : { 10, 63, 100 })
+	{
+		Throttler.SetLimit(i);
+		check(Throttler.GetAllowance( OneSecond          ) == (i << 10));
+		check(Throttler.GetAllowance((OneSecond + 1) >> 1) == (i << 9));
+		check(Throttler.GetAllowance((OneSecond + 3) >> 2) == (i << 8));
+	}
+
+	uint32 Limit = 17 << 10;
+	Throttler = FThrottler();
+	Throttler.SetLimit(Limit >> 10);
+	Throttler.GetAllowance(OneSecond);
+	Throttler.ReturnUnused(Limit >> 1);
+	check(Throttler.GetAllowance(OneSecond) == (Limit + (Limit >> 1)));
+
+	// runaway
+	check(Throttler.GetAllowance(OneSecond * 3) == Limit * 3);
+	check(Throttler.GetAllowance(OneSecond * 4) == Limit * 4);
+	check(Throttler.GetAllowance(OneSecond * 5) == Limit * 4);
+
+	// idle
+	check(Throttler.GetAllowance(OneSecond * 64) == Limit);
 
+	// threshold
+	Limit = FThrottler::THRESHOLD;
+	Throttler = FThrottler();
+	Throttler.SetLimit(Limit >> 10);
+	for (int64 Counter = OneSecond;;)
+	{
+		int64 Delta = (OneSecond + 15) >> 4;
+		Counter -= Delta;
+		uint32 Allowance = Throttler.GetAllowance(Delta);
+		if (Allowance != 0)
+		{
+			check(Allowance == Limit);
+			check(Counter < 10);
+			break;
+		}
+	};
+
+	// overflow(ish)
+	Throttler = FThrottler();
+	Throttler.SetLimit((512 << 10) - 1);
+	Throttler.GetAllowance((OneSecond * 790) / 100);
+
+	// timing test
+	FIoBuffer RecvData;
+	for (uint32 SizeKiB : { 25, 60 })
+	{
+		const uint32 ThrottleKiB = 5;
+
+		TAnsiStringBuilder<128> TestUrl;
+		TestUrl << "http://localhost:9493/data/";
+		TestUrl << (SizeKiB << 10);
+
+		FEventLoop Loop;
+		Loop.Throttle(ThrottleKiB);
+
+		FRequest Request = Loop.Request("GET", TestUrl).Accept("*/*");
+		Loop.Send(MoveTemp(Request), [&] (const FTicketStatus& Status) {
+			check(Status.GetId() != FTicketStatus::EId::Error);
+			if (Status.GetId() == FTicketStatus::EId::Response)
+			{
+				Status.GetResponse().SetDestination(&RecvData);
+			}
+		});
+
+		uint64 Time = FPlatformTime::Cycles64();
+		while (Loop.Tick(-1));
+		Time = FPlatformTime::Cycles64() - Time;
+		Time /= OneSecond;
+
+		// It's dangerous stuff testing elapsed time you know. The +1 is because
+		// throttling assumes one second has already passed when initialised.
+		check(Time + 1 == (SizeKiB / ThrottleKiB));
+
+		RecvData = FIoBuffer();
+	}
+}
+
+////////////////////////////////////////////////////////////////////////////////
+IOSTOREONDEMAND_API void IasHttpTest()
+{
 #if PLATFORM_WINDOWS
 	WSADATA WsaData;
 	if (WSAStartup(MAKEWORD(2, 2), &WsaData) == 0x0a9e0493)
@@ -2334,6 +2554,9 @@ IOSTOREONDEMAND_API void IasHttpTest()
 	ON_SCOPE_EXIT { WSACleanup(); };
 #endif
 
+	MiscTest();
+	ThrottleTest();
+
 #define HOST_NAME "10.24.101.89"
 	FAnsiStringView TestUrl = "http://" HOST_NAME ":9493/data";
 	FAnsiStringView TestHostName = "localhost";
