Fixed IAS' narrowing warnings

#rnx

diff --git a/Engine/Source/Runtime/Experimental/IoStoreOnDemand/Private/CoreHttp/Http.cpp b/Engine/Source/Runtime/Experimental/IoStoreOnDemand/Private/CoreHttp/Http.cpp
--- a/Engine/Source/Runtime/Experimental/IoStoreOnDemand/Private/CoreHttp/Http.cpp
+++ b/Engine/Source/Runtime/Experimental/IoStoreOnDemand/Private/CoreHttp/Http.cpp
@@ -123,7 +123,7 @@ static void EnumerateHeaders(FAnsiStringView Headers, LambdaType&& Lambda)
 			Cursor = Right;
 			for (; Cursor > Left + 1 && IsOws(Cursor[-1]); --Cursor);
 
-			FAnsiStringView Value (Left, ptrdiff_t(Cursor - Left));
+			FAnsiStringView Value (Left, int32(ptrdiff_t(Cursor - Left)));
 
 			if (!Lambda(Name, Value))
 			{
@@ -204,7 +204,7 @@ static int32 ParseMessage(FAnsiStringView Message, FMessageOffsets& Out)
 
 	// Trim left and tightly reject anything adventurous
 	for (int n = 32; Cursor[i] == ' ' && i < n; ++i);
-	Out.StatusCode = i;
+	Out.StatusCode = uint8(i);
 
 	// At least one status line digit. (Note to self; expect exactly three)
 	for (int n = 32; uint32(Cursor[i] - 0x30) <= 9 && i < n; ++i);
@@ -215,7 +215,7 @@ static int32 ParseMessage(FAnsiStringView Message, FMessageOffsets& Out)
 
 	// Trim left
 	for (int n = 32; Cursor[i] == ' ' && i < n; ++i);
-	Out.Message = i;
+	Out.Message = uint8(i);
 
 	// Extra conservative length allowance
 	if (i > 32)
@@ -235,7 +235,7 @@ static int32 ParseMessage(FAnsiStringView Message, FMessageOffsets& Out)
 	{
 		return -1;
 	}
-	Out.Headers = i + 2;
+	Out.Headers = uint16(i + 2);
 
 	return 1;
 }
@@ -245,6 +245,8 @@ struct FUrlOffsets
 {
 	struct Slice
 	{
+						Slice() = default;
+						Slice(int32 l, int32 r) : Left(uint8(l)), Right(uint8(r)) {}
 		FAnsiStringView	Get(FAnsiStringView Url) const { return Url.Mid(Left, Right - Left); }
 						operator bool () const { return Left > 0; }
 		int32			Len() const { return Right - Left; }
@@ -330,14 +332,14 @@ static int32 ParseUrl(FAnsiStringView Url, FUrlOffsets& Out)
 	case 1:
 		if (Seps[0].c == ':')
 		{
-			Out.Port = { uint8(Seps[0].i + 1), uint8(i) };
-			Out.HostName.Right = Seps[0].i;
+			Out.Port = { Seps[0].i + 1, i };
+			Out.HostName.Right = uint8(Seps[0].i);
 		}
 		else
 		{
-			Out.UserInfo = { Out.HostName.Left, uint8(Seps[0].i) };
-			Out.HostName.Left += Seps[0].i + 1;
-			Out.HostName.Right += Seps[0].i + 1;
+			Out.UserInfo = { Out.HostName.Left, Seps[0].i };
+			Out.HostName.Left += uint8(Seps[0].i + 1);
+			Out.HostName.Right += uint8(Seps[0].i + 1);
 		}
 		break;
 
@@ -346,8 +348,8 @@ static int32 ParseUrl(FAnsiStringView Url, FUrlOffsets& Out)
 		{
 			return -1;
 		}
-		Out.UserInfo = { Out.HostName.Left, uint8(Seps[0].i) };
-		Out.Port.Left = Seps[1].i + 1;
+		Out.UserInfo = { Out.HostName.Left, Seps[0].i };
+		Out.Port.Left = uint8(Seps[1].i + 1);
 		Out.Port.Right = Out.HostName.Right;
 		Out.HostName.Left = Out.UserInfo.Right + 1;
 		Out.HostName.Right = Out.Port.Left - 1;
@@ -617,8 +619,8 @@ private:
 ////////////////////////////////////////////////////////////////////////////////
 FSocketPool::FSocketPool(FAnsiStringView InHostName, uint32 InPort, uint32 InMaxLeases)
 : HostName(InHostName)
-, Port(InPort)
-, MaxLeases(InMaxLeases)
+, Port(uint16(InPort))
+, MaxLeases(uint8(InMaxLeases))
 , State(EState::Unresolved)
 {
 	check(MaxLeases == InMaxLeases); // field overflow
@@ -693,7 +695,7 @@ bool FSocketPool::AddIpAddress(uint32 Address)
 ////////////////////////////////////////////////////////////////////////////////
 void FSocketPool::SetBufferSize(EDirection Dir, int32 Size)
 {
-	((Dir == EDirection::Send) ? SendBufKb : RecvBufKb) = Size >> 10;
+	(Dir == EDirection::Send) ? SendBufKb : RecvBufKb = uint16(Size >> 10);
 }
 
 ////////////////////////////////////////////////////////////////////////////////
@@ -718,7 +720,7 @@ int32 FConnectionPool::FParams::SetHostFromUrl(FAnsiStringView Url)
 	if (Offsets.Port)
 	{
 		FAnsiStringView PortView = Offsets.Port.Get(Url);
-		Host.Port = CrudeToInt(PortView);
+		Host.Port = uint16(CrudeToInt(PortView));
 	}
 
 	return Offsets.Path;
@@ -1443,7 +1445,7 @@ static int32 DoConnect(FActivity* Activity)
 
 	sockaddr_in AddrInet = { sizeof(sockaddr_in) };
 	AddrInet.sin_family = AF_INET;
-	AddrInet.sin_port = htons(Activity->Pool->GetPort());
+	AddrInet.sin_port = htons(uint16(Activity->Pool->GetPort()));
 	memcpy(&(AddrInet.sin_addr), &IpAddress, sizeof(IpAddress));
 	{
 		int Result = connect(Candidate, &(sockaddr&)AddrInet, sizeof(AddrInet));
@@ -1634,7 +1636,7 @@ static int32 DoRecvMessage(FActivity* Activity)
 
 			if (Name.Equals("Content-Length", ESearchCase::IgnoreCase))
 			{
-				ContentLength = CrudeToInt(Value);
+				ContentLength = int32(CrudeToInt(Value));
 				return true;
 			}
 			
@@ -2260,7 +2262,7 @@ FRequest FEventLoop::Request(
 	if (UrlOffsets.Port)
 	{
 		FAnsiStringView PortView = UrlOffsets.Port.Get(Url);
-		Port = CrudeToInt(PortView);
+		Port = uint32(CrudeToInt(PortView));
 	}
 
 	FAnsiStringView Path;
@@ -2712,7 +2714,7 @@ IOSTOREONDEMAND_API void IasHttpTest()
 		auto ErrorSink = [&] (const FTicketStatus& Status)
 		{
 			FTicket Ticket = Status.GetTicket();
-			uint32 Index = 63 - FMath::CountLeadingZeros64(uint64(Ticket));
+			uint32 Index = uint32(63 - FMath::CountLeadingZeros64(uint64(Ticket)));
 
 			if (Status.GetId() == FTicketStatus::EId::Error)
 			{
diff --git a/Engine/Source/Runtime/Experimental/IoStoreOnDemand/Private/Governor.inl b/Engine/Source/Runtime/Experimental/IoStoreOnDemand/Private/Governor.inl
--- a/Engine/Source/Runtime/Experimental/IoStoreOnDemand/Private/Governor.inl
+++ b/Engine/Source/Runtime/Experimental/IoStoreOnDemand/Private/Governor.inl
@@ -41,7 +41,7 @@ inline uint32 FGovernorExternal::TickAllowance()
 	CycleLast = Cycle;
 
 	IPlatformFile& Ipf = IPlatformFile::GetPlatformPhysical();
-	return Ipf.GetAllowedBytesToWriteThrottledStorage();
+	return uint32(Ipf.GetAllowedBytesToWriteThrottledStorage());
 }
 
 
diff --git a/Engine/Source/Runtime/Experimental/IoStoreOnDemand/Private/IoStoreOnDemand.cpp b/Engine/Source/Runtime/Experimental/IoStoreOnDemand/Private/IoStoreOnDemand.cpp
--- a/Engine/Source/Runtime/Experimental/IoStoreOnDemand/Private/IoStoreOnDemand.cpp
+++ b/Engine/Source/Runtime/Experimental/IoStoreOnDemand/Private/IoStoreOnDemand.cpp
@@ -854,7 +854,7 @@ TIoStatusOr<FIoStoreUploadResult> UploadContainerFiles(
 			for (const FIoStoreCompressedBlockInfo& BlockInfo : ReadResult.Blocks)
 			{
 				const uint64 EncodedBlockSize = Align(BlockInfo.CompressedSize, FAES::AESBlockSize);
-				ContainerEntry.BlockSizes.Add(EncodedBlockSize);
+				ContainerEntry.BlockSizes.Add(uint32(EncodedBlockSize));
 
 				FMemoryView EncodedBlock = EncodedBlocks.Left(EncodedBlockSize);
 				EncodedBlocks += EncodedBlock.GetSize();
diff --git a/Engine/Source/Runtime/Experimental/IoStoreOnDemand/Private/JournaledCache.cpp b/Engine/Source/Runtime/Experimental/IoStoreOnDemand/Private/JournaledCache.cpp
--- a/Engine/Source/Runtime/Experimental/IoStoreOnDemand/Private/JournaledCache.cpp
+++ b/Engine/Source/Runtime/Experimental/IoStoreOnDemand/Private/JournaledCache.cpp
@@ -189,7 +189,7 @@ bool FPending::Put(uint64 Key, FIoBuffer&& Data)
 		return true;
 	}
 
-	uint32 Size = Data.GetSize();
+	uint32 Size = uint32(Data.GetSize());
 	if (Size == 0 || MaxSize < Size)
 	{
 		return false;
@@ -225,7 +225,7 @@ uint32 FPending::DebugVisit(void* Param, FDebugCacheEntry::Callback* Callback)
 	for (auto& Item : Items)
 	{
 		Out.Key = Item.Key;
-		Out.Size = Item.Data.GetSize();
+		Out.Size = uint32(Item.Data.GetSize());
 		Callback(Param, Out);
 	}
 	return Items.Num();
@@ -246,7 +246,7 @@ int32 FPending::DropImpl(uint32 Size, Lambda&& Callback)
 
 		uint32 Index = n ? (Size * 0x0'a9e0'493) % n : 0;
 
-		Size = Items[Index].Data.GetSize();
+		Size = uint32(Items[Index].Data.GetSize());
 		UsedSize -= Size;
 		DropSize += Size;
 
@@ -476,14 +476,14 @@ bool FJournal::Materialize(EntryHandle Handle, FIoBuffer& Out, uint32 Offset) co
 
 	const FMapEntry& Entry = *(FMapEntry*)Handle;
 
-	uint32 ReadSize = Entry.Size - Offset;
+	uint32 ReadSize = uint32(Entry.Size) - Offset;
 
 	if (Out.GetData() == nullptr)
 	{
 		Out = FIoBuffer(ReadSize);
 	}
 
-	ReadSize = FMath::Min<uint32>(Out.GetSize(), ReadSize);
+	ReadSize = FMath::Min<uint32>(uint32(Out.GetSize()), ReadSize);
 
 	IPlatformFile& Ipf = IPlatformFile::GetPlatformPhysical();
 	TUniquePtr<IFileHandle> File(Ipf.OpenRead(*BinPath, false));
@@ -679,7 +679,7 @@ bool FJournal::Load()
 
 		Out.Phrase = Header;
 		Out.Marker = Header->Marker;
-		Out.DataSize = LastEntry->Offset + LastEntry->Size + sizeof(MarkerType);
+		Out.DataSize = uint32(LastEntry->Offset + LastEntry->Size) + sizeof(MarkerType);
 
 		return Cursor;
 	};
@@ -1202,7 +1202,6 @@ FJournaledCache::GetRetType	FJournaledCache::Get(
 
 		FIoBuffer Buffer;
 
-		uint64 Offset = Options.GetOffset();
 		int64 Size = Options.GetSize();
 		if (void* DestAddr = Options.GetTargetVa(); DestAddr != nullptr)
 		{
@@ -1214,6 +1213,7 @@ FJournaledCache::GetRetType	FJournaledCache::Get(
 			Buffer = FIoBuffer(Size);
 		}
 
+		uint32 Offset = uint32(Options.GetOffset());
 		if (!Entry.Materialize(Buffer, Offset))
 		{
 			return TIoStatusOr<FIoBuffer>(FIoStatus(EIoErrorCode::ReadError));
@@ -1267,7 +1267,7 @@ static uint64 KeyGen(const uint8* Data, uint32 Size)
 
 static uint64 KeyGen(const FIoBuffer& Data)
 {
-	return KeyGen(Data.GetData(), Data.GetSize());
+	return KeyGen(Data.GetData(), uint32(Data.GetSize()));
 }
 
 ////////////////////////////////////////////////////////////////////////////////
@@ -1282,15 +1282,15 @@ IOSTOREONDEMAND_API void Tests()
 	};
 
 	// Some data
-	uint32 WorkingSize = 1_Mi;
+	uint64 WorkingSize = 1_Mi;
 	TUniquePtr<uint8[]> WorkingScope(new uint8[WorkingSize]);
 	uint8* Working = WorkingScope.Get();
 	for (uint32 i = 0; i < WorkingSize; i += 8)
 	{
 		*(uint64*)(Working + i) = MixTh();
 	}
-	auto DummyData = [&] (uint32 Size) {
-		uint32 Offset = MixTh() % (WorkingSize - Size);
+	auto DummyData = [&] (uint64 Size) {
+		uint64 Offset = MixTh() % (WorkingSize - Size);
 		return FIoBuffer(FIoBuffer::Wrap, Working + Offset, Size);
 	};
 
@@ -1355,18 +1355,18 @@ IOSTOREONDEMAND_API void Tests()
 	{
 		FCache::FConfig Config;
 		Config.Path = FPaths::ProjectPersistentDownloadDir() / TEXT("ias_cache_test");
-		Config.MemoryQuota = 512_Ki;
+		Config.MemoryQuota = uint32(512_Ki);
 		Config.DiskQuota = 8_Mi;
-		Config.JournalQuota = 7_Ki;
+		Config.JournalQuota = uint32(7_Ki);
 		Config.JournalFlushInterval = 4;
 		Config.DropCache = true;
 		FCache Cache(MoveTemp(Config));
 
-		auto PrimePuts = [&] (int32 PutMax) {
+		auto PrimePuts = [&] (int64 PutMax) {
 			TMap<uint64, FIoBuffer> Ret;
 			while (true)
 			{
-				uint32 Size = MixTh() & ((128_Ki) - 1);
+				uint64 Size = MixTh() & ((128_Ki) - 1);
 				if ((PutMax -= Size) < 0)
 				{
 					break;
@@ -1382,12 +1382,12 @@ IOSTOREONDEMAND_API void Tests()
 		uint32 WriteAllowance;
 
 		// no-op
-		WriteAllowance = 1_Ki;
+		WriteAllowance = uint32(1_Ki);
 		Cache.Flush(WriteAllowance);
 		Cache.Flush(WriteAllowance);
 		Cache.Reset();
 
-		WriteAllowance = 512_Ki;
+		WriteAllowance = uint32(512_Ki);
 		PrimePuts(WriteAllowance);
 		Cache.Flush(0);
 		Cache.Flush(WriteAllowance);
@@ -1420,7 +1420,7 @@ IOSTOREONDEMAND_API void Tests()
 		};
 		check(Validate() == 0);
 
-		WriteAllowance = 512_Ki;
+		WriteAllowance = uint32(512_Ki);
 
 		// simple
 		for (uint32 i : {1, 2, 4, 7, 11})
