Set test hostname to localhost, removed unused variable, and setting Result to zero feels like a good idea.

#rnx

diff --git a/Engine/Source/Runtime/Experimental/IoStoreOnDemand/Private/CoreHttp/Http.cpp b/Engine/Source/Runtime/Experimental/IoStoreOnDemand/Private/CoreHttp/Http.cpp
--- a/Engine/Source/Runtime/Experimental/IoStoreOnDemand/Private/CoreHttp/Http.cpp
+++ b/Engine/Source/Runtime/Experimental/IoStoreOnDemand/Private/CoreHttp/Http.cpp
@@ -2173,6 +2173,7 @@ uint32 FEventLoop::FImpl::Tick(uint32 PollTimeoutMs)
 
 		case FActivity::EState::Completed:
 			--BusyCount;
+			Result = 0;
 			break;
 		}
 
@@ -2559,9 +2560,8 @@ IOSTOREONDEMAND_API void IasHttpTest()
 	MiscTest();
 	ThrottleTest();
 
-#define HOST_NAME "10.24.101.89"
+#define HOST_NAME "localhost"
 	FAnsiStringView TestUrl = "http://" HOST_NAME ":9493/data";
-	FAnsiStringView TestHostName = "localhost";
 
 	struct
 	{
