Fixed narrowing warnings from Clang, enabled the compiler flag.

#rnx

diff --git a/Engine/Source/Runtime/Experimental/IoStoreOnDemand/IoStoreOnDemand.Build.cs b/Engine/Source/Runtime/Experimental/IoStoreOnDemand/IoStoreOnDemand.Build.cs
--- a/Engine/Source/Runtime/Experimental/IoStoreOnDemand/IoStoreOnDemand.Build.cs
+++ b/Engine/Source/Runtime/Experimental/IoStoreOnDemand/IoStoreOnDemand.Build.cs
@@ -16,6 +16,7 @@ public class IoStoreOnDemand : ModuleRules
 		);
 
 		bAllowConfidentialPlatformDefines = true;
+		UnsafeTypeCastWarningLevel = WarningLevel.Error; 
 
 		if (Target.Platform.IsInGroup(UnrealPlatformGroup.Desktop) &&
 			(Target.Type == TargetType.Editor || Target.Type == TargetType.Program))
diff --git a/Engine/Source/Runtime/Experimental/IoStoreOnDemand/Private/CoreHttp/Http.cpp b/Engine/Source/Runtime/Experimental/IoStoreOnDemand/Private/CoreHttp/Http.cpp
--- a/Engine/Source/Runtime/Experimental/IoStoreOnDemand/Private/CoreHttp/Http.cpp
+++ b/Engine/Source/Runtime/Experimental/IoStoreOnDemand/Private/CoreHttp/Http.cpp
@@ -1220,11 +1220,15 @@ static uint64 ReadyCheck(FActivity** Activities, uint32 Num, uint32 TimeoutMs)
 
 		case EWait::Read:
 		case EWait::Write: {
+			// The following looks odd because POLLFD varies subtly from one platform
+			// to the next. To cleanly set members to zero and to not get narrowing
+			// warnings from the compiler, we list-init and don't assume POD types.
 			FSelect& Select = Selects[SelectNum];
+			auto SelectEvent = decltype(FSelect::events)((Activity->SocketWait == EWait::Read) ? POLLIN : POLLOUT);
 			Select = {
 				Activity->Socket,
-				(Activity->SocketWait == EWait::Read) ? POLLIN : POLLOUT,
-				/* 0 */ // list-init so this value zeros with out per-platform hoops
+				SelectEvent,
+				/* 0 */
 			};
 			++SelectNum;
 			} break;
diff --git a/Engine/Source/Runtime/Experimental/IoStoreOnDemand/Private/Statistics.cpp b/Engine/Source/Runtime/Experimental/IoStoreOnDemand/Private/Statistics.cpp
--- a/Engine/Source/Runtime/Experimental/IoStoreOnDemand/Private/Statistics.cpp
+++ b/Engine/Source/Runtime/Experimental/IoStoreOnDemand/Private/Statistics.cpp
@@ -129,13 +129,13 @@ public:
 
 		const double TotalPopulation = static_cast<double>(Count + Other.Count);
 		const double MeanDifference = Mean - Other.Mean;
-		const double A = ((Count - 1) * GetVariance()) + ((Other.Count - 1) * Other.GetVariance());
-		const double B = (MeanDifference) * (MeanDifference) * (Count * Other.Count / TotalPopulation);
+		const double A = (double(Count - 1) * GetVariance()) + (double(Other.Count - 1) * Other.GetVariance());
+		const double B = (MeanDifference) * (MeanDifference) * (double(Count) * double(Other.Count) / TotalPopulation);
 		const double MergedVariance = (A + B) / (TotalPopulation - 1);
 
 		const uint64 NewCount = Count + Other.Count;
 		const double NewMean = ((Mean * double(Count)) + (Other.Mean * double(Other.Count))) / double(NewCount);
-		const double NewVarianceAccumulator = MergedVariance * (NewCount - 1);
+		const double NewVarianceAccumulator = MergedVariance * double(NewCount - 1);
 
 		Count = NewCount;
 		Mean = NewMean;
@@ -449,8 +449,8 @@ void FOnDemandIoBackendStats::OnHttpGet(uint64 SizeBytes, uint64 DurationMs)
 
 	GHttpAvgDuration.Increment(static_cast<double>(DurationMs));
 
-	const double SizeMiB = SizeBytes / (1024.0 * 1024.0);
-	const double DurationSeconds = DurationMs / 1000.0;
+	const double SizeMiB = double(SizeBytes) / (1024.0 * 1024.0);
+	const double DurationSeconds = double(DurationMs) / 1000.0;
 
 	GHttpAvgRate.Increment(SizeMiB / DurationSeconds);
 
