Tweaked my test functions to take a host for easier testing on kits.

#rnx

diff --git a/Engine/Source/Runtime/Experimental/IoStoreOnDemand/Private/CoreHttp/Http.cpp b/Engine/Source/Runtime/Experimental/IoStoreOnDemand/Private/CoreHttp/Http.cpp
--- a/Engine/Source/Runtime/Experimental/IoStoreOnDemand/Private/CoreHttp/Http.cpp
+++ b/Engine/Source/Runtime/Experimental/IoStoreOnDemand/Private/CoreHttp/Http.cpp
@@ -1852,7 +1852,7 @@ static void DoCancel(FActivity* Activity)
 // {{{1 throttler ..............................................................
 
 ////////////////////////////////////////////////////////////////////////////////
-static void ThrottleTest();
+static void ThrottleTest(FAnsiStringView);
 
 ////////////////////////////////////////////////////////////////////////////////
 class FThrottler
@@ -1864,7 +1864,7 @@ public:
 	void	ReturnUnused(uint32 Unused);
 
 private:
-	friend	void ThrottleTest();
+	friend	void ThrottleTest(FAnsiStringView);
 	int32	GetAllowance(uint64 CycleDelta);
 	uint64	CycleFreq;
 	uint64	CycleLast;
@@ -2457,7 +2457,7 @@ static void MiscTest()
 }
 
 ////////////////////////////////////////////////////////////////////////////////
-static void ThrottleTest()
+static void ThrottleTest(FAnsiStringView TestUrl)
 {
 	enum { TheMax = 0x7fff'fffful };
 
@@ -2522,14 +2522,14 @@ static void ThrottleTest()
 	{
 		const uint32 ThrottleKiB = 5;
 
-		TAnsiStringBuilder<128> TestUrl;
-		TestUrl << "http://localhost:9493/data/";
-		TestUrl << (SizeKiB << 10);
+		TAnsiStringBuilder<128> Url;
+		Url << TestUrl;
+		Url << (SizeKiB << 10);
 
 		FEventLoop Loop;
 		Loop.Throttle(ThrottleKiB);
 
-		FRequest Request = Loop.Request("GET", TestUrl).Accept("*/*");
+		FRequest Request = Loop.Request("GET", Url).Accept("*/*");
 		Loop.Send(MoveTemp(Request), [&] (const FTicketStatus& Status) {
 			check(Status.GetId() != FTicketStatus::EId::Error);
 			if (Status.GetId() == FTicketStatus::EId::Response)
@@ -2552,7 +2552,7 @@ static void ThrottleTest()
 }
 
 ////////////////////////////////////////////////////////////////////////////////
-IOSTOREONDEMAND_API void IasHttpTest()
+IOSTOREONDEMAND_API void IasHttpTest(const ANSICHAR* TestHost="localhost")
 {
 #if PLATFORM_WINDOWS
 	WSADATA WsaData;
@@ -2561,11 +2561,17 @@ IOSTOREONDEMAND_API void IasHttpTest()
 	ON_SCOPE_EXIT { WSACleanup(); };
 #endif
 
-	MiscTest();
-	ThrottleTest();
+	TAnsiStringBuilder<64> Ret;
+	auto BuildUrl = [&] (const ANSICHAR* Suffix=nullptr, int32 Port=9493) -> const auto& {
+		Ret.Reset();
+		Ret << "http://";
+		Ret << TestHost;
+		Ret << ":" << Port;
+		return (Suffix != nullptr) ? (Ret << Suffix) : Ret;
+	};
 
-#define HOST_NAME "localhost"
-	FAnsiStringView TestUrl = "http://" HOST_NAME ":9493/data";
+	MiscTest();
+	ThrottleTest(BuildUrl("/data/"));
 
 	struct
 	{
@@ -2659,12 +2665,12 @@ IOSTOREONDEMAND_API void IasHttpTest()
 
 	// unused request
 	{
-		FRequest Request = Loop.Request("GET", TestUrl);
+		FRequest Request = Loop.Request("GET", BuildUrl("/data"));
 	}
 
 	// foundational
 	{
-		FRequest Request = Loop.Request("GET", "http://" HOST_NAME ":9493/seed/493");
+		FRequest Request = Loop.Request("GET", BuildUrl("/seed/493"));
 		Request.Accept(EMimeType::Json);
 
 		FTicket Ticket = Loop.Send(MoveTemp(Request), NullSink);
@@ -2674,11 +2680,11 @@ IOSTOREONDEMAND_API void IasHttpTest()
 
 	// convenience
 	{
-		FRequest Request = Loop.Get(TestUrl).Accept(EMimeType::Json);
+		FRequest Request = Loop.Get(BuildUrl("/data")).Accept(EMimeType::Json);
 
 		FTicket Tickets[] = {
 			Loop.Send(MoveTemp(Request), HashSink),
-			Loop.Send(Loop.Get(TestUrl).Accept(EMimeType::Json), HashSink),
+			Loop.Send(Loop.Get(BuildUrl("/data")).Accept(EMimeType::Json), HashSink),
 			Loop.Send(Loop.Get("http://httpbin.org/get"), NoErrorSink),
 		};
 		WaitForLoopIdle();
@@ -2686,7 +2692,7 @@ IOSTOREONDEMAND_API void IasHttpTest()
 
 	// convenience
 	{
-		FRequest Request = Loop.Get(TestUrl).Accept(EMimeType::Json);
+		FRequest Request = Loop.Get(BuildUrl("/data")).Accept(EMimeType::Json);
 		FTicket Ticket = Loop.Send(MoveTemp(Request), HashSink);
 		WaitForLoopIdle();
 	}
@@ -2694,7 +2700,7 @@ IOSTOREONDEMAND_API void IasHttpTest()
 	// no connect
 	{
 		FRequest Request[] = {
-			Loop.Request("GET", "http://" HOST_NAME ":10930"),
+			Loop.Request("GET", BuildUrl(nullptr, 10930)),
 			Loop.Request("GET", "http://thisdoesnotexistihope/"),
 		};
 		Loop.Send(MoveTemp(Request[0]), NullSink);
@@ -2707,11 +2713,11 @@ IOSTOREONDEMAND_API void IasHttpTest()
 		const uint32 StressLoad = 32;
 
 		struct {
-			FAnsiStringView Url;
+			const ANSICHAR* Uri;
 			bool Disconnect;
 		} StressUrls[] = {
-			{ "http://" HOST_NAME ":9494/data",				false },
-			{ "http://" HOST_NAME ":9494/data?disconnect",	true },
+			{ "/data",				false },
+			{ "/data?disconnect",	true },
 		};
 
 		uint64 Errors = 0;
@@ -2737,18 +2743,11 @@ IOSTOREONDEMAND_API void IasHttpTest()
 			check(false);
 		};
 
-		for (const auto& [StressUrl, ExpectDisconnect] : StressUrls)
+		for (const auto& [StressUri, ExpectDisconnect] : StressUrls)
 		{
-			FTicketSink Sink;
-			if (ExpectDisconnect)
-			{
-				Sink = ErrorSink;
-			}
-			else
-			{
-				Sink = HashSink;
-			}
+			FTicketSink Sink = ExpectDisconnect ? FTicketSink(ErrorSink) : FTicketSink(HashSink);
 
+			const auto& StressUrl = BuildUrl(StressUri, 9494);
 			for (bool AddDelay : {false, true})
 			{
 				FTicket Tickets[StressLoad];
@@ -2772,7 +2771,7 @@ IOSTOREONDEMAND_API void IasHttpTest()
 		const uint32 StressTaskCount = 7;
 		static_assert(StressLoad * StressTaskCount <= 32);
 
-		FAnsiStringView Url = "http://" HOST_NAME ":9494/data";
+		FAnsiStringView Url = BuildUrl("/data");
 
 		auto StressTaskEntry = [&] {
 			for (uint32 i = 0; i < StressLoad; ++i)
