Iris
- Implemented initial support for being able to skip over data for objects that fail to be instantiated

diff --git a/Engine/Plugins/Runtime/ReplicationSystemTestPlugin/Source/Private/Tests/ReplicationSystem/ReplicatedTestObject.cpp b/Engine/Plugins/Runtime/ReplicationSystemTestPlugin/Source/Private/Tests/ReplicationSystem/ReplicatedTestObject.cpp
--- a/Engine/Plugins/Runtime/ReplicationSystemTestPlugin/Source/Private/Tests/ReplicationSystem/ReplicatedTestObject.cpp
+++ b/Engine/Plugins/Runtime/ReplicationSystemTestPlugin/Source/Private/Tests/ReplicationSystem/ReplicatedTestObject.cpp
@@ -104,9 +104,13 @@ bool UReplicatedTestObjectBridge::WriteCreationHeader(UE::Net::FNetSerialization
 	uint16 NumIrisComponentsToSpawn = 0U;
 	uint16 NumDynamicComponentsToSpawn = 0U;
 	uint16 NumConnectionFilteredComponentsToSpawn = 0U;
+	bool bForceFailToInstantiateOnRemote = false;
 
 	const UObject* Object = GetReplicatedObject(Handle);
-	
+	if (const UReplicatedTestObject* ReplicatedTestObject = Cast<UReplicatedTestObject>(Object))
+	{
+		bForceFailToInstantiateOnRemote = ReplicatedTestObject->bForceFailToInstantiateOnRemote;
+	}
 	if (const UTestReplicatedIrisObject* TestReplicatedIrisObject = Cast<UTestReplicatedIrisObject>(Object))
 	{
 		NumComponentsToSpawn = TestReplicatedIrisObject->Components.Num();
@@ -123,6 +127,7 @@ bool UReplicatedTestObjectBridge::WriteCreationHeader(UE::Net::FNetSerialization
 	Writer.WriteBits(NumIrisComponentsToSpawn, 16);
 	Writer.WriteBits(NumDynamicComponentsToSpawn, 16);
 	Writer.WriteBits(NumConnectionFilteredComponentsToSpawn, 16);
+	Writer.WriteBool(bForceFailToInstantiateOnRemote);
 
 	return !Writer.IsOverflown();
 }
@@ -140,6 +145,7 @@ UObjectReplicationBridge::FCreationHeader* UReplicatedTestObjectBridge::ReadCrea
 	Header->NumIrisComponentsToSpawn = Reader.ReadBits(16);
 	Header->NumDynamicComponentsToSpawn = Reader.ReadBits(16);
 	Header->NumConnectionFilteredComponentsToSpawn = Reader.ReadBits(16);
+	Header->bForceFailCreateRemoteInstance = Reader.ReadBool();
 
 	if (Reader.IsOverflown())
 	{
@@ -153,6 +159,12 @@ FObjectReplicationBridgeInstantiateResult UReplicatedTestObjectBridge::BeginInst
 {
 	const FReplicationTestObjectCreationHeader* Header = static_cast<const FReplicationTestObjectCreationHeader*>(InHeader);
 
+	// Force fail to create this remote instance
+	if (Header->bForceFailCreateRemoteInstance)
+	{
+		return FObjectReplicationBridgeInstantiateResult();
+	}
+
 	UObject* ArcheType = StaticFindObject(UObject::StaticClass(), nullptr, *Header->ArchetypeName, false);
 
 	check(ArcheType);
diff --git a/Engine/Plugins/Runtime/ReplicationSystemTestPlugin/Source/Private/Tests/ReplicationSystem/ReplicatedTestObject.h b/Engine/Plugins/Runtime/ReplicationSystemTestPlugin/Source/Private/Tests/ReplicationSystem/ReplicatedTestObject.h
--- a/Engine/Plugins/Runtime/ReplicationSystemTestPlugin/Source/Private/Tests/ReplicationSystem/ReplicatedTestObject.h
+++ b/Engine/Plugins/Runtime/ReplicationSystemTestPlugin/Source/Private/Tests/ReplicationSystem/ReplicatedTestObject.h
@@ -36,6 +36,9 @@ public:
 
 	/** Cached NetRefHandle to simplify testing. DO NOT use this for multi-system tests. */
 	UE::Net::FNetRefHandle NetRefHandle;
+
+	/** If this is set to true, this instance will fail to instantiate on remote end. */
+	bool bForceFailToInstantiateOnRemote = false;
 };
 
 USTRUCT()
@@ -452,6 +455,25 @@ public:
 
 	float GetMaxTickRate() const { return Super::GetMaxTickRate(); }
 
+	// Scope to suppress ensure when testing instantiation failure on clients
+	class FSupressCreateInstanceFailedEnsureScope
+	{
+	public:
+		explicit FSupressCreateInstanceFailedEnsureScope(UReplicatedTestObjectBridge& BridgeIn) : Bridge(BridgeIn), bSuppressCreateInstanceFailedEnsure(Bridge.bSuppressCreateInstanceFailedEnsure)
+		{
+			Bridge.bSuppressCreateInstanceFailedEnsure = true;
+		}
+		~FSupressCreateInstanceFailedEnsureScope()
+		{
+			// Restore
+			Bridge.bSuppressCreateInstanceFailedEnsure = bSuppressCreateInstanceFailedEnsure;			
+		}
+			
+	private:
+		UReplicatedTestObjectBridge& Bridge;
+		bool bSuppressCreateInstanceFailedEnsure;
+	};
+
 protected:
 	// Type specifics for serializing creation data this will most likely be made into a separate interface to support different types of header data for different types
 	// But if we can avoid having custom data per type for instantiating remote objects as we would like to be able to fully express the state of a replicated object using the define protocol alone. 
@@ -463,6 +485,7 @@ protected:
 		uint32 NumIrisComponentsToSpawn;
 		uint32 NumDynamicComponentsToSpawn;
 		uint32 NumConnectionFilteredComponentsToSpawn;
+		bool bForceFailCreateRemoteInstance;
 	};
 
 	virtual bool WriteCreationHeader(UE::Net::FNetSerializationContext& Context, FNetRefHandle Handle) override;
@@ -476,6 +499,7 @@ protected:
 	TArray<TStrongObjectPtr<UObject>>* CreatedObjectsOnNode;
 
 	TFunction<void(FNetRefHandle NetHandle, const UObject* ReplicatedObject, FVector& OutLocation, float& OutCullDistance)> WorldLocationUpdateFunc;
+	bool bForceFailCreateRemoteInstance = false;
 };
 
 extern const UE::Net::FRepTag RepTag_FakeGeneratedReplicationState_IntB;
diff --git a/Engine/Plugins/Runtime/ReplicationSystemTestPlugin/Source/Private/Tests/ReplicationSystem/TestReplicationSystem.cpp b/Engine/Plugins/Runtime/ReplicationSystemTestPlugin/Source/Private/Tests/ReplicationSystem/TestReplicationSystem.cpp
--- a/Engine/Plugins/Runtime/ReplicationSystemTestPlugin/Source/Private/Tests/ReplicationSystem/TestReplicationSystem.cpp
+++ b/Engine/Plugins/Runtime/ReplicationSystemTestPlugin/Source/Private/Tests/ReplicationSystem/TestReplicationSystem.cpp
@@ -4,9 +4,9 @@
 #include "Iris/Serialization/NetBitStreamReader.h"
 #include "Iris/Serialization/NetBitStreamWriter.h"
 #include "Iris/ReplicationSystem/ReplicationSystem.h"
-
 #include "Iris/ReplicationSystem/ReplicationSystemInternal.h"
 #include "Iris/ReplicationSystem/NetTokenStore.h"
+#include "Iris/Core/IrisLog.h"
 
 namespace UE::Net::Private
 {
@@ -1712,4 +1712,90 @@ UE_NET_TEST_FIXTURE(FReplicationSystemServerClientTestFixture, TestObjectPollFra
 	UE_NET_ASSERT_TRUE(SlowPollObjectHasBeenInequal);
 }
 
+// Test that broken objects can be skipped by client
+UE_NET_TEST_FIXTURE(FReplicationSystemServerClientTestFixture, TestClientCanSkipBrokenObject)
+{
+	UReplicationSystem* ReplicationSystem = Server->ReplicationSystem;
+
+	// Add a client
+	FReplicationSystemTestClient* Client = CreateClient();
+
+	// Spawn objects on server
+	UTestReplicatedIrisObject* ServerObjectA = Server->CreateObject(0,0);
+	UTestReplicatedIrisObject* ServerObjectB = Server->CreateObject(0,0);
+
+	{
+		// Setup client to fail to create next remote object
+		ServerObjectA->bForceFailToInstantiateOnRemote = true;
+
+		// Suppress ensure that will occur due to failing to instantiate the object
+		UReplicatedTestObjectBridge::FSupressCreateInstanceFailedEnsureScope SuppressEnsureScope(*Client->GetReplicationBridge());
+
+		// Disable error logging as we know we will fail.
+		auto IrisLogVerbosity = UE_GET_LOG_VERBOSITY(LogIris);
+		LogIris.SetVerbosity(ELogVerbosity::NoLogging);
+
+		// Send and deliver packet
+		Server->PreSendUpdate();
+		Server->SendAndDeliverTo(Client, true);
+		Server->PostSendUpdate();
+
+		// Restore LogVerbosity
+		LogIris.SetVerbosity(IrisLogVerbosity);
+	}
+
+	// We expect replication of ObjectA to have failed
+	{
+		UTestReplicatedIrisObject* ClientObjectA = Cast<UTestReplicatedIrisObject>(Client->GetReplicationBridge()->GetReplicatedObject(ServerObjectA->NetRefHandle));
+		UE_NET_ASSERT_TRUE(ClientObjectA == nullptr);
+	}
+
+	// ObjectB should have been replicated ok
+	{
+		UTestReplicatedIrisObject* ClientObjectB = Cast<UTestReplicatedIrisObject>(Client->GetReplicationBridge()->GetReplicatedObject(ServerObjectB->NetRefHandle));
+		UE_NET_ASSERT_TRUE(ClientObjectB != nullptr);
+	}
+
+	// Modify both objects to make them replicate again
+	++ServerObjectA->IntA;
+	++ServerObjectB->IntA;
+
+	// Send and deliver packet to verify that client ignores the broken object
+	Server->PreSendUpdate();
+	Server->SendAndDeliverTo(Client, true);
+	Server->PostSendUpdate();
+
+	// We expect replication of ObjectA to have failed
+	{
+		UTestReplicatedIrisObject* ClientObjectA = Cast<UTestReplicatedIrisObject>(Client->GetReplicationBridge()->GetReplicatedObject(ServerObjectA->NetRefHandle));
+		UE_NET_ASSERT_TRUE(ClientObjectA == nullptr);
+	}
+
+	// Filter out ObjectA to tell the client that the object has gone out of scope
+	ReplicationSystem->AddToGroup(ReplicationSystem->GetNotReplicatedNetObjectGroup(), ServerObjectA->NetRefHandle);
+
+	// Send and deliver packet, the client should now remove the broken object from the list of broken objects
+	Server->PreSendUpdate();
+	Server->SendAndDeliverTo(Client, true);
+	Server->PostSendUpdate();
+
+	// Enable replication of ObjectA again to try to replicate it to server now that it should succeed
+	ReplicationSystem->RemoveFromGroup(ReplicationSystem->GetNotReplicatedNetObjectGroup(), ServerObjectA->NetRefHandle);
+
+	// Set ObjectA to be able instantiate on client again
+	ServerObjectA->bForceFailToInstantiateOnRemote = false;
+
+	// Client should now be able to instantiate the object
+	Server->PreSendUpdate();
+	Server->SendAndDeliverTo(Client, true);
+	Server->PostSendUpdate();
+
+	// We expect replication of ObjectA to have succeeded this time
+	{
+		UTestReplicatedIrisObject* ClientObjectA = Cast<UTestReplicatedIrisObject>(Client->GetReplicationBridge()->GetReplicatedObject(ServerObjectA->NetRefHandle));
+		UE_NET_ASSERT_TRUE(ClientObjectA == nullptr);
+	}
+}
+
+
 }
diff --git a/Engine/Source/Runtime/Engine/Private/Net/Iris/ReplicationSystem/ActorReplicationBridge.cpp b/Engine/Source/Runtime/Engine/Private/Net/Iris/ReplicationSystem/ActorReplicationBridge.cpp
--- a/Engine/Source/Runtime/Engine/Private/Net/Iris/ReplicationSystem/ActorReplicationBridge.cpp
+++ b/Engine/Source/Runtime/Engine/Private/Net/Iris/ReplicationSystem/ActorReplicationBridge.cpp
@@ -608,8 +608,7 @@ FObjectReplicationBridgeInstantiateResult UActorReplicationBridge::BeginInstanti
 			}
 			else
 			{
-				UE_LOG_ACTORREPLICATIONBRIDGE(Warning, TEXT("OnBeginInstantiateFromRemote Unable to spawn object, failed to resolve archetype with %s"), *(Header->ArchetypeReference.GetRefHandle().ToString()));
-				check(false);
+				UE_LOG_ACTORREPLICATIONBRIDGE(Error, TEXT("OnBeginInstantiateFromRemote Unable to spawn object, failed to resolve archetype with %s"), *DescribeObjectReference(Header->ArchetypeReference, ResolveContext));
 			}
 		}
 		else
@@ -626,7 +625,7 @@ FObjectReplicationBridgeInstantiateResult UActorReplicationBridge::BeginInstanti
 			}
 			else
 			{
-				UE_LOG_ACTORREPLICATIONBRIDGE(Warning, TEXT("OnBeginInstantiateFromRemote Failed to find Resolve ObjectReference for static Actor %s"), *Header->ObjectReference.ToString());
+				UE_LOG_ACTORREPLICATIONBRIDGE(Error, TEXT("OnBeginInstantiateFromRemote Failed to find Resolve ObjectReference for static Actor %s"), *DescribeObjectReference(Header->ObjectReference, ResolveContext));
 			}
 		}
 	}
@@ -652,7 +651,7 @@ FObjectReplicationBridgeInstantiateResult UActorReplicationBridge::BeginInstanti
 
 				if (!SubObj)
 				{
-					UE_LOG_ACTORREPLICATIONBRIDGE(Warning, TEXT("BeginInstantiateFromRemote Failed to find subobjectReference for dynamic SubObject %s, Owner %s (%s)"), *Header->ObjectReference.ToString(), *SubObjectOwnerNetHandle.ToString(), *GetPathNameSafe(GetReplicatedObject(SubObjectOwnerNetHandle)));
+					UE_LOG_ACTORREPLICATIONBRIDGE(Error, TEXT("BeginInstantiateFromRemote Failed to find subobjectReference for dynamic SubObject %s, Owner %s (%s)"), *DescribeObjectReference(Header->ObjectReference, ResolveContext), *SubObjectOwnerNetHandle.ToString(), *GetPathNameSafe(GetReplicatedObject(SubObjectOwnerNetHandle)));
 				}
 			}
 			else
@@ -693,7 +692,7 @@ FObjectReplicationBridgeInstantiateResult UActorReplicationBridge::BeginInstanti
 			}
 			else
 			{
-				UE_LOG_ACTORREPLICATIONBRIDGE(Warning, TEXT("BeginInstantiateFromRemote Failed to find Resolve SubObjectReference for static SubObject %s, Owner %s (%s)"), *Header->ObjectReference.ToString(), *SubObjectOwnerNetHandle.ToString(), *GetPathNameSafe(GetReplicatedObject(SubObjectOwnerNetHandle)));
+				UE_LOG_ACTORREPLICATIONBRIDGE(Error, TEXT("BeginInstantiateFromRemote Failed to find Resolve SubObjectReference for static SubObject %s, Owner %s (%s)"), *DescribeObjectReference(Header->ObjectReference, ResolveContext), *SubObjectOwnerNetHandle.ToString(), *GetPathNameSafe(GetReplicatedObject(SubObjectOwnerNetHandle)));
 			}			
 		}
 	}
diff --git a/Engine/Source/Runtime/Experimental/Iris/Core/Private/Iris/ReplicationSystem/ObjectReferenceCache.cpp b/Engine/Source/Runtime/Experimental/Iris/Core/Private/Iris/ReplicationSystem/ObjectReferenceCache.cpp
--- a/Engine/Source/Runtime/Experimental/Iris/Core/Private/Iris/ReplicationSystem/ObjectReferenceCache.cpp
+++ b/Engine/Source/Runtime/Experimental/Iris/Core/Private/Iris/ReplicationSystem/ObjectReferenceCache.cpp
@@ -119,7 +119,7 @@ bool FObjectReferenceCache::SupportsObjectInternal(const UObject* Object) const
 // $IRIS: $TODO: refactor into utility methods as this is coped from GUIDCache
 bool FObjectReferenceCache::IsDynamicObject(const UObject* Object) const
 {
-	checkSlow(Object != NULL);
+	checkSlow(Object != nullptr);
 	checkSlow(Object->IsSupportedForNetworking());
 
 	// Any non net addressable object is dynamic
@@ -241,42 +241,42 @@ bool FObjectReferenceCache::CreateObjectReferenceInternal(const UObject* Object,
 		if (FCachedNetObjectReference* CachedObjectPtr = ReferenceHandleToCachedReference.Find(RefHandle))
 		{
 			FCachedNetObjectReference& CachedObject = *CachedObjectPtr;
-		if (CachedObject.Object.Get() == Object)
-		{
+			if (CachedObject.Object.Get() == Object)
+			{
 				UE_LOG_REFERENCECACHE(VeryVerbose, TEXT("ObjectReferenceCache::CreateObjectReferenceHandle Found existing %s for ObjectPath %s Object: %s (0x%p), OuterNetRefHandle: %s"),
 				*CachedObject.NetRefHandle.ToString(), ToCStr(Object->GetPathName()), *GetNameSafe(Object), Object, *CachedObject.OuterNetRefHandle.ToString());
 				OutReference = MakeNetObjectReference(CachedObject);
-			return true;
-		}
-		else
-		{
+				return true;
+			}
+			else
+			{
 				UE_LOG_REFERENCECACHE(Verbose, TEXT("ObjectReferenceCache::CreateObjectReferenceHandle Removed %s from ObjectToNetReferenceHandle due to stale cache. Cached Object (0x%p).  New object %s (0x%p)"),
 					*CachedObject.NetRefHandle.ToString(), CachedObject.ObjectKey, *GetNameSafe(Object), Object);
 
-			ObjectToNetReferenceHandle.Remove(CachedObject.ObjectKey);
-			ObjectToNetReferenceHandle.Remove(Object);
+				ObjectToNetReferenceHandle.Remove(CachedObject.ObjectKey);
+				ObjectToNetReferenceHandle.Remove(Object);
 				RefHandlePtr = nullptr;
 
-			if (RefHandle.IsStatic())
-			{
+				if (RefHandle.IsStatic())
+				{
 					UE_LOG_REFERENCECACHE(Verbose, TEXT("ObjectReferenceCache::CreateObjectReferenceHandle clearing stale cache for static handle %s. Cached Object (0x%p).  New object %s (0x%p)"),
 						*CachedObject.NetRefHandle.ToString(), CachedObject.ObjectKey, *GetNameSafe(Object), Object);
 
-				// Note we only cleanse the object reference
-				// we still keep the cachedObject data around in order to be able to serialize static destruction infos
-				CachedObject.ObjectKey = nullptr;
-				CachedObject.Object = nullptr;
+					// Note we only cleanse the object reference
+					// we still keep the cachedObject data around in order to be able to serialize static destruction infos
+					CachedObject.ObjectKey = nullptr;
+					CachedObject.Object = nullptr;
 
-			}
-			else
-			{
+				}
+				else
+				{
 					UE_LOG_REFERENCECACHE(Verbose, TEXT("ObjectReferenceCache::CreateObjectReferenceHandle removing cache for handle %s. Cached Object (0x%p). New object %s (0x%p)"),
 						*CachedObject.NetRefHandle.ToString(), CachedObject.ObjectKey, *GetNameSafe(Object), Object);
 
-				ReferenceHandleToCachedReference.Remove(RefHandle);
+					ReferenceHandleToCachedReference.Remove(RefHandle);
+				}
 			}
 		}
-	}
 		else
 		{
 			UE_LOG_REFERENCECACHE(Warning, TEXT("ObjectReferenceCache::CreateObjectReferenceInternal removed %s from ObjectToNetReferenceHandle due to mismatch with cache. Object %s (0x%p) "), *RefHandle.ToString(), *GetNameSafe(Object), Object);
@@ -799,7 +799,7 @@ UObject* FObjectReferenceCache::ResolveObjectReferenceHandleInternal(FNetRefHand
 				}
 
 				// There is nothing else to do except wait on the delegate to tell us this package is done loading
-				return NULL;
+				return nullptr;
 			}
 			else
 			{
@@ -1682,6 +1682,34 @@ void FObjectReferenceCache::ReadMustBeMappedExports(FNetSerializationContext& Co
 	}
 }
 
+FString FObjectReferenceCache::DescribeObjectReference(const FNetObjectReference Ref, const FNetObjectResolveContext& ResolveContext)
+{
+	FString FullPath;
+
+	GenerateFullPath_r(Ref.GetRefHandle(), ResolveContext, FullPath);
+
+	// Only client assigned FNetObjectReferences has a path stored directly
+	if (Ref.PathToken.IsValid())
+	{
+		// This path is only used by Client to Server references
+		const TCHAR* ResolvedToken = StringTokenStore->ResolveRemoteToken(Ref.PathToken, *ResolveContext.RemoteNetTokenStoreState);
+		if (ResolvedToken)
+		{
+			if (!FullPath.IsEmpty())
+			{
+				FullPath += TEXT(".");
+			}
+
+			FString ObjectPath(ResolvedToken);
+			constexpr bool bReading = true;
+			RenamePathForPie(ResolveContext.ConnectionId, ObjectPath, bReading);
+			FullPath += FString::Printf(TEXT("%s"), ResolvedToken);
+		}
+	}
+
+	return FullPath;
+}
+
 FString FObjectReferenceCache::FullPath(FNetRefHandle RefHandle, const FNetObjectResolveContext& ResolveContext) const
 {
 	FString FullPath;
diff --git a/Engine/Source/Runtime/Experimental/Iris/Core/Private/Iris/ReplicationSystem/ObjectReferenceCache.h b/Engine/Source/Runtime/Experimental/Iris/Core/Private/Iris/ReplicationSystem/ObjectReferenceCache.h
--- a/Engine/Source/Runtime/Experimental/Iris/Core/Private/Iris/ReplicationSystem/ObjectReferenceCache.h
+++ b/Engine/Source/Runtime/Experimental/Iris/Core/Private/Iris/ReplicationSystem/ObjectReferenceCache.h
@@ -131,6 +131,8 @@ public:
 	void UpdateTrackedQueuedBatchObjectReference(const FNetRefHandle InHandle, const UObject* NewObject);
 	void RemoveTrackedQueuedBatchObjectReference(const FNetRefHandle InHandle);
 
+	FString DescribeObjectReference(const FNetObjectReference Ref, const FNetObjectResolveContext& ResolveContext);
+
 private:
 
 	struct FCachedNetObjectReference
diff --git a/Engine/Source/Runtime/Experimental/Iris/Core/Private/Iris/ReplicationSystem/ObjectReplicationBridge.cpp b/Engine/Source/Runtime/Experimental/Iris/Core/Private/Iris/ReplicationSystem/ObjectReplicationBridge.cpp
--- a/Engine/Source/Runtime/Experimental/Iris/Core/Private/Iris/ReplicationSystem/ObjectReplicationBridge.cpp
+++ b/Engine/Source/Runtime/Experimental/Iris/Core/Private/Iris/ReplicationSystem/ObjectReplicationBridge.cpp
@@ -154,6 +154,11 @@ UObject* UObjectReplicationBridge::ResolveObjectReference(const UE::Net::FNetObj
 	return GetObjectReferenceCache()->ResolveObjectReference(Reference, ResolveContext);
 }
 
+FString UObjectReplicationBridge::DescribeObjectReference(const UE::Net::FNetObjectReference& Reference, const UE::Net::FNetObjectResolveContext& ResolveContext)
+{
+	return GetObjectReferenceCache()->DescribeObjectReference(Reference, ResolveContext);
+}
+
 UE::Net::FNetObjectReference UObjectReplicationBridge::GetOrCreateObjectReference(const UObject* Instance) const
 {
 	return GetObjectReferenceCache()->GetOrCreateObjectReference(Instance);
@@ -581,8 +586,9 @@ FReplicationBridgeCreateNetRefHandleResult UObjectReplicationBridge::CreateNetRe
 	// https://jira.it.epicgames.com/browse/UE-127369	
 	FObjectReplicationBridgeInstantiateResult InstantiateResult = BeginInstantiateFromRemote(SubObjectOwnerNetHandle, Context.SerializationContext.GetInternalContext()->ResolveContext, Header.Get());
 	UObject* InstancePtr = InstantiateResult.Object;
-	if (!ensureAlwaysMsgf(InstancePtr, TEXT("Failed to instantiate Handle: %s"), *WantedNetHandle.ToString()))
+	if (!InstancePtr)
 	{
+		ensureAlwaysMsgf(bSuppressCreateInstanceFailedEnsure, TEXT("Failed to instantiate Handle: %s"), *WantedNetHandle.ToString());
 		return CreateResult;
 	}
 
diff --git a/Engine/Source/Runtime/Experimental/Iris/Core/Private/Iris/ReplicationSystem/ReplicationReader.cpp b/Engine/Source/Runtime/Experimental/Iris/Core/Private/Iris/ReplicationSystem/ReplicationReader.cpp
--- a/Engine/Source/Runtime/Experimental/Iris/Core/Private/Iris/ReplicationSystem/ReplicationReader.cpp
+++ b/Engine/Source/Runtime/Experimental/Iris/Core/Private/Iris/ReplicationSystem/ReplicationReader.cpp
@@ -24,6 +24,7 @@
 #include "Iris/Serialization/NetSerializer.h"
 #include "HAL/IConsoleManager.h"
 #include "UObject/Class.h"
+#include "Algo/RemoveIf.h"
 
 #if UE_NET_ENABLE_REPLICATIONREADER_LOG
 #	define UE_LOG_REPLICATIONREADER(Format, ...)  UE_LOG(LogIris, Log, Format, ##__VA_ARGS__)
@@ -650,6 +651,19 @@ uint32 FReplicationReader::ReadObjectBatch(FNetSerializationContext& Context)
 		Reader.Seek(ReturnPos);
 	}
 
+	// Skip over broken objects
+	const bool bIsBroken = BrokenObjects.FindByPredicate([IncompleteHandle](const FNetRefHandle& Entry) { return Entry.GetId() == IncompleteHandle.GetId(); } ) != nullptr;
+	if (bIsBroken)
+	{
+		UE_NET_TRACE_OBJECT_SCOPE(IncompleteHandle, Reader, Context.GetTraceCollector(), ENetTraceVerbosity::Trace);
+		UE_NET_TRACE_SCOPE(SkippedData, Reader, Context.GetTraceCollector(), ENetTraceVerbosity::Trace);
+
+		Reader.Seek(BatchEndPos);
+		return 0U;
+	}
+
+	FPendingBatchData* PendingBatch = PendingBatches.FindByPredicate([&IncompleteHandle](const FPendingBatchData& Entry) { return Entry.Handle == IncompleteHandle; });
+
 	// This object has pending must be mapped references that must be resolved before we can process the data.
 	FPendingBatchData* PendingBatchData = ObjectReferenceCache->ShouldAsyncLoad() ? UpdateUnresolvedMustBeMappedReferences(IncompleteHandle, TempMustBeMappedReferences) : nullptr;
 	if (PendingBatchData)
@@ -689,6 +703,24 @@ uint32 FReplicationReader::ReadObjectBatch(FNetSerializationContext& Context)
 
 		if (Context.HasErrorOrOverflow())
 		{
+			if (Context.GetError() == GNetError_BrokenNetHandle)
+			{
+				// $TODO: Report this to the server so it knows that the state of data in the batch is unknown
+
+				// Log error and try to recover, if get more incoming data for an object in the broken state we will skip it.
+				UE_LOG_REPLICATIONREADER_ERROR(TEXT("FReplicationReader::ReadObject Failed to read object batch handle: %s skipping batch data"), ToCStr(IncompleteHandle.ToString()));
+					
+				BrokenObjects.AddUnique(IncompleteHandle);
+
+				Context.ResetErrorContext();
+
+				UE_NET_TRACE_OBJECT_SCOPE(IncompleteHandle, Reader, Context.GetTraceCollector(), ENetTraceVerbosity::Trace);
+				UE_NET_TRACE_SCOPE(SkippedData, Reader, Context.GetTraceCollector(), ENetTraceVerbosity::Trace);
+
+				// Skip this batch
+				Reader.Seek(BatchEndPos);
+			}
+			
 			return 0U;
 		}
 	}
@@ -782,7 +814,10 @@ void FReplicationReader::ReadObjectInBatch(FNetSerializationContext& Context, FN
 		if (!NetRefHandle.IsValid())
 		{	
 			UE_LOG_REPLICATIONREADER_ERROR(TEXT("FReplicationReader::ReadObject Unable to create handle for %s."), *IncompleteHandle.ToString());
-			Context.SetError(GNetError_InvalidNetHandle);
+
+			// Mark error, but do not mark the bitstream as overflown as we want to handle this error.
+			Context.SetError(GNetError_BrokenNetHandle, false);
+
 			bHasErrors = true;
 			goto ErrorHandling;
 		}
@@ -1745,7 +1780,7 @@ void FReplicationReader::ProcessQueuedBatches()
 		TempMustBeMappedReferences.Reset();
 		UpdateUnresolvedMustBeMappedReferences(PendingBatchData.Handle, TempMustBeMappedReferences);
 
-		// If we have no more pending must be referenes we can apply the received state
+		// If we have no more pending must be references we can apply the received state
 		if (PendingBatchData.PendingMustBeMappedReferences.IsEmpty())		
 		{
 			UE_LOG(LogIris, Verbose, TEXT("ProcessQueuedBatches processing %d queued batches for Handle %s "), PendingBatchData.QueuedDataChunks.Num(), *PendingBatchData.Handle.ToString());
@@ -1772,10 +1807,23 @@ void FReplicationReader::ProcessQueuedBatches()
 						EndReplication(InternalIndex, false, bShouldDestroyInstance);
 					}
 
+					// Remove from broken list
+					BrokenObjects.SetNum(Algo::RemoveIf(BrokenObjects, [&NetRefHandleToEndReplication](const FNetRefHandle& Handle)
+					{
+						return Handle.GetId() == NetRefHandleToEndReplication.GetId();
+					}));
+
 					UE_LOG(LogIris, Verbose, TEXT("FReplicationReader::ProcessQueuedBatches EndReplication for %s while processing queued batches for %s"), *NetRefHandleToEndReplication.ToString(), *PendingBatchData.Handle.ToString());
 					continue;
 				}
 
+				// Skip over broken objects, we still process remaining chunk if the object has been destroyed.
+				const bool bIsBroken = (BrokenObjects.FindByPredicate([&PendingBatchData](const FNetRefHandle& Entry) { return Entry.GetId() == PendingBatchData.Handle.GetId(); }) != nullptr);
+				if (bIsBroken)
+				{
+					continue;
+				}
+
 				// Read and process chunk as it was a received packet for now at least
 				FMemMark TempAllocatorScope(TempLinearAllocator);
 
@@ -1788,8 +1836,20 @@ void FReplicationReader::ProcessQueuedBatches()
 				// $IRIS: $TODO: Implement special dispatch to defer RepNotifies if we are processing multiple batches for the same object.
 				ReadObjectsInBatch(Context, PendingBatchData.Handle, CurrentChunk.bHasBatchOwnerData, CurrentChunk.NumBits);
 
-				// $IRIS: $TODO: What to do if we fail to process this batch? Just delete it? Might need to report this to server as a broken object and build logic to reset replication of the object
-				ensureAlwaysMsgf(!Context.HasErrorOrOverflow(), TEXT("FReplicationReader::ProcessQueuedBatches - Failed to process enqueued batch for %s - %s"), *PendingBatchData.Handle.ToString(), *Context.GetError().ToString());
+				if (Context.HasErrorOrOverflow())
+				{
+					if (Context.GetError() == GNetError_BrokenNetHandle)
+					{
+						// $TODO: Report this to the server so it knows that the state of data in the batch is unknown
+
+						// Log error and try to recover, if get more incoming data for an object in the broken state we will skip it.
+						UE_LOG_REPLICATIONREADER_ERROR(TEXT("FReplicationReader::ProcessQueuedBatches Failed to process object batch handle: %s skipping batch data"), ToCStr(PendingBatchData.Handle.ToString()));
+					
+						BrokenObjects.AddUnique(PendingBatchData.Handle);
+
+						Context.ResetErrorContext();
+					}
+				}
 			
 				// Apply received data and resolve dependencies
 				DispatchStateData(Context);
diff --git a/Engine/Source/Runtime/Experimental/Iris/Core/Private/Iris/ReplicationSystem/ReplicationReader.h b/Engine/Source/Runtime/Experimental/Iris/Core/Private/Iris/ReplicationSystem/ReplicationReader.h
--- a/Engine/Source/Runtime/Experimental/Iris/Core/Private/Iris/ReplicationSystem/ReplicationReader.h
+++ b/Engine/Source/Runtime/Experimental/Iris/Core/Private/Iris/ReplicationSystem/ReplicationReader.h
@@ -270,6 +270,9 @@ private:
 	// We do not expect to have many objects in this state
 	TArray<FPendingBatchData> PendingBatches;
 
+	// We do not expect many objects to be broken
+	TArray<FNetRefHandle> BrokenObjects;
+
 	// Used during receive and processing of pending batches
 	TArray<FNetRefHandle> TempMustBeMappedReferences;
 
diff --git a/Engine/Source/Runtime/Experimental/Iris/Core/Private/Iris/Serialization/NetErrorContext.cpp b/Engine/Source/Runtime/Experimental/Iris/Core/Private/Iris/Serialization/NetErrorContext.cpp
--- a/Engine/Source/Runtime/Experimental/Iris/Core/Private/Iris/Serialization/NetErrorContext.cpp
+++ b/Engine/Source/Runtime/Experimental/Iris/Core/Private/Iris/Serialization/NetErrorContext.cpp
@@ -9,6 +9,7 @@ const FName GNetError_BitStreamOverflow("BitStream overflow");
 const FName GNetError_BitStreamError("BitStream error");
 const FName GNetError_ArraySizeTooLarge("Array size is too large");
 const FName GNetError_InvalidNetHandle("Invalid NetHandle");
+const FName GNetError_BrokenNetHandle("Broken NetHandle");
 const FName GNetError_InvalidValue("Invalid value");
 
 void FNetErrorContext::SetError(const FName InError)
diff --git a/Engine/Source/Runtime/Experimental/Iris/Core/Public/Iris/ReplicationSystem/ObjectReplicationBridge.h b/Engine/Source/Runtime/Experimental/Iris/Core/Public/Iris/ReplicationSystem/ObjectReplicationBridge.h
--- a/Engine/Source/Runtime/Experimental/Iris/Core/Public/Iris/ReplicationSystem/ObjectReplicationBridge.h
+++ b/Engine/Source/Runtime/Experimental/Iris/Core/Public/Iris/ReplicationSystem/ObjectReplicationBridge.h
@@ -80,6 +80,9 @@ public:
 	/** Try to resolve UObject from NetObjectReference, this function tries to resolve the object by loading if necessary. */
 	IRISCORE_API UObject* ResolveObjectReference(const UE::Net::FNetObjectReference& ObjectRef, const UE::Net::FNetObjectResolveContext& ResolveContext);
 
+	/** Describe the NetObjectReference */
+	IRISCORE_API FString DescribeObjectReference(const UE::Net::FNetObjectReference& ObjectRef, const UE::Net::FNetObjectResolveContext& ResolveContext);
+
 	/** Get or create NetObjectReference for object instance. */
 	IRISCORE_API UE::Net::FNetObjectReference GetOrCreateObjectReference(const UObject* Instance) const;
 
@@ -341,6 +344,9 @@ private:
 
 	bool bHasPollOverrides = false;
 	bool bHasDirtyClassesInPollPeriodOverrides = false;
+
+protected:
+	bool bSuppressCreateInstanceFailedEnsure = false;
 };
 
 
diff --git a/Engine/Source/Runtime/Experimental/Iris/Core/Public/Iris/Serialization/NetErrorContext.h b/Engine/Source/Runtime/Experimental/Iris/Core/Public/Iris/Serialization/NetErrorContext.h
--- a/Engine/Source/Runtime/Experimental/Iris/Core/Public/Iris/Serialization/NetErrorContext.h
+++ b/Engine/Source/Runtime/Experimental/Iris/Core/Public/Iris/Serialization/NetErrorContext.h
@@ -12,6 +12,7 @@ IRISCORE_API extern const FName GNetError_BitStreamOverflow;
 IRISCORE_API extern const FName GNetError_BitStreamError;
 IRISCORE_API extern const FName GNetError_ArraySizeTooLarge;
 IRISCORE_API extern const FName GNetError_InvalidNetHandle;
+IRISCORE_API extern const FName GNetError_BrokenNetHandle;
 IRISCORE_API extern const FName GNetError_InvalidValue;
 
 class FNetErrorContext
diff --git a/Engine/Source/Runtime/Experimental/Iris/Core/Public/Iris/Serialization/NetSerializationContext.h b/Engine/Source/Runtime/Experimental/Iris/Core/Public/Iris/Serialization/NetSerializationContext.h
--- a/Engine/Source/Runtime/Experimental/Iris/Core/Public/Iris/Serialization/NetSerializationContext.h
+++ b/Engine/Source/Runtime/Experimental/Iris/Core/Public/Iris/Serialization/NetSerializationContext.h
@@ -41,10 +41,13 @@ public:
 
 	bool HasError() const { return ErrorContext.HasError(); }
 	bool HasErrorOrOverflow() const;
-	/** If an error has already been set calling this function again will be a no-op. */
-	void SetError(const FName Error);
+	/** If an error has already been set calling this function again will be a no-op, if bDoOverFlow is true, the function will also mark the current bitstream as overflown */
+	void SetError(const FName Error, bool bDoOverFlow = true);
 	FName GetError() const { return ErrorContext.GetError(); }
 
+	/** There are cases where an error is handled and reported where we want to stay calm, reset the error context and carry on */
+	void ResetErrorContext() { ErrorContext = FNetErrorContext(); }
+
 	void SetIsInitState(bool bInIsInitState) { bIsInitState = bInIsInitState; }
 	bool IsInitState() const { return bIsInitState; }
 
@@ -149,9 +152,12 @@ inline FNetSerializationContext FNetSerializationContext::MakeSubContext(FNetBit
 	return SubContext;
 }
 
-inline void FNetSerializationContext::SetError(const FName Error)
+inline void FNetSerializationContext::SetError(const FName Error, bool bDoOverFlow)
 {
-	SetBitStreamOverflow();
+	if (bDoOverFlow)
+	{
+		SetBitStreamOverflow();
+	}
 	ErrorContext.SetError(Error);
 }
 
