IAS - only use valid encryption keys when decoding chunk blocks

#rnx

diff --git a/Engine/Source/Runtime/Experimental/IoStoreOnDemand/Private/OnDemandIoDispatcherBackend.cpp b/Engine/Source/Runtime/Experimental/IoStoreOnDemand/Private/OnDemandIoDispatcherBackend.cpp
--- a/Engine/Source/Runtime/Experimental/IoStoreOnDemand/Private/OnDemandIoDispatcherBackend.cpp
+++ b/Engine/Source/Runtime/Experimental/IoStoreOnDemand/Private/OnDemandIoDispatcherBackend.cpp
@@ -339,6 +339,7 @@ void FOnDemandIoStore::AddDeferredContainers()
 		FContainer* Container = *It;
 		if (Container->EncryptionKeyGuid.IsEmpty())
 		{
+			check(Container->EncryptionKey.IsValid() == false);
 			UE_LOG(LogIas, Log, TEXT("Mounting container '%s' (%d entries)"), *Container->Name, Container->TocEntries.Num());
 			RegisteredContainers.Add(Container);
 			It.RemoveCurrent();
@@ -539,8 +540,10 @@ struct FChunkRequestParams
 
 	FIoChunkDecodingParams GetDecodingParams() const
 	{
+		const FAES::FAESKey& EncryptionKey = ChunkInfo.Container->EncryptionKey;
+
 		FIoChunkDecodingParams Params;
-		Params.EncryptionKey = MakeMemoryView(ChunkInfo.Container->EncryptionKey.Key, FAES::FAESKey::KeySize);
+		Params.EncryptionKey = EncryptionKey.IsValid() ? MakeMemoryView(EncryptionKey.Key, FAES::FAESKey::KeySize) : FMemoryView();
 		Params.CompressionFormat = ChunkInfo.Container->CompressionFormat;
 		Params.BlockSize = ChunkInfo.Container->BlockSize;
 		Params.TotalRawSize = ChunkInfo.Entry->RawSize;
