Minor tweak to my test functions

#rnx

diff --git a/Engine/Source/Runtime/Experimental/IoStoreOnDemand/Private/CoreHttp/Http.cpp b/Engine/Source/Runtime/Experimental/IoStoreOnDemand/Private/CoreHttp/Http.cpp
--- a/Engine/Source/Runtime/Experimental/IoStoreOnDemand/Private/CoreHttp/Http.cpp
+++ b/Engine/Source/Runtime/Experimental/IoStoreOnDemand/Private/CoreHttp/Http.cpp
@@ -2466,6 +2466,18 @@ static void ThrottleTest(FAnsiStringView TestUrl)
 	FThrottler Throttler;
 	uint64 OneSecond = Throttler.CycleFreq;
 
+	for (uint32 TargetElapsed : { 1, 5, 7 })
+	{
+		uint64 CycleTest = FPlatformTime::Cycles64();
+		for (uint32 i = 0; i < TargetElapsed; ++i)
+		{
+			FPlatformProcess::Sleep(1.0f);
+		}
+		CycleTest = FPlatformTime::Cycles64() - CycleTest;
+		check((CycleTest + (OneSecond / 8)) / OneSecond == TargetElapsed);
+	}
+
+
 	Throttler.SetLimit(0);
 	check(Throttler.GetAllowance(0) >= TheMax);
 	check(Throttler.GetAllowance()  >= TheMax);
@@ -2545,7 +2557,9 @@ static void ThrottleTest(FAnsiStringView TestUrl)
 
 		// It's dangerous stuff testing elapsed time you know. The +1 is because
 		// throttling assumes one second has already passed when initialised.
+#if PLATFORM_WINDOWS
 		check(Time + 1 == (SizeKiB / ThrottleKiB));
+#endif
 
 		RecvData = FIoBuffer();
 	}
