ClangGame narrowing compile fix

#rnx

diff --git a/Engine/Source/Runtime/Experimental/IoStoreOnDemand/Private/CoreHttp/Http.cpp b/Engine/Source/Runtime/Experimental/IoStoreOnDemand/Private/CoreHttp/Http.cpp
--- a/Engine/Source/Runtime/Experimental/IoStoreOnDemand/Private/CoreHttp/Http.cpp
+++ b/Engine/Source/Runtime/Experimental/IoStoreOnDemand/Private/CoreHttp/Http.cpp
@@ -73,7 +73,7 @@ static_assert(sizeof(sockaddr_in::sin_addr) == sizeof(uint32));
 	static_assert(EWOULDBLOCK == EAGAIN);
 #endif
 
-enum : SocketType { InvalidSocket = -1 };
+static const SocketType InvalidSocket = ~SocketType(0);
 
 // }}}
 
