IAS - added meta data to TOC file with build version and target platform

#rnx

diff --git a/Engine/Source/Runtime/Experimental/IoStoreOnDemand/Private/IoStoreOnDemand.cpp b/Engine/Source/Runtime/Experimental/IoStoreOnDemand/Private/IoStoreOnDemand.cpp
--- a/Engine/Source/Runtime/Experimental/IoStoreOnDemand/Private/IoStoreOnDemand.cpp
+++ b/Engine/Source/Runtime/Experimental/IoStoreOnDemand/Private/IoStoreOnDemand.cpp
@@ -16,6 +16,7 @@
 #include "Misc/CommandLine.h"
 #include "Misc/ConfigCacheIni.h"
 #include "Misc/CoreDelegates.h"
+#include "Misc/DateTime.h"
 #include "Misc/FileHelper.h"
 #include "Misc/Parse.h"
 #include "Modules/ModuleManager.h"
@@ -226,6 +227,38 @@ static FIasCacheConfig GetIasCacheConfig(const TCHAR* CommandLine)
 }
 
 ////////////////////////////////////////////////////////////////////////////////
+FArchive& operator<<(FArchive& Ar, FTocMeta& Meta)
+{
+	Ar << Meta.EpochTimestamp;
+	Ar << Meta.BuildVersion;
+	Ar << Meta.TargetPlatform;
+	return Ar;
+}
+
+FCbWriter& operator<<(FCbWriter& Writer, const FTocMeta& Meta)
+{
+	Writer.BeginObject();
+	Writer.AddInteger(UTF8TEXTVIEW("EpochTimestamp"), Meta.EpochTimestamp);
+	Writer.AddString(UTF8TEXTVIEW("BuildVersion"), Meta.BuildVersion);
+	Writer.AddString(UTF8TEXTVIEW("TargetPlatform"), Meta.TargetPlatform);
+	Writer.EndObject();
+
+	return Writer;
+}
+
+bool LoadFromCompactBinary(FCbFieldView Field, FTocMeta& OutMeta)
+{
+	if (FCbObjectView Obj = Field.AsObjectView())
+	{
+		OutMeta.EpochTimestamp = Obj["EpochTimestamp"].AsInt64();
+		OutMeta.BuildVersion = FString(Obj["BuildVersion"].AsString());
+		OutMeta.TargetPlatform = FString(Obj["TargetPlatform"].AsString());
+		return true;
+	}
+	
+	return false;
+}
+
 FArchive& operator<<(FArchive& Ar, FOnDemandTocHeader& Header)
 {
 	if (Ar.IsLoading() && Ar.TotalSize() < sizeof(FOnDemandTocHeader))
@@ -440,6 +473,11 @@ FArchive& operator<<(FArchive& Ar, FOnDemandToc& Toc)
 	}
 
 	Ar.SetCustomVersion(Toc.VersionGuid, int32(Toc.Header.Version), TEXT("OnDemandToc"));
+
+	if (uint32(Toc.Header.Version) >= uint32(EOnDemandTocVersion::Meta))
+	{
+		Ar << Toc.Meta;
+	}
 	Ar << Toc.Containers;
 
 	return Ar;
@@ -472,6 +510,14 @@ bool LoadFromCompactBinary(FCbFieldView Field, FOnDemandToc& OutToc)
 			return false;
 		}
 
+		if (uint32(OutToc.Header.Version) >= uint32(EOnDemandTocVersion::Meta))
+		{
+			if (!LoadFromCompactBinary(Obj["Meta"], OutToc.Meta))
+			{
+				return false;
+			}
+		}
+
 		FCbArrayView Containers = Obj["Containers"].AsArrayView();
 		OutToc.Containers.Reserve(int32(Containers.Num()));
 		for (FCbFieldView ArrayField : Containers)
@@ -638,6 +684,8 @@ TIoStatusOr<FIoStoreUploadParams> FIoStoreUploadParams::Parse(const TCHAR* Comma
 	FParse::Value(CommandLine, TEXT("SessionToken="), Params.SessionToken);
 	FParse::Value(CommandLine, TEXT("CredentialsFile="), Params.CredentialsFile);
 	FParse::Value(CommandLine, TEXT("CredentialsFileKeyName="), Params.CredentialsFileKeyName);
+	FParse::Value(CommandLine, TEXT("BuildVersion="), Params.BuildVersion);
+	FParse::Value(CommandLine, TEXT("TargetPlatform="), Params.TargetPlatform);
 	Params.bDeleteContainerFiles = !FParse::Param(CommandLine, TEXT("KeepContainerFiles"));
 	Params.bDeletePakFiles = !FParse::Param(CommandLine, TEXT("KeepPakFiles"));
 
@@ -899,13 +947,7 @@ TIoStatusOr<FIoStoreUploadResult> UploadContainerFiles(
 					<< TEXT("/") << HashString
 					<< TEXT(".iochunk");
 
-				const bool bEnqueued = UploadQueue.Enqueue(Key.ToString(), ReadResult.IoBuffer);
-
-				if (bEnqueued)
-				{
-					UE_LOG(LogIas, Display, TEXT("Uploaded chunk '%s/%s/%s'"), *Client.GetConfig().ServiceUrl, *UploadParams.Bucket, Key.ToString());
-				}
-				else
+				if (UploadQueue.Enqueue(Key.ToString(), ReadResult.IoBuffer) == false)
 				{
 					return FIoStatus(EIoErrorCode::WriteError, TEXT("Failed to upload chunk"));
 				}
@@ -971,6 +1013,10 @@ TIoStatusOr<FIoStoreUploadResult> UploadContainerFiles(
 
 	FIoStoreUploadResult UploadResult;
 	{
+		OnDemandToc.Meta.EpochTimestamp = FDateTime::Now().ToUnixTimestamp();
+		OnDemandToc.Meta.BuildVersion = UploadParams.BuildVersion;
+		OnDemandToc.Meta.TargetPlatform = UploadParams.TargetPlatform;
+
 		FLargeMemoryWriter Ar;
 		Ar << OnDemandToc;
 
@@ -1013,6 +1059,8 @@ TIoStatusOr<FIoStoreUploadResult> UploadContainerFiles(
 		UE_LOG(LogIas, Display, TEXT("------------------------------------------------- Upload Summary -------------------------------------------------"));
 		UE_LOG(LogIas, Display, TEXT("%-15s: %s"), TEXT("Service URL"), *UploadParams.ServiceUrl);
 		UE_LOG(LogIas, Display, TEXT("%-15s: %s"), TEXT("Bucket"), *UploadParams.Bucket);
+		UE_LOG(LogIas, Display, TEXT("%-15s: %s"), TEXT("TargetPlatform"), *UploadParams.TargetPlatform);
+		UE_LOG(LogIas, Display, TEXT("%-15s: %s"), TEXT("BuildVersion"), *UploadParams.BuildVersion);
 		UE_LOG(LogIas, Display, TEXT("%-15s: %s"), TEXT("TOC path"), *UploadResult.TocPath);
 		UE_LOG(LogIas, Display, TEXT("%-15s: %.2lf KiB"), TEXT("TOC size"), double(UploadResult.TocSize) / 1024.0);
 		UE_LOG(LogIas, Display, TEXT("%-15s: %.2lf second(s)"), TEXT("Duration"), Duration);
diff --git a/Engine/Source/Runtime/Experimental/IoStoreOnDemand/Public/IO/IoStoreOnDemand.h b/Engine/Source/Runtime/Experimental/IoStoreOnDemand/Public/IO/IoStoreOnDemand.h
--- a/Engine/Source/Runtime/Experimental/IoStoreOnDemand/Public/IO/IoStoreOnDemand.h
+++ b/Engine/Source/Runtime/Experimental/IoStoreOnDemand/Public/IO/IoStoreOnDemand.h
@@ -39,6 +39,7 @@ enum class EOnDemandTocVersion : uint32
 	UTocHash		= 2,
 	BlockHash32		= 3,
 	NoRawHash		= 4,
+	Meta			= 5,
 
 	LatestPlusOne,
 	Latest			= (LatestPlusOne - 1)
@@ -53,6 +54,18 @@ enum class EOnDemandChunkVersion : uint32
 	Latest			= (LatestPlusOne - 1)
 };
 
+struct FTocMeta
+{
+	int64 EpochTimestamp = 0;
+	FString BuildVersion;
+	FString TargetPlatform;
+
+	UE_API friend FArchive& operator<<(FArchive& Ar, FTocMeta& Meta);
+	UE_API friend FCbWriter& operator<<(FCbWriter& Writer, const FTocMeta& Meta);
+};
+
+UE_API bool LoadFromCompactBinary(FCbFieldView Field, FTocMeta& OutMeta);
+
 struct FOnDemandTocHeader
 {
 	static constexpr uint64 ExpectedMagic = 0x6f6e64656d616e64; // ondemand
@@ -103,6 +116,7 @@ UE_API bool LoadFromCompactBinary(FCbFieldView Field, FOnDemandTocContainerEntry
 struct FOnDemandToc
 {
 	FOnDemandTocHeader Header;
+	FTocMeta Meta;
 	TArray<FOnDemandTocContainerEntry> Containers;
 
 	UE_API friend FArchive& operator<<(FArchive& Ar, FOnDemandToc& Toc);
@@ -127,6 +141,8 @@ struct FIoStoreUploadParams
 	FString SessionToken;
 	FString CredentialsFile;
 	FString CredentialsFileKeyName;
+	FString BuildVersion;
+	FString TargetPlatform;
 	int32 MaxConcurrentUploads = 16;
 	bool bDeleteContainerFiles = true;
 	bool bDeletePakFiles = true;
