Geometry collection : Fix memory leak when setting rest state and not using scene proxy

diff --git a/Engine/Source/Runtime/Experimental/GeometryCollectionEngine/Private/GeometryCollection/GeometryCollectionComponent.cpp b/Engine/Source/Runtime/Experimental/GeometryCollectionEngine/Private/GeometryCollection/GeometryCollectionComponent.cpp
--- a/Engine/Source/Runtime/Experimental/GeometryCollectionEngine/Private/GeometryCollection/GeometryCollectionComponent.cpp
+++ b/Engine/Source/Runtime/Experimental/GeometryCollectionEngine/Private/GeometryCollection/GeometryCollectionComponent.cpp
@@ -1350,14 +1350,14 @@ void UGeometryCollectionComponent::SetRestState(TArray<FTransform>&& InRestTrans
 		SetInitialTransforms(RestTransforms);
 	}
 
-	FGeometryCollectionDynamicData* DynamicData = GDynamicDataPool.Allocate();
-	DynamicData->SetPrevTransforms(ComponentSpaceTransforms);
-	CalculateGlobalMatrices();
-	DynamicData->SetTransforms(ComponentSpaceTransforms);
-	DynamicData->IsDynamic = true;
-
 	if (SceneProxy)
 	{
+		FGeometryCollectionDynamicData* DynamicData = GDynamicDataPool.Allocate();
+		DynamicData->SetPrevTransforms(ComponentSpaceTransforms);
+		CalculateGlobalMatrices();
+		DynamicData->SetTransforms(ComponentSpaceTransforms);
+		DynamicData->IsDynamic = true;
+
 #if WITH_EDITOR
 			// We need to do this in case we're controlled by Sequencer in editor, which doesn't invoke PostEditChangeProperty
 			UpdateCachedBounds();
@@ -1384,6 +1384,10 @@ void UGeometryCollectionComponent::SetRestState(TArray<FTransform>&& InRestTrans
 			);
 		}
 	}
+	else
+	{
+		CalculateGlobalMatrices();
+	}
 
 	RefreshEmbeddedGeometry();
 	RefreshCustomRenderer();
