Chaos - increase default velocity iterations to 2 (from 1) and disable the average contact solve

diff --git a/Engine/Source/Runtime/Experimental/Chaos/Private/Chaos/Collision/PBDCollisionSolver.cpp b/Engine/Source/Runtime/Experimental/Chaos/Private/Chaos/Collision/PBDCollisionSolver.cpp
--- a/Engine/Source/Runtime/Experimental/Chaos/Private/Chaos/Collision/PBDCollisionSolver.cpp
+++ b/Engine/Source/Runtime/Experimental/Chaos/Private/Chaos/Collision/PBDCollisionSolver.cpp
@@ -46,7 +46,7 @@ namespace Chaos
 		float Chaos_PBDCollisionSolver_Velocity_MinInvMassScale = Chaos_PBDCollisionSolver_Position_MinInvMassScale;
 		bool bChaos_PBDCollisionSolver_Velocity_FrictionEnabled = true;
 		float Chaos_PBDCollisionSolver_Velocity_StaticFrictionStiffness = 1.0f;
-		bool bChaos_PBDCollisionSolver_Velocity_AveragePointEnabled = true;
+		bool bChaos_PBDCollisionSolver_Velocity_AveragePointEnabled = false;
 
 		FAutoConsoleVariableRef CVarChaos_PBDCollisionSolver_Velocity_SolveEnabled(TEXT("p.Chaos.PBDCollisionSolver.Velocity.SolveEnabled"), bChaos_PBDCollisionSolver_Velocity_SolveEnabled, TEXT(""));
 		FAutoConsoleVariableRef CVarChaos_PBDCollisionSolver_Velocity_MinInvMassScale(TEXT("p.Chaos.PBDCollisionSolver.Velocity.MinInvMassScale"), Chaos_PBDCollisionSolver_Velocity_MinInvMassScale, TEXT(""));
diff --git a/Engine/Source/Runtime/Experimental/Chaos/Public/Chaos/Collision/PBDCollisionSolverSettings.h b/Engine/Source/Runtime/Experimental/Chaos/Public/Chaos/Collision/PBDCollisionSolverSettings.h
--- a/Engine/Source/Runtime/Experimental/Chaos/Public/Chaos/Collision/PBDCollisionSolverSettings.h
+++ b/Engine/Source/Runtime/Experimental/Chaos/Public/Chaos/Collision/PBDCollisionSolverSettings.h
@@ -18,7 +18,7 @@ namespace Chaos
 			, NumPositionFrictionIterations(4)
 			, NumVelocityFrictionIterations(1)
 			, NumPositionShockPropagationIterations(3)
-			, NumVelocityShockPropagationIterations(1)
+			, NumVelocityShockPropagationIterations(2)
 		{
 		}
 
diff --git a/Engine/Source/Runtime/Experimental/Chaos/Public/Chaos/PBDRigidsEvolutionGBF.h b/Engine/Source/Runtime/Experimental/Chaos/Public/Chaos/PBDRigidsEvolutionGBF.h
--- a/Engine/Source/Runtime/Experimental/Chaos/Public/Chaos/PBDRigidsEvolutionGBF.h
+++ b/Engine/Source/Runtime/Experimental/Chaos/Public/Chaos/PBDRigidsEvolutionGBF.h
@@ -58,7 +58,7 @@ namespace Chaos
 
 		// Default iteration counts
 		static constexpr int32 DefaultNumPositionIterations = 8;
-		static constexpr int32 DefaultNumVelocityIterations = 1;
+		static constexpr int32 DefaultNumVelocityIterations = 2;
 		static constexpr int32 DefaultNumProjectionIterations = 1;
 		static constexpr FRealSingle DefaultCollisionMarginFraction = 0.05f;
 		static constexpr FRealSingle DefaultCollisionMarginMax = 10.0f;
