Chaos - enable large manifolds for mesh collisions by default

diff --git a/Engine/Source/Runtime/Experimental/Chaos/Private/Chaos/CollisionOneShotManifolds.cpp b/Engine/Source/Runtime/Experimental/Chaos/Private/Chaos/CollisionOneShotManifolds.cpp
--- a/Engine/Source/Runtime/Experimental/Chaos/Private/Chaos/CollisionOneShotManifolds.cpp
+++ b/Engine/Source/Runtime/Experimental/Chaos/Private/Chaos/CollisionOneShotManifolds.cpp
@@ -82,7 +82,7 @@ namespace Chaos
 	FAutoConsoleVariableRef CVarChaos_Collision_EnableEdgePrune(TEXT("p.Chaos.Collision.EnableEdgePrune"), bChaos_Collision_EnableEdgePrune, TEXT(""));
 	FAutoConsoleVariableRef CVarChaos_Collision_EdgePrunePlaneDistance(TEXT("p.Chaos.Collision.EdgePrunePlaneDistance"), Chaos_Collision_EdgePrunePlaneDistance, TEXT(""));
 
-	int32 Chaos_Collision_ReduceMeshManifoldThreshold = 1000;	// Effectively Disabled
+	int32 Chaos_Collision_ReduceMeshManifoldThreshold = 20;
 	FAutoConsoleVariableRef CVarChaos_Collision_ReduceMeshManifoldThreshold(TEXT("p.Chaos.Collision.ReduceMeshManifoldThreshold"), Chaos_Collision_ReduceMeshManifoldThreshold, TEXT("If an object has more contacts with a mesh than this, do not attempt to reduce the manifold"));
 
 	FRealSingle Chaos_Collision_MeshContactNormalThreshold = 0.98f;	// ~10deg
