Make ISM Pool actor allocation be per level
This fixes an issue where geometry collection using custom renderer do not work with bRequireFullVisibilityToRender

diff --git a/Engine/Source/Runtime/Experimental/GeometryCollectionEngine/Private/GeometryCollection/GeometryCollectionISMPoolRenderer.cpp b/Engine/Source/Runtime/Experimental/GeometryCollectionEngine/Private/GeometryCollection/GeometryCollectionISMPoolRenderer.cpp
--- a/Engine/Source/Runtime/Experimental/GeometryCollectionEngine/Private/GeometryCollection/GeometryCollectionISMPoolRenderer.cpp
+++ b/Engine/Source/Runtime/Experimental/GeometryCollectionEngine/Private/GeometryCollection/GeometryCollectionISMPoolRenderer.cpp
@@ -15,6 +15,7 @@
 
 void UGeometryCollectionISMPoolRenderer::OnRegisterGeometryCollection(UGeometryCollectionComponent const& InComponent)
 {
+	OwningLevel = InComponent.GetComponentLevel();
 }
 
 void UGeometryCollectionISMPoolRenderer::OnUnregisterGeometryCollection()
@@ -72,7 +73,8 @@ UGeometryCollectionISMPoolComponent* UGeometryCollectionISMPoolRenderer::GetOrCr
 	{
 		if (UGeometryCollectionISMPoolSubSystem* ISMPoolSubSystem = UWorld::GetSubsystem<UGeometryCollectionISMPoolSubSystem>(GetWorld()))
 		{
-			ISMPoolActor = ISMPoolSubSystem->FindISMPoolActor();
+			check(OwningLevel);
+			ISMPoolActor = ISMPoolSubSystem->FindISMPoolActor(OwningLevel);
 		}
 	}
 	return ISMPoolActor != nullptr ? ISMPoolActor->GetISMPoolComp() : nullptr;
diff --git a/Engine/Source/Runtime/Experimental/GeometryCollectionEngine/Private/GeometryCollection/GeometryCollectionISMPoolRenderer.h b/Engine/Source/Runtime/Experimental/GeometryCollectionEngine/Private/GeometryCollection/GeometryCollectionISMPoolRenderer.h
--- a/Engine/Source/Runtime/Experimental/GeometryCollectionEngine/Private/GeometryCollection/GeometryCollectionISMPoolRenderer.h
+++ b/Engine/Source/Runtime/Experimental/GeometryCollectionEngine/Private/GeometryCollection/GeometryCollectionISMPoolRenderer.h
@@ -8,6 +8,8 @@
 
 class AGeometryCollectionISMPoolActor;
 class UGeometryCollectionISMPoolComponent;
+class UGeometryCollectionComponent;
+class ULevel;
 
 /** Implementation of a geometry collection custom renderer that pushes AutoInstanceMeshes to an ISMPool. */
 UCLASS()
@@ -43,6 +45,9 @@ protected:
 	FISMPoolGroup MergedMeshGroup;
 	FISMPoolGroup InstancesGroup;
 
+	/** level of the owning component of this renderer */
+	ULevel* OwningLevel = nullptr;
+
 private:
 	UGeometryCollectionISMPoolComponent* GetOrCreateISMPoolComponent();
 	void InitMergedMeshFromGeometryCollection(UGeometryCollection const& InGeometryCollection);
diff --git a/Engine/Source/Runtime/Experimental/GeometryCollectionEngine/Private/GeometryCollection/GeometryCollectionISMPoolSubSystem.cpp b/Engine/Source/Runtime/Experimental/GeometryCollectionEngine/Private/GeometryCollection/GeometryCollectionISMPoolSubSystem.cpp
--- a/Engine/Source/Runtime/Experimental/GeometryCollectionEngine/Private/GeometryCollection/GeometryCollectionISMPoolSubSystem.cpp
+++ b/Engine/Source/Runtime/Experimental/GeometryCollectionEngine/Private/GeometryCollection/GeometryCollectionISMPoolSubSystem.cpp
@@ -8,7 +8,6 @@
 // UGeometryCollectionISMPoolSubSystem
 //----------------------------------------------------------------------//
 UGeometryCollectionISMPoolSubSystem::UGeometryCollectionISMPoolSubSystem()
-	: ISMPoolActor(nullptr)
 {
 }
 
@@ -20,34 +19,53 @@ void UGeometryCollectionISMPoolSubSystem::Initialize(FSubsystemCollectionBase& C
 
 void UGeometryCollectionISMPoolSubSystem::Deinitialize()
 {
-	if (ISMPoolActor)
-	{
-		if (UWorld* World = GetWorld())
-		{
-			GetWorld()->DestroyActor(ISMPoolActor);
-		}
-	}
+	PerLevelISMPoolActors.Reset();
 	Super::Deinitialize();
 }
 
-AGeometryCollectionISMPoolActor* UGeometryCollectionISMPoolSubSystem::FindISMPoolActor()
+AGeometryCollectionISMPoolActor* UGeometryCollectionISMPoolSubSystem::FindISMPoolActor(ULevel* Level)
 {
-	// on demand creation of the actor 
-	// very simple logic for now, we can extend that in the future to return a specific actor based on the criteria
-	if (!ISMPoolActor)
+	AGeometryCollectionISMPoolActor* ISMPoolActor = nullptr;
+
+	// on demand creation of the actor based on level
+	TObjectPtr<AGeometryCollectionISMPoolActor>* ISMPoolActorInLevel = PerLevelISMPoolActors.Find(Level);
+	if (ISMPoolActorInLevel)
+	{
+		ISMPoolActor = ISMPoolActorInLevel->Get();
+	}
+	else
 	{
 		// we keep it as transient to avoid accumulation of those actors in saved levels
 		FActorSpawnParameters Params;
 		Params.ObjectFlags = EObjectFlags::RF_DuplicateTransient | EObjectFlags::RF_Transient;
+		Params.OverrideLevel = Level;
 		ISMPoolActor = GetWorld()->SpawnActor<AGeometryCollectionISMPoolActor>(Params);
+		// spawn can still fail if we are in the middle or tearing down the world
+		if (ISMPoolActor)
+		{
+			PerLevelISMPoolActors.Add(Level, ISMPoolActor);
+			// make sure we capture when the actor get removed so we can update our internal structure accordingly 
+			ISMPoolActor->OnEndPlay.AddDynamic(this, &UGeometryCollectionISMPoolSubSystem::OnActorEndPlay);
+		}
 	}
 	return ISMPoolActor;
 }
 
+void UGeometryCollectionISMPoolSubSystem::OnActorEndPlay(AActor* InSource, EEndPlayReason::Type Reason)
+{
+	if (ULevel* ActorLevel = InSource->GetLevel())
+	{
+		PerLevelISMPoolActors.Remove(ActorLevel);
+	}
+}
+
 void UGeometryCollectionISMPoolSubSystem::GetISMPoolActors(TArray<AGeometryCollectionISMPoolActor*>& OutActors) const
 {
-	if (ISMPoolActor != nullptr)
+	for (const auto& MapEntry : PerLevelISMPoolActors)
 	{
-		OutActors.Add(ISMPoolActor);
+		if (MapEntry.Value != nullptr)
+		{
+			OutActors.Add(MapEntry.Value);
+		}
 	}
 }
diff --git a/Engine/Source/Runtime/Experimental/GeometryCollectionEngine/Public/GeometryCollection/GeometryCollectionISMPoolSubSystem.h b/Engine/Source/Runtime/Experimental/GeometryCollectionEngine/Public/GeometryCollection/GeometryCollectionISMPoolSubSystem.h
--- a/Engine/Source/Runtime/Experimental/GeometryCollectionEngine/Public/GeometryCollection/GeometryCollectionISMPoolSubSystem.h
+++ b/Engine/Source/Runtime/Experimental/GeometryCollectionEngine/Public/GeometryCollection/GeometryCollectionISMPoolSubSystem.h
@@ -6,7 +6,7 @@
 #include "GeometryCollectionISMPoolSubSystem.generated.h"
 
 class AGeometryCollectionISMPoolActor;
-
+class ULevel;
 /**
  * A subsystem managing ISMPool actors.
  * Used by geometry collection now but repurposed for more general use.
@@ -25,13 +25,15 @@ public:
 	// USubsystem END
 
 	/** Finds or creates an actor. */
-	GEOMETRYCOLLECTIONENGINE_API AGeometryCollectionISMPoolActor* FindISMPoolActor();
+	GEOMETRYCOLLECTIONENGINE_API AGeometryCollectionISMPoolActor* FindISMPoolActor(ULevel* Level);
 	
 	/** Get all actors managed by the subsystem. */
 	GEOMETRYCOLLECTIONENGINE_API void GetISMPoolActors(TArray<AGeometryCollectionISMPoolActor*>& OutActors) const;
 
 protected:
-	/** For now we only use one ISMPool actor per world, but we could extend the system to manage many more and return the right one based on search criteria. */
-	UPROPERTY(Transient)
-	TObjectPtr<AGeometryCollectionISMPoolActor> ISMPoolActor;
+	UFUNCTION()
+	void OnActorEndPlay(AActor* InSource, EEndPlayReason::Type Reason);
+
+	/** ISMPool are per level **/
+	TMap<TObjectPtr<ULevel>, TObjectPtr<AGeometryCollectionISMPoolActor> > PerLevelISMPoolActors;
 };
