Chaos - minor mesh collision pruning optimization

diff --git a/Engine/Source/Runtime/Experimental/Chaos/Private/Chaos/Collision/ContactTriangles.cpp b/Engine/Source/Runtime/Experimental/Chaos/Private/Chaos/Collision/ContactTriangles.cpp
--- a/Engine/Source/Runtime/Experimental/Chaos/Private/Chaos/Collision/ContactTriangles.cpp
+++ b/Engine/Source/Runtime/Experimental/Chaos/Private/Chaos/Collision/ContactTriangles.cpp
@@ -411,6 +411,8 @@ namespace Chaos
 	// infinitly extending the face along the normal) then fix it to be in the valid range.
 	void FContactTriangleCollector::FixInvalidNormalContactPoints()
 	{
+		const FReal ContactDotNormalThreshold = Chaos_Collision_MeshContactNormalThreshold;
+
 		const int32 NumContactPoints = TriangleContactPoints.Num();
 		for (int32 ContactIndex = 0; ContactIndex < NumContactPoints; ++ContactIndex)
 		{
@@ -429,6 +431,13 @@ namespace Chaos
 				const int32 ContactTriangleIndex = ContactPointData.GetTriangleIndex();
 				const FContactTriangle& ContactTriangle = ContactTriangles[ContactTriangleIndex];
 
+				// If the normal points roughly along the triangle normal, it won't be too bad even if it's technically invalid so just leave it
+				const FReal ContactDotNormal = FReal(ContactPointData.GetContactNormalDotTriangleNormal());
+				if (ContactDotNormal > ContactDotNormalThreshold)
+				{
+					continue;
+				}
+
 				// Find the vertex positions and indices for this contact
 				int32 VertexIndexA, VertexIndexB;
 				FVec3 VertexA, VertexB;
@@ -445,7 +454,6 @@ namespace Chaos
 						{
 							const FContactTriangle& OtherContactTriangle = ContactTriangles[OtherContactTriangleIndex];
 							const FReal MinContactDotNormal = FVec3::DotProduct(OtherContactTriangle.FaceNormal, ContactTriangle.FaceNormal);
-							const FReal ContactDotNormal = FVec3::DotProduct(ContactPoint.ShapeContactNormal, ContactTriangle.FaceNormal);
 							if (ContactDotNormal < MinContactDotNormal)
 							{
 								// We are outside the valid normal range for this edge
@@ -548,7 +556,7 @@ namespace Chaos
 
 			if (ContactPointData.IsEnabled())
 			{
-				const FReal ContactNormalDotFaceNormal = FVec3::DotProduct(ContactPoint.ShapeContactNormal, ContactTriangle.FaceNormal);
+				const FReal ContactNormalDotFaceNormal = ContactPointData.GetContactNormalDotTriangleNormal();// FVec3::DotProduct(ContactPoint.ShapeContactNormal, ContactTriangle.FaceNormal);
 				if (ContactNormalDotFaceNormal < 0)
 				{
 					DisableContact(ContactIndex);
@@ -609,9 +617,9 @@ namespace Chaos
 				if (bIsSamePosition)
 				{
 					// Keep the contact with the normal that is most along the triangle face normal. If normals are similar, keep the deepest.
-					const FReal NormalDot0 = FVec3::DotProduct(ContactPoint0.ShapeContactNormal, ContactTriangle0.FaceNormal);
-					const FReal NormalDot1 = FVec3::DotProduct(ContactPoint1.ShapeContactNormal, ContactTriangle1.FaceNormal);
-					const bool bSameNormal = (FMath::Abs(NormalDot0 - NormalDot1) < 0.01);
+					const FReal NormalDot0 = ContactPointData0.GetContactNormalDotTriangleNormal();// FVec3::DotProduct(ContactPoint0.ShapeContactNormal, ContactTriangle0.FaceNormal);
+					const FReal NormalDot1 = ContactPointData1.GetContactNormalDotTriangleNormal();// FVec3::DotProduct(ContactPoint1.ShapeContactNormal, ContactTriangle1.FaceNormal);
+					const bool bSameNormal = FMath::IsNearlyEqual(NormalDot0, NormalDot1, FReal(0.01));
 					const bool bTake0 = (bSameNormal && (ContactPoint0.Phi < ContactPoint1.Phi)) || (!bSameNormal && (NormalDot0 > NormalDot1));
 					if (bTake0)
 					{
