Chaos - support manifold reuse for mesh contacts (if large manifolds are enabled)

diff --git a/Engine/Source/Runtime/Experimental/Chaos/Private/Chaos/Collision/ContactTriangles.cpp b/Engine/Source/Runtime/Experimental/Chaos/Private/Chaos/Collision/ContactTriangles.cpp
--- a/Engine/Source/Runtime/Experimental/Chaos/Private/Chaos/Collision/ContactTriangles.cpp
+++ b/Engine/Source/Runtime/Experimental/Chaos/Private/Chaos/Collision/ContactTriangles.cpp
@@ -9,7 +9,7 @@ namespace Chaos
 		extern int32 ChaosSolverDebugDrawMeshContacts;
 	}
 	extern bool bChaos_Collision_EnableEdgePrune;
-	extern int32 Chaos_Collision_ReduceMeshManifoldThreshold;
+	extern bool bChaos_Collision_EnableLargeMeshManifolds;
 	extern FRealSingle Chaos_Collision_MeshContactNormalThreshold;
 	extern bool bChaos_Collision_MeshManifoldSortByDistance;
 
@@ -198,9 +198,6 @@ namespace Chaos
 
 	void FContactTriangleCollector::ProcessContacts(const FRigidTransform3& MeshToConvexTransform)
 	{
-		// @todo(chaos): Some of the pruning mechanisms don't work very well now that we have support for large manifolds.
-		const bool bEnableFullPruning = ((TriangleContactPoints.Num() - NumDisabledTriangleContactPoints) < Chaos_Collision_ReduceMeshManifoldThreshold);
-
 		// Build the feature set from the full list of contacts before pruning
 		BuildContactFeatureSets();
 
@@ -210,7 +207,7 @@ namespace Chaos
 		if ((TriangleContactPoints.Num() > 1))
 		{
 			const bool bPruneEdges = true;
-			const bool bPruneVertices = bEnableFullPruning;	// @todo(chaos): fix - when set, removes too many interior contacts
+			const bool bPruneVertices = !bChaos_Collision_EnableLargeMeshManifolds;	// @todo(chaos): fix - when set, removes too many interior contacts
 			PruneEdgeAndVertexContactPoints(bPruneEdges, bPruneVertices);
 		}
 
@@ -233,13 +230,9 @@ namespace Chaos
 
 		DebugDrawContactPoints(FColor::Yellow, 0.7);
 
-		// Reduce to only 4 contact points from here
-		// @todo(chaos): A 4 point manifold is often not sufficient when we have low curvature and we end up removing contacts that
-		// would be required during the solve to prevent jitter. For now we disable the reduction when we are colliding
-		// with "lots" of triangles (large objects relative to triangle size), but keep the original behaviour for 
-		// smaller objects so we don't impact perf for the most common cases until we do this properly.
-		if (bEnableFullPruning)
+		if (!bChaos_Collision_EnableLargeMeshManifolds)
 		{
+			// Reduce to 4 contact points from here
 			ReduceManifoldContactPointsTriangeMesh();
 		}
 
diff --git a/Engine/Source/Runtime/Experimental/Chaos/Private/Chaos/Collision/PBDCollisionConstraint.cpp b/Engine/Source/Runtime/Experimental/Chaos/Private/Chaos/Collision/PBDCollisionConstraint.cpp
--- a/Engine/Source/Runtime/Experimental/Chaos/Private/Chaos/Collision/PBDCollisionConstraint.cpp
+++ b/Engine/Source/Runtime/Experimental/Chaos/Private/Chaos/Collision/PBDCollisionConstraint.cpp
@@ -691,10 +691,6 @@ namespace Chaos
 
 		// Reset current closest point
 		ClosestManifoldPointIndex = INDEX_NONE;
-
-		// How many manifold points we expect. E.g., for Box-box this will be 4 or 1 depending on whether
-		// we have a face or edge contact. We don't reuse the manifold if we lose points after culling here
-		ExpectedNumManifoldPoints = ManifoldPoints.Num();
 		Flags.bWasManifoldRestored = false;
 
 		if (!Flags.bCanRestoreManifold)
@@ -730,16 +726,22 @@ namespace Chaos
 
 		// Either update or remove each manifold point depending on how far it has moved from its initial relative point
 		// NOTE: We do not reset if we have 0 points - we can still "restore" a zero point manifold if the bodies have not moved
-		int32 ManifoldPointToRemove = INDEX_NONE;
 		if (ManifoldPoints.Num() > 0)
 		{
 			const FRigidTransform3 Shape0ToShape1Transform = ShapeWorldTransforms[0].GetRelativeTransformNoScale(ShapeWorldTransforms[1]);
 			
-			// Update or prune manifold points. If we would end up removing more than 1 point, we just throw the 
-			// whole manifold away because it will get rebuilt in the narrow phasee anyway.
+			// Update or prune manifold points
+			int32 NumActiveManifoldPoints = ManifoldPoints.Num();
 			for (int32 ManifoldPointIndex = 0; ManifoldPointIndex < ManifoldPoints.Num(); ++ManifoldPointIndex)
 			{
 				FManifoldPoint& ManifoldPoint = ManifoldPoints[ManifoldPointIndex];
+				
+				// Point may have been disabled on a previous tick
+				if (ManifoldPoint.Flags.bDisabled)
+				{
+					--NumActiveManifoldPoints;
+					continue;
+				}
 
 				// Calculate the world-space contact location and separation at the current shape transforms
 				// @todo(chaos): this should use the normal owner. Currently we assume body 1 is the owner
@@ -780,27 +782,24 @@ namespace Chaos
 						ClosestManifoldPointIndex = ManifoldPointIndex;
 					}
 				}
-				else if ((ManifoldPointToRemove == INDEX_NONE) && (bChaos_Collision_EnableManifoldGJKReplace || bChaos_Collision_EnableManifoldGJKInject))
-				{
-					// We can reject up to 1 point (if we have GJK point injection enabled)
-					ManifoldPointToRemove = ManifoldPointIndex;
-				}
 				else
 				{
-					// We want to remove a(nother) point, but we will never reuse the manifold now so throw it away
-					return false;
+					// This point moved too far - disable it
+					ManifoldPoint.Flags.bDisabled = true;
+					ManifoldPoint.Flags.bWasRestored = false;
+					ManifoldPoint.Flags.bWasReplaced = false;
+					--NumActiveManifoldPoints;
 				}
 			}
 
-			// Remove points - only one point removal support required (see above)
-			if (ManifoldPointToRemove != INDEX_NONE)
+			// We want to retain a stable supporting manifold, so at least 4 points. Unless we started with less than that, 
+			// in which case make sure it doesn't go any lower.
+			ExpectedNumManifoldPoints = FMath::Min(ManifoldPoints.Num(), 4);
+
+			// If we removed too many points, we must rebuild the manifold
+			// NOTE: We can only add up to one point
+			if (NumActiveManifoldPoints < ExpectedNumManifoldPoints)
 			{
-				ManifoldPoints.RemoveAt(ManifoldPointToRemove);
-				if ((ManifoldPointToRemove < ClosestManifoldPointIndex) && (ClosestManifoldPointIndex != INDEX_NONE))
-				{
-					--ClosestManifoldPointIndex;
-					check(ClosestManifoldPointIndex >= 0);
-				}
 				return false;
 			}
 		}
diff --git a/Engine/Source/Runtime/Experimental/Chaos/Private/Chaos/CollisionOneShotManifolds.cpp b/Engine/Source/Runtime/Experimental/Chaos/Private/Chaos/CollisionOneShotManifolds.cpp
--- a/Engine/Source/Runtime/Experimental/Chaos/Private/Chaos/CollisionOneShotManifolds.cpp
+++ b/Engine/Source/Runtime/Experimental/Chaos/Private/Chaos/CollisionOneShotManifolds.cpp
@@ -82,8 +82,8 @@ namespace Chaos
 	FAutoConsoleVariableRef CVarChaos_Collision_EnableEdgePrune(TEXT("p.Chaos.Collision.EnableEdgePrune"), bChaos_Collision_EnableEdgePrune, TEXT(""));
 	FAutoConsoleVariableRef CVarChaos_Collision_EdgePrunePlaneDistance(TEXT("p.Chaos.Collision.EdgePrunePlaneDistance"), Chaos_Collision_EdgePrunePlaneDistance, TEXT(""));
 
-	int32 Chaos_Collision_ReduceMeshManifoldThreshold = 20;
-	FAutoConsoleVariableRef CVarChaos_Collision_ReduceMeshManifoldThreshold(TEXT("p.Chaos.Collision.ReduceMeshManifoldThreshold"), Chaos_Collision_ReduceMeshManifoldThreshold, TEXT("If an object has more contacts with a mesh than this, do not attempt to reduce the manifold"));
+	bool bChaos_Collision_EnableLargeMeshManifolds = 1;
+	FAutoConsoleVariableRef CVarChaos_Collision_EnableLargeMeshManifolds(TEXT("p.Chaos.Collision.EnableLargeMeshManifolds"), bChaos_Collision_EnableLargeMeshManifolds, TEXT("Whether to allow large mesh manifolds for collisions against meshes (required for good behaviour)"));
 
 	FRealSingle Chaos_Collision_MeshContactNormalThreshold = 0.98f;	// ~10deg
 	FAutoConsoleVariableRef CVarChaos_Collision_MeshContactNormalThreshold(TEXT("p.Chaos.Collision.MeshContactNormalThreshold"), Chaos_Collision_MeshContactNormalThreshold, TEXT("Treat contact with a dot product between the normal and the triangle face greater than this as face collisions"));
diff --git a/Engine/Source/Runtime/Experimental/Chaos/Public/Chaos/Collision/CollisionUtil.h b/Engine/Source/Runtime/Experimental/Chaos/Public/Chaos/Collision/CollisionUtil.h
--- a/Engine/Source/Runtime/Experimental/Chaos/Public/Chaos/Collision/CollisionUtil.h
+++ b/Engine/Source/Runtime/Experimental/Chaos/Public/Chaos/Collision/CollisionUtil.h
@@ -8,6 +8,7 @@
 namespace Chaos
 {
 	extern bool bChaos_Collision_EnableBoundsChecks;
+	extern bool bChaos_Collision_EnableLargeMeshManifolds;
 
 	namespace Private
 	{
@@ -50,8 +51,8 @@ namespace Chaos
 			const bool bIsSphere1 = (ImplicitType1 == ImplicitObjectType::Sphere);
 			const bool bIsCapsule0 = (ImplicitType0 == ImplicitObjectType::Capsule);
 			const bool bIsCapsule1 = (ImplicitType1 == ImplicitObjectType::Capsule);
-			const bool bIsTriangle0 = (ImplicitType0 == ImplicitObjectType::TriangleMesh) || (ImplicitType0 == ImplicitObjectType::HeightField);
-			const bool bIsTriangle1 = (ImplicitType1 == ImplicitObjectType::TriangleMesh) || (ImplicitType1 == ImplicitObjectType::HeightField);
+			const bool bIsMesh0 = (ImplicitType0 == ImplicitObjectType::TriangleMesh) || (ImplicitType0 == ImplicitObjectType::HeightField);
+			const bool bIsMesh1 = (ImplicitType1 == ImplicitObjectType::TriangleMesh) || (ImplicitType1 == ImplicitObjectType::HeightField);
 			const bool bIsLevelSet = ((ImplicitType0 == ImplicitObjectType::LevelSet) || (ImplicitType1 == ImplicitObjectType::LevelSet));
 
 			const bool bHasBounds0 = (Implicit0 != nullptr) && Implicit0->HasBoundingBox();
@@ -67,7 +68,14 @@ namespace Chaos
 
 			// Do not try to reuse manifold points for capsules or spheres (against anything)
 			// NOTE: This can also be disabled for all shape types by the solver (see GenerateCollisionImpl and the Context)
-			Flags.bEnableManifoldUpdate = !bIsSphere0 && !bIsSphere1 && !bIsCapsule0 && !bIsCapsule1 && !bIsTriangle0 && !bIsTriangle1 && !bIsLevelSet;
+			Flags.bEnableManifoldUpdate = !bIsSphere0 && !bIsSphere1 && !bIsCapsule0 && !bIsCapsule1 && !bIsLevelSet;
+
+			// If we don't allow large mesh manifolds we can't reuse manifolds for meshes because we may have 
+			// thrown away some points that we needed at the new transform
+			if (!bChaos_Collision_EnableLargeMeshManifolds)
+			{
+				Flags.bEnableManifoldUpdate = Flags.bEnableManifoldUpdate && !bIsMesh0 && !bIsMesh1;
+			}
 
 			// Mark probe flag now so we know which GenerateCollisions to use
 			// @todo(chaos): it looks like this can be changed by a collision modifier so we should not be caching it
