Renamed the component parts of IAS' cache to better reflect their role and to avoid overloading terms and renaming things is always fun

#rnx

diff --git a/Engine/Source/Runtime/Experimental/IoStoreOnDemand/Private/JournaledCache.cpp b/Engine/Source/Runtime/Experimental/IoStoreOnDemand/Private/JournaledCache.cpp
--- a/Engine/Source/Runtime/Experimental/IoStoreOnDemand/Private/JournaledCache.cpp
+++ b/Engine/Source/Runtime/Experimental/IoStoreOnDemand/Private/JournaledCache.cpp
@@ -53,7 +53,7 @@ struct FDebugCacheEntry
 	using Callback = void(void* Param, const FDebugCacheEntry&);
 	uint64 Key;
 	uint32 Size;
-	uint32 IsPending : 1;
+	uint32 IsMemCache : 1;
 	uint32 _Unused : 31;
 };
 
@@ -91,10 +91,10 @@ using EntryHandle = UPTRINT;
 
 
 
-// {{{1 pending ................................................................
+// {{{1 mem-cache ..............................................................
 
 ////////////////////////////////////////////////////////////////////////////////
-class FPending
+class FMemCache
 {
 public:
 	struct FItem
@@ -105,7 +105,7 @@ public:
 
 	using PeelItems = TArray<FItem>;
 
-					FPending(uint32 InMaxSize=64 << 10);
+					FMemCache(uint32 InMaxSize=64 << 10);
 	void			Reset();
 	uint32			GetCount() const	{ return Items.Num(); }
 	uint32			GetUsed() const		{ return UsedSize; }
@@ -126,20 +126,20 @@ private:
 };
 
 ////////////////////////////////////////////////////////////////////////////////
-FPending::FPending(uint32 InMaxSize)
+FMemCache::FMemCache(uint32 InMaxSize)
 : MaxSize(InMaxSize)
 {
 }
 
 ////////////////////////////////////////////////////////////////////////////////
-void FPending::Reset()
+void FMemCache::Reset()
 {
 	Items.Reset();
 	UsedSize = 0;
 }
 
 ////////////////////////////////////////////////////////////////////////////////
-EntryHandle FPending::Get(uint64 Key) const
+EntryHandle FMemCache::Get(uint64 Key) const
 {
 	for (auto& Item : Items)
 	{
@@ -153,7 +153,7 @@ EntryHandle FPending::Get(uint64 Key) const
 }
 
 ////////////////////////////////////////////////////////////////////////////////
-bool FPending::Materialize(EntryHandle Handle, FIoBuffer& Out, uint32 Offset) const
+bool FMemCache::Materialize(EntryHandle Handle, FIoBuffer& Out, uint32 Offset) const
 {
 	TRACE_CPUPROFILER_EVENT_SCOPE(IasCache::Materialize_Memory);
 
@@ -180,7 +180,7 @@ bool FPending::Materialize(EntryHandle Handle, FIoBuffer& Out, uint32 Offset) co
 }
 
 ////////////////////////////////////////////////////////////////////////////////
-bool FPending::Put(uint64 Key, FIoBuffer&& Data)
+bool FMemCache::Put(uint64 Key, FIoBuffer&& Data)
 {
 	if (Get(Key) != 0)
 	{
@@ -210,7 +210,7 @@ bool FPending::Put(uint64 Key, FIoBuffer&& Data)
 }
 
 ////////////////////////////////////////////////////////////////////////////////
-int32 FPending::Peel(int32 PeelSize, PeelItems& Out)
+int32 FMemCache::Peel(int32 PeelSize, PeelItems& Out)
 {
 	return DropImpl(PeelSize, [&Out] (FItem&& Item) {
 		Out.Add(MoveTemp(Item));
@@ -218,10 +218,10 @@ int32 FPending::Peel(int32 PeelSize, PeelItems& Out)
 }
 
 ////////////////////////////////////////////////////////////////////////////////
-uint32 FPending::DebugVisit(void* Param, FDebugCacheEntry::Callback* Callback)
+uint32 FMemCache::DebugVisit(void* Param, FDebugCacheEntry::Callback* Callback)
 {
 	FDebugCacheEntry Out = {};
-	Out.IsPending = 1;
+	Out.IsMemCache = 1;
 	for (auto& Item : Items)
 	{
 		Out.Key = Item.Key;
@@ -233,7 +233,7 @@ uint32 FPending::DebugVisit(void* Param, FDebugCacheEntry::Callback* Callback)
 
 ////////////////////////////////////////////////////////////////////////////////
 template <typename Lambda>
-int32 FPending::DropImpl(uint32 Size, Lambda&& Callback)
+int32 FMemCache::DropImpl(uint32 Size, Lambda&& Callback)
 {
 	int32 DropSize = 0;
 	int32 TargetSize = FMath::Min<int32>(Size, UsedSize);
@@ -261,15 +261,15 @@ int32 FPending::DropImpl(uint32 Size, Lambda&& Callback)
 }
 
 ////////////////////////////////////////////////////////////////////////////////
-int32 FPending::Drop(uint32 Size)
+int32 FMemCache::Drop(uint32 Size)
 {
 	return DropImpl(Size, [] (FItem&&) {});
 }
 
-// {{{1 journal ................................................................
+// {{{1 disk-cache .............................................................
 
 ////////////////////////////////////////////////////////////////////////////////
-class FJournal
+class FDiskCache
 {
 private:
 	static const uint32 MAGIC = 0x04930001;
@@ -308,7 +308,7 @@ public:
 		int32				MaxEntries;
 	};
 
-							FJournal(FString&& Path, uint64 InMaxDataSize, uint32 InJournalSize);
+							FDiskCache(FString&& Path, uint64 InMaxDataSize, uint32 InJournalSize);
 	void					Reset();
 	FPhrase					OpenPhrase(uint32 DataSize);
 	void					ClosePhrase(FPhrase& Phrase);
@@ -349,7 +349,7 @@ private:
 };
 
 ////////////////////////////////////////////////////////////////////////////////
-bool FJournal::FPhrase::Add(uint64 Key, FIoBuffer&& Data)
+bool FDiskCache::FPhrase::Add(uint64 Key, FIoBuffer&& Data)
 {
 	check(MaxEntries > 0);
 	uint32 DataSize = uint32(Data.GetSize());
@@ -361,7 +361,7 @@ bool FJournal::FPhrase::Add(uint64 Key, FIoBuffer&& Data)
 }
 
 ////////////////////////////////////////////////////////////////////////////////
-FJournal::FJournal(FString&& Path, uint64 InMaxDataSize, uint32 InJournalSize)
+FDiskCache::FDiskCache(FString&& Path, uint64 InMaxDataSize, uint32 InJournalSize)
 : BinPath(Path)
 , MaxDataSize(InMaxDataSize)
 , JournalSize(InJournalSize)
@@ -375,7 +375,7 @@ FJournal::FJournal(FString&& Path, uint64 InMaxDataSize, uint32 InJournalSize)
 }
 
 ////////////////////////////////////////////////////////////////////////////////
-void FJournal::OpenJrnFile()
+void FDiskCache::OpenJrnFile()
 {
 	TStringBuilder<265> JrnPath;
 	JrnPath << BinPath;
@@ -387,7 +387,7 @@ void FJournal::OpenJrnFile()
 }
 
 ////////////////////////////////////////////////////////////////////////////////
-void FJournal::Reset()
+void FDiskCache::Reset()
 {
 	DataMap.Reset();
 	MappedBytes = 0;
@@ -399,7 +399,7 @@ void FJournal::Reset()
 }
 
 ////////////////////////////////////////////////////////////////////////////////
-FJournal::FPhrase FJournal::OpenPhrase(uint32 DataSize)
+FDiskCache::FPhrase FDiskCache::OpenPhrase(uint32 DataSize)
 {
 	check((JournalCursor & (sizeof(FDataEntry) - 1)) == 0);
 
@@ -416,7 +416,7 @@ FJournal::FPhrase FJournal::OpenPhrase(uint32 DataSize)
 }
 
 ////////////////////////////////////////////////////////////////////////////////
-void FJournal::ClosePhrase(FPhrase& Phrase)
+void FDiskCache::ClosePhrase(FPhrase& Phrase)
 {
 	int32 EntryCount = Entries.Num() - Phrase.Index - 1;
 	if (EntryCount <= 0)
@@ -464,13 +464,13 @@ void FJournal::ClosePhrase(FPhrase& Phrase)
 }
 
 ////////////////////////////////////////////////////////////////////////////////
-EntryHandle FJournal::Get(uint64 Key) const
+EntryHandle FDiskCache::Get(uint64 Key) const
 {
 	return UPTRINT(DataMap.Find(Key));
 }
 
 ////////////////////////////////////////////////////////////////////////////////
-bool FJournal::Materialize(EntryHandle Handle, FIoBuffer& Out, uint32 Offset) const
+bool FDiskCache::Materialize(EntryHandle Handle, FIoBuffer& Out, uint32 Offset) const
 {
 	TRACE_CPUPROFILER_EVENT_SCOPE(IasCache::Materialize_Disk);
 
@@ -497,7 +497,7 @@ bool FJournal::Materialize(EntryHandle Handle, FIoBuffer& Out, uint32 Offset) co
 }
 
 ////////////////////////////////////////////////////////////////////////////////
-uint64 FJournal::Insert(uint64 DataBase, const FDataEntry& Entry)
+uint64 FDiskCache::Insert(uint64 DataBase, const FDataEntry& Entry)
 {
 	FMapEntry Value;
 	Value.DataCursor = DataBase + Entry.Offset;
@@ -508,7 +508,7 @@ uint64 FJournal::Insert(uint64 DataBase, const FDataEntry& Entry)
 }
 
 ////////////////////////////////////////////////////////////////////////////////
-uint64 FJournal::Insert(const FPhraseDesc* Phrase)
+uint64 FDiskCache::Insert(const FPhraseDesc* Phrase)
 {
 	uint64 TotalSize = 0;
 	uint64 DataBase = Phrase->DataCursor;
@@ -523,7 +523,7 @@ uint64 FJournal::Insert(const FPhraseDesc* Phrase)
 }
 
 ////////////////////////////////////////////////////////////////////////////////
-void FJournal::Prune(uint64 DataBase, uint32 Size)
+void FDiskCache::Prune(uint64 DataBase, uint32 Size)
 {
 	TRACE_CPUPROFILER_EVENT_SCOPE(IasCache::Prune);
 
@@ -568,14 +568,14 @@ void FJournal::Prune(uint64 DataBase, uint32 Size)
 }
 
 ////////////////////////////////////////////////////////////////////////////////
-int32 FJournal::Flush()
+int32 FDiskCache::Flush()
 {
 	if (Entries.IsEmpty())
 	{
 		return 0;
 	}
 
-	TRACE_CPUPROFILER_EVENT_SCOPE(IasCache::Flush_Journal);
+	TRACE_CPUPROFILER_EVENT_SCOPE(IasCache::Flush_DiskCache);
 
 	uint32 Size = Entries.Num() * sizeof(Entries[0]);
 
@@ -602,7 +602,7 @@ int32 FJournal::Flush()
 }
 
 ////////////////////////////////////////////////////////////////////////////////
-void FJournal::Drop()
+void FDiskCache::Drop()
 {
 	IPlatformFile& Ipf = IPlatformFile::GetPlatformPhysical();
 	Ipf.DeleteFile(*BinPath);
@@ -611,7 +611,7 @@ void FJournal::Drop()
 }
 
 ////////////////////////////////////////////////////////////////////////////////
-bool FJournal::Load()
+bool FDiskCache::Load()
 {
 	IPlatformFile& Ipf = IPlatformFile::GetPlatformPhysical();
 
@@ -809,7 +809,7 @@ bool FJournal::Load()
 		Insert(Holm.Phrase);
 	}
 	
-	// Prime the journal's state
+	// Prime the disk-cache's state
 	const FParagraph& LastPara = Paragraphs.Last();
 	const FPhraseDesc* LastPhrase = LastPara.Phrase;
 	Marker = LastPara.Marker + 1;
@@ -836,7 +836,7 @@ bool FJournal::Load()
 }
 
 ////////////////////////////////////////////////////////////////////////////////
-void FJournal::Spam()
+void FDiskCache::Spam()
 {
 	UE_LOG(LogIas, VeryVerbose,
 		TEXT("JournaledCache: MappedKiB=%llu Entries=%d DataCur=%llu JournalCur=%u Marker=%u)"),
@@ -849,7 +849,7 @@ void FJournal::Spam()
 }
 
 ////////////////////////////////////////////////////////////////////////////////
-uint32 FJournal::DebugVisit(void* Param, FDebugCacheEntry::Callback* Callback)
+uint32 FDiskCache::DebugVisit(void* Param, FDebugCacheEntry::Callback* Callback)
 {
 	FDebugCacheEntry Out = {};
 	for (const auto& Entry : DataMap)
@@ -862,7 +862,7 @@ uint32 FJournal::DebugVisit(void* Param, FDebugCacheEntry::Callback* Callback)
 }
 
 ////////////////////////////////////////////////////////////////////////////////
-void FJournal::PhraseReset()
+void FDiskCache::PhraseReset()
 {
 	Entries.Reset();
 }
@@ -914,8 +914,8 @@ public:
 
 private:
 	mutable FRWLock	Lock;
-	FPending		Pending;
-	FJournal		Journal;
+	FMemCache		MemCache;
+	FDiskCache		DiskCache;
 	uint32			FlushIndex = 0;
 	uint32			FlushPeriod = 4;
 };
@@ -930,8 +930,8 @@ bool FCache::FEntry::Materialize(FIoBuffer& Out, uint32 Offset)
 
 	switch (HitType)
 	{
-	case EHit::Memory:	((const FPending*)Owner)->Materialize(Handle, Out, Offset); break;
-	case EHit::Disk:	((const FJournal*)Owner)->Materialize(Handle, Out, Offset); break;
+	case EHit::Memory:	((const FMemCache*)Owner)->Materialize(Handle, Out, Offset); break;
+	case EHit::Disk:	((const FDiskCache*)Owner)->Materialize(Handle, Out, Offset); break;
 	case EHit::None:	return false;
 	}
 
@@ -941,13 +941,13 @@ bool FCache::FEntry::Materialize(FIoBuffer& Out, uint32 Offset)
 
 ////////////////////////////////////////////////////////////////////////////////
 FCache::FCache(FConfig&& Config)
-: Pending(Config.MemoryQuota)
-, Journal(MoveTemp(Config.Path), Config.DiskQuota, Config.JournalQuota)
+: MemCache(Config.MemoryQuota)
+, DiskCache(MoveTemp(Config.Path), Config.DiskQuota, Config.JournalQuota)
 , FlushPeriod(Config.JournalFlushInterval)
 {
 	if (Config.DropCache)
 	{
-		Journal.Drop();
+		DiskCache.Drop();
 	}
 }
 
@@ -955,8 +955,8 @@ FCache::FCache(FConfig&& Config)
 void FCache::Reset()
 {
 	FWriteAccess _(Lock);
-	Pending.Reset();
-	Journal.Reset();
+	MemCache.Reset();
+	DiskCache.Reset();
 	FlushIndex = 0;
 }
 
@@ -964,7 +964,7 @@ void FCache::Reset()
 bool FCache::Load()
 {
 	FWriteAccess _(Lock);
-	return Journal.Load();
+	return DiskCache.Load();
 }
 
 ////////////////////////////////////////////////////////////////////////////////
@@ -973,16 +973,16 @@ FCache::FEntry FCache::Get(uint64 Key) const
 	FEntry Ret;
 	Ret.Lock = FReadAccess(Lock);
 
-	if (Ret.Handle = Journal.Get(Key); Ret.Handle)
+	if (Ret.Handle = DiskCache.Get(Key); Ret.Handle)
 	{
 		Ret.HitType = EHit::Disk;
-		Ret.Owner = &Journal;
+		Ret.Owner = &DiskCache;
 	}
 		
-	else if (Ret.Handle = Pending.Get(Key); Ret.Handle)
+	else if (Ret.Handle = MemCache.Get(Key); Ret.Handle)
 	{
 		Ret.HitType = EHit::Memory;
-		Ret.Owner = &Pending;
+		Ret.Owner = &MemCache;
 	}
 
 	return Ret;
@@ -993,7 +993,7 @@ bool FCache::Put(uint64 Key, FIoBuffer& Data)
 {
 	FIoBuffer Cloned = Data;
 	FWriteAccess _(Lock);
-	return Pending.Put(Key, MoveTemp(Cloned));
+	return MemCache.Put(Key, MoveTemp(Cloned));
 }
 
 ////////////////////////////////////////////////////////////////////////////////
@@ -1006,26 +1006,26 @@ int32 FCache::Flush(int32 Allowance)
 		{
 			FlushIndex -= FlushPeriod;
 			FWriteAccess _(Lock);
-			return Allowance - Journal.Flush();
+			return Allowance - DiskCache.Flush();
 		}
 
-		if (Pending.GetUsed() == 0)
+		if (MemCache.GetUsed() == 0)
 		{
 			FlushIndex -= (FlushIndex == 1);
 			FWriteAccess _(Lock);
-			return Allowance - Journal.Flush();
+			return Allowance - DiskCache.Flush();
 		}
 	}
 
-	TRACE_CPUPROFILER_EVENT_SCOPE(IasCache::Flush_Pending);
+	TRACE_CPUPROFILER_EVENT_SCOPE(IasCache::Flush_MemCache);
 
-	FPending::PeelItems PeelItems;
+	FMemCache::PeelItems PeelItems;
 
 	FWriteAccess _(Lock);
 
-	int32 PendingSize = Pending.Peel(Allowance, PeelItems);
+	int32 MemCacheSize = MemCache.Peel(Allowance, PeelItems);
 
-	FJournal::FPhrase Phrase = Journal.OpenPhrase(PendingSize);
+	FDiskCache::FPhrase Phrase = DiskCache.OpenPhrase(MemCacheSize);
 	int32 PeelIndex = -1;
 	for (int32 i = 0, n = PeelItems.Num(); i < n; ++i)
 	{
@@ -1036,20 +1036,20 @@ int32 FCache::Flush(int32 Allowance)
 			break;
 		}
 	}
-	Journal.ClosePhrase(Phrase);
+	DiskCache.ClosePhrase(Phrase);
 
 	if (PeelIndex >= 0)
 	{
 		/* end of journal reached so not all peeled items could be added, may
-		 * we can re-add leftover peeled items back to pending? */
+		 * we can re-add leftover peeled items back to mem-cache? */
 	}
 
 	if (FlushPeriod <= 1)
 	{
-		Allowance -= Journal.Flush();
+		Allowance -= DiskCache.Flush();
 	}
 
-	return Allowance - PendingSize;
+	return Allowance - MemCacheSize;
 }
 
 ////////////////////////////////////////////////////////////////////////////////
@@ -1057,8 +1057,8 @@ uint32 FCache::DebugVisit(void* Param, FDebugCacheEntry::Callback* Callback)
 {
 	FReadAccess _(Lock);
 	uint32 Count = 0;
-	Count += Pending.DebugVisit(Param, Callback);
-	Count += Journal.DebugVisit(Param, Callback);
+	Count += MemCache.DebugVisit(Param, Callback);
+	Count += DiskCache.DebugVisit(Param, Callback);
 	return Count;
 }
 
@@ -1294,7 +1294,7 @@ IOSTOREONDEMAND_API void Tests()
 		return FIoBuffer(FIoBuffer::Wrap, Working + Offset, Size);
 	};
 
-	// Pending {{{2
+	// MemCache {{{2
 	{
 		struct {
 			int32 Size;
@@ -1309,44 +1309,44 @@ IOSTOREONDEMAND_API void Tests()
 
 		for (auto& [Size, Expected] : TestCases)
 		{
-			FPending Pending(Size);
+			FMemCache MemCache(Size);
 
-			Pending.Put(0x493, DummyData(0));
-			check(Pending.GetCount() == 0);
+			MemCache.Put(0x493, DummyData(0));
+			check(MemCache.GetCount() == 0);
 
-			Pending.Put(0x493, DummyData(513));
-			Pending.Put(0xa9e, DummyData(511));
-			check(Pending.GetUsed() == Expected);
+			MemCache.Put(0x493, DummyData(513));
+			MemCache.Put(0xa9e, DummyData(511));
+			check(MemCache.GetUsed() == Expected);
 
-			Pending.Put(0x49e, DummyData(11));
+			MemCache.Put(0x49e, DummyData(11));
 			Expected = (Expected == 0) ? 0 : (511 + 11);
-			check(Pending.GetUsed() == Expected);
+			check(MemCache.GetUsed() == Expected);
 		}
 
-		FPending::PeelItems Peeled;
+		FMemCache::PeelItems Peeled;
 
-		FPending Pending(64);
-		Pending.Put(1, DummyData(1));
-		check(Pending.Peel(0, Peeled) == 0);
+		FMemCache MemCache(64);
+		MemCache.Put(1, DummyData(1));
+		check(MemCache.Peel(0, Peeled) == 0);
 		check(Peeled.Num() == 0);
-		check(Pending.Peel(64, Peeled) == 1);
+		check(MemCache.Peel(64, Peeled) == 1);
 		check(Peeled.Num() == 1);
-		check(Pending.GetUsed() == 0);
+		check(MemCache.GetUsed() == 0);
 		Peeled.Reset();
 
-		Pending = FPending(64);
+		MemCache = FMemCache(64);
 		for (int32 i = 0; i < 64; ++i)
 		{
-			Pending.Put(i + 1, DummyData(1));
+			MemCache.Put(i + 1, DummyData(1));
 		}
 
-		check(Pending.Peel(32, Peeled) == 32);
+		check(MemCache.Peel(32, Peeled) == 32);
 		check(Peeled.Num() == 32);
-		check(Pending.GetUsed() == 32);
+		check(MemCache.GetUsed() == 32);
 		for (auto& [Key, _] : Peeled)
 		{
 			FIoBuffer Data;
-			check(Pending.Get(Key) == 0);
+			check(MemCache.Get(Key) == 0);
 		}
 		Peeled.Reset();
 	} // }}}
@@ -1413,7 +1413,7 @@ IOSTOREONDEMAND_API void Tests()
 				int32 IsFromDisk = 0;
 				IsFromDisk |= Data.GetData() >= State->WorkRange[1];
 				IsFromDisk |= (Data.GetData() + Data.GetSize()) <= State->WorkRange[0];
-				check(Entry.IsPending != IsFromDisk);
+				check(Entry.IsMemCache != IsFromDisk);
 			};
 			FVisitState State = {Cache, {Working, Working + WorkingSize}};
 			return Cache.DebugVisit(&State, Visitor);
