Refactor phrase building out to be more standalone so we're better placed to create multiple disk lanes for IAS' cache

#rnx

diff --git a/Engine/Source/Runtime/Experimental/IoStoreOnDemand/Private/JournaledCache.cpp b/Engine/Source/Runtime/Experimental/IoStoreOnDemand/Private/JournaledCache.cpp
--- a/Engine/Source/Runtime/Experimental/IoStoreOnDemand/Private/JournaledCache.cpp
+++ b/Engine/Source/Runtime/Experimental/IoStoreOnDemand/Private/JournaledCache.cpp
@@ -269,49 +269,115 @@ int32 FMemCache::Drop(uint32 Size)
 // {{{1 disk-cache .............................................................
 
 ////////////////////////////////////////////////////////////////////////////////
-class FDiskCache
+static const uint32 MAGIC = 0x04930001;
+static const uint32 SIZE_BITS = 25;
+
+using MarkerType = uint32;
+
+struct FDataEntry
+{
+	uint64		Key;
+	uint64		Offset : 23;
+	uint64		Size : SIZE_BITS;
+	uint64		EntryCount : 16;
+};
+
+struct FPhraseDesc
 {
+	uint32		Magic;
+	MarkerType	Marker;
+	uint64		EntryCount : 16;
+	uint64		DataCursor : 48;
+	FDataEntry	Entries[];
+};
+
+static_assert(sizeof(FPhraseDesc) == 16);
+static_assert(sizeof(FPhraseDesc) == sizeof(FDataEntry));
+
+////////////////////////////////////////////////////////////////////////////////
+struct FDiskPhrase
+{
+						FDiskPhrase(TArray<FDataEntry>& InEntries, int32 InMaxEntries, uint32 DataSize);
+						~FDiskPhrase();
+	bool				Add(uint64 Key, FIoBuffer&& Data);
+	const uint8*		Finalize(MarkerType Marker, uint64 DataCursor) const;
+	int32				GetEntryCount() const	{ return Entries.Num() - Index - 1; }
+	uint32				GetSize() const			{ return Cursor; }
+	const FPhraseDesc&	GetDesc() const			{ return (FPhraseDesc&)(Entries[Index]); }
+
 private:
-	static const uint32 MAGIC = 0x04930001;
-	static const uint32 SIZE_BITS = 25;
+	TUniquePtr<uint8[]>	Buffer;
+	TArray<FDataEntry>&	Entries;
+	uint32				Cursor = sizeof(MarkerType);
+	uint32				Index;
+	int32				MaxEntries;
 
-	using MarkerType = uint32;
+private:
+						FDiskPhrase(const FDiskPhrase&) = delete;
+	FDiskPhrase&		operator = (const FDiskPhrase&) = delete;
+};
 
-	struct FDataEntry
-	{
-		uint64		Key;
-		uint64		Offset : 23;
-		uint64		Size : SIZE_BITS;
-		uint64		EntryCount : 16;
-	};
+////////////////////////////////////////////////////////////////////////////////
+FDiskPhrase::FDiskPhrase(TArray<FDataEntry>& InEntries, int32 InMaxEntries, uint32 DataSize)
+: Entries(InEntries)
+, Index(Entries.Num())
+, MaxEntries(InMaxEntries)
+{
+	DataSize += Cursor;
+	Buffer = TUniquePtr<uint8[]>(new uint8[DataSize]);
 
-	struct FPhraseDesc
+	Entries.Add(FDataEntry{});
+}
+
+////////////////////////////////////////////////////////////////////////////////
+FDiskPhrase::~FDiskPhrase()
+{
+	if (GetEntryCount() == 0)
 	{
-		uint32		Magic;
-		MarkerType	Marker;
-		uint64		EntryCount : 16;
-		uint64		DataCursor : 48;
-		FDataEntry	Entries[];
-	};
+		Entries.Pop();
+	}
+}
+
+////////////////////////////////////////////////////////////////////////////////
+bool FDiskPhrase::Add(uint64 Key, FIoBuffer&& Data)
+{
+	check(MaxEntries > 0);
+	uint32 DataSize = uint32(Data.GetSize());
+	check(DataSize < (1 << SIZE_BITS));
+	Entries.Add({Key, Cursor, DataSize});
+	std::memcpy(Buffer.Get() + Cursor, Data.GetData(), DataSize);
+	Cursor += DataSize;
+	return (--MaxEntries > 0);
+}
+
+////////////////////////////////////////////////////////////////////////////////
+const uint8* FDiskPhrase::Finalize(MarkerType Marker, uint64 DataCursor) const
+{
+	uint32 EntryCount = GetEntryCount();
+	Entries.Last().EntryCount = uint16(EntryCount);
+
+	auto& Desc = (FPhraseDesc&)(Entries[Index]);
+	Desc.Magic = MAGIC;
+	Desc.Marker = Marker;
+	Desc.DataCursor = DataCursor;
+	Desc.EntryCount = uint16(EntryCount);
+
+	uint8* Ret = Buffer.Get();
+	std::memcpy(Ret, &Marker, sizeof(MarkerType));
+	return Ret;
+}
 
-	static_assert(sizeof(FPhraseDesc) == 16);
-	static_assert(sizeof(FPhraseDesc) == sizeof(FDataEntry));
 
-public:
-	struct FPhrase
-	{
-		bool				Add(uint64 Key, FIoBuffer&& Data);
-		TUniquePtr<uint8[]>	Buffer;
-		TArray<FDataEntry>*	Entries;
-		uint32				Cursor;
-		uint32				Index;
-		int32				MaxEntries;
-	};
 
+////////////////////////////////////////////////////////////////////////////////
+class FDiskCache
+{
+private:
+public:
 							FDiskCache(FString&& Path, uint64 InMaxDataSize, uint32 InJournalSize);
 	void					Reset();
-	FPhrase					OpenPhrase(uint32 DataSize);
-	void					ClosePhrase(FPhrase& Phrase);
+	FDiskPhrase				OpenPhrase(uint32 DataSize);
+	void					ClosePhrase(FDiskPhrase&& Phrase);
 	EntryHandle				Get(uint64 Key) const;
 	bool					Materialize(EntryHandle Handle, FIoBuffer& Out, uint32 Offset=0) const;
 	int32					Flush();
@@ -333,7 +399,7 @@ private:
 	void					Spam();
 	void					PhraseReset();
 	uint64					Insert(uint64 DataBase, const FDataEntry& Entry);
-	uint64					Insert(const FPhraseDesc* Phrase);
+	uint64					Insert(const FPhraseDesc& Phrase);
 	void					Prune(uint64 DataBase, uint32 Size);
 	TArray<FDataEntry>		Entries;
 	FString					BinPath;
@@ -348,18 +414,6 @@ private:
 	MarkerType				Marker = 0;
 };
 
-////////////////////////////////////////////////////////////////////////////////
-bool FDiskCache::FPhrase::Add(uint64 Key, FIoBuffer&& Data)
-{
-	check(MaxEntries > 0);
-	uint32 DataSize = uint32(Data.GetSize());
-	check(DataSize < (1 << SIZE_BITS));
-	Entries->Add({Key, Cursor, DataSize});
-	std::memcpy(Buffer.Get() + Cursor, Data.GetData(), DataSize);
-	Cursor += DataSize;
-	return (--MaxEntries > 0);
-}
-
 ////////////////////////////////////////////////////////////////////////////////
 FDiskCache::FDiskCache(FString&& Path, uint64 InMaxDataSize, uint32 InJournalSize)
 : BinPath(Path)
@@ -399,50 +453,32 @@ void FDiskCache::Reset()
 }
 
 ////////////////////////////////////////////////////////////////////////////////
-FDiskCache::FPhrase FDiskCache::OpenPhrase(uint32 DataSize)
+FDiskPhrase FDiskCache::OpenPhrase(uint32 DataSize)
 {
 	check((JournalCursor & (sizeof(FDataEntry) - 1)) == 0);
 
-	FPhrase Ret;
-	Ret.Buffer = TUniquePtr<uint8[]>(new uint8[DataSize + sizeof(MarkerType)]);
-	Ret.Entries = &Entries;
-	Ret.Cursor = sizeof(MarkerType);
-	Ret.Index = Entries.Num();
-	Ret.MaxEntries = int32((JournalSize - JournalCursor) / sizeof(FDataEntry)) - 1;
+	int32 MaxEntries = int32((JournalSize - JournalCursor) / sizeof(FDataEntry)) - 1;
 
-	Entries.Add(FDataEntry{});
-
-	return Ret;
+	return FDiskPhrase(Entries, MaxEntries, DataSize);
 }
 
 ////////////////////////////////////////////////////////////////////////////////
-void FDiskCache::ClosePhrase(FPhrase& Phrase)
+void FDiskCache::ClosePhrase(FDiskPhrase&& Phrase)
 {
-	int32 EntryCount = Entries.Num() - Phrase.Index - 1;
+	int32 EntryCount = Phrase.GetEntryCount();
 	if (EntryCount <= 0)
 	{
-		Entries.Pop();
 		return;
 	}
 
-	Entries.Last().EntryCount = uint16(EntryCount);
-
-	uint32 WriteSize = Phrase.Cursor;
+	uint32 WriteSize = Phrase.GetSize();
 	if (DataCursor + WriteSize > MaxDataSize)
 	{
 		OverRemoval = 0;
 		DataCursor = 0;
 	}
 
-	auto& Header = (FPhraseDesc&)(Entries[Phrase.Index]);
-	Header.Magic = MAGIC;
-	Header.Marker = Marker;
-	Header.DataCursor = DataCursor;
-	Header.EntryCount = uint16(EntryCount);
-
-	uint8* Buffer = Phrase.Buffer.Get();
-
-	*(MarkerType*)Buffer = Marker;
+	const uint8* Buffer = Phrase.Finalize(Marker, DataCursor);
 	++Marker;
 
 	IPlatformFile& Ipf = IPlatformFile::GetPlatformPhysical();
@@ -455,12 +491,10 @@ void FDiskCache::ClosePhrase(FPhrase& Phrase)
 		File.Reset();
 
 		Prune(DataCursor, WriteSize);
-		Insert(&Header);
+		Insert(Phrase.GetDesc());
 
 		DataCursor += WriteSize;
 	}
-
-	Phrase.Buffer.Reset();
 }
 
 ////////////////////////////////////////////////////////////////////////////////
@@ -508,13 +542,13 @@ uint64 FDiskCache::Insert(uint64 DataBase, const FDataEntry& Entry)
 }
 
 ////////////////////////////////////////////////////////////////////////////////
-uint64 FDiskCache::Insert(const FPhraseDesc* Phrase)
+uint64 FDiskCache::Insert(const FPhraseDesc& Phrase)
 {
 	uint64 TotalSize = 0;
-	uint64 DataBase = Phrase->DataCursor;
-	for (uint32 i = 0; i < Phrase->EntryCount; ++i)
+	uint64 DataBase = Phrase.DataCursor;
+	for (uint32 i = 0; i < Phrase.EntryCount; ++i)
 	{
-		TotalSize += Insert(DataBase, Phrase->Entries[i]);
+		TotalSize += Insert(DataBase, Phrase.Entries[i]);
 	}
 
 	MappedBytes += TotalSize;
@@ -806,7 +840,7 @@ bool FDiskCache::Load()
 	for (uint32 i = BasisIndex, n = Paragraphs.Num(); i < n; ++i)
 	{
 		const FParagraph& Holm = Paragraphs[i];
-		Insert(Holm.Phrase);
+		Insert(Holm.Phrase[0]);
 	}
 	
 	// Prime the disk-cache's state
@@ -1025,7 +1059,7 @@ int32 FCache::Flush(int32 Allowance)
 
 	int32 MemCacheSize = MemCache.Peel(Allowance, PeelItems);
 
-	FDiskCache::FPhrase Phrase = DiskCache.OpenPhrase(MemCacheSize);
+	FDiskPhrase Phrase = DiskCache.OpenPhrase(MemCacheSize);
 	int32 PeelIndex = -1;
 	for (int32 i = 0, n = PeelItems.Num(); i < n; ++i)
 	{
@@ -1036,7 +1070,7 @@ int32 FCache::Flush(int32 Allowance)
 			break;
 		}
 	}
-	DiskCache.ClosePhrase(Phrase);
+	DiskCache.ClosePhrase(MoveTemp(Phrase));
 
 	if (PeelIndex >= 0)
 	{
