Iris - Allow sending dynamic objects in SoftObjectPtr.

diff --git a/Engine/Plugins/Runtime/ReplicationSystemTestPlugin/Source/Private/Tests/ReplicationSystem/ReplicatedTestObject.cpp b/Engine/Plugins/Runtime/ReplicationSystemTestPlugin/Source/Private/Tests/ReplicationSystem/ReplicatedTestObject.cpp
--- a/Engine/Plugins/Runtime/ReplicationSystemTestPlugin/Source/Private/Tests/ReplicationSystem/ReplicatedTestObject.cpp
+++ b/Engine/Plugins/Runtime/ReplicationSystemTestPlugin/Source/Private/Tests/ReplicationSystem/ReplicatedTestObject.cpp
@@ -449,6 +449,7 @@ void UTestReplicatedIrisObjectWithObjectReference::GetLifetimeReplicatedProps( T
 	DOREPLIFETIME(ThisClass, IntC);
 	DOREPLIFETIME(ThisClass, RawObjectPtrRef);
 	DOREPLIFETIME(ThisClass, WeakObjectPtrObjectRef);
+	DOREPLIFETIME(ThisClass, SoftObjectPtrRef);
 }
 
 UTestReplicatedIrisObjectWithObjectReference::UTestReplicatedIrisObjectWithObjectReference()
diff --git a/Engine/Plugins/Runtime/ReplicationSystemTestPlugin/Source/Private/Tests/ReplicationSystem/ReplicatedTestObject.h b/Engine/Plugins/Runtime/ReplicationSystemTestPlugin/Source/Private/Tests/ReplicationSystem/ReplicatedTestObject.h
--- a/Engine/Plugins/Runtime/ReplicationSystemTestPlugin/Source/Private/Tests/ReplicationSystem/ReplicatedTestObject.h
+++ b/Engine/Plugins/Runtime/ReplicationSystemTestPlugin/Source/Private/Tests/ReplicationSystem/ReplicatedTestObject.h
@@ -345,6 +345,9 @@ public:
 	UPROPERTY(Transient, Replicated)
 	TWeakObjectPtr<UObject> WeakObjectPtrObjectRef;
 
+	UPROPERTY(Transient, Replicated)
+	TSoftObjectPtr<UObject> SoftObjectPtrRef;
+
 public:
 	// Network data only for test
 	TArray<UE::Net::FReplicationFragment*> ReplicationFragments;
diff --git a/Engine/Plugins/Runtime/ReplicationSystemTestPlugin/Source/Private/Tests/Serialization/TestObjectNetSerializer.cpp b/Engine/Plugins/Runtime/ReplicationSystemTestPlugin/Source/Private/Tests/Serialization/TestObjectNetSerializer.cpp
--- a/Engine/Plugins/Runtime/ReplicationSystemTestPlugin/Source/Private/Tests/Serialization/TestObjectNetSerializer.cpp
+++ b/Engine/Plugins/Runtime/ReplicationSystemTestPlugin/Source/Private/Tests/Serialization/TestObjectNetSerializer.cpp
@@ -44,10 +44,10 @@ UE_NET_TEST_FIXTURE(FReplicationSystemServerClientTestFixture, TestObjectReferen
 	auto ClientObjectC = Client->GetObjectAs<UTestReplicatedIrisObjectWithObjectReference>(ObjectC->NetRefHandle);
 	auto ClientObjectD = Client->GetObjectAs<UTestReplicatedIrisObjectWithObjectReference>(ObjectD->NetRefHandle);
 	
-	UE_NET_ASSERT_TRUE(ClientObjectA != nullptr);
-	UE_NET_ASSERT_TRUE(ClientObjectB != nullptr);
-	UE_NET_ASSERT_TRUE(ClientObjectC != nullptr);
-	UE_NET_ASSERT_TRUE(ClientObjectD != nullptr);
+	UE_NET_ASSERT_NE(ClientObjectA, nullptr);
+	UE_NET_ASSERT_NE(ClientObjectB, nullptr);
+	UE_NET_ASSERT_NE(ClientObjectC, nullptr);
+	UE_NET_ASSERT_NE(ClientObjectD, nullptr);
 
 	// Verify that we managed to resolved the references
 	UE_NET_ASSERT_EQ((UObject*)ClientObjectA->RawObjectPtrRef, (UObject*)ClientObjectB);
@@ -84,10 +84,10 @@ UE_NET_TEST_FIXTURE(FReplicationSystemServerClientTestFixture, TestCircularRefer
 	auto ClientObjectC = Client->GetObjectAs<UTestReplicatedIrisObjectWithObjectReference>(ObjectC->NetRefHandle);
 	auto ClientObjectD = Client->GetObjectAs<UTestReplicatedIrisObjectWithObjectReference>(ObjectD->NetRefHandle);
 	
-	UE_NET_ASSERT_TRUE(ClientObjectA != nullptr);
-	UE_NET_ASSERT_TRUE(ClientObjectB != nullptr);
-	UE_NET_ASSERT_TRUE(ClientObjectC != nullptr);
-	UE_NET_ASSERT_TRUE(ClientObjectD != nullptr);
+	UE_NET_ASSERT_NE(ClientObjectA, nullptr);
+	UE_NET_ASSERT_NE(ClientObjectB, nullptr);
+	UE_NET_ASSERT_NE(ClientObjectC, nullptr);
+	UE_NET_ASSERT_NE(ClientObjectD, nullptr);
 
 	// Verify that we managed to resolved the references
 	UE_NET_ASSERT_EQ((UObject*)ClientObjectA->RawObjectPtrRef, (UObject*)ClientObjectB);
@@ -147,10 +147,10 @@ UE_NET_TEST_FIXTURE(FReplicationSystemServerClientTestFixture, TestCircularRefer
 	auto ClientObjectC = Client->GetObjectAs<UTestReplicatedIrisObjectWithObjectReference>(ObjectC->NetRefHandle);
 	auto ClientObjectD = Client->GetObjectAs<UTestReplicatedIrisObjectWithObjectReference>(ObjectD->NetRefHandle);
 	
-	UE_NET_ASSERT_TRUE(ClientObjectA != nullptr);
-	UE_NET_ASSERT_TRUE(ClientObjectB != nullptr);
-	UE_NET_ASSERT_TRUE(ClientObjectC != nullptr);
-	UE_NET_ASSERT_TRUE(ClientObjectD != nullptr);
+	UE_NET_ASSERT_NE(ClientObjectA, nullptr);
+	UE_NET_ASSERT_NE(ClientObjectB, nullptr);
+	UE_NET_ASSERT_NE(ClientObjectC, nullptr);
+	UE_NET_ASSERT_NE(ClientObjectD, nullptr);
 
 	// Verify that we managed to resolved the references
 	UE_NET_ASSERT_EQ((UObject*)ClientObjectA->RawObjectPtrRef, (UObject*)ClientObjectB);
@@ -159,4 +159,32 @@ UE_NET_TEST_FIXTURE(FReplicationSystemServerClientTestFixture, TestCircularRefer
 	UE_NET_ASSERT_EQ((UObject*)ClientObjectD->RawObjectPtrRef, (UObject*)ClientObjectA);
 }
 
+UE_NET_TEST_FIXTURE(FReplicationSystemServerClientTestFixture, TestSoftObjectReferenceToDynamicObject)
+{
+	// Add a client
+	FReplicationSystemTestClient* Client = CreateClient();
+
+	// Spawn objects on server
+	UTestReplicatedIrisObjectWithObjectReference* ObjectB = Server->CreateObject<UTestReplicatedIrisObjectWithObjectReference>();
+	UTestReplicatedIrisObjectWithObjectReference* ObjectA = Server->CreateObject<UTestReplicatedIrisObjectWithObjectReference>();
+
+	ObjectA->SoftObjectPtrRef = ObjectB;
+
+	// Send and deliver packet
+	Server->PreSendUpdate();
+	Server->SendAndDeliverTo(Client, DeliverPacket);
+	Server->PostSendUpdate();
+
+	// Check state, we cheat a bit and use the server NetRefHandle for lookup
+	// THIS will fail if we start using full handle compares
+	auto ClientObjectA = Client->GetObjectAs<UTestReplicatedIrisObjectWithObjectReference>(ObjectA->NetRefHandle);
+	auto ClientObjectB = Client->GetObjectAs<UTestReplicatedIrisObjectWithObjectReference>(ObjectB->NetRefHandle);
+
+	UE_NET_ASSERT_NE(ClientObjectA, nullptr);
+	UE_NET_ASSERT_NE(ClientObjectB, nullptr);
+
+	// Verify that we managed to resolved the reference
+	UE_NET_ASSERT_EQ((UObject*)ClientObjectA->SoftObjectPtrRef.Get(), (UObject*)ClientObjectB);
+}
+
 }
diff --git a/Engine/Plugins/Runtime/ReplicationSystemTestPlugin/Source/Private/Tests/Serialization/TestSoftObjectNetSerializer.cpp b/Engine/Plugins/Runtime/ReplicationSystemTestPlugin/Source/Private/Tests/Serialization/TestSoftObjectNetSerializer.cpp
--- a/Engine/Plugins/Runtime/ReplicationSystemTestPlugin/Source/Private/Tests/Serialization/TestSoftObjectNetSerializer.cpp
+++ b/Engine/Plugins/Runtime/ReplicationSystemTestPlugin/Source/Private/Tests/Serialization/TestSoftObjectNetSerializer.cpp
@@ -97,17 +97,10 @@ void FTestSoftObjectNetSerializer::TearDown()
 
 void FTestSoftObjectNetSerializer::SetUpTestValues()
 {
-	CreatedObjects.Add(TStrongObjectPtr<UObject>(NewObject<UTestReplicatedIrisObject>()));
-	CreatedObjects.Add(TStrongObjectPtr<UObject>(NewObject<UTestReplicatedIrisObject>()));
-
 	TestValues.Add(FSoftObjectPtr());
 	TestValues.Add(FSoftObjectPtr(FSoftObjectPath(TEXT("/Script/NonExistingPlugin.NonExistingClass:NonExistingInstance"))));
 
-	for (const TStrongObjectPtr<UObject>& Object : CreatedObjects)
-	{
-		TestValues.Add(FSoftObjectPtr(Object.Get()));
-		check(TestValues[TestValues.Num() - 1].Get() != nullptr);
-	}
+	// Non-stably named objects are tested in TestObjectNetSerializer.cpp.
 }
 
 void FTestSoftObjectNetSerializer::TestValidate()
diff --git a/Engine/Source/Runtime/Experimental/Iris/Core/Private/Iris/Serialization/SoftObjectNetSerializers.cpp b/Engine/Source/Runtime/Experimental/Iris/Core/Private/Iris/Serialization/SoftObjectNetSerializers.cpp
--- a/Engine/Source/Runtime/Experimental/Iris/Core/Private/Iris/Serialization/SoftObjectNetSerializers.cpp
+++ b/Engine/Source/Runtime/Experimental/Iris/Core/Private/Iris/Serialization/SoftObjectNetSerializers.cpp
@@ -1,10 +1,30 @@
 // Copyright Epic Games, Inc. All Rights Reserved.
 
 #include "Iris/Serialization/SoftObjectNetSerializers.h"
-#include "Iris/Serialization/StringNetSerializerUtils.h"
+#include "Iris/Core/NetObjectReference.h"
 #include "Iris/Serialization/InternalNetSerializer.h"
+#include "Iris/Serialization/NetReferenceCollector.h"
+#include "Iris/Serialization/ObjectNetSerializer.h"
+#include "Iris/Serialization/StringNetSerializerUtils.h"
+#include "Templates/IsPODType.h"
 #include "UObject/SoftObjectPtr.h"
 
+
+namespace UE::Net
+{
+
+struct FFSoftObjectNetSerializerQuantizedData
+{
+	Private::FStringNetSerializerBase::QuantizedType Path;
+	FNetObjectReference ObjectReference;
+	bool bIsObject;
+	uint8 Padding[7];
+};
+
+}
+
+template<> struct TIsPODType<UE::Net::FFSoftObjectNetSerializerQuantizedData> { enum { Value = true }; };
+
 namespace UE::Net
 {
 
@@ -13,21 +33,41 @@ struct FSoftObjectNetSerializer : public Private::FStringNetSerializerBase
 	// Version
 	static const uint32 Version = 0;
 
+	// Traits
+	static constexpr bool bIsForwardingSerializer = true;
+	static constexpr bool bHasCustomNetReference = true;
+
 	// Types
 	typedef FSoftObjectPtr SourceType;
+	typedef FFSoftObjectNetSerializerQuantizedData QuantizedType;
 	typedef FSoftObjectNetSerializerConfig ConfigType;
 
-	static const ConfigType DefaultConfig;
+	inline static const ConfigType DefaultConfig;
 
 	// Implementation
 	static void Quantize(FNetSerializationContext&, const FNetQuantizeArgs& Args);
 	static void Dequantize(FNetSerializationContext&, const FNetDequantizeArgs& Args);
 
+	static void Serialize(FNetSerializationContext&, const FNetSerializeArgs& Args);
+	static void Deserialize(FNetSerializationContext&, const FNetDeserializeArgs& Args);
+
+	static void SerializeDelta(FNetSerializationContext&, const FNetSerializeDeltaArgs& Args);
+	static void DeserializeDelta(FNetSerializationContext&, const FNetDeserializeDeltaArgs& Args);
+
 	static bool IsEqual(FNetSerializationContext&, const FNetIsEqualArgs& Args);
 	static bool Validate(FNetSerializationContext&, const FNetValidateArgs& Args);
+
+	static void CloneDynamicState(FNetSerializationContext&, const FNetCloneDynamicStateArgs&);
+	static void FreeDynamicState(FNetSerializationContext&, const FNetFreeDynamicStateArgs&);
+	static void CollectNetReferences(FNetSerializationContext&, const FNetCollectReferencesArgs&);
+
+private:
+	typedef Private::FStringNetSerializerBase Super;
+
+	inline static const FNetSerializer* ObjectNetSerializer = &UE_NET_GET_SERIALIZER(FObjectNetSerializer);
+	inline static const FNetSerializerConfig* ObjectNetSerializerConfig = UE_NET_GET_SERIALIZER_DEFAULT_CONFIG(FObjectNetSerializer);
 };
 UE_NET_IMPLEMENT_SERIALIZER_INTERNAL(FSoftObjectNetSerializer);
-const FSoftObjectNetSerializer::ConfigType FSoftObjectNetSerializer::DefaultConfig;
 
 struct FSoftObjectPathNetSerializer : public Private::FStringNetSerializerBase
 {
@@ -38,7 +78,7 @@ struct FSoftObjectPathNetSerializer : public Private::FStringNetSerializerBase
 	typedef FSoftObjectPath SourceType;
 	typedef FSoftObjectPathNetSerializerConfig ConfigType;
 
-	static const ConfigType DefaultConfig;
+	inline static const ConfigType DefaultConfig;
 
 	// Implementation
 	static void Quantize(FNetSerializationContext&, const FNetQuantizeArgs& Args);
@@ -48,7 +88,6 @@ struct FSoftObjectPathNetSerializer : public Private::FStringNetSerializerBase
 	static bool Validate(FNetSerializationContext&, const FNetValidateArgs& Args);
 };
 UE_NET_IMPLEMENT_SERIALIZER_INTERNAL(FSoftObjectPathNetSerializer);
-const FSoftObjectPathNetSerializer::ConfigType FSoftObjectPathNetSerializer::DefaultConfig;
 
 struct FSoftClassPathNetSerializer : public Private::FStringNetSerializerBase
 {
@@ -75,27 +114,167 @@ const FSoftClassPathNetSerializer::ConfigType FSoftClassPathNetSerializer::Defau
 void FSoftObjectNetSerializer::Quantize(FNetSerializationContext& Context, const FNetQuantizeArgs& Args)
 {
 	const SourceType& Source = *reinterpret_cast<const SourceType*>(Args.Source);
-	const FSoftObjectPath& SoftObjectPath = Source.GetUniqueID();
-	const FString& SoftObjectPathString = SoftObjectPath.ToString();
-	return FStringNetSerializerBase::Quantize(Context, Args, SoftObjectPathString);
+	QuantizedType& Target = *reinterpret_cast<QuantizedType*>(Args.Target);
+
+	QuantizedType TempValue = {};
+
+	const UObject* Object = Source.Get();
+	// Use path whenever possible. Only resort to object serialization for non-stably named objects.
+	const bool bIsObject = Object && !Object->IsFullNameStableForNetworking();
+	if (bIsObject && !Target.bIsObject)
+	{
+		FNetFreeDynamicStateArgs FreeArgs;
+		FreeArgs.Source = reinterpret_cast<NetSerializerValuePointer>(&Target.Path);
+		Super::FreeDynamicState(Context, FreeArgs);
+	}
+
+	TempValue.bIsObject = bIsObject;
+	if (bIsObject)
+	{
+		FNetQuantizeArgs QuantizeArgs = Args;
+		QuantizeArgs.NetSerializerConfig = reinterpret_cast<NetSerializerConfigParam>(ObjectNetSerializerConfig);
+		QuantizeArgs.Source = reinterpret_cast<NetSerializerValuePointer>(&Object);
+		QuantizeArgs.Target = reinterpret_cast<NetSerializerValuePointer>(&TempValue.ObjectReference);
+		ObjectNetSerializer->Quantize(Context, QuantizeArgs);
+	}
+	else
+	{
+		const FSoftObjectPath& SoftObjectPath = Source.GetUniqueID();
+		const FString& SoftObjectPathString = SoftObjectPath.ToString();
+
+		FNetQuantizeArgs QuantizeArgs = Args;
+		QuantizeArgs.Source = 0;
+		QuantizeArgs.Target = reinterpret_cast<NetSerializerValuePointer>(&TempValue.Path);
+		FStringNetSerializerBase::Quantize(Context, QuantizeArgs, SoftObjectPathString);
+	}
+
+	Target = TempValue;
 }
 
 void FSoftObjectNetSerializer::Dequantize(FNetSerializationContext& Context, const FNetDequantizeArgs& Args)
 {
-	FString SoftObjectPathString;
-	FStringNetSerializerBase::Dequantize(Context, Args, SoftObjectPathString);
-
+	const QuantizedType& Source = *reinterpret_cast<const QuantizedType*>(Args.Source);
 	SourceType& Target = *reinterpret_cast<SourceType*>(Args.Target);
-	Target.ResetWeakPtr();
-	FSoftObjectPath& SoftObjectPath = Target.GetUniqueID();
-	SoftObjectPath.SetPath(SoftObjectPathString);
+
+	if (Source.bIsObject)
+	{
+		UObject* Object = nullptr;
+
+		FNetDequantizeArgs DequantizeArgs = Args;
+		DequantizeArgs.NetSerializerConfig = ObjectNetSerializerConfig;
+		DequantizeArgs.Source = reinterpret_cast<NetSerializerValuePointer>(&Source.ObjectReference);
+		DequantizeArgs.Target = reinterpret_cast<NetSerializerValuePointer>(&Object);
+		ObjectNetSerializer->Dequantize(Context, DequantizeArgs);
+		
+		// This will release any previous weak object and create the unique ID as well. 
+		Target = Object;
+	}
+	else
+	{
+		FString SoftObjectPathString;
+		Super::Dequantize(Context, Args, SoftObjectPathString);
+
+		Target.ResetWeakPtr();
+		Target.GetUniqueID().SetPath(MoveTemp(SoftObjectPathString));
+	}
+
+	Target.GetUniqueID().FixupForPIE();
+}
+
+void FSoftObjectNetSerializer::Serialize(FNetSerializationContext& Context, const FNetSerializeArgs& Args)
+{
+	const QuantizedType& Source = *reinterpret_cast<const QuantizedType*>(Args.Source);
+
+	FNetBitStreamWriter* Writer = Context.GetBitStreamWriter();
+	if (Writer->WriteBool(Source.bIsObject))
+	{
+		FNetSerializeArgs SerializeArgs = Args;
+		SerializeArgs.NetSerializerConfig = ObjectNetSerializerConfig;
+		SerializeArgs.Source = reinterpret_cast<NetSerializerValuePointer>(&Source.ObjectReference);
+		ObjectNetSerializer->Serialize(Context, SerializeArgs);
+	}
+	else
+	{
+		FNetSerializeArgs SerializeArgs = Args;
+		SerializeArgs.NetSerializerConfig = 0;
+		SerializeArgs.Source = reinterpret_cast<NetSerializerValuePointer>(&Source.Path);
+		Super::Serialize(Context, SerializeArgs);
+	}
+}
+
+void FSoftObjectNetSerializer::Deserialize(FNetSerializationContext& Context, const FNetDeserializeArgs& Args)
+{
+	QuantizedType& Target = *reinterpret_cast<QuantizedType*>(Args.Target);
+
+	FNetBitStreamReader* Reader = Context.GetBitStreamReader();
+	QuantizedType TempValue = {};
+	TempValue.bIsObject = Reader->ReadBool();
+	if (TempValue.bIsObject)
+	{
+		FNetDeserializeArgs DeserializeArgs = Args;
+		DeserializeArgs.NetSerializerConfig = ObjectNetSerializerConfig;
+		DeserializeArgs.Target = reinterpret_cast<NetSerializerValuePointer>(&TempValue.ObjectReference);
+		ObjectNetSerializer->Deserialize(Context, DeserializeArgs);
+
+		// Free previous path allocation
+		if (!Target.bIsObject)
+		{
+			FNetFreeDynamicStateArgs FreeArgs = {};
+			FreeArgs.Source = reinterpret_cast<NetSerializerValuePointer>(&Target.Path);
+			Super::FreeDynamicState(Context, FreeArgs);
+		}
+	}
+	else
+	{
+		// Reuse memory.
+		TempValue.Path = Target.Path;
+
+		FNetDeserializeArgs DeserializeArgs = Args;
+		DeserializeArgs.NetSerializerConfig = 0;
+		DeserializeArgs.Target = reinterpret_cast<NetSerializerValuePointer>(&TempValue.Path);
+		Super::Deserialize(Context, DeserializeArgs);
+	}
+
+	Target = TempValue;
+}
+
+void FSoftObjectNetSerializer::SerializeDelta(FNetSerializationContext& Context, const FNetSerializeDeltaArgs& Args)
+{
+	Serialize(Context, Args);
+}
+
+void FSoftObjectNetSerializer::DeserializeDelta(FNetSerializationContext& Context, const FNetDeserializeDeltaArgs& Args)
+{
+	Deserialize(Context, Args);
 }
 
 bool FSoftObjectNetSerializer::IsEqual(FNetSerializationContext& Context, const FNetIsEqualArgs& Args)
 {
 	if (Args.bStateIsQuantized)
 	{
-		return IsQuantizedEqual(Context, Args);
+		const QuantizedType& Source0 = *reinterpret_cast<const QuantizedType*>(Args.Source0);
+		const QuantizedType& Source1 = *reinterpret_cast<const QuantizedType*>(Args.Source1);
+
+		if (Source0.bIsObject != Source1.bIsObject)
+		{
+			return false;
+		}
+
+		if (Source0.bIsObject)
+		{
+			FNetIsEqualArgs EqualArgs = Args;
+			EqualArgs.NetSerializerConfig = ObjectNetSerializerConfig;
+			EqualArgs.Source0 = reinterpret_cast<NetSerializerValuePointer>(&Source0.Path);
+			EqualArgs.Source1 = reinterpret_cast<NetSerializerValuePointer>(&Source1.Path);
+			return ObjectNetSerializer->IsEqual(Context, EqualArgs);
+		}
+		else
+		{
+			FNetIsEqualArgs EqualArgs = Args;
+			EqualArgs.Source0 = reinterpret_cast<NetSerializerValuePointer>(&Source0.Path);
+			EqualArgs.Source1 = reinterpret_cast<NetSerializerValuePointer>(&Source1.Path);
+			return Super::IsQuantizedEqual(Context, Args);
+		}
 	}
 	else
 	{
@@ -117,6 +296,49 @@ bool FSoftObjectNetSerializer::Validate(FNetSerializationContext& Context, const
 	return FStringNetSerializerBase::Validate(Context, Args, SoftObjectPathString);
 }
 
+void FSoftObjectNetSerializer::CloneDynamicState(FNetSerializationContext& Context, const FNetCloneDynamicStateArgs& Args)
+{
+	const QuantizedType& Source = *reinterpret_cast<const QuantizedType*>(Args.Source);
+	QuantizedType& Target = *reinterpret_cast<QuantizedType*>(Args.Target);
+
+	Target = Source;
+	if (!Source.bIsObject)
+	{
+		FNetCloneDynamicStateArgs CloneArgs = Args;
+		CloneArgs.NetSerializerConfig = 0;
+		CloneArgs.Source = NetSerializerValuePointer(&Source.Path);
+		CloneArgs.Target = NetSerializerValuePointer(&Target.Path);
+		Super::CloneDynamicState(Context, CloneArgs);
+	}
+}
+
+void FSoftObjectNetSerializer::FreeDynamicState(FNetSerializationContext& Context, const FNetFreeDynamicStateArgs& Args)
+{
+	QuantizedType& Source = *reinterpret_cast<QuantizedType*>(Args.Source);
+
+	if (!Source.bIsObject)
+	{
+		FNetFreeDynamicStateArgs FreeArgs = Args;
+		FreeArgs.NetSerializerConfig = 0;
+		FreeArgs.Source = NetSerializerValuePointer(&Source.Path);
+		Super::FreeDynamicState(Context, FreeArgs);
+	}
+}
+
+void FSoftObjectNetSerializer::CollectNetReferences(FNetSerializationContext&, const FNetCollectReferencesArgs& Args)
+{
+	const QuantizedType& Source = *reinterpret_cast<const QuantizedType*>(Args.Source);
+
+	if (Source.bIsObject)
+	{
+		const QuantizedType& Value = *reinterpret_cast<const QuantizedType*>(Args.Source);
+		FNetReferenceCollector& Collector = *reinterpret_cast<FNetReferenceCollector*>(Args.Collector);
+
+		const FNetReferenceInfo ReferenceInfo(FNetReferenceInfo::EResolveType::ResolveOnClient);
+		Collector.Add(ReferenceInfo, Value.ObjectReference, Args.ChangeMaskInfo);
+	}
+}
+
 // FSoftObjectPathNetSerializer
 void FSoftObjectPathNetSerializer::Quantize(FNetSerializationContext& Context, const FNetQuantizeArgs& Args)
 {
@@ -198,3 +420,4 @@ bool FSoftClassPathNetSerializer::Validate(FNetSerializationContext& Context, co
 }
 
 }
+
