Iris - Fixed ensure

diff --git a/Engine/Source/Runtime/Experimental/Iris/Core/Private/Iris/ReplicationSystem/ReplicationBridge.cpp b/Engine/Source/Runtime/Experimental/Iris/Core/Private/Iris/ReplicationSystem/ReplicationBridge.cpp
--- a/Engine/Source/Runtime/Experimental/Iris/Core/Private/Iris/ReplicationSystem/ReplicationBridge.cpp
+++ b/Engine/Source/Runtime/Experimental/Iris/Core/Private/Iris/ReplicationSystem/ReplicationBridge.cpp
@@ -820,7 +820,11 @@ UE::Net::FNetRefHandle UReplicationBridge::InternalAddDestructionInfo(FNetRefHan
 	// Create handle for DestructionInfo to be able to scope destruction infos
 	const FNetObjectGroupHandle LevelGroupHandle = GetLevelGroup(Parameters.Level);
 	FNetRefHandle DestructionInfoHandle = NetRefHandleManager->CreateHandleForDestructionInfo(Handle, DestructionInfoProtocol);
-	GetReplicationSystem()->AddToGroup(LevelGroupHandle, DestructionInfoHandle);
+
+	if (LevelGroupHandle.IsValid())
+	{
+		GetReplicationSystem()->AddToGroup(LevelGroupHandle, DestructionInfoHandle);
+	}
 
 	const FInternalNetRefIndex InternalReplicationIndex = NetRefHandleManager->GetInternalIndex(DestructionInfoHandle);
 
