Iris - Fix NetDeserializeDeltaDefault such that it properly frees previous memory in target memory.

#rnx

diff --git a/Engine/Source/Runtime/Experimental/Iris/Core/Public/Iris/Serialization/NetSerializerBuilder.inl b/Engine/Source/Runtime/Experimental/Iris/Core/Public/Iris/Serialization/NetSerializerBuilder.inl
--- a/Engine/Source/Runtime/Experimental/Iris/Core/Public/Iris/Serialization/NetSerializerBuilder.inl
+++ b/Engine/Source/Runtime/Experimental/Iris/Core/Public/Iris/Serialization/NetSerializerBuilder.inl
@@ -44,13 +44,23 @@ NetSerializeDeltaDefault(FNetSerializationContext& Context, const FNetSerializeD
 	Serialize(Context, Args);
 };
 
-template<uint32 QuantizedTypeSize, NetDeserializeFunction Deserialize, NetCloneDynamicStateFunction CloneDynamicState>
+template<uint32 QuantizedTypeSize, NetDeserializeFunction Deserialize, NetFreeDynamicStateFunction FreeDynamicState, NetCloneDynamicStateFunction CloneDynamicState>
 void
 NetDeserializeDeltaDefault(FNetSerializationContext& Context, const FNetDeserializeDeltaArgs& Args)
 {
 	if (Context.GetBitStreamReader()->ReadBool())
 	{
-		// Clone from Prev
+		// Clone from prev. Need to free target first.
+		if (FreeDynamicState != NetFreeDynamicStateFunction(nullptr))
+		{
+			FNetFreeDynamicStateArgs FreeArgs;
+			FreeArgs.Version = 0;
+			FreeArgs.NetSerializerConfig = Args.NetSerializerConfig;
+			FreeArgs.Source = Args.Target;
+
+			FreeDynamicState(Context, FreeArgs);
+		}
+
 		FMemory::Memcpy(reinterpret_cast<uint8*>(Args.Target), reinterpret_cast<uint8*>(Args.Prev), QuantizedTypeSize);
 
 		if (CloneDynamicState != NetCloneDynamicStateFunction(nullptr))
@@ -343,11 +353,11 @@ public:
 	template<typename T = void, typename U = typename TEnableIf<!HasDeserializeDelta && !ShouldUseDefaultDelta(), T>::Type, int V = 0>
 	static NetDeserializeDeltaFunction GetDeserializeDeltaFunction(const void* = nullptr) { return NetDeserializeDeltaDefault<NetSerializerImpl::Deserialize>; }
 
-	template<typename T = void, typename U = typename TEnableIf<!HasDeserializeDelta && ShouldUseDefaultDelta() && HasCloneDynamicState, T>::Type, char V = 0>
-	static NetDeserializeDeltaFunction GetDeserializeDeltaFunction(const void* = nullptr) { return NetDeserializeDeltaDefault<GetQuantizedTypeSize(), NetSerializerImpl::Deserialize, NetSerializerImpl::CloneDynamicState>; }
+	template<typename T = void, typename U = typename TEnableIf<!HasDeserializeDelta && ShouldUseDefaultDelta() && (HasCloneDynamicState && HasFreeDynamicState), T>::Type, char V = 0>
+	static NetDeserializeDeltaFunction GetDeserializeDeltaFunction(const void* = nullptr) { return NetDeserializeDeltaDefault<GetQuantizedTypeSize(), NetSerializerImpl::Deserialize, NetSerializerImpl::FreeDynamicState, NetSerializerImpl::CloneDynamicState>; }
 
-	template<typename T = void, typename U = typename TEnableIf<!HasDeserializeDelta && ShouldUseDefaultDelta() && !HasCloneDynamicState, T>::Type, unsigned char V = 0>
-	static NetDeserializeDeltaFunction GetDeserializeDeltaFunction(const void* = nullptr) { return NetDeserializeDeltaDefault<GetQuantizedTypeSize(), NetSerializerImpl::Deserialize, NetCloneDynamicStateFunction(nullptr)>; }
+	template<typename T = void, typename U = typename TEnableIf<!HasDeserializeDelta && ShouldUseDefaultDelta() && !(HasCloneDynamicState && HasFreeDynamicState), T>::Type, unsigned char V = 0>
+	static NetDeserializeDeltaFunction GetDeserializeDeltaFunction(const void* = nullptr) { return NetDeserializeDeltaDefault<GetQuantizedTypeSize(), NetSerializerImpl::Deserialize, NetFreeDynamicStateFunction(nullptr), NetCloneDynamicStateFunction(nullptr)>; }
 
 	// Provide a default Quantize implementation if needed. The default will copy the value.
 	template<typename T = void, typename U = typename TEnableIf<HasQuantize, T>::Type, bool V = true>
