Iris - SoftObjectNetSerializer same value delta optimization.

#rnx

diff --git a/Engine/Source/Runtime/Experimental/Iris/Core/Private/Iris/Serialization/SoftObjectNetSerializers.cpp b/Engine/Source/Runtime/Experimental/Iris/Core/Private/Iris/Serialization/SoftObjectNetSerializers.cpp
--- a/Engine/Source/Runtime/Experimental/Iris/Core/Private/Iris/Serialization/SoftObjectNetSerializers.cpp
+++ b/Engine/Source/Runtime/Experimental/Iris/Core/Private/Iris/Serialization/SoftObjectNetSerializers.cpp
@@ -240,12 +240,12 @@ void FSoftObjectNetSerializer::Deserialize(FNetSerializationContext& Context, co
 
 void FSoftObjectNetSerializer::SerializeDelta(FNetSerializationContext& Context, const FNetSerializeDeltaArgs& Args)
 {
-	Serialize(Context, Args);
+	NetSerializeDeltaDefault<FSoftObjectNetSerializer::Serialize, FSoftObjectNetSerializer::IsEqual>(Context, Args);
 }
 
 void FSoftObjectNetSerializer::DeserializeDelta(FNetSerializationContext& Context, const FNetDeserializeDeltaArgs& Args)
 {
-	Deserialize(Context, Args);
+	NetDeserializeDeltaDefault<sizeof(QuantizedType), FSoftObjectNetSerializer::Deserialize, FSoftObjectNetSerializer::FreeDynamicState, FSoftObjectNetSerializer::CloneDynamicState>(Context, Args);
 }
 
 bool FSoftObjectNetSerializer::IsEqual(FNetSerializationContext& Context, const FNetIsEqualArgs& Args)
