Move code for loading IAS toc to a common function.

#rnx

- We now only log an error when failing to load the toc once all retries have failed. We don't want to mark CIS processes as failed just because the first attempt to download a toc failed if we managed to get it on the second attempt.
- Sorted the includes alphabetically to avoid IWYU doing so later.

diff --git a/Engine/Source/Runtime/Experimental/IoStoreOnDemand/Private/IoStoreOnDemand.cpp b/Engine/Source/Runtime/Experimental/IoStoreOnDemand/Private/IoStoreOnDemand.cpp
--- a/Engine/Source/Runtime/Experimental/IoStoreOnDemand/Private/IoStoreOnDemand.cpp
+++ b/Engine/Source/Runtime/Experimental/IoStoreOnDemand/Private/IoStoreOnDemand.cpp
@@ -3,15 +3,12 @@
 #include "IO/IoStoreOnDemand.h"
 
 #include "EncryptionKeyManager.h"
-#include "IasCache.h"
-#include "LatencyInjector.h"
-#include "OnDemandIoDispatcherBackend.h"
-#include "Statistics.h"
-
+#include "HAL/FileManager.h"
 #include "HAL/LowLevelMemTracker.h"
 #include "HAL/PlatformMisc.h"
 #include "HAL/PlatformTime.h"
-#include "HAL/FileManager.h"
+#include "IasCache.h"
+#include "LatencyInjector.h"
 #include "Misc/Base64.h"
 #include "Misc/CommandLine.h"
 #include "Misc/ConfigCacheIni.h"
@@ -20,11 +17,14 @@
 #include "Misc/FileHelper.h"
 #include "Misc/Parse.h"
 #include "Modules/ModuleManager.h"
+#include "OnDemandHttpClient.h"
+#include "OnDemandIoDispatcherBackend.h"
 #include "Serialization/Archive.h"
-#include "Serialization/CompactBinaryWriter.h"
 #include "Serialization/CompactBinarySerialization.h"
+#include "Serialization/CompactBinaryWriter.h"
 #include "Serialization/LargeMemoryWriter.h"
 #include "Serialization/MemoryReader.h"
+#include "Statistics.h"
 #include "String/LexFromString.h"
 
 #if (PLATFORM_DESKTOP && (IS_PROGRAM || WITH_EDITOR))
@@ -534,6 +534,59 @@ bool LoadFromCompactBinary(FCbFieldView Field, FOnDemandToc& OutToc)
 	return false;
 }
 
+TIoStatusOr<FOnDemandToc> LoadTocFromUrl(const FString& ServiceURL, const FString& TocPath, int32 RetryCount)
+{
+	TRACE_CPUPROFILER_EVENT_SCOPE(LoadTocFromUrl);
+
+	FString ErrorMsg;
+
+	for (int32 Attempt = 0; Attempt <= RetryCount; ++Attempt)
+	{
+		TUniquePtr<FOnDemandHttpClient> HttpClient = MakeUnique<FOnDemandHttpClient>(ServiceURL, 1);
+		TAnsiStringBuilder<256> Url;
+
+		Url << "/" << TocPath;
+		UE_LOG(LogIas, Log, TEXT("Fetching TOC '%s/%s' (#%d/%d)"), *HttpClient->ServiceUrl(), *TocPath, Attempt + 1, RetryCount);
+
+		TIoStatusOr<FOnDemandToc> Toc;
+		HttpClient->Get(Url.ToView(), [&Toc, &ErrorMsg](TIoStatusOr<FIoBuffer> Response, uint64 DurationMs)
+			{
+				if (Response.IsOk())
+				{
+					TRACE_CPUPROFILER_EVENT_SCOPE(LoadTocFromEndpoint);
+					FIoBuffer Buffer = Response.ConsumeValueOrDie();
+					FOnDemandToc NewToc;
+
+					FMemoryReaderView Ar(Buffer.GetView());
+					Ar << NewToc;
+					if (!Ar.IsError())
+					{
+						Toc = TIoStatusOr<FOnDemandToc>(MoveTemp(NewToc));
+					}
+					else
+					{
+						ErrorMsg = TEXT("Failed loading on demand TOC from compact binary");
+					}
+				}
+				else
+				{
+					ErrorMsg = FString::Printf(TEXT("Failed fetching TOC, reason '%s'"), *Response.Status().ToString());
+				}
+			});
+
+		while (HttpClient->Tick());
+
+		if (Toc.IsOk())
+		{
+			return Toc;
+		}
+	}
+
+	UE_LOG(LogIas, Error, TEXT("%s"), *ErrorMsg);
+
+	return TIoStatusOr<FOnDemandToc>(FIoStatus(EIoErrorCode::NotFound));
+}
+
 ////////////////////////////////////////////////////////////////////////////////
 #if (PLATFORM_DESKTOP && (IS_PROGRAM || WITH_EDITOR))
 class FS3UploadQueue
diff --git a/Engine/Source/Runtime/Experimental/IoStoreOnDemand/Private/OnDemandIoDispatcherBackend.cpp b/Engine/Source/Runtime/Experimental/IoStoreOnDemand/Private/OnDemandIoDispatcherBackend.cpp
--- a/Engine/Source/Runtime/Experimental/IoStoreOnDemand/Private/OnDemandIoDispatcherBackend.cpp
+++ b/Engine/Source/Runtime/Experimental/IoStoreOnDemand/Private/OnDemandIoDispatcherBackend.cpp
@@ -1306,59 +1306,9 @@ TArray<FIoChunkId> FOnDemandIoBackend::GetAllChunkIds()
 
 #endif // IS_PROGRAM || WITH_EDITOR
 
-TIoStatusOr<FOnDemandToc> FOnDemandIoBackend::GetToc(const FOnDemandEndpoint& Endpoint)
-{
-	TRACE_CPUPROFILER_EVENT_SCOPE(FOnDemandIoBackend::GetToc);
-
-	for (int32 Attempt = 0, MaxAttempts = GIasMaxHttpRetryCount; Attempt <= MaxAttempts; ++Attempt)
-	{
-		TUniquePtr<FOnDemandHttpClient> HttpClient = MakeUnique<FOnDemandHttpClient>(Endpoint.ServiceUrl, GIasMaxHttpConnectionCount);
-		TAnsiStringBuilder<256> Url;
-
-		Url << "/" << Endpoint.TocPath;
-		UE_LOG(LogIas, Log, TEXT("Fetching TOC '%s/%s' (#%d/%d)"), *HttpClient->ServiceUrl(), *Endpoint.TocPath, Attempt + 1, MaxAttempts);
-		
-		TIoStatusOr<FOnDemandToc> Toc;
-		HttpClient->Get(Url.ToView(), [&Toc](TIoStatusOr<FIoBuffer> Response, uint64 DurationMs)
-		{
-			if (Response.IsOk())
-			{
-				TRACE_CPUPROFILER_EVENT_SCOPE(FOnDemandIoBackend::SerializeToc);
-				FIoBuffer Buffer = Response.ConsumeValueOrDie();
-				FOnDemandToc NewToc;	
-
-				FMemoryReaderView Ar(Buffer.GetView());
-				Ar << NewToc;
-				if (!Ar.IsError())
-				{
-					Toc = TIoStatusOr<FOnDemandToc>(MoveTemp(NewToc));
-				}
-				else
-				{
-					UE_LOG(LogIas, Error, TEXT("Failed loading on demand TOC from compact binary"));
-				}
-			}
-			else
-			{
-				UE_LOG(LogIas, Error, TEXT("Failed fetching TOC, reason '%s'"), *Response.Status().ToString());
-			}
-		});
-
-		const bool bBlock = true;
-		while (HttpClient->Tick(bBlock));
-
-		if (Toc.IsOk())
-		{
-			return Toc;
-		}
-	}
-
-	return TIoStatusOr<FOnDemandToc>(FIoStatus(EIoErrorCode::NotFound));
-}
-
 FIoStatus FOnDemandIoBackend::AddToc(const FOnDemandEndpoint& Endpoint)
 {
-	TIoStatusOr<FOnDemandToc> Toc = GetToc(Endpoint);
+	TIoStatusOr<FOnDemandToc> Toc = LoadTocFromUrl(Endpoint.ServiceUrl, Endpoint.TocPath, GIasMaxHttpRetryCount);
 	if (!Toc.IsOk())
 	{
 		return FIoStatus(Toc.Status());
diff --git a/Engine/Source/Runtime/Experimental/IoStoreOnDemand/Public/IO/IoStoreOnDemand.h b/Engine/Source/Runtime/Experimental/IoStoreOnDemand/Public/IO/IoStoreOnDemand.h
--- a/Engine/Source/Runtime/Experimental/IoStoreOnDemand/Public/IO/IoStoreOnDemand.h
+++ b/Engine/Source/Runtime/Experimental/IoStoreOnDemand/Public/IO/IoStoreOnDemand.h
@@ -127,6 +127,8 @@ struct FOnDemandToc
 
 UE_API bool LoadFromCompactBinary(FCbFieldView Field, FOnDemandToc& OutToc);
 
+TIoStatusOr<FOnDemandToc> LoadTocFromUrl(const FString& ServiceURL, const FString& TocPath, int32 RetryCount);
+
 #if (IS_PROGRAM || WITH_EDITOR)
 
 ////////////////////////////////////////////////////////////////////////////////
@@ -182,7 +184,6 @@ struct FIoStoreDownloadParams
 UE_API FIoStatus DownloadContainerFiles(const FIoStoreDownloadParams& DownloadParams, const FString& TocPath);
 
 UE_API FIoStatus PrimeEndPoint(FStringView IoStoreOnDemandIniPath);
-
 #endif // (IS_PROGRAM || WITH_EDITOR)
 
 class IOnDemandIoDispatcherBackend;
