Geometry collection : Add "use material damage modifier" option to GC component

diff --git a/Engine/Source/Runtime/Experimental/GeometryCollectionEngine/Private/GeometryCollection/GeometryCollectionComponent.cpp b/Engine/Source/Runtime/Experimental/GeometryCollectionEngine/Private/GeometryCollection/GeometryCollectionComponent.cpp
--- a/Engine/Source/Runtime/Experimental/GeometryCollectionEngine/Private/GeometryCollection/GeometryCollectionComponent.cpp
+++ b/Engine/Source/Runtime/Experimental/GeometryCollectionEngine/Private/GeometryCollection/GeometryCollectionComponent.cpp
@@ -331,6 +331,7 @@ UGeometryCollectionComponent::UGeometryCollectionComponent(const FObjectInitiali
 	, MaxSimulatedLevel(100)
 	, DamageThreshold({ 500000.f, 50000.f, 5000.f })
 	, bUseSizeSpecificDamageThreshold(false)
+	, bUseMaterialDamageModifiers(false)
 	, bEnableDamageFromCollision(true)
 	, bAllowRemovalOnSleep(true)
 	, bAllowRemovalOnBreak(true)
@@ -3068,7 +3069,6 @@ void UGeometryCollectionComponent::RegisterAndInitializePhysicsProxy()
 #endif
 		EClusterConnectionTypeEnum ClusterCollectionType = ClusterConnectionType_DEPRECATED;
 		float ConnectionGraphBoundsFilteringMargin = 0;
-		bool bUseMaterialDamageModifiers = false;
 		if (RestCollection)
 		{
 			RestCollection->GetSharedSimulationParams(SimulationParameters.Shared);
@@ -3076,7 +3076,6 @@ void UGeometryCollectionComponent::RegisterAndInitializePhysicsProxy()
 			SimulationParameters.InitialRootIndex = RestCollection->GetRootIndex();
 			ClusterCollectionType = RestCollection->ClusterConnectionType;
 			ConnectionGraphBoundsFilteringMargin = RestCollection->ConnectionGraphBoundsFilteringMargin;
-			bUseMaterialDamageModifiers = RestCollection->bUseMaterialDamageModifiers;
 		}
 		SimulationParameters.Simulating = BodyInstance.bSimulatePhysics;
 		SimulationParameters.EnableClustering = EnableClustering;
@@ -3689,6 +3688,7 @@ void UGeometryCollectionComponent::ApplyAssetDefaults()
 		DamageModel = RestCollection->DamageModel;
 		DamageThreshold = RestCollection->DamageThreshold;
 		bUseSizeSpecificDamageThreshold = RestCollection->bUseSizeSpecificDamageThreshold;
+		bUseMaterialDamageModifiers = RestCollection->bUseMaterialDamageModifiers;
 
 		// initialize the component damage progataion data from the asset defaults 
 		DamagePropagationData = RestCollection->DamagePropagationData;
diff --git a/Engine/Source/Runtime/Experimental/GeometryCollectionEngine/Public/GeometryCollection/GeometryCollectionComponent.h b/Engine/Source/Runtime/Experimental/GeometryCollectionEngine/Public/GeometryCollection/GeometryCollectionComponent.h
--- a/Engine/Source/Runtime/Experimental/GeometryCollectionEngine/Public/GeometryCollection/GeometryCollectionComponent.h
+++ b/Engine/Source/Runtime/Experimental/GeometryCollectionEngine/Public/GeometryCollection/GeometryCollectionComponent.h
@@ -826,6 +826,10 @@ public:
 	UPROPERTY(EditAnywhere, BlueprintReadWrite, Category = "ChaosPhysics|Damage", meta = (EditCondition = "DamageModel == EDamageModelTypeEnum::Chaos_Damage_Model_UserDefined_Damage_Threshold"))
 	bool bUseSizeSpecificDamageThreshold;
 
+	/** When on , use the modifiers on the material to adjust the user defined damage threshold values */
+	UPROPERTY(EditAnywhere, BlueprintReadWrite, Category = "ChaosPhysics|Damage", meta = (EditCondition = "DamageModel == EDamageModelTypeEnum::Chaos_Damage_Model_UserDefined_Damage_Threshold"))
+	bool bUseMaterialDamageModifiers;
+
 	/** Data about how damage propagation shoudl behave. */
 	UPROPERTY(EditAnywhere, BlueprintReadWrite, Category = "ChaosPhysics|Damage")
 	FGeometryCollectionDamagePropagationData DamagePropagationData;
