Geometry collection : make sure SetPerParticleCollisionProfileName do not call LoadCollisionProfiles if the array has not changed

diff --git a/Engine/Source/Runtime/Experimental/GeometryCollectionEngine/Private/GeometryCollection/GeometryCollectionComponent.cpp b/Engine/Source/Runtime/Experimental/GeometryCollectionEngine/Private/GeometryCollection/GeometryCollectionComponent.cpp
--- a/Engine/Source/Runtime/Experimental/GeometryCollectionEngine/Private/GeometryCollection/GeometryCollectionComponent.cpp
+++ b/Engine/Source/Runtime/Experimental/GeometryCollectionEngine/Private/GeometryCollection/GeometryCollectionComponent.cpp
@@ -3399,9 +3399,12 @@ void UGeometryCollectionComponent::SetPerParticleCollisionProfileName(const TArr
 		return;
 	}
 
+	bool bHasChanged = false;
+
 	const int32 NumTransforms = RestCollection->GetGeometryCollection()->Transform.Num();
 	if (CollisionProfilePerParticle.Num() != NumTransforms)
 	{
+		bHasChanged = true;
 		CollisionProfilePerParticle.SetNumZeroed(NumTransforms);
 	}
 
@@ -3409,11 +3412,18 @@ void UGeometryCollectionComponent::SetPerParticleCollisionProfileName(const TArr
 	{
 		if (CollisionProfilePerParticle.IsValidIndex(Id))
 		{
-			CollisionProfilePerParticle[Id] = ProfileName;
+			if (CollisionProfilePerParticle[Id] != ProfileName)
+			{
+				bHasChanged = true;
+				CollisionProfilePerParticle[Id] = ProfileName;
+			}
 		}
 	}
 
-	LoadCollisionProfiles();
+	if (bHasChanged)
+	{
+		LoadCollisionProfiles();
+	}
 }
 
 void UGeometryCollectionComponent::LoadCollisionProfiles()
