EngineTest
* Made EngineTest replication test work with Iris
* Suobjects now properly create their protocols
* Changed StableSubObject creation so it's not pointing to the cdo's subobject anymore

#rnx

diff --git a/Engine/Source/Runtime/Experimental/Iris/Core/Private/Iris/ReplicationSystem/ObjectReplicationBridge.cpp b/Engine/Source/Runtime/Experimental/Iris/Core/Private/Iris/ReplicationSystem/ObjectReplicationBridge.cpp
--- a/Engine/Source/Runtime/Experimental/Iris/Core/Private/Iris/ReplicationSystem/ObjectReplicationBridge.cpp
+++ b/Engine/Source/Runtime/Experimental/Iris/Core/Private/Iris/ReplicationSystem/ObjectReplicationBridge.cpp
@@ -388,10 +388,11 @@ UE::Net::FNetRefHandle UObjectReplicationBridge::BeginReplication(FNetRefHandle
 	const FReplicationSystemInternal* ReplicationSystemInternal = GetReplicationSystem()->GetReplicationSystemInternal();
 	const FNetRefHandleManager& LocalNetRefHandleManager = ReplicationSystemInternal->GetNetRefHandleManager();
 
-	// Owner must be replicated
-	check(IsReplicatedHandle(OwnerRefHandle));
-	// Verify assumptions
-	check(!Instance->HasAnyFlags(RF_ArchetypeObject | RF_ClassDefaultObject));
+	checkf(IsReplicatedHandle(OwnerRefHandle), TEXT("Owner %s (%s) must be replicated for subobject %s to replicate."), 
+		*GetNameSafe(LocalNetRefHandleManager.GetReplicatedObjectInstance(LocalNetRefHandleManager.GetInternalIndex(OwnerRefHandle))), *OwnerRefHandle.ToString(), *GetNameSafe(Instance));
+
+	checkf(!Instance->HasAnyFlags(RF_ArchetypeObject | RF_ClassDefaultObject), TEXT("Iris cannot replicate subobject %s owned by %s because it's an %s"), 
+		*GetNameSafe(Instance), *GetNameSafe(LocalNetRefHandleManager.GetReplicatedObjectInstance(LocalNetRefHandleManager.GetInternalIndex(OwnerRefHandle))), Instance->HasAnyFlags(RF_ArchetypeObject)?TEXT("Archetype"):TEXT("DefaultObject"));
 
 	FNetRefHandle SubObjectRefHandle = GetReplicatedRefHandle(Instance);
 	if (SubObjectRefHandle.IsValid())
