Move ::LatencyTest to it's own code file for reuse

#rnx

diff --git a/Engine/Source/Runtime/Experimental/IoStoreOnDemand/Private/CoreHttp/LatencyTesting.cpp b/Engine/Source/Runtime/Experimental/IoStoreOnDemand/Private/CoreHttp/LatencyTesting.cpp
new file mode 100644
--- /dev/null
+++ b/Engine/Source/Runtime/Experimental/IoStoreOnDemand/Private/CoreHttp/LatencyTesting.cpp
@@ -0,0 +1,65 @@
+// Copyright Epic Games, Inc. All Rights Reserved.
+
+#include "LatencyTesting.h"
+
+#include "Client.h"
+#include "Containers/StringConv.h"
+#include "Containers/StringView.h"
+#include "HAL/PlatformTime.h"
+#include "Misc/StringBuilder.h"
+#include "Templates/Function.h"
+
+#if !UE_BUILD_SHIPPING
+
+namespace UE::IO::IAS::HTTP
+{
+
+void LatencyTest(FStringView InUrl, FStringView InPath, TArrayView<int32> OutResults)
+{
+	auto AnsiUrl = StringCast<ANSICHAR>(InUrl.GetData(), InUrl.Len());
+
+	FConnectionPool::FParams PoolParams;
+	PoolParams.SetHostFromUrl(AnsiUrl);
+	PoolParams.ConnectionCount = 1;
+	FConnectionPool Pool(PoolParams);
+
+	TAnsiStringBuilder<256> AnsiPath;
+	AnsiPath << "/";
+	AnsiPath << InPath;
+
+	FEventLoop Loop;
+
+	for (int32& Result : OutResults)
+	{
+		bool Ok = false;
+
+		FRequest Request = Loop.Request("HEAD", AnsiPath, Pool);
+		Loop.Send(MoveTemp(Request), [&](const FTicketStatus& Status)
+			{
+				if (Status.GetId() != FTicketStatus::EId::Response)
+					return;
+
+				const FResponse& Response = Status.GetResponse();
+				Ok = (Response.GetStatus() == EStatusCodeClass::Successful);
+			});
+
+		uint64 Cycles = FPlatformTime::Cycles64();
+		while (Loop.Tick(-1) != 0);
+		Cycles = FPlatformTime::Cycles64() - Cycles;
+
+		Result = Ok ? int32(Cycles) : -1;
+	}
+
+	int64 Freq = int64(1.0 / FPlatformTime::GetSecondsPerCycle());
+	for (int32& Result : OutResults)
+	{
+		if (Result == -1)
+			continue;
+
+		Result = int32((int64(Result) * 1000) / Freq);
+	}
+}
+
+} // namespace UE::IO::IAS::HTTP
+
+#endif // !UE_BUILD_SHIPPING
\ No newline at end of file
diff --git a/Engine/Source/Runtime/Experimental/IoStoreOnDemand/Private/CoreHttp/LatencyTesting.h b/Engine/Source/Runtime/Experimental/IoStoreOnDemand/Private/CoreHttp/LatencyTesting.h
new file mode 100644
--- /dev/null
+++ b/Engine/Source/Runtime/Experimental/IoStoreOnDemand/Private/CoreHttp/LatencyTesting.h
@@ -0,0 +1,17 @@
+// Copyright Epic Games, Inc. All Rights Reserved.
+
+#pragma once
+
+#include "Containers/ArrayView.h"
+#include "Containers/StringFwd.h"
+
+#if !UE_BUILD_SHIPPING
+
+namespace UE::IO::IAS::HTTP
+{
+
+void LatencyTest(FStringView InUrl, FStringView InPath, TArrayView<int32> OutResults);
+
+} // namespace UE::IO::IAS::HTTP
+
+#endif // !UE_BUILD_SHIPPING
diff --git a/Engine/Source/Runtime/Experimental/IoStoreOnDemand/Private/OnDemandIoDispatcherBackend.cpp b/Engine/Source/Runtime/Experimental/IoStoreOnDemand/Private/OnDemandIoDispatcherBackend.cpp
--- a/Engine/Source/Runtime/Experimental/IoStoreOnDemand/Private/OnDemandIoDispatcherBackend.cpp
+++ b/Engine/Source/Runtime/Experimental/IoStoreOnDemand/Private/OnDemandIoDispatcherBackend.cpp
@@ -5,6 +5,7 @@
 #include "AnalyticsEventAttribute.h"
 #include "CancellationToken.h"
 #include "Containers/StringView.h"
+#include "CoreHttp/LatencyTesting.h"
 #include "DistributionEndpoints.h"
 #include "EncryptionKeyManager.h"
 #include "GenericPlatform/GenericPlatformCrashContext.h"
@@ -77,50 +78,8 @@ static TAutoConsoleVariable<bool> CVar_IoReportAnalytics(
 #if !UE_BUILD_SHIPPING
 static void LatencyTest(FStringView InUrl, FStringView InPath)
 {
-	auto AnsiUrl = StringCast<ANSICHAR>(InUrl.GetData(), InUrl.Len());
-
-	using namespace UE::IO::IAS::HTTP;
-
-	FConnectionPool::FParams PoolParams;
-	PoolParams.SetHostFromUrl(AnsiUrl);
-	PoolParams.ConnectionCount = 1;
-	FConnectionPool Pool(PoolParams);
-
-	TAnsiStringBuilder<256> AnsiPath;
-	AnsiPath << "/";
-	AnsiPath << InPath;
-
-	FEventLoop Loop;
 	int32 Results[4] = {};
-	for (uint32 i = 0; i < UE_ARRAY_COUNT(Results); ++i)
-	{
-		bool Ok = false;
-
-		FRequest Request = Loop.Request("HEAD", AnsiPath, Pool);
-		Loop.Send(MoveTemp(Request), [&] (const FTicketStatus& Status)
-		{
-			if (Status.GetId() != FTicketStatus::EId::Response)
-				return;
-
-			const FResponse& Response = Status.GetResponse();
-			Ok = (Response.GetStatus() == EStatusCodeClass::Successful);
-		});
-
-		uint64 Cycles = FPlatformTime::Cycles64();
-		while (Loop.Tick(-1) != 0);
-		Cycles = FPlatformTime::Cycles64() - Cycles;
-
-		Results[i] = Ok ? int32(Cycles) : -1;
-	}
-
-	int64 Freq = int64(1.0 / FPlatformTime::GetSecondsPerCycle());
-	for (int32& Result : Results)
-	{
-		if (Result == -1)
-			continue;
-
-		Result = int32((int64(Result) * 1000) / Freq);
-	}
+	UE::IO::IAS::HTTP::LatencyTest(InUrl, InPath, MakeArrayView(Results));
 
 	UE_LOG(LogIas, VeryVerbose, TEXT("HEAD latencies (ms); %d %d %d %d (%s)"),
 		Results[0], Results[1], Results[2], Results[3], InUrl.GetData());
