Warn on conversion of a large mesh (>1M tris) to geometry collection, and show a progress bar

diff --git a/Engine/Plugins/Experimental/ChaosEditor/Source/FractureEditor/Private/FractureToolGenerators.cpp b/Engine/Plugins/Experimental/ChaosEditor/Source/FractureEditor/Private/FractureToolGenerators.cpp
--- a/Engine/Plugins/Experimental/ChaosEditor/Source/FractureEditor/Private/FractureToolGenerators.cpp
+++ b/Engine/Plugins/Experimental/ChaosEditor/Source/FractureEditor/Private/FractureToolGenerators.cpp
@@ -26,6 +26,7 @@
 #include "GeometryCollection/GeometryCollectionUtility.h"
 #include "GeometryCollection/GeometryCollectionProximityUtility.h"
 #include "FractureToolContext.h"
+#include "FractureToolBackgroundTask.h"
 
 #include "Widgets/DeclarativeSyntaxSupport.h"
 #include "Input/Reply.h"
@@ -606,9 +607,13 @@ void UFractureToolGenerateAsset::OnGenerateAssetPathChosen(const FString& InAsse
 	{
 		AActor* FirstActor = Actors[0];
 
-		AGeometryCollectionActor* GeometryCollectionActor = Cast<AGeometryCollectionActor>(FirstActor);
+		AGeometryCollectionActor* GeometryCollectionActor = nullptr;
 		
 		GeometryCollectionActor = ConvertActorsToGeometryCollection(InAssetPath, false/*bAddInternalMaterials*/, bSplitComponents, Actors, bFromToMeshTool);
+		if (!GeometryCollectionActor)
+		{
+			return;
+		}
 
 		GeometryCollectionComponent = GeometryCollectionActor->GetGeometryCollectionComponent();
 
@@ -662,6 +667,57 @@ AGeometryCollectionActor* UFractureToolGenerateAsset::ConvertActorsToGeometryCol
 	const FString& Name = FirstActor->GetActorLabel();
 	const FVector FirstActorLocation(FirstActor->GetActorLocation());
 
+	using FGeometryCollectionSharedPtr = TSharedPtr<FGeometryCollection, ESPMode::ThreadSafe>;
+
+	// Count the total triangles of all input sources
+	int32 SourceTriCount = 0;
+	{
+		for (AActor* Actor : Actors)
+		{
+			TArray<UStaticMeshComponent*> StaticMeshComponents;
+			Actor->GetComponents(StaticMeshComponents, true);
+			for (int32 ii = 0, ni = StaticMeshComponents.Num(); ii < ni; ++ii)
+			{
+				if (UStaticMeshComponent* StaticMeshComponent = StaticMeshComponents[ii])
+				{
+					if (UStaticMesh* StaticMesh = StaticMeshComponent->GetStaticMesh())
+					{
+						SourceTriCount += StaticMesh->GetNumTriangles(0);
+					}
+				}
+			}
+
+			TArray<UGeometryCollectionComponent*> GeometryCollectionComponents;
+			Actor->GetComponents(GeometryCollectionComponents, true);
+			for (int32 ii = 0, ni = GeometryCollectionComponents.Num(); ii < ni; ++ii)
+			{
+				if (UGeometryCollectionComponent* GeometryCollectionComponent = GeometryCollectionComponents[ii])
+				{
+					if (const UGeometryCollection* RestCollection = GeometryCollectionComponent->GetRestCollection())
+					{
+						if (const FGeometryCollectionSharedPtr Collection = RestCollection->GetGeometryCollection())
+						{
+							SourceTriCount += Collection->Indices.Num();
+						}
+					}
+				}
+			}
+		}
+	}
+	constexpr int32 ThresholdToCautionAboutLargeInput = 1000000;
+	if (SourceTriCount > ThresholdToCautionAboutLargeInput)
+	{
+		EAppReturnType::Type Ret = FMessageDialog::Open(EAppMsgType::YesNo,
+			FText::Format(LOCTEXT("FractureLargeInputQuestion", "Sources with a large number of triangles ({0}) may take time to convert to Geometry Collection. Continue?"), SourceTriCount),
+			LOCTEXT("FractureLargeInputTitle", "Process large input?"));
+		if (Ret == EAppReturnType::No)
+		{
+			return nullptr;
+		}
+	}
+
+	FScopedSlowTask SlowTask(SourceTriCount, LOCTEXT("FractureCreateNewGeometryCollection", "Creating new Geometry Collection"));
+	SlowTask.MakeDialog();
 
 	AGeometryCollectionActor* NewActor = CreateNewGeometryActor(InAssetPath, FTransform(), true);
 
@@ -687,6 +743,7 @@ AGeometryCollectionActor* UFractureToolGenerateAsset::ConvertActorsToGeometryCol
 				{
 					// If any of the static meshes have Nanite enabled, also enable on the new geometry collection asset for convenience.
 					FracturedGeometryCollection->EnableNanite |= ComponentStaticMesh->IsNaniteEnabled();
+					SlowTask.EnterProgressFrame(ComponentStaticMesh->GetNumTriangles(0));
 				}
 
 				FTransform ComponentTransform(StaticMeshComponent->GetComponentTransform());
@@ -715,6 +772,11 @@ AGeometryCollectionActor* UFractureToolGenerateAsset::ConvertActorsToGeometryCol
 				{
 					// If any of the static meshes have Nanite enabled, also enable on the new geometry collection asset for convenience.
 					FracturedGeometryCollection->EnableNanite |= RestCollection->EnableNanite;
+
+					if (const FGeometryCollectionSharedPtr Collection = RestCollection->GetGeometryCollection())
+					{
+						SlowTask.EnterProgressFrame(Collection->Indices.Num());
+					}
 				}
 
 				FTransform ComponentTransform(GeometryCollectionComponent->GetComponentTransform());
diff --git a/Engine/Source/Runtime/Experimental/GeometryCollectionEngine/Private/GeometryCollection/GeometryCollectionEngineConversion.cpp b/Engine/Source/Runtime/Experimental/GeometryCollectionEngine/Private/GeometryCollection/GeometryCollectionEngineConversion.cpp
--- a/Engine/Source/Runtime/Experimental/GeometryCollectionEngine/Private/GeometryCollection/GeometryCollectionEngineConversion.cpp
+++ b/Engine/Source/Runtime/Experimental/GeometryCollectionEngine/Private/GeometryCollection/GeometryCollectionEngineConversion.cpp
@@ -25,6 +25,7 @@
 #include "Materials/Material.h"
 #include "Rendering/SkeletalMeshRenderData.h"
 #include "MeshDescription.h"
+#include "Misc/ScopedSlowTask.h"
 #include "StaticMeshAttributes.h"
 #include "StaticMeshOperations.h"
 #include "Physics/Experimental/ChaosInterfaceUtils.h"
@@ -36,6 +37,8 @@
 
 DEFINE_LOG_CATEGORY_STATIC(UGeometryCollectionConversionLogging, Log, All);
 
+#define LOCTEXT_NAMESPACE "GeometryCollectionConversion"
+
 struct FUniqueVertex
 {
 	FVector3f Normal;
@@ -104,6 +107,13 @@ void FGeometryCollectionEngineConversion::AppendMeshDescription(
 
 	check(GeometryCollection);
 
+	// prepare to tick progress per 100k vertices
+	const int32 ReportProgressSpacing = 100000;
+	int32 NumVertProgressSteps = int32(MeshDescription->Vertices().GetArraySize() / ReportProgressSpacing);
+
+	FScopedSlowTask AppendMeshDescriptionTask(6 + 2*NumVertProgressSteps, LOCTEXT("AppendMeshDescriptionTask", "Appending Mesh Description Data"));
+	AppendMeshDescriptionTask.EnterProgressFrame(1);
+
 	// source vertex information
 	FStaticMeshConstAttributes Attributes(*MeshDescription);
 	TArrayView<const FVector3f> SourcePosition = Attributes.GetVertexPositions().GetRawArray();
@@ -145,9 +155,17 @@ void FGeometryCollectionEngineConversion::AppendMeshDescription(
 	// A new mapping of MeshDescription vertex instances to the split vertices is maintained.
 	TMap<FVertexInstanceID, int32> VertexInstanceToGeometryCollectionVertex;
 	VertexInstanceToGeometryCollectionVertex.Reserve(Attributes.GetVertexInstanceNormals().GetNumElements());
-		
+
+
+	int32 LastProgress = 0;		
 	for (const FVertexID VertexIndex : MeshDescription->Vertices().GetElementIDs())
-	{		
+	{
+		int32 Progress = int32(VertexIndex / ReportProgressSpacing);
+		if (Progress > LastProgress)
+		{
+			AppendMeshDescriptionTask.EnterProgressFrame(Progress - LastProgress);
+			LastProgress = Progress;
+		}
 		TArrayView<const FVertexInstanceID> ReferencingVertexInstances = MeshDescription->GetVertexVertexInstanceIDs(VertexIndex);
 
 		// Generate per instance hash of splittable attributes.
@@ -202,6 +220,16 @@ void FGeometryCollectionEngineConversion::AppendMeshDescription(
 		}
 	}
 
+	if (LastProgress < NumVertProgressSteps)
+	{
+		AppendMeshDescriptionTask.EnterProgressFrame(NumVertProgressSteps - LastProgress);
+		LastProgress = NumVertProgressSteps;
+	}
+
+	// enter a progress frame for triangle processing w/ size equivalent to the vertex processing (as a heuristic)
+	// (note: could instead tick this per 100k triangles as we do with vertices above, if more responsive progress tracking is desired)
+	AppendMeshDescriptionTask.EnterProgressFrame(NumVertProgressSteps);
+
 	// target triangle indices
 	TManagedArray<FIntVector>& TargetIndices = GeometryCollection->Indices;
 	TManagedArray<bool>& TargetVisible = GeometryCollection->Visible;
@@ -243,6 +271,8 @@ void FGeometryCollectionEngineConversion::AppendMeshDescription(
 		++TargetIndex;
 	}
 
+	AppendMeshDescriptionTask.EnterProgressFrame(1);
+
 	// Geometry transform
 	TManagedArray<FTransform>& Transform = GeometryCollection->Transform;
 
@@ -331,6 +361,8 @@ void FGeometryCollectionEngineConversion::AppendMeshDescription(
 	}
 	if (VertexCount) Center /= VertexCount;
 
+	AppendMeshDescriptionTask.EnterProgressFrame(1);
+
 	// Inner/Outer edges, bounding box
 	BoundingBox[GeometryIndex] = FBox(ForceInitToZero);
 	InnerRadius[GeometryIndex] = FLT_MAX;
@@ -344,6 +376,8 @@ void FGeometryCollectionEngineConversion::AppendMeshDescription(
 		OuterRadius[GeometryIndex] = FMath::Max(OuterRadius[GeometryIndex], Delta);
 	}
 
+	AppendMeshDescriptionTask.EnterProgressFrame(1);
+
 	// Inner/Outer centroid
 	for (int fdx = IndicesStart; fdx < IndicesStart + IndicesCount; fdx++)
 	{
@@ -359,6 +393,8 @@ void FGeometryCollectionEngineConversion::AppendMeshDescription(
 		OuterRadius[GeometryIndex] = FMath::Max(OuterRadius[GeometryIndex], Delta);
 	}
 
+	AppendMeshDescriptionTask.EnterProgressFrame(1);
+
 	// Inner/Outer edges
 	for (int fdx = IndicesStart; fdx < IndicesStart + IndicesCount; fdx++)
 	{
@@ -372,6 +408,8 @@ void FGeometryCollectionEngineConversion::AppendMeshDescription(
 		}
 	}
 
+	AppendMeshDescriptionTask.EnterProgressFrame(1);
+
 	if (ReindexMaterials) {
 		GeometryCollection->ReindexMaterials();
 	}
@@ -554,8 +592,11 @@ bool FGeometryCollectionEngineConversion::AppendStaticMesh(const UStaticMesh* St
 	FGeometryCollection* GeometryCollection, bool bReindexMaterials, bool bAddInternalMaterials, bool bSplitComponents, bool bSetInternalFromMaterialIndex)
 {
 #if WITH_EDITORONLY_DATA
+	FScopedSlowTask AppendStaticMeshTask(bSplitComponents ? 3 : 2, LOCTEXT("AppendStaticMeshTask", "Appending Static Mesh"));
+
 	if (StaticMesh)
 	{
+		AppendStaticMeshTask.EnterProgressFrame(1);
 		FMeshDescription* MeshDescription = GetMaxResMeshDescriptionWithNormalsAndTangents(StaticMesh);
 
 		check(GeometryCollection);
@@ -571,6 +612,8 @@ bool FGeometryCollectionEngineConversion::AppendStaticMesh(const UStaticMesh* St
 
 			if (bSplitComponents)
 			{
+				AppendStaticMeshTask.EnterProgressFrame(1);
+
 				int32 MaxVID = MeshDescription->Vertices().Num();
 				UE::Geometry::FVertexConnectedComponents Components(MaxVID);
 				for (const FTriangleID TriangleID : MeshDescription->Triangles().GetElementIDs())
@@ -695,6 +738,7 @@ bool FGeometryCollectionEngineConversion::AppendStaticMesh(const UStaticMesh* St
 				// else only one component -- fall back to just using the original mesh description
 			}
 
+			AppendStaticMeshTask.EnterProgressFrame(1);
 			AppendMeshDescription(MeshDescription, StaticMesh->GetName(), StartMaterialIndex, MeshTransform, GeometryCollection, StaticMesh->GetBodySetup(), bReindexMaterials, bAddInternalMaterials, bSetInternalFromMaterialIndex);
 			return true;
 		}
@@ -1415,4 +1459,6 @@ void FGeometryCollectionEngineConversion::AppendGeometryCollectionSource(const F
 				);
 		}
 	}
-}
\ No newline at end of file
+}
+
+#undef LOCTEXT_NAMESPACE 
\ No newline at end of file
