Fix interface UGeometryCollectionExternalRenderInterface usage, should be using TScriptInterface as per disfunctional UHT error

diff --git a/Engine/Source/Runtime/Experimental/GeometryCollectionEngine/Private/GeometryCollection/GeometryCollectionComponent.cpp b/Engine/Source/Runtime/Experimental/GeometryCollectionEngine/Private/GeometryCollection/GeometryCollectionComponent.cpp
--- a/Engine/Source/Runtime/Experimental/GeometryCollectionEngine/Private/GeometryCollection/GeometryCollectionComponent.cpp
+++ b/Engine/Source/Runtime/Experimental/GeometryCollectionEngine/Private/GeometryCollection/GeometryCollectionComponent.cpp
@@ -2802,7 +2802,7 @@ void UGeometryCollectionComponent::OnRegister()
 	UClass* Type = bOverrideCustomRenderer ? CustomRendererType : (RestCollection ? RestCollection->CustomRendererType : nullptr);
 	if (Type && Type->ImplementsInterface(UGeometryCollectionExternalRenderInterface::StaticClass()))
 	{
-		CustomRenderer = (UGeometryCollectionExternalRenderInterface*)NewObject<UObject>(this, Type);
+		CustomRenderer = NewObject<UObject>(this, Type);
 		RegisterCustomRenderer();
 	}
 
@@ -2823,7 +2823,7 @@ void UGeometryCollectionComponent::OnUnregister()
 
 void UGeometryCollectionComponent::RegisterCustomRenderer()
 {
-	if (IGeometryCollectionExternalRenderInterface* RendererInterface = Cast<IGeometryCollectionExternalRenderInterface>(CustomRenderer))
+	if (IGeometryCollectionExternalRenderInterface* RendererInterface = CustomRenderer.GetInterface())
 	{
 		RendererInterface->OnRegisterGeometryCollection(*this);
 	}
@@ -2831,7 +2831,7 @@ void UGeometryCollectionComponent::RegisterCustomRenderer()
 
 void UGeometryCollectionComponent::UnregisterCustomRenderer()
 {
-	if (IGeometryCollectionExternalRenderInterface* RendererInterface = Cast<IGeometryCollectionExternalRenderInterface>(CustomRenderer))
+	if (IGeometryCollectionExternalRenderInterface* RendererInterface = CustomRenderer.GetInterface())
 	{
 		RendererInterface->OnUnregisterGeometryCollection();
 	}
@@ -4902,7 +4902,7 @@ void UGeometryCollectionComponent::RefreshCustomRenderer()
 		// Don't refresh the custom renderer on the server but we still need to do the work of computing component space transforms.
 		const bool bUpdateRenderer = !IsNetMode(NM_DedicatedServer);
 
-		if (IGeometryCollectionExternalRenderInterface* RendererInterface = Cast<IGeometryCollectionExternalRenderInterface>(CustomRenderer))
+		if (IGeometryCollectionExternalRenderInterface* RendererInterface = CustomRenderer.GetInterface())
 		{
 			if (RestCollection != nullptr)
 			{
diff --git a/Engine/Source/Runtime/Experimental/GeometryCollectionEngine/Public/GeometryCollection/GeometryCollectionComponent.h b/Engine/Source/Runtime/Experimental/GeometryCollectionEngine/Public/GeometryCollection/GeometryCollectionComponent.h
--- a/Engine/Source/Runtime/Experimental/GeometryCollectionEngine/Public/GeometryCollection/GeometryCollectionComponent.h
+++ b/Engine/Source/Runtime/Experimental/GeometryCollectionEngine/Public/GeometryCollection/GeometryCollectionComponent.h
@@ -48,7 +48,7 @@ struct FGeometryCollectionSection;
 struct FDamageCollector;
 class FPhysScene_Chaos;
 class AGeometryCollectionISMPoolActor;
-class UGeometryCollectionExternalRenderInterface;
+class IGeometryCollectionExternalRenderInterface;
 enum ESimulationInitializationState : uint8;
 enum class EClusterConnectionTypeEnum : uint8;
 enum class EInitialVelocityTypeEnum : uint8;
@@ -1217,7 +1217,7 @@ protected:
 
 	/** A custom renderer object created from CustomRenderType. */
 	UPROPERTY(Transient)
-	TObjectPtr<UGeometryCollectionExternalRenderInterface> CustomRenderer;
+	TScriptInterface<IGeometryCollectionExternalRenderInterface> CustomRenderer;
 
 	/** Populate the dynamic particle data for the render thread. */
 	GEOMETRYCOLLECTIONENGINE_API FGeometryCollectionDynamicData* InitDynamicData(bool bInitialization);
