Improve the 'PrimeEndPoint' command to run many times faster and add additional options.

#rnx

- Moved the priming code to it's own code files.
- Priming no longer uses FOnDemandIoBackend but manages its own requests.
- Bumped up the number of connections to the largest supported by FEventLoop (63) to speed up the priming as much as possible.
- Added additional cmdline options although -PrimeEndPoint is still required:
-- '-PrimeAll' will cause all available CDNs to be primed.
-- '-PrimeUrl=' will override the CDN url with the one provided.
--- '-List' will print all of the available CDN urls to screen along with the average latency to it.
-- '-Help' prints help on the new commands to the screen.
- Fix minor bug where UnrealPak was returning none zero for success and zero for error cases.
- Automatically call ::ResolveDeferredEndpoints as part of FDistributionEndpoints::Flush if needed. Failure to do so can lead to flush waiting forever on work that has never been initiated.

diff --git a/Engine/Source/Developer/IoStoreUtilities/Private/IoStoreUtilities.cpp b/Engine/Source/Developer/IoStoreUtilities/Private/IoStoreUtilities.cpp
--- a/Engine/Source/Developer/IoStoreUtilities/Private/IoStoreUtilities.cpp
+++ b/Engine/Source/Developer/IoStoreUtilities/Private/IoStoreUtilities.cpp
@@ -8782,10 +8782,13 @@ bool DownloadIoStoreContainerFiles(const TCHAR* TocPath)
 bool PrimeEndPoint(FStringView IoStoreOnDemandIniPath)
 {
 	FIoStatus Status = UE::IO::IAS::PrimeEndPoint(IoStoreOnDemandIniPath);
-	if (Status.IsOk() == false)
+	if (Status.IsOk())
 	{
-		UE_LOG(LogIoStore, Error, TEXT("Failed to prime end points, reason '%s'"), *Status.ToString());
+		return true;
+	}
+	else
+	{
+		UE_LOG(LogIoStore, Error, TEXT("PrimeEndPoint command failed due to: %s"), *Status.ToString());
+		return false;
 	}
-
-	return Status.IsOk();
 }
diff --git a/Engine/Source/Developer/PakFileUtilities/Private/PakFileUtilities.cpp b/Engine/Source/Developer/PakFileUtilities/Private/PakFileUtilities.cpp
--- a/Engine/Source/Developer/PakFileUtilities/Private/PakFileUtilities.cpp
+++ b/Engine/Source/Developer/PakFileUtilities/Private/PakFileUtilities.cpp
@@ -5458,19 +5458,22 @@ bool ExecuteUnrealPak(const TCHAR* CmdLine)
 			return CreateIoStoreContainerFiles(CmdLine) == 0;
 		}
 
-		if (FParse::Value(CmdLine, TEXT("-Upload="), IoStoreArg))
+		// IAS commands
 		{
-			return UploadIoStoreContainerFiles(*IoStoreArg) == 0;
-		}
+			if (FParse::Value(CmdLine, TEXT("-Upload="), IoStoreArg))
+			{
+				return UploadIoStoreContainerFiles(*IoStoreArg) == 0;
+			}
 
-		if (FParse::Value(CmdLine, TEXT("-Download="), IoStoreArg))
-		{
-			return DownloadIoStoreContainerFiles(*IoStoreArg) == 0;
-		}
+			if (FParse::Value(CmdLine, TEXT("-Download="), IoStoreArg))
+			{
+				return DownloadIoStoreContainerFiles(*IoStoreArg) == 0;
+			}
 
-		if (FParse::Value(CmdLine, TEXT("-PrimeEndPoint="), IoStoreArg))
-		{
-			return PrimeEndPoint(*IoStoreArg) == 0;
+			if (FParse::Value(CmdLine, TEXT("-PrimeEndPoint="), IoStoreArg))
+			{
+				return PrimeEndPoint(*IoStoreArg);
+			}
 		}
 	}
 
diff --git a/Engine/Source/Runtime/Experimental/IoStoreOnDemand/Private/DistributionEndpoints.cpp b/Engine/Source/Runtime/Experimental/IoStoreOnDemand/Private/DistributionEndpoints.cpp
--- a/Engine/Source/Runtime/Experimental/IoStoreOnDemand/Private/DistributionEndpoints.cpp
+++ b/Engine/Source/Runtime/Experimental/IoStoreOnDemand/Private/DistributionEndpoints.cpp
@@ -29,6 +29,11 @@ bool FDistributionEndpoints::Flush(double TimeOut)
 {
 	TRACE_CPUPROFILER_EVENT_SCOPE(FDistributionEndpoints::Flush);
 
+	if (!bInitialized)
+	{
+		ResolveDeferredEndpoints();
+	}
+
 	FHttpManager& HttpManager = FHttpModule::Get().GetHttpManager();
 
 	const double StartTime = FPlatformTime::Seconds();
diff --git a/Engine/Source/Runtime/Experimental/IoStoreOnDemand/Private/IoStoreOnDemand.cpp b/Engine/Source/Runtime/Experimental/IoStoreOnDemand/Private/IoStoreOnDemand.cpp
--- a/Engine/Source/Runtime/Experimental/IoStoreOnDemand/Private/IoStoreOnDemand.cpp
+++ b/Engine/Source/Runtime/Experimental/IoStoreOnDemand/Private/IoStoreOnDemand.cpp
@@ -19,6 +19,7 @@
 #include "Modules/ModuleManager.h"
 #include "OnDemandHttpClient.h"
 #include "OnDemandIoDispatcherBackend.h"
+#include "PrimeEndpoint.h"
 #include "Serialization/Archive.h"
 #include "Serialization/CompactBinarySerialization.h"
 #include "Serialization/CompactBinaryWriter.h"
@@ -160,7 +161,7 @@ static bool TryParseConfigFileFromPlatformPackage(FOnDemandEndpoint& OutEndpoint
 }
 
 ////////////////////////////////////////////////////////////////////////////////
-static bool TryParseConfigFile(const FString& ConfigPath, FOnDemandEndpoint& OutEndpoint)
+bool TryParseConfigFile(const FString& ConfigPath, FOnDemandEndpoint& OutEndpoint)
 {
 	const FString ConfigFileName = TEXT("IoStoreOnDemand.ini");
 	FString ConfigContent; 
@@ -1464,66 +1465,9 @@ FIoStatus DownloadContainerFiles(const FIoStoreDownloadParams& DownloadParams, c
 
 FIoStatus PrimeEndPoint(FStringView IoStoreOnDemandIniPath)
 {
-	TRACE_CPUPROFILER_EVENT_SCOPE(PrimeEndPoint);
-
-	FOnDemandEndpoint EndPoint;
-	if (!TryParseConfigFile(FString(IoStoreOnDemandIniPath), EndPoint))
-	{
-		return FIoStatus(EIoErrorCode::Unknown, TEXT("Failed to parse config file"));
-	}
-
-	EndPoint.EndpointType = EOnDemandEndpointType::CDN;
-
-	TSharedPtr <IOnDemandIoDispatcherBackend> Backend = MakeOnDemandIoDispatcherBackend(nullptr);
-	Backend->Mount(EndPoint);
-
-	check(!FIoDispatcher::IsInitialized()); // Assume this is only run in standalone programs that have not
-											// yet setup the IoDispatcher
-	FIoDispatcher::Initialize();
-	FIoDispatcher::InitializePostSettings();
-	ON_SCOPE_EXIT
-	{
-		FIoDispatcher::Shutdown();
-	};
-
-	FIoDispatcher::Get().Mount(Backend.ToSharedRef(), MAX_int32);
-
-	if (!Backend->FlushDeferedEndPoints(30.0))
-	{
-		// TODO: Real error value
-		return FIoStatus(EIoErrorCode::Unknown, TEXT("Unable to connect to endpoint, timed out..."));
-	}
-
-	UE_LOG(LogIas, Display, TEXT("Finding on demand chunks..."));
-	TArray<FIoChunkId> Chunks = Backend->GetAllChunkIds();
-	
-	UE_LOG(LogIas, Display, TEXT("Found %d chunks"), Chunks.Num());
-
-	std::atomic<int32> CompletedRequests = 0;
-	FIoBatch Batch;
-	for (const FIoChunkId& ChunkId : Chunks)
-	{
-		Batch.ReadWithCallback(ChunkId, FIoReadOptions(), IoDispatcherPriority_Medium,
-			[&](TIoStatusOr<FIoBuffer> Result)
-			{
-				CompletedRequests++;
-			});
-	}
-
-	UE_LOG(LogIas, Display, TEXT("Priming all chunks for endpoint..."));
-
-	FEventRef BatchCompleted;
-	Batch.IssueAndTriggerEvent(BatchCompleted.Get());
-	
-	while (!BatchCompleted->Wait(FTimespan::FromSeconds(5.0)))
-	{
-		UE_LOG(LogIas, Display, TEXT("Completed %d/%d"), CompletedRequests.load(), Chunks.Num());
-	}
-
-	UE_LOG(LogIas, Display, TEXT("All chunks primed!"));
-
-	return FIoStatus::Ok;
+	return UE::IO::IAS::PrimeEndpointInternal(FString(IoStoreOnDemandIniPath));
 }
+
 #endif // (PLATFORM_DESKTOP && (IS_PROGRAM || WITH_EDITOR))
 
 
diff --git a/Engine/Source/Runtime/Experimental/IoStoreOnDemand/Private/PrimeEndpoint.cpp b/Engine/Source/Runtime/Experimental/IoStoreOnDemand/Private/PrimeEndpoint.cpp
new file mode 100644
--- /dev/null
+++ b/Engine/Source/Runtime/Experimental/IoStoreOnDemand/Private/PrimeEndpoint.cpp
@@ -0,0 +1,325 @@
+// Copyright Epic Games, Inc. All Rights Reserved.
+
+#if (IS_PROGRAM || WITH_EDITOR)
+
+#include "PrimeEndpoint.h"
+
+#include "CoreHttp/LatencyTesting.h"
+#include "DistributionEndpoints.h"
+#include "IO/IoStoreOnDemand.h"
+#include "Misc/CommandLine.h"
+#include "Misc/PathViews.h"
+#include "OnDemandHttpClient.h"
+#include "OnDemandIoDispatcherBackend.h" // For FOnDemandEndpoint
+
+#include <atomic>
+
+namespace UE::IO::IAS
+{
+
+// TODO
+extern bool TryParseConfigFile(const FString& ConfigPath, FOnDemandEndpoint& OutEndpoint);
+
+/*** Utility for creating a url to a .iochunk */
+static void CreateIoChunkUrl(const FOnDemandTocEntry& Entry, FStringView ChunksDirectory, FAnsiStringBuilderBase& Url)
+{
+	// Same code as FChunkRequestParams, could factor out to avoid duplication
+	const FString HashString = LexToString(Entry.Hash);
+	Url << "/" << ChunksDirectory << "/" << HashString.Left(2) << "/" << HashString << ANSITEXTVIEW(".iochunk");
+}
+
+/** Find all of the CDN urls available from the distribution URl */
+static TIoStatusOr<TArray<FString>> FindCDNUrls(const FString& DistributionUrl)
+{
+	TArray<FString> CDNUrls;
+
+	FDistributionEndpoints DistributionInfo;
+	DistributionInfo.ResolveEndpoints(DistributionUrl, [&CDNUrls](const FString& DistributionUrl, TConstArrayView<FString> Urls)
+		{
+			CDNUrls = Urls;
+		});
+
+	if (!DistributionInfo.Flush(30.0))
+	{
+		return FIoStatus(EIoErrorCode::Unknown, TEXT("Unable to connect to endpoint, timed out..."));
+	}
+
+	return CDNUrls;
+}
+
+/** Utility to format time in seconds to a more human readable form */
+static FString SecondsToString(uint64 TotalSeconds)
+{
+	// There is probably code somewhere in the engine for this but I was unable to find
+	// any that gave the display we wanted (FText for example was a bit off)
+	// Replace with generic code when found.
+
+	const uint64 Hours = (TotalSeconds % 86400) / 3600;
+	const uint64 Minutes = (TotalSeconds % 3600) / 60;
+	const uint64 Seconds = (TotalSeconds % 60);
+
+	if (Hours > 0)
+	{
+		return FString::Printf(TEXT("%02dh %02dm %02ds"), Hours, Minutes, Seconds);
+	}
+	else if (Minutes > 0)
+	{
+		return FString::Printf(TEXT("%02dm %02ds"), Minutes, Seconds);
+	}
+	else
+	{
+		return FString::Printf(TEXT("%02ds"), Seconds);
+	}
+}
+
+/** Prints the various commands that can be used when priming an endpoint */
+static void PrintHelp()
+{
+	UE_LOG(LogIas, Display, TEXT("PrimeEndPoint Help"));
+	UE_LOG(LogIas, Display, TEXT("Required:"));
+	UE_LOG(LogIas, Display, TEXT("-PrimeEndPoint=<Path>"));
+	UE_LOG(LogIas, Display, TEXT("\tWhere <Path> is an absolute path to the IoStoreOnDemand.ini file you wish to prime"));
+	UE_LOG(LogIas, Display, TEXT(""));
+	UE_LOG(LogIas, Display, TEXT("Optional:"));
+	UE_LOG(LogIas, Display, TEXT("-All"));
+	UE_LOG(LogIas, Display, TEXT("\tPrimes all avaliable CDNs rather than just the primary"));
+	UE_LOG(LogIas, Display, TEXT("-PrimeUrl=<url>"));
+	UE_LOG(LogIas, Display, TEXT("\tPrimes the CDN found at the provided <url>) "));
+	UE_LOG(LogIas, Display, TEXT("-List"));
+	UE_LOG(LogIas, Display, TEXT("\tPrints a list of avaliable CDNs along with the average latency"));
+}
+
+/** Utility to ping a CDN and find the average latency to it */
+static int32 PingCDN(const FString& Url, const FString& TocPath)
+{
+#if !UE_BUILD_SHIPPING
+	int32 Results[4] = {};
+	UE::IO::IAS::HTTP::LatencyTest(Url, TocPath, MakeArrayView(Results));
+
+	// Ignore the first result as it might be wildly higher due to caching on the CDN
+	return (Results[1] + Results[2] + Results[3]) / 3;
+#else
+	return -1;
+#endif // !UE_BUILD_SHIPPING
+}
+
+/** Prints all of the CDNs to screen along with the average latency for each one */
+static void ListAvailableCDNs(TArrayView<FString> CDNUrls, const FString& TocPath)
+{
+	// Collect all ping info up front so we don't delay while printing the list
+	TArray<int32> Pings;
+	Pings.SetNum(CDNUrls.Num());
+
+	for (int32 Index = 0; Index < CDNUrls.Num(); ++Index)
+	{
+		Pings[Index] = PingCDN(CDNUrls[Index], TocPath);
+	}
+
+	UE_LOG(LogIas, Display, TEXT("Listing CDN Urls:"));
+	for (int32 Index = 0; Index < CDNUrls.Num(); ++Index)
+	{
+		UE_LOG(LogIas, Display, TEXT("%02d: %s (Ping %d ms)"), Index + 1, *CDNUrls[Index], Pings[Index]);
+	}
+}
+
+/** Primes a single CDN */
+static FIoStatus PrimeCDN(FString CDNUrl, const FPrimeEndpointArgs& Args)
+{
+	const double StartTime = FPlatformTime::Seconds();
+
+	CDNUrl = CDNUrl.Replace(TEXT("https"), TEXT("http"));
+
+	UE_LOG(LogIas, Display, TEXT("Priming '%s'"), *CDNUrl);
+	
+	UE_LOG(LogIas, Display, TEXT("\tDownloading Toc from CDN..."));
+	TIoStatusOr<FOnDemandToc> TocResult = LoadTocFromUrl(CDNUrl, Args.TocPath, 3);
+	if (!TocResult.IsOk())
+	{
+		return TocResult.Status();
+	}
+
+	FOnDemandToc Toc = TocResult.ConsumeValueOrDie();
+
+	TArray<FOnDemandTocEntry> Entries;
+	for (const FOnDemandTocContainerEntry& Container : Toc.Containers)
+	{
+		Entries.Append(Container.Entries);
+	}
+
+	// Need to build the directory where the chunks are stored based on the Toc path and info from it's header
+	TStringBuilder<512> ChunksDirectory;
+	{
+		int32 Idx = INDEX_NONE;
+		if (Args.TocPath.FindLastChar(TCHAR('/'), Idx))
+		{
+			ChunksDirectory << FStringView(Args.TocPath).Left(Idx);
+		}
+
+		FPathViews::Append(ChunksDirectory, Toc.Header.ChunksDirectory);
+	}
+
+	UE_LOG(LogIas, Display, TEXT("\tToc contained %d entries to prime"), Entries.Num());
+
+	// Note the max number of connections that FEventLoop supports is  (2^6)-1
+	const int32 NumConnections = 63;
+
+	UE_LOG(LogIas, Display, TEXT("\tPriming entries with %d connections..."), NumConnections);
+
+	{
+		auto AnsiUrl = StringCast<ANSICHAR>(*CDNUrl);
+
+		using namespace UE::IO::IAS::HTTP;
+
+		FConnectionPool::FParams PoolParams;
+		PoolParams.SetHostFromUrl(AnsiUrl);
+		PoolParams.ConnectionCount = NumConnections;
+		FConnectionPool Pool(PoolParams);
+
+		FEventLoop Loop;
+
+		std::atomic<int32> Inflight = 0;
+		std::atomic<int32> Completed = 0;
+		std::atomic<int64> TotalQueryTimeMs = 0;
+
+		double TotalTime = FPlatformTime::Seconds();
+		double Timer = FPlatformTime::Seconds();
+
+		// Called when we are waiting for requests to complete
+		auto BusyLoop = [&Loop, &Completed, &Timer,&TotalTime, &TotalQueryTimeMs, Total = Entries.Num()]()
+			{
+				Loop.Tick(-1);
+
+				const double DisplayFrequency = 10.0;
+				if (FPlatformTime::Seconds() - Timer >= DisplayFrequency)
+				{
+					const double ItemsPerSecond = Completed / (FPlatformTime::Seconds() - TotalTime);
+					const double TimeEstimate = static_cast<double>(Total - Completed) / ItemsPerSecond;
+					const FString EstimateString = SecondsToString(static_cast<uint64>(TimeEstimate));
+
+					const int64 AvgQueryTimeMs = TotalQueryTimeMs / Completed;
+
+					UE_LOG(LogIas, Display, TEXT("\tCompleted %d/%d | Avg Query (ms) %lld | Est Time Remaining: %s"), Completed.load(), Total, AvgQueryTimeMs , *EstimateString);
+					Timer = FPlatformTime::Seconds();
+				}
+			};
+
+		for (const FOnDemandTocEntry& Entry : Entries)
+		{
+			TAnsiStringBuilder<256> Url;
+			CreateIoChunkUrl(Entry, ChunksDirectory, Url);
+
+			while (Inflight >= NumConnections)
+			{
+				BusyLoop();
+			}
+
+			FRequest Request = Loop.Request("HEAD", Url, Pool);
+			Loop.Send(MoveTemp(Request), [&Completed, &Inflight, &TotalQueryTimeMs, StartTime = FPlatformTime::Seconds()](const FTicketStatus& Status)
+				{
+					const FTicketStatus::EId StatusId = Status.GetId();
+					if (StatusId == FTicketStatus::EId::Response)
+					{
+						return;
+					}
+
+					const double QueryTimeMs = (FPlatformTime::Seconds() - StartTime) * 1000.0;
+					TotalQueryTimeMs += FMath::RoundToInt(QueryTimeMs);
+
+					Completed++;
+					--Inflight;
+				});
+
+			Inflight++;
+		}
+
+		// Wait for the remaining inflight requests to finish up
+		while (Completed != Entries.Num())
+		{
+			BusyLoop();
+		}
+	}
+
+	const double TotalTimeTaken = FPlatformTime::Seconds() - StartTime;
+	UE_LOG(LogIas, Display, TEXT("\tPriming took %s"), *SecondsToString(static_cast<uint64>(TotalTimeTaken)));
+
+	return FIoStatus::Ok;
+}
+
+/** Primes all of the given CDNs */
+FIoStatus PrimeAllCDNs(TArray<FString> Urls, const FPrimeEndpointArgs& Args)
+{
+	if (Urls.IsEmpty())
+	{
+		return FIoStatus(EIoErrorCode::NotFound, TEXT("No CDN urls were found!"));
+	}
+
+	for (const FString& Url : Urls)
+	{
+		PrimeCDN(Url, Args);
+	}
+
+	return FIoStatus::Ok;
+}
+
+FIoStatus PrimeEndpointInternal(const FString& IoStoreOnDemandIniPath)
+{
+	TRACE_CPUPROFILER_EVENT_SCOPE(PrimeEndpointInternalArgsIni);
+
+	if (FParse::Param(FCommandLine::Get(), TEXT("Help")))
+	{
+		PrintHelp();
+		return FIoStatus::Ok;
+	}
+
+	FOnDemandEndpoint Endpoint;
+
+	UE_LOG(LogIas, Display, TEXT("Parsing '%s'..."), *IoStoreOnDemandIniPath);
+	if (!TryParseConfigFile(IoStoreOnDemandIniPath, Endpoint))
+	{
+		return FIoStatus(EIoErrorCode::Unknown, TEXT("Failed to parse config file"));
+	}
+
+	FPrimeEndpointArgs Args(MoveTemp(Endpoint.DistributionUrl), MoveTemp(Endpoint.TocPath));
+	return PrimeEndpointInternal(Args);
+}
+
+FIoStatus PrimeEndpointInternal(const FPrimeEndpointArgs & Args)
+{
+	TRACE_CPUPROFILER_EVENT_SCOPE(PrimeEndpointInternalArgs);
+
+	TIoStatusOr<TArray<FString>> Result = FindCDNUrls(Args.DistributionUrl);
+	if (!Result.IsOk())
+	{
+		return Result.Status();
+	}
+
+	TArray<FString> CDNUrls = Result.ConsumeValueOrDie();
+
+	FString CmdlineUrl;
+	if (FParse::Param(FCommandLine::Get(), TEXT("PrimeAll")))
+	{
+		return PrimeAllCDNs(CDNUrls, Args);
+	}
+	else if (FParse::Value(FCommandLine::Get(), TEXT("PrimeUrl="), CmdlineUrl))
+	{
+		return PrimeCDN(CmdlineUrl, Args);
+	}
+	else if (FParse::Param(FCommandLine::Get(), TEXT("List")))
+	{
+		ListAvailableCDNs(CDNUrls, Args.TocPath);
+		return FIoStatus::Ok;
+	}
+	else
+	{
+		if (CDNUrls.IsEmpty())
+		{
+			return FIoStatus(EIoErrorCode::NotFound, TEXT("No CDN urls were found!"));
+		}
+
+		return PrimeCDN(CDNUrls[0], Args);
+	}
+}
+
+} // namespace UE::IO::IAS
+
+#endif // (IS_PROGRAM || WITH_EDITOR)
diff --git a/Engine/Source/Runtime/Experimental/IoStoreOnDemand/Private/PrimeEndpoint.h b/Engine/Source/Runtime/Experimental/IoStoreOnDemand/Private/PrimeEndpoint.h
new file mode 100644
--- /dev/null
+++ b/Engine/Source/Runtime/Experimental/IoStoreOnDemand/Private/PrimeEndpoint.h
@@ -0,0 +1,42 @@
+// Copyright Epic Games, Inc. All Rights Reserved.
+
+#pragma once
+
+#if (IS_PROGRAM || WITH_EDITOR)
+
+#include "Containers/StringFwd.h"
+#include "IO/IoStatus.h"
+
+namespace UE::IO::IAS
+{
+
+struct FPrimeEndpointArgs
+{
+	FPrimeEndpointArgs() = delete;
+
+	FPrimeEndpointArgs(FStringView InDistributionUrl, FStringView InTocPath)
+		: DistributionUrl(InDistributionUrl)
+		, TocPath(InTocPath)
+	{
+	}
+
+	FPrimeEndpointArgs(FString&& InDistributionUrl, FString&& InTocPath)
+		: DistributionUrl(MoveTemp(InDistributionUrl))
+		, TocPath(MoveTemp(InTocPath))
+	{
+	}
+
+	~FPrimeEndpointArgs() = default;
+
+	FString DistributionUrl;
+	FString TocPath;
+};
+
+/** Derives FPrimeEndpointArgs based on the provided ini file */
+FIoStatus PrimeEndpointInternal(const FString& IoStoreOnDemandIniPath);
+
+FIoStatus PrimeEndpointInternal(const FPrimeEndpointArgs& Args);
+
+} //namespace UE::IO::IAS
+
+#endif // (IS_PROGRAM || WITH_EDITOR)
diff --git a/Engine/Source/Runtime/Experimental/IoStoreOnDemand/Public/IO/IoStoreOnDemand.h b/Engine/Source/Runtime/Experimental/IoStoreOnDemand/Public/IO/IoStoreOnDemand.h
--- a/Engine/Source/Runtime/Experimental/IoStoreOnDemand/Public/IO/IoStoreOnDemand.h
+++ b/Engine/Source/Runtime/Experimental/IoStoreOnDemand/Public/IO/IoStoreOnDemand.h
@@ -25,12 +25,17 @@ class IIoStoreWriter;
 struct FAnalyticsEventAttribute;
 struct FIoContainerSettings;
 struct FIoStoreWriterSettings;
+namespace UE::IO::IAS { struct FOnDemandEndpoint; }
 
 UE_API DECLARE_LOG_CATEGORY_EXTERN(LogIas, VeryVerbose, All);
 
 namespace UE::IO::IAS
 {
 
+////////////////////////////////////////////////////////////////////////////////
+
+bool TryParseConfigFile(const FString& ConfigPath, FOnDemandEndpoint& OutEndpoint);
+
 ////////////////////////////////////////////////////////////////////////////////
 enum class EOnDemandTocVersion : uint32
 {
