IAS - append slash to persistent download dir if missing

#rnx

diff --git a/Engine/Source/Runtime/Experimental/IoStoreOnDemand/Private/JournaledCache.cpp b/Engine/Source/Runtime/Experimental/IoStoreOnDemand/Private/JournaledCache.cpp
--- a/Engine/Source/Runtime/Experimental/IoStoreOnDemand/Private/JournaledCache.cpp
+++ b/Engine/Source/Runtime/Experimental/IoStoreOnDemand/Private/JournaledCache.cpp
@@ -1142,6 +1142,10 @@ FJournaledCache::FJournaledCache(const FIasCacheConfig& Config)
 {
 	TStringBuilder<256> CachePath;
 	CachePath << FPaths::ProjectPersistentDownloadDir();
+	if (CachePath.LastChar() != TEXT('/'))
+	{
+		CachePath.AppendChar(TEXT('/'));
+	}
 	CachePath << TEXT("ias/");
 	CachePath << Config.Name;
 	CachePath << TEXT(".cache.0");
