Iris - Match alloc/free.

diff --git a/Engine/Source/Runtime/Experimental/Iris/Core/Private/Iris/ReplicationState/ReplicationStateDescriptor.cpp b/Engine/Source/Runtime/Experimental/Iris/Core/Private/Iris/ReplicationState/ReplicationStateDescriptor.cpp
--- a/Engine/Source/Runtime/Experimental/Iris/Core/Private/Iris/ReplicationState/ReplicationStateDescriptor.cpp
+++ b/Engine/Source/Runtime/Experimental/Iris/Core/Private/Iris/ReplicationState/ReplicationStateDescriptor.cpp
@@ -93,7 +93,8 @@ void FReplicationStateDescriptor::Release() const
 				FunctionDescriptor.Descriptor->Release();
 			}
 
-			delete this;
+			// Allocated via FMemory::Malloc in ReplicationStateDescriptorBuilder
+			FMemory::Free(const_cast<FReplicationStateDescriptor*>(this));
 		}
 	}
 }
