PLAY-9856 & UE-193935
Remove warning log when a SimObject with resim option register without resim being enabled.
Dynamically scale down memory allocation for rewind caching when cache shrinks.

diff --git a/Engine/Source/Runtime/Experimental/Chaos/Private/PBDRigidsSolver.cpp b/Engine/Source/Runtime/Experimental/Chaos/Private/PBDRigidsSolver.cpp
--- a/Engine/Source/Runtime/Experimental/Chaos/Private/PBDRigidsSolver.cpp
+++ b/Engine/Source/Runtime/Experimental/Chaos/Private/PBDRigidsSolver.cpp
@@ -1682,7 +1682,6 @@ namespace Chaos
 				else
 				{
 					ensure(!IsNetworkPhysicsPredictionEnabled());
-					UE_LOG(LogChaos, Warning, TEXT("A SimCallbackObject with ESimCallbackOptions::Rewind defined is registering but there is no IRewindCallback set up. Make sure to enable 'Project Settings -> Physics -> Physics Prediction'"));
 				}
 			}
 
diff --git a/Engine/Source/Runtime/Experimental/Chaos/Public/Chaos/EvolutionResimCache.h b/Engine/Source/Runtime/Experimental/Chaos/Public/Chaos/EvolutionResimCache.h
--- a/Engine/Source/Runtime/Experimental/Chaos/Public/Chaos/EvolutionResimCache.h
+++ b/Engine/Source/Runtime/Experimental/Chaos/Public/Chaos/EvolutionResimCache.h
@@ -10,6 +10,9 @@ namespace Chaos
 {
 	class FCollisionResimCache;
 
+	static float ResimCacheReallocationLeniency = 10;
+	static FAutoConsoleVariableRef CVarResimCacheReallocationLeniency(TEXT("np2.ResimCacheReallocationLeniency"), ResimCacheReallocationLeniency, TEXT("Default = 10. Percentage value from 0-100 how much the particle cache needs to grow or shrink before a memory reallocation is performed"));
+
 	class FEvolutionResimCache : public IResimCacheBase
 	{
 	public:
@@ -17,9 +20,28 @@ namespace Chaos
 		virtual ~FEvolutionResimCache() = default;
 		void ResetCache()
 		{
+			// If the number of cached particles is shrinking, reduce the memory allocation but leave room for 10% growth without needing to reallocate
+			// Memory will be reallocated if the cached particles grow or shrink by 10%, growing reallocation is done automatically, this logic handles shrink reallocation and leaving room for 10% growth.
+			// Example: At 100 particles it's allowed to populate the TMap between 90-110 particles without a memory reallocation.
+			const int32 CurrentSize = ParticleToCachedSolve.Num();
+			const int32 SizeLeniency = FMath::CeilToInt(FMath::Clamp((1.0f / ResimCacheReallocationLeniency), 0.0f, 1.0f) * CurrentSize); // Example: 10% leniency
+			ParticleCacheAllocationSize = FMath::Max(CurrentSize, ParticleCacheAllocationSize);
+			const int32 ReallocationLimit = FMath::Max(ParticleCacheAllocationSize - (SizeLeniency * 2), 0);
+
+			if (CurrentSize < ReallocationLimit)
+			{
+				const int32 PreferredSize = CurrentSize + SizeLeniency;
+
+				ParticleToCachedSolve.Empty(PreferredSize);
+				ParticleCacheAllocationSize = PreferredSize;
+			}
+			else
+			{
+				ParticleToCachedSolve.Reset();
+			}
+
 			SavedConstraints.Reset();
 			WeakSinglePointConstraints.Reset();
-			ParticleToCachedSolve.Reset();
 		}
 
 		void SaveParticlePostSolve(const FPBDRigidParticleHandle& Particle)
@@ -151,6 +173,7 @@ namespace Chaos
 		//TODO: better way to handle this?
 		TArray<FWeakConstraintPair> WeakSinglePointConstraints;
 		TMap<FUniqueIdx, FPBDSolveCache> ParticleToCachedSolve;
+		int32 ParticleCacheAllocationSize = 0;
 	};
 
 } // namespace Chaos
