Fix tree synchronization by replicating the construction order from the server

diff --git a/Engine/Source/Runtime/Experimental/ChaosVehicles/ChaosVehiclesCore/Private/SimModule/SimModuleTree.cpp b/Engine/Source/Runtime/Experimental/ChaosVehicles/ChaosVehiclesCore/Private/SimModule/SimModuleTree.cpp
--- a/Engine/Source/Runtime/Experimental/ChaosVehicles/ChaosVehiclesCore/Private/SimModule/SimModuleTree.cpp
+++ b/Engine/Source/Runtime/Experimental/ChaosVehicles/ChaosVehiclesCore/Private/SimModule/SimModuleTree.cpp
@@ -432,11 +432,14 @@ void FSimModuleTree::GenerateReplicationStructure(Chaos::FModuleNetDataArray& Ne
 	NetData.Reserve(Tree.Num());
 	for (int Index = 0; Index < Tree.Num(); Index++)
 	{
-		TSharedPtr<FModuleNetData>&& Data = Tree[Index].SimModule->GenerateNetData(Index);
-		// not all modules will have net replication data - nullptr is a valid response
-		if (Data)
+		if (ISimulationModuleBase* SimModule = Tree[Index].SimModule)
 		{
-			NetData.Emplace(Data);
+			TSharedPtr<FModuleNetData>&& Data = SimModule->GenerateNetData(Index);
+			// not all modules will have net replication data - nullptr is a valid response
+			if (Data)
+			{
+				NetData.Emplace(Data);
+			}
 		}
 	}
 }
@@ -450,9 +453,12 @@ void FSimModuleTree::SetNetState(Chaos::FModuleNetDataArray& ModuleDatas)
 
 	for (TSharedPtr<FModuleNetData>& DataElement : ModuleDatas)
 	{
-		if (!SimulationModuleTree.IsEmpty() && SimulationModuleTree[DataElement->SimArrayIndex].SimModule)
+		if (!SimulationModuleTree.IsEmpty() && DataElement->SimArrayIndex < SimulationModuleTree.Num())
 		{
-			DataElement->FillNetState(SimulationModuleTree[DataElement->SimArrayIndex].SimModule);
+			if (ISimulationModuleBase* SimModule = SimulationModuleTree[DataElement->SimArrayIndex].SimModule)
+			{
+				DataElement->FillNetState(SimModule);
+			}
 		}
 	}
 }
@@ -461,9 +467,12 @@ void FSimModuleTree::SetSimState(const Chaos::FModuleNetDataArray& ModuleDatas)
 {
 	for (const TSharedPtr<FModuleNetData>& DataElement : ModuleDatas)
 	{
-		if (!SimulationModuleTree.IsEmpty() && DataElement->SimArrayIndex < SimulationModuleTree.Num() && SimulationModuleTree[DataElement->SimArrayIndex].SimModule)
+		if (!SimulationModuleTree.IsEmpty() && DataElement->SimArrayIndex < SimulationModuleTree.Num())
 		{
-			DataElement->FillSimState(SimulationModuleTree[DataElement->SimArrayIndex].SimModule);
+			if (ISimulationModuleBase* SimModule = SimulationModuleTree[DataElement->SimArrayIndex].SimModule)
+			{
+				DataElement->FillSimState(SimModule);
+			}
 		}
 	}
 }
diff --git a/Engine/Source/Runtime/Experimental/ChaosVehicles/ChaosVehiclesCore/Public/SimModule/TorqueSimModule.h b/Engine/Source/Runtime/Experimental/ChaosVehicles/ChaosVehiclesCore/Public/SimModule/TorqueSimModule.h
--- a/Engine/Source/Runtime/Experimental/ChaosVehicles/ChaosVehiclesCore/Public/SimModule/TorqueSimModule.h
+++ b/Engine/Source/Runtime/Experimental/ChaosVehicles/ChaosVehiclesCore/Public/SimModule/TorqueSimModule.h
@@ -73,6 +73,8 @@ public:
 		);
 	}
 
+	virtual eSimType GetSimType() const { return eSimType::TorqueSim; }
+
 	void SetDriveTorque(float TorqueIn) { DriveTorque = TorqueIn; }
 	float GetDriveTorque() const { return DriveTorque; }
 
