[NetPhysics] Fix int to float implicit cast.

diff --git a/Engine/Source/Runtime/Experimental/Chaos/Public/Chaos/EvolutionResimCache.h b/Engine/Source/Runtime/Experimental/Chaos/Public/Chaos/EvolutionResimCache.h
--- a/Engine/Source/Runtime/Experimental/Chaos/Public/Chaos/EvolutionResimCache.h
+++ b/Engine/Source/Runtime/Experimental/Chaos/Public/Chaos/EvolutionResimCache.h
@@ -24,7 +24,7 @@ namespace Chaos
 			// Memory will be reallocated if the cached particles grow or shrink by 10%, growing reallocation is done automatically, this logic handles shrink reallocation and leaving room for 10% growth.
 			// Example: At 100 particles it's allowed to populate the TMap between 90-110 particles without a memory reallocation.
 			const int32 CurrentSize = ParticleToCachedSolve.Num();
-			const int32 SizeLeniency = FMath::CeilToInt(FMath::Clamp((1.0f / ResimCacheReallocationLeniency), 0.0f, 1.0f) * CurrentSize); // Example: 10% leniency
+			const int32 SizeLeniency = FMath::CeilToInt(FMath::Clamp((1.0f / ResimCacheReallocationLeniency), 0.0f, 1.0f) * static_cast<float>(CurrentSize)); // Example: 10% leniency
 			ParticleCacheAllocationSize = FMath::Max(CurrentSize, ParticleCacheAllocationSize);
 			const int32 ReallocationLimit = FMath::Max(ParticleCacheAllocationSize - (SizeLeniency * 2), 0);
 
