[NetPhysics] Fix CVar

diff --git a/Engine/Source/Runtime/Experimental/Chaos/Public/Chaos/EvolutionResimCache.h b/Engine/Source/Runtime/Experimental/Chaos/Public/Chaos/EvolutionResimCache.h
--- a/Engine/Source/Runtime/Experimental/Chaos/Public/Chaos/EvolutionResimCache.h
+++ b/Engine/Source/Runtime/Experimental/Chaos/Public/Chaos/EvolutionResimCache.h
@@ -10,9 +10,6 @@ namespace Chaos
 {
 	class FCollisionResimCache;
 
-	static float ResimCacheReallocationLeniency = 10;
-	static FAutoConsoleVariableRef CVarResimCacheReallocationLeniency(TEXT("np2.ResimCacheReallocationLeniency"), ResimCacheReallocationLeniency, TEXT("Default = 10. Percentage value from 0-100 how much the particle cache needs to grow or shrink before a memory reallocation is performed"));
-
 	class FEvolutionResimCache : public IResimCacheBase
 	{
 	public:
@@ -24,7 +21,7 @@ namespace Chaos
 			// Memory will be reallocated if the cached particles grow or shrink by 10%, growing reallocation is done automatically, this logic handles shrink reallocation and leaving room for 10% growth.
 			// Example: At 100 particles it's allowed to populate the TMap between 90-110 particles without a memory reallocation.
 			const int32 CurrentSize = ParticleToCachedSolve.Num();
-			const int32 SizeLeniency = FMath::CeilToInt(FMath::Clamp((1.0f / ResimCacheReallocationLeniency), 0.0f, 1.0f) * static_cast<float>(CurrentSize)); // Example: 10% leniency
+			const int32 SizeLeniency = FMath::CeilToInt(0.1f * static_cast<float>(CurrentSize)); // 10% leniency
 			ParticleCacheAllocationSize = FMath::Max(CurrentSize, ParticleCacheAllocationSize);
 			const int32 ReallocationLimit = FMath::Max(ParticleCacheAllocationSize - (SizeLeniency * 2), 0);
 
