Iris
* Added ability to define the inline allocator size for ReplicationViews via UE_IRIS_INLINE_VIEWS_PER_CONNECTION
* Set this to the maximum amount of splitscreen players per connection (ex: splitscreen) your game can have.

#rnx

diff --git a/Engine/Source/Runtime/Experimental/Iris/Core/Private/Iris/ReplicationSystem/ReplicationConnections.h b/Engine/Source/Runtime/Experimental/Iris/Core/Private/Iris/ReplicationSystem/ReplicationConnections.h
--- a/Engine/Source/Runtime/Experimental/Iris/Core/Private/Iris/ReplicationSystem/ReplicationConnections.h
+++ b/Engine/Source/Runtime/Experimental/Iris/Core/Private/Iris/ReplicationSystem/ReplicationConnections.h
@@ -81,7 +81,7 @@ public:
 	void InitDataStreams(uint32 ReplicationSystemId, uint32 ConnectionId, UDataStreamManager* DataStreamManager);
 
 	void SetReplicationView(uint32 ConnectionId, const FReplicationView& ViewInfo);
-	FReplicationView GetReplicationView(uint32 ConnectionId) const { return ReplicationViews[ConnectionId]; }
+	const FReplicationView& GetReplicationView(uint32 ConnectionId) const { return ReplicationViews[ConnectionId]; }
 
 	const FNetTokenStoreState& GetRemoteNetTokenStoreState(uint32 ConnectionId) const { return RemoteNetTokenStoreStates[ConnectionId]; }
 	FNetTokenStoreState& GetRemoteNetTokenStoreState(uint32 ConnectionId) { return RemoteNetTokenStoreStates[ConnectionId]; }
diff --git a/Engine/Source/Runtime/Experimental/Iris/Core/Private/Iris/ReplicationSystem/ReplicationOperationsInternal.cpp b/Engine/Source/Runtime/Experimental/Iris/Core/Private/Iris/ReplicationSystem/ReplicationOperationsInternal.cpp
--- a/Engine/Source/Runtime/Experimental/Iris/Core/Private/Iris/ReplicationSystem/ReplicationOperationsInternal.cpp
+++ b/Engine/Source/Runtime/Experimental/Iris/Core/Private/Iris/ReplicationSystem/ReplicationOperationsInternal.cpp
@@ -82,7 +82,7 @@ uint32 FReplicationInstanceOperationsInternal::CopyObjectStateData(FNetBitStream
 			// if the object was scopable prev frame we can do partial copy
 			bool bShouldPropagateChangedStates = Object.bShouldPropagateChangedStates;
 
-			// We do a full CopyAndQunatize if the cvar bCVarForceFullCopyAndQuantize is set or that the object is marked as needing a FullCopyAndQuantize
+			// We do a full CopyAndQuantize if the cvar bCVarForceFullCopyAndQuantize is set or that the object is marked as needing a FullCopyAndQuantize
 			const bool bUseFullCopyAndQuantize = bCVarForceFullCopyAndQuantize || Object.bNeedsFullCopyAndQuantize;
 			Object.bNeedsFullCopyAndQuantize = 0U;
 
diff --git a/Engine/Source/Runtime/Experimental/Iris/Core/Public/Iris/ReplicationSystem/ReplicationView.h b/Engine/Source/Runtime/Experimental/Iris/Core/Public/Iris/ReplicationSystem/ReplicationView.h
--- a/Engine/Source/Runtime/Experimental/Iris/Core/Public/Iris/ReplicationSystem/ReplicationView.h
+++ b/Engine/Source/Runtime/Experimental/Iris/Core/Public/Iris/ReplicationSystem/ReplicationView.h
@@ -6,6 +6,11 @@
 #include "Containers/Array.h"
 #include "Math/Vector.h"
 
+// Define this macro in your game's Target.cs and set it to the maximum clients per connection your game can have (ex: splitscreen).
+#ifndef UE_IRIS_INLINE_VIEWS_PER_CONNECTION
+	#define UE_IRIS_INLINE_VIEWS_PER_CONNECTION 4
+#endif
+
 namespace UE::Net
 {
 
@@ -18,7 +23,7 @@ struct FReplicationView
 		float FoVRadians;
 	};
 
-	TArray<FView, TInlineAllocator<4>> Views;
+	TArray<FView, TInlineAllocator<UE_IRIS_INLINE_VIEWS_PER_CONNECTION>> Views;
 };
 
 }
