Iris - Moved NetCullDistanceSquared property to lifetime conditions state and marked it to never replicate and push based. This will prevent overwriting any values on the receiving end.

diff --git a/Engine/Source/Runtime/Experimental/Iris/Core/Private/Iris/ReplicationState/ReplicationStateDescriptorBuilder.cpp b/Engine/Source/Runtime/Experimental/Iris/Core/Private/Iris/ReplicationState/ReplicationStateDescriptorBuilder.cpp
--- a/Engine/Source/Runtime/Experimental/Iris/Core/Private/Iris/ReplicationState/ReplicationStateDescriptorBuilder.cpp
+++ b/Engine/Source/Runtime/Experimental/Iris/Core/Private/Iris/ReplicationState/ReplicationStateDescriptorBuilder.cpp
@@ -2723,10 +2723,10 @@ SIZE_T FReplicationStateDescriptorBuilder::CreateDescriptorsForClass(FResult& Cr
 				FPropertyReplicationStateDescriptorBuilder::GetSerializerTraits(NetCullDistanceSqrMemberProperty, NetCullDistanceSquaredProperty, NetCullDistanceSqrMemberProperty.SerializerInfo);
 
 				// These needs to be set/fixed after the Traits calls.
-				NetCullDistanceSqrMemberProperty.Traits |= EMemberPropertyTraits::InitOnly;
-				NetCullDistanceSqrMemberProperty.ReplicationCondition = COND_InitialOnly;
+				NetCullDistanceSqrMemberProperty.Traits |= EMemberPropertyTraits::HasLifetimeConditionals | EMemberPropertyTraits::HasPushBasedDirtiness;
+				NetCullDistanceSqrMemberProperty.ReplicationCondition = COND_Never;
 
-				Builders[InitPropertyReplicationStateBuilderIndex].AddMemberProperty(NetCullDistanceSqrMemberProperty);
+				Builders[LifetimeConditionalsReplicationStateBuilderIndex].AddMemberProperty(NetCullDistanceSqrMemberProperty);
 			}
 		}
 	}
