Fix constant falloff for Apply Radial Impulse

diff --git a/Engine/Source/Runtime/Experimental/Chaos/Private/Chaos/PhysicsObjectInterface.cpp b/Engine/Source/Runtime/Experimental/Chaos/Private/Chaos/PhysicsObjectInterface.cpp
--- a/Engine/Source/Runtime/Experimental/Chaos/Private/Chaos/PhysicsObjectInterface.cpp
+++ b/Engine/Source/Runtime/Experimental/Chaos/Private/Chaos/PhysicsObjectInterface.cpp
@@ -84,31 +84,39 @@ namespace
 	{
 		using namespace Chaos;
 
-		if (ParticleHandle == nullptr)
-		{
-			return 0.0f;
-		}
+		float FalloffStrength = 0.f;
 
-		const FVec3 CurrentImpulseVelocity = ParticleHandle->LinearImpulseVelocity();
-
-		float FalloffStrength = Strength;
-		if (Falloff == ERadialImpulseFalloff::RIF_Linear)
+		if (ParticleHandle)
 		{
-			//Radius should always be greater than distance - if we get here and it's not, something has gone
-			//wrong with detection
-			double Distance = FVec3::Distance(ParticleHandle->X(), Origin);
-			FalloffStrength = static_cast<float>(FalloffStrength * (1.0f - FMath::Min(1.0f, Distance / Radius)));
-		}
+			const FVec3 ParticleToOrigin = ParticleHandle->X() - Origin;
+			const float DistanceSquared = (float)ParticleToOrigin.SizeSquared();
+			const float RadiusSquared = Radius * Radius;
 
-		FVec3 Impulse = ParticleHandle->X() - Origin;
-		Impulse = Impulse.GetSafeNormal();
-		Impulse = Impulse * FalloffStrength;
-		FVec3 Velocity = Impulse * ParticleHandle->InvM() * VelocityRatio;
+			if (DistanceSquared < RadiusSquared)
+			{
+				// by default we are within the radius so strength is maximum
+				// equivalent to ERadialImpulseFalloff::RIF_Constant
+				FalloffStrength = Strength;
+
+				if (Falloff == ERadialImpulseFalloff::RIF_Linear)
+				{
+					const float Distance = FMath::Sqrt(DistanceSquared);
+					FalloffStrength = static_cast<float>(FalloffStrength * (1.0f - Distance / Radius));
+				}
 
-		FVec3 FinalVeloc = CurrentImpulseVelocity + Velocity;
+				// if the strength was still strong enough to consider
+				if (FalloffStrength > 0)
+				{
+					const FVec3 Normal = ParticleToOrigin.GetSafeNormal();
+					const FVec3 Impulse = Normal * FalloffStrength;
+					const FVec3 Velocity = Impulse * ParticleHandle->InvM() * VelocityRatio;
 
-		ParticleHandle->SetLinearImpulseVelocity(CurrentImpulseVelocity + Velocity, bInvalidate);
+					const FVec3 CurrentImpulseVelocity = ParticleHandle->LinearImpulseVelocity();
 
+					ParticleHandle->SetLinearImpulseVelocity(CurrentImpulseVelocity + Velocity, bInvalidate);
+				}
+			}
+		}
 		return FalloffStrength;
 	}
 }
@@ -909,8 +917,10 @@ namespace Chaos
 
 									//to do: remove cvar when material system is in place and densities are updated
 									FalloffStrength = PhysicsObjectInterfaceCVars::StrainModifier * FalloffStrength;
-
-									Clustering.SetExternalStrain(ChildHandle->CastToClustered(), FalloffStrength);
+									if (FalloffStrength > 0)
+									{
+										Clustering.SetExternalStrain(ChildHandle->CastToClustered(), FalloffStrength);
+									}
 								}
 							}
 							else
@@ -1132,4 +1142,4 @@ namespace Chaos
 
 	template class FWritePhysicsObjectInterface<EThreadContext::External>;
 	template class FWritePhysicsObjectInterface<EThreadContext::Internal>;
-}
\ No newline at end of file
+}
