Strategic submit of IAS cache configuration that will be needed in the future.

#rnx

diff --git a/Engine/Source/Runtime/Experimental/IoStoreOnDemand/Private/IasCache.h b/Engine/Source/Runtime/Experimental/IoStoreOnDemand/Private/IasCache.h
--- a/Engine/Source/Runtime/Experimental/IoStoreOnDemand/Private/IasCache.h
+++ b/Engine/Source/Runtime/Experimental/IoStoreOnDemand/Private/IasCache.h
@@ -42,12 +42,20 @@ struct FIasCacheConfig
 		uint32	Seconds = 60;
 	};
 
+	struct FDemand
+	{
+		uint8	Threshold = 30;
+		uint8	Boost = 60;
+		uint8	SuperBoost = 87;
+	};
+
 	const TCHAR*Name = TEXT("ias");
 	uint64		DiskQuota = 512ull << 20;
 	uint32		MemoryQuota = 2 << 20;
 	uint32		JournalQuota = 4 << 20; // description in JournalCache.cpp
 	uint32		JournalFlushInterval = 4;
 	FRate		WriteRate;
+	FDemand		Demand;
 	bool		DropCache = false;
 };
 
diff --git a/Engine/Source/Runtime/Experimental/IoStoreOnDemand/Private/IoStoreOnDemand.cpp b/Engine/Source/Runtime/Experimental/IoStoreOnDemand/Private/IoStoreOnDemand.cpp
--- a/Engine/Source/Runtime/Experimental/IoStoreOnDemand/Private/IoStoreOnDemand.cpp
+++ b/Engine/Source/Runtime/Experimental/IoStoreOnDemand/Private/IoStoreOnDemand.cpp
@@ -211,6 +211,9 @@ static FIasCacheConfig GetIasCacheConfig(const TCHAR* CommandLine)
 	GetConfigInt(TEXT("MemoryQuota"),			Ret.MemoryQuota);
 	GetConfigInt(TEXT("JournalQuota"),			Ret.JournalQuota);
 	GetConfigInt(TEXT("JournalFlushInterval"),	Ret.JournalFlushInterval);
+	GetConfigInt(TEXT("DemandThreshold"),		Ret.Demand.Threshold);
+	GetConfigInt(TEXT("DemandBoost"),			Ret.Demand.Boost);
+	GetConfigInt(TEXT("DemandSuperBoost"),		Ret.Demand.SuperBoost);
 #undef GetConfigInt
 
 #if !UE_BUILD_SHIPPING
