Separate out IAS' disk cache's journaling.

These is predominantly strategic to give a clearer separation upon which to manage data read/writes. Loading is a bit awkward as that requires multiple components to collaborate whereas. To avoid complexity intertwining that the parts that collaborate, loading's been moved out to a friendly function. LoadCache() leaves a lot to be desired frankly, but so be it.

#rnx

diff --git a/Engine/Source/Runtime/Experimental/IoStoreOnDemand/Private/JournaledCache.cpp b/Engine/Source/Runtime/Experimental/IoStoreOnDemand/Private/JournaledCache.cpp
--- a/Engine/Source/Runtime/Experimental/IoStoreOnDemand/Private/JournaledCache.cpp
+++ b/Engine/Source/Runtime/Experimental/IoStoreOnDemand/Private/JournaledCache.cpp
@@ -45,7 +45,9 @@ two writes per allowed op-throttle).
 namespace UE::IO::IAS::JournaledCache
 {
 
-// {{{1 structs ................................................................
+// {{{1 misc ...................................................................
+
+static bool LoadCache(class FDiskCache&);
 
 ////////////////////////////////////////////////////////////////////////////////
 struct FDebugCacheEntry
@@ -266,7 +268,7 @@ int32 FMemCache::Drop(uint32 Size)
 	return DropImpl(Size, [] (FItem&&) {});
 }
 
-// {{{1 disk-cache .............................................................
+// {{{1 phrase .................................................................
 
 ////////////////////////////////////////////////////////////////////////////////
 static const uint32 MAGIC = 0x04930001;
@@ -298,12 +300,14 @@ static_assert(sizeof(FPhraseDesc) == sizeof(FDataEntry));
 struct FDiskPhrase
 {
 						FDiskPhrase(TArray<FDataEntry>& InEntries, int32 InMaxEntries, uint32 DataSize);
-						~FDiskPhrase();
+						FDiskPhrase(FDiskPhrase&&) = default;
 	bool				Add(uint64 Key, FIoBuffer&& Data);
-	const uint8*		Finalize(MarkerType Marker, uint64 DataCursor) const;
-	int32				GetEntryCount() const	{ return Entries.Num() - Index - 1; }
-	uint32				GetSize() const			{ return Cursor; }
-	const FPhraseDesc&	GetDesc() const			{ return (FPhraseDesc&)(Entries[Index]); }
+	void				Drop()					{ return Entries.SetNumUninitialized(Index); }
+	const FDataEntry*	GetEntries() const		{ return Entries.GetData() + Index; }
+	int32				GetEntryCount() const	{ return Entries.Num() - Index; }
+	uint8*				GetPhraseData() const	{ return Buffer.Get(); }
+	uint32				GetPhraseSize() const	{ return Cursor; }
+	uint32				GetDataSize() const		{ return Cursor - sizeof(MarkerType); }
 
 private:
 	TUniquePtr<uint8[]>	Buffer;
@@ -315,6 +319,7 @@ private:
 private:
 						FDiskPhrase(const FDiskPhrase&) = delete;
 	FDiskPhrase&		operator = (const FDiskPhrase&) = delete;
+	FDiskPhrase&		operator = (FDiskPhrase&&) = delete;
 };
 
 ////////////////////////////////////////////////////////////////////////////////
@@ -325,17 +330,6 @@ FDiskPhrase::FDiskPhrase(TArray<FDataEntry>& InEntries, int32 InMaxEntries, uint
 {
 	DataSize += Cursor;
 	Buffer = TUniquePtr<uint8[]>(new uint8[DataSize]);
-
-	Entries.Add(FDataEntry{});
-}
-
-////////////////////////////////////////////////////////////////////////////////
-FDiskPhrase::~FDiskPhrase()
-{
-	if (GetEntryCount() == 0)
-	{
-		Entries.Pop();
-	}
 }
 
 ////////////////////////////////////////////////////////////////////////////////
@@ -350,29 +344,165 @@ bool FDiskPhrase::Add(uint64 Key, FIoBuffer&& Data)
 	return (--MaxEntries > 0);
 }
 
+
+
+// {{{1 journal ................................................................
+
+////////////////////////////////////////////////////////////////////////////////
+class FDiskJournal
+{
+public:
+							FDiskJournal(FStringView InRootPath, uint32 InMaxSize);
+	void					Reset();
+	void					Drop();
+	int32					Flush();
+	FDiskPhrase				OpenPhrase(uint32 DataSize);
+	void					ClosePhrase(FDiskPhrase&& Phrase, uint64 DataCursor);
+	uint32					GetMaxSize() const	{ return MaxSize; }
+	uint32					GetCursor() const	{ return Cursor; }
+	uint32					GetMarker() const	{ return Marker; }
+
+private:
+	friend bool				LoadCache(FDiskCache&);
+	void					GetPath(TStringBuilder<64>& Out);
+	void					OpenJrnFile();
+	TArray<FDataEntry>		Entries;
+	FStringView				RootPath;
+	MarkerType				Marker = 0;
+	TUniquePtr<IFileHandle>	JrnHandle;
+	uint32					Cursor = 0;
+	uint32					MaxSize;
+};
+
+////////////////////////////////////////////////////////////////////////////////
+FDiskJournal::FDiskJournal(FStringView InRootPath, uint32 InMaxSize)
+: RootPath(InRootPath)
+, MaxSize(InMaxSize)
+{
+	// Align down to keep to some assumptions
+	MaxSize &= ~(sizeof(FDataEntry) - 1);
+
+	OpenJrnFile();
+}
+
+////////////////////////////////////////////////////////////////////////////////
+void FDiskJournal::Reset()
+{
+	// Marker = 0; // We'll just lets this roll along in its own little world
+	Cursor = 0;
+	Entries.Reset();
+}
+
+////////////////////////////////////////////////////////////////////////////////
+void FDiskJournal::Drop()
+{
+	JrnHandle.Reset();
+
+	TStringBuilder<64> JrnPath;
+	GetPath(JrnPath);
+
+	IPlatformFile& Ipf = IPlatformFile::GetPlatformPhysical();
+	Ipf.DeleteFile(*JrnPath);
+
+	Reset();
+	OpenJrnFile();
+}
+
+////////////////////////////////////////////////////////////////////////////////
+void FDiskJournal::OpenJrnFile()
+{
+	TStringBuilder<64> JrnPath;
+	GetPath(JrnPath);
+
+	IPlatformFile& Ipf = IPlatformFile::GetPlatformPhysical();
+	IFileHandle* Handle = Ipf.OpenWrite(*JrnPath, true, true);
+	JrnHandle.Reset(Handle);
+}
+
+////////////////////////////////////////////////////////////////////////////////
+void FDiskJournal::GetPath(TStringBuilder<64>& Out)
+{
+	Out << RootPath;
+	Out << TEXT(".jrn");
+}
+
+////////////////////////////////////////////////////////////////////////////////
+FDiskPhrase FDiskJournal::OpenPhrase(uint32 DataSize)
+{
+	check((Cursor & (sizeof(FDataEntry) - 1)) == 0);
+
+	Entries.Add(FDataEntry{});
+	int32 MaxEntries = int32((MaxSize - Cursor) / sizeof(FDataEntry)) - 1;
+
+	FDiskPhrase Ret(Entries, MaxEntries, DataSize);
+
+	uint8* PhraseData = Ret.GetPhraseData();
+	std::memcpy(PhraseData, &Marker, sizeof(MarkerType));
+
+	return MoveTemp(Ret);
+}
+
 ////////////////////////////////////////////////////////////////////////////////
-const uint8* FDiskPhrase::Finalize(MarkerType Marker, uint64 DataCursor) const
+void FDiskJournal::ClosePhrase(FDiskPhrase&& Phrase, uint64 DataCursor)
 {
-	uint32 EntryCount = GetEntryCount();
+	uint32 EntryCount = Phrase.GetEntryCount();
+	if (EntryCount == 0)
+	{
+		Entries.Pop();
+		return;
+	}
+
 	Entries.Last().EntryCount = uint16(EntryCount);
 
-	auto& Desc = (FPhraseDesc&)(Entries[Index]);
+	auto& Desc = (FPhraseDesc&)(Phrase.GetEntries()[-1]);
 	Desc.Magic = MAGIC;
 	Desc.Marker = Marker;
-	Desc.DataCursor = DataCursor;
 	Desc.EntryCount = uint16(EntryCount);
+	Desc.DataCursor = DataCursor;
 
-	uint8* Ret = Buffer.Get();
-	std::memcpy(Ret, &Marker, sizeof(MarkerType));
-	return Ret;
+	++Marker;
 }
 
+////////////////////////////////////////////////////////////////////////////////
+int32 FDiskJournal::Flush()
+{
+	if (Entries.IsEmpty())
+	{
+		return 0;
+	}
 
+	TRACE_CPUPROFILER_EVENT_SCOPE(IasCache::Flush_DiskCache);
+
+	uint32 Size = Entries.Num() * sizeof(Entries[0]);
+
+	if (Cursor + Size > MaxSize)
+	{
+		Cursor = 0;
+	}
+
+	if (JrnHandle.IsValid())
+	{
+		TRACE_CPUPROFILER_EVENT_SCOPE(IasCache::JournalWrite);
+
+		JrnHandle->Seek(Cursor);
+		JrnHandle->Write((uint8*)(Entries.GetData()), Size);
+		JrnHandle.Reset();
+		Cursor += Size;
+	}
+
+	Entries.Reset();
+	OpenJrnFile();
+
+	return Size;
+}
+
+
+
+// {{{1 disk-cache .............................................................
 
 ////////////////////////////////////////////////////////////////////////////////
 class FDiskCache
 {
-private:
 public:
 							FDiskCache(FString&& Path, uint64 InMaxDataSize, uint32 InJournalSize);
 	void					Reset();
@@ -382,7 +512,6 @@ public:
 	bool					Materialize(EntryHandle Handle, FIoBuffer& Out, uint32 Offset=0) const;
 	int32					Flush();
 	void					Drop();
-	bool					Load();
 	uint32					DebugVisit(void* Param, FDebugCacheEntry::Callback* Callback);
 
 private:
@@ -394,72 +523,51 @@ private:
 	};
 	static_assert(sizeof(FMapEntry) == sizeof(uint64));
 
+	friend bool				LoadCache(FDiskCache&);
 	using					FDataMap = TMap<uint64, FMapEntry>;
-	void					OpenJrnFile();
 	void					Spam();
-	void					PhraseReset();
 	uint64					Insert(uint64 DataBase, const FDataEntry& Entry);
-	uint64					Insert(const FPhraseDesc& Phrase);
+	uint64					Insert(uint64 DataBase, const FDataEntry* Entries, uint32 EntryCount);
 	void					Prune(uint64 DataBase, uint32 Size);
-	TArray<FDataEntry>		Entries;
 	FString					BinPath;
 	FDataMap				DataMap;
-	TUniquePtr<IFileHandle>	JrnHandle;
 	uint64					MappedBytes;
+
 	uint64					MaxDataSize;
 	uint64					DataCursor;
-	uint32					JournalSize;
-	uint32					JournalCursor;
+
 	uint32					OverRemoval;
-	MarkerType				Marker = 0;
+
+	FDiskJournal			Journal;
 };
 
 ////////////////////////////////////////////////////////////////////////////////
 FDiskCache::FDiskCache(FString&& Path, uint64 InMaxDataSize, uint32 InJournalSize)
-: BinPath(Path)
+: BinPath(MoveTemp(Path))
 , MaxDataSize(InMaxDataSize)
-, JournalSize(InJournalSize)
+, Journal(BinPath, InJournalSize)
 {
 	// Align down to keep to some assumptions
-	JournalSize &= ~(sizeof(FDataEntry) - 1);
-	MaxDataSize = (MaxDataSize - JournalSize) & ~((1ull << 20) - 1);
+	InJournalSize = Journal.GetMaxSize();
+	MaxDataSize = (MaxDataSize - InJournalSize) & ~((1ull << 20) - 1);
 
 	Reset();
-	OpenJrnFile();
-}
-
-////////////////////////////////////////////////////////////////////////////////
-void FDiskCache::OpenJrnFile()
-{
-	TStringBuilder<265> JrnPath;
-	JrnPath << BinPath;
-	JrnPath << TEXT(".jrn");
-
-	IPlatformFile& Ipf = IPlatformFile::GetPlatformPhysical();
-	IFileHandle* Handle = Ipf.OpenWrite(*JrnPath, true, false);
-	JrnHandle.Reset(Handle);
 }
 
 ////////////////////////////////////////////////////////////////////////////////
 void FDiskCache::Reset()
 {
 	DataMap.Reset();
+	Journal.Reset();
 	MappedBytes = 0;
 	DataCursor = 0;
-	// Marker = 0; // We'll just lets this roll along in its own little world
-	JournalCursor = 0;
 	OverRemoval = 0;
-	PhraseReset();
 }
 
 ////////////////////////////////////////////////////////////////////////////////
 FDiskPhrase FDiskCache::OpenPhrase(uint32 DataSize)
 {
-	check((JournalCursor & (sizeof(FDataEntry) - 1)) == 0);
-
-	int32 MaxEntries = int32((JournalSize - JournalCursor) / sizeof(FDataEntry)) - 1;
-
-	return FDiskPhrase(Entries, MaxEntries, DataSize);
+	return Journal.OpenPhrase(DataSize);
 }
 
 ////////////////////////////////////////////////////////////////////////////////
@@ -468,33 +576,41 @@ void FDiskCache::ClosePhrase(FDiskPhrase&& Phrase)
 	int32 EntryCount = Phrase.GetEntryCount();
 	if (EntryCount <= 0)
 	{
+		Journal.ClosePhrase(MoveTemp(Phrase), 0);
 		return;
 	}
 
-	uint32 WriteSize = Phrase.GetSize();
+	TRACE_CPUPROFILER_EVENT_SCOPE(IasCache::ClosePhrase);
+
+	IPlatformFile& Ipf = IPlatformFile::GetPlatformPhysical();
+	TUniquePtr<IFileHandle> File(Ipf.OpenWrite(*BinPath, true, false));
+	if (!File.IsValid())
+	{
+		Phrase.Drop();
+		Journal.ClosePhrase(MoveTemp(Phrase), 0);
+		return;
+	}
+
+	uint32 WriteSize = Phrase.GetPhraseSize();
 	if (DataCursor + WriteSize > MaxDataSize)
 	{
 		OverRemoval = 0;
 		DataCursor = 0;
 	}
 
-	const uint8* Buffer = Phrase.Finalize(Marker, DataCursor);
-	++Marker;
-
-	IPlatformFile& Ipf = IPlatformFile::GetPlatformPhysical();
-	if (TUniquePtr<IFileHandle> File(Ipf.OpenWrite(*BinPath, true, false)); File.IsValid())
 	{
 		TRACE_CPUPROFILER_EVENT_SCOPE(IasCache::DataWrite);
-
+		const uint8* Buffer = Phrase.GetPhraseData();
 		File->Seek(DataCursor);
 		File->Write(Buffer, WriteSize);
 		File.Reset();
+	}
 
-		Prune(DataCursor, WriteSize);
-		Insert(Phrase.GetDesc());
+	Prune(DataCursor, WriteSize);
+	Insert(DataCursor, Phrase.GetEntries(), EntryCount);
 
-		DataCursor += WriteSize;
-	}
+	Journal.ClosePhrase(MoveTemp(Phrase), DataCursor);
+	DataCursor += WriteSize;
 }
 
 ////////////////////////////////////////////////////////////////////////////////
@@ -542,13 +658,12 @@ uint64 FDiskCache::Insert(uint64 DataBase, const FDataEntry& Entry)
 }
 
 ////////////////////////////////////////////////////////////////////////////////
-uint64 FDiskCache::Insert(const FPhraseDesc& Phrase)
+uint64 FDiskCache::Insert(uint64 DataBase, const FDataEntry* Entries, uint32 EntryCount)
 {
 	uint64 TotalSize = 0;
-	uint64 DataBase = Phrase.DataCursor;
-	for (uint32 i = 0; i < Phrase.EntryCount; ++i)
+	for (uint32 i = 0; i < EntryCount; ++i)
 	{
-		TotalSize += Insert(DataBase, Phrase.Entries[i]);
+		TotalSize += Insert(DataBase, Entries[i]);
 	}
 
 	MappedBytes += TotalSize;
@@ -604,35 +719,9 @@ void FDiskCache::Prune(uint64 DataBase, uint32 Size)
 ////////////////////////////////////////////////////////////////////////////////
 int32 FDiskCache::Flush()
 {
-	if (Entries.IsEmpty())
-	{
-		return 0;
-	}
-
-	TRACE_CPUPROFILER_EVENT_SCOPE(IasCache::Flush_DiskCache);
-
-	uint32 Size = Entries.Num() * sizeof(Entries[0]);
-
-	if (JournalCursor + Size > JournalSize)
-	{
-		JournalCursor = 0;
-	}
-
-	if (JrnHandle.IsValid())
-	{
-		TRACE_CPUPROFILER_EVENT_SCOPE(IasCache::JournalWrite);
-
-		JrnHandle->Seek(JournalCursor);
-		JrnHandle->Write((uint8*)(Entries.GetData()), Size);
-		JrnHandle.Reset();
-		JournalCursor += Size;
-	}
-
-	PhraseReset();
-	OpenJrnFile();
-
+	int32 Ret = Journal.Flush();
 	Spam();
-	return Size;
+	return Ret;
 }
 
 ////////////////////////////////////////////////////////////////////////////////
@@ -640,26 +729,64 @@ void FDiskCache::Drop()
 {
 	IPlatformFile& Ipf = IPlatformFile::GetPlatformPhysical();
 	Ipf.DeleteFile(*BinPath);
-	Ipf.DeleteFile(*(BinPath + TEXT(".jrn")));
+
+	Journal.Drop();
+
 	Reset();
 }
 
 ////////////////////////////////////////////////////////////////////////////////
-bool FDiskCache::Load()
+void FDiskCache::Spam()
 {
-	IPlatformFile& Ipf = IPlatformFile::GetPlatformPhysical();
+	UE_LOG(LogIas, VeryVerbose,
+		TEXT("JournaledCache: MappedKiB=%llu Entries=%d DataCur=%llu JournalCur=%u Marker=%u)"),
+		(MappedBytes >> 10),
+		DataMap.Num(),
+		DataCursor,
+		Journal.GetCursor(),
+		Journal.GetMarker()
+	);
+}
+
+////////////////////////////////////////////////////////////////////////////////
+uint32 FDiskCache::DebugVisit(void* Param, FDebugCacheEntry::Callback* Callback)
+{
+	FDebugCacheEntry Out = {};
+	for (const auto& Entry : DataMap)
+	{
+		Out.Key = Entry.Key;
+		Out.Size = Entry.Value.Size;
+		Callback(Param, Out);
+	}
+	return DataMap.Num();
+}
+
+
+
+// {{{1 loader .................................................................
+
+////////////////////////////////////////////////////////////////////////////////
+static bool LoadCache(FDiskCache& DiskCache)
+{
+	FDiskJournal& Journal = DiskCache.Journal;
 
 	uint32 DataSize = 0;
 	TUniquePtr<uint8[]> Data;
 
-	TStringBuilder<256> JrnPath;
-	JrnPath << BinPath;
-	JrnPath << TEXT(".jrn");
-	if (TUniquePtr<IFileHandle> File(Ipf.OpenRead(*JrnPath, false)); File.IsValid())
+	if (auto& Handle = Journal.JrnHandle; Handle.IsValid())
 	{
-		DataSize = uint32(File->Size());
+		DataSize = uint32(Handle->Size());
+		if (DataSize == 0)
+		{
+			return false;
+		}
+
 		Data = TUniquePtr<uint8[]>(new uint8[DataSize]);
-		File->Read(Data.Get(), DataSize);
+		Handle->Seek(0);
+		if (!Handle->Read(Data.Get(), DataSize))
+		{
+			UE_LOG(LogIas, VeryVerbose, TEXT("JournaledCache: failed reading journal"));
+		}
 	}
 
 	if (DataSize == 0)
@@ -672,7 +799,7 @@ bool FDiskCache::Load()
 	struct FParagraph
 	{
 		const FPhraseDesc*	Phrase;
-		uint32				Marker;
+		MarkerType			Marker;
 		uint32				DataSize;
 	};
 
@@ -681,9 +808,7 @@ bool FDiskCache::Load()
 		return (UPTRINT(Address) - UPTRINT(Data.Get())) > DataSize;
 	};
 
-	auto ReadPhrases = [&IsOob] (
-		const uint8* Cursor,
-		FParagraph& Out) -> const uint8*
+	auto ReadPhrases = [&IsOob] (const uint8* Cursor, FParagraph& Out) -> const uint8*
 	{
 		// Only proceed if we can read at least three integers
 		if (IsOob(Cursor + sizeof(FPhraseDesc)))
@@ -730,6 +855,7 @@ bool FDiskCache::Load()
 		{
 			break;
 		}
+
 		Paragraphs.Add(Paragraph);
 	}
 
@@ -766,8 +892,8 @@ bool FDiskCache::Load()
 		const FParagraph& Lhs,
 		const FParagraph& Rhs)
 	{
-        uint32 L = Lhs.Marker;
-        uint32 R = Rhs.Marker;
+        MarkerType L = Lhs.Marker;
+        MarkerType R = Rhs.Marker;
         enum : uint32 { LowQuarter = 1u << 30, HighQuarter = 3u << 30 };
         int32 Wrap = (L < LowQuarter) & (R >= HighQuarter);
         Wrap |= (R < LowQuarter) & (L >= HighQuarter);
@@ -777,7 +903,7 @@ bool FDiskCache::Load()
 
 	// Eliminate any discontinuities and find where data wrapped
 	int32 BasisIndex = 0;
-	int64 Remaining = MaxDataSize;
+	int64 Remaining = DiskCache.MaxDataSize;
 	for (int32 i = Paragraphs.Num() - 2; i >= 0; BasisIndex = i--)
 	{
 		const FParagraph& Newer = Paragraphs[i + 1];
@@ -796,27 +922,31 @@ bool FDiskCache::Load()
 		}
 	}
 
-	TUniquePtr<IFileHandle> File(Ipf.OpenRead(*BinPath, false));
+	IPlatformFile& Ipf = IPlatformFile::GetPlatformPhysical();
+
+	const TCHAR* BinPath = *(DiskCache.BinPath);
+	TUniquePtr<IFileHandle> File(Ipf.OpenRead(BinPath, false));
 	if (!File.IsValid())
 	{
-		UE_LOG(LogIas, VeryVerbose, TEXT("JournaledCache: unable to open '%s'"), *BinPath);
+		UE_LOG(LogIas, VeryVerbose, TEXT("JournaledCache: unable to open '%s'"), BinPath);
 		return false;
 	}
 
-	if (uint64(File->Size()) > MaxDataSize)
+	if (uint64(File->Size()) > DiskCache.MaxDataSize)
 	{
-		UE_LOG(LogIas, VeryVerbose, TEXT("JournaledCache: Dropping - existing cache to bi; %llu/%llu"), uint64(File->Size()), MaxDataSize);
-		Drop();
+		UE_LOG(LogIas, VeryVerbose,
+			TEXT("JournaledCache: Dropping - existing cache to bi; %llu/%llu"),
+			uint64(File->Size()), DiskCache.MaxDataSize
+		);
+		DiskCache.Drop();
 		return false;
 	}
 
 	// Detect data writes that are newer than any journal flushes.
-	auto ReadBinMarker = [File = File.Get()] (int64 Cursor)
+	auto ReadBinMarker = [File = File.Get()] (int64 Cursor, uint32& Out)
 	{
-		uint32 Value;
 		File->Seek(Cursor);
-		File->Read((uint8*)(&Value), sizeof(Value));
-		return Value;
+		return File->Read((uint8*)(&Out), sizeof(Out));
 	};
 
 	for (; BasisIndex < Paragraphs.Num(); ++BasisIndex)
@@ -824,13 +954,13 @@ bool FDiskCache::Load()
 		const FParagraph& Stock = Paragraphs[BasisIndex];
 
 		uint64 DataBase = Stock.Phrase->DataCursor;
-		if (DataBase + Stock.DataSize > MaxDataSize)
+		if (DataBase + Stock.DataSize > DiskCache.MaxDataSize)
 		{
 			DataBase = 0;
 		}
 
-		uint32 DataMark = ReadBinMarker(DataBase);
-		if (DataMark == Stock.Marker)
+		uint32 DataMark;
+		if (ReadBinMarker(DataBase, DataMark) && (DataMark == Stock.Marker))
 		{
 			break;
 		}
@@ -839,68 +969,45 @@ bool FDiskCache::Load()
 	// Add known entries into the tree.
 	for (uint32 i = BasisIndex, n = Paragraphs.Num(); i < n; ++i)
 	{
-		const FParagraph& Holm = Paragraphs[i];
-		Insert(Holm.Phrase[0]);
+		const FPhraseDesc& Holm = Paragraphs[i].Phrase[0];
+		DiskCache.Insert(Holm.DataCursor, Holm.Entries, Holm.EntryCount);
 	}
 	
-	// Prime the disk-cache's state
+	// Prime the journal's state
 	const FParagraph& LastPara = Paragraphs.Last();
 	const FPhraseDesc* LastPhrase = LastPara.Phrase;
-	Marker = LastPara.Marker + 1;
+	Journal.Marker = LastPara.Marker + 1;
 
-	if (DataSize <= JournalSize)
+	if (DataSize <= Journal.MaxSize)
 	{
-		JournalCursor = uint32(UPTRINT(LastPhrase + LastPhrase->EntryCount + 1) - UPTRINT(Data.Get()));
+		Journal.Cursor = uint32(UPTRINT(LastPhrase + LastPhrase->EntryCount + 1) - UPTRINT(Data.Get()));
 	}
 	else
 	{
-		Ipf.DeleteFile(*JrnPath);
+		UE_LOG(LogIas, VeryVerbose,
+			TEXT("JournaledCache: Journal exceeds given size - dropping; %u/%u"),
+			DataSize, Journal.MaxSize
+		);
+		Journal.Drop();
+		return false;
 	}
 
-	DataCursor = LastPhrase->DataCursor + LastPara.DataSize;
-	if (DataCursor > MaxDataSize)
+	// Prime the disk-cache's state
+	DiskCache.DataCursor = LastPhrase->DataCursor + LastPara.DataSize;
+	if (DiskCache.DataCursor > DiskCache.MaxDataSize)
 	{
-		UE_LOG(LogIas, VeryVerbose, TEXT("JournaledCache: Dropping - DataCursor too big; %llu/%llu"), DataCursor, MaxDataSize);
-		Drop();
+		UE_LOG(LogIas, VeryVerbose,
+			TEXT("JournaledCache: Dropping - DataCursor too big; %llu/%llu"),
+			DiskCache.DataCursor, DiskCache.MaxDataSize
+		);
+		DiskCache.Drop();
 		return false;
 	}
 
-	Spam();
+	DiskCache.Spam();
 	return true;
 }
 
-////////////////////////////////////////////////////////////////////////////////
-void FDiskCache::Spam()
-{
-	UE_LOG(LogIas, VeryVerbose,
-		TEXT("JournaledCache: MappedKiB=%llu Entries=%d DataCur=%llu JournalCur=%u Marker=%u)"),
-		(MappedBytes >> 10),
-		DataMap.Num(),
-		DataCursor,
-		JournalCursor,
-		Marker
-	);
-}
-
-////////////////////////////////////////////////////////////////////////////////
-uint32 FDiskCache::DebugVisit(void* Param, FDebugCacheEntry::Callback* Callback)
-{
-	FDebugCacheEntry Out = {};
-	for (const auto& Entry : DataMap)
-	{
-		Out.Key = Entry.Key;
-		Out.Size = Entry.Value.Size;
-		Callback(Param, Out);
-	}
-	return DataMap.Num();
-}
-
-////////////////////////////////////////////////////////////////////////////////
-void FDiskCache::PhraseReset()
-{
-	Entries.Reset();
-}
-
 // {{{1 cache ..................................................................
 
 ////////////////////////////////////////////////////////////////////////////////
@@ -998,7 +1105,7 @@ void FCache::Reset()
 bool FCache::Load()
 {
 	FWriteAccess _(Lock);
-	return DiskCache.Load();
+	return LoadCache(DiskCache);
 }
 
 ////////////////////////////////////////////////////////////////////////////////
@@ -1054,9 +1161,7 @@ int32 FCache::Flush(int32 Allowance)
 	TRACE_CPUPROFILER_EVENT_SCOPE(IasCache::Flush_MemCache);
 
 	FMemCache::PeelItems PeelItems;
-
 	FWriteAccess _(Lock);
-
 	int32 MemCacheSize = MemCache.Peel(Allowance, PeelItems);
 
 	FDiskPhrase Phrase = DiskCache.OpenPhrase(MemCacheSize);
@@ -1472,8 +1577,8 @@ IOSTOREONDEMAND_API void Tests()
 			check(Validate() == i);
 
 			Cache.Reset();
-			Cache.Load();
-			check(Validate() <= i);
+			check(Cache.Load());
+			check(Validate() == 0); // not enough flushes to write a journal
 		}
 		Cache.Reset();
 
@@ -1485,11 +1590,15 @@ IOSTOREONDEMAND_API void Tests()
 				PrimePuts(WriteAllowance);
 				Cache.Flush(WriteAllowance);
 			}
-			uint32 Count = Validate();
+			uint32 PreCount = Validate();
 
 			Cache.Reset();
-			Cache.Load();
-			check(Validate() <= Count);
+			check(Cache.Load());
+
+			uint32 PostCount = Validate();
+			check(!PostCount == !(i / 4)); // JournalFlushInterval
+			check(PostCount <= PreCount);
+			check((PostCount == PreCount) == ((i & 3) == 0));
 		}
 		Cache.Reset();
 
