Minor edit to keep it help me remember that mem-cache's internal items and peeled items are the same underlying type.

#rnx

diff --git a/Engine/Source/Runtime/Experimental/IoStoreOnDemand/Private/JournaledCache.cpp b/Engine/Source/Runtime/Experimental/IoStoreOnDemand/Private/JournaledCache.cpp
--- a/Engine/Source/Runtime/Experimental/IoStoreOnDemand/Private/JournaledCache.cpp
+++ b/Engine/Source/Runtime/Experimental/IoStoreOnDemand/Private/JournaledCache.cpp
@@ -105,7 +105,8 @@ public:
 		FIoBuffer	Data;
 	};
 
-	using PeelItems = TArray<FItem>;
+	using ItemArray = TArray<FItem>;
+	using PeelItems = ItemArray;
 
 					FMemCache(uint32 InMaxSize=64 << 10);
 	void			Reset();
@@ -124,7 +125,7 @@ private:
 	int32			Drop(uint32 Size);
 	uint32			MaxSize;
 	uint32			UsedSize = 0;
-	TArray<FItem>	Items;
+	ItemArray		Items;
 };
 
 ////////////////////////////////////////////////////////////////////////////////
