If the whole mem-cache can be persisted to disk we can just transfer ownership directly.

#rnx

diff --git a/Engine/Source/Runtime/Experimental/IoStoreOnDemand/Private/JournaledCache.cpp b/Engine/Source/Runtime/Experimental/IoStoreOnDemand/Private/JournaledCache.cpp
--- a/Engine/Source/Runtime/Experimental/IoStoreOnDemand/Private/JournaledCache.cpp
+++ b/Engine/Source/Runtime/Experimental/IoStoreOnDemand/Private/JournaledCache.cpp
@@ -215,9 +215,17 @@ bool FMemCache::Put(uint64 Key, FIoBuffer&& Data)
 ////////////////////////////////////////////////////////////////////////////////
 int32 FMemCache::Peel(int32 PeelSize, PeelItems& Out)
 {
-	return DropImpl(PeelSize, [&Out] (FItem&& Item) {
-		Out.Add(MoveTemp(Item));
-	});
+	if (PeelSize < int32(UsedSize))
+	{
+		return DropImpl(PeelSize, [&Out] (FItem&& It) { Out.Add(MoveTemp(It)); });
+	}
+
+	Out = MoveTemp(Items);
+
+	int32 Ret = UsedSize;
+	UsedSize = 0;
+	FOnDemandIoBackendStats::Get()->OnCachePendingBytes(0);
+	return Ret;
 }
 
 ////////////////////////////////////////////////////////////////////////////////
