Changed peeling of IAS' in-memory cache to favour large items first and then try and fill the remaining space with smaller items. The hope is that this will give better utilisation.

#rnx

diff --git a/Engine/Source/Runtime/Experimental/IoStoreOnDemand/Private/JournaledCache.cpp b/Engine/Source/Runtime/Experimental/IoStoreOnDemand/Private/JournaledCache.cpp
--- a/Engine/Source/Runtime/Experimental/IoStoreOnDemand/Private/JournaledCache.cpp
+++ b/Engine/Source/Runtime/Experimental/IoStoreOnDemand/Private/JournaledCache.cpp
@@ -215,17 +215,54 @@ bool FMemCache::Put(uint64 Key, FIoBuffer&& Data)
 ////////////////////////////////////////////////////////////////////////////////
 int32 FMemCache::Peel(int32 PeelSize, PeelItems& Out)
 {
-	if (PeelSize < int32(UsedSize))
+	TRACE_CPUPROFILER_EVENT_SCOPE(IasCache::Peel);
+
+	Algo::Sort(Items, [] (const FItem& Lhs, const FItem& Rhs) {
+		return Lhs.Data.GetSize() < Rhs.Data.GetSize();
+	});
+
+	// Add large items
+	int32 NumItems = Items.Num();
+	int32 DropSize = 0;
+	while (NumItems > 0)
 	{
-		return DropImpl(PeelSize, [&Out] (FItem&& It) { Out.Add(MoveTemp(It)); });
+		FItem& Item = Items[NumItems - 1];
+
+		int32 NextDropSize = DropSize + int32(Item.Data.GetSize());
+		if (NextDropSize > PeelSize)
+		{
+			break;
+		}
+
+		Out.Add(MoveTemp(Item));
+		DropSize = NextDropSize;
+		--NumItems;
 	}
 
-	Out = MoveTemp(Items);
+	// Fill remaining space with small items
+	for (int32 i = 0; i < NumItems; ++i)
+	{
+		FItem& Item = Items[i];
 
-	int32 Ret = UsedSize;
-	UsedSize = 0;
-	FOnDemandIoBackendStats::Get()->OnCachePendingBytes(0);
-	return Ret;
+		int32 NextDropSize = DropSize + int32(Item.Data.GetSize());
+		if (NextDropSize > PeelSize)
+		{
+			break;
+		}
+
+		Out.Add(MoveTemp(Item));
+		Items.Swap(i, NumItems - 1);
+		--NumItems;
+
+		DropSize = NextDropSize;
+	}
+
+	Items.SetNum(NumItems);
+
+	UsedSize -= DropSize;
+	FOnDemandIoBackendStats::Get()->OnCachePendingBytes(UsedSize);
+
+	return DropSize;
 }
 
 ////////////////////////////////////////////////////////////////////////////////
@@ -258,7 +295,6 @@ int32 FMemCache::DropImpl(uint32 Size, Lambda&& Callback)
 		uint32 Index = n ? (Size * 0x0'a9e0'493) % n : 0;
 
 		Size = uint32(Items[Index].Data.GetSize());
-		UsedSize -= Size;
 		DropSize += Size;
 
 		Callback(MoveTemp(Items[Index]));
@@ -267,7 +303,9 @@ int32 FMemCache::DropImpl(uint32 Size, Lambda&& Callback)
 		Items.Pop();
 	}
 
+	UsedSize -= DropSize;
 	FOnDemandIoBackendStats::Get()->OnCachePendingBytes(UsedSize);
+
 	return DropSize;
 }
 
