Track the demand on IAS' in-memory cache as a percentage.

#rnx

diff --git a/Engine/Source/Runtime/Experimental/IoStoreOnDemand/Private/JournaledCache.cpp b/Engine/Source/Runtime/Experimental/IoStoreOnDemand/Private/JournaledCache.cpp
--- a/Engine/Source/Runtime/Experimental/IoStoreOnDemand/Private/JournaledCache.cpp
+++ b/Engine/Source/Runtime/Experimental/IoStoreOnDemand/Private/JournaledCache.cpp
@@ -110,6 +110,7 @@ public:
 
 					FMemCache(uint32 InMaxSize=64 << 10);
 	void			Reset();
+	uint32			GetDemand() const;
 	uint32			GetCount() const	{ return Items.Num(); }
 	uint32			GetUsed() const		{ return UsedSize; }
 	uint32			GetMax() const		{ return MaxSize; }
@@ -141,6 +142,12 @@ void FMemCache::Reset()
 	UsedSize = 0;
 }
 
+////////////////////////////////////////////////////////////////////////////////
+uint32 FMemCache::GetDemand() const
+{
+	return (GetUsed() * 100) / MaxSize;
+}
+
 ////////////////////////////////////////////////////////////////////////////////
 EntryHandle FMemCache::Get(uint64 Key) const
 {
@@ -1095,6 +1102,7 @@ public:
 					FCache(FConfig&& Config);
 	void			Reset();
 	bool			Load();
+	uint32			GetDemand() const;
 	FEntry			Get(uint64 Key) const;
 	bool			Put(uint64 Key, FIoBuffer& Data);
 	int32			Flush(int32 Allowance);
@@ -1104,6 +1112,7 @@ private:
 	mutable FRWLock	Lock;
 	FMemCache		MemCache;
 	FDiskCache		DiskCache;
+	std::atomic_int	Demand;
 	uint32			FlushIndex = 0;
 	uint32			FlushPeriod = 4;
 };
@@ -1155,6 +1164,12 @@ bool FCache::Load()
 	return LoadCache(DiskCache);
 }
 
+////////////////////////////////////////////////////////////////////////////////
+uint32 FCache::GetDemand() const
+{
+	return Demand.load(std::memory_order_relaxed);
+}
+
 ////////////////////////////////////////////////////////////////////////////////
 FCache::FEntry FCache::Get(uint64 Key) const
 {
@@ -1181,7 +1196,13 @@ bool FCache::Put(uint64 Key, FIoBuffer& Data)
 {
 	FIoBuffer Cloned = Data;
 	FWriteAccess _(Lock);
-	return MemCache.Put(Key, MoveTemp(Cloned));
+	bool Ok = MemCache.Put(Key, MoveTemp(Cloned));
+	if (Ok)
+	{
+		uint32 NewDemand = MemCache.GetDemand();
+		Demand.store(NewDemand, std::memory_order_relaxed);
+	}
+	return Ok;
 }
 
 ////////////////////////////////////////////////////////////////////////////////
@@ -1235,6 +1256,9 @@ int32 FCache::Flush(int32 Allowance)
 		Allowance -= DiskCache.Flush();
 	}
 
+	uint32 NewDemand = MemCache.GetDemand();
+	Demand.store(NewDemand, std::memory_order_relaxed);
+
 	return Allowance - MemCacheSize;
 }
 
