Compile out tests for now to reduce noise.

#rnx

diff --git a/Engine/Source/Runtime/Experimental/IoStoreOnDemand/Private/JournaledCache.cpp b/Engine/Source/Runtime/Experimental/IoStoreOnDemand/Private/JournaledCache.cpp
--- a/Engine/Source/Runtime/Experimental/IoStoreOnDemand/Private/JournaledCache.cpp
+++ b/Engine/Source/Runtime/Experimental/IoStoreOnDemand/Private/JournaledCache.cpp
@@ -1487,6 +1487,7 @@ static uint64 KeyGen(const FIoBuffer& Data)
 ////////////////////////////////////////////////////////////////////////////////
 IOSTOREONDEMAND_API void Tests()
 {
+#if 0
 	using namespace JournaledCache;
 
 	// Some randomness
@@ -1704,6 +1705,7 @@ IOSTOREONDEMAND_API void Tests()
 
 		// don't load-and-sort so many paragraphs (only need max-data size)
 	} // }}}
+#endif // 0
 }
 
 } // namespace IasJournaledFileCacheTest
