Moved make() function up in the file so it's next to the class it creates.

#rnx

diff --git a/Engine/Source/Runtime/Experimental/IoStoreOnDemand/Private/JournaledCache.cpp b/Engine/Source/Runtime/Experimental/IoStoreOnDemand/Private/JournaledCache.cpp
--- a/Engine/Source/Runtime/Experimental/IoStoreOnDemand/Private/JournaledCache.cpp
+++ b/Engine/Source/Runtime/Experimental/IoStoreOnDemand/Private/JournaledCache.cpp
@@ -1456,7 +1456,14 @@ uint64 FJournaledCache::ReduceKey(const FIoHash& Key)
 	return (Ret[0] + Ret[2]) ^ Ret[1];
 }
 
-// }}}
+
+
+////////////////////////////////////////////////////////////////////////////////
+TUniquePtr<IIasCache> MakeIasCache(const FIasCacheConfig& Config)
+{
+	LLM_SCOPE_BYTAG(Ias);
+	return MakeUnique<FJournaledCache>(Config);
+}
 
 
 
@@ -1714,13 +1721,6 @@ IOSTOREONDEMAND_API void Tests()
 
 // }}}
 
-////////////////////////////////////////////////////////////////////////////////
-TUniquePtr<IIasCache> MakeIasCache(const FIasCacheConfig& Config)
-{
-	LLM_SCOPE_BYTAG(Ias);
-	return MakeUnique<FJournaledCache>(Config);
-}
-
 } // namespace UE::IO::IAS
 
 /* vim: set noet foldlevel=1 : */
