Clean up IAS analytics

#rnx

- We now strip the http/https prefix from the CDN url as some analytic visualization systems have trouble parsing it. In addition the value was renamed from "IasCDNBackend" to "IasCDNUrl" as this a) makes the name clearer b) makes it easier to differentiate between this and the older full url.
- Removed the http download rate statistic as we found no real use for it.

diff --git a/Engine/Source/Runtime/Experimental/IoStoreOnDemand/Private/OnDemandIoDispatcherBackend.cpp b/Engine/Source/Runtime/Experimental/IoStoreOnDemand/Private/OnDemandIoDispatcherBackend.cpp
--- a/Engine/Source/Runtime/Experimental/IoStoreOnDemand/Private/OnDemandIoDispatcherBackend.cpp
+++ b/Engine/Source/Runtime/Experimental/IoStoreOnDemand/Private/OnDemandIoDispatcherBackend.cpp
@@ -1237,7 +1237,15 @@ void FOnDemandIoBackend::ReportAnalytics(TArray<FAnalyticsEventAttribute>& OutAn
 
 	if (HttpClient.IsValid())
 	{
-		AppendAnalyticsEventAttributeArray(OutAnalyticsArray, TEXT("IasCDNBackend"), HttpClient->ServiceUrl());
+		FString CdnUrl = HttpClient->ServiceUrl();
+
+		// Strip the prefix from the url as some analytics systems may have trouble dealing with it
+		if (!CdnUrl.RemoveFromStart(TEXT("http://")))
+		{
+			CdnUrl.RemoveFromStart(TEXT("https://"));
+		}
+
+		AppendAnalyticsEventAttributeArray(OutAnalyticsArray, TEXT("IasCdnUrl"), MoveTemp(CdnUrl));
 
 		Stats.ReportAnalytics(OutAnalyticsArray);
 	}
diff --git a/Engine/Source/Runtime/Experimental/IoStoreOnDemand/Private/Statistics.cpp b/Engine/Source/Runtime/Experimental/IoStoreOnDemand/Private/Statistics.cpp
--- a/Engine/Source/Runtime/Experimental/IoStoreOnDemand/Private/Statistics.cpp
+++ b/Engine/Source/Runtime/Experimental/IoStoreOnDemand/Private/Statistics.cpp
@@ -213,7 +213,6 @@ int64					GHttpHistoryTotalDuration = 0;
 int64					GHttpHistoryTotalBytes = 0;
 int64 					GHttpHistoryIndex = 0;
 FIncrementalVariance	GHttpAvgDuration; // Duration of the http requests, in milliseconds
-FIncrementalVariance	GHttpAvgRate; // The download rate of the http requests, in MiB/s
 
 ////////////////////////////////////////////////////////////////////////////////
 // 
@@ -317,8 +316,6 @@ void FOnDemandIoBackendStats::ReportAnalytics(TArray<FAnalyticsEventAttribute>&
 			TEXT("IasHttpDownloadedBytes"), GHttpDownloadedBytes.Get(),
 			TEXT("IasHttpDurationMeanAvg"), GHttpAvgDuration.GetMean(),
 			TEXT("IasHttpDurationStdDev"), GHttpAvgDuration.GetDeviation(),
-			TEXT("IasHttpRateMeanAvg"), GHttpAvgRate.GetMean(),
-			TEXT("IasHttpRateStdDev"), GHttpAvgRate.GetDeviation(),
 
 			TEXT("IasHttpDuration0"), GHttpDurationBuckets[0],
 			TEXT("IasHttpDuration1"), GHttpDurationBuckets[1],
@@ -449,11 +446,6 @@ void FOnDemandIoBackendStats::OnHttpGet(uint64 SizeBytes, uint64 DurationMs)
 
 	GHttpAvgDuration.Increment(static_cast<double>(DurationMs));
 
-	const double SizeMiB = double(SizeBytes) / (1024.0 * 1024.0);
-	const double DurationSeconds = double(DurationMs) / 1000.0;
-
-	GHttpAvgRate.Increment(SizeMiB / DurationSeconds);
-
 	GHttpDurationBuckets[FindDurationBucket(DurationMs)]++;
 }
 
