Split Flush() into two methods as it does two separate things. Fixed check() if we try to read more than the data file has available.

#rnx

diff --git a/Engine/Source/Runtime/Experimental/IoStoreOnDemand/Private/JournaledCache.cpp b/Engine/Source/Runtime/Experimental/IoStoreOnDemand/Private/JournaledCache.cpp
--- a/Engine/Source/Runtime/Experimental/IoStoreOnDemand/Private/JournaledCache.cpp
+++ b/Engine/Source/Runtime/Experimental/IoStoreOnDemand/Private/JournaledCache.cpp
@@ -997,8 +997,13 @@ static bool LoadCache(FDiskCache& DiskCache)
 	}
 
 	// Detect data writes that are newer than any journal flushes.
-	auto ReadBinMarker = [File = File.Get()] (int64 Cursor, uint32& Out)
+	auto ReadBinMarker = [File = File.Get()] (uint64 Cursor, uint32& Out)
 	{
+		if (Cursor + sizeof(Out) > uint64(File->Size()))
+		{
+			return false;
+		}
+
 		File->Seek(Cursor);
 		return File->Read((uint8*)(&Out), sizeof(Out));
 	};
@@ -1106,6 +1111,7 @@ public:
 	FEntry			Get(uint64 Key) const;
 	bool			Put(uint64 Key, FIoBuffer& Data);
 	int32			Flush(int32 Allowance);
+	uint32			WriteMemToDisk(int32 Allowance);
 	uint32			DebugVisit(void* Param, FDebugCacheEntry::Callback* Callback);
 
 private:
@@ -1226,6 +1232,12 @@ int32 FCache::Flush(int32 Allowance)
 		}
 	}
 
+	return WriteMemToDisk(Allowance);
+}
+
+////////////////////////////////////////////////////////////////////////////////
+uint32 FCache::WriteMemToDisk(int32 Allowance)
+{
 	TRACE_CPUPROFILER_EVENT_SCOPE(IasCache::Flush_MemCache);
 
 	FMemCache::PeelItems PeelItems;
