IAS - fix bug accessing I/O dispatche request when chunk request is cancelled

#rnx

diff --git a/Engine/Source/Runtime/Experimental/IoStoreOnDemand/Private/OnDemandIoDispatcherBackend.cpp b/Engine/Source/Runtime/Experimental/IoStoreOnDemand/Private/OnDemandIoDispatcherBackend.cpp
--- a/Engine/Source/Runtime/Experimental/IoStoreOnDemand/Private/OnDemandIoDispatcherBackend.cpp
+++ b/Engine/Source/Runtime/Experimental/IoStoreOnDemand/Private/OnDemandIoDispatcherBackend.cpp
@@ -1032,10 +1032,9 @@ bool FOnDemandIoBackend::Resolve(FIoRequestImpl* Request)
 			}
 		}
 
-		const bool bHttpDisabled = IsHttpEnabled() == false || IsHttpEnabled(ChunkRequest->GetChunkId()) == false;
-		if (bHttpDisabled || ChunkRequest->CancellationToken.IsCancelled())
+		if (ChunkRequest->CancellationToken.IsCancelled() || IsHttpEnabled() == false)
 		{
-			UE_CLOG(bHttpDisabled, LogIas, Log, TEXT("Chunk was not found in the cache and HTTP is disabled"));
+			UE_CLOG(IsHttpEnabled() == false, LogIas, Log, TEXT("Chunk was not found in the cache and HTTP is disabled"));
 			return CompleteRequest(ChunkRequest);
 		}
 
