Addition of a boolean to control navigation relevancy

diff --git a/Engine/Source/Runtime/Experimental/GeometryCollectionEngine/Private/GeometryCollection/GeometryCollectionComponent.cpp b/Engine/Source/Runtime/Experimental/GeometryCollectionEngine/Private/GeometryCollection/GeometryCollectionComponent.cpp
--- a/Engine/Source/Runtime/Experimental/GeometryCollectionEngine/Private/GeometryCollection/GeometryCollectionComponent.cpp
+++ b/Engine/Source/Runtime/Experimental/GeometryCollectionEngine/Private/GeometryCollection/GeometryCollectionComponent.cpp
@@ -1427,6 +1427,10 @@ void UGeometryCollectionComponent::GetResourceSizeEx(FResourceSizeEx& Cumulative
 	CumulativeResourceSize.AddDedicatedSystemMemoryBytes(SizeBytes);
 }
 
+bool UGeometryCollectionComponent::IsNavigationRelevant() const
+{
+	return bIsCurrentlyNavigationRelevant && Super::IsNavigationRelevant();
+}
 
 #if WITH_EDITOR
 FDelegateHandle UGeometryCollectionComponent::RegisterOnGeometryCollectionPropertyChanged(const FOnGeometryCollectionPropertyChanged& Delegate)
diff --git a/Engine/Source/Runtime/Experimental/GeometryCollectionEngine/Public/GeometryCollection/GeometryCollectionComponent.h b/Engine/Source/Runtime/Experimental/GeometryCollectionEngine/Public/GeometryCollection/GeometryCollectionComponent.h
--- a/Engine/Source/Runtime/Experimental/GeometryCollectionEngine/Public/GeometryCollection/GeometryCollectionComponent.h
+++ b/Engine/Source/Runtime/Experimental/GeometryCollectionEngine/Public/GeometryCollection/GeometryCollectionComponent.h
@@ -482,9 +482,12 @@ public:
 	GEOMETRYCOLLECTIONENGINE_API FDelegateHandle RegisterOnGeometryCollectionPropertyChanged(const FOnGeometryCollectionPropertyChanged& Delegate);
 	GEOMETRYCOLLECTIONENGINE_API void UnregisterOnGeometryCollectionPropertyChanged(FDelegateHandle Handle);
 #endif
-	//~ Begin UActorComponent Interface. 
-
+	//~ End UActorComponent Interface. 
 
+	//~ Begin INavRelevantInterface Interface
+	GEOMETRYCOLLECTIONENGINE_API virtual bool IsNavigationRelevant() const override;
+	//~ End INavRelevantInterface Interface
+	
 	//~ Begin USceneComponent Interface.
 	GEOMETRYCOLLECTIONENGINE_API virtual FBoxSphereBounds CalcBounds(const FTransform& LocalToWorld) const override;
 	virtual FBoxSphereBounds CalcLocalBounds() const { return LocalBounds; }
@@ -498,7 +501,7 @@ public:
 	GEOMETRYCOLLECTIONENGINE_API virtual void AsyncPhysicsTickComponent(float DeltaTime, float SimTime) override;
 
 	GEOMETRYCOLLECTIONENGINE_API virtual void OnHiddenInGameChanged() override;
-	//~ Begin USceneComponent Interface.
+	//~ End USceneComponent Interface.
 
 
 	//~ Begin UPrimitiveComponent Interface.
@@ -1177,6 +1180,9 @@ public:
 	UPROPERTY(EditAnywhere, BlueprintReadOnly, Category = "ChaosPhysics|General")
 	bool bStoreVelocities;
 
+	UPROPERTY(Transient)
+	bool bIsCurrentlyNavigationRelevant = true;
+
 protected:
 	/** Display Bone Colors instead of assigned materials */
 	UPROPERTY(EditAnywhere, BlueprintReadWrite, Category = "ChaosPhysics|General")
