make non-recursive versions of the commonly-used recursive global matrix helpers for geometry collection algo, and guard against an invalid parent array (so we would ensure rather than infinite loop in that case)

diff --git a/Engine/Source/Runtime/Experimental/Chaos/Private/GeometryCollection/GeometryCollectionAlgo.cpp b/Engine/Source/Runtime/Experimental/Chaos/Private/GeometryCollection/GeometryCollectionAlgo.cpp
--- a/Engine/Source/Runtime/Experimental/Chaos/Private/GeometryCollection/GeometryCollectionAlgo.cpp
+++ b/Engine/Source/Runtime/Experimental/Chaos/Private/GeometryCollection/GeometryCollectionAlgo.cpp
@@ -298,76 +298,142 @@ namespace GeometryCollectionAlgo
 		}
 	}
 
-	void GlobalMatricesRecursive(const int32 Index, const TManagedArray<int32>& Parents, const TManagedArray<FTransform>& Transform, TArray<bool>& IsTransformComputed, const TManagedArray<FTransform>* UniformScale, TArray<FMatrix>& OutGlobalTransforms)
+	// Helper type for computing global matrices
+	using FIndicesNeedMatricesArray = TArray<int32, TInlineAllocator<16>>;
+
+	// return false if we failed to get the indices, which can happen if the parent array does not describe a valid tree (e.g., if it contains a loop)
+	bool GlobalMatricesGetIndicesToProcessHelper(const int32 Index, const TManagedArray<int32>& Parents, TArray<bool>& IsTransformComputed, FIndicesNeedMatricesArray& OutToProcess)
 	{
-		if (!IsTransformComputed[Index])
+		checkSlow(Index != FGeometryCollection::Invalid);
+		checkSlow(OutToProcess.IsEmpty());
+
+		OutToProcess.Add(Index);
+		bool bFoundRootOrComputed = false;
+		const int32 MaxDepth = Parents.Num();
+		while (OutToProcess.Num() <= MaxDepth)
+		{
+			int32 Parent = Parents[OutToProcess.Last()];
+			if (Parent < 0 || IsTransformComputed[Parent])
+			{
+				bFoundRootOrComputed = true;
+				break;
+			}
+			OutToProcess.Add(Parent);
+		}
+		if (!ensureMsgf(bFoundRootOrComputed, TEXT("Geometry Collection has invalid parent hierarchy, could not find root to create global transforms")))
+		{
+			return false;
+		}
+
+		return true;
+	}
+
+	void GlobalMatricesHelper(const int32 Index, const TManagedArray<int32>& Parents, const TManagedArray<FTransform>& Transform, TArray<bool>& IsTransformComputed, const TManagedArray<FTransform>* UniformScale, TArray<FMatrix>& OutGlobalTransforms)
+	{
+		if (IsTransformComputed[Index])
+		{
+			return;
+		}
+
+		FIndicesNeedMatricesArray ToProcess;
+		if (!GlobalMatricesGetIndicesToProcessHelper(Index, Parents, IsTransformComputed, ToProcess))
+		{
+			return;
+		}
+
+		while (!ToProcess.IsEmpty())
 		{
-			FMatrix Result = Transform[Index].ToMatrixWithScale();
-			if (Parents[Index] != FGeometryCollection::Invalid)
+			const int32 ProcessIndex = ToProcess.Pop(false);
+			const int32 ParentIndex = Parents[ProcessIndex];
+			FMatrix Result = Transform[ProcessIndex].ToMatrixWithScale();
+			if (ParentIndex != FGeometryCollection::Invalid)
 			{
-				GlobalMatricesRecursive(Parents[Index], Parents, Transform, IsTransformComputed, UniformScale, OutGlobalTransforms);
-				Result *= OutGlobalTransforms[Parents[Index]];
+				Result *= OutGlobalTransforms[ParentIndex];
 			}
 
 			if (UniformScale)
 			{
-				OutGlobalTransforms[Index] = (*UniformScale)[Index].ToMatrixWithScale() * Result;
+				OutGlobalTransforms[ProcessIndex] = (*UniformScale)[ProcessIndex].ToMatrixWithScale() * Result;
 			}
 			else
-			{ 
-				OutGlobalTransforms[Index] = Result;
+			{
+				OutGlobalTransforms[ProcessIndex] = Result;
 			}
-			
-			IsTransformComputed[Index] = true;
+
+			IsTransformComputed[ProcessIndex] = true;
 		}
 	}
 
-	// #note: this verison returns an FTransform to support functionality for getting global matrices for an array of indices.
-	void GlobalMatricesRecursive(const int32 Index, const TManagedArray<int32>& Parents, const TManagedArray<FTransform>& Transform, TArray<bool>& IsTransformComputed, const TManagedArray<FTransform>* UniformScale, TArray<FTransform>& OutGlobalTransforms)
+	// #note: this version outputs FTransforms to support functionality for getting global matrices for an array of indices.
+	void GlobalMatricesHelper(const int32 Index, const TManagedArray<int32>& Parents, const TManagedArray<FTransform>& Transform, TArray<bool>& IsTransformComputed, const TManagedArray<FTransform>* UniformScale, TArray<FTransform>& OutGlobalTransforms)
 	{
-		if (!IsTransformComputed[Index])
+		if (IsTransformComputed[Index])
 		{
-			FTransform Result = Transform[Index];
-			if (Parents[Index] != FGeometryCollection::Invalid)
+			return;
+		}
+
+		FIndicesNeedMatricesArray ToProcess;
+		if (!GlobalMatricesGetIndicesToProcessHelper(Index, Parents, IsTransformComputed, ToProcess))
+		{
+			return;
+		}
+
+		while (!ToProcess.IsEmpty())
+		{
+			const int32 ProcessIndex = ToProcess.Pop(false);
+			const int32 ParentIndex = Parents[ProcessIndex];
+			FTransform Result = Transform[ProcessIndex];
+			if (ParentIndex != FGeometryCollection::Invalid)
 			{
-				GlobalMatricesRecursive(Parents[Index], Parents, Transform, IsTransformComputed, UniformScale, OutGlobalTransforms);
-				Result *= OutGlobalTransforms[Parents[Index]];
+				Result *= OutGlobalTransforms[ParentIndex];
 			}
-			
+
 			if (UniformScale)
 			{
-				OutGlobalTransforms[Index] = (*UniformScale)[Index] * Result;
+				OutGlobalTransforms[ProcessIndex] = (*UniformScale)[ProcessIndex] * Result;
 			}
 			else
 			{
-				OutGlobalTransforms[Index] = Result;
+				OutGlobalTransforms[ProcessIndex] = Result;
 			}
-			
-			IsTransformComputed[Index] = true;
+
+			IsTransformComputed[ProcessIndex] = true;
 		}
 	}
 
-	FTransform GlobalMatricesRecursiveForIndices(const int32 Index, const TManagedArray<int32>& Parents, const TManagedArray<FTransform>& Transform, TArray<bool>& IsTransformComputed, const TManagedArray<FTransform>* UniformScale, TArray<FTransform>& TransformCache)
+	FTransform GlobalMatricesHelperForIndices(const int32 Index, const TManagedArray<int32>& Parents, const TManagedArray<FTransform>& Transform, TArray<bool>& IsTransformComputed, const TManagedArray<FTransform>* UniformScale, TArray<FTransform>& TransformCache)
 	{
-		if (!IsTransformComputed[Index])
+		if (IsTransformComputed[Index])
+		{
+			return TransformCache[Index];
+		}
+
+		FIndicesNeedMatricesArray ToProcess;
+		if (!GlobalMatricesGetIndicesToProcessHelper(Index, Parents, IsTransformComputed, ToProcess))
 		{
-			FTransform Result = Transform[Index];
-			if (Parents[Index] != FGeometryCollection::Invalid)
+			return TransformCache[Index];
+		}
+
+		while (!ToProcess.IsEmpty())
+		{
+			const int32 ProcessIndex = ToProcess.Pop(false);
+			const int32 ParentIndex = Parents[ProcessIndex];
+			FTransform Result = Transform[ProcessIndex];
+			if (ParentIndex != FGeometryCollection::Invalid)
 			{
-				GlobalMatricesRecursive(Parents[Index], Parents, Transform, IsTransformComputed, UniformScale, TransformCache);
-				Result *= TransformCache[Parents[Index]];
+				Result *= TransformCache[ParentIndex];
 			}
-			
+
 			if (UniformScale)
 			{
-				TransformCache[Index] = (*UniformScale)[Index] * Result;
+				TransformCache[ProcessIndex] = (*UniformScale)[ProcessIndex] * Result;
 			}
 			else
 			{
-				TransformCache[Index] = Result;
+				TransformCache[ProcessIndex] = Result;
 			}
 
-			IsTransformComputed[Index] = true;
+			IsTransformComputed[ProcessIndex] = true;
 		}
 
 		return TransformCache[Index];
@@ -402,7 +468,7 @@ namespace GeometryCollectionAlgo
 		OutGlobalTransforms.SetNumUninitialized(Indices.Num(), false);
 		for (int Idx = 0; Idx < Indices.Num(); Idx++)
 		{
-			OutGlobalTransforms[Idx] = GlobalMatricesRecursiveForIndices(Indices[Idx], Parents, RelativeTransforms, IsTransformComputed, &UniformScale, TransformCache);
+			OutGlobalTransforms[Idx] = GlobalMatricesHelperForIndices(Indices[Idx], Parents, RelativeTransforms, IsTransformComputed, &UniformScale, TransformCache);
 		}
 	}
 
@@ -417,7 +483,7 @@ namespace GeometryCollectionAlgo
 		OutGlobalTransforms.SetNumUninitialized(Indices.Num(), false);
 		for (int Idx = 0; Idx < Indices.Num(); Idx++)
 		{
-			OutGlobalTransforms[Idx] = GlobalMatricesRecursiveForIndices(Indices[Idx], Parents, RelativeTransforms, IsTransformComputed, nullptr, TransformCache);
+			OutGlobalTransforms[Idx] = GlobalMatricesHelperForIndices(Indices[Idx], Parents, RelativeTransforms, IsTransformComputed, nullptr, TransformCache);
 		}
 	}
 
@@ -447,7 +513,7 @@ namespace GeometryCollectionAlgo
 
 		for (int BoneIdx = 0; BoneIdx < NumTransforms; ++BoneIdx)
 		{
-			GlobalMatricesRecursive(BoneIdx, Parents, RelativeTransforms, IsTransformComputed, &UniformScale, OutGlobalTransforms);
+			GlobalMatricesHelper(BoneIdx, Parents, RelativeTransforms, IsTransformComputed, &UniformScale, OutGlobalTransforms);
 		}
 	}
 
@@ -463,7 +529,7 @@ namespace GeometryCollectionAlgo
 
 		for (int BoneIdx = 0; BoneIdx < NumTransforms; ++BoneIdx)
 		{
-			GlobalMatricesRecursive(BoneIdx, Parents, RelativeTransforms, IsTransformComputed, nullptr, OutGlobalTransforms);
+			GlobalMatricesHelper(BoneIdx, Parents, RelativeTransforms, IsTransformComputed, nullptr, OutGlobalTransforms);
 		}
 	}
 	template void CHAOS_API GlobalMatrices<FTransform>(const TManagedArray<FTransform>&, const TManagedArray<int32>&, const TManagedArray<FTransform>&, TArray<FTransform>&);
