Add missing const with TObjectPtr
#rnx

diff --git a/Engine/Plugins/Runtime/ModelViewViewModel/Source/ModelViewViewModel/Private/Bindings/MVVMCompiledBindingLibrary.cpp b/Engine/Plugins/Runtime/ModelViewViewModel/Source/ModelViewViewModel/Private/Bindings/MVVMCompiledBindingLibrary.cpp
--- a/Engine/Plugins/Runtime/ModelViewViewModel/Source/ModelViewViewModel/Private/Bindings/MVVMCompiledBindingLibrary.cpp
+++ b/Engine/Plugins/Runtime/ModelViewViewModel/Source/ModelViewViewModel/Private/Bindings/MVVMCompiledBindingLibrary.cpp
@@ -45,7 +45,7 @@ UFunction* FMVVMVCompiledFields::GetFunction(FName FunctionName) const
 UE::FieldNotification::FFieldId FMVVMVCompiledFields::GetFieldId(FName FieldName) const
 {
 	// ScriptStruct can't implements UNotifyFieldValueChanged
-	UClass* Class = CastChecked<UClass>(ClassOrScriptStruct);
+	const UClass* Class = CastChecked<UClass>(ClassOrScriptStruct);
 	check(Class->ImplementsInterface(UNotifyFieldValueChanged::StaticClass()));
 
 	TScriptInterface<INotifyFieldValueChanged> Interface = Class->GetDefaultObject();
diff --git a/Engine/Source/Runtime/Engine/Classes/Components/HierarchicalInstancedStaticMeshComponent.h b/Engine/Source/Runtime/Engine/Classes/Components/HierarchicalInstancedStaticMeshComponent.h
--- a/Engine/Source/Runtime/Engine/Classes/Components/HierarchicalInstancedStaticMeshComponent.h
+++ b/Engine/Source/Runtime/Engine/Classes/Components/HierarchicalInstancedStaticMeshComponent.h
@@ -259,7 +259,7 @@ public:
 	ENGINE_API void AcceptPrebuiltTree(TArray<FInstancedStaticMeshInstanceData>& InInstanceData, TArray<FClusterNode>& InClusterTree, int32 InOcclusionLayerNumNodes, int32 InNumBuiltRenderInstances);
 	bool IsAsyncBuilding() const { return bIsAsyncBuilding; }
 	bool IsTreeFullyBuilt() const { return !bIsOutOfDate; }
-	ENGINE_API void GetTree(TArray<FClusterNode>& OutClusterTree);
+	ENGINE_API void GetTree(TArray<FClusterNode>& OutClusterTree) const;
 	ENGINE_API FVector GetAverageScale() const;
 
 	/** Heuristic for the number of leaves in the tree **/
diff --git a/Engine/Source/Runtime/Engine/Private/HierarchicalInstancedStaticMesh.cpp b/Engine/Source/Runtime/Engine/Private/HierarchicalInstancedStaticMesh.cpp
--- a/Engine/Source/Runtime/Engine/Private/HierarchicalInstancedStaticMesh.cpp
+++ b/Engine/Source/Runtime/Engine/Private/HierarchicalInstancedStaticMesh.cpp
@@ -3531,7 +3531,7 @@ void UHierarchicalInstancedStaticMeshComponent::GetOverlappingBoxTransforms(cons
 	});
 }
 
-void UHierarchicalInstancedStaticMeshComponent::GetTree(TArray<FClusterNode>& OutClusterTree)
+void UHierarchicalInstancedStaticMeshComponent::GetTree(TArray<FClusterNode>& OutClusterTree) const
 {
 	OutClusterTree = *ClusterTreePtr;
 }
diff --git a/Engine/Source/Runtime/Experimental/GeometryCollectionEngine/Public/GeometryCollection/GeometryCollectionISMPoolDebugDrawComponent.cpp b/Engine/Source/Runtime/Experimental/GeometryCollectionEngine/Public/GeometryCollection/GeometryCollectionISMPoolDebugDrawComponent.cpp
--- a/Engine/Source/Runtime/Experimental/GeometryCollectionEngine/Public/GeometryCollection/GeometryCollectionISMPoolDebugDrawComponent.cpp
+++ b/Engine/Source/Runtime/Experimental/GeometryCollectionEngine/Public/GeometryCollection/GeometryCollectionISMPoolDebugDrawComponent.cpp
@@ -253,14 +253,14 @@ FDebugRenderSceneProxy* UGeometryCollectionISMPoolDebugDrawComponent::CreateDebu
 
 	FDebugRenderSceneProxy* DebugProxy = new FGeometryCollectionISMPoolDebugDrawSceneProxy(this);
 
-	if (UHierarchicalInstancedStaticMeshComponent* HISMComponent = Cast<UHierarchicalInstancedStaticMeshComponent>(SelectedComponent))
+	if (const UHierarchicalInstancedStaticMeshComponent* HISMComponent = Cast<UHierarchicalInstancedStaticMeshComponent>(SelectedComponent))
 	{
 		const FTransform LocalToWorldTransform = HISMComponent->GetComponentTransform();
 		TArray<FClusterNode> TreeNodes;
 		HISMComponent->GetTree(TreeNodes);
 		DrawBoxAndChildren(DebugProxy, LocalToWorldTransform, TreeNodes, 0, 0);
 	}
-	else if(UInstancedStaticMeshComponent* ISMComponent = Cast<UInstancedStaticMeshComponent>(SelectedComponent))
+	else if(const UInstancedStaticMeshComponent* ISMComponent = Cast<UInstancedStaticMeshComponent>(SelectedComponent))
 	{
 		const FBox LocalBounds = ISMComponent->CalcBounds(FTransform::Identity).GetBox();
 		const FTransform LocalToWorldTransform = ISMComponent->GetComponentTransform();
