Chaos - mesh collision optimization
- don't generate manifolds when shared vertices have already collided
- alternative SAT edge-pair loop (currently slower but code remains, defined out - it may be better after SIMD)
- don't save friction anchors for contacts with no impulse if there are lots of manifold points (optimize friction anchor search for large manifolds)

diff --git a/Engine/Source/Runtime/Experimental/Chaos/Private/Chaos/Collision/ConvexTriangleContactPoint.cpp b/Engine/Source/Runtime/Experimental/Chaos/Private/Chaos/Collision/ConvexTriangleContactPoint.cpp
--- a/Engine/Source/Runtime/Experimental/Chaos/Private/Chaos/Collision/ConvexTriangleContactPoint.cpp
+++ b/Engine/Source/Runtime/Experimental/Chaos/Private/Chaos/Collision/ConvexTriangleContactPoint.cpp
@@ -9,7 +9,12 @@
 #include "Chaos/Triangle.h"
 #include "Misc/MemStack.h"
 
-//PRAGMA_DISABLE_OPTIMIZATION
+//UE_DISABLE_OPTIMIZATION
+
+// Method to use when checking edge pairs in convex-triangle SAT
+// 0: Iterate over convex edges, use IsMinkowskiSumConvexTriangle to skip invalid edge pairs
+// 1: Iterate over convex faces, then their edges. This has the advantage we don't need the edge list on FConvex
+#define CHAOS_CONVEX_TRIANGLE_EDGEEDGE_METHOD 0
 
 namespace Chaos
 {
@@ -206,7 +211,8 @@ namespace Chaos
 		//
 		// SAT: Convex edges vs triangle edges
 		//
-		
+
+#if (CHAOS_CONVEX_TRIANGLE_EDGEEDGE_METHOD == 0)
 		// Calculate the distance of triangle each edge to the convex separating plane
 		FReal TriVertexConvexDMin0 = FVec3::DotProduct(Triangle.GetVertex(0) - ConvexPlaneX, ConvexPlaneN);
 		FReal TriVertexConvexDMin1 = FVec3::DotProduct(Triangle.GetVertex(1) - ConvexPlaneX, ConvexPlaneN);
@@ -217,7 +223,7 @@ namespace Chaos
 			FMath::Min(TriVertexConvexDMin0, TriVertexConvexDMin1),
 			FMath::Min(TriVertexConvexDMin1, TriVertexConvexDMin2),
 		};
-
+		
 		FVec3 EdgeEdgeN = FVec3(0);
 		FReal EdgeEdgeDMin = InvalidPhi;
 		int32 ConvexEdgeIndexMin = INDEX_NONE;
@@ -302,10 +308,107 @@ namespace Chaos
 				}
 			}
 		}
+#elif (CHAOS_CONVEX_TRIANGLE_EDGEEDGE_METHOD == 1)
+		// Loop over all the convex planes and then its edges, and check the separating axis defined by the convex edge and each tri edge
+		// NOTE: This performs duplicate work as each edge is shared with 2 convex faces, but we also reject half the planes (for most
+		// convexes, although it depends) so it ends up faster than using IsMinkowskiSumConvexTriangle on all the edge pairs
+		FVec3 EdgeEdgeN = FVec3(0);
+		FReal EdgeEdgeDMin = InvalidPhi;
+		int32 ConvexEdgeIndexMin = INDEX_NONE;
+		int32 TriEdgeIndexMin = INDEX_NONE;
+		for (int32 PlaneIndex = 0; PlaneIndex < Convex.NumPlanes(); ++PlaneIndex)
+		{
+			FVec3 ConN, ConX;
+			Convex.GetPlaneNX(PlaneIndex, ConN, ConX);
+
+			// Ignore faces with normals in same direction as triangle
+			const FReal NormalDot = FVec3::DotProduct(ConN, TriN);
+			if (NormalDot > 0)
+			{
+				continue;
+			}
+
+			// Check all the edges on this face
+			const int32 NumPlaneVerts = Convex.NumPlaneVertices(PlaneIndex);
+			int32 ConvexEdgeVertexIndex = Convex.GetPlaneVertex(PlaneIndex, NumPlaneVerts - 1);
+			FVec3 ConvexEdgeV0 = Convex.GetVertex(ConvexEdgeVertexIndex);
+			for (int32 PlaneVertexIndex = 0; PlaneVertexIndex < NumPlaneVerts; ++PlaneVertexIndex)
+			{
+				ConvexEdgeVertexIndex = Convex.GetPlaneVertex(PlaneIndex, PlaneVertexIndex);
+				const FVec3 ConvexEdgeV1 = Convex.GetVertex(ConvexEdgeVertexIndex);
+
+				FVec3 TriEdgeV0 = Triangle.GetVertex(2);
+				for (int32 TriVertexIndex = 0; TriVertexIndex < 3; ++TriVertexIndex)
+				{
+					const FVec3 TriEdgeV1 = Triangle.GetVertex(TriVertexIndex);
+
+					FVec3 EdgeN = FVec3::CrossProduct(ConvexEdgeV1 - ConvexEdgeV0, TriEdgeV1 - TriEdgeV0);
+
+					// Ignore parallel edge pairs
+					const FReal EdgeNLenSq = EdgeN.SizeSquared();
+					if (EdgeNLenSq < NormalToleranceSq)
+					{
+						continue;
+					}
+
+					// Project the convex and triangle onto the edge axis
+					// NOTE: EdgeN is not normalized yet, so these projections will be scaled by EdgeN.Size()
+					FReal ThisConvexEdgeDMin, ThisConvexEdgeDMax, ThisTriEdgeDMin, ThisTriEdgeDMax;;
+					int32 ThisConvexEdgeIndexMin, ThisConvexEdgeIndexMax, ThisTriEdgeIndexMin, ThisTriEdgeIndexMax;
+					ProjectOntoAxis(Convex, EdgeN, ConvexEdgeV0, ThisConvexEdgeDMin, ThisConvexEdgeDMax, ThisConvexEdgeIndexMin, ThisConvexEdgeIndexMax, nullptr);
+					ProjectOntoAxis(Triangle, EdgeN, ConvexEdgeV0, ThisTriEdgeDMin, ThisTriEdgeDMax, ThisTriEdgeIndexMin, ThisTriEdgeIndexMax);
+
+					// Separation will be the larger of these
+					// NOTE: these separations are scaled by EdgeN.Size()
+					const FReal ScaledPhi0 = ThisConvexEdgeDMin - ThisTriEdgeDMax;
+					const FReal ScaledPhi1 = ThisTriEdgeDMin - ThisConvexEdgeDMax;
+
+					// Ignore separations greater than cull distance (taking the non-noemlized edge axis into account)
+					const FReal ScaledCullDistanceSq = FMath::Square(CullDistance) * EdgeNLenSq;
+					if (FMath::Square(ScaledPhi0) > ScaledCullDistanceSq)
+					{
+						continue;
+					}
+					if (FMath::Square(ScaledPhi1) > ScaledCullDistanceSq)
+					{
+						continue;
+					}
+
+					// Normalize the edge axis
+					const FReal InvEdgeNLen = FMath::InvSqrt(EdgeNLenSq);
+					const FReal Phi0 = ScaledPhi0 * InvEdgeNLen;
+					const FReal Phi1 = ScaledPhi1 * InvEdgeNLen;
+					EdgeN = EdgeN * InvEdgeNLen;
+
+					// See if we have a new shallowest overlap
+					if ((Phi0 >= Phi1) && (Phi0 > EdgeEdgeDMin))
+					{
+						EdgeEdgeN = EdgeN;
+						EdgeEdgeDMin = Phi0;
+						ConvexEdgeIndexMin = ThisConvexEdgeIndexMin;
+						TriEdgeIndexMin = ThisTriEdgeIndexMax;
+					}
+					else if (Phi1 > EdgeEdgeDMin)
+					{
+						EdgeEdgeN = EdgeN;
+						EdgeEdgeDMin = Phi1;
+						ConvexEdgeIndexMin = ThisConvexEdgeIndexMax;
+						TriEdgeIndexMin = ThisTriEdgeIndexMin;
+					}
+
+					TriEdgeV0 = TriEdgeV1;
+				}
+
+				ConvexEdgeV0 = ConvexEdgeV1;
+			}
+		}
+#else
+		#error Invalid CHAOS_CONVEX_TRIANGLE_EDGEEDGE_METHOD
+#endif
 
 		// Determine which of the features we want to use
 		// NOTE: we rely on the fact that all valid Phi values are greater than InvalidPhi here
-		const FReal TriFaceBias = FReal(1.e-2);	// Prevent flip=flip on near parallel cases
+		const FReal TriFaceBias = FReal(1.e-2);	// Prevent flip-flop on near-parallel cases
 		EContactPointType ContactType = EContactPointType::Unknown;
 		if ((TriPlaneDMin != InvalidPhi) && (TriPlaneDMin + TriFaceBias > ConvexPlaneDMin) && (TriPlaneDMin + TriFaceBias > EdgeEdgeDMin))
 		{
diff --git a/Engine/Source/Runtime/Experimental/Chaos/Private/Chaos/Collision/MeshContactGenerator.cpp b/Engine/Source/Runtime/Experimental/Chaos/Private/Chaos/Collision/MeshContactGenerator.cpp
new file mode 100644
--- /dev/null
+++ b/Engine/Source/Runtime/Experimental/Chaos/Private/Chaos/Collision/MeshContactGenerator.cpp
@@ -0,0 +1,371 @@
+// Copyright Epic Games, Inc.All Rights Reserved.
+#include "Chaos/Collision/MeshContactGenerator.h"
+#include "Chaos/DebugDrawQueue.h"
+#include "Chaos/Framework/UncheckedArray.h"
+
+namespace Chaos
+{
+	extern FRealSingle Chaos_Collision_MeshContactNormalThreshold;
+}
+
+namespace Chaos::CVars
+{
+	extern int32 ChaosSolverDebugDrawMeshContacts;
+}
+
+namespace Chaos::Private
+{
+	FMeshContactGenerator::FMeshContactGenerator(const int32 InHashSize)
+		: EdgeTriangleIndicesMap(InHashSize)
+		, VertexContactIndicesMap(InHashSize)
+	{
+	}
+
+	void FMeshContactGenerator::Reset(const int32 InMaxTriangles, const int32 InMaxContacts)
+	{
+		Triangles.Reset(InMaxTriangles);
+
+		Contacts.Reset(InMaxContacts);
+		ContactDatas.Reset(InMaxContacts);
+
+		EdgeTriangleIndicesMap.Reset(InMaxTriangles);
+		VertexContactIndicesMap.Reset(InMaxContacts);
+	}
+
+	void FMeshContactGenerator::AddTriangleContacts(const FContactPointManifold& TriangleContactPoints, const int32 LocalTriangleIndex)
+	{
+		// Clamp the number of contacts we can add per mesh
+		if (Contacts.Num() + TriangleContactPoints.Num() > Contacts.Max())
+		{
+			return;
+		}
+
+		FTriangleExt& Triangle = Triangles[LocalTriangleIndex];
+
+		// We need to know how many vertex and edge collisions a triangle has on it (from collision with other triangle 
+		// that share edges and vertices) so we avoid visiting triangles if we already know the contacts. Check all
+		// contacts and assigne edge/vertex status as required. See GenerateMeshContacts()
+		const FReal BarycentricTolerance = 1.e-3;
+		for (const FContactPoint& ContactPoint : TriangleContactPoints)
+		{
+			// See if we have a vertex or edge contact
+			// @todo(chaos): we should be able to produce this as output from the manifold creator
+			int32 LocalVertexID0, LocalVertexID1;
+			GetTriangleEdgeVerticesAtPosition(
+				ContactPoint.ShapeContactPoints[1],
+				Triangle.Triangle.GetVertex(0), Triangle.Triangle.GetVertex(1), Triangle.Triangle.GetVertex(2),
+				LocalVertexID0, LocalVertexID1,
+				BarycentricTolerance);
+
+			const int32 VertexID0 = (LocalVertexID0 != INDEX_NONE) ? Triangle.VertexIndices[LocalVertexID0] : INDEX_NONE;
+			const int32 VertexID1 = (LocalVertexID1 != INDEX_NONE) ? Triangle.VertexIndices[LocalVertexID1] : INDEX_NONE;
+
+			const int32 NewContactIndex = Contacts.Num();
+
+			const FRealSingle ContactNormalDotTriangleNormal = FRealSingle(FVec3::DotProduct(ContactPoint.ShapeContactNormal, Triangle.Normal));
+
+			const FReal FaceNormalThreshold = FReal(0.998);	// 3deg
+			const bool bIsFaceContact = (ContactNormalDotTriangleNormal > FaceNormalThreshold);
+
+			// Register vertex and edge collisions
+			if ((VertexID0 != INDEX_NONE) && (VertexID1 != INDEX_NONE))
+			{
+				// Edge collisions - if it's a face normal, tell each triangle using the edge it has a face contact
+				if (bIsFaceContact)
+				{
+					Triangles[LocalTriangleIndex].NumFaceEdgeCollisions += 1;
+
+					const int32 OtherLocalTriangleIndex = GetOtherTriangleIndexForEdge(LocalTriangleIndex, FContactEdgeID(VertexID0, VertexID1));
+					if (OtherLocalTriangleIndex != INDEX_NONE)
+					{
+						Triangles[OtherLocalTriangleIndex].NumFaceEdgeCollisions += 1;
+					}
+				}
+			}
+			else if ((VertexID0 != INDEX_NONE) && (VertexID1 == INDEX_NONE))
+			{
+				// Vertex collisions - we only keep one vertex collision per vertex
+				if (FVertexContactIndex* ExistingContactIndex = VertexContactIndicesMap.Find(VertexID0))
+				{
+					// We have an existing contact at this vertex. 
+					// If it was a face contact just ignore the new contact, 
+					// otherwise keep the contact with the most face-pointing normal.
+					if (!ExistingContactIndex->bIsFaceContact)
+					{
+						FContactPoint& ExistingContact = Contacts[ExistingContactIndex->ContactIndex];
+						FTriangleContactPointData& ExistingContactData = ContactDatas[ExistingContactIndex->ContactIndex];
+						if (ContactNormalDotTriangleNormal > ExistingContactData.GetContactNormalDotTriangleNormal())
+						{
+							ExistingContact = ContactPoint;
+
+							ExistingContactData.SetVertexID(VertexID0);
+							ExistingContactData.SetContactNormalDotTriangleNormal(ContactNormalDotTriangleNormal);
+							ExistingContactData.SetTriangleIndex(LocalTriangleIndex);
+
+							ExistingContactIndex->bIsFaceContact = bIsFaceContact;
+						}
+					}
+
+					// Don't add this contact
+					continue;
+				}
+
+				// This is the first time we hit this vertex so store the contact index with the vertex
+				VertexContactIndicesMap.Emplace(VertexID0, VertexID0, NewContactIndex, bIsFaceContact);
+			}
+
+			// Store and enable the contact
+			Contacts.Add(ContactPoint);
+			ContactDatas.AddDefaulted();
+			ContactDatas[NewContactIndex].SetEdgeOrVertexID({ VertexID0, VertexID1 });
+			ContactDatas[NewContactIndex].SetContactNormalDotTriangleNormal(ContactNormalDotTriangleNormal);
+			ContactDatas[NewContactIndex].SetTriangleIndex(LocalTriangleIndex);
+			ContactDatas[NewContactIndex].SetEnabled();
+		}
+	}
+
+	void FMeshContactGenerator::ProcessGeneratedContacts(const FRigidTransform3& ConvexTransform, const FRigidTransform3& MeshToConvexTransform)
+	{
+		// Contacts that get pruned or corrected will show as green
+		DebugDrawContacts(ConvexTransform, FColor::Green, 0.5);
+
+		PruneAndCorrectContacts();
+
+		// Final contacts will be red
+		DebugDrawContacts(ConvexTransform, FColor::Red, 1.0);
+
+		// Visited triangles are white, ignored triangles are gray
+		DebugDrawTriangles(ConvexTransform, FColor::White, FColor::Silver);
+
+		FinalizeContacts(MeshToConvexTransform);
+	}
+
+	void FMeshContactGenerator::PruneAndCorrectContacts()
+	{
+		const FRealSingle NormalThreshold = FRealSingle(0.998);
+
+		for (int32 ContactIndex = 0; ContactIndex < Contacts.Num(); ++ContactIndex)
+		{
+			FTriangleContactPointData& ContactPointData = ContactDatas[ContactIndex];
+
+			// Reject back-faces
+			if (ContactPointData.GetContactNormalDotTriangleNormal() < 0)
+			{
+				ContactPointData.SetDisabled();
+				continue;
+			}
+
+			// Fix edge normals
+			if (ContactPointData.GetContactNormalDotTriangleNormal() < Chaos_Collision_MeshContactNormalThreshold)
+			{
+				FixContactNormal(ContactIndex);
+			}
+		}
+
+		// re-pack the contact array
+		RemoveDisabledContacts();
+	}
+
+	void FMeshContactGenerator::FixContactNormal(const int32 ContactIndex)
+	{
+		FContactPoint& ContactPoint = Contacts[ContactIndex];
+		FTriangleContactPointData& ContactPointData = ContactDatas[ContactIndex];
+
+		const int32 LocalTriangleIndex = ContactPointData.GetTriangleIndex();
+		const FTriangleExt& Triangle = Triangles[LocalTriangleIndex];
+
+		// If we have an edge or vertex contact, make sure that the normal is in a valiid range, based
+		// on the triangles that share that edge or vertex.
+		if (ContactPointData.IsEdge())
+		{
+			// We have an edge collision. The contact normal must lie between the normals of the two faces using the edge
+			const int32 OtherLocalTriangleIndex = GetOtherTriangleIndexForEdge(LocalTriangleIndex, ContactPointData.GetEdgeID());
+			if (OtherLocalTriangleIndex != INDEX_NONE)
+			{
+				const FTriangleExt& OtherTriangle = Triangles[OtherLocalTriangleIndex];
+
+				const FReal ContactDotNormal = ContactPointData.GetContactNormalDotTriangleNormal();
+				const FReal MinContactDotNormal = FRealSingle(FVec3::DotProduct(OtherTriangle.Normal, Triangle.Normal));
+				if (ContactDotNormal < MinContactDotNormal)
+				{
+					// We are outside the valid normal range for this edge
+					// Convert the edge collision to a face collision on one of the faces, selected to get the smallest depth
+					FVec3 CorrectedContactNormal;
+					int32 CorrectedTriangleIndex;
+					const FReal OtherContactDotNormal = FVec3::DotProduct(ContactPoint.ShapeContactNormal, OtherTriangle.Normal);
+					if (ContactDotNormal >= OtherContactDotNormal)
+					{
+						CorrectedContactNormal = (ContactDotNormal > -SMALL_NUMBER) ? Triangle.Normal : -Triangle.Normal;
+						CorrectedTriangleIndex = LocalTriangleIndex;
+					}
+					else
+					{
+						CorrectedContactNormal = (OtherContactDotNormal > -SMALL_NUMBER) ? OtherTriangle.Normal : -OtherTriangle.Normal;
+						CorrectedTriangleIndex = OtherLocalTriangleIndex;
+					}
+
+					// NOTE: We keep the contact depth as it is because we know that the depth is a lower-bound. 
+					// We have to update the ShapeContactPoint[0] because Phi is actually derived from the positions (the value in ContactPoint is just a cache of current state)
+					// @todo(chaos): face selection logic might be better if we knew the contact velocity
+					ContactPoint.ShapeContactNormal = CorrectedContactNormal;
+					ContactPoint.ShapeContactPoints[0] = ContactPoint.ShapeContactPoints[1] + ContactPoint.Phi * CorrectedContactNormal;
+					ContactPointData.SetTriangleIndex(CorrectedTriangleIndex);
+				}
+			}
+		}
+		else if (ContactPointData.IsVertex())
+		{
+			// We have a vertex collision. Ensure that the contact normal is in a valid range for the vertex based on the
+			// triangles that share the vertex (there can be arbitarily many of these).
+			for (int32 TriangleIndex = 0; TriangleIndex < Triangles.Num(); ++TriangleIndex)
+			{
+				const FTriangleExt& OtherTriangle = Triangles[TriangleIndex];
+				if (OtherTriangle.HasVertexIndex(ContactPointData.GetVertexID()))
+				{
+					FVec3 VertexA, VertexB, VertexC;
+					if (OtherTriangle.GetVertexPosition(ContactPointData.GetVertexID(), VertexA) && OtherTriangle.GetOtherVertexPositions(ContactPointData.GetVertexID(), VertexB, VertexC))
+					{
+						// Does the contact normal point into the infinite prism formed by extruding the triangle along the face normal?
+						// It does if the contact normal dotted with the edge plane normal is negative for both edge planes on the triangle that use the vertex.
+						const FVec3 EdgeDelta0 = VertexB - VertexA;
+						const FVec3 EdgeDelta1 = VertexC - VertexA;
+						const FReal EdgeSign0 = FVec3::DotProduct(FVec3::CrossProduct(ContactPoint.ShapeContactNormal, VertexB - VertexA), OtherTriangle.Normal);
+						const FReal EdgeSign1 = FVec3::DotProduct(FVec3::CrossProduct(ContactPoint.ShapeContactNormal, VertexC - VertexA), OtherTriangle.Normal);
+						if (FMath::Sign(EdgeSign0) == FMath::Sign(EdgeSign1))
+						{
+							const FVec3 Centroid = OtherTriangle.Triangle.GetCentroid();
+							const FReal NormalDotCentroid = FVec3::DotProduct(ContactPoint.ShapeContactNormal, Centroid - ContactPoint.ShapeContactPoints[1]);
+							if (NormalDotCentroid > 0)
+							{
+								const FReal OtherContactDotNormal = FVec3::DotProduct(ContactPoint.ShapeContactNormal, OtherTriangle.Normal);
+								const FVec3 CorrectedContactNormal = OtherTriangle.Normal; //(OtherContactDotNormal > -SMALL_NUMBER) ? OtherTriangle.Normal : -OtherTriangle.Normal;
+
+								ContactPoint.ShapeContactNormal = CorrectedContactNormal;
+								ContactPoint.ShapeContactPoints[0] = ContactPoint.ShapeContactPoints[1] + ContactPoint.Phi * CorrectedContactNormal;
+							}
+						}
+					}
+				}
+			}
+		}
+	}
+
+	void FMeshContactGenerator::RemoveDisabledContacts()
+	{
+		// Re-pack the contact point array without re-ordering
+		const int32 NumContactPoints = Contacts.Num();
+		int32 DestContactIndex = 0;		// Index to where the next enabled item goes
+		int32 SrcContactIndex = 0;		// Index to the next enabled item
+		while (SrcContactIndex < NumContactPoints)
+		{
+			if (!ContactDatas[SrcContactIndex].IsEnabled())
+			{
+				// Find the next enabled point index to use as the copy source
+				while (++SrcContactIndex < NumContactPoints)
+				{
+					if (ContactDatas[SrcContactIndex].IsEnabled())
+					{
+						break;
+					}
+				}
+
+				if (SrcContactIndex == NumContactPoints)
+				{
+					break;
+				}
+			}
+
+			// If we have removed elements, copy source to dest
+			if (DestContactIndex != SrcContactIndex)
+			{
+				Contacts[DestContactIndex] = Contacts[SrcContactIndex];
+				ContactDatas[DestContactIndex] = ContactDatas[SrcContactIndex];
+			}
+
+			++DestContactIndex;
+			++SrcContactIndex;
+		}
+
+		// Clip the array to the enabled set
+		Contacts.SetNum(DestContactIndex, false);
+		ContactDatas.SetNum(DestContactIndex, false);
+	}
+
+	void FMeshContactGenerator::FinalizeContacts(const FRigidTransform3& MeshToConvexTransform)
+	{
+		for (int32 ContactIndex = 0; ContactIndex < Contacts.Num(); ++ContactIndex)
+		{
+			FContactPoint& ContactPoint = Contacts[ContactIndex];
+
+			ContactPoint.ShapeContactPoints[1] = MeshToConvexTransform.InverseTransformPositionNoScale(ContactPoint.ShapeContactPoints[1]);
+			ContactPoint.ShapeContactNormal = MeshToConvexTransform.InverseTransformVectorNoScale(ContactPoint.ShapeContactNormal);
+		}
+	}
+
+	void FMeshContactGenerator::DebugDrawContacts(const FRigidTransform3& ConvexTransform, const FColor& Color, const FReal LineScale)
+	{
+#if CHAOS_DEBUG_DRAW
+		if (CVars::ChaosSolverDebugDrawMeshContacts && FDebugDrawQueue::GetInstance().IsDebugDrawingEnabled())
+		{
+			const FReal Duration = 0;
+			const int32 DrawPriority = 10;
+			for (int32 ContactIndex = 0; ContactIndex < Contacts.Num(); ++ContactIndex)
+			{
+				const FTriangleContactPointData& ContactPointData = ContactDatas[ContactIndex];
+				if (ContactPointData.IsEnabled())
+				{
+					const FContactPoint& ContactPoint = Contacts[ContactIndex];
+					const FVec3 P1 = ConvexTransform.TransformPosition(ContactPoint.ShapeContactPoints[1]);
+					const FVec3 N = ConvexTransform.TransformVectorNoScale(ContactPoint.ShapeContactNormal);
+
+					// Draw the normal from the triangle face
+					FDebugDrawQueue::GetInstance().DrawDebugLine(P1, P1 + LineScale * FReal(50) * N, Color, false, FRealSingle(Duration), DrawPriority, FRealSingle(LineScale * 2));
+				}
+			}
+		}
+#endif
+	}
+
+	void FMeshContactGenerator::DebugDrawTriangles(const FRigidTransform3& ConvexTransform, const FColor& VisitedColor, const FColor& IgnoredColor)
+	{
+#if CHAOS_DEBUG_DRAW
+		if (CVars::ChaosSolverDebugDrawMeshContacts && FDebugDrawQueue::GetInstance().IsDebugDrawingEnabled())
+		{
+			// NOTE: drawing in two loops so that the visited triangle edges draw over the ignored ones (priority doesn't seem to work)
+			for (int32 TriangleIndex = 0; TriangleIndex < Triangles.Num(); ++TriangleIndex)
+			{
+				const FTriangleExt& Triangle = Triangles[TriangleIndex];
+				if (Triangle.VisitIndex == INDEX_NONE)
+				{
+					DebugDrawTriangle(ConvexTransform, Triangle, IgnoredColor);
+				}
+			}
+			for (int32 TriangleIndex = 0; TriangleIndex < Triangles.Num(); ++TriangleIndex)
+			{
+				const FTriangleExt& Triangle = Triangles[TriangleIndex];
+				if (Triangle.VisitIndex != INDEX_NONE)
+				{
+					DebugDrawTriangle(ConvexTransform, Triangle, VisitedColor);
+				}
+			}
+		}
+#endif
+	}
+
+	void FMeshContactGenerator::DebugDrawTriangle(const FRigidTransform3& ConvexTransform, const FTriangleExt& Triangle, const FColor& Color)
+	{
+#if CHAOS_DEBUG_DRAW
+		const FReal Duration = 0;
+		const FReal LineScale = 1;
+		const int8 DrawPriority = 10;
+
+		const FVec3 V0 = ConvexTransform.TransformPosition(Triangle.Triangle.GetVertex(0));
+		const FVec3 V1 = ConvexTransform.TransformPosition(Triangle.Triangle.GetVertex(1));
+		const FVec3 V2 = ConvexTransform.TransformPosition(Triangle.Triangle.GetVertex(2));
+
+		FDebugDrawQueue::GetInstance().DrawDebugLine(V0, V1, Color, false, FRealSingle(Duration), DrawPriority, FRealSingle(LineScale));
+		FDebugDrawQueue::GetInstance().DrawDebugLine(V1, V2, Color, false, FRealSingle(Duration), DrawPriority, FRealSingle(LineScale));
+		FDebugDrawQueue::GetInstance().DrawDebugLine(V2, V0, Color, false, FRealSingle(Duration), DrawPriority, FRealSingle(LineScale));
+#endif
+	}
+}
diff --git a/Engine/Source/Runtime/Experimental/Chaos/Private/Chaos/CollisionOneShotManifolds.cpp b/Engine/Source/Runtime/Experimental/Chaos/Private/Chaos/CollisionOneShotManifolds.cpp
--- a/Engine/Source/Runtime/Experimental/Chaos/Private/Chaos/CollisionOneShotManifolds.cpp
+++ b/Engine/Source/Runtime/Experimental/Chaos/Private/Chaos/CollisionOneShotManifolds.cpp
@@ -91,6 +91,13 @@ namespace Chaos
 	bool bChaos_Collision_MeshManifoldSortByDistance = false;
 	FAutoConsoleVariableRef CVarChaos_Collision_LargeMeshManifoldSortByDistance(TEXT("p.Chaos.Collision.SortMeshManifoldByDistance"), bChaos_Collision_MeshManifoldSortByDistance, TEXT("Sort large mesh manifold points by |RxN| for improved solver stability (less rotation in first iteration)"));
 
+	int32 Chaos_Collision_MeshManifoldHashSize = 256;
+	FAutoConsoleVariableRef CVarChaos_Collision_MeshManifoldHashSize(TEXT("p.Chaos.Collision.MeshManifoldHashSize"), Chaos_Collision_MeshManifoldHashSize, TEXT("Hash table size to use in vertex and edge maps in convex-mesh collision"));
+
+	// @todo(chaos): Temp while we test the new convex-mesh collision optimizations
+	bool bChaos_Collision_EnableMeshManifoldOptimizedLoop = true;
+	FAutoConsoleVariableRef CVarChaos_Collision_EnableMeshManifoldOptimizedLoop(TEXT("p.Chaos.Collision.EnableMeshManifoldOptimizedLoop"), bChaos_Collision_EnableMeshManifoldOptimizedLoop, TEXT(""));
+
 	// Whether to use the new index-less GJK. 
 	// @todo(chaos): This should be removed once soaked for a bit (enabled 7 June 2022)
 	bool bChaos_Collision_UseGJK2 = false;
diff --git a/Engine/Source/Runtime/Experimental/Chaos/Private/Chaos/CollisionOneShotManifoldsMeshShapes.cpp b/Engine/Source/Runtime/Experimental/Chaos/Private/Chaos/CollisionOneShotManifoldsMeshShapes.cpp
--- a/Engine/Source/Runtime/Experimental/Chaos/Private/Chaos/CollisionOneShotManifoldsMeshShapes.cpp
+++ b/Engine/Source/Runtime/Experimental/Chaos/Private/Chaos/CollisionOneShotManifoldsMeshShapes.cpp
@@ -7,6 +7,7 @@
 #include "Chaos/Collision/ContactPointsMiscShapes.h"
 #include "Chaos/Collision/ContactTriangles.h"
 #include "Chaos/Collision/ConvexTriangleContactPoint.h"
+#include "Chaos/Collision/MeshContactGenerator.h"
 #include "Chaos/Collision/PBDCollisionConstraint.h"
 #include "Chaos/Collision/SphereTriangleContactPoint.h"
 #include "Chaos/Convex.h"
@@ -28,6 +29,8 @@ namespace Chaos
 	extern bool bChaos_Collision_OneSidedHeightField;
 	extern FRealSingle Chaos_Collision_TriMeshDistanceTolerance;
 	extern FRealSingle Chaos_Collision_TriMeshPhiToleranceScale;
+	extern int32 Chaos_Collision_MeshManifoldHashSize;
+	extern bool bChaos_Collision_EnableMeshManifoldOptimizedLoop;
 
 	extern bool bChaos_Collision_UseCapsuleTriMesh2;
 	extern bool bChaos_Collision_UseConvexTriMesh2;
@@ -108,6 +111,33 @@ namespace Chaos
 			ConstructSphereTriangleOneShotManifold(Sphere, Triangle, CullDistance, OutContactPoints);
 		}
 
+		template<typename ConvexType, typename MeshType>
+		void ConstructConvexMeshOneShotManifold2(const ConvexType& Convex, const FRigidTransform3& ConvexTransform, const MeshType& Mesh, const FRigidTransform3& MeshTransform, const FVec3& MeshScale, const FReal CullDistance, Private::FMeshContactGenerator& ContactGenerator)
+		{
+			FRigidTransform3 MeshToConvexTransform = MeshTransform.GetRelativeTransformNoScale(ConvexTransform);
+			MeshToConvexTransform.SetScale3D(MeshScale);
+
+			// @todo(chaos): add Convex.CalculateInverseTransformed bounds with scale support (to optimize sphere and capsule)
+			const FAABB3 ConvexBounds = FAABB3(Convex.BoundingBox()).Thicken(CullDistance);
+			const FAABB3 MeshQueryBounds = ConvexBounds.InverseTransformedAABB(MeshToConvexTransform);
+
+			// Generate the contact manifold between Convex and a Triangle
+			const auto& GenerateConvexTriangleContacts =
+				[&Convex, &MeshToConvexTransform, CullDistance](const FTriangle& Triangle, FContactPointManifold& OutContactPoints)
+				{
+					ConstructConvexTriangleOneShotManifold2(Convex, Triangle, CullDistance, OutContactPoints);
+				};
+
+			// Collect all the triangles that overlap our convex, transformed into Convex space
+			Mesh.CollectTriangles(MeshQueryBounds, MeshToConvexTransform, ConvexBounds, ContactGenerator);
+
+			// Generate a set of contact points for all triangles
+			ContactGenerator.GenerateMeshContacts(GenerateConvexTriangleContacts);
+
+			// Process the contacts to minimize manifold etc
+			ContactGenerator.ProcessGeneratedContacts(ConvexTransform, MeshToConvexTransform);
+		}
+
 		/**
 		* @brief Create a minimized set of contact points between a convex polyhedron (box, convex) and a non-convex mesh (trimesh, heightfield)
 		*
@@ -277,34 +307,66 @@ namespace Chaos
 			ensure(MeshTransform.GetScale3D() == FVec3(1));
 
 			const FVec3 MeshScale = FVec3(1);	// Scale is built into heightfield
-
 			const FReal CullDistance = Constraint.GetCullDistance();
 			const FReal PhiTolerance = CalculateTriMeshPhiTolerance(CullDistance);
 			const FReal DistanceTolerance = Chaos_Collision_TriMeshDistanceTolerance;
-			FContactTriangleCollector MeshContacts(bChaos_Collision_OneSidedHeightField, PhiTolerance, DistanceTolerance, ConvexTransform);
 
-			if (const FImplicitBox3* RawBox = Convex.template GetObject<FImplicitBox3>())
+			if (bChaos_Collision_EnableMeshManifoldOptimizedLoop)
 			{
-				ConstructConvexMeshOneShotManifold(*RawBox, ConvexTransform, Mesh, MeshTransform, MeshScale, CullDistance, MeshContacts);
-			}
-			else if (const TImplicitObjectScaled<FImplicitConvex3>* ScaledConvex = Convex.template GetObject<TImplicitObjectScaled<FImplicitConvex3>>())
-			{
-				ConstructConvexMeshOneShotManifold(*ScaledConvex, ConvexTransform, Mesh, MeshTransform, MeshScale, CullDistance, MeshContacts);
-			}
-			else if (const TImplicitObjectInstanced<FImplicitConvex3>* InstancedConvex = Convex.template GetObject<TImplicitObjectInstanced<FImplicitConvex3>>())
-			{
-				ConstructConvexMeshOneShotManifold(*InstancedConvex, ConvexTransform, Mesh, MeshTransform, MeshScale, CullDistance, MeshContacts);
-			}
-			else if (const FImplicitConvex3* RawConvex = Convex.template GetObject<FImplicitConvex3>())
-			{
-				ConstructConvexMeshOneShotManifold(*RawConvex, ConvexTransform, Mesh, MeshTransform, MeshScale, CullDistance, MeshContacts);
+				// New version uses a two-pass loop over triangles to avoid visiting triangles whose vertices are all colliding as a result of checking adjacent triangles
+				const int32 HashSize = FMath::RoundUpToPowerOfTwo(Chaos_Collision_MeshManifoldHashSize);
+				Private::FMeshContactGenerator ContactGenerator(HashSize);
+
+				if (const FImplicitBox3* RawBox = Convex.template GetObject<FImplicitBox3>())
+				{
+					ConstructConvexMeshOneShotManifold2(*RawBox, ConvexTransform, Mesh, MeshTransform, MeshScale, CullDistance, ContactGenerator);
+				}
+				else if (const TImplicitObjectScaled<FImplicitConvex3>* ScaledConvex = Convex.template GetObject<TImplicitObjectScaled<FImplicitConvex3>>())
+				{
+					ConstructConvexMeshOneShotManifold2(*ScaledConvex, ConvexTransform, Mesh, MeshTransform, MeshScale, CullDistance, ContactGenerator);
+				}
+				else if (const TImplicitObjectInstanced<FImplicitConvex3>* InstancedConvex = Convex.template GetObject<TImplicitObjectInstanced<FImplicitConvex3>>())
+				{
+					ConstructConvexMeshOneShotManifold2(*InstancedConvex, ConvexTransform, Mesh, MeshTransform, MeshScale, CullDistance, ContactGenerator);
+				}
+				else if (const FImplicitConvex3* RawConvex = Convex.template GetObject<FImplicitConvex3>())
+				{
+					ConstructConvexMeshOneShotManifold2(*RawConvex, ConvexTransform, Mesh, MeshTransform, MeshScale, CullDistance, ContactGenerator);
+				}
+				else
+				{
+					check(false);
+				}
+
+				Constraint.SetOneShotManifoldContacts(ContactGenerator.GetContactPoints());
 			}
 			else
 			{
-				check(false);
-			}
+				FContactTriangleCollector MeshContacts(bChaos_Collision_OneSidedHeightField, PhiTolerance, DistanceTolerance, ConvexTransform);
 
-			Constraint.SetOneShotManifoldContacts(MeshContacts.GetContactPoints());
+				if (const FImplicitBox3* RawBox = Convex.template GetObject<FImplicitBox3>())
+				{
+					ConstructConvexMeshOneShotManifold(*RawBox, ConvexTransform, Mesh, MeshTransform, MeshScale, CullDistance, MeshContacts);
+				}
+				else if (const TImplicitObjectScaled<FImplicitConvex3>* ScaledConvex = Convex.template GetObject<TImplicitObjectScaled<FImplicitConvex3>>())
+				{
+					ConstructConvexMeshOneShotManifold(*ScaledConvex, ConvexTransform, Mesh, MeshTransform, MeshScale, CullDistance, MeshContacts);
+				}
+				else if (const TImplicitObjectInstanced<FImplicitConvex3>* InstancedConvex = Convex.template GetObject<TImplicitObjectInstanced<FImplicitConvex3>>())
+				{
+					ConstructConvexMeshOneShotManifold(*InstancedConvex, ConvexTransform, Mesh, MeshTransform, MeshScale, CullDistance, MeshContacts);
+				}
+				else if (const FImplicitConvex3* RawConvex = Convex.template GetObject<FImplicitConvex3>())
+				{
+					ConstructConvexMeshOneShotManifold(*RawConvex, ConvexTransform, Mesh, MeshTransform, MeshScale, CullDistance, MeshContacts);
+				}
+				else
+				{
+					check(false);
+				}
+
+				Constraint.SetOneShotManifoldContacts(MeshContacts.GetContactPoints());
+			}
 		}
 	}
 }
diff --git a/Engine/Source/Runtime/Experimental/Chaos/Private/Chaos/HeightField.cpp b/Engine/Source/Runtime/Experimental/Chaos/Private/Chaos/HeightField.cpp
--- a/Engine/Source/Runtime/Experimental/Chaos/Private/Chaos/HeightField.cpp
+++ b/Engine/Source/Runtime/Experimental/Chaos/Private/Chaos/HeightField.cpp
@@ -3,6 +3,7 @@
 #include "Chaos/HeightField.h"
 #include "Chaos/Collision/ContactPoint.h"
 #include "Chaos/Collision/ContactPointsMiscShapes.h"
+#include "Chaos/Collision/MeshContactGenerator.h"
 #include "Chaos/Collision/TriangleOverlap.h"
 #include "Chaos/Collision/PBDCollisionConstraint.h"
 #include "Chaos/CollisionOneShotManifolds.h"
@@ -2509,4 +2510,103 @@ namespace Chaos
 
 		return OutPhi;
 	}
+
+	bool FHeightField::GetOverlappingCellRange(const FAABB3& QueryBounds, TVec2<int32>& OutCellBegin, TVec2<int32>& OutCellEnd) const
+	{
+		// Calculate the square cell range that overlaps the input query bounds (which is in local space)
+		const FBounds2D HeightfieldFlatBounds = GetFlatBounds();
+		FVec2 Scale2D(GeomData.Scale[0], GeomData.Scale[1]);
+		FBounds2D GridQueryBounds;
+		GridQueryBounds.Min = FVec2(QueryBounds.Min()[0], QueryBounds.Min()[1]);
+		GridQueryBounds.Max = FVec2(QueryBounds.Max()[0], QueryBounds.Max()[1]);
+		GridQueryBounds = FBounds2D::FromPoints(HeightfieldFlatBounds.Clamp(GridQueryBounds.Min) / Scale2D, HeightfieldFlatBounds.Clamp(GridQueryBounds.Max) / Scale2D);
+		OutCellBegin = FlatGrid.Cell(GridQueryBounds.Min);
+		OutCellEnd = FlatGrid.Cell(GridQueryBounds.Max);
+		OutCellEnd[0] += 1;
+		OutCellEnd[1] += 1;
+
+		// @todo(chaos): this should return false if QueryBounds does not overlap the heightfield
+		return true;
+	}
+
+	void FHeightField::CollectTriangles(const FAABB3& QueryBounds, const FRigidTransform3& QueryTransform, const FAABB3& ObjectBounds, Private::FMeshContactGenerator& Collector) const
+	{
+		TVec2<int32> CellBegin, CellEnd;
+		if (!GetOverlappingCellRange(QueryBounds, CellBegin, CellEnd))
+		{
+			return;
+		}
+
+		// Tell the collector how many triangles to expect
+		const int32 NumCells = (CellEnd[0] - CellBegin[0]) * (CellEnd[1] - CellBegin[1]);
+		Collector.BeginCollect(2 * NumCells);
+
+		FVec3 Points[4];
+		FTriangle Triangles[2];
+		const bool bStandardWinding = ((GeomData.Scale.X * GeomData.Scale.Y * GeomData.Scale.Z) >= FReal(0));
+
+		// Visit all the cells in stored order
+		for (int32 CellX = CellBegin[0]; CellX < CellEnd[0]; ++CellX)
+		{
+			for (int32 CellY = CellBegin[1]; CellY < CellEnd[1]; ++CellY)
+			{
+				const int32 CellVertexIndex = CellY * GeomData.NumCols + CellX;	// First vertex in cell
+				const int32 CellIndex = CellY * (GeomData.NumCols - 1) + CellX;
+
+				// Check for holes and skip checking if we'll never collide
+				if (GeomData.MaterialIndices.IsValidIndex(CellIndex) && GeomData.MaterialIndices[CellIndex] == TNumericLimits<uint8>::Max())
+				{
+					continue;
+				}
+
+				// The triangle is solid so proceed to test it
+				FAABB3 CellBounds;
+				GeomData.GetPointsAndBoundsScaled(CellVertexIndex, Points, CellBounds);
+
+				// Perform a 3D bounds check now in local space
+				if (CellBounds.Intersects(QueryBounds))
+				{
+					// Transform points into the space of the query
+					// @todo(chaos): duplicate work here when we overlap lots of cells. We could generate the transformed verts for the full query once...
+					Points[0] = QueryTransform.TransformPositionNoScale(Points[0]);
+					Points[1] = QueryTransform.TransformPositionNoScale(Points[1]);
+					Points[2] = QueryTransform.TransformPositionNoScale(Points[2]);
+					Points[3] = QueryTransform.TransformPositionNoScale(Points[3]);
+
+					// Perform a 3D bounds check in object space
+					FAABB3 ObjectCellBounds = FAABB3::FromPoints(Points[0], Points[1], Points[2], Points[3]);
+
+					if (ObjectBounds.Intersects(ObjectCellBounds))
+					{
+						// @todo(chaos): utility function for this
+						const int32 VertexIndex0 = CellVertexIndex;
+						const int32 VertexIndex1 = CellVertexIndex + 1;
+						const int32 VertexIndex2 = CellVertexIndex + GeomData.NumCols;
+						const int32 VertexIndex3 = CellVertexIndex + GeomData.NumCols + 1;
+
+						const int32 FaceIndex0 = CellIndex * 2 + 0;
+						const int32 FaceIndex1 = CellIndex * 2 + 1;
+
+						if (bStandardWinding)
+						{
+							Triangles[0] = FTriangle(Points[0], Points[1], Points[3]);
+							Triangles[1] = FTriangle(Points[0], Points[3], Points[2]);
+							Collector.AddTriangle(Triangles[0], FaceIndex0, VertexIndex0, VertexIndex1, VertexIndex3);
+							Collector.AddTriangle(Triangles[1], FaceIndex1, VertexIndex0, VertexIndex3, VertexIndex2);
+						}
+						else
+						{
+							Triangles[0] = FTriangle(Points[0], Points[3], Points[1]);
+							Triangles[1] = FTriangle(Points[0], Points[2], Points[3]);
+							Collector.AddTriangle(Triangles[0], FaceIndex0, VertexIndex0, VertexIndex3, VertexIndex1);
+							Collector.AddTriangle(Triangles[1], FaceIndex1, VertexIndex0, VertexIndex2, VertexIndex3);
+						}
+					}
+				}
+			}
+		}
+
+		Collector.EndCollect();
+	}
+
 }
diff --git a/Engine/Source/Runtime/Experimental/Chaos/Public/Chaos/Collision/ContactTriangles.h b/Engine/Source/Runtime/Experimental/Chaos/Public/Chaos/Collision/ContactTriangles.h
--- a/Engine/Source/Runtime/Experimental/Chaos/Public/Chaos/Collision/ContactTriangles.h
+++ b/Engine/Source/Runtime/Experimental/Chaos/Public/Chaos/Collision/ContactTriangles.h
@@ -9,20 +9,20 @@
 
 namespace Chaos
 {
-	// Given a point and a triangle, check if the point is onan edge or vertex and return vertex indices if so.
+	// Given a point and a triangle, check if the point is on an edge or vertex and return vertex indices if so.
 	// If Position is at one of the vertices, OutEdgeVertexIndexA will an index into Vertices
 	// If Position is one one of the edges, both OutEdgeVertexIndexA and OutEdgeVertexIndexB will be an index into Vertices
 	// If both OutEdgeVertexIndexA and OutEdgeVertexIndexB are INDEX_NONE, Position is not on the triangle edges
 	inline bool GetTriangleEdgeVerticesAtPosition(
 		const FVec3& Position, 
-		const FVec3 Vertices[],
+		const FVec3 VertexA, const FVec3 VertexB, const FVec3 VertexC,
 		int32& OutEdgeVertexIndexA, int32& OutEdgeVertexIndexB, 
 		const FReal BaryCentricTolerance = UE_KINDA_SMALL_NUMBER)
 	{
 		OutEdgeVertexIndexA = INDEX_NONE;
 		OutEdgeVertexIndexB = INDEX_NONE;
 
-		const FVec3 BaryCentric = ToBarycentric(Position, Vertices[0], Vertices[1], Vertices[2]);
+		const FVec3 BaryCentric = ToBarycentric(Position, VertexA, VertexB, VertexC);
 
 		// Is it a vertex contact?
 		if (FMath::IsNearlyEqual(BaryCentric.X, FReal(1), BaryCentricTolerance))
@@ -64,6 +64,19 @@ namespace Chaos
 		return false;
 	}
 
+	inline bool GetTriangleEdgeVerticesAtPosition(
+		const FVec3& Position,
+		const FVec3 Vertices[],
+		int32& OutEdgeVertexIndexA, int32& OutEdgeVertexIndexB,
+		const FReal BaryCentricTolerance = UE_KINDA_SMALL_NUMBER)
+	{
+		return GetTriangleEdgeVerticesAtPosition(
+			Position, 
+			Vertices[0], Vertices[1], Vertices[2], 
+			OutEdgeVertexIndexA, OutEdgeVertexIndexB, 
+			BaryCentricTolerance);
+	}
+
 
 	/**
 	 * @brief Data held alongside contact points when generating contacts against a (likely non-convex) mesh of triangles
@@ -156,7 +169,8 @@ namespace Chaos
 		FORCEINLINE FContactEdgeID(const FContactVertexID VertexIndexA, const FContactVertexID VertexIndexB)
 		{
 			// EdgeID is the same if we swap the vertex indices
-			if (VertexIndexA < VertexIndexB)
+			// The check for INDEX_NONE allows us to use an EdgeID as a VertexID without swapping order
+			if ((VertexIndexA < VertexIndexB) || (VertexIndexB == INDEX_NONE))
 			{
 				VertexIDs[0] = VertexIndexA;
 				VertexIDs[1] = VertexIndexB;
@@ -201,30 +215,30 @@ namespace Chaos
 	};
 
 	
-	class UE_DEPRECATED(5.4, "No longer used") FTriangleContactPoint : public FContactPoint
-	{
-	public:
-		FTriangleContactPoint()
-			: FContactPoint()
-			, EdgeID(INDEX_NONE, INDEX_NONE)
-			, VertexID(INDEX_NONE)
-			, ContactTriangleIndex(INDEX_NONE)
-		{
-		}
-
-		FTriangleContactPoint(const FContactPoint& InContactPoint)
-			: FContactPoint(InContactPoint)
-			, EdgeID(INDEX_NONE, INDEX_NONE)
-			, VertexID(INDEX_NONE)
-			, ContactTriangleIndex(INDEX_NONE)
-		{
-		}
-
-		FContactEdgeID EdgeID;
-		FContactVertexID VertexID;
-
-		int32 ContactTriangleIndex;
-	};
+	//class UE_DEPRECATED(5.4, "No longer used") FTriangleContactPoint : public FContactPoint
+	//{
+	//public:
+	//	FTriangleContactPoint()
+	//		: FContactPoint()
+	//		, EdgeID(INDEX_NONE, INDEX_NONE)
+	//		, VertexID(INDEX_NONE)
+	//		, ContactTriangleIndex(INDEX_NONE)
+	//	{
+	//	}
+
+	//	FTriangleContactPoint(const FContactPoint& InContactPoint)
+	//		: FContactPoint(InContactPoint)
+	//		, EdgeID(INDEX_NONE, INDEX_NONE)
+	//		, VertexID(INDEX_NONE)
+	//		, ContactTriangleIndex(INDEX_NONE)
+	//	{
+	//	}
+
+	//	FContactEdgeID EdgeID;
+	//	FContactVertexID VertexID;
+
+	//	int32 ContactTriangleIndex;
+	//};
 
 	/**
 	* An ID for an Edge or Vertex in a triangle mesh.
@@ -324,6 +338,11 @@ namespace Chaos
 			EdgeOrVertexID = FContactEdgeOrVertexID(InVertexIDA, InVertexIDB);
 		}
 
+		FORCEINLINE void SetEdgeOrVertexID(const FContactEdgeOrVertexID InEdgeOrVertexID)
+		{
+			EdgeOrVertexID = InEdgeOrVertexID;
+		}
+
 		FORCEINLINE void SetTriangleIndex(const int32 InTriangleIndex)
 		{
 			ContactTriangleIndex = InTriangleIndex;
diff --git a/Engine/Source/Runtime/Experimental/Chaos/Public/Chaos/Collision/MeshContactGenerator.h b/Engine/Source/Runtime/Experimental/Chaos/Public/Chaos/Collision/MeshContactGenerator.h
new file mode 100644
--- /dev/null
+++ b/Engine/Source/Runtime/Experimental/Chaos/Public/Chaos/Collision/MeshContactGenerator.h
@@ -0,0 +1,289 @@
+// Copyright Epic Games, Inc.All Rights Reserved.
+#pragma once
+
+#include "Chaos/Core.h"
+#include "Chaos/Collision/ContactPoint.h"
+#include "Chaos/Collision/ContactTriangles.h"
+#include "Chaos/CollisionResolutionTypes.h"
+#include "Chaos/Framework/HashMappedArray.h"
+#include "Chaos/Triangle.h"
+
+namespace Chaos::Private
+{
+	// Adapter for FContactEdgeID use in THashMappedArray
+	template<> struct THashMappedArrayIDTraits<FContactEdgeID>
+	{
+		using FIDType = FContactEdgeID;
+
+		static uint32 HashID(const FIDType& ID)
+		{
+			return uint32(MurmurFinalize64(ID.EdgeID));
+		}
+	};
+
+	/**
+	* Generate contacts between a collision shape and the triangles from a mesh.
+	*/
+	class FMeshContactGenerator
+	{
+	public:
+		FMeshContactGenerator(const int32 InHashSize);
+
+		// Clear and initialize buffers
+		void BeginCollect(const int32 InNumTriangles)
+		{
+			const int32 ExpectedNumContacts = FMath::Min(InNumTriangles * 4, 1000);
+			Reset(InNumTriangles, ExpectedNumContacts);
+		}
+
+		// Add a triangle that we might be overlapping
+		void AddTriangle(const FTriangle& MeshTriangle, const int32 MeshTriangleIndex, const int32 VertexIndex0, const int32 VertexIndex1, const int32 VertexIndex2)
+		{
+			Triangles.Emplace(MeshTriangle, MeshTriangleIndex, VertexIndex0, VertexIndex1, VertexIndex2);
+		}
+
+		void EndCollect()
+		{
+			for (int32 LocalTriangleIndex = 0; LocalTriangleIndex < GetNumTriangles(); ++LocalTriangleIndex)
+			{
+				const FTriangleExt& Triangle = Triangles[LocalTriangleIndex];
+				AddTriangleEdge(LocalTriangleIndex, Triangle.VertexIndices[0], Triangle.VertexIndices[1]);
+				AddTriangleEdge(LocalTriangleIndex, Triangle.VertexIndices[1], Triangle.VertexIndices[2]);
+				AddTriangleEdge(LocalTriangleIndex, Triangle.VertexIndices[2], Triangle.VertexIndices[0]);
+			}
+		}
+
+		// Loop over (the required subset of) all triangles and call the TriangleContactGenerator to create a manifold for each.
+		// TriangleContactGeneratorType: void(const FTriangle& Triangle, FContactPointManifold& OutContactPoints)
+		template<typename TriangleContactGeneratorType>
+		void GenerateMeshContacts(const TriangleContactGeneratorType& TriangleContactGenerator)
+		{
+			FContactPointManifold TriangleContactPoints;
+
+			// First loop: Visit triangles that do not have any collisions on any of their vertices or edges.
+			// This will skip all triangles whose neighbours have already been processed and generated a contact
+			// on a shared edge/vertex.
+			for (int32 LocalTriangleIndex = 0; LocalTriangleIndex < GetNumTriangles(); ++LocalTriangleIndex)
+			{
+				if (GetNumTriangleFaceCollisions(LocalTriangleIndex) == 0)
+				{
+					TriangleContactPoints.Reset();
+					TriangleContactGenerator(Triangles[LocalTriangleIndex].Triangle, TriangleContactPoints);
+
+					AddTriangleContacts(TriangleContactPoints, LocalTriangleIndex);
+
+					SetTriangleVisited(LocalTriangleIndex, 0);
+				}
+			}
+
+			// Second loop: Visit remaining triangles that have less than 3 contacts on them. This will skip all triangles
+			// that have a full manifold as a result of collisions on shared edges/vertices from adjacent triangles.
+			for (int32 LocalTriangleIndex = 0; LocalTriangleIndex < GetNumTriangles(); ++LocalTriangleIndex)
+			{
+				if (!IsTriangleVisited(LocalTriangleIndex) && (GetNumTriangleFaceCollisions(LocalTriangleIndex) < 3))
+				{
+					TriangleContactPoints.Reset();
+					TriangleContactGenerator(Triangles[LocalTriangleIndex].Triangle, TriangleContactPoints);
+
+					AddTriangleContacts(TriangleContactPoints, LocalTriangleIndex);
+
+					SetTriangleVisited(LocalTriangleIndex, 1);
+				}
+			}
+		}
+
+		// Process all the contact points generated by GenerateMeshContacts. This prunes duplicates, fixes normals, and
+		// transforms the contact data back into shape-local space.
+		void ProcessGeneratedContacts(const FRigidTransform3& ConvexTransform, const FRigidTransform3& MeshToConvexTransform);
+
+		// The results of contact generation (must call ProcessGeneratedContacts prior to GetContactPoints)
+		TArrayView<const FContactPoint> GetContactPoints() const
+		{
+			return MakeArrayView(Contacts);
+		}
+
+	private:
+		// A triangle plus some extended data and state
+		class FTriangleExt
+		{
+		public:
+			FTriangleExt(const FTriangle& InTriangle, const int32 InTriangleIndex, const int32 InVertexIndex0, const int32 InVertexIndex1, const int32 InVertexIndex2)
+				: Triangle(InTriangle)
+				, Normal(InTriangle.GetNormal())
+				, TriangleIndex(InTriangleIndex)
+				, VertexIndices{ InVertexIndex0, InVertexIndex1, InVertexIndex2 }
+				, NumFaceEdgeCollisions(0)
+				, VisitIndex(INDEX_NONE)
+				, bEnabled(true)
+			{
+			}
+
+			// Does this triangle contains the specified vertex? (VertexIndex is an index into the owning mesh's vertices)
+			inline bool HasVertexIndex(const int32 VertexIndex) const
+			{
+				return (VertexIndices[0] == VertexIndex) || (VertexIndices[1] == VertexIndex) || (VertexIndices[2] == VertexIndex);
+			}
+
+			inline bool GetVertexPosition(const int32 VertexIndex, FVec3& OutVertex) const
+			{
+				if (VertexIndex == VertexIndices[0])
+				{
+					OutVertex = Triangle.GetVertex(0);
+					return true;
+				}
+				else if (VertexIndex == VertexIndices[1])
+				{
+					OutVertex = Triangle.GetVertex(1);
+					return true;
+				}
+				else if (VertexIndex == VertexIndices[2])
+				{
+					OutVertex = Triangle.GetVertex(2);
+					return true;
+				}
+				return false;
+			}
+
+			// Get the positions of the other two vertices in the triangle. (VertexIndex is an index into the owning mesh's vertices)
+			inline bool GetOtherVertexPositions(const int32 VertexIndex, FVec3& OutVertex0, FVec3& OutVertex1) const
+			{
+				if (VertexIndex == VertexIndices[0])
+				{
+					OutVertex0 = Triangle.GetVertex(1);
+					OutVertex1 = Triangle.GetVertex(2);
+					return true;
+				}
+				else if (VertexIndex == VertexIndices[1])
+				{
+					OutVertex0 = Triangle.GetVertex(2);
+					OutVertex1 = Triangle.GetVertex(0);
+					return true;
+				}
+				else if (VertexIndex == VertexIndices[2])
+				{
+					OutVertex0 = Triangle.GetVertex(0);
+					OutVertex1 = Triangle.GetVertex(1);
+					return true;
+				}
+				return false;
+			}
+
+			FTriangle Triangle;
+			FVec3 Normal;
+			int32 TriangleIndex;
+			int32 VertexIndices[3];
+			int8 NumFaceEdgeCollisions;
+			int8 VisitIndex;
+			bool bEnabled;
+		};
+
+		// A contact index combined with a flag to indicate if the normal is roughly along the triangle face
+		struct FVertexContactIndex
+		{
+			FContactVertexID ID;
+			int32 ContactIndex;
+			bool bIsFaceContact;
+		};
+
+		// The triangle indices that share an edge (assumes only 2)
+		struct FEdgeTriangleIndices
+		{
+			FEdgeTriangleIndices(const FContactEdgeID& InEdgeID, const int32 InIndex0, const int32 InIndex1)
+				: ID(InEdgeID)
+				, LocalTriangleIndices{ InIndex0, InIndex1 }
+			{
+			}
+
+			FContactEdgeID ID;
+			int32 LocalTriangleIndices[2];
+		};
+
+		void Reset(const int32 InMaxTriangles, const int32 InMaxContacts);
+
+		void AddTriangleEdge(const int32 LocalTriangleIndex, const int32 VertexIndex0, const int32 VertexIndex1)
+		{
+			const FContactEdgeID EdgeID = FContactEdgeID(VertexIndex0, VertexIndex1);
+			FEdgeTriangleIndices* EdgeTriangleIndices = EdgeTriangleIndicesMap.Find(EdgeID);
+			if (EdgeTriangleIndices == nullptr)
+			{
+				EdgeTriangleIndicesMap.Emplace(EdgeID, EdgeID, LocalTriangleIndex, INDEX_NONE);
+			}
+			else
+			{
+				EdgeTriangleIndices->LocalTriangleIndices[1] = LocalTriangleIndex;
+			}
+		}
+
+		int32 GetOtherTriangleIndexForEdge(const int32 LocalTriangleIndex, const FContactEdgeID EdgeID)
+		{
+			FEdgeTriangleIndices* EdgeTriangleIndices = EdgeTriangleIndicesMap.Find(EdgeID);
+			if (EdgeTriangleIndices != nullptr)
+			{
+				return (EdgeTriangleIndices->LocalTriangleIndices[0] == LocalTriangleIndex) ? EdgeTriangleIndices->LocalTriangleIndices[1] : EdgeTriangleIndices->LocalTriangleIndices[0];
+			}
+			return INDEX_NONE;
+		}
+
+		bool HasFaceVertexCollision(const FContactVertexID VertexID) const
+		{
+			if (const FVertexContactIndex* ContactIndex = VertexContactIndicesMap.Find(VertexID))
+			{
+				return ContactIndex->bIsFaceContact;
+			}
+			return false;
+		}
+
+		int32 GetNumTriangleFaceCollisions(const int32 LocalTriangleIndex) const
+		{
+			const int32 MeshVertexIndex0 = Triangles[LocalTriangleIndex].VertexIndices[0];
+			const int32 MeshVertexIndex1 = Triangles[LocalTriangleIndex].VertexIndices[1];
+			const int32 MeshVertexIndex2 = Triangles[LocalTriangleIndex].VertexIndices[2];
+			int32 NumCollisions = Triangles[LocalTriangleIndex].NumFaceEdgeCollisions;
+			NumCollisions += HasFaceVertexCollision(MeshVertexIndex0) ? 1 : 0;
+			NumCollisions += HasFaceVertexCollision(MeshVertexIndex1) ? 1 : 0;
+			NumCollisions += HasFaceVertexCollision(MeshVertexIndex2) ? 1 : 0;
+			return NumCollisions;
+		}
+
+		bool IsTriangleVisited(const int32 LocalTriangleIndex) const
+		{
+			return Triangles[LocalTriangleIndex].VisitIndex != INDEX_NONE;
+		}
+
+		void SetTriangleVisited(const int32 LocalTriangleIndex, const int8 VisitIndex)
+		{
+			Triangles[LocalTriangleIndex].VisitIndex = VisitIndex;
+		}
+
+		int32 GetNumTriangles() const
+		{
+			return Triangles.Num();
+		}
+
+		void AddTriangleContacts(const FContactPointManifold& TriangleContactPoints, const int32 LocalTriangleIndex);
+		void PruneAndCorrectContacts();
+		void FixContactNormal(const int32 ContactIndex);
+		void RemoveDisabledContacts();
+		void FinalizeContacts(const FRigidTransform3& MeshToConvexTransform);
+
+		void DebugDrawContacts(const FRigidTransform3& ConvexTransform, const FColor& Color, const FReal LineScale);
+		void DebugDrawTriangles(const FRigidTransform3& ConvexTransform, const FColor& VisitedColor, const FColor& IgnoredColor);
+		void DebugDrawTriangle(const FRigidTransform3& ConvexTransform, const FTriangleExt& TriangleData, const FColor& Color);
+
+	private:
+		// All the triangles we might collide with
+		TArray<FTriangleExt> Triangles;
+
+		// The contact data. This is split into the output data (FContactPoint) and the extra
+		// per-contact metadata required during processing (FTriangleContactPointData).
+		TArray<FContactPoint> Contacts;
+		TArray<FTriangleContactPointData> ContactDatas;
+
+		// A map of EdgeID to the two triangles (indices) that use the Edge
+		THashMappedArray<FContactEdgeID, FEdgeTriangleIndices> EdgeTriangleIndicesMap;
+
+		// A map of VertexID to contact index on that vertex - we only ever keep one contact per vertex
+		THashMappedArray<FContactVertexID, FVertexContactIndex> VertexContactIndicesMap;
+	};
+
+}
\ No newline at end of file
diff --git a/Engine/Source/Runtime/Experimental/Chaos/Public/Chaos/Collision/PBDCollisionConstraint.h b/Engine/Source/Runtime/Experimental/Chaos/Public/Chaos/Collision/PBDCollisionConstraint.h
--- a/Engine/Source/Runtime/Experimental/Chaos/Public/Chaos/Collision/PBDCollisionConstraint.h
+++ b/Engine/Source/Runtime/Experimental/Chaos/Public/Chaos/Collision/PBDCollisionConstraint.h
@@ -673,7 +673,8 @@ namespace Chaos
 		FORCEINLINE void ResetSolverResults()
 		{
 			// NOTE: does not initalize any data. All properties will be written to in SetSolverResults
-			SavedManifoldPoints.SetNum(ManifoldPoints.Num());
+			//SavedManifoldPoints.SetNum(ManifoldPoints.Num());
+			SavedManifoldPoints.Reset(ManifoldPoints.Num());
 			ManifoldPointResults.SetNum(ManifoldPoints.Num());
 		}
 
@@ -689,7 +690,6 @@ namespace Chaos
 		{
 			FManifoldPoint& ManifoldPoint = ManifoldPoints[ManifoldPointIndex];
 			FManifoldPointResult& ManifoldPointResult = ManifoldPointResults[ManifoldPointIndex];
-			FSavedManifoldPoint& SavedManifoldPoint = SavedManifoldPoints[ManifoldPointIndex];
 
 			ManifoldPointResult.NetPushOut = NetPushOut;
 			ManifoldPointResult.NetImpulse = NetImpulse;
@@ -704,17 +704,26 @@ namespace Chaos
 			if (StaticFrictionRatio >= FReal(1.0f - UE_KINDA_SMALL_NUMBER))
 			{
 				// StaticFrictionRatio ~= 1: Static friction held - we keep the same contacts points as-is for use next frame
+				FSavedManifoldPoint& SavedManifoldPoint = SavedManifoldPoints[SavedManifoldPoints.AddUninitialized()];
 				SavedManifoldPoint.ShapeContactPoints[0] = ManifoldPoint.ShapeAnchorPoints[0];
 				SavedManifoldPoint.ShapeContactPoints[1] = ManifoldPoint.ShapeAnchorPoints[1];
 				ManifoldPointResult.bInsideStaticFrictionCone = true;
 			}
 			else if (StaticFrictionRatio < FReal(UE_KINDA_SMALL_NUMBER))
 			{
-				// StaticFrictionRatio ~= 0: No friction (or no contact) - discard the friction anchors
-				const FVec3 Anchor0 = ManifoldPoint.ContactPoint.ShapeContactPoints[0];
-				const FVec3 Anchor1 = ManifoldPoint.ContactPoint.ShapeContactPoints[1];
-				SavedManifoldPoint.ShapeContactPoints[0] = Anchor0;
-				SavedManifoldPoint.ShapeContactPoints[1] = Anchor1;
+				// StaticFrictionRatio ~= 0: No friction (or no contact/impulse) - discard the friction anchors
+				// If we have a lot of manifold points, we don't store the previous position for these contacts because we assume
+				// that there will be others that actually applied an impulse and will fall into the other two branches.
+				// This is an optimization when we have to match the new contacts with the saved ones (see AssignSavedManifoldPoints)
+				const int32 SmallNumManifoldPoints = 8;
+				if (ManifoldPoints.Num() < SmallNumManifoldPoints)
+				{
+					const FVec3 Anchor0 = ManifoldPoint.ContactPoint.ShapeContactPoints[0];
+					const FVec3 Anchor1 = ManifoldPoint.ContactPoint.ShapeContactPoints[1];
+					FSavedManifoldPoint& SavedManifoldPoint = SavedManifoldPoints[SavedManifoldPoints.AddUninitialized()];
+					SavedManifoldPoint.ShapeContactPoints[0] = Anchor0;
+					SavedManifoldPoint.ShapeContactPoints[1] = Anchor1;
+				}
 			}
 			else
 			{
@@ -722,6 +731,7 @@ namespace Chaos
 				// toward the last-detected contact position so that it sits at the edge of the friction cone.
 				const FVec3 Anchor0 = FVec3::Lerp(ManifoldPoint.ContactPoint.ShapeContactPoints[0], ManifoldPoint.ShapeAnchorPoints[0], StaticFrictionRatio);
 				const FVec3 Anchor1 = FVec3::Lerp(ManifoldPoint.ContactPoint.ShapeContactPoints[1], ManifoldPoint.ShapeAnchorPoints[1], StaticFrictionRatio);
+				FSavedManifoldPoint& SavedManifoldPoint = SavedManifoldPoints[SavedManifoldPoints.AddUninitialized()];
 				SavedManifoldPoint.ShapeContactPoints[0] = Anchor0;
 				SavedManifoldPoint.ShapeContactPoints[1] = Anchor1;
 			}
diff --git a/Engine/Source/Runtime/Experimental/Chaos/Public/Chaos/Framework/HashMappedArray.h b/Engine/Source/Runtime/Experimental/Chaos/Public/Chaos/Framework/HashMappedArray.h
new file mode 100644
--- /dev/null
+++ b/Engine/Source/Runtime/Experimental/Chaos/Public/Chaos/Framework/HashMappedArray.h
@@ -0,0 +1,130 @@
+// Copyright Epic Games, Inc.All Rights Reserved.
+#pragma once
+
+#include "Containers/HashTable.h"
+
+namespace Chaos::Private
+{
+	// Default ID traits for use with THashMappedArray that works for all types that can be cast to int32
+	template<typename TIDType>
+	struct THashMappedArrayIDTraits
+	{
+		using FIDType = TIDType;
+
+		// Hash the ID to a 32 bit unsigned int for use with FHashTable
+		static uint32 HashID(const FIDType& ID)
+		{
+			return MurmurFinalize32(uint32(ID));
+		}
+	};
+
+	// Default Element traits for THashMappedArray that works for all types that have an member variable ID of type TIDType
+	template<typename TElementType, typename TIDType>
+	struct THashMappedArrayElementTraits
+	{
+		using FIDType = TIDType;
+		using FElementType = TElementType;
+
+		// Return true if the element is the one with the specified ID
+		static bool Is(const FElementType& Element, const FIDType& ID)
+		{
+			return Element.ID == ID;
+		}
+	};
+
+	/**
+	* A HashMap using FHashTable to index an array of elements of type TElementType, whcih should be uniquely identified by an object of type TIDType.
+	* 
+	* E.g.,
+	*	using FMyDataID = int32;
+	*	struct FMyData
+	*	{
+	*		FMyDataID ID;	// Every FMyData will require a unique ID if using the default THashMappedArrayElementTraits
+	*		float MyValue;
+	*	};
+	* 
+	*	const int32 HashTableSize = 128;				// Must be power of 2
+	*	THashMapedArray<FMyDataID, FMyData> MyDataMap(HashTableSize);
+	* 
+	*	MyDataMap.Add(1, { 1, 1.0 });					// NOTE: ID passed twice. Once for the hash map and once to construct FMyData
+	*	MyDataMap.Emplace(2, 2, 2.0);					// NOTE: ID passed twice. Once for the hash map and once for forwarding args to FMyData
+	* 
+	*	const FMyData* MyData2 = MyDataMap.Find(2);		// MyData2->MyValue == 2.0
+	* 
+	*/
+	template<typename TIDType, typename TElementType>
+	class THashMappedArray
+	{
+	public:
+		using FIDType = TIDType;
+		using FElementType = TElementType;
+		using FIDTraits = THashMappedArrayIDTraits<FIDType>;
+		using FElementTraits = THashMappedArrayElementTraits<FElementType, FIDType>;
+		using FHashType = uint32;
+		using FType = THashMappedArray<FIDType, FElementType>;
+
+		// Initialize the hash table. InHashSize must be a power of two (asserted)
+		THashMappedArray(const int32 InHashSize)
+			: HashTable(InHashSize)
+		{
+		}
+
+		// Clear the hash map and reserve space for the specified number of elements (will not shrink)
+		void Reset(const int32 InReserveElements)
+		{
+			HashTable.Clear();
+			HashTable.Resize(InReserveElements);
+			Elements.Reset(InReserveElements);
+		}
+
+		// Add an element with the specified ID to the map. 
+		FORCEINLINE void Add(const FIDType ID, const FElementType& Element)
+		{
+			checkSlow(Find(ID) == nullptr);
+
+			const int32 Index = Elements.Add(Element);
+			const FHashType Key = HashMapKey(ID);
+
+			HashTable.Add(Key, Index);
+		}
+
+		// Add an element with the specified ID to the map. 
+		// NOTE: since your element type will also need to contain the ID, you usually have to pass the ID twice
+		// to emplace, which is a little annoying, but shouldn't affect much.
+		template <typename... ArgsType>
+		FORCEINLINE void Emplace(const FIDType ID, ArgsType&&... Args)
+		{
+			checkSlow(Find(Key) == nullptr);
+
+			const int32 Index = Elements.Emplace(Forward<ArgsType>(Args)...);
+			const FHashType Key = FIDTraits::HashID(ID);
+
+			HashTable.Add(Key, Index);
+		}
+
+		// Find the element with the specified ID. Roughly O(Max(1,N/M)) for N elements with a hash table of size M
+		const FElementType* Find(const FIDType ID) const
+		{
+			return const_cast<FType*>(this)->Find(ID);
+		}
+
+		// Find the element with the specified ID. Roughly O(Max(1,N/M)) for N elements with a hash table of size M
+		FElementType* Find(const FIDType ID)
+		{
+			const FHashType Key = FIDTraits::HashID(ID);
+			for (uint32 Index = HashTable.First(Key); HashTable.IsValid(Index); Index = HashTable.Next(Index))
+			{
+				if (FElementTraits::Is(Elements[Index], ID))
+				{
+					return &Elements[Index];
+				}
+			}
+			return nullptr;
+		}
+
+	private:
+		FHashTable HashTable;
+		TArray<FElementType> Elements;
+	};
+
+}
\ No newline at end of file
diff --git a/Engine/Source/Runtime/Experimental/Chaos/Public/Chaos/HeightField.h b/Engine/Source/Runtime/Experimental/Chaos/Public/Chaos/HeightField.h
--- a/Engine/Source/Runtime/Experimental/Chaos/Public/Chaos/HeightField.h
+++ b/Engine/Source/Runtime/Experimental/Chaos/Public/Chaos/HeightField.h
@@ -21,6 +21,11 @@ namespace Chaos
 	struct FMTDInfo;
 }
 
+namespace Chaos::Private
+{
+	class FMeshContactGenerator;
+}
+
 namespace Chaos
 {
 	class FHeightField final : public FImplicitObject
@@ -239,6 +244,9 @@ namespace Chaos
 			}
 		}
 
+		// Internal: do not use - this API will change as we optimize mesh collision
+		void CollectTriangles(const FAABB3& QueryBounds, const FRigidTransform3& QueryTransform, const FAABB3& ObjectBounds, Private::FMeshContactGenerator& Collector) const;
+
 		virtual void VisitOverlappingLeafObjectsImpl(
 			const FAABB3& QueryBounds,
 			const FRigidTransform3& ObjectTransform,
@@ -1077,6 +1085,10 @@ namespace Chaos
 		CHAOS_API bool GetGridIntersections(FBounds2D InFlatBounds, TArray<TVec2<int32>>& OutInterssctions) const;
 		CHAOS_API bool GetGridIntersectionsBatch(FBounds2D InFlatBounds, TArray<TVec2<int32>>& OutIntersections, const FAABBVectorized& Bounds) const;
 		
+		// Get the cell range that overlaps the QueryBounds (local space). NOTE: OutEndCell is like an end iterator (i.e., BeginCell==EndCell is an empty range)
+		// NOTE: This is a 2D overlap, ignoring height
+		CHAOS_API bool GetOverlappingCellRange(const FAABB3& QueryBounds, TVec2<int32>& OutBeginCell, TVec2<int32>& OutEndCell) const;
+
 		CHAOS_API FBounds2D GetFlatBounds() const;
 
 		// Grid for queries, faster than bounding volumes for heightfields
