Fix CIS - linux/max build error

diff --git a/Engine/Source/Runtime/Experimental/Chaos/Private/Chaos/Collision/MeshContactGenerator.cpp b/Engine/Source/Runtime/Experimental/Chaos/Private/Chaos/Collision/MeshContactGenerator.cpp
--- a/Engine/Source/Runtime/Experimental/Chaos/Private/Chaos/Collision/MeshContactGenerator.cpp
+++ b/Engine/Source/Runtime/Experimental/Chaos/Private/Chaos/Collision/MeshContactGenerator.cpp
@@ -1,7 +1,6 @@
 // Copyright Epic Games, Inc.All Rights Reserved.
 #include "Chaos/Collision/MeshContactGenerator.h"
 #include "Chaos/DebugDrawQueue.h"
-#include "Chaos/Framework/UncheckedArray.h"
 
 namespace Chaos
 {
diff --git a/Engine/Source/Runtime/Experimental/Chaos/Public/Chaos/Collision/MeshContactGenerator.h b/Engine/Source/Runtime/Experimental/Chaos/Public/Chaos/Collision/MeshContactGenerator.h
--- a/Engine/Source/Runtime/Experimental/Chaos/Public/Chaos/Collision/MeshContactGenerator.h
+++ b/Engine/Source/Runtime/Experimental/Chaos/Public/Chaos/Collision/MeshContactGenerator.h
@@ -6,6 +6,7 @@
 #include "Chaos/Collision/ContactTriangles.h"
 #include "Chaos/CollisionResolutionTypes.h"
 #include "Chaos/Framework/HashMappedArray.h"
+#include "Chaos/Framework/UncheckedArray.h"
 #include "Chaos/Triangle.h"
 
 namespace Chaos::Private
@@ -180,6 +181,17 @@ namespace Chaos::Private
 		// A contact index combined with a flag to indicate if the normal is roughly along the triangle face
 		struct FVertexContactIndex
 		{
+			FVertexContactIndex()
+			{
+			}
+			
+			FVertexContactIndex(const FContactVertexID InID, const int32 InContactIndex, const bool bInIsFaceContact)
+				: ID(InID)
+				, ContactIndex(InContactIndex)
+				, bIsFaceContact(bInIsFaceContact)
+			{
+			}
+
 			FContactVertexID ID;
 			int32 ContactIndex;
 			bool bIsFaceContact;
