Chaos - minor mesh collision optimization - lazy evaluation of triangle normal

diff --git a/Engine/Source/Runtime/Experimental/Chaos/Private/Chaos/Collision/MeshContactGenerator.cpp b/Engine/Source/Runtime/Experimental/Chaos/Private/Chaos/Collision/MeshContactGenerator.cpp
--- a/Engine/Source/Runtime/Experimental/Chaos/Private/Chaos/Collision/MeshContactGenerator.cpp
+++ b/Engine/Source/Runtime/Experimental/Chaos/Private/Chaos/Collision/MeshContactGenerator.cpp
@@ -27,7 +27,11 @@ namespace Chaos::Private
 		Contacts.Reset(InMaxContacts);
 		ContactDatas.Reset(InMaxContacts);
 
-		EdgeTriangleIndicesMap.Reset(InMaxTriangles);
+		// If all edges were shared between 2 triangles we have NumEdges = (3 * NumTriangles) / 2
+		// but not all are shared so lets go with NumEdges = (2 * NumTriangles)
+		EdgeTriangleIndicesMap.Reset(InMaxTriangles * 2);
+
+		// Worst case - assume all contacts are vertex contacts
 		VertexContactIndicesMap.Reset(InMaxContacts);
 	}
 
@@ -52,16 +56,16 @@ namespace Chaos::Private
 			int32 LocalVertexID0, LocalVertexID1;
 			GetTriangleEdgeVerticesAtPosition(
 				ContactPoint.ShapeContactPoints[1],
-				Triangle.Triangle.GetVertex(0), Triangle.Triangle.GetVertex(1), Triangle.Triangle.GetVertex(2),
+				Triangle.GetVertex(0), Triangle.GetVertex(1), Triangle.GetVertex(2),
 				LocalVertexID0, LocalVertexID1,
 				BarycentricTolerance);
 
-			const int32 VertexID0 = (LocalVertexID0 != INDEX_NONE) ? Triangle.VertexIndices[LocalVertexID0] : INDEX_NONE;
-			const int32 VertexID1 = (LocalVertexID1 != INDEX_NONE) ? Triangle.VertexIndices[LocalVertexID1] : INDEX_NONE;
+			const int32 VertexID0 = (LocalVertexID0 != INDEX_NONE) ? Triangle.GetVertexIndex(LocalVertexID0) : INDEX_NONE;
+			const int32 VertexID1 = (LocalVertexID1 != INDEX_NONE) ? Triangle.GetVertexIndex(LocalVertexID1) : INDEX_NONE;
 
 			const int32 NewContactIndex = Contacts.Num();
 
-			const FRealSingle ContactNormalDotTriangleNormal = FRealSingle(FVec3::DotProduct(ContactPoint.ShapeContactNormal, Triangle.Normal));
+			const FRealSingle ContactNormalDotTriangleNormal = FRealSingle(FVec3::DotProduct(ContactPoint.ShapeContactNormal, Triangle.GetNormal()));
 
 			const FReal FaceNormalThreshold = FReal(0.998);	// 3deg
 			const bool bIsFaceContact = (ContactNormalDotTriangleNormal > FaceNormalThreshold);
@@ -72,12 +76,12 @@ namespace Chaos::Private
 				// Edge collisions - if it's a face normal, tell each triangle using the edge it has a face contact
 				if (bIsFaceContact)
 				{
-					Triangles[LocalTriangleIndex].NumFaceEdgeCollisions += 1;
+					Triangles[LocalTriangleIndex].AddFaceEdgeCollision();
 
 					const int32 OtherLocalTriangleIndex = GetOtherTriangleIndexForEdge(LocalTriangleIndex, FContactEdgeID(VertexID0, VertexID1));
 					if (OtherLocalTriangleIndex != INDEX_NONE)
 					{
-						Triangles[OtherLocalTriangleIndex].NumFaceEdgeCollisions += 1;
+						Triangles[OtherLocalTriangleIndex].AddFaceEdgeCollision();
 					}
 				}
 			}
@@ -172,6 +176,7 @@ namespace Chaos::Private
 
 		const int32 LocalTriangleIndex = ContactPointData.GetTriangleIndex();
 		const FTriangleExt& Triangle = Triangles[LocalTriangleIndex];
+		const FVec3& TriangleNormal = Triangle.GetNormal();
 
 		// If we have an edge or vertex contact, make sure that the normal is in a valiid range, based
 		// on the triangles that share that edge or vertex.
@@ -182,24 +187,25 @@ namespace Chaos::Private
 			if (OtherLocalTriangleIndex != INDEX_NONE)
 			{
 				const FTriangleExt& OtherTriangle = Triangles[OtherLocalTriangleIndex];
+				const FVec3& OtherTriangleNormal = OtherTriangle.GetNormal();
 
 				const FReal ContactDotNormal = ContactPointData.GetContactNormalDotTriangleNormal();
-				const FReal MinContactDotNormal = FRealSingle(FVec3::DotProduct(OtherTriangle.Normal, Triangle.Normal));
+				const FReal MinContactDotNormal = FRealSingle(FVec3::DotProduct(OtherTriangleNormal, TriangleNormal));
 				if (ContactDotNormal < MinContactDotNormal)
 				{
 					// We are outside the valid normal range for this edge
 					// Convert the edge collision to a face collision on one of the faces, selected to get the smallest depth
 					FVec3 CorrectedContactNormal;
 					int32 CorrectedTriangleIndex;
-					const FReal OtherContactDotNormal = FVec3::DotProduct(ContactPoint.ShapeContactNormal, OtherTriangle.Normal);
+					const FReal OtherContactDotNormal = FVec3::DotProduct(ContactPoint.ShapeContactNormal, OtherTriangleNormal);
 					if (ContactDotNormal >= OtherContactDotNormal)
 					{
-						CorrectedContactNormal = (ContactDotNormal > -SMALL_NUMBER) ? Triangle.Normal : -Triangle.Normal;
+						CorrectedContactNormal = (ContactDotNormal > -SMALL_NUMBER) ? TriangleNormal : -TriangleNormal;
 						CorrectedTriangleIndex = LocalTriangleIndex;
 					}
 					else
 					{
-						CorrectedContactNormal = (OtherContactDotNormal > -SMALL_NUMBER) ? OtherTriangle.Normal : -OtherTriangle.Normal;
+						CorrectedContactNormal = (OtherContactDotNormal > -SMALL_NUMBER) ? OtherTriangleNormal : -OtherTriangleNormal;
 						CorrectedTriangleIndex = OtherLocalTriangleIndex;
 					}
 
@@ -228,16 +234,17 @@ namespace Chaos::Private
 						// It does if the contact normal dotted with the edge plane normal is negative for both edge planes on the triangle that use the vertex.
 						const FVec3 EdgeDelta0 = VertexB - VertexA;
 						const FVec3 EdgeDelta1 = VertexC - VertexA;
-						const FReal EdgeSign0 = FVec3::DotProduct(FVec3::CrossProduct(ContactPoint.ShapeContactNormal, VertexB - VertexA), OtherTriangle.Normal);
-						const FReal EdgeSign1 = FVec3::DotProduct(FVec3::CrossProduct(ContactPoint.ShapeContactNormal, VertexC - VertexA), OtherTriangle.Normal);
+						const FVec3& OtherTriangleNormal = OtherTriangle.GetNormal();
+						const FReal EdgeSign0 = FVec3::DotProduct(FVec3::CrossProduct(ContactPoint.ShapeContactNormal, VertexB - VertexA), OtherTriangleNormal);
+						const FReal EdgeSign1 = FVec3::DotProduct(FVec3::CrossProduct(ContactPoint.ShapeContactNormal, VertexC - VertexA), OtherTriangleNormal);
 						if (FMath::Sign(EdgeSign0) == FMath::Sign(EdgeSign1))
 						{
-							const FVec3 Centroid = OtherTriangle.Triangle.GetCentroid();
+							const FVec3 Centroid = OtherTriangle.GetCentroid();
 							const FReal NormalDotCentroid = FVec3::DotProduct(ContactPoint.ShapeContactNormal, Centroid - ContactPoint.ShapeContactPoints[1]);
 							if (NormalDotCentroid > 0)
 							{
-								const FReal OtherContactDotNormal = FVec3::DotProduct(ContactPoint.ShapeContactNormal, OtherTriangle.Normal);
-								const FVec3 CorrectedContactNormal = OtherTriangle.Normal; //(OtherContactDotNormal > -SMALL_NUMBER) ? OtherTriangle.Normal : -OtherTriangle.Normal;
+								const FReal OtherContactDotNormal = FVec3::DotProduct(ContactPoint.ShapeContactNormal, OtherTriangleNormal);
+								const FVec3 CorrectedContactNormal = OtherTriangleNormal;
 
 								ContactPoint.ShapeContactNormal = CorrectedContactNormal;
 								ContactPoint.ShapeContactPoints[0] = ContactPoint.ShapeContactPoints[1] + ContactPoint.Phi * CorrectedContactNormal;
@@ -334,7 +341,7 @@ namespace Chaos::Private
 			for (int32 TriangleIndex = 0; TriangleIndex < Triangles.Num(); ++TriangleIndex)
 			{
 				const FTriangleExt& Triangle = Triangles[TriangleIndex];
-				if (Triangle.VisitIndex == INDEX_NONE)
+				if (Triangle.GetVisitIndex() == INDEX_NONE)
 				{
 					DebugDrawTriangle(ConvexTransform, Triangle, IgnoredColor);
 				}
@@ -342,7 +349,7 @@ namespace Chaos::Private
 			for (int32 TriangleIndex = 0; TriangleIndex < Triangles.Num(); ++TriangleIndex)
 			{
 				const FTriangleExt& Triangle = Triangles[TriangleIndex];
-				if (Triangle.VisitIndex != INDEX_NONE)
+				if (Triangle.GetVisitIndex() != INDEX_NONE)
 				{
 					DebugDrawTriangle(ConvexTransform, Triangle, VisitedColor);
 				}
@@ -358,9 +365,9 @@ namespace Chaos::Private
 		const FReal LineScale = 1;
 		const int8 DrawPriority = 10;
 
-		const FVec3 V0 = ConvexTransform.TransformPosition(Triangle.Triangle.GetVertex(0));
-		const FVec3 V1 = ConvexTransform.TransformPosition(Triangle.Triangle.GetVertex(1));
-		const FVec3 V2 = ConvexTransform.TransformPosition(Triangle.Triangle.GetVertex(2));
+		const FVec3 V0 = ConvexTransform.TransformPosition(Triangle.GetVertex(0));
+		const FVec3 V1 = ConvexTransform.TransformPosition(Triangle.GetVertex(1));
+		const FVec3 V2 = ConvexTransform.TransformPosition(Triangle.GetVertex(2));
 
 		FDebugDrawQueue::GetInstance().DrawDebugLine(V0, V1, Color, false, FRealSingle(Duration), DrawPriority, FRealSingle(LineScale));
 		FDebugDrawQueue::GetInstance().DrawDebugLine(V1, V2, Color, false, FRealSingle(Duration), DrawPriority, FRealSingle(LineScale));
diff --git a/Engine/Source/Runtime/Experimental/Chaos/Public/Chaos/Collision/MeshContactGenerator.h b/Engine/Source/Runtime/Experimental/Chaos/Public/Chaos/Collision/MeshContactGenerator.h
--- a/Engine/Source/Runtime/Experimental/Chaos/Public/Chaos/Collision/MeshContactGenerator.h
+++ b/Engine/Source/Runtime/Experimental/Chaos/Public/Chaos/Collision/MeshContactGenerator.h
@@ -48,9 +48,9 @@ namespace Chaos::Private
 			for (int32 LocalTriangleIndex = 0; LocalTriangleIndex < GetNumTriangles(); ++LocalTriangleIndex)
 			{
 				const FTriangleExt& Triangle = Triangles[LocalTriangleIndex];
-				AddTriangleEdge(LocalTriangleIndex, Triangle.VertexIndices[0], Triangle.VertexIndices[1]);
-				AddTriangleEdge(LocalTriangleIndex, Triangle.VertexIndices[1], Triangle.VertexIndices[2]);
-				AddTriangleEdge(LocalTriangleIndex, Triangle.VertexIndices[2], Triangle.VertexIndices[0]);
+				AddTriangleEdge(LocalTriangleIndex, Triangle.GetVertexIndex(0), Triangle.GetVertexIndex(1));
+				AddTriangleEdge(LocalTriangleIndex, Triangle.GetVertexIndex(1), Triangle.GetVertexIndex(2));
+				AddTriangleEdge(LocalTriangleIndex, Triangle.GetVertexIndex(2), Triangle.GetVertexIndex(0));
 			}
 		}
 
@@ -69,7 +69,7 @@ namespace Chaos::Private
 				if (GetNumTriangleFaceCollisions(LocalTriangleIndex) == 0)
 				{
 					TriangleContactPoints.Reset();
-					TriangleContactGenerator(Triangles[LocalTriangleIndex].Triangle, TriangleContactPoints);
+					TriangleContactGenerator(Triangles[LocalTriangleIndex].GetTriangle(), TriangleContactPoints);
 
 					AddTriangleContacts(TriangleContactPoints, LocalTriangleIndex);
 
@@ -84,7 +84,7 @@ namespace Chaos::Private
 				if (!IsTriangleVisited(LocalTriangleIndex) && (GetNumTriangleFaceCollisions(LocalTriangleIndex) < 3))
 				{
 					TriangleContactPoints.Reset();
-					TriangleContactGenerator(Triangles[LocalTriangleIndex].Triangle, TriangleContactPoints);
+					TriangleContactGenerator(Triangles[LocalTriangleIndex].GetTriangle(), TriangleContactPoints);
 
 					AddTriangleContacts(TriangleContactPoints, LocalTriangleIndex);
 
@@ -107,10 +107,12 @@ namespace Chaos::Private
 		// A triangle plus some extended data and state
 		class FTriangleExt
 		{
+			static constexpr FReal InvalidNormalMarker = std::numeric_limits<FReal>::max();
+
 		public:
 			FTriangleExt(const FTriangle& InTriangle, const int32 InTriangleIndex, const int32 InVertexIndex0, const int32 InVertexIndex1, const int32 InVertexIndex2)
 				: Triangle(InTriangle)
-				, Normal(InTriangle.GetNormal())
+				, Normal(InvalidNormalMarker)
 				, TriangleIndex(InTriangleIndex)
 				, VertexIndices{ InVertexIndex0, InVertexIndex1, InVertexIndex2 }
 				, NumFaceEdgeCollisions(0)
@@ -169,8 +171,73 @@ namespace Chaos::Private
 				return false;
 			}
 
+			const FTriangle& GetTriangle() const
+			{
+				return Triangle;
+			}
+
+			const FVec3& GetVertex(const int32 LocalVertexIndex) const
+			{
+				return Triangle.GetVertex(LocalVertexIndex);
+			}
+
+			int32 GetTriangleIndex() const
+			{
+				return TriangleIndex;
+			}
+
+			int32 GetVertexIndex(const int32 LocalIndex) const
+			{
+				return VertexIndices[LocalIndex];
+			}
+
+			const FVec3& GetNormal() const
+			{
+				if (Normal.X == InvalidNormalMarker)
+				{
+					Normal = Triangle.GetNormal();
+				}
+				return Normal;
+			}
+
+			FVec3 GetCentroid() const
+			{
+				return Triangle.GetCentroid();
+			}
+
+			void SetVisitIndex(const int8 InVisitIndex)
+			{
+				VisitIndex = InVisitIndex;
+			}
+
+			int8 GetVisitIndex() const
+			{
+				return VisitIndex;
+			}
+
+			void SetEnabled(const bool bInEnabled)
+			{
+				bEnabled = bInEnabled;
+			}
+
+			bool GetIsEnabled() const
+			{
+				return bEnabled;
+			}
+
+			void AddFaceEdgeCollision()
+			{
+				++NumFaceEdgeCollisions;
+			}
+
+			int32 GetNumFaceEdgeCollisions() const
+			{
+				return NumFaceEdgeCollisions;
+			}
+
+		private:
 			FTriangle Triangle;
-			FVec3 Normal;
+			mutable FVec3 Normal;
 			int32 TriangleIndex;
 			int32 VertexIndices[3];
 			int8 NumFaceEdgeCollisions;
@@ -247,10 +314,10 @@ namespace Chaos::Private
 
 		int32 GetNumTriangleFaceCollisions(const int32 LocalTriangleIndex) const
 		{
-			const int32 MeshVertexIndex0 = Triangles[LocalTriangleIndex].VertexIndices[0];
-			const int32 MeshVertexIndex1 = Triangles[LocalTriangleIndex].VertexIndices[1];
-			const int32 MeshVertexIndex2 = Triangles[LocalTriangleIndex].VertexIndices[2];
-			int32 NumCollisions = Triangles[LocalTriangleIndex].NumFaceEdgeCollisions;
+			const int32 MeshVertexIndex0 = Triangles[LocalTriangleIndex].GetVertexIndex(0);
+			const int32 MeshVertexIndex1 = Triangles[LocalTriangleIndex].GetVertexIndex(1);
+			const int32 MeshVertexIndex2 = Triangles[LocalTriangleIndex].GetVertexIndex(2);
+			int32 NumCollisions = Triangles[LocalTriangleIndex].GetNumFaceEdgeCollisions();
 			NumCollisions += HasFaceVertexCollision(MeshVertexIndex0) ? 1 : 0;
 			NumCollisions += HasFaceVertexCollision(MeshVertexIndex1) ? 1 : 0;
 			NumCollisions += HasFaceVertexCollision(MeshVertexIndex2) ? 1 : 0;
@@ -259,12 +326,12 @@ namespace Chaos::Private
 
 		bool IsTriangleVisited(const int32 LocalTriangleIndex) const
 		{
-			return Triangles[LocalTriangleIndex].VisitIndex != INDEX_NONE;
+			return Triangles[LocalTriangleIndex].GetVisitIndex() != INDEX_NONE;
 		}
 
 		void SetTriangleVisited(const int32 LocalTriangleIndex, const int8 VisitIndex)
 		{
-			Triangles[LocalTriangleIndex].VisitIndex = VisitIndex;
+			Triangles[LocalTriangleIndex].SetVisitIndex(VisitIndex);
 		}
 
 		int32 GetNumTriangles() const
