Fixed a bug that was preventing the journal cache file from being created on disk

#rnx

- IPlatformFile::OpenWrite requires that all directories in the given path have already been created or the open will fail. So if the open for write does fail we should first make sure that the directories do exist and then try to open the file again.
- Added a log error if we fail to open the file entirely.

diff --git a/Engine/Source/Runtime/Experimental/IoStoreOnDemand/Private/JournaledCache.cpp b/Engine/Source/Runtime/Experimental/IoStoreOnDemand/Private/JournaledCache.cpp
--- a/Engine/Source/Runtime/Experimental/IoStoreOnDemand/Private/JournaledCache.cpp
+++ b/Engine/Source/Runtime/Experimental/IoStoreOnDemand/Private/JournaledCache.cpp
@@ -13,10 +13,11 @@
 #include "HAL/PlatformProcess.h"
 #include "HAL/Runnable.h"
 #include "HAL/RunnableThread.h"
-#include "IO/IoDispatcher.h"
 #include "IO/IoBuffer.h"
+#include "IO/IoDispatcher.h"
 #include "IO/IoHash.h"
 #include "Math/UnrealMath.h"
+#include "Misc/PathViews.h"
 #include "Misc/Paths.h"
 #include "Misc/ScopeRWLock.h"
 #include "Misc/StringBuilder.h"
@@ -469,7 +470,19 @@ void FDiskJournal::OpenJrnFile()
 	GetPath(JrnPath);
 
 	IPlatformFile& Ipf = IPlatformFile::GetPlatformPhysical();
+
 	IFileHandle* Handle = Ipf.OpenWrite(*JrnPath, true, true);
+	if (Handle == nullptr)
+	{	
+		// If the open failed it could be because we still need to create the directory
+		TStringBuilder<64> JrnDir = WriteToString<64>(FPathViews::GetPath(JrnPath));
+		if (Ipf.CreateDirectory(*JrnDir))
+		{
+			Handle = Ipf.OpenWrite(*JrnPath, true, true);
+		}
+	}
+
+	UE_CLOG(Handle == nullptr, LogIas, Error, TEXT("Failed to open '%s' for FDiskJournal"), *JrnPath);
 	JrnHandle.Reset(Handle);
 }
 
