Use the same file handle for reading and writing IAS cache data

#rnx
(git:5ea5a6e - Martin Ridgers - [REDACTED_PHONE]:14:25 +0200)

diff --git a/Engine/Source/Runtime/Experimental/IoStoreOnDemand/Private/JournaledCache.cpp b/Engine/Source/Runtime/Experimental/IoStoreOnDemand/Private/JournaledCache.cpp
--- a/Engine/Source/Runtime/Experimental/IoStoreOnDemand/Private/JournaledCache.cpp
+++ b/Engine/Source/Runtime/Experimental/IoStoreOnDemand/Private/JournaledCache.cpp
@@ -22,6 +22,7 @@
 #include "Misc/ScopeRWLock.h"
 #include "Misc/StringBuilder.h"
 #include "ProfilingDebugging/CpuProfilerTrace.h"
+#include "Tasks/Pipe.h"
 #include "Tasks/Task.h"
 #include "Templates/UniquePtr.h"
 
@@ -591,6 +592,7 @@ private:
 	static_assert(sizeof(FMapEntry) == sizeof(uint64));
 
 	friend bool				LoadCache(FDiskCache&);
+	void					OpenDataFile();
 	using					FDataMap = TMap<uint64, FMapEntry>;
 	void					Spam();
 	uint64					Insert(uint64 DataBase, const FDataEntry& Entry);
@@ -602,7 +604,7 @@ private:
 
 	uint64					MaxDataSize;
 	uint64					DataCursor;
-
+	TUniquePtr<IFileHandle>	DataHandle;
 	uint32					OverRemoval;
 
 	FDiskJournal			Journal;
@@ -619,6 +621,7 @@ FDiskCache::FDiskCache(FString&& Path, uint64 InMaxDataSize, uint32 InJournalSiz
 	MaxDataSize = (MaxDataSize - InJournalSize) & ~((1ull << 20) - 1);
 
 	Reset();
+	OpenDataFile();
 }
 
 ////////////////////////////////////////////////////////////////////////////////
@@ -631,6 +634,15 @@ void FDiskCache::Reset()
 	OverRemoval = 0;
 }
 
+////////////////////////////////////////////////////////////////////////////////
+void FDiskCache::OpenDataFile()
+{
+	IPlatformFile& Ipf = IPlatformFile::GetPlatformPhysical();
+
+	IFileHandle* Handle(Ipf.OpenWrite(*BinPath, true, true));
+	DataHandle.Reset(Handle);
+}
+
 ////////////////////////////////////////////////////////////////////////////////
 FDiskPhrase FDiskCache::OpenPhrase(uint32 DataSize)
 {
@@ -649,9 +661,7 @@ void FDiskCache::ClosePhrase(FDiskPhrase&& Phrase)
 
 	TRACE_CPUPROFILER_EVENT_SCOPE(IasCache::ClosePhrase);
 
-	IPlatformFile& Ipf = IPlatformFile::GetPlatformPhysical();
-	TUniquePtr<IFileHandle> File(Ipf.OpenWrite(*BinPath, true, false));
-	if (!File.IsValid())
+	if (!DataHandle.IsValid())
 	{
 		Phrase.Drop();
 		Journal.ClosePhrase(MoveTemp(Phrase), 0);
@@ -668,9 +678,8 @@ void FDiskCache::ClosePhrase(FDiskPhrase&& Phrase)
 	{
 		TRACE_CPUPROFILER_EVENT_SCOPE(IasCache::DataWrite);
 		const uint8* Buffer = Phrase.GetPhraseData();
-		File->Seek(DataCursor);
-		File->Write(Buffer, WriteSize);
-		File.Reset();
+		DataHandle->Seek(DataCursor);
+		DataHandle->Write(Buffer, WriteSize);
 	}
 
 	Prune(DataCursor, WriteSize);
@@ -702,15 +711,13 @@ bool FDiskCache::Materialize(EntryHandle Handle, FIoBuffer& Out, uint32 Offset)
 
 	ReadSize = FMath::Min<uint32>(uint32(Out.GetSize()), ReadSize);
 
-	IPlatformFile& Ipf = IPlatformFile::GetPlatformPhysical();
-	TUniquePtr<IFileHandle> File(Ipf.OpenRead(*BinPath, false));
-	if (!File.IsValid())
+	if (!DataHandle.IsValid())
 	{
 		return false;
 	}
 
-	File->Seek(Entry.DataCursor + Offset);
-	return File->Read(Out.GetData(), ReadSize);
+	DataHandle->Seek(Entry.DataCursor + Offset);
+	return DataHandle->Read(Out.GetData(), ReadSize);
 }
 
 ////////////////////////////////////////////////////////////////////////////////
@@ -786,20 +793,25 @@ void FDiskCache::Prune(uint64 DataBase, uint32 Size)
 ////////////////////////////////////////////////////////////////////////////////
 int32 FDiskCache::Flush()
 {
+	DataHandle.Reset();
 	int32 Ret = Journal.Flush();
 	Spam();
+	OpenDataFile();
 	return Ret;
 }
 
 ////////////////////////////////////////////////////////////////////////////////
 void FDiskCache::Drop()
 {
+	DataHandle.Reset();
+
 	IPlatformFile& Ipf = IPlatformFile::GetPlatformPhysical();
 	Ipf.DeleteFile(*BinPath);
 
 	Journal.Drop();
 
 	Reset();
+	OpenDataFile();
 }
 
 ////////////////////////////////////////////////////////////////////////////////
@@ -991,14 +1003,13 @@ static bool LoadCache(FDiskCache& DiskCache)
 
 	IPlatformFile& Ipf = IPlatformFile::GetPlatformPhysical();
 
-	const TCHAR* BinPath = *(DiskCache.BinPath);
-	TUniquePtr<IFileHandle> File(Ipf.OpenRead(BinPath, false));
-	if (!File.IsValid())
+	if (!DiskCache.DataHandle.IsValid())
 	{
-		UE_LOG(LogIas, VeryVerbose, TEXT("JournaledCache: unable to open '%s'"), BinPath);
+		UE_LOG(LogIas, VeryVerbose, TEXT("JournaledCache: unable to open '%s'"), *DiskCache.BinPath);
 		return false;
 	}
 
+	IFileHandle* File = DiskCache.DataHandle.Get();
 	if (uint64(File->Size()) > DiskCache.MaxDataSize)
 	{
 		UE_LOG(LogIas, VeryVerbose,
@@ -1010,7 +1021,7 @@ static bool LoadCache(FDiskCache& DiskCache)
 	}
 
 	// Detect data writes that are newer than any journal flushes.
-	auto ReadBinMarker = [File = File.Get()] (uint64 Cursor, uint32& Out)
+	auto ReadBinMarker = [File] (uint64 Cursor, uint32& Out)
 	{
 		if (Cursor + sizeof(Out) > uint64(File->Size()))
 		{
@@ -1328,6 +1339,7 @@ private:
 	static uint64				ReduceKey(const FIoHash& Key);
 	TUniquePtr<FCacheInner>		Cache;
 	FGovernor					Governor;
+	UE::Tasks::FPipe			GetPipe = UE::Tasks::FPipe(TEXT("IasCacheGetPipe"));
 
 	// FRunnable
 	virtual uint32				Run() override;
@@ -1425,7 +1437,7 @@ FJournaledCache::GetRetType	FJournaledCache::Get(
 	const FIoCancellationToken* CancelToken)
 {
 	uint64 InnerKey = ReduceKey(Key);
-	return UE::Tasks::Launch(TEXT("IasCacheGet"), [this, InnerKey, Options, CancelToken] () {
+	return GetPipe.Launch(TEXT("IasCacheGet"), [this, InnerKey, Options, CancelToken] () {
 		LLM_SCOPE_BYTAG(Ias);
 
 		FCacheInner::FEntry Entry = Cache->Get(InnerKey);
