Deleted a few unnecessary indirections

#rnx
(git:3c17201 - Martin Ridgers - [REDACTED_PHONE]:38:47 +0200)

diff --git a/Engine/Source/Runtime/Experimental/IoStoreOnDemand/Private/CoreHttp/Http.cpp b/Engine/Source/Runtime/Experimental/IoStoreOnDemand/Private/CoreHttp/Http.cpp
--- a/Engine/Source/Runtime/Experimental/IoStoreOnDemand/Private/CoreHttp/Http.cpp
+++ b/Engine/Source/Runtime/Experimental/IoStoreOnDemand/Private/CoreHttp/Http.cpp
@@ -1430,17 +1430,17 @@ static int32 DoConnect(FActivity* Activity)
 	}
 
 	// Adjust socket send and recv buffer sizes
-	if (int32 OptValue = Activity->Pool->GetBufferSize(FSocketPool::EDirection::Send); OptValue >= 0)
+	if (int32 OptValue = Pool->GetBufferSize(FSocketPool::EDirection::Send); OptValue >= 0)
 	{
 		setsockopt(Candidate, SOL_SOCKET, SO_SNDBUF, &(char&)OptValue, sizeof(OptValue));
 	}
-	if (int32 OptValue = Activity->Pool->GetBufferSize(FSocketPool::EDirection::Recv); OptValue >= 0)
+	if (int32 OptValue = Pool->GetBufferSize(FSocketPool::EDirection::Recv); OptValue >= 0)
 	{
 		setsockopt(Candidate, SOL_SOCKET, SO_RCVBUF, &(char&)OptValue, sizeof(OptValue));
 	}
 
 	// connect
-	uint32 IpAddress = Activity->Pool->GetIpAddress();
+	uint32 IpAddress = Pool->GetIpAddress();
 	if (IpAddress == 0)
 	{
 		Activity_SetError(Activity, "No IP address to connect to");
@@ -1449,7 +1449,7 @@ static int32 DoConnect(FActivity* Activity)
 
 	sockaddr_in AddrInet = { sizeof(sockaddr_in) };
 	AddrInet.sin_family = AF_INET;
-	AddrInet.sin_port = htons(uint16(Activity->Pool->GetPort()));
+	AddrInet.sin_port = htons(uint16(Pool->GetPort()));
 	memcpy(&(AddrInet.sin_addr), &IpAddress, sizeof(IpAddress));
 	{
 		int Result = connect(Candidate, &(sockaddr&)AddrInet, sizeof(AddrInet));
