Missing default values for send/recv buffer sizes

#rnx
(git:de3dc1f - Martin Ridgers - [REDACTED_PHONE]:54:23 +0200)

diff --git a/Engine/Source/Runtime/Experimental/IoStoreOnDemand/Private/CoreHttp/Http.cpp b/Engine/Source/Runtime/Experimental/IoStoreOnDemand/Private/CoreHttp/Http.cpp
--- a/Engine/Source/Runtime/Experimental/IoStoreOnDemand/Private/CoreHttp/Http.cpp
+++ b/Engine/Source/Runtime/Experimental/IoStoreOnDemand/Private/CoreHttp/Http.cpp
@@ -607,8 +607,8 @@ public:
 private:
 	FAnsiStringView	HostName;
 	uint32			IpAddresses[4] = {};
-	int16			SendBufKb;
-	int16			RecvBufKb;
+	int16			SendBufKb = -1;
+	int16			RecvBufKb = -1;
 	uint16			Port;
 	uint8			LeaseCount = 0;
 	uint8			MaxLeases : 6;
