New IAS cache governor that matches target's write access accounting

#rnx
(git:7811c32 - Martin Ridgers - [REDACTED_PHONE]:20:23 +0200)

diff --git a/Engine/Source/Runtime/Experimental/IoStoreOnDemand/Private/JournaledCache.cpp b/Engine/Source/Runtime/Experimental/IoStoreOnDemand/Private/JournaledCache.cpp
--- a/Engine/Source/Runtime/Experimental/IoStoreOnDemand/Private/JournaledCache.cpp
+++ b/Engine/Source/Runtime/Experimental/IoStoreOnDemand/Private/JournaledCache.cpp
@@ -51,6 +51,19 @@ namespace UE::IO::IAS::JournaledCache
 
 static bool LoadCache(class FDiskCache&);
 
+#if !defined(IAS_HAS_WRITE_COMMIT_THRESHOLD)
+#	define IAS_HAS_WRITE_COMMIT_THRESHOLD 0
+#endif
+
+#if IAS_HAS_WRITE_COMMIT_THRESHOLD
+	int32 GetWriteCommitThreshold();
+#else
+	static int32 GetWriteCommitThreshold()
+	{
+		return 0;
+	}
+#endif
+
 ////////////////////////////////////////////////////////////////////////////////
 struct FDebugCacheEntry
 {
@@ -1134,7 +1147,7 @@ public:
 	uint32			GetDemand() const;
 	FEntry			Get(uint64 Key) const;
 	bool			Put(uint64 Key, FIoBuffer& Data);
-	int32			Flush(int32 Allowance);
+	uint32			Flush();
 	uint32			WriteMemToDisk(int32 Allowance);
 	uint32			DebugVisit(void* Param, FDebugCacheEntry::Callback* Callback);
 
@@ -1143,8 +1156,6 @@ private:
 	FMemCache		MemCache;
 	FDiskCache		DiskCache;
 	std::atomic_int	Demand;
-	uint32			FlushIndex = 0;
-	uint32			FlushPeriod = 4;
 };
 
 ////////////////////////////////////////////////////////////////////////////////
@@ -1170,7 +1181,6 @@ bool FCache::FEntry::Materialize(FIoBuffer& Out, uint32 Offset)
 FCache::FCache(FConfig&& Config)
 : MemCache(Config.MemoryQuota)
 , DiskCache(MoveTemp(Config.Path), Config.DiskQuota, Config.JournalQuota)
-, FlushPeriod(Config.JournalFlushInterval)
 {
 	if (Config.DropCache)
 	{
@@ -1184,7 +1194,6 @@ void FCache::Reset()
 	FWriteAccess _(Lock);
 	MemCache.Reset();
 	DiskCache.Reset();
-	FlushIndex = 0;
 }
 
 ////////////////////////////////////////////////////////////////////////////////
@@ -1236,27 +1245,9 @@ bool FCache::Put(uint64 Key, FIoBuffer& Data)
 }
 
 ////////////////////////////////////////////////////////////////////////////////
-int32 FCache::Flush(int32 Allowance)
+uint32 FCache::Flush()
 {
-	if (FlushPeriod > 1)
-	{
-		FlushIndex += 1;
-		if (FlushIndex >= FlushPeriod)
-		{
-			FlushIndex -= FlushPeriod;
-			FWriteAccess _(Lock);
-			return Allowance - DiskCache.Flush();
-		}
-
-		if (MemCache.GetUsed() == 0)
-		{
-			FlushIndex -= (FlushIndex == 1);
-			FWriteAccess _(Lock);
-			return Allowance - DiskCache.Flush();
-		}
-	}
-
-	return WriteMemToDisk(Allowance);
+	return DiskCache.Flush();
 }
 
 ////////////////////////////////////////////////////////////////////////////////
@@ -1287,15 +1278,11 @@ uint32 FCache::WriteMemToDisk(int32 Allowance)
 		 * we can re-add leftover peeled items back to mem-cache? */
 	}
 
-	if (FlushPeriod <= 1)
-	{
-		Allowance -= DiskCache.Flush();
-	}
 
 	uint32 NewDemand = MemCache.GetDemand();
 	Demand.store(NewDemand, std::memory_order_relaxed);
 
-	return Allowance - MemCacheSize;
+	return MemCacheSize;
 }
 
 ////////////////////////////////////////////////////////////////////////////////
@@ -1308,11 +1295,140 @@ uint32 FCache::DebugVisit(void* Param, FDebugCacheEntry::Callback* Callback)
 	return Count;
 }
 
+
+
+// {{{1 governor ...............................................................
+
+////////////////////////////////////////////////////////////////////////////////
+class FGovernor
+{
+public:
+			FGovernor() = default;
+	void	Set(uint32 Allowance, uint32 Ops, uint32 Seconds);
+	void	SetDemands(uint32 Threshold, uint32 Boost, uint32 SuperBoost);
+	uint32	BeginAllowance(uint32 DemandPercent);
+	bool	EndAllowance(uint32 UnusedAllowance);
+
+private:
+	enum class EState : uint8
+	{
+		Waiting,
+		Rolling,
+	};
+
+	void	Set(uint32 Allowance, uint32 Ops, uint32 Seconds, int64 CycleFreq);
+	uint32	BeginInternal(uint32 Demand, int64 Cycle);
+	int64	FlushInterval;
+	int64	PrevCycles;
+	uint32	RunOff = 0;
+	uint32	OpCount = 0;
+	uint32	MaxOpCount;
+	uint32	AllowanceInterval;
+	uint8	DemandThreshold = 30;
+	uint8	DemandBoost = 60;
+	uint8	DemandSuperBoost = 87;
+	uint8	VerySuperBoost = 0;
+	EState	State = EState::Waiting;
+};
+
+////////////////////////////////////////////////////////////////////////////////
+void FGovernor::Set(uint32 Allowance, uint32 Ops, uint32 Seconds)
+{
+	int64 CycleFreq	= int64(1.0 / FPlatformTime::GetSecondsPerCycle());
+	return Set(Allowance, Ops, Seconds, CycleFreq);
+}
+
+////////////////////////////////////////////////////////////////////////////////
+void FGovernor::Set(uint32 Allowance, uint32 Ops, uint32 Seconds, int64 CycleFreq)
+{
+	int32 CommitBufferSize = JournaledCache::GetWriteCommitThreshold();
+	if (CommitBufferSize == 0)
+	{
+		AllowanceInterval = Allowance / Ops;
+		FlushInterval = (CycleFreq * Seconds) / Ops;
+		MaxOpCount = 4;
+		return;
+	}
+
+	int32 BlockCount = Allowance / CommitBufferSize;
+	int32 CommitOpCost = BlockCount * 3;
+	MaxOpCount = (Ops - CommitOpCost) / BlockCount;
+
+	AllowanceInterval = CommitBufferSize / MaxOpCount;
+	FlushInterval = (Seconds * CycleFreq) / (BlockCount * (MaxOpCount - 1));
+}
+
+////////////////////////////////////////////////////////////////////////////////
+void FGovernor::SetDemands(uint32 Threshold, uint32 Boost, uint32 SuperBoost)
+{
+	DemandThreshold = uint8(Threshold);
+	DemandBoost = uint8(Boost);
+	DemandSuperBoost = uint8(SuperBoost);
+}
+
+////////////////////////////////////////////////////////////////////////////////
+uint32 FGovernor::BeginInternal(uint32 Demand, int64 Cycles)
+{
+	int64 Interval = FlushInterval;
+	Interval >>= int32(Demand >= DemandBoost);
+	Interval >>= int32(Demand >= DemandSuperBoost) << VerySuperBoost;
+	Interval <<= int32(Demand <= DemandThreshold);
+
+	int64 Delta = Cycles - PrevCycles;
+	if (Delta <= Interval)
+	{
+		return 0;
+	}
+
+	do { Delta -= Interval; } while (Delta > Interval);
+	PrevCycles = Cycles - Delta;
+
+	OpCount++;
+
+	return AllowanceInterval + RunOff;
+}
+
+////////////////////////////////////////////////////////////////////////////////
+uint32 FGovernor::BeginAllowance(uint32 DemandPercent)
+{
+	if (State == EState::Rolling)
+	{
+		int64 Cycles = FPlatformTime::Cycles64();
+		return BeginInternal(DemandPercent, Cycles);
+	}
+
+	if (DemandPercent < DemandThreshold)
+	{
+		return 0;
+	}
+
+	State = EState::Rolling;
+	PrevCycles = FPlatformTime::Cycles64();
+	OpCount = 1;
+	RunOff = 0;
+	return AllowanceInterval;
+}
+
+////////////////////////////////////////////////////////////////////////////////
+bool FGovernor::EndAllowance(uint32 UnusedAllowance)
+{
+	RunOff = UnusedAllowance;
+
+	if (OpCount >= MaxOpCount)
+	{
+		State = EState::Waiting;
+		return true;
+	}
+
+	return false;
+}
+
 // }}}
 
 } // namespace UE::IO::IAS::JournaledCache
 
 
+
 namespace UE::IO::IAS {
 
 // {{{1 journaled-cache ........................................................
@@ -1338,8 +1454,8 @@ private:
 	void						Update();
 	static uint64				ReduceKey(const FIoHash& Key);
 	TUniquePtr<FCacheInner>		Cache;
-	FGovernor					Governor;
 	UE::Tasks::FPipe			GetPipe = UE::Tasks::FPipe(TEXT("IasCacheGetPipe"));
+	JournaledCache::FGovernor	Governor;
 
 	// FRunnable
 	virtual uint32				Run() override;
@@ -1370,7 +1486,9 @@ FJournaledCache::FJournaledCache(const FIasCacheConfig& Config)
 	Cache->Load();
 
 	const FIasCacheConfig::FRate& WriteRate = Config.WriteRate;
+	const FIasCacheConfig::FDemand& Demand = Config.Demand;
 	Governor.Set(WriteRate.Allowance, WriteRate.Ops, WriteRate.Seconds);
+	Governor.SetDemands(Demand.Threshold, Demand.Boost, Demand.SuperBoost);
 
 	StartThread();
 }
@@ -1392,17 +1510,26 @@ void FJournaledCache::Update()
 {
 	LLM_SCOPE_BYTAG(Ias);
 
-	int32 WriteAllowance = Governor.TickAllowance();
-	if (!WriteAllowance)
+	uint32 Demand = Cache->GetDemand();
+	uint32 Allowance = Governor.BeginAllowance(Demand);
+	if (Allowance == 0)
 	{
 		return;
 	}
 
-	TRACE_CPUPROFILER_EVENT_SCOPE(IasCache::Update);
+	do
+	{
+		TRACE_CPUPROFILER_EVENT_SCOPE(IasCache::Update);
 
-	int32 AllowanceUsed = Cache->Flush(WriteAllowance);
+		uint32 AllowanceUsed = Cache->WriteMemToDisk(Allowance);
+		uint32 Unused = Allowance - AllowanceUsed;
 
-	Governor.Return(WriteAllowance - AllowanceUsed);
+		if (Governor.EndAllowance(Unused))
+		{
+			Cache->Flush();
+		}
+	}
+	while (false);
 }
 
 ////////////////////////////////////////////////////////////////////////////////
