Moved resolve handling to the connection pool

#rnx
(git:ad4a235 - Martin Ridgers - [REDACTED_PHONE]:09:31 +0200)

diff --git a/Engine/Source/Runtime/Experimental/IoStoreOnDemand/Private/CoreHttp/Http.cpp b/Engine/Source/Runtime/Experimental/IoStoreOnDemand/Private/CoreHttp/Http.cpp
--- a/Engine/Source/Runtime/Experimental/IoStoreOnDemand/Private/CoreHttp/Http.cpp
+++ b/Engine/Source/Runtime/Experimental/IoStoreOnDemand/Private/CoreHttp/Http.cpp
@@ -598,6 +598,8 @@ public:
 	bool			AddIpAddress(uint32 Address);
 	void			SetBufferSize(EDirection Dir, int32 Size);
 	int32			GetBufferSize(EDirection Dir) const;
+	int32			IsResolved() const;
+	int32			ResolveHostName();
 	uint32			GetIpAddress() const	{ return IpAddresses[0]; }
 	FAnsiStringView	GetHostName() const		{ return HostName; }
 	uint32			GetPort() const			{ return Port; }
@@ -704,6 +706,81 @@ int32 FSocketPool::GetBufferSize(EDirection Dir) const
 	return int32((Dir == EDirection::Send) ? SendBufKb : RecvBufKb) << 10;
 }
 
+////////////////////////////////////////////////////////////////////////////////
+int32 FSocketPool::IsResolved() const
+{
+	switch (IpAddresses[0])
+	{
+	case 0:  return 0;
+	case 1:  return -1;
+	default: return 1;
+	}
+}
+
+////////////////////////////////////////////////////////////////////////////////
+int32 FSocketPool::ResolveHostName()
+{
+	// todo: GetAddrInfoW() for async resolve on Windows
+
+	TRACE_CPUPROFILER_EVENT_SCOPE(IasHttp::PoolResolve);
+
+	IpAddresses[0] = 1;
+
+	addrinfo* Info = nullptr;
+	ON_SCOPE_EXIT { if (Info != nullptr) freeaddrinfo(Info); };
+
+	const FAnsiStringView& Host = GetHostName();
+
+	addrinfo Hints = {};
+	Hints.ai_family = AF_INET;
+	Hints.ai_socktype = SOCK_STREAM;
+	Hints.ai_protocol = IPPROTO_TCP;
+	auto Result = getaddrinfo(Host.GetData(), nullptr, &Hints, &Info);
+	if (uint32(Result) || Info == nullptr)
+	{
+		IpAddresses[0] = 2;
+		return -1;
+	}
+
+	if (Info->ai_family != AF_INET)
+	{
+		IpAddresses[0] = 2;
+		return -2;
+	}
+
+	uint32 AddressCount = 0;
+	for (const addrinfo* Cursor = Info; Cursor != nullptr; Cursor = Cursor->ai_next)
+	{
+		const auto* AddrInet = (sockaddr_in*)(Cursor->ai_addr);
+		if (AddrInet->sin_family != AF_INET)
+		{
+			continue;
+		}
+
+		uint32 IpAddress = 0;
+		memcpy(&IpAddress, &(AddrInet->sin_addr), sizeof(uint32));
+
+		if (IpAddress == 0)
+		{
+			break;
+		}
+
+		IpAddresses[AddressCount] = IpAddress;
+		if (++AddressCount >= UE_ARRAY_COUNT(IpAddresses))
+		{
+			break;
+		}
+	}
+
+	if (AddressCount > 0)
+	{
+		return AddressCount;
+	}
+
+	IpAddresses[0] = 2;
+	return 0;
+}
+
 
 
 ////////////////////////////////////////////////////////////////////////////////
@@ -847,7 +924,7 @@ static void Activity_Free(FActivity* Activity)
 			Activity->Socket = InvalidSocket;
 		}
 
-		if (Activity->Pool->GetState() == FSocketPool::EState::Resolved)
+		if (Activity->Pool->GetIpAddress() > 0x00ff'ffff)
 		{
 			Activity->Pool->ReturnLease(Activity->Socket);
 		}
@@ -1211,7 +1288,7 @@ static uint64 ReadyCheck(FActivity** Activities, uint32 Num, uint32 TimeoutMs)
 			break;
 
 		case EWait::Pool:
-			if (Activity->Pool->GetState() >= FSocketPool::EState::Resolved)
+			if (Activity->Pool->GetIpAddress() > 0x00ff'ffff)
 			{
 				Activities[i]->SocketWait = EWait::None;
 				Ret |= (1ull << Activity->Slot);
@@ -1280,83 +1357,33 @@ struct FHandlerResult
 ////////////////////////////////////////////////////////////////////////////////
 static int32 DoResolve(FActivity* Activity)
 {
-	TRACE_CPUPROFILER_EVENT_SCOPE(IasHttp::DoResolve);
-
-	// todo: GetAddrInfoW() for async resolve on Windows
-
 	// There could be many activities using the same socket pool. Resolving only
 	// needs to happen once, the first activity in can do the honours. Everyone
 	// else can wait.
 	FSocketPool* Pool = Activity->Pool;
-	switch (Pool->GetState())
+	int32 Result = Pool->IsResolved();
+	if (Result > 0)
 	{
-	case FSocketPool::EState::Error:
-	case FSocketPool::EState::Resolved:
 		Activity->State = FActivity::EState::Connect;
 		return 0;
+	}
 
-	case FSocketPool::EState::Busy:
+	if (Result < 0)
+	{
 		Activity->SocketWait = FActivity::EWait::Pool;
 		Activity->State = FActivity::EState::Connect;
 		return 1;
 	}
 
-	Pool->SetState(FSocketPool::EState::Busy);
-
-	addrinfo* Info = nullptr;
-	ON_SCOPE_EXIT { if (Info != nullptr) freeaddrinfo(Info); };
-
-	const FAnsiStringView& HostName = Pool->GetHostName();
-
-	addrinfo Hints = {};
-	Hints.ai_family = AF_INET;
-	Hints.ai_socktype = SOCK_STREAM;
-	Hints.ai_protocol = IPPROTO_TCP;
-	auto Result = getaddrinfo(HostName.GetData(), nullptr, &Hints, &Info);
-	if (uint32(Result) || Info == nullptr)
-	{
-		Pool->SetState(FSocketPool::EState::Error);
-		Activity_SetError(Activity, "Error encountered resolving");
-		return -1;
-	}
-
-	if (Info->ai_family != AF_INET)
+	// We won! We WON!
+	Result = Pool->ResolveHostName();
+	switch (Result)
 	{
-		Pool->SetState(FSocketPool::EState::Error);
-		Activity_SetError(Activity, "Unexpected address family during resolve");
-		return -1;
+	case 0:  Activity_SetError(Activity, "Unable to resolve host"); return -1;
+	case -1: Activity_SetError(Activity, "Error encountered resolving"); return -1;
+	case -2: Activity_SetError(Activity, "Unexpected address family during resolve"); return -1;
 	}
 
-	uint32 AddressCount = 0;
-	for (const addrinfo* Cursor = Info; Cursor != nullptr; Cursor = Cursor->ai_next)
-	{
-		const auto* AddrInet = (sockaddr_in*)(Cursor->ai_addr);
-		if (AddrInet->sin_family != AF_INET)
-		{
-			continue;
-		}
-
-		uint32 IpAddress = 0;
-		memcpy(&IpAddress, &(AddrInet->sin_addr), sizeof(uint32));
-
-		if (IpAddress == 0)
-		{
-			break;
-		}
-
-		if (!Pool->AddIpAddress(IpAddress))
-		{
-			break;
-		}
-
-		++AddressCount;
-	}
-
-	auto NextState = AddressCount
-		? FSocketPool::EState::Resolved
-		: FSocketPool::EState::Error;
-	Pool->SetState(NextState);
-
 	Activity->State = FActivity::EState::Connect;
 	return 0;
 }
@@ -1367,12 +1394,14 @@ static int32 DoConnect(FActivity* Activity)
 	TRACE_CPUPROFILER_EVENT_SCOPE(IasHttp::DoConnect);
 
 	FSocketPool* Pool = Activity->Pool;
-	if (Pool->GetState() == FSocketPool::EState::Error)
+	check(Pool->IsResolved() > 0);
+
+	uint32 IpAddress = Pool->GetIpAddress();
+	if (IpAddress <= 0x00ff'ffff)
 	{
-		Activity_SetError(Activity, "Unable to resolve host");
+		Activity_SetError(Activity, "Unresolved host");
 		return -1;
 	}
-	check(Pool->GetState() == FSocketPool::EState::Resolved);
 
 	// Claim an existing socket from the pool.
 	SocketType Candidate;
@@ -1440,13 +1469,6 @@ static int32 DoConnect(FActivity* Activity)
 	}
 
 	// connect
-	uint32 IpAddress = Pool->GetIpAddress();
-	if (IpAddress == 0)
-	{
-		Activity_SetError(Activity, "No IP address to connect to");
-		return -1;
-	}
-
 	sockaddr_in AddrInet = { sizeof(sockaddr_in) };
 	AddrInet.sin_family = AF_INET;
 	AddrInet.sin_port = htons(uint16(Pool->GetPort()));
