Removed now-unused JournalInterval config

#rnx
(git:20a927a - Martin Ridgers - [REDACTED_PHONE]:25:03 +0200)

diff --git a/Engine/Source/Runtime/Experimental/IoStoreOnDemand/Private/IasCache.h b/Engine/Source/Runtime/Experimental/IoStoreOnDemand/Private/IasCache.h
--- a/Engine/Source/Runtime/Experimental/IoStoreOnDemand/Private/IasCache.h
+++ b/Engine/Source/Runtime/Experimental/IoStoreOnDemand/Private/IasCache.h
@@ -53,7 +53,6 @@ struct FIasCacheConfig
 	uint64		DiskQuota = 512ull << 20;
 	uint32		MemoryQuota = 2 << 20;
 	uint32		JournalQuota = 4 << 20; // description in JournalCache.cpp
-	uint32		JournalFlushInterval = 4;
 	FRate		WriteRate;
 	FDemand		Demand;
 	bool		DropCache = false;
diff --git a/Engine/Source/Runtime/Experimental/IoStoreOnDemand/Private/IoStoreOnDemand.cpp b/Engine/Source/Runtime/Experimental/IoStoreOnDemand/Private/IoStoreOnDemand.cpp
--- a/Engine/Source/Runtime/Experimental/IoStoreOnDemand/Private/IoStoreOnDemand.cpp
+++ b/Engine/Source/Runtime/Experimental/IoStoreOnDemand/Private/IoStoreOnDemand.cpp
@@ -210,7 +210,6 @@ static FIasCacheConfig GetIasCacheConfig(const TCHAR* CommandLine)
 	GetConfigInt(TEXT("DiskQuota"),				Ret.DiskQuota);
 	GetConfigInt(TEXT("MemoryQuota"),			Ret.MemoryQuota);
 	GetConfigInt(TEXT("JournalQuota"),			Ret.JournalQuota);
-	GetConfigInt(TEXT("JournalFlushInterval"),	Ret.JournalFlushInterval);
 	GetConfigInt(TEXT("DemandThreshold"),		Ret.Demand.Threshold);
 	GetConfigInt(TEXT("DemandBoost"),			Ret.Demand.Boost);
 	GetConfigInt(TEXT("DemandSuperBoost"),		Ret.Demand.SuperBoost);
diff --git a/Engine/Source/Runtime/Experimental/IoStoreOnDemand/Private/JournaledCache.cpp b/Engine/Source/Runtime/Experimental/IoStoreOnDemand/Private/JournaledCache.cpp
--- a/Engine/Source/Runtime/Experimental/IoStoreOnDemand/Private/JournaledCache.cpp
+++ b/Engine/Source/Runtime/Experimental/IoStoreOnDemand/Private/JournaledCache.cpp
@@ -38,10 +38,6 @@ detected. For example, assuming an average size of cache items of 2KiB (very
 conservative), a 512MiB cache can hold 256K items. Journal entries are 16
 bytes, so a 256K * 16 is sufficient; 256K * 16 = 4MiB.
 
-JournalFlushInterval dictates how often a write-op allowance is spent on
-writing the journal. For example, a value of 4 would be 3 data writes followed
-by 1 journal write. A value of <=1 will flush immediately (N.B. doing so causes
-two writes per allowed op-throttle).
 */
 
 namespace UE::IO::IAS::JournaledCache
