Added trace counters for IAS' cache writing

#rnx
(git:b8f0f28 - Martin Ridgers - [REDACTED_PHONE]:32:01 +0200)

diff --git a/Engine/Source/Runtime/Experimental/IoStoreOnDemand/Private/JournaledCache.cpp b/Engine/Source/Runtime/Experimental/IoStoreOnDemand/Private/JournaledCache.cpp
--- a/Engine/Source/Runtime/Experimental/IoStoreOnDemand/Private/JournaledCache.cpp
+++ b/Engine/Source/Runtime/Experimental/IoStoreOnDemand/Private/JournaledCache.cpp
@@ -21,6 +21,7 @@
 #include "Misc/Paths.h"
 #include "Misc/ScopeRWLock.h"
 #include "Misc/StringBuilder.h"
+#include "ProfilingDebugging/CountersTrace.h"
 #include "ProfilingDebugging/CpuProfilerTrace.h"
 #include "Tasks/Pipe.h"
 #include "Tasks/Task.h"
@@ -1427,6 +1428,10 @@ namespace UE::IO::IAS {
 
 // {{{1 journaled-cache ........................................................
 
+////////////////////////////////////////////////////////////////////////////////
+TRACE_DECLARE_INT_COUNTER(IasMemDemand, TEXT("Ias/CacheMemDemand"));
+TRACE_DECLARE_INT_COUNTER(IasAllowance, TEXT("Ias/CacheAllowance"));
+TRACE_DECLARE_INT_COUNTER(IasOpCount,   TEXT("Ias/CacheOpCount"));
 
 ////////////////////////////////////////////////////////////////////////////////
 class FJournaledCache
@@ -1511,19 +1516,27 @@ void FJournaledCache::Update()
 		return;
 	}
 
+	TRACE_COUNTER_SET(IasMemDemand, Demand);
+	TRACE_COUNTER_SET(IasAllowance, 0);
+	TRACE_COUNTER_SET(IasAllowance, Allowance);
 	do
 	{
 		TRACE_CPUPROFILER_EVENT_SCOPE(IasCache::Update);
 
 		uint32 AllowanceUsed = Cache->WriteMemToDisk(Allowance);
 		uint32 Unused = Allowance - AllowanceUsed;
+		TRACE_COUNTER_SET(IasAllowance, Unused);
 
+		TRACE_COUNTER_ADD(IasOpCount, 1);
 		if (Governor.EndAllowance(Unused))
 		{
+			TRACE_COUNTER_ADD(IasOpCount, 1); // flush from closing .bin file - we can remove this!
+			TRACE_COUNTER_ADD(IasOpCount, 3); // write-flush-commit from .jrn.
 			Cache->Flush();
 		}
 	}
 	while (false);
+	TRACE_COUNTER_SET(IasAllowance, 0);
 }
 
 ////////////////////////////////////////////////////////////////////////////////
