Added some notes about IAS cache's demand config

#rnx
(git:4220844 - Martin Ridgers - [REDACTED_PHONE]:32:52 +0200)

diff --git a/Engine/Source/Runtime/Experimental/IoStoreOnDemand/Private/JournaledCache.cpp b/Engine/Source/Runtime/Experimental/IoStoreOnDemand/Private/JournaledCache.cpp
--- a/Engine/Source/Runtime/Experimental/IoStoreOnDemand/Private/JournaledCache.cpp
+++ b/Engine/Source/Runtime/Experimental/IoStoreOnDemand/Private/JournaledCache.cpp
@@ -37,6 +37,9 @@ detected. For example, assuming an average size of cache items of 2KiB (very
 conservative), a 512MiB cache can hold 256K items. Journal entries are 16
 bytes, so a 256K * 16 is sufficient; 256K * 16 = 4MiB.
 
+DemandThreshold, DemandBoost, and DemandSuperBoost slow down or speed up the
+writing of data (and journal) to disk depending on how full the memory cache
+is. They are expressed as percentages.
 */
 
 namespace UE::IO::IAS::JournaledCache
