Deleted unused methods

#rnx
(git:2b9cae1 - Martin Ridgers - [REDACTED_PHONE]:10:08 +0200)

diff --git a/Engine/Source/Runtime/Experimental/IoStoreOnDemand/Private/CoreHttp/Http.cpp b/Engine/Source/Runtime/Experimental/IoStoreOnDemand/Private/CoreHttp/Http.cpp
--- a/Engine/Source/Runtime/Experimental/IoStoreOnDemand/Private/CoreHttp/Http.cpp
+++ b/Engine/Source/Runtime/Experimental/IoStoreOnDemand/Private/CoreHttp/Http.cpp
@@ -587,7 +587,6 @@ FMessageBuilder& FMessageBuilder::operator << (FAnsiStringView Lhs)
 class FSocketPool
 {
 public:
-	enum class EState : uint8 { Unresolved, Busy, Resolved, Error };
 	enum class EDirection : uint8 { Send, Recv };
 
 					FSocketPool(FAnsiStringView InHostName, uint32 InPort, uint32 InMaxLeases);
@@ -595,7 +594,6 @@ public:
 	static uint32	GetAllocSize(uint32 MaxLeases);
 	bool			LeaseSocket(SocketType& Out);
 	void			ReturnLease(SocketType Socket);
-	bool			AddIpAddress(uint32 Address);
 	void			SetBufferSize(EDirection Dir, int32 Size);
 	int32			GetBufferSize(EDirection Dir) const;
 	int32			IsResolved() const;
@@ -603,8 +601,6 @@ public:
 	uint32			GetIpAddress() const	{ return IpAddresses[0]; }
 	FAnsiStringView	GetHostName() const		{ return HostName; }
 	uint32			GetPort() const			{ return Port; }
-	EState			GetState() const		{ return State; }
-	void			SetState(EState Value)	{ State = Value; }
 
 private:
 	FAnsiStringView	HostName;
@@ -613,8 +609,7 @@ private:
 	int16			RecvBufKb = -1;
 	uint16			Port;
 	uint8			LeaseCount = 0;
-	uint8			MaxLeases : 6;
-	EState			State : 2;
+	uint8			MaxLeases;
 	SocketType		Sockets[1/*...N*/]; // this should be the last member
 };
 
@@ -623,7 +618,6 @@ FSocketPool::FSocketPool(FAnsiStringView InHostName, uint32 InPort, uint32 InMax
 : HostName(InHostName)
 , Port(uint16(InPort))
 , MaxLeases(uint8(InMaxLeases))
-, State(EState::Unresolved)
 {
 	check(MaxLeases == InMaxLeases); // field overflow
 
@@ -677,23 +671,6 @@ void FSocketPool::ReturnLease(SocketType Socket)
 	Sockets[LeaseCount] = Socket;
 }
 
-////////////////////////////////////////////////////////////////////////////////
-bool FSocketPool::AddIpAddress(uint32 Address)
-{
-	for (uint32& Entry : IpAddresses)
-	{
-		if (Entry != 0)
-		{
-			continue;
-		}
-
-		Entry = Address;
-		return true;
-	}
-
-	return false;
-}
-
 ////////////////////////////////////////////////////////////////////////////////
 void FSocketPool::SetBufferSize(EDirection Dir, int32 Size)
 {
