Correct return value if writing to disk was capped by remaining journal space

#rnx

(git:848c388 - Martin Ridgers - [REDACTED_PHONE]:33:19 +0200)

diff --git a/Engine/Source/Runtime/Experimental/IoStoreOnDemand/Private/JournaledCache.cpp b/Engine/Source/Runtime/Experimental/IoStoreOnDemand/Private/JournaledCache.cpp
--- a/Engine/Source/Runtime/Experimental/IoStoreOnDemand/Private/JournaledCache.cpp
+++ b/Engine/Source/Runtime/Experimental/IoStoreOnDemand/Private/JournaledCache.cpp
@@ -1268,14 +1268,15 @@ uint32 FCache::WriteMemToDisk(int32 Allowance)
 			break;
 		}
 	}
-	DiskCache.ClosePhrase(MoveTemp(Phrase));
 
 	if (PeelIndex >= 0)
 	{
 		/* end of journal reached so not all peeled items could be added, may
 		 * we can re-add leftover peeled items back to mem-cache? */
+		MemCacheSize = Phrase.GetDataSize();
 	}
 
+	DiskCache.ClosePhrase(MoveTemp(Phrase));
 
 	uint32 NewDemand = MemCache.GetDemand();
 	Demand.store(NewDemand, std::memory_order_relaxed);
