Creators of connection pools can optionally resolve it prior to use

#rnx
(git:3c0782a - Martin Ridgers - [REDACTED_PHONE]:10:23 +0200)

diff --git a/Engine/Source/Runtime/Experimental/IoStoreOnDemand/Private/CoreHttp/Client.h b/Engine/Source/Runtime/Experimental/IoStoreOnDemand/Private/CoreHttp/Client.h
--- a/Engine/Source/Runtime/Experimental/IoStoreOnDemand/Private/CoreHttp/Client.h
+++ b/Engine/Source/Runtime/Experimental/IoStoreOnDemand/Private/CoreHttp/Client.h
@@ -71,6 +71,7 @@ public:
 							~FConnectionPool();
 							FConnectionPool(FConnectionPool&& Rhs)	{ *this = MoveTemp(Rhs); }
 	FConnectionPool&		operator = (FConnectionPool&& Rhs)		{ Swap(Ptr, Rhs.Ptr); return *this; }
+	bool					Resolve();
 
 private:
 	friend					class FEventLoop;
diff --git a/Engine/Source/Runtime/Experimental/IoStoreOnDemand/Private/CoreHttp/Http.cpp b/Engine/Source/Runtime/Experimental/IoStoreOnDemand/Private/CoreHttp/Http.cpp
--- a/Engine/Source/Runtime/Experimental/IoStoreOnDemand/Private/CoreHttp/Http.cpp
+++ b/Engine/Source/Runtime/Experimental/IoStoreOnDemand/Private/CoreHttp/Http.cpp
@@ -818,6 +818,12 @@ FConnectionPool::~FConnectionPool()
 	}
 }
 
+////////////////////////////////////////////////////////////////////////////////
+bool FConnectionPool::Resolve()
+{
+	return (Ptr->ResolveHostName() > 0);
+}
+
 
 
 // {{{1 activity ...............................................................
