Added debug support for routing IAS traffic through a SOCKS4/5 proxy

#rnx
(git:944f561 - Martin Ridgers - [REDACTED_PHONE]:27:46 +0200)

diff --git a/Engine/Source/Runtime/Experimental/IoStoreOnDemand/Private/CoreHttp/Http.cpp b/Engine/Source/Runtime/Experimental/IoStoreOnDemand/Private/CoreHttp/Http.cpp
--- a/Engine/Source/Runtime/Experimental/IoStoreOnDemand/Private/CoreHttp/Http.cpp
+++ b/Engine/Source/Runtime/Experimental/IoStoreOnDemand/Private/CoreHttp/Http.cpp
@@ -82,6 +82,68 @@ namespace UE::IO::IAS::HTTP
 
 // {{{1 misc ...................................................................
 
+///////////////////////////////////////////////////////////////////////////////
+static uint32	GSocksIpAddress	= 0;
+static int32	GSocksPort		= 1080; // default SOCKS5 port
+static int32	GSocksVersion	= 5;
+
+#if !UE_BUILD_SHIPPING
+static void SetSocksIpAddress(const TCHAR* Value)
+{
+	uint32 IpAddress = 0;
+	uint32 Accumulator = 0;
+	while (true)
+	{
+		uint32 c = *Value++;
+
+		if (c - '0' <= '9' - '0')
+		{
+			Accumulator *= 10;
+			Accumulator += (c - '0');
+			continue;
+		}
+
+		if (c == '.' || c == '\0')
+		{
+			IpAddress <<= 8;
+			IpAddress |= Accumulator;
+			Accumulator = 0;
+			if (c == '\0')
+			{
+				break;
+			}
+			continue;
+		}
+
+		GSocksIpAddress = 0;
+		return;
+	}
+	GSocksIpAddress = htonl(IpAddress);
+}
+
+static FString GSocksIpAddressStr;
+static FAutoConsoleVariableRef CVar_IasSocksIpAddress(
+	TEXT("ias.SocksIp"),
+	GSocksIpAddressStr,
+	TEXT("Routes all IAS HTTP traffic through the given SOCKS5 proxy"),
+	FConsoleVariableDelegate::CreateStatic([] (IConsoleVariable*) {
+		SetSocksIpAddress(*GSocksIpAddressStr);
+	})
+);
+
+static FAutoConsoleVariableRef CVar_IasSocksPort(
+	TEXT("ias.SocksPort"),
+	GSocksPort,
+	TEXT("Port of the SOCKS5 proxy to use")
+);
+
+static FAutoConsoleVariableRef CVar_IasSocksVersion(
+	TEXT("ias.SocksVersion"),
+	GSocksVersion,
+	TEXT("SOCKS proxy protocol version to use")
+);
+#endif
+
 ////////////////////////////////////////////////////////////////////////////////
 template <typename LambdaType>
 static void EnumerateHeaders(FAnsiStringView Headers, LambdaType&& Lambda)
@@ -849,6 +911,7 @@ struct alignas(16) FActivity
 		Build,
 		Resolve,
 		Connect,
+		Socks,
 		Send,
 		RecvMessage,
 		RecvContent,
@@ -1451,10 +1514,19 @@ static int32 DoConnect(FActivity* Activity)
 		setsockopt(Candidate, SOL_SOCKET, SO_RCVBUF, &(char&)OptValue, sizeof(OptValue));
 	}
 
+	uint16 Port = uint16(Pool->GetPort());
+
+	// Redirect connect to a SOCKS proxy if one is set
+	if (GSocksIpAddress != 0)
+	{
+		IpAddress = GSocksIpAddress;
+		Port = uint16(GSocksPort);
+	}
+
 	// connect
 	sockaddr_in AddrInet = { sizeof(sockaddr_in) };
 	AddrInet.sin_family = AF_INET;
-	AddrInet.sin_port = htons(uint16(Pool->GetPort()));
+	AddrInet.sin_port = htons(Port);
 	memcpy(&(AddrInet.sin_addr), &IpAddress, sizeof(IpAddress));
 	{
 		int Result = connect(Candidate, &(sockaddr&)AddrInet, sizeof(AddrInet));
@@ -1467,12 +1539,144 @@ static int32 DoConnect(FActivity* Activity)
 
 	Activity->Socket = Candidate;
 	Activity->SocketWait = FActivity::EWait::Write;
-	Activity->State = FActivity::EState::Send;
-	Activity->StateParam = 0;
 	Candidate = InvalidSocket;
+
+	Activity->StateParam = 0;
+	Activity->State = (GSocksIpAddress != 0)
+		? FActivity::EState::Socks
+		: FActivity::EState::Send;
+
 	return 1;
 }
 
+////////////////////////////////////////////////////////////////////////////////
+static int32 DoSocks4(FActivity* Activity)
+{
+	struct FSocks4Request
+	{
+		uint8	Version = 4;
+		uint8	Command = 1;
+		uint16	Port;
+		uint32	IpAddress;
+	};
+
+	struct FSocks4Reply
+	{
+		uint8	Version;
+		uint8	Code;
+		uint16	Port;
+		uint32	IpAddress;
+	};
+
+	SocketType Socket = Activity->Socket;
+	int32 Result = 0;
+
+	FSocks4Request Request = {
+		.Port		= htons(uint16(Activity->Pool->GetPort())),
+		.IpAddress	= Activity->Pool->GetIpAddress(),
+	};
+	Result = send(Socket, (const char*)&Request, sizeof(Request), MsgFlagType(0));
+	if (Result <= 0)
+	{
+		Activity_SetError(Activity, "Failed sending SOCKS4 request");
+		return -1;
+	}
+
+	FSocks4Reply Reply;
+	do {
+		Result = recv(Socket, (char*)&Reply, sizeof(Reply), MsgFlagType(0));
+	} while ((Result < 0) & (IsSocketResult(EWOULDBLOCK) | IsSocketResult(EINPROGRESS)));
+	if (Result <= 0)
+	{
+		Activity_SetError(Activity, "Failed receiving SOCKS4 reply");
+		return -1;
+	}
+
+	return 0;
+}
+
+////////////////////////////////////////////////////////////////////////////////
+static int32 DoSocks5(FActivity* Activity)
+{
+	SocketType Socket = Activity->Socket;
+	int32 Result;
+
+	// Greeting
+	const char Greeting[] = { 5, 1, 0 };
+	Result = send(Socket, Greeting, sizeof(Greeting), MsgFlagType(0));
+	if (Result != sizeof(Greeting))
+	{
+		Activity_SetError(Activity, "Failed sending SOCKS5 greeting");
+		return -1;
+	}
+
+	// Server auth-choice
+	char ServerChoice[1 + 1];
+	do {
+		Result = recv(Socket, ServerChoice, sizeof(ServerChoice), MsgFlagType(0));
+	} while ((Result < 0) & (IsSocketResult(EWOULDBLOCK) | IsSocketResult(EINPROGRESS)));
+	if (Result != sizeof(ServerChoice))
+	{
+		Activity_SetError(Activity, "Failed getting SOCKS5 server choice");
+		return -1;
+	}
+
+	if (ServerChoice[0] != 0x05 || ServerChoice[1] != 0x00)
+	{
+		Activity_SetError(Activity, "Unable to authenticate with SOCKS");
+		return -1;
+	}
+
+	// Connection request
+	uint32 IpAddress = Activity->Pool->GetIpAddress();
+	uint16 Port = htons(uint16(Activity->Pool->GetPort()));
+	char Request[] = { 5, 1, 0, 1, 0x11,0x11,0x11,0x11, 0x22,0x22 };
+	std::memcpy(Request + 4, &IpAddress, sizeof(IpAddress));
+	std::memcpy(Request + 8, &Port, sizeof(Port));
+	Result = send(Socket, Request, sizeof(Request), MsgFlagType(0));
+	if (Result != sizeof(Request))
+	{
+		Activity_SetError(Activity, "Failed requesting SOCKS5 connection");
+		return -1;
+	}
+
+	// Connect reply
+	char Reply[3 + (1 + 4) + 2];
+	do {
+		Result = recv(Socket, Reply, sizeof(Reply), MsgFlagType(0));
+	} while ((Result < 0) & (IsSocketResult(EWOULDBLOCK) | IsSocketResult(EINPROGRESS)));
+	if (Result != sizeof(Reply))
+	{
+		Activity_SetError(Activity, "Failed SOCKS5 connection");
+		return -1;
+	}
+
+	if (Reply[0] != 0x05 && Reply[1] != 0x00)
+	{
+		Activity_SetError(Activity, "SOCKS5 connection denied");
+		return -1;
+	}
+
+	return 0;
+}
+
+////////////////////////////////////////////////////////////////////////////////
+static int32 DoSocks(FActivity* Activity)
+{
+	switch (GSocksVersion)
+	{
+	case 4: if (DoSocks4(Activity) == -1) return -1; break;
+	case 5: if (DoSocks5(Activity) == -1) return -1; break;
+	default:
+		Activity_SetError(Activity, "Unsupported SOCKS proxy version");
+		return -1;
+	}
+
+	Activity->State = FActivity::EState::Send;
+	Activity->StateParam = 0;
+	return 0;
+}
+
 ////////////////////////////////////////////////////////////////////////////////
 static int32 DoSend(FActivity* Activity)
 {
@@ -2184,6 +2388,10 @@ uint32 FEventLoop::FImpl::Tick(uint32 PollTimeoutMs)
 			--BusyCount;
 			Result = 0;
 			break;
+
+		case FActivity::EState::Socks:
+			Result = DoSocks(Activity);
+			break;
 		}
 
 		if (Result == -1)
