Placate MSVC's static analyzer

#rnx

diff --git a/Engine/Source/Runtime/Experimental/IoStoreOnDemand/Private/CoreHttp/Http.cpp b/Engine/Source/Runtime/Experimental/IoStoreOnDemand/Private/CoreHttp/Http.cpp
--- a/Engine/Source/Runtime/Experimental/IoStoreOnDemand/Private/CoreHttp/Http.cpp
+++ b/Engine/Source/Runtime/Experimental/IoStoreOnDemand/Private/CoreHttp/Http.cpp
@@ -1598,6 +1598,13 @@ static int32 DoSocks4(FActivity* Activity)
 ////////////////////////////////////////////////////////////////////////////////
 static int32 DoSocks5(FActivity* Activity)
 {
+#ifdef _MSC_VER
+	// MSVC's static analysis doesn't see that 'Result' from recv() is checked
+	// to be the exact size of the destination buffer.
+#pragma warning(push)
+#pragma warning(disable : 6385)
+#endif
+
 	SocketType Socket = Activity->Socket;
 	int32 Result;
 
@@ -1658,6 +1665,10 @@ static int32 DoSocks5(FActivity* Activity)
 	}
 
 	return 0;
+
+#ifdef _MSC_VER
+#pragma warning(pop)
+#endif
 }
 
 ////////////////////////////////////////////////////////////////////////////////
@@ -1702,6 +1713,10 @@ static int32 DoSend(FActivity* Activity)
 		SendSize = 2;
 		break;
 		}
+
+	default:
+		check(false);
+		return -1;
 	}
 
 	SendData += Remaining;
diff --git a/Engine/Source/Runtime/Experimental/IoStoreOnDemand/Private/JournaledCache.cpp b/Engine/Source/Runtime/Experimental/IoStoreOnDemand/Private/JournaledCache.cpp
--- a/Engine/Source/Runtime/Experimental/IoStoreOnDemand/Private/JournaledCache.cpp
+++ b/Engine/Source/Runtime/Experimental/IoStoreOnDemand/Private/JournaledCache.cpp
@@ -1326,7 +1326,6 @@ private:
 	uint8	DemandThreshold = 30;
 	uint8	DemandBoost = 60;
 	uint8	DemandSuperBoost = 87;
-	uint8	VerySuperBoost = 0;
 	EState	State = EState::Waiting;
 };
 
@@ -1370,7 +1369,7 @@ uint32 FGovernor::BeginInternal(uint32 Demand, int64 Cycles)
 {
 	int64 Interval = FlushInterval;
 	Interval >>= int32(Demand >= DemandBoost);
-	Interval >>= int32(Demand >= DemandSuperBoost) << VerySuperBoost;
+	Interval >>= int32(Demand >= DemandSuperBoost);
 	Interval <<= int32(Demand <= DemandThreshold);
 
 	int64 Delta = Cycles - PrevCycles;
