MSVC analysis fix.

#rnx

diff --git a/Engine/Source/Runtime/Experimental/IoStoreOnDemand/Private/CoreHttp/Http.cpp b/Engine/Source/Runtime/Experimental/IoStoreOnDemand/Private/CoreHttp/Http.cpp
--- a/Engine/Source/Runtime/Experimental/IoStoreOnDemand/Private/CoreHttp/Http.cpp
+++ b/Engine/Source/Runtime/Experimental/IoStoreOnDemand/Private/CoreHttp/Http.cpp
@@ -1713,16 +1713,12 @@ static int32 DoSend(FActivity* Activity)
 		SendSize = 2;
 		break;
 		}
-
-	default:
-		check(false);
-		return -1;
 	}
 
 	SendData += Remaining;
 	SendSize -= Remaining;
 
-	if (SendSize == 0)
+	if (SendSize == 0 || SendData == nullptr)
 	{
 		// It is expected there will be enough space for a RespInt object
 		Buffer.Reset();
