Iris - Optimized poll frequency update for AArch64.

#rnx

diff --git a/Engine/Source/Runtime/Experimental/Iris/Core/Private/Iris/ReplicationSystem/ObjectPollFrequencyLimiter.cpp b/Engine/Source/Runtime/Experimental/Iris/Core/Private/Iris/ReplicationSystem/ObjectPollFrequencyLimiter.cpp
--- a/Engine/Source/Runtime/Experimental/Iris/Core/Private/Iris/ReplicationSystem/ObjectPollFrequencyLimiter.cpp
+++ b/Engine/Source/Runtime/Experimental/Iris/Core/Private/Iris/ReplicationSystem/ObjectPollFrequencyLimiter.cpp
@@ -96,6 +96,56 @@ void FObjectPollFrequencyLimiter::Update(const FNetBitArrayView& RelevantObjects
 			*ObjectsToPollData = (ObjectsToPoll | DirtyObjectWord) & ObjectsInScopeWord;
 		}
 	}
+#elif PLATFORM_ENABLE_VECTORINTRINSICS == 1 && PLATFORM_ENABLE_VECTORINTRINSICS_NEON
+	{
+		const uint16x8_t AndMask = vdupq_n_u16(0x0180U);
+		const uint8x16_t AllBitsSet = vdupq_n_u8(255U);
+		const int16_t ShiftAmountsData[8] = { -7, -5, -3, -1, 1, 3, 5, 7 };
+		const int16x8_t ShiftAmounts = vld1q_s16(ShiftAmountsData);
+
+		for (uint32 ObjectIndex = 0, ObjectEndIndex = MaxInternalHandle;
+			ObjectIndex <= ObjectEndIndex;
+			ObjectIndex += WordBitCount, CountersData += WordBitCount, FramesBetweenUpdatesData += WordBitCount, ++RelevantObjectsData, ++DirtyObjectsData, ++ObjectsToPollData)
+		{
+			// Skip ranges with no scopable objects.
+			const WordType ObjectsInScopeWord = *RelevantObjectsData;
+			if (!ObjectsInScopeWord)
+			{
+				continue;
+			}
+
+			uint8x16_t Values0 = vld1q_u8(CountersData + 0U);
+			uint8x16_t Values1 = vld1q_u8(CountersData + 16U);
+
+			const uint8x16_t FramesBetweenUpdates0 = vld1q_u8(FramesBetweenUpdatesData + 0U);
+			const uint8x16_t FramesBetweenUpdates1 = vld1q_u8(FramesBetweenUpdatesData + 16U);
+
+			uint8x16_t MaskToUpdate0 = vceqzq_u8(Values0);
+			Values0 = vaddq_u8(Values0, AllBitsSet);
+			Values0 = vbslq_u8(MaskToUpdate0, FramesBetweenUpdates0, Values0);
+
+			uint8x16_t MaskToUpdate1 = vceqzq_u8(Values1);
+			Values1 = vaddq_u8(Values1, AllBitsSet);
+			Values1 = vbslq_u8(MaskToUpdate1, FramesBetweenUpdates1, Values1);
+
+			// Create bitmasks via masking and horizontal add.
+			uint16x8_t MiddleBits0 = vandq_u16(vreinterpretq_u16_u8(MaskToUpdate0), AndMask);
+			uint16x8_t MiddleBits1 = vandq_u16(vreinterpretq_u16_u8(MaskToUpdate1), AndMask);
+
+			uint16x8_t BitsInPlace0 = vshlq_u16(MiddleBits0, ShiftAmounts);
+			uint16x8_t BitsInPlace1 = vshlq_u16(MiddleBits1, ShiftAmounts);
+
+			const WordType ObjectsToPoll0 = vaddvq_u16(BitsInPlace0);
+			const WordType ObjectsToPoll1 = vaddvq_u16(BitsInPlace1);
+
+			vst1q_u8(CountersData + 0U, Values0);
+			vst1q_u8(CountersData + 16U, Values1);
+
+			const WordType ObjectsToPoll = ObjectsToPoll0 | (ObjectsToPoll1 << 16U);
+			const WordType DirtyObjectWord = *DirtyObjectsData;
+			*ObjectsToPollData = (ObjectsToPoll | DirtyObjectWord) & ObjectsInScopeWord;
+		}
+	}
 #else
 	// Slower path using scalar integer instructions.
 	{
