Fixed decompression failures while flush-committing IAS' journal.

#rnx

diff --git a/Engine/Source/Runtime/Experimental/IoStoreOnDemand/Private/JournaledCache.cpp b/Engine/Source/Runtime/Experimental/IoStoreOnDemand/Private/JournaledCache.cpp
--- a/Engine/Source/Runtime/Experimental/IoStoreOnDemand/Private/JournaledCache.cpp
+++ b/Engine/Source/Runtime/Experimental/IoStoreOnDemand/Private/JournaledCache.cpp
@@ -1245,6 +1245,7 @@ bool FCache::Put(uint64 Key, FIoBuffer& Data)
 ////////////////////////////////////////////////////////////////////////////////
 uint32 FCache::Flush()
 {
+	FWriteAccess _(Lock);
 	return DiskCache.Flush();
 }
 
