Moved some test support code out to its own object

#rnx

(git:8cd2643 - Martin Ridgers - [REDACTED_PHONE]:00:22 +0200)

diff --git a/Engine/Source/Runtime/Experimental/IoStoreOnDemand/Private/JournaledCache.cpp b/Engine/Source/Runtime/Experimental/IoStoreOnDemand/Private/JournaledCache.cpp
--- a/Engine/Source/Runtime/Experimental/IoStoreOnDemand/Private/JournaledCache.cpp
+++ b/Engine/Source/Runtime/Experimental/IoStoreOnDemand/Private/JournaledCache.cpp
@@ -1666,30 +1666,37 @@ static uint64 KeyGen(const FIoBuffer& Data)
 	return KeyGen(Data.GetData(), uint32(Data.GetSize()));
 }
 
+////////////////////////////////////////////////////////////////////////////////
+struct FSupport
+{
+	FSupport()
+	{
+		for (uint32 i = 0; i < WorkingSize; i += 8)
+		{
+			*(uint64*)(Working + i) = Mix();
+		}
+	}
+
+	auto DummyData(uint64 Size)
+	{
+		uint64 Offset = Mix() % (WorkingSize - Size);
+		return FIoBuffer(FIoBuffer::Wrap, Working + Offset, Size);
+	}
+
+	uint64				Mix() { return Th *= 0x369dea0f31a53f85ull; }
+	uint64				Th = 0x0'a9e0'493; // prime!
+	const uint64		WorkingSize = 1_Mi;
+	TUniquePtr<uint8[]> WorkingScope = TUniquePtr<uint8[]>(new uint8[WorkingSize]);
+	uint8*				Working = WorkingScope.Get();
+};
+
 ////////////////////////////////////////////////////////////////////////////////
 IOSTOREONDEMAND_API void Tests()
 {
 #if 0
 	using namespace JournaledCache;
 
-	// Some randomness
-	uint64 Th = 0x0'a9e0'493; // prime!
-	auto MixTh = [&] {
-		return Th *= 0x369dea0f31a53f85ull;
-	};
-
-	// Some data
-	uint64 WorkingSize = 1_Mi;
-	TUniquePtr<uint8[]> WorkingScope(new uint8[WorkingSize]);
-	uint8* Working = WorkingScope.Get();
-	for (uint32 i = 0; i < WorkingSize; i += 8)
-	{
-		*(uint64*)(Working + i) = MixTh();
-	}
-	auto DummyData = [&] (uint64 Size) {
-		uint64 Offset = MixTh() % (WorkingSize - Size);
-		return FIoBuffer(FIoBuffer::Wrap, Working + Offset, Size);
-	};
+	FSupport Support;
 
 	// MemCache {{{2
 	{
