Reinstated mem-cache tests

#rnx
(git:b850731 - Martin Ridgers - [REDACTED_PHONE]:01:49 +0200)

diff --git a/Engine/Source/Runtime/Experimental/IoStoreOnDemand/Private/JournaledCache.cpp b/Engine/Source/Runtime/Experimental/IoStoreOnDemand/Private/JournaledCache.cpp
--- a/Engine/Source/Runtime/Experimental/IoStoreOnDemand/Private/JournaledCache.cpp
+++ b/Engine/Source/Runtime/Experimental/IoStoreOnDemand/Private/JournaledCache.cpp
@@ -1691,70 +1691,73 @@ struct FSupport
 };
 
 ////////////////////////////////////////////////////////////////////////////////
-IOSTOREONDEMAND_API void Tests()
+static void MemCacheTests(FSupport& Support)
 {
-#if 0
 	using namespace JournaledCache;
 
-	FSupport Support;
+	struct {
+		int32 Size;
+		int32 Expected;
+	} TestCases[] = {
+		{ 0, 0 },
+		{ 10, 0 },
+		{ 1023, 511 },
+		{ 1024, 1024 },
+		{ 1025, 1024 },
+	};
 
-	// MemCache {{{2
-	{
-		struct {
-			int32 Size;
-			int32 Expected;
-		} TestCases[] = {
-			{ 0, 0 },
-			{ 10, 0 },
-			{ 1023, 511 },
-			{ 1024, 1024 },
-			{ 1025, 1024 },
-		};
+	for (auto& [Size, Expected] : TestCases)
+	{
+		FMemCache MemCache(Size);
 
-		for (auto& [Size, Expected] : TestCases)
-		{
-			FMemCache MemCache(Size);
+		MemCache.Put(0x493, FIoBuffer());
+		MemCache.Put(0x493, Support.DummyData(0));
+		check(MemCache.GetCount() == 0);
 
-			MemCache.Put(0x493, DummyData(0));
-			check(MemCache.GetCount() == 0);
+		MemCache.Put(0x493, Support.DummyData(513));
+		MemCache.Put(0xa9e, Support.DummyData(511));
+		check(MemCache.GetUsed() == Expected);
 
-			MemCache.Put(0x493, DummyData(513));
-			MemCache.Put(0xa9e, DummyData(511));
-			check(MemCache.GetUsed() == Expected);
+		MemCache.Put(0x49e, Support.DummyData(11));
+		Expected = (Expected == 0) ? 0 : (511 + 11);
+		check(MemCache.GetUsed() == Expected);
+	}
 
-			MemCache.Put(0x49e, DummyData(11));
-			Expected = (Expected == 0) ? 0 : (511 + 11);
-			check(MemCache.GetUsed() == Expected);
-		}
+	FMemCache::PeelItems Peeled;
 
-		FMemCache::PeelItems Peeled;
+	FMemCache MemCache(64);
+	MemCache.Put(1, Support.DummyData(1));
+	check(MemCache.Peel(0, Peeled) == 0);
+	check(Peeled.Num() == 0);
+	check(MemCache.Peel(64, Peeled) == 1);
+	check(Peeled.Num() == 1);
+	check(MemCache.GetUsed() == 0);
+	Peeled.Reset();
 
-		FMemCache MemCache(64);
-		MemCache.Put(1, DummyData(1));
-		check(MemCache.Peel(0, Peeled) == 0);
-		check(Peeled.Num() == 0);
-		check(MemCache.Peel(64, Peeled) == 1);
-		check(Peeled.Num() == 1);
-		check(MemCache.GetUsed() == 0);
-		Peeled.Reset();
+	MemCache = FMemCache(64);
+	for (int32 i = 0; i < 64; ++i)
+	{
+		MemCache.Put(i + 1, Support.DummyData(1));
+	}
 
-		MemCache = FMemCache(64);
-		for (int32 i = 0; i < 64; ++i)
-		{
-			MemCache.Put(i + 1, DummyData(1));
-		}
+	check(MemCache.Peel(32, Peeled) == 32);
+	check(Peeled.Num() == 32);
+	check(MemCache.GetUsed() == 32);
+	for (auto& [Key, _] : Peeled)
+	{
+		FIoBuffer Data;
+		check(MemCache.Get(Key) == 0);
+	}
+	Peeled.Reset();
+}
 
-		check(MemCache.Peel(32, Peeled) == 32);
-		check(Peeled.Num() == 32);
-		check(MemCache.GetUsed() == 32);
-		for (auto& [Key, _] : Peeled)
-		{
-			FIoBuffer Data;
-			check(MemCache.Get(Key) == 0);
-		}
-		Peeled.Reset();
-	} // }}}
+////////////////////////////////////////////////////////////////////////////////
+IOSTOREONDEMAND_API void Tests()
+{
+	FSupport Support;
+	MemCacheTests(Support);
 
+#if 0
 	// Cache {{{2
 	{
 		FCache::FConfig Config;
