[NetPhysics] Resimulation bugfixes and optimizations.

- Fix issue where FindValidRewindFrame could return a frame that was earlier than the earliest frame in history.
- Fix issue where RewindData would use the wrong value of saved history frame count, it was clamped at 1 lower than the actual history size.
- Optimize FindValidRewindFrame to only process particles that are marked for resimulation.
- Optimize FindValidRewindFrame to return the requested frame if no better and valid frame was found.

- The changes fixes issue where resimulation could get blocked indefinitely due to the client receiving target states in alternating frames from the server for different object which would never result in a valid rewind frame where all objects have targets available.

diff --git a/Engine/Source/Runtime/Engine/Private/PhysicsEngine/NetworkPhysicsComponent.cpp b/Engine/Source/Runtime/Engine/Private/PhysicsEngine/NetworkPhysicsComponent.cpp
--- a/Engine/Source/Runtime/Engine/Private/PhysicsEngine/NetworkPhysicsComponent.cpp
+++ b/Engine/Source/Runtime/Engine/Private/PhysicsEngine/NetworkPhysicsComponent.cpp
@@ -171,17 +171,13 @@ int32 FNetworkPhysicsCallback::TriggerRewindIfNeeded_Internal(int32 LatestStepCo
 		ResimFrame = (ResimFrame == INDEX_NONE) ? CallbackFrame : FMath::Min(CallbackFrame, ResimFrame);
 	}
 
-#if DEBUG_NETWORK_PHYSICS
-	UE_LOG(LogTemp, Log, TEXT("COMMON | PT | TriggerRewindIfNeeded_Internal | Callbacks Frame = %d"), ResimFrame);
-#endif
-
 	if (RewindData)
 	{
 		if (NetMode == NM_Client)
 		{
 			const int32 ReplicationFrame = RewindData->GetResimFrame();
 
-#if DEBUG_NETWORK_PHYSICS
+#if DEBUG_NETWORK_PHYSICS || DEBUG_REWIND_DATA
 			UE_LOG(LogTemp, Log, TEXT("CLIENT | PT | TriggerRewindIfNeeded_Internal | Replication Frame = %d"), ReplicationFrame);
 #endif
 			ResimFrame = (ResimFrame == INDEX_NONE) ? ReplicationFrame : (ReplicationFrame == INDEX_NONE) ? ResimFrame : FMath::Min(ReplicationFrame, ResimFrame);
@@ -191,7 +187,7 @@ int32 FNetworkPhysicsCallback::TriggerRewindIfNeeded_Internal(int32 LatestStepCo
 		if (ResimFrame != INDEX_NONE)
 		{
 			const int32 ValidFrame = RewindData->FindValidResimFrame(ResimFrame);
-#if DEBUG_NETWORK_PHYSICS
+#if DEBUG_NETWORK_PHYSICS || DEBUG_REWIND_DATA
 			UE_LOG(LogTemp, Log, TEXT("CLIENT | PT | TriggerRewindIfNeeded_Internal | Resim Frame = %d | Valid Frame = %d"), ResimFrame, ValidFrame);
 #endif
 			ResimFrame = ValidFrame;
@@ -576,14 +572,14 @@ void UNetworkPhysicsComponent::CorrectServerToLocalOffset(const int32 LocalToSer
 		int32 ServerToLocalOffset = LocalToServerOffset;
 		for (int32 FrameIndex = 0; FrameIndex < LocalFrames.Num(); ++FrameIndex)
 		{
-#if DEBUG_NETWORK_PHYSICS
+#if DEBUG_NETWORK_PHYSICS || DEBUG_REWIND_DATA
 			UE_LOG(LogTemp, Log, TEXT("CLIENT | GT | CorrectServerToLocalOffset | Server frame = %d | Client Frame = %d"), ServerFrames[FrameIndex], InputFrames[FrameIndex]);
 #endif
 			ServerToLocalOffset = FMath::Min(ServerToLocalOffset, ServerFrames[FrameIndex] - InputFrames[FrameIndex]);
 		}
 
 		GetPlayerController()->SetServerToLocalAsyncPhysicsTickOffset(ServerToLocalOffset);
-#if DEBUG_NETWORK_PHYSICS
+#if DEBUG_NETWORK_PHYSICS || DEBUG_REWIND_DATA
 		UE_LOG(LogTemp, Log, TEXT("CLIENT | GT | CorrectServerToLocalOffset | Server to local offset = %d | Local to server offset = %d"), ServerToLocalOffset, LocalToServerOffset);
 #endif
 	}
diff --git a/Engine/Source/Runtime/Experimental/Chaos/Private/PBDRigidsSolver.cpp b/Engine/Source/Runtime/Experimental/Chaos/Private/PBDRigidsSolver.cpp
--- a/Engine/Source/Runtime/Experimental/Chaos/Private/PBDRigidsSolver.cpp
+++ b/Engine/Source/Runtime/Experimental/Chaos/Private/PBDRigidsSolver.cpp
@@ -1804,6 +1804,11 @@ namespace Chaos
 			const int32 ResimStep = MRewindCallback->TriggerRewindIfNeeded_Internal(LastStep);
 			const int32 NumResimSteps = LastStep - ResimStep + 1;
 
+			if (ResimStep < 0)
+			{
+				return;
+			}
+
 			const bool bEnableNetworkPredictionDebug = IsNetworkPhysicsPredictionEnabled() && CanDebugNetworkPhysicsPrediction();
 			
 #if !(UE_BUILD_SHIPPING || UE_BUILD_TEST)
@@ -1825,13 +1830,17 @@ namespace Chaos
 				}
 			}
 #endif
-			
-			if(ResimStep >= 0 && (ResimStep < LastStep) && NumResimSteps <= MarshallingManager.GetNumHistory_Internal())
+
+			if ((ResimStep < LastStep) && NumResimSteps <= MarshallingManager.GetNumHistory_Internal())
 			{
 				FResimDebugInfo DebugInfo;
 				QUICK_SCOPE_CYCLE_COUNTER(ChaosRewindAndResim);
 				if(MRewindData->RewindToFrame(ResimStep))
 				{
+#if DEBUG_REWIND_DATA
+					UE_LOG(LogTemp, Warning, TEXT("COMMON | PT | ConditionalApplyRewind_Internal | PERFORMING RESIMULATION | Resim From Frame = %d | Num Steps = %d | To Current Frame: %d"), ResimStep, NumResimSteps, CurrentFrame);
+#endif
+
 					GetEvolution()->SetResim(true);
 					CurrentFrame = ResimStep;
 					
@@ -1899,7 +1908,19 @@ namespace Chaos
 					ResimTimer.Stop();
 					MRewindCallback->SetResimDebugInfo_Internal(DebugInfo);
 				}
+#if DEBUG_REWIND_DATA
+				else
+				{
+					UE_LOG(LogTemp, Log, TEXT("COMMON | PT | ConditionalApplyRewind_Internal | Resimulation failed, FRewindData::RewindToFrame returned false | Current Frame = %d | Num Steps = %d | Resim Frame = %d | Last Frame = %d | Rewind History Size = %d"), CurrentFrame, NumResimSteps, ResimStep, LastStep, MarshallingManager.GetNumHistory_Internal());
+				}
+#endif
 			}
+#if DEBUG_REWIND_DATA
+			else
+			{
+				UE_LOG(LogTemp, Log, TEXT("COMMON | PT | ConditionalApplyRewind_Internal | Resimulation failed, invalid rewind frame data | Current Frame = %d | Num Steps = %d | Resim Frame = %d | Last Frame = %d | Rewind History Size = %d"), CurrentFrame, NumResimSteps, ResimStep, LastStep, MarshallingManager.GetNumHistory_Internal());
+			}
+#endif
 		}
 	}
 
diff --git a/Engine/Source/Runtime/Experimental/Chaos/Private/RewindData.cpp b/Engine/Source/Runtime/Experimental/Chaos/Private/RewindData.cpp
--- a/Engine/Source/Runtime/Experimental/Chaos/Private/RewindData.cpp
+++ b/Engine/Source/Runtime/Experimental/Chaos/Private/RewindData.cpp
@@ -265,12 +265,18 @@ bool FRewindData::RewindToFrame(int32 Frame)
 	const int32 EarliestFrame = GetEarliestFrame_Internal();
 	if (Frame < EarliestFrame)
 	{
+#if DEBUG_REWIND_DATA
+		UE_LOG(LogTemp, Log, TEXT("COMMON | PT | RewindToFrame | Failed due to rewind frame earlier than available history | Rewind Frame: %d | Earliest Frame: %d"), Frame, EarliestFrame);
+#endif
 		return false;
 	}
 
 	//If we need to save and we are right on the edge of the buffer, we can't go back to earliest frame
 	if(Frame == EarliestFrame && bNeedsSave && FramesSaved == Managers.Capacity())
 	{
+#if DEBUG_REWIND_DATA
+		UE_LOG(LogTemp, Log, TEXT("COMMON | PT | RewindToFrame | Failed due to rewinding to last available frame and bNeedsSave is set to true"));
+#endif
 		return false;
 	}
 
@@ -504,9 +510,9 @@ FAutoConsoleVariableRef CVarSkipDesyncTest(TEXT("p.SkipDesyncTest"), SkipDesyncT
 
 void FRewindData::AdvanceFrameImp(IResimCacheBase* ResimCache)
 {
-	FramesSaved = FMath::Min(FramesSaved+1,static_cast<int32>(Managers.Capacity()-1));
+	FramesSaved = FMath::Min(FramesSaved + 1, static_cast<int32>(Managers.Capacity()));
 
-	const int32 EarliestFrame = CurFrame - 1 - FramesSaved;
+	const int32 EarliestFrame = CurFrame - FramesSaved;
 	const FFrameAndPhase FrameAndPhase{ CurFrame, FFrameAndPhase::PostCallbacks };
 
 	auto AdvanceHelper = [this, EarliestFrame, FrameAndPhase](auto& DirtyObjects, const auto& DesyncFunc, const auto& AdvanceDirtyFunc)
@@ -832,37 +838,65 @@ void FRewindData::ExtendHistoryWithFrame(const int32 Frame)
 	FramesSaved = FMath::Max(CurFrame - Frame+1, FramesSaved);
 }
 
+CHAOS_API bool ResimFrameValidationLeniency = true;
+FAutoConsoleVariableRef CVarResimFrameValidationLeniency(TEXT("p.ResimFrameValidationLeniency"), ResimFrameValidationLeniency, TEXT("Lenient resim frame validation finds a frame to resim from where the particle that is triggering the resim has a valid target. Setting this to false will require all replicated particles to have a valid target and both valid component input and state history."));
+CHAOS_API bool ResimIncompleteHistory = true;
+FAutoConsoleVariableRef CVarResimIncompleteHistory(TEXT("p.ResimIncompleteHistory"), ResimIncompleteHistory, TEXT("If a valid resim frame can't be found, use the requested resim frame and perform a resimulation with incomplete data."));
+
 int32 FRewindData::FindValidResimFrame(const int32 RequestedFrame)
 {
 	int32 ValidFrame = INDEX_NONE;
-	if (RequestedFrame > BlockResimFrame)
+
+	if (RequestedFrame <= BlockResimFrame)
 	{
-		EnsureIsInPhysicsThreadContext();
+#if DEBUG_REWIND_DATA
+		UE_LOG(LogTemp, Log, TEXT("COMMON | PT | FindValidResimFrame | Resim is blocked | BlockResimFrame: %d | RequestedFrame: %d"), BlockResimFrame, RequestedFrame);
+#endif
+
+		return ValidFrame;
+	}
+
+	EnsureIsInPhysicsThreadContext();
 
-		// First frame of the history datas
-		const int32 EarliestFrame = FMath::Max(GetEarliestFrame_Internal(), BlockResimFrame);
+	// First frame of the history datas
+	const int32 EarliestFrame = FMath::Max(GetEarliestFrame_Internal(), BlockResimFrame + 1);
+	bool bHasTargetHistory = false;
 
-		for (ValidFrame = RequestedFrame; ValidFrame >= EarliestFrame; --ValidFrame)
+	for (ValidFrame = RequestedFrame; ValidFrame > EarliestFrame; ValidFrame--)
+	{
+		bHasTargetHistory = true;
+		for (FDirtyParticleInfo& DirtyParticleInfo : DirtyParticles)
 		{
-			bool bHasTargetHistory = true;
-			for (FDirtyParticleInfo& DirtyParticleInfo : DirtyParticles)
+			FGeometryParticleStateBase& History = DirtyParticleInfo.GetHistory();
+			FGeometryParticleHandle* Handle = DirtyParticleInfo.GetObjectPtr();
+			const bool bResimAsFollower = DirtyParticleInfo.bResimAsFollower;
+
+			// If the particle is not marked for resimulation, don't bother checking for valid target states.
+			if (ResimFrameValidationLeniency && Solver->GetEvolution()->GetIslandManager().GetParticleResimFrame(Handle) == INDEX_NONE)
 			{
-				FGeometryParticleStateBase& History = DirtyParticleInfo.GetHistory();
-				const bool bResimAsFollower = DirtyParticleInfo.bResimAsFollower;
+				continue;
+			}
 
-				const FFrameAndPhase FrameAndPhase{ ValidFrame, FFrameAndPhase::PostPushData };
-				if (const FParticleDynamicMisc* DynamicMisc = History.DynamicsMisc.Read(FrameAndPhase, PropertiesPool))
+#if DEBUG_REWIND_DATA
+			UE_LOG(LogTemp, Log, TEXT("COMMON | PT | FindValidResimFrame | Processing resim particle | History Frame: %d | Total Particle Count: %d"), ValidFrame, DirtyParticles.Num());
+#endif
+
+			const FFrameAndPhase FrameAndPhase{ ValidFrame, FFrameAndPhase::PostPushData };
+			if (const FParticleDynamicMisc* DynamicMisc = History.DynamicsMisc.Read(FrameAndPhase, PropertiesPool))
+			{
+				if (!DynamicMisc->Disabled() && (DynamicMisc->ObjectState() == EObjectStateType::Dynamic) && !History.TargetPositions.IsEmpty() && !History.TargetVelocities.IsEmpty() && !History.TargetStates.IsEmpty())
 				{
-					if (!DynamicMisc->Disabled() && (DynamicMisc->ObjectState() == EObjectStateType::Dynamic) && !History.TargetPositions.IsEmpty() && !History.TargetVelocities.IsEmpty() && !History.TargetStates.IsEmpty())
+					if (bResimAsFollower || History.TargetPositions.IsClean(FrameAndPhase) || History.TargetVelocities.IsClean(FrameAndPhase) || History.TargetStates.IsClean(FrameAndPhase))
 					{
-						if (bResimAsFollower || History.TargetPositions.IsClean(FrameAndPhase) || History.TargetVelocities.IsClean(FrameAndPhase) || History.TargetStates.IsClean(FrameAndPhase))
-						{
-							bHasTargetHistory = false;
-							break;
-						}
+						bHasTargetHistory = false;
+						break;
 					}
 				}
 			}
+		}
+
+		if (!ResimFrameValidationLeniency)
+		{
 			if (bHasTargetHistory)
 			{
 				for (auto& InputsHistory : InputsHistories)
@@ -885,13 +919,23 @@ int32 FRewindData::FindValidResimFrame(const int32 RequestedFrame)
 					}
 				}
 			}
+		}
 
-			if (bHasTargetHistory)
-			{
-				break;
-			}
+		if (bHasTargetHistory)
+		{
+			break;
 		}
 	}
+	
+	if (!bHasTargetHistory)
+	{
+		ValidFrame = ResimIncompleteHistory ? RequestedFrame : INDEX_NONE;
+
+#if DEBUG_REWIND_DATA
+		UE_LOG(LogTemp, Warning, TEXT("COMMON | PT | FindValidResimFrame | No valid resim frame found | RequestedFrame: %d | ValidFrame: %d | EarliestFrame: %d | HasTargetHistory: %d | EarliestHistoryFrame: %d | CurrentFrame: %d | FramesSaved: %d, ResimFrameValidationLeniency:%d"), RequestedFrame, ValidFrame, EarliestFrame, bHasTargetHistory, GetEarliestFrame_Internal(), CurrentFrame(), FramesSaved, ResimFrameValidationLeniency);
+#endif
+	}
+
 	return ValidFrame;
 }
 
diff --git a/Engine/Source/Runtime/Experimental/Chaos/Public/RewindData.h b/Engine/Source/Runtime/Experimental/Chaos/Public/RewindData.h
--- a/Engine/Source/Runtime/Experimental/Chaos/Public/RewindData.h
+++ b/Engine/Source/Runtime/Experimental/Chaos/Public/RewindData.h
@@ -13,6 +13,10 @@
 #define VALIDATE_REWIND_DATA 0
 #endif
 
+#ifndef DEBUG_REWIND_DATA
+#define DEBUG_REWIND_DATA 0
+#endif
+
 #ifndef DEBUG_NETWORK_PHYSICS
 #define DEBUG_NETWORK_PHYSICS 0
 #endif
