Improve analytics reporting of the IAS file cache.

#rnx

- Track the max size of the cache.
- Remove a number of unused items being sent to analytics.
- Enable reporting of the file cache statistics by default.
- Reformatted the analytics code so it is easier to comment out individual items if needed.

diff --git a/Engine/Source/Runtime/Experimental/IoStoreOnDemand/Private/JournaledCache.cpp b/Engine/Source/Runtime/Experimental/IoStoreOnDemand/Private/JournaledCache.cpp
--- a/Engine/Source/Runtime/Experimental/IoStoreOnDemand/Private/JournaledCache.cpp
+++ b/Engine/Source/Runtime/Experimental/IoStoreOnDemand/Private/JournaledCache.cpp
@@ -633,6 +633,8 @@ FDiskCache::FDiskCache(FString&& Path, uint64 InMaxDataSize, uint32 InJournalSiz
 
 	Reset();
 	OpenDataFile();
+	
+	FOnDemandIoBackendStats::Get()->OnCacheSetMaxBytes(MaxDataSize);
 }
 
 ////////////////////////////////////////////////////////////////////////////////
diff --git a/Engine/Source/Runtime/Experimental/IoStoreOnDemand/Private/Statistics.cpp b/Engine/Source/Runtime/Experimental/IoStoreOnDemand/Private/Statistics.cpp
--- a/Engine/Source/Runtime/Experimental/IoStoreOnDemand/Private/Statistics.cpp
+++ b/Engine/Source/Runtime/Experimental/IoStoreOnDemand/Private/Statistics.cpp
@@ -20,7 +20,7 @@ static TAutoConsoleVariable<bool> CVar_ReportHttpAnalytics(
 
 static TAutoConsoleVariable<bool> CVar_ReportCacheAnalytics(
 	TEXT("ias.ReportCacheAnalytics"),
-	false,
+	true,
 	TEXT("Enables reporting statics on our file cache usage to the analytics system"));
 
 ////////////////////////////////////////////////////////////////////////////////
@@ -191,6 +191,7 @@ FCounterAtomicInt		GCachePutCount(TEXT("Ias/CachePutCount"), TraceCounterDisplay
 FCounterAtomicInt		GCachePutExistingCount(TEXT("Ias/CachePutExistingCount"), TraceCounterDisplayHint_None);
 FCounterAtomicInt		GCachePutRejectCount(TEXT("Ias/CachePutRejectCount"), TraceCounterDisplayHint_None);
 FCounterAtomicInt		GCacheCachedBytes(TEXT("Ias/CacheCachedBytes"), TraceCounterDisplayHint_Memory);
+int64					GCacheMaxBytes = 0;
 FCounterAtomicInt		GCachePendingBytes(TEXT("Ias/CachePendingBytes"), TraceCounterDisplayHint_Memory);
 FCounterAtomicInt		GCacheReadBytes(TEXT("Ias/CacheReadBytes"), TraceCounterDisplayHint_Memory);
 FCounterAtomicInt		GCacheRejectBytes(TEXT("Ias/CachePutRejectBytes"), TraceCounterDisplayHint_Memory);
@@ -309,33 +310,33 @@ void FOnDemandIoBackendStats::ReportAnalytics(TArray<FAnalyticsEventAttribute>&
 {
 	if (CVar_ReportHttpAnalytics.GetValueOnAnyThread())
 	{
-		AppendAnalyticsEventAttributeArray(OutAnalyticsArray,
-			TEXT("IasHttpErrorCount"), GHttpErrorCount.Get(),
-			TEXT("IasHttpRetryCount"), GHttpRetryCount.Get(),
-			TEXT("IasHttpGetCount"), GHttpGetCount.Get(),
-			TEXT("IasHttpDownloadedBytes"), GHttpDownloadedBytes.Get(),
-			TEXT("IasHttpDurationMeanAvg"), GHttpAvgDuration.GetMean(),
-			TEXT("IasHttpDurationStdDev"), GHttpAvgDuration.GetDeviation(),
-
-			TEXT("IasHttpDuration0"), GHttpDurationBuckets[0],
-			TEXT("IasHttpDuration1"), GHttpDurationBuckets[1],
-			TEXT("IasHttpDuration2"), GHttpDurationBuckets[2],
-			TEXT("IasHttpDuration3"), GHttpDurationBuckets[3],
-			TEXT("IasHttpDuration4"), GHttpDurationBuckets[4]
+		AppendAnalyticsEventAttributeArray(OutAnalyticsArray
+			,TEXT("IasHttpErrorCount"), GHttpErrorCount.Get()
+			,TEXT("IasHttpRetryCount"), GHttpRetryCount.Get()
+			,TEXT("IasHttpGetCount"), GHttpGetCount.Get()
+			,TEXT("IasHttpDownloadedBytes"), GHttpDownloadedBytes.Get()
+			,TEXT("IasHttpDurationMeanAvg"), GHttpAvgDuration.GetMean()
+			,TEXT("IasHttpDurationStdDev"), GHttpAvgDuration.GetDeviation()
+
+			,TEXT("IasHttpDuration0"), GHttpDurationBuckets[0]
+			,TEXT("IasHttpDuration1"), GHttpDurationBuckets[1]
+			,TEXT("IasHttpDuration2"), GHttpDurationBuckets[2]
+			,TEXT("IasHttpDuration3"), GHttpDurationBuckets[3]
+			,TEXT("IasHttpDuration4"), GHttpDurationBuckets[4]
 		);
 	}
 
 	if (CVar_ReportCacheAnalytics.GetValueOnAnyThread())
 	{
-		AppendAnalyticsEventAttributeArray(OutAnalyticsArray,
-			TEXT("IasCacheErrorCount"), GCacheErrorCount.Get(),
-			TEXT("IasCacheGetCount"), GCacheGetCount.Get(),
-			TEXT("IasCachePutCount"), GCachePutCount.Get(),
-			TEXT("IasCachetRejectCount"), GCachePutRejectCount.Get(),
-
-			TEXT("IasCacheCachedBytes"), GCacheCachedBytes.Get(),
-			TEXT("IasCacheReadBytes"), GCacheReadBytes.Get(),
-			TEXT("IasCacheRejectBytes"), GCacheRejectBytes.Get()
+		AppendAnalyticsEventAttributeArray(OutAnalyticsArray
+
+			,TEXT("IasCacheErrorCount"), GCacheErrorCount.Get()
+
+			,TEXT("IasCacheCachedBytes"), GCacheCachedBytes.Get()
+			,TEXT("IasCacheMaxBytes"), GCacheMaxBytes
+
+			,TEXT("IasCacheReadBytes"), GCacheReadBytes.Get()
+			,TEXT("IasCacheRejectBytes"), GCacheRejectBytes.Get()
 		);
 	}
 }
@@ -404,6 +405,11 @@ void FOnDemandIoBackendStats::OnCachePersistedBytes(uint64 TotalSize)
 	GCacheCachedBytes.Set(TotalSize);
 }
 
+void FOnDemandIoBackendStats::OnCacheSetMaxBytes(uint64 TotalSize)
+{
+	GCacheMaxBytes = TotalSize;
+}
+
 void FOnDemandIoBackendStats::OnHttpEnqueue()
 {
 	GHttpPendingCount.Add(1);
diff --git a/Engine/Source/Runtime/Experimental/IoStoreOnDemand/Private/Statistics.h b/Engine/Source/Runtime/Experimental/IoStoreOnDemand/Private/Statistics.h
--- a/Engine/Source/Runtime/Experimental/IoStoreOnDemand/Private/Statistics.h
+++ b/Engine/Source/Runtime/Experimental/IoStoreOnDemand/Private/Statistics.h
@@ -47,6 +47,7 @@ public:
 	void OnCachePutReject(uint64 DataSize) IAS_STATISTICS_IMPL()
 	void OnCachePendingBytes(uint64 TotalSize) IAS_STATISTICS_IMPL()
 	void OnCachePersistedBytes(uint64 TotalSize) IAS_STATISTICS_IMPL()
+	void OnCacheSetMaxBytes(uint64 TotalSize) IAS_STATISTICS_IMPL()
 
 	void OnHttpEnqueue() IAS_STATISTICS_IMPL()
 	void OnHttpDequeue() IAS_STATISTICS_IMPL()
