Iris - Fix COND_SkipReplay not being allowed to be replicated and dynamic conditions not being initialized.

diff --git a/Engine/Source/Runtime/Experimental/Iris/Core/Private/Iris/ReplicationSystem/Conditionals/ReplicationConditionals.cpp b/Engine/Source/Runtime/Experimental/Iris/Core/Private/Iris/ReplicationSystem/Conditionals/ReplicationConditionals.cpp
--- a/Engine/Source/Runtime/Experimental/Iris/Core/Private/Iris/ReplicationSystem/Conditionals/ReplicationConditionals.cpp
+++ b/Engine/Source/Runtime/Experimental/Iris/Core/Private/Iris/ReplicationSystem/Conditionals/ReplicationConditionals.cpp
@@ -171,6 +171,15 @@ void FReplicationConditionals::InitPropertyCustomConditions(FInternalNetRefIndex
 					const FReplicationStateMemberChangeMaskDescriptor& MemberChangeMaskDescriptor = StateDescriptor->MemberChangeMaskDescriptors[MemberIndex];
 					ConditionalChangeMask.ClearBits(MemberChangeMaskDescriptor.BitOffset, MemberChangeMaskDescriptor.BitCount);
 				}
+
+				if (MemberLifeTimeConditionDescriptor.Condition == COND_Dynamic)
+				{
+					ELifetimeCondition Condition = Tracker->GetDynamicCondition(RepIndex);
+					if (Condition != COND_Dynamic)
+					{
+						SetDynamicCondition(ObjectIndex, RepIndex, Condition);
+					}
+				}
 			}
 		}
 	}
@@ -754,6 +763,7 @@ FReplicationConditionals::FConditionalsMask FReplicationConditionals::GetLifetim
 	ConditionalsMask.SetConditionEnabled(COND_ReplayOrOwner, bIsReplicatingToOwner);
 	ConditionalsMask.SetConditionEnabled(COND_SimulatedOnlyNoReplay, bRoleSimulated);
 	ConditionalsMask.SetConditionEnabled(COND_SimulatedOrPhysicsNoReplay, bRoleSimulated | bRepPhysics);
+	ConditionalsMask.SetConditionEnabled(COND_SkipReplay, true);
 
 	return ConditionalsMask;
 }
