Strategic change for better handling of dirs and files later.

#rnx

diff --git a/Engine/Source/Runtime/Experimental/IoStoreOnDemand/Private/JournaledCache.cpp b/Engine/Source/Runtime/Experimental/IoStoreOnDemand/Private/JournaledCache.cpp
--- a/Engine/Source/Runtime/Experimental/IoStoreOnDemand/Private/JournaledCache.cpp
+++ b/Engine/Source/Runtime/Experimental/IoStoreOnDemand/Private/JournaledCache.cpp
@@ -1447,8 +1447,8 @@ class FJournaledCache
 public:
 	using GetRetType = UE::Tasks::TTask<TIoStatusOr<FIoBuffer>>;
 
-								FJournaledCache(const FIasCacheConfig& Config);
-								~FJournaledCache();
+								FJournaledCache() = default;
+	bool						Initialize(const FIasCacheConfig& Config);
 	virtual bool				ContainsChunk(const FIoHash& Key) const override;
 	virtual GetRetType			Get(const FIoHash& Key, const FIoReadOptions& Options, const FIoCancellationToken* CancellationToken) override;
 	virtual FIoStatus			Put(const FIoHash& Key, FIoBuffer& Data) override;
@@ -1472,7 +1472,7 @@ private:
 };
 
 ////////////////////////////////////////////////////////////////////////////////
-FJournaledCache::FJournaledCache(const FIasCacheConfig& Config)
+bool FJournaledCache::Initialize(const FIasCacheConfig& Config)
 {
 	TStringBuilder<256> CachePath;
 	CachePath << FPaths::ProjectPersistentDownloadDir();
@@ -1496,6 +1496,8 @@ FJournaledCache::FJournaledCache(const FIasCacheConfig& Config)
 	Governor.SetDemands(Demand.Threshold, Demand.Boost, Demand.SuperBoost);
 
 	StartThread();
+
+	return true;
 }
 
 ////////////////////////////////////////////////////////////////////////////////
@@ -1505,11 +1507,6 @@ void FJournaledCache::StartThread()
 	Thread.Reset(Inst);
 }
 
-////////////////////////////////////////////////////////////////////////////////
-FJournaledCache::~FJournaledCache()
-{
-}
-
 ////////////////////////////////////////////////////////////////////////////////
 void FJournaledCache::Update()
 {
@@ -1639,7 +1636,14 @@ uint64 FJournaledCache::ReduceKey(const FIoHash& Key)
 TUniquePtr<IIasCache> MakeIasCache(const FIasCacheConfig& Config)
 {
 	LLM_SCOPE_BYTAG(Ias);
-	return MakeUnique<FJournaledCache>(Config);
+
+	FJournaledCache* Cache = new FJournaledCache();
+	if (Cache->Initialize(Config))
+	{
+		return TUniquePtr<IIasCache>(Cache);
+	}
+
+	return nullptr;
 }
 
 
