Count cache read errors

#rnx

diff --git a/Engine/Source/Runtime/Experimental/IoStoreOnDemand/Private/OnDemandIoDispatcherBackend.cpp b/Engine/Source/Runtime/Experimental/IoStoreOnDemand/Private/OnDemandIoDispatcherBackend.cpp
--- a/Engine/Source/Runtime/Experimental/IoStoreOnDemand/Private/OnDemandIoDispatcherBackend.cpp
+++ b/Engine/Source/Runtime/Experimental/IoStoreOnDemand/Private/OnDemandIoDispatcherBackend.cpp
@@ -1030,6 +1030,10 @@ bool FOnDemandIoBackend::Resolve(FIoRequestImpl* Request)
 				ChunkRequest->bCached = true;
 				return CompleteRequest(ChunkRequest);
 			}
+			else if (Status.Status().GetErrorCode() == EIoErrorCode::ReadError)
+			{
+				FOnDemandIoBackendStats::Get()->OnCacheError();
+			}
 		}
 
 		if (ChunkRequest->CancellationToken.IsCancelled() || IsHttpEnabled() == false)
