Chaos: HashTable implementation for AABBTree payload tracking to reduce memory costs of AABB trees in when used as a local sub-acceleration.

Users of TAABBTree can provide a storage policy that currently only handles storage requirements for the PayloadToInfo mapping store. The main scene
uses a flat array sized to the unique indices of the particles inside it. This is OK for the main scene but when used as a sub-acceleration later in a run
when the unique index counter is very high there's a lot of empty space in the payload map.

Here cluster unions defer to a fast hash table for this mapping rather than a flat array. Because cluster unions are usually small, at least smaller than
the whole scene the lookup in the map shouldn't often chain and should be relatively fast.

diff --git a/Engine/Source/Runtime/Engine/Classes/PhysicsEngine/ClusterUnionComponent.h b/Engine/Source/Runtime/Engine/Classes/PhysicsEngine/ClusterUnionComponent.h
--- a/Engine/Source/Runtime/Engine/Classes/PhysicsEngine/ClusterUnionComponent.h
+++ b/Engine/Source/Runtime/Engine/Classes/PhysicsEngine/ClusterUnionComponent.h
@@ -124,6 +124,9 @@ class UClusterUnionComponent : public UPrimitiveComponent
 {
 	GENERATED_BODY()
 public:
+
+	using FSpatialAcceleration = Chaos::ISpatialAcceleration<FExternalSpatialAccelerationPayload, Chaos::FReal, 3>;
+
 	ENGINE_API UClusterUnionComponent(const FObjectInitializer& ObjectInitializer);
 
 	UFUNCTION(BlueprintCallable, Category="Cluster Union")
@@ -173,6 +176,8 @@ public:
 	// and doesn't need to wait for the next OnSyncBodies.
 	ENGINE_API void ForceRebuildGTParticleGeometry();
 
+	ENGINE_API const FSpatialAcceleration* GetSpatialAcceleration() const;
+
 	friend class UClusterUnionReplicatedProxyComponent;
 	friend class UModularVehicleBaseComponent;
 protected:
@@ -234,7 +239,7 @@ private:
 	FChaosUserData PhysicsUserData;
 
 	// An acceleration structure of all children components managed by the cluster union itself.
-	TUniquePtr<Chaos::ISpatialAcceleration<FExternalSpatialAccelerationPayload, Chaos::FReal, 3>> AccelerationStructure;
+	TUniquePtr<FSpatialAcceleration> AccelerationStructure;
 
 	// Need to handle the fact that this component may or may not be initialized prior to the components referenced in
 	// ClusteredComponentsReferences. This function lets us listen to OnComponentPhysicsStateChanged on the incoming
diff --git a/Engine/Source/Runtime/Engine/Private/PhysicsEngine/ClusterUnionComponent.cpp b/Engine/Source/Runtime/Engine/Private/PhysicsEngine/ClusterUnionComponent.cpp
--- a/Engine/Source/Runtime/Engine/Private/PhysicsEngine/ClusterUnionComponent.cpp
+++ b/Engine/Source/Runtime/Engine/Private/PhysicsEngine/ClusterUnionComponent.cpp
@@ -28,20 +28,39 @@ namespace
 	bool bIncrementalUnionBuild = false;
 	FAutoConsoleVariableRef CVarIncrementalUnionBuild(TEXT("ClusterUnion.IncrementalUnionBuild"), bIncrementalUnionBuild , TEXT("Cvar to incrementally build the union."));
 	
+	template<typename PayloadType>
+	struct TClusterUnionAABBTreeStorageTraits
+	{
+		constexpr static uint32 HashTableBucketsSize = 256;
+
+		using PayloadToInfoType = Chaos::TSQMap<PayloadType, Chaos::FAABBTreePayloadInfo>;
+
+		static void InitPayloadToInfo(PayloadToInfoType& PayloadToInfo)
+		{
+			ensureMsgf(PayloadToInfo.Num() == 0, TEXT("Expected an empty map in InitPayloadToInfo, this will incur a rehash of %d elems"), PayloadToInfo.Num());
+			PayloadToInfo.ResizeHashBuckets(HashTableBucketsSize);
+		}
+	};
+
 	// TODO: Should this be exposed in Chaos instead?
-	using FAccelerationStructure = Chaos::TAABBTree<FExternalSpatialAccelerationPayload, Chaos::TAABBTreeLeafArray<FExternalSpatialAccelerationPayload>>;
+	using FAccelerationStructure = Chaos::TAABBTree<
+		FExternalSpatialAccelerationPayload, 
+		Chaos::TAABBTreeLeafArray<FExternalSpatialAccelerationPayload>,
+		true,
+		Chaos::FReal,
+		TClusterUnionAABBTreeStorageTraits<FExternalSpatialAccelerationPayload>>;
 
 	TUniquePtr<Chaos::ISpatialAcceleration<FExternalSpatialAccelerationPayload, Chaos::FReal, 3>> CreateEmptyAccelerationStructure()
 	{
-		FAccelerationStructure* Structure = new FAccelerationStructure(
-			nullptr,
+		FAccelerationStructure* Structure = new FAccelerationStructure{
+			FAccelerationStructure::EmptyInit{},
 			FAccelerationStructure::DefaultMaxChildrenInLeaf,
 			FAccelerationStructure::DefaultMaxTreeDepth,
 			FAccelerationStructure::DefaultMaxPayloadBounds,
 			FAccelerationStructure::DefaultMaxNumToProcess,
 			true,
 			false
-		);
+		};
 
 		check(Structure != nullptr);
 		return TUniquePtr<Chaos::ISpatialAcceleration<FExternalSpatialAccelerationPayload, Chaos::FReal, 3>>(Structure);
@@ -509,6 +528,11 @@ void UClusterUnionComponent::ForceRebuildGTParticleGeometry()
 	PhysicsProxy->SetGeometry_External(Chaos::FImplicitObjectPtr(NewGeometry), Particles); 
 }
 
+const UClusterUnionComponent::FSpatialAcceleration* UClusterUnionComponent::GetSpatialAcceleration() const
+{
+	return AccelerationStructure.Get();
+}
+
 TArray<UPrimitiveComponent*> UClusterUnionComponent::GetPrimitiveComponents()
 {
 	TArray<UPrimitiveComponent*> PrimitiveComponents;
diff --git a/Engine/Source/Runtime/Engine/Private/PhysicsEngine/Experimental/PhysScene_Chaos.cpp b/Engine/Source/Runtime/Engine/Private/PhysicsEngine/Experimental/PhysScene_Chaos.cpp
--- a/Engine/Source/Runtime/Engine/Private/PhysicsEngine/Experimental/PhysScene_Chaos.cpp
+++ b/Engine/Source/Runtime/Engine/Private/PhysicsEngine/Experimental/PhysScene_Chaos.cpp
@@ -26,6 +26,7 @@
 #include "PhysicsProxy/StaticMeshPhysicsProxy.h"
 #include "Chaos/PendingSpatialData.h"
 #include "Chaos/PhysicsSolverBaseImpl.h"
+#include "Misc/CoreMisc.h"
 
 #if WITH_EDITOR
 #include "Editor.h"
@@ -70,6 +71,53 @@ struct FPendingAsyncPhysicsCommand
 	bool bEnableResim = true;
 };
 
+class FPhysSceneExecHandler : public FSelfRegisteringExec
+{
+	bool Exec_Runtime(UWorld* InWorld, const TCHAR* Cmd, FOutputDevice& Ar) override
+	{
+#if !UE_BUILD_SHIPPING
+		if(FParse::Command(&Cmd, TEXT("LIST SQ")) && InWorld)
+		{
+			FPhysScene_Chaos* Scene = InWorld->GetPhysicsScene();
+
+			if(!Scene)
+			{
+				return false;
+			}
+
+			Ar.Logf(TEXT("----- Begin SQ Listing ----------------------------------------"));
+			Ar.Logf(TEXT("----- World SQ         ----------------------------------------"));
+			Ar.Logf(TEXT("SQ Data for world %s"), *InWorld->GetName());
+			Scene->GetSpacialAcceleration()->DumpStatsTo(Ar);
+
+			Ar.Logf(TEXT("----- Cluster Union SQ ----------------------------------------"));
+
+			for(TObjectIterator<UClusterUnionComponent> Iter; Iter; ++Iter)
+			{
+				UClusterUnionComponent* Union = *Iter;
+				if(Union->GetWorld() != InWorld)
+				{
+					continue;
+				}
+
+				Ar.Logf(TEXT("Inner Acceleration data for Cluster Union Component %s"), *Union->GetName());
+				if(Union->GetSpatialAcceleration())
+				{
+					Union->GetSpatialAcceleration()->DumpStatsTo(Ar);
+				}
+				Ar.Logf(TEXT(""));
+			}
+
+			Ar.Logf(TEXT("----- End SQ Listing   ----------------------------------------"));
+
+			return true;
+		}
+#endif
+
+		return false;
+	}
+};
+static FPhysSceneExecHandler GPhysSceneExecHandler;
 
 class FAsyncPhysicsTickCallback : public Chaos::TSimCallbackObject<
 	Chaos::FSimCallbackNoInput,
diff --git a/Engine/Source/Runtime/Experimental/Chaos/Public/Chaos/AABBTree.h b/Engine/Source/Runtime/Experimental/Chaos/Public/Chaos/AABBTree.h
--- a/Engine/Source/Runtime/Experimental/Chaos/Public/Chaos/AABBTree.h
+++ b/Engine/Source/Runtime/Experimental/Chaos/Public/Chaos/AABBTree.h
@@ -241,7 +241,7 @@ struct TAABBTreeLeafArray : public TBoundsWrapperHelper<TPayloadType, T, bComput
 		this->ComputeBounds(Elems);
 	}
 
-	void GatherElements(TArray<TPayloadBoundsElement<TPayloadType, T>>& OutElements)
+	void GatherElements(TArray<TPayloadBoundsElement<TPayloadType, T>>& OutElements) const
 	{
 		OutElements.Append(Elems);
 	}
@@ -731,9 +731,21 @@ struct DirtyGridHashEntry
 	int32 Count;  // Number of valid entries from Index in FlattenedCellArrayOfDirtyIndices
 };
 
-template <typename TPayloadType, typename TLeafType, bool bMutable = true, typename T = FReal>
+template<typename PayloadType>
+struct TDefaultAABBTreeStorageTraits
+{
+	using PayloadToInfoType = TArrayAsMap<PayloadType, FAABBTreePayloadInfo>;
+
+	static void InitPayloadToInfo(PayloadToInfoType& PayloadToInfo)
+	{}
+};
+
+template <typename TPayloadType, typename TLeafType, bool bMutable = true, typename T = FReal, typename StorageTraits = TDefaultAABBTreeStorageTraits<TPayloadType>>
 class TAABBTree final : public ISpatialAcceleration<TPayloadType, T, 3> 
 {
+private:
+	using FElement = TPayloadBoundsElement<TPayloadType, T>;
+	using FNode = TAABBTreeNode<T>;
 public:
 	using PayloadType = TPayloadType;
 	static constexpr int D = 3;
@@ -759,6 +771,8 @@ public:
 		, bBuildOverlapCache(true)		
 	{
 		GetCVars();
+
+		StorageTraits::InitPayloadToInfo(PayloadToInfo);
 	}
 
 	virtual void Reset() override
@@ -837,12 +851,38 @@ public:
 	{
 		if (bInUseDirtyTree)
 		{
-			DirtyElementTree = TUniquePtr<TAABBTree<TPayloadType, TLeafType, bMutable, T>>(new TAABBTree<TPayloadType, TLeafType, bMutable, T>());
+			DirtyElementTree = TUniquePtr<TAABBTree>(new TAABBTree());
 			DirtyElementTree->SetTreeToDynamic();
-		}		
+		}
+
+		StorageTraits::InitPayloadToInfo(PayloadToInfo);
+
 		GenerateTree(Particles);
 	}
 
+	// Tag dispatch enable for the below constructor to allow setting up the defaults without an initial set of particles
+	struct EmptyInit {};
+
+	TAABBTree(EmptyInit, int32 InMaxChildrenInLeaf = DefaultMaxChildrenInLeaf, int32 InMaxTreeDepth = DefaultMaxTreeDepth, T InMaxPayloadBounds = DefaultMaxPayloadBounds, int32 InMaxNumToProcess = DefaultMaxNumToProcess, bool bInDynamicTree = false, bool bInUseDirtyTree = false, bool bInBuildOverlapCache = true)
+		: ISpatialAcceleration<TPayloadType, T, 3>(StaticType)
+		, bDynamicTree(bInDynamicTree)
+		, MaxChildrenInLeaf(InMaxChildrenInLeaf)
+		, MaxTreeDepth(InMaxTreeDepth)
+		, MaxPayloadBounds(InMaxPayloadBounds)
+		, MaxNumToProcess(InMaxNumToProcess)
+		, bModifyingTreeMultiThreadingFastCheck(false)
+		, bShouldRebuild(true)
+		, bBuildOverlapCache(bInBuildOverlapCache)
+	{
+		if(bInUseDirtyTree)
+		{
+			DirtyElementTree = TUniquePtr<TAABBTree>(new TAABBTree());
+			DirtyElementTree->SetTreeToDynamic();
+		}
+
+		StorageTraits::InitPayloadToInfo(PayloadToInfo);
+	}
+
 	template <typename ParticleView>
 	void Reinitialize(const ParticleView& Particles, int32 InMaxChildrenInLeaf = DefaultMaxChildrenInLeaf, int32 InMaxTreeDepth = DefaultMaxTreeDepth, T InMaxPayloadBounds = DefaultMaxPayloadBounds, int32 InMaxNumToProcess = DefaultMaxNumToProcess, bool bInDynamicTree = false, bool bInbBuildOverlapCache = true)
 	{
@@ -876,14 +916,14 @@ public:
 
 	virtual ~TAABBTree() {}
 
-	void CopyFrom(const TAABBTree<TPayloadType, TLeafType, bMutable, T>& Other)
+	void CopyFrom(const TAABBTree& Other)
 	{
 		(*this) = Other;
 	}
 
 	virtual TUniquePtr<ISpatialAcceleration<TPayloadType, T, 3>> Copy() const override
 	{
-		return TUniquePtr<ISpatialAcceleration<TPayloadType, T, 3>>(new TAABBTree<TPayloadType, TLeafType, bMutable, T>(*this));
+		return TUniquePtr<ISpatialAcceleration<TPayloadType, T, 3>>(new TAABBTree(*this));
 	}
 
 	virtual void Raycast(const FVec3& Start, const FVec3& Dir, const FReal Length, ISpatialVisitor<TPayloadType, FReal>& Visitor) const override
@@ -1147,11 +1187,18 @@ public:
 		return DirtyGridOverflowIdx;
 	}
 
-	// Expensive function: Don't call unless debugging
-	void DynamicTreeDebugStats()
+	struct FElementsCollection
 	{
 		TArray<FElement> AllElements;
-		for (int LeafIndex = 0; LeafIndex < Leaves.Num(); LeafIndex++)
+		int32 DepthTotal;
+		int32 MaxDepth;
+		int32 DirtyElementCount;
+	};
+
+	FElementsCollection DebugGetElementsCollection() const 
+	{
+		TArray<FElement> AllElements;
+		for(int LeafIndex = 0; LeafIndex < Leaves.Num(); LeafIndex++)
 		{
 			const TLeafType& Leaf = Leaves[LeafIndex];
 			Leaf.GatherElements(AllElements);
@@ -1159,34 +1206,85 @@ public:
 
 		int32 MaxDepth = 0;
 		int32 DepthTotal = 0;
-		for (const FElement& Element : AllElements)
+		for(const FElement& Element : AllElements)
 		{
-			FAABBTreePayloadInfo* PayloadInfo = PayloadToInfo.Find(Element.Payload);
+			const FAABBTreePayloadInfo* PayloadInfo = PayloadToInfo.Find(Element.Payload);
+
+			if(!PayloadInfo)
+			{
+				continue;
+			}
+
 			int32 Depth = 0;
 			int32 Node = PayloadInfo->NodeIdx;
-			check(Node != INDEX_NONE);
-			while (Node != INDEX_NONE)
+
+			while(Node != INDEX_NONE)
 			{
 				Node = Nodes[Node].ParentNode;
-				if (Node != INDEX_NONE)
+				if(Node != INDEX_NONE)
 				{
 					Depth++;
 				}
 			}
-			if (Depth > MaxDepth)
+			if(Depth > MaxDepth)
 			{
 				MaxDepth = Depth;
 			}
 			DepthTotal += Depth;
 		}
+
+		return { AllElements, DepthTotal, MaxDepth, DirtyElements.Num() };
+	}
+
+	// Expensive function: Don't call unless debugging
+	void DynamicTreeDebugStats() const
+	{
+		const FElementsCollection ElemData = DebugGetElementsCollection();
+
 #if !WITH_EDITOR
-		CSV_CUSTOM_STAT(ChaosPhysicsTimers, MaximumTreeDepth, MaxDepth, ECsvCustomStatOp::Max);
-		CSV_CUSTOM_STAT(ChaosPhysicsTimers, AvgTreeDepth, DepthTotal / AllElements.Num(), ECsvCustomStatOp::Max);
-		CSV_CUSTOM_STAT(ChaosPhysicsTimers, Dirty,DirtyElements.Num(), ECsvCustomStatOp::Max);
+		CSV_CUSTOM_STAT(ChaosPhysicsTimers, MaximumTreeDepth, ElemData.MaxDepth, ECsvCustomStatOp::Max);
+		CSV_CUSTOM_STAT(ChaosPhysicsTimers, AvgTreeDepth, ElemData.DepthTotal / ElemData.AllElements.Num(), ECsvCustomStatOp::Max);
+		CSV_CUSTOM_STAT(ChaosPhysicsTimers, Dirty, ElemData.DirtyElementCount, ECsvCustomStatOp::Max);
 #endif
 
 	}
 
+#if !UE_BUILD_SHIPPING
+	void DumpStats() const override
+	{
+		if(GLog)
+		{
+			DumpStatsTo(*GLog);
+		}
+	}
+
+	void DumpStatsTo(FOutputDevice& Ar) const override
+	{
+		const FElementsCollection ElemData = DebugGetElementsCollection();
+
+		const int32 ElementsNum = ElemData.AllElements.Num();
+		const int32 PayloadMapNum = PayloadToInfo.Num();
+		const int32 PayloadMapCapacity = PayloadToInfo.Capacity();
+		const uint64 PayloadAllocsize = (uint64)PayloadToInfo.GetAllocatedSize();
+		const float AvgDepth = ElementsNum > 0 ? (float)ElemData.DepthTotal / (float)ElementsNum : 0.0f;
+
+		Ar.Logf(ELogVerbosity::Log, TEXT("\t\tContains %d elements"), ElementsNum);
+		Ar.Logf(ELogVerbosity::Log, TEXT("\t\tMax depth is %d"), ElemData.MaxDepth);
+		Ar.Logf(ELogVerbosity::Log, TEXT("\t\tAvg depth is %.3f"), AvgDepth);
+		Ar.Logf(ELogVerbosity::Log, TEXT("\t\tDirty element count is %d"), ElemData.DirtyElementCount);
+		Ar.Logf(ELogVerbosity::Log, TEXT("\t\tPayload container size is %d elements"), PayloadMapNum);
+		Ar.Logf(ELogVerbosity::Log, TEXT("\t\tPayload container capacity is %d elements"), PayloadMapCapacity);
+		Ar.Logf(ELogVerbosity::Log, TEXT("\t\tAllocated size of payload container is %u bytes (%u per tree element)"), PayloadAllocsize, ElementsNum > 0 ? PayloadAllocsize / (uint32)ElementsNum : 0);
+		
+		if(DirtyElementTree)
+		{
+			Ar.Logf(ELogVerbosity::Log, TEXT(""));
+			Ar.Logf(ELogVerbosity::Log, TEXT("\t\tDirty Tree:"));
+			DirtyElementTree->DumpStatsTo(Ar);
+		}
+	}
+#endif
+
 	int32 AllocateInternalNode()
 	{
 		int32 AllocatedNodeIdx = FirstFreeInternalNode;
@@ -1304,10 +1402,6 @@ public:
 		return(WhichChildAmI(NodeIdx) ^ 1);
 	}
 
-private:
-	using FElement = TPayloadBoundsElement<TPayloadType, T>;
-	using FNode = TAABBTreeNode<T>;
-
 public:
 	int32 FindBestSibling(const TAABB<T, 3>& InNewBounds, bool& bOutAddToLeaf)
 	{
@@ -2012,7 +2106,7 @@ public:
 	{
 		check(this != &InFrom);
 		check(InFrom.GetType() == ESpatialAcceleration::AABBTree);
-		const TAABBTree<TPayloadType, TLeafType, bMutable, T>& From = static_cast<const TAABBTree<TPayloadType, TLeafType, bMutable, T>&>(InFrom);
+		const TAABBTree& From = static_cast<const TAABBTree&>(InFrom);
 
 		Reset();
 
@@ -2058,7 +2152,7 @@ public:
 		{
 			if (!DirtyElementTree)
 			{
-				DirtyElementTree = TUniquePtr<TAABBTree<TPayloadType, TLeafType, bMutable, T>>(new TAABBTree<TPayloadType, TLeafType, bMutable, T>());
+				DirtyElementTree = TUniquePtr<TAABBTree>(new TAABBTree());
 			}
 			DirtyElementTree->PrepareCopyTimeSliced(*(From.DirtyElementTree));
 		}
@@ -2072,7 +2166,7 @@ public:
 	{
 		check(this != &InFrom);
 		check(InFrom.GetType() == ESpatialAcceleration::AABBTree);
-		const TAABBTree<TPayloadType, TLeafType, bMutable, T>& From = static_cast<const TAABBTree<TPayloadType, TLeafType, bMutable, T>&>(InFrom);
+		const TAABBTree& From = static_cast<const TAABBTree&>(InFrom);
 
 		int32 SizeToCopyLeft = MaximumBytesToCopy;
 		check(From.CellHashToFlatArray.Num() == 0); // Partial Copy of TMAPs not implemented, and this should be empty for our current use cases
@@ -2520,7 +2614,7 @@ private:
 			Leaf.GatherElements(AllElements);
 		}
 
-		TAABBTree<TPayloadType,TLeafType,bMutable, T> NewTree(AllElements);
+		TAABBTree NewTree(AllElements);
 		*this = NewTree;
 		bShouldRebuild = true; // No changes since last time tree was built
 	}
@@ -3464,6 +3558,10 @@ private:
 		{
 			ContainerTo.AddFrom(ContainerFrom, Index);
 		}
+		else if constexpr(std::is_same_v<ContainerType, TSQMap<TPayloadType, FAABBTreePayloadInfo>>)
+		{
+			ContainerTo.AddFrom(ContainerFrom, Index);
+		}
 		else
 		{
 			ContainerTo.Add(ContainerFrom[Index]);
@@ -3535,7 +3633,7 @@ private:
 #endif
 
 
-	TAABBTree(const TAABBTree<TPayloadType, TLeafType, bMutable, T>& Other)
+	TAABBTree(const TAABBTree& Other)
 		: ISpatialAcceleration<TPayloadType, T, 3>(StaticType)
 		, Nodes(Other.Nodes)
 		, Leaves(Other.Leaves)
@@ -3574,17 +3672,17 @@ private:
 		PriorityQ.Reserve(32);
 		if (Other.DirtyElementTree)
 		{
-			DirtyElementTree = TUniquePtr<TAABBTree<TPayloadType, TLeafType, bMutable, T>>(new TAABBTree<TPayloadType, TLeafType, bMutable, T>(*(Other.DirtyElementTree)));
+			DirtyElementTree = TUniquePtr<TAABBTree>(new TAABBTree(*(Other.DirtyElementTree)));
 		}
 	}
 
 	virtual void DeepAssign(const ISpatialAcceleration<TPayloadType, T, 3>& Other) override
 	{
 		check(Other.GetType() == ESpatialAcceleration::AABBTree);
-		*this = static_cast<const TAABBTree<TPayloadType, TLeafType, bMutable, T>&>(Other);
+		*this = static_cast<const TAABBTree&>(Other);
 	}
 
-	TAABBTree<TPayloadType,TLeafType, bMutable, T>& operator=(const TAABBTree<TPayloadType,TLeafType,bMutable, T>& Rhs)
+	TAABBTree& operator=(const TAABBTree& Rhs)
 	{
 		ISpatialAcceleration<TPayloadType, T, 3>::DeepAssign(Rhs);
 		ensure(Rhs.WorkStack.Num() == 0);
@@ -3629,7 +3727,7 @@ private:
 			{
 				if (!DirtyElementTree)
 				{
-					DirtyElementTree = TUniquePtr<TAABBTree<TPayloadType, TLeafType, bMutable, T>>(new TAABBTree<TPayloadType, TLeafType, bMutable, T>());
+					DirtyElementTree = TUniquePtr<TAABBTree>(new TAABBTree());
 				}				
 				*DirtyElementTree = *Rhs.DirtyElementTree;
 			}
@@ -3659,8 +3757,7 @@ private:
 	TArray<int32> DirtyElementsGridOverflow; // Array of indices of DirtyElements that is not in the grid for some reason
 
 	// Members for using a dynamic tree as a dirty element acceleration structure
-	TUniquePtr<TAABBTree<TPayloadType, TLeafType, bMutable, T>> DirtyElementTree;
-
+	TUniquePtr<TAABBTree> DirtyElementTree;
 
 	// Copy of CVARS
 	T DirtyElementGridCellSize;
@@ -3673,7 +3770,7 @@ private:
 	AABBTreeExpensiveStatistics TreeExpensiveStats;
 
 	TArray<FElement> GlobalPayloads;
-	TArrayAsMap<TPayloadType, FAABBTreePayloadInfo> PayloadToInfo;
+	typename StorageTraits::PayloadToInfoType PayloadToInfo;
 
 	int32 MaxChildrenInLeaf;
 	int32 MaxTreeDepth;
@@ -3711,8 +3808,8 @@ private:
 
 };
 
-template<typename TPayloadType, typename TLeafType, bool bMutable, typename T>
-FArchive& operator<<(FChaosArchive& Ar, TAABBTree<TPayloadType, TLeafType, bMutable, T>& AABBTree)
+template<typename TPayloadType, typename TLeafType, bool bMutable, typename T, typename StorageTraits>
+FArchive& operator<<(FChaosArchive& Ar, TAABBTree<TPayloadType, TLeafType, bMutable, T, StorageTraits>& AABBTree)
 {
 	AABBTree.Serialize(Ar);
 	return Ar;
diff --git a/Engine/Source/Runtime/Experimental/Chaos/Public/Chaos/BoundingVolume.h b/Engine/Source/Runtime/Experimental/Chaos/Public/Chaos/BoundingVolume.h
--- a/Engine/Source/Runtime/Experimental/Chaos/Public/Chaos/BoundingVolume.h
+++ b/Engine/Source/Runtime/Experimental/Chaos/Public/Chaos/BoundingVolume.h
@@ -423,7 +423,7 @@ public:
 
 	}
 
-	void GatherElements(TArray<TPayloadBoundsElement<TPayloadType, T>>& OutElements)
+	void GatherElements(TArray<TPayloadBoundsElement<TPayloadType, T>>& OutElements) const
 	{
 		OutElements.Reserve(GetReserveCount());
 		OutElements.Append(MGlobalPayloads);
diff --git a/Engine/Source/Runtime/Experimental/Chaos/Public/Chaos/ISpatialAcceleration.h b/Engine/Source/Runtime/Experimental/Chaos/Public/Chaos/ISpatialAcceleration.h
--- a/Engine/Source/Runtime/Experimental/Chaos/Public/Chaos/ISpatialAcceleration.h
+++ b/Engine/Source/Runtime/Experimental/Chaos/Public/Chaos/ISpatialAcceleration.h
@@ -341,6 +341,7 @@ public:
 	virtual void DebugDraw(ISpacialDebugDrawInterface<T>* InInterface) const {}
 	virtual void DebugDrawLeaf(ISpacialDebugDrawInterface<T>& InInterface, const FLinearColor& InLinearColor, float InThickness) const {}
 	virtual void DumpStats() const {}
+	virtual void DumpStatsTo(class FOutputDevice& Ar) const {}
 #endif
 
 	static ISpatialAcceleration<TPayloadType, T, d>* SerializationFactory(FChaosArchive& Ar, ISpatialAcceleration<TPayloadType, T, d>* Accel)
@@ -496,10 +497,43 @@ private:
 template <typename TKey, typename TValue>
 class TArrayAsMap
 {
+	struct FEntry
+	{
+		TValue Value;
+		bool bSet;
+
+		FEntry()
+			: bSet(false)
+		{
+
+		}
+	};
+
 public:
 	// @todo(chaos): rename with "F"
 	using ElementType = TValue;
 
+	static constexpr uint32 GetTypeSize()
+	{
+		return sizeof(FEntry);
+	}
+
+	SIZE_T GetAllocatedSize() const
+	{
+		SIZE_T AllocatedSize = Entries.GetAllocatedSize();
+
+#if CHAOS_SERIALIZE_OUT
+		AllocatedSize += KeysToSerializeOut.GetAllocatedSize();
+#endif
+
+		return AllocatedSize;
+	}
+
+	int32 Capacity() const
+	{
+		return Entries.Max();
+	}
+
 	// @todo(chaos): rename with "F"
 	struct Element
 	{
@@ -656,19 +690,7 @@ public:
 	}
 
 private:
-	
-	struct FEntry
-	{
-		TValue Value;
-		bool bSet;
-
-		FEntry()
-			: bSet(false)
-		{
 
-		}
-	};
-	
 	TArray<FEntry> Entries;
 
 #if CHAOS_SERIALIZE_OUT
@@ -678,6 +700,282 @@ private:
 #endif
 };
 
+template <typename TKey, typename TValue>
+struct SQMapKeyWithValue
+{
+	SQMapKeyWithValue() = default;
+	SQMapKeyWithValue(const TKey& InKey, const TValue& InValue)
+		: Key(InKey)
+		, Value(InValue)
+	{}
+
+	TKey Key;
+	TValue Value;
+};
+
+template <typename TKey, typename TValue>
+FChaosArchive& operator<< (FChaosArchive& Ar, SQMapKeyWithValue<TKey, TValue>& Pair)
+{
+	Ar << Pair.Key;
+	Ar << Pair.Value;
+
+	return Ar;
+}
+
+class FSQHashTable : public FHashTable
+{
+public:
+	SIZE_T GetAllocatedSize() const
+	{
+		return IndexSize == 0 ? 0 : (IndexSize + HashSize) * sizeof(uint32);
+	}
+
+	uint32 GetIndexSize() const
+	{
+		return IndexSize;
+	}
+};
+
+/**
+ * Map structure using FHashTable as a base store for payload data
+ * FHashTable is a fast, limited API map that requires some management.
+ * In this case as we already have unique values we can just mix the bits
+ * to get a reasonable distribution (the calls to MurmurFinalize32).
+ */
+template <typename TKey, typename TValue>
+class TSQMap
+{
+public:
+
+	using ElementType = TValue;
+	using PairType = SQMapKeyWithValue<TKey, TValue>;
+
+	static constexpr uint32 GetTypeSize()
+	{
+		return sizeof(PairType);
+	}
+
+	SIZE_T GetAllocatedSize() const
+	{
+		return Elements.GetAllocatedSize() + HashTable.GetAllocatedSize();
+	}
+
+	int32 Num() const
+	{
+		return Elements.Num();
+	}
+
+	int32 Capacity() const
+	{
+		return HashTable.GetIndexSize();
+	}
+
+	void Reserve(int32 NumToReserve)
+	{
+		Elements.Reserve(NumToReserve);
+
+		// FHashTable requires powers of two but doesn't assert it.
+		const uint32 RequiredSize = FMath::RoundUpToPowerOfTwo(NumToReserve + 1);
+
+		// #TODO should we allow SQ shrinking?
+		if(RequiredSize > HashTable.GetIndexSize())
+		{
+			HashTable.Resize(RequiredSize);
+		}
+	}
+
+	void ResizeHashBuckets(uint32 NewSize)
+	{
+		const uint32 NumElements = Elements.Num();
+		HashTable.Clear(NewSize, FMath::RoundUpToPowerOfTwo(NumElements + 1));
+
+		// Rehash the elements, we don't use the other Add function here
+		// as that will search for the element and we know it's not in here
+		// as we're building a new hash table
+		for(uint32 ElemIdx = 0; ElemIdx < NumElements; ++ElemIdx)
+		{
+			const uint32 ElemHash = MurmurFinalize32(GetUniqueIdx(Elements[ElemIdx].Key).Idx);
+			HashTable.Add(ElemHash, ElemIdx);
+		}
+	}
+
+	TValue* Find(const TKey& Key)
+	{
+		int32 Index = FindIndex(Key);
+
+		if(Index != INDEX_NONE)
+		{
+			return &Elements[Index].Value;
+		}
+
+		return nullptr;
+	}
+
+	const TValue* Find(const TKey& Key) const
+	{
+		int32 Index = FindIndex(Key);
+
+		if(Index != INDEX_NONE)
+		{
+			return &Elements[Index].Value;
+		}
+
+		return nullptr;
+	}
+
+	TValue& FindChecked(const TKey& Key)
+	{
+		int32 Index = FindIndex(Key);
+		check(Index != INDEX_NONE);
+		return Elements[Index].Value;
+	}
+
+	const TValue& FindChecked(const TKey& Key) const
+	{
+		int32 Index = FindIndex(Key);
+		check(Index != INDEX_NONE);
+		return Elements[Index].Value;
+	}
+
+	TValue& FindOrAdd(const TKey& Key)
+	{
+		const int32 KeyIndex = GetUniqueIdx(Key).Idx;
+		const uint32 Hash = MurmurFinalize32(KeyIndex);
+		int32 Index = FindIndex(Hash, KeyIndex);
+		
+		if(Index == INDEX_NONE)
+		{
+			return Add(Hash, Key);
+		}
+
+		return Elements[Index].Value;
+	}
+
+	void Empty()
+	{
+		Elements.Empty();
+		HashTable.Free();
+	}
+
+	FORCEINLINE TValue& Add(const TKey& Key)
+	{
+		return Add(MurmurFinalize32(GetUniqueIdx(Key).Idx), Key);
+	}
+
+	FORCEINLINE TValue& Add(uint32 InHash, const TKey& Key)
+	{
+		int32 ExistingIndex = FindIndex(InHash, GetUniqueIdx(Key).Idx);
+
+		if(ExistingIndex == INDEX_NONE)
+		{
+			Elements.Emplace(Key, TValue{});
+			ExistingIndex = Elements.Num() - 1;
+			HashTable.Add(InHash, ExistingIndex);
+		}
+
+		return Elements[ExistingIndex].Value;
+	}
+
+	void Add(const TKey& Key, const TValue& Value)
+	{
+		Add(Key) = Value;
+	}
+
+	void RemoveChecked(const TKey& Key)
+	{
+		const int32 NumBefore = Elements.Num();
+		Remove(Key);
+		check(NumBefore > Elements.Num());
+	}
+
+	void Remove(const TKey& Key)
+	{
+		const int32 KeyIndex = GetUniqueIdx(Key).Idx;
+		const uint32 Hash = MurmurFinalize32(KeyIndex);
+		const int32 Index = FindIndex(Hash, KeyIndex);
+
+		if(Index != INDEX_NONE)
+		{
+			HashTable.Remove(Hash, Index);
+
+			// Before removing from the array, if we have more than 1 element we
+			// swap the back element into the empty slot and rehash it. Maybe
+			// we should use a freelist? (profile this)
+			const int32 NumElems = Elements.Num();
+			const int32 BackIndex = NumElems - 1;
+			if(NumElems > 1)
+			{
+				if(Index == BackIndex)
+				{
+					// If we're already the back element, there's nothing to re-add
+					Elements.RemoveAtSwap(Index, 1, false);
+				}
+				else
+				{
+					PairType& BackPair = Elements[NumElems - 1];
+					const uint32 BackHash = MurmurFinalize32(GetUniqueIdx(BackPair.Key).Idx);
+					HashTable.Remove(BackHash, NumElems - 1);
+					Elements.RemoveAtSwap(Index, 1, false);
+					HashTable.Add(BackHash, Index);
+				}
+			}
+			else
+			{
+				Elements.Reset();
+			}
+		}
+	}
+
+	void Reset()
+	{
+		Elements.Reset();
+		HashTable.Clear();
+	}
+
+	FORCEINLINE int32 FindIndex(const TKey& Key) const
+	{
+		const int32 KeyIndex = GetUniqueIdx(Key).Idx;
+		return FindIndex(MurmurFinalize32(KeyIndex), KeyIndex);
+	}
+
+	FORCEINLINE int32 FindIndex(uint32 InHash, int32 UniqueIndex) const
+	{
+		for(uint32 i = HashTable.First(InHash); HashTable.IsValid(i); i = HashTable.Next(i))
+		{
+			if(GetUniqueIdx(Elements[i].Key).Idx == UniqueIndex)
+			{
+				return i;
+			}
+		}
+
+		return INDEX_NONE;
+	}
+
+	void Serialize(FChaosArchive& Ar)
+	{
+		Ar << Elements;
+
+		if(Ar.IsLoading())
+		{
+			const uint32 NumElements = Elements.Num();
+			HashTable.Clear(1024, FMath::RoundUpToPowerOfTwo(NumElements));
+			for(uint32 i = 0; i < NumElements; ++i)
+			{
+				HashTable.Add(MurmurFinalize32(GetUniqueIdx(Elements[i].Key).Idx), i);
+			}
+		}
+	}
+
+	void AddFrom(const TSQMap<TKey, TValue>& Source, int32 SourceIndex)
+	{
+		const PairType& SourcePair = Source.Elements[SourceIndex];
+		Add(SourcePair.Key, SourcePair.Value);
+	}
+
+	TArray<PairType> Elements;
+	FSQHashTable HashTable;
+};
+
 template <typename TKey, typename TValue>
 FChaosArchive& operator<< (FChaosArchive& Ar, TArrayAsMap<TKey, TValue>& Map)
 {
@@ -685,6 +983,12 @@ FChaosArchive& operator<< (FChaosArchive& Ar, TArrayAsMap<TKey, TValue>& Map)
 	return Ar;
 }
 
+template <typename TKey, typename TValue>
+FChaosArchive& operator<< (FChaosArchive& Ar, TSQMap<TKey, TValue>& Map)
+{
+	Map.Serialize(Ar);
+	return Ar;
+}
 
 template <typename TPayload, typename TVisitor>
 typename TEnableIf<!TIsPointer<TPayload>::Value, bool>::Type PrePreFilterHelper(const TPayload& Payload, const TVisitor& Visitor)
diff --git a/Engine/Source/Runtime/Experimental/Chaos/Public/Chaos/SpatialAccelerationCollection.h b/Engine/Source/Runtime/Experimental/Chaos/Public/Chaos/SpatialAccelerationCollection.h
--- a/Engine/Source/Runtime/Experimental/Chaos/Public/Chaos/SpatialAccelerationCollection.h
+++ b/Engine/Source/Runtime/Experimental/Chaos/Public/Chaos/SpatialAccelerationCollection.h
@@ -6,6 +6,7 @@
 #include "Chaos/Collision/StatsData.h"
 #include "ChaosStats.h"
 #include "GeometryParticlesfwd.h"
+#include "Misc/OutputDevice.h"
 
 #include <tuple>
 
@@ -685,6 +686,32 @@ public:
 		}
 	}
 
+#if !UE_BUILD_SHIPPING
+	void DumpStats() const override
+	{
+		if(GLog)
+		{
+			DumpStatsTo(*GLog);
+		}
+	}
+
+	void DumpStatsTo(FOutputDevice& Ar) const override
+	{
+		for(int BucketIdx = 0; BucketIdx < MaxBuckets; ++BucketIdx)
+		{
+			const BucketType& Bucket = Buckets[BucketIdx];
+			Ar.Logf(TEXT("Bucket %d (%d entries):"), BucketIdx, Bucket.Objects.Num());
+
+			for(int EntryIdx = 0; EntryIdx < Bucket.Objects.Num(); ++EntryIdx)
+			{
+				Ar.Logf(TEXT("\tEntry %d"), EntryIdx);
+				Bucket.Objects[EntryIdx].Acceleration->DumpStatsTo(Ar);
+				Ar.Logf(TEXT(""));
+			}
+		}
+	}
+#endif
+
 private:
 
 	TSpatialAccelerationCollection(const TSpatialAccelerationCollection<TSpatialAccelerationTypes...>& Other)
