NonUnity fix

#rnx

diff --git a/Engine/Source/Runtime/Experimental/Chaos/Public/Chaos/ISpatialAcceleration.h b/Engine/Source/Runtime/Experimental/Chaos/Public/Chaos/ISpatialAcceleration.h
--- a/Engine/Source/Runtime/Experimental/Chaos/Public/Chaos/ISpatialAcceleration.h
+++ b/Engine/Source/Runtime/Experimental/Chaos/Public/Chaos/ISpatialAcceleration.h
@@ -5,6 +5,7 @@
 #include "GeometryParticlesfwd.h"
 #include "ChaosCheck.h"
 #include "ChaosDebugDrawDeclares.h"
+#include "Containers/HashTable.h"
 
 namespace Chaos
 {
