Do not log a streaming error if the IoStoreOnDemand.ini file cannot be found.

#rnx

- FPlatformMisc::LoadTextFileFromPlatformPackage will print a LogStreaming error if the given file cannot be found. However we are enabling/disabling Ias based on file existing or not so the file not existing is perfectly valid and should not raise an error.
- Check for the file existing first before trying to load and parse it.

diff --git a/Engine/Source/Runtime/Experimental/IoStoreOnDemand/Private/IoStoreOnDemand.cpp b/Engine/Source/Runtime/Experimental/IoStoreOnDemand/Private/IoStoreOnDemand.cpp
--- a/Engine/Source/Runtime/Experimental/IoStoreOnDemand/Private/IoStoreOnDemand.cpp
+++ b/Engine/Source/Runtime/Experimental/IoStoreOnDemand/Private/IoStoreOnDemand.cpp
@@ -155,9 +155,16 @@ static bool TryParseConfigFileFromPlatformPackage(FOnDemandEndpoint& OutEndpoint
 {
 	const FString ConfigFileName = TEXT("IoStoreOnDemand.ini");
 	const FString ConfigPath = FPaths::Combine(TEXT("Cloud"), ConfigFileName);
-	const FString ConfigContent = FPlatformMisc::LoadTextFileFromPlatformPackage(ConfigPath);
-
-	return TryParseConfigContent(ConfigContent, ConfigFileName, OutEndpoint);
+	
+	if (FPlatformMisc::FileExistsInPlatformPackage(ConfigPath))
+	{
+		const FString ConfigContent = FPlatformMisc::LoadTextFileFromPlatformPackage(ConfigPath);
+		return TryParseConfigContent(ConfigContent, ConfigFileName, OutEndpoint);
+	}
+	else
+	{
+		return false;
+	}
 }
 
 ////////////////////////////////////////////////////////////////////////////////
