Properly register new components in UGeometryCollectionComponent::CreateProxyComponents()

diff --git a/Engine/Source/Runtime/Experimental/GeometryCollectionEngine/Private/GeometryCollection/GeometryCollectionComponent.cpp b/Engine/Source/Runtime/Experimental/GeometryCollectionEngine/Private/GeometryCollection/GeometryCollectionComponent.cpp
--- a/Engine/Source/Runtime/Experimental/GeometryCollectionEngine/Private/GeometryCollection/GeometryCollectionComponent.cpp
+++ b/Engine/Source/Runtime/Experimental/GeometryCollectionEngine/Private/GeometryCollection/GeometryCollectionComponent.cpp
@@ -4961,7 +4961,9 @@ TArray<UStaticMeshComponent*> UGeometryCollectionComponent::CreateProxyComponent
 			{
 				UStaticMeshComponent* NewComponent = NewObject<UStaticMeshComponent>(GetOwner());
 				NewComponent->SetStaticMesh(Mesh);
-				NewComponent->SetComponentToWorld(GetComponentToWorld());
+				NewComponent->SetRelativeTransform(GetComponentTransform());
+				NewComponent->RegisterComponent();
+
 				Components.Add(NewComponent);
 			}
 		}
