Ias: Hook up http client throttling to a tweakable "ias.HttpRateLimitKiBPerSecond" cvar.

#rnx

diff --git a/Engine/Source/Runtime/Experimental/IoStoreOnDemand/Private/OnDemandHttpClient.cpp b/Engine/Source/Runtime/Experimental/IoStoreOnDemand/Private/OnDemandHttpClient.cpp
--- a/Engine/Source/Runtime/Experimental/IoStoreOnDemand/Private/OnDemandHttpClient.cpp
+++ b/Engine/Source/Runtime/Experimental/IoStoreOnDemand/Private/OnDemandHttpClient.cpp
@@ -11,11 +11,18 @@
 namespace UE::IO::IAS
 {
 
+int32 GIasHttpRateLimitKiBPerSecond = 0;
+static FAutoConsoleVariableRef CVar_GIasHttpRateLimitKiBPerSecond(
+	TEXT("ias.HttpRateLimitKiBPerSecond"),
+	GIasHttpRateLimitKiBPerSecond,
+	TEXT("Http throttle limit in KiBPerSecond")
+);
+
 int32 GIasHttpPollTimeoutMs = 0;
 static FAutoConsoleVariableRef CVar_GIasHttpPollTimeoutMs(
-	TEXT("ias.HttpPollTimeout"),
+	TEXT("ias.HttpPollTimeoutMs"),
 	GIasHttpPollTimeoutMs,
-	TEXT("FHttpClient::Tick poll timeout in milliseconds")
+	TEXT("Http tick poll timeout in milliseconds")
 );
 
 static void LogHttpResult(const TCHAR* Url, uint32 StatusCode, uint64 DurationMs, uint64 Size, uint64 Offset, const char* Memo = "ok")
@@ -108,6 +115,7 @@ void FOnDemandHttpClient::Issue(FAnsiStringView Url, FGetCallback&& Callback, FI
 bool FOnDemandHttpClient::Tick(bool Block)
 {
 	int32 TimeoutMs = Block ? -1 : GIasHttpPollTimeoutMs;
+	EventLoop.Throttle(GIasHttpRateLimitKiBPerSecond);
 	return EventLoop.Tick(TimeoutMs) != 0;
 }
 
