Make sure IAS directory exists when creating cache. Root directory is now a parameter. File-open handling to follow.

#rnx

diff --git a/Engine/Source/Runtime/Experimental/IoStoreOnDemand/Private/IasCache.h b/Engine/Source/Runtime/Experimental/IoStoreOnDemand/Private/IasCache.h
--- a/Engine/Source/Runtime/Experimental/IoStoreOnDemand/Private/IasCache.h
+++ b/Engine/Source/Runtime/Experimental/IoStoreOnDemand/Private/IasCache.h
@@ -3,6 +3,7 @@
 #pragma once
 
 #include "CancellationToken.h"
+#include "Containers/StringView.h"
 #include "IO/IoStatus.h"
 #include "Memory/MemoryFwd.h"
 #include "Tasks/Task.h"
@@ -49,7 +50,7 @@ struct FIasCacheConfig
 		uint8	SuperBoost = 87;
 	};
 
-	const TCHAR*Name = TEXT("ias");
+	FStringView Name = TEXT("ias");
 	uint64		DiskQuota = 512ull << 20;
 	uint32		MemoryQuota = 2 << 20;
 	uint32		JournalQuota = 4 << 20; // description in JournalCache.cpp
@@ -58,6 +59,6 @@ struct FIasCacheConfig
 	bool		DropCache = false;
 };
 
-TUniquePtr<IIasCache> MakeIasCache(const FIasCacheConfig& Config);
+TUniquePtr<IIasCache> MakeIasCache(const TCHAR* RootDir, const FIasCacheConfig& Config);
 
 } // namespace UE::IO::IAS
diff --git a/Engine/Source/Runtime/Experimental/IoStoreOnDemand/Private/IoStoreOnDemand.cpp b/Engine/Source/Runtime/Experimental/IoStoreOnDemand/Private/IoStoreOnDemand.cpp
--- a/Engine/Source/Runtime/Experimental/IoStoreOnDemand/Private/IoStoreOnDemand.cpp
+++ b/Engine/Source/Runtime/Experimental/IoStoreOnDemand/Private/IoStoreOnDemand.cpp
@@ -1570,7 +1570,8 @@ void FIoStoreOnDemandModule::StartupModule()
 	TSharedPtr<IIasCache> Cache;
 	if (FIasCacheConfig Config = GetIasCacheConfig(CommandLine); Config.DiskQuota > 0)
 	{
-		Cache = MakeShareable(MakeIasCache(Config).Release());
+		FString CacheDir = FPaths::ProjectPersistentDownloadDir();
+		Cache = MakeShareable(MakeIasCache(*CacheDir, Config).Release());
 	}
 	else
 	{
diff --git a/Engine/Source/Runtime/Experimental/IoStoreOnDemand/Private/JournaledCache.cpp b/Engine/Source/Runtime/Experimental/IoStoreOnDemand/Private/JournaledCache.cpp
--- a/Engine/Source/Runtime/Experimental/IoStoreOnDemand/Private/JournaledCache.cpp
+++ b/Engine/Source/Runtime/Experimental/IoStoreOnDemand/Private/JournaledCache.cpp
@@ -8,6 +8,7 @@
 #include "Containers/Array.h"
 #include "Containers/Map.h"
 #include "Containers/UnrealString.h"
+#include "HAL/FileManager.h"
 #include "HAL/LowLevelMemTracker.h"
 #include "HAL/PlatformFile.h"
 #include "HAL/PlatformProcess.h"
@@ -484,16 +485,6 @@ void FDiskJournal::OpenJrnFile()
 	IPlatformFile& Ipf = IPlatformFile::GetPlatformPhysical();
 
 	IFileHandle* Handle = Ipf.OpenWrite(*JrnPath, true, true);
-	if (Handle == nullptr)
-	{	
-		// If the open failed it could be because we still need to create the directory
-		TStringBuilder<64> JrnDir = WriteToString<64>(FPathViews::GetPath(JrnPath));
-		if (Ipf.CreateDirectory(*JrnDir))
-		{
-			Handle = Ipf.OpenWrite(*JrnPath, true, true);
-		}
-	}
-
 	UE_CLOG(Handle == nullptr, LogIas, Error, TEXT("Failed to open '%s' for FDiskJournal"), *JrnPath);
 	JrnHandle.Reset(Handle);
 }
@@ -1448,7 +1439,7 @@ public:
 	using GetRetType = UE::Tasks::TTask<TIoStatusOr<FIoBuffer>>;
 
 								FJournaledCache() = default;
-	bool						Initialize(const FIasCacheConfig& Config);
+	bool						Initialize(const TCHAR* RootDir, const FIasCacheConfig& Config);
 	virtual bool				ContainsChunk(const FIoHash& Key) const override;
 	virtual GetRetType			Get(const FIoHash& Key, const FIoReadOptions& Options, const FIoCancellationToken* CancellationToken) override;
 	virtual FIoStatus			Put(const FIoHash& Key, FIoBuffer& Data) override;
@@ -1472,24 +1463,34 @@ private:
 };
 
 ////////////////////////////////////////////////////////////////////////////////
-bool FJournaledCache::Initialize(const FIasCacheConfig& Config)
+bool FJournaledCache::Initialize(const TCHAR* RootDir, const FIasCacheConfig& Config)
 {
+	// Filesystem setup
+	FStringView Name = Config.Name;
+	check(Name.Len() > 0 && !Name.EndsWith('/') && !Name.EndsWith('\\'));
+
 	TStringBuilder<256> CachePath;
-	CachePath << FPaths::ProjectPersistentDownloadDir();
-	if (CachePath.LastChar() != TEXT('/'))
+	CachePath << RootDir;
+	FPathViews::Append(CachePath, "ias");
+	FPathViews::Append(CachePath, FPathViews::GetPath(Name));
+
+	if (IFileManager& Ifm = IFileManager::Get(); !Ifm.MakeDirectory(CachePath.ToString(), true))
 	{
-		CachePath.AppendChar(TEXT('/'));
+		UE_LOG(LogIas, Error, TEXT("JournaledCache: Unable to create directory '%s'"), CachePath.ToString());
+		return false;
 	}
-	CachePath << TEXT("ias/");
-	CachePath << Config.Name;
+
+	FPathViews::Append(CachePath, FPathViews::GetBaseFilename(Name));
 	CachePath << TEXT(".cache.0");
 
+	// Inner cache
 	FCacheInner::FConfig EventualConfig;
 	static_cast<FIasCacheConfig&>(EventualConfig) = Config;
 	EventualConfig.Path = CachePath;
 	Cache = MakeUnique<FCacheInner>(MoveTemp(EventualConfig));
 	Cache->Load();
 
+	// Thread setup
 	const FIasCacheConfig::FRate& WriteRate = Config.WriteRate;
 	const FIasCacheConfig::FDemand& Demand = Config.Demand;
 	Governor.Set(WriteRate.Allowance, WriteRate.Ops, WriteRate.Seconds);
@@ -1633,12 +1634,12 @@ uint64 FJournaledCache::ReduceKey(const FIoHash& Key)
 
 
 ////////////////////////////////////////////////////////////////////////////////
-TUniquePtr<IIasCache> MakeIasCache(const FIasCacheConfig& Config)
+TUniquePtr<IIasCache> MakeIasCache(const TCHAR* RootPath, const FIasCacheConfig& Config)
 {
 	LLM_SCOPE_BYTAG(Ias);
 
 	FJournaledCache* Cache = new FJournaledCache();
-	if (Cache->Initialize(Config))
+	if (Cache->Initialize(RootPath, Config))
 	{
 		return TUniquePtr<IIasCache>(Cache);
 	}
