fix broken GeometryCollection HasCycle functions, which were running a recursive function but passing a tracking array by value

Instead, just walk parents until we've found a root or walked > NumNodes steps; this avoids creating the temp array and should be faster in the non-error case, and not terrible in the error case (which is a very bad failure case and should not happen generally ...)

diff --git a/Engine/Source/Runtime/Experimental/Chaos/Private/GeometryCollection/GeometryCollectionAlgo.cpp b/Engine/Source/Runtime/Experimental/Chaos/Private/GeometryCollection/GeometryCollectionAlgo.cpp
--- a/Engine/Source/Runtime/Experimental/Chaos/Private/GeometryCollection/GeometryCollectionAlgo.cpp
+++ b/Engine/Source/Runtime/Experimental/Chaos/Private/GeometryCollection/GeometryCollectionAlgo.cpp
@@ -175,36 +175,26 @@ namespace GeometryCollectionAlgo
 		return false;
 	}
 
-	bool HasCycleRec(TManagedArray<int32>& Parents, int32 Node, TArray<bool> Visited)
+	bool HasCycle(TManagedArray<int32>& Parents, int32 Node)
 	{
-		ensure(0 <= Node && Node < Visited.Num());
-
-		if (Visited[Node])
-			return true;
-		Visited[Node] = true;
-
-		if (Parents[Node] != FGeometryCollection::Invalid)
+		const int32 NumParents = Parents.Num();
+		int32 WalkNode = Node;
+		for (int32 Iters = 0; WalkNode != FGeometryCollection::Invalid && Iters < NumParents; ++Iters)
 		{
-			return HasCycleRec(Parents, Parents[Node], Visited);
+			WalkNode = Parents[WalkNode];
 		}
-		return false;
-	}
-	bool HasCycle(TManagedArray<int32>& Parents, int32 Node)
-	{
-		TArray<bool> Visited;
-		Visited.Init(false, Parents.Num());
-		return HasCycleRec(Parents, Node, Visited);
+		return WalkNode != FGeometryCollection::Invalid;
 	}
 	bool HasCycle(TManagedArray<int32>& Parents, const TArray<int32>& SelectedBones)
 	{
-		bool result = false;
-		TArray<bool> Visited;
-		Visited.Init(false, Parents.Num());
-		for (int32 Index = 0; Index < SelectedBones.Num(); Index++)
+		for (int32 Bone : SelectedBones)
 		{
-			result |= HasCycleRec(Parents, SelectedBones[Index], Visited);
+			if (HasCycle(Parents, Bone))
+			{
+				return true;
+			}
 		}
-		return result;
+		return false;
 	}
 
 	void ParentTransform(FTransformCollection* GeometryCollection, const int32 TransformIndex, const int32 ChildIndex)
