Fixed crash if TAABBtree becomes too deep

#[review]-27587475 [REDACTED_HANDLE]

diff --git a/Engine/Source/Runtime/Experimental/Chaos/Public/Chaos/AABBTree.h b/Engine/Source/Runtime/Experimental/Chaos/Public/Chaos/AABBTree.h
--- a/Engine/Source/Runtime/Experimental/Chaos/Public/Chaos/AABBTree.h
+++ b/Engine/Source/Runtime/Experimental/Chaos/Public/Chaos/AABBTree.h
@@ -2900,7 +2900,7 @@ private:
 			FReal TOI;
 		};
 
-		// Caching is for now only available for dyanmic tree
+		// Caching is for now only available for dynanmic tree
 		if (bDynamicTree && !OverlappingLeaves.IsEmpty())
 		{
 			// For overlap query and dynamic tree we are using the cached overlapping leaves
@@ -2911,21 +2911,20 @@ private:
 			}
 		}
 		
-		constexpr int32 MaxNodeStackNum = 255;
-		check(MaxTreeDepth + 2 <= MaxNodeStackNum);
-		FNodeQueueEntry NodeStack[MaxNodeStackNum];
-		int32 NodeStackNum = 0;
+		constexpr int32 MaxNodeStackNumOnSystemStack = 255;
+		TArray<FNodeQueueEntry, TSizedInlineAllocator<MaxNodeStackNumOnSystemStack,32> > NodeStack;
+		
 		if (bDynamicTree)
 		{
 
 			if (RootNode != INDEX_NONE)
 			{
-				NodeStack[NodeStackNum++] = FNodeQueueEntry{ RootNode, 0 };
+				NodeStack.Emplace(FNodeQueueEntry{RootNode, 0});
 			}
 		}
 		else if (Nodes.Num())
 		{
-			NodeStack[NodeStackNum++] = FNodeQueueEntry{ 0, 0 };
+			NodeStack.Emplace(FNodeQueueEntry{ 0, 0 });
 		}
 
 // Slow debug code
@@ -2953,7 +2952,7 @@ private:
 			LengthSimd = VectorSetDouble1(CurData.CurrentLength);
 		}
 
-		while (NodeStackNum)
+		while (NodeStack.Num())
 		{
 			PHYSICS_CSV_SCOPED_VERY_EXPENSIVE(PhysicsVerbose, QueryImp_NodeTraverse);
 
@@ -2963,8 +2962,7 @@ private:
 //				CSV_CUSTOM_STAT(ChaosPhysicsTimers, AABBCheckCount, 1, ECsvCustomStatOp::Accumulate);
 //			}
 //#endif
-
-			const FNodeQueueEntry NodeEntry = NodeStack[--NodeStackNum];
+			const FNodeQueueEntry NodeEntry = NodeStack.Pop(false);
 			if constexpr (Query != EAABBQueryType::Overlap)
 			{
 				if (NodeEntry.TOI > CurData.CurrentLength)
@@ -3029,22 +3027,22 @@ private:
 					{
 						if (TOI1 > TOI0)
 						{
-							NodeStack[NodeStackNum++] = FNodeQueueEntry{ Node.ChildrenNodes[1], TOI1 };
-							NodeStack[NodeStackNum++] = FNodeQueueEntry{ Node.ChildrenNodes[0], TOI0 };
+							NodeStack.Emplace(FNodeQueueEntry{Node.ChildrenNodes[1], TOI1});
+							NodeStack.Emplace(FNodeQueueEntry{Node.ChildrenNodes[0], TOI0});
 						}
 						else
 						{
-							NodeStack[NodeStackNum++] = FNodeQueueEntry{ Node.ChildrenNodes[0], TOI0 };
-							NodeStack[NodeStackNum++] = FNodeQueueEntry{ Node.ChildrenNodes[1], TOI1 };
+							NodeStack.Emplace(FNodeQueueEntry{Node.ChildrenNodes[0], TOI0});
+							NodeStack.Emplace(FNodeQueueEntry{ Node.ChildrenNodes[1], TOI1 });
 						}
 					}
 					else if (bIntersect0)
 					{
-						NodeStack[NodeStackNum++] = FNodeQueueEntry{ Node.ChildrenNodes[0], TOI0 };
+						NodeStack.Emplace(FNodeQueueEntry{Node.ChildrenNodes[0], TOI0});
 					}
 					else if (bIntersect1)
 					{
-						NodeStack[NodeStackNum++] = FNodeQueueEntry{ Node.ChildrenNodes[1], TOI1 };
+						NodeStack.Emplace(FNodeQueueEntry{ Node.ChildrenNodes[1], TOI1 });
 					}
 				}
 				else
@@ -3053,7 +3051,7 @@ private:
 					{
 						if (TAABBTreeIntersectionHelper<TQueryFastData, Query>::Intersects(Start, CurData, TOI, FAABB3(AABB.Min(), AABB.Max()), QueryBounds, QueryHalfExtents, Dir, InvDir, bParallel))
 						{
-							NodeStack[NodeStackNum++] = FNodeQueueEntry{ Node.ChildrenNodes[Idx], TOI };
+							NodeStack.Emplace(FNodeQueueEntry{ Node.ChildrenNodes[Idx], TOI });
 						}
 						++Idx;
 					}
