Iris - Support changing poll frequency on the fly.

diff --git a/Engine/Source/Runtime/Engine/Private/Net/Iris/ReplicationSystem/ReplicationSystemUtil.cpp b/Engine/Source/Runtime/Engine/Private/Net/Iris/ReplicationSystem/ReplicationSystemUtil.cpp
--- a/Engine/Source/Runtime/Engine/Private/Net/Iris/ReplicationSystem/ReplicationSystemUtil.cpp
+++ b/Engine/Source/Runtime/Engine/Private/Net/Iris/ReplicationSystem/ReplicationSystemUtil.cpp
@@ -105,6 +105,12 @@ FNetHandle FReplicationSystemUtil::GetNetHandle(const UActorComponent* SubObject
 	return NetHandle;
 }
 
+FNetHandle FReplicationSystemUtil::GetNetHandle(const UObject* Object)
+{
+	FNetHandle NetHandle = FNetHandleManager::GetNetHandle(Object);
+	return NetHandle;
+}
+
 void FReplicationSystemUtil::BeginReplication(AActor* Actor, const FActorBeginReplicationParams& Params)
 {
 	if (const UWorld* World = Actor->GetWorld())
@@ -635,6 +641,30 @@ void FReplicationSystemUtil::ClearCullDistanceSqrOverride(const AActor* Actor)
 	});
 }
 
+void FReplicationSystemUtil::SetPollFrequency(const UObject* Object, float PollFrequency)
+{
+	FNetHandle NetHandle = GetNetHandle(Object);
+	if (!NetHandle.IsValid())
+	{
+		return;
+	}
+	
+	ReplicationSystemUtil::ForEachReplicationSystem([NetHandle, PollFrequency](UReplicationSystem* ReplicationSystem)
+	{
+		if (ReplicationSystem->IsServer())
+		{
+			if (UObjectReplicationBridge* Bridge = ReplicationSystem->GetReplicationBridgeAs<UObjectReplicationBridge>())
+			{
+				FNetRefHandle RefHandle = Bridge->GetReplicatedRefHandle(NetHandle);
+				if (RefHandle.IsValid())
+				{
+					Bridge->SetPollFrequency(RefHandle, PollFrequency);
+				}
+			}
+		}
+	});
+}
+
 }
 
 #endif
diff --git a/Engine/Source/Runtime/Engine/Public/Net/Iris/ReplicationSystem/ReplicationSystemUtil.h b/Engine/Source/Runtime/Engine/Public/Net/Iris/ReplicationSystem/ReplicationSystemUtil.h
--- a/Engine/Source/Runtime/Engine/Public/Net/Iris/ReplicationSystem/ReplicationSystemUtil.h
+++ b/Engine/Source/Runtime/Engine/Public/Net/Iris/ReplicationSystem/ReplicationSystemUtil.h
@@ -46,6 +46,9 @@ struct FReplicationSystemUtil
 	/** Returns the NetHandle for an actor component. The returned handle may be invalid. */
 	ENGINE_API static FNetHandle GetNetHandle(const UActorComponent* SubObject);
 
+	/** Returns the NetHandle for an object. The returned handle may be invalid. */
+	ENGINE_API static FNetHandle GetNetHandle(const UObject* Object);
+
 	/** Begins replication of an actor and all of its registered subobjects. If any ReplicationSystem wants to replicate the actor a NetHandle will be created. */
 	ENGINE_API static void BeginReplication(AActor* Actor, const FActorBeginReplicationParams& Params);
 
@@ -127,6 +130,9 @@ struct FReplicationSystemUtil
 
 	/** Clears any previously set squared cull distance for an actor. This will cause affected code to respect the NetCullDistanceSquared property. */
 	ENGINE_API static void ClearCullDistanceSqrOverride(const AActor* Actor);
+
+	/** Set the poll frequency for an object and its subobjects. */
+	ENGINE_API static void SetPollFrequency(const UObject* Object, float CullDistSqr);
 };
 
 }
diff --git a/Engine/Source/Runtime/Experimental/Iris/Core/Private/Iris/ReplicationSystem/ObjectReplicationBridge.cpp b/Engine/Source/Runtime/Experimental/Iris/Core/Private/Iris/ReplicationSystem/ObjectReplicationBridge.cpp
--- a/Engine/Source/Runtime/Experimental/Iris/Core/Private/Iris/ReplicationSystem/ObjectReplicationBridge.cpp
+++ b/Engine/Source/Runtime/Experimental/Iris/Core/Private/Iris/ReplicationSystem/ObjectReplicationBridge.cpp
@@ -1473,10 +1473,10 @@ void UObjectReplicationBridge::ReinitPollFrequency()
 
 			PollFrequencyLimiter->SetPollFramePeriod(RootObjectIndex, PollFramePeriod);
 
-			// Set the subobjects of the object
+			// Make sure the subobjects are polled the same frame as the root object.
 			for (const FInternalNetRefIndex SubObjectIndex : LocalNetRefHandleManager.GetSubObjects(RootObjectIndex))
 			{
-				PollFrequencyLimiter->SetPollFramePeriod(SubObjectIndex, PollFramePeriod);
+				PollFrequencyLimiter->SetPollWithObject(RootObjectIndex, SubObjectIndex);
 			}
 		}
 	};
@@ -1486,3 +1486,27 @@ void UObjectReplicationBridge::ReinitPollFrequency()
 
 	FNetBitArrayView::ForAllSetBits(RootObjects, SubObjects, FNetBitArrayView::AndNotOp, UpdatePollFrequency);
 }
+
+void UObjectReplicationBridge::SetPollFrequency(FNetRefHandle RefHandle, float PollFrequency)
+{
+	using namespace UE::Net;
+	using namespace UE::Net::Private;
+
+	FReplicationSystemInternal* ReplicationSystemInternal = GetReplicationSystem()->GetReplicationSystemInternal();
+	const FNetRefHandleManager& LocalNetRefHandleManager = ReplicationSystemInternal->GetNetRefHandleManager();
+
+	FInternalNetRefIndex RootObjectIndex = LocalNetRefHandleManager.GetInternalIndex(RefHandle);
+	if (RootObjectIndex == FNetRefHandleManager::InvalidInternalIndex)
+	{
+		return;
+	}
+
+	const uint8 PollFramePeriod = ConvertPollFrequencyIntoFrames(PollFrequency);
+	PollFrequencyLimiter->SetPollFramePeriod(RootObjectIndex, PollFramePeriod);
+
+	// Make sure the subobjects are polled the same frame as the root object.
+	for (const FInternalNetRefIndex SubObjectIndex : LocalNetRefHandleManager.GetSubObjects(RootObjectIndex))
+	{
+		PollFrequencyLimiter->SetPollWithObject(RootObjectIndex, SubObjectIndex);
+	}
+}
diff --git a/Engine/Source/Runtime/Experimental/Iris/Core/Public/Iris/ReplicationSystem/ObjectReplicationBridge.h b/Engine/Source/Runtime/Experimental/Iris/Core/Public/Iris/ReplicationSystem/ObjectReplicationBridge.h
--- a/Engine/Source/Runtime/Experimental/Iris/Core/Public/Iris/ReplicationSystem/ObjectReplicationBridge.h
+++ b/Engine/Source/Runtime/Experimental/Iris/Core/Public/Iris/ReplicationSystem/ObjectReplicationBridge.h
@@ -148,6 +148,9 @@ public:
 	/** Trigger replication of dirty state for object wanting to be dormant. */
 	IRISCORE_API void ForceUpdateWantsToBeDormantObject(FNetRefHandle Handle);	
 
+	/** Set poll frequency on root object and its subobjects. They will be polled on the same frame. */
+	IRISCORE_API void SetPollFrequency(FNetRefHandle RootHandle, float PollFrequency);
+
 protected:
 	IRISCORE_API virtual ~UObjectReplicationBridge();
 
