Collated some literals that are dotted about the IAS cache.

#rnx

diff --git a/Engine/Source/Runtime/Experimental/IoStoreOnDemand/Private/JournaledCache.cpp b/Engine/Source/Runtime/Experimental/IoStoreOnDemand/Private/JournaledCache.cpp
--- a/Engine/Source/Runtime/Experimental/IoStoreOnDemand/Private/JournaledCache.cpp
+++ b/Engine/Source/Runtime/Experimental/IoStoreOnDemand/Private/JournaledCache.cpp
@@ -73,6 +73,11 @@ struct FDebugCacheEntry
 	uint32 _Unused : 31;
 };
 
+////////////////////////////////////////////////////////////////////////////////
+static constexpr FStringView GetCacheFsDir()	 { return FStringView(TEXT("ias")); }
+static constexpr FStringView GetCacheFsSuffix()  { return FStringView(TEXT(".cache.0")); }
+static constexpr FStringView GetCacheJrnSuffix() { return FStringView(TEXT(".jrn")); }
+
 ////////////////////////////////////////////////////////////////////////////////
 template <bool IsExclusive>
 class FAccessScope
@@ -493,7 +498,7 @@ void FDiskJournal::OpenJrnFile()
 void FDiskJournal::GetPath(TStringBuilder<64>& Out)
 {
 	Out << RootPath;
-	Out << TEXT(".jrn");
+	Out << GetCacheJrnSuffix();
 }
 
 ////////////////////////////////////////////////////////////////////////////////
@@ -1465,13 +1470,15 @@ private:
 ////////////////////////////////////////////////////////////////////////////////
 bool FJournaledCache::Initialize(const TCHAR* RootDir, const FIasCacheConfig& Config)
 {
+	using namespace JournaledCache;
+
 	// Filesystem setup
 	FStringView Name = Config.Name;
 	check(Name.Len() > 0 && !Name.EndsWith('/') && !Name.EndsWith('\\'));
 
 	TStringBuilder<256> CachePath;
 	CachePath << RootDir;
-	FPathViews::Append(CachePath, "ias");
+	FPathViews::Append(CachePath, GetCacheFsDir());
 	FPathViews::Append(CachePath, FPathViews::GetPath(Name));
 
 	if (IFileManager& Ifm = IFileManager::Get(); !Ifm.MakeDirectory(CachePath.ToString(), true))
@@ -1481,7 +1488,7 @@ bool FJournaledCache::Initialize(const TCHAR* RootDir, const FIasCacheConfig& Co
 	}
 
 	FPathViews::Append(CachePath, FPathViews::GetBaseFilename(Name));
-	CachePath << TEXT(".cache.0");
+	CachePath << GetCacheFsSuffix();
 
 	// Inner cache
 	FCacheInner::FConfig EventualConfig;
