Moved cache tests out to their own function (albeit partially disabled)

#rnx
(git:14863e7 - Martin Ridgers - [REDACTED_PHONE]:21:57 +0200)

diff --git a/Engine/Source/Runtime/Experimental/IoStoreOnDemand/Private/JournaledCache.cpp b/Engine/Source/Runtime/Experimental/IoStoreOnDemand/Private/JournaledCache.cpp
--- a/Engine/Source/Runtime/Experimental/IoStoreOnDemand/Private/JournaledCache.cpp
+++ b/Engine/Source/Runtime/Experimental/IoStoreOnDemand/Private/JournaledCache.cpp
@@ -1766,154 +1766,156 @@ static void MemCacheTests(FSupport& Support)
 }
 
 ////////////////////////////////////////////////////////////////////////////////
-IOSTOREONDEMAND_API void Tests()
+static void CacheTests(FSupport& Support)
 {
-	FSupport Support;
-	MemCacheTests(Support);
+	using namespace JournaledCache;
 
-#if 0
-	// Cache {{{2
-	{
-		FCache::FConfig Config;
-		Config.Path = FPaths::ProjectPersistentDownloadDir() / TEXT("ias_cache_test");
-		Config.MemoryQuota = uint32(512_Ki);
-		Config.DiskQuota = 8_Mi;
-		Config.JournalQuota = uint32(7_Ki);
-		Config.JournalFlushInterval = 4;
-		Config.DropCache = true;
-		FCache Cache(MoveTemp(Config));
-
-		auto PrimePuts = [&] (int64 PutMax) {
-			TMap<uint64, FIoBuffer> Ret;
-			while (true)
+	FCache::FConfig Config;
+	Config.Path = FPaths::ProjectPersistentDownloadDir() / TEXT("ias_cache_test");
+	Config.MemoryQuota = uint32(512_Ki);
+	Config.DiskQuota = 8_Mi;
+	Config.JournalQuota = uint32(7_Ki);
+	Config.DropCache = true;
+	FCache Cache(MoveTemp(Config));
+
+	auto PrimePuts = [&] (int64 PutMax) {
+		TMap<uint64, FIoBuffer> Ret;
+		while (true)
+		{
+			uint64 Size = Support.Mix() & ((128_Ki) - 1);
+			if ((PutMax -= Size) < 0)
 			{
-				uint64 Size = MixTh() & ((128_Ki) - 1);
-				if ((PutMax -= Size) < 0)
-				{
-					break;
-				}
-				FIoBuffer Data = DummyData(Size);
-				uint64 Key = KeyGen(Data);
-				Cache.Put(Key, Data);
-				Ret.Add(Key, Data);
+				break;
 			}
-			return Ret;
-		};
+			FIoBuffer Data = Support.DummyData(Size);
+			uint64 Key = KeyGen(Data);
+			Cache.Put(Key, Data);
+			Ret.Add(Key, Data);
+		}
+		return Ret;
+	};
 
-		uint32 WriteAllowance;
+	uint32 WriteAllowance;
 
-		// no-op
-		WriteAllowance = uint32(1_Ki);
-		Cache.Flush(WriteAllowance);
-		Cache.Flush(WriteAllowance);
-		Cache.Reset();
+	// no-op
+	WriteAllowance = uint32(1_Ki);
+	Cache.WriteMemToDisk(WriteAllowance);
+	Cache.Flush();
+	Cache.Reset();
 
-		WriteAllowance = uint32(512_Ki);
-		PrimePuts(WriteAllowance);
-		Cache.Flush(0);
-		Cache.Flush(WriteAllowance);
-		Cache.Flush(WriteAllowance);
-		Cache.Flush(WriteAllowance);
-		Cache.Reset();
+	WriteAllowance = uint32(512_Ki);
+	PrimePuts(WriteAllowance);
+	Cache.WriteMemToDisk(0);
+	Cache.WriteMemToDisk(WriteAllowance);
+	Cache.Flush();
+	Cache.Reset();
 
-		auto Validate = [&] () {
-			struct FVisitState {
-				FCache& Cache;
-				const uint8* WorkRange[2];
-			};
-			auto Visitor = [] (void* Param, const FDebugCacheEntry& Entry) {
-				auto* State = (FVisitState*)Param;
-				
-				auto GetEntry = State->Cache.Get(Entry.Key);
-				check(GetEntry.IsHit());
-				FIoBuffer Data;
-				check(GetEntry.Materialize(Data));
-				check(Data.GetSize() == Entry.Size);
-				check(KeyGen(Data) == Entry.Key);
-
-				int32 IsFromDisk = 0;
-				IsFromDisk |= Data.GetData() >= State->WorkRange[1];
-				IsFromDisk |= (Data.GetData() + Data.GetSize()) <= State->WorkRange[0];
-				check(Entry.IsMemCache != IsFromDisk);
-			};
-			FVisitState State = {Cache, {Working, Working + WorkingSize}};
-			return Cache.DebugVisit(&State, Visitor);
+#if 0
+	auto Validate = [&] () {
+		struct FVisitState {
+			FCache& Cache;
+			const uint8* WorkRange[2];
+		};
+		auto Visitor = [] (void* Param, const FDebugCacheEntry& Entry) {
+			auto* State = (FVisitState*)Param;
+			
+			auto GetEntry = State->Cache.Get(Entry.Key);
+			check(GetEntry.IsHit());
+			FIoBuffer Data;
+			check(GetEntry.Materialize(Data));
+			check(Data.GetSize() == Entry.Size);
+			check(KeyGen(Data) == Entry.Key);
+
+			int32 IsFromDisk = 0;
+			IsFromDisk |= Data.GetData() >= State->WorkRange[1];
+			IsFromDisk |= (Data.GetData() + Data.GetSize()) <= State->WorkRange[0];
+			check(Entry.IsMemCache != IsFromDisk);
 		};
-		check(Validate() == 0);
+		FVisitState State = {Cache, {Working, Working + WorkingSize}};
+		return Cache.DebugVisit(&State, Visitor);
+	};
+	check(Validate() == 0);
 
-		WriteAllowance = uint32(512_Ki);
+	WriteAllowance = uint32(512_Ki);
 
-		// simple
-		for (uint32 i : {1, 2, 4, 7, 11})
+	// simple
+	for (uint32 i : {1, 2, 4, 7, 11})
+	{
+		for (uint32 j = 0; j < i; ++j)
 		{
-			for (uint32 j = 0; j < i; ++j)
-			{
-				FIoBuffer Data = DummyData(32);
-				Cache.Put(KeyGen(Data), Data);
-			}
-			Cache.Flush(WriteAllowance);
-			check(Validate() == i);
-
-			Cache.Reset();
-			check(Cache.Load());
-			check(Validate() == 0); // not enough flushes to write a journal
+			FIoBuffer Data = Support.DummyData(32);
+			Cache.Put(KeyGen(Data), Data);
 		}
+		Cache.Flush(WriteAllowance);
+		check(Validate() == i);
+
 		Cache.Reset();
+		check(Cache.Load());
+		check(Validate() == 0); // not enough flushes to write a journal
+	}
+	Cache.Reset();
 
-		// general
-		for (int32 i : {1, 4, 136, 137})
+	// general
+	for (int32 i : {1, 4, 136, 137})
+	{
+		for (int32 j = 0; j < i; ++j)
 		{
-			for (int32 j = 0; j < i; ++j)
-			{
-				PrimePuts(WriteAllowance);
-				Cache.Flush(WriteAllowance);
-			}
-			uint32 PreCount = Validate();
-
-			Cache.Reset();
-			check(Cache.Load());
-
-			uint32 PostCount = Validate();
-			check(!PostCount == !(i / 4)); // JournalFlushInterval
-			check(PostCount <= PreCount);
-			check((PostCount == PreCount) == ((i & 3) == 0));
+			PrimePuts(WriteAllowance);
+			Cache.Flush(WriteAllowance);
 		}
+		uint32 PreCount = Validate();
+
 		Cache.Reset();
+		check(Cache.Load());
+
+		uint32 PostCount = Validate();
+		check(!PostCount == !(i / 4)); // JournalFlushInterval
+		check(PostCount <= PreCount);
+		check((PostCount == PreCount) == ((i & 3) == 0));
+	}
+	Cache.Reset();
 
-		// power 2
-		for (int32 i : {74, 75})
+	// power 2
+	for (int32 i : {74, 75})
+	{
+		while (i--)
 		{
-			while (i--)
+			for (int32 j = 0; j < 3; ++j)
 			{
-				for (int32 j = 0; j < 3; ++j)
-				{
-					FIoBuffer Data = DummyData(64_Ki - 4);
-					Cache.Put(KeyGen(Data), Data);
-					Cache.Flush(WriteAllowance);
-				}
+				FIoBuffer Data = Support.DummyData(64_Ki - 4);
+				Cache.Put(KeyGen(Data), Data);
 				Cache.Flush(WriteAllowance);
 			}
-			Validate();
-			Cache.Reset();
+			Cache.Flush(WriteAllowance);
 		}
+		Validate();
+		Cache.Reset();
+	}
 
-		// marker wrap
-		// one-phrase journal
-		// journal paragraphs that are all the same size
-		// phrases with no entries
+	// marker wrap
+	// one-phrase journal
+	// journal paragraphs that are all the same size
+	// phrases with no entries
 
-		// cache items larger than pending memory
-		// cache items larger than write allowance
+	// cache items larger than pending memory
+	// cache items larger than write allowance
 
-		// journal wrapping without truncation
+	// journal wrapping without truncation
 
-		// changes in max data/journal size
+	// changes in max data/journal size
 
-		// don't load-and-sort so many paragraphs (only need max-data size)
-	} // }}}
+	// don't load-and-sort so many paragraphs (only need max-data size)
 #endif // 0
 }
 
+////////////////////////////////////////////////////////////////////////////////
+IOSTOREONDEMAND_API void Tests()
+{
+	FSupport Support;
+	MemCacheTests(Support);
+	CacheTests(Support);
+}
+
 } // namespace IasJournaledFileCacheTest
 
 #endif // IAS_JOURNALED_CACHE_TEST
