Should delete the cache object that fails to initialize

#rnx

diff --git a/Engine/Source/Runtime/Experimental/IoStoreOnDemand/Private/JournaledCache.cpp b/Engine/Source/Runtime/Experimental/IoStoreOnDemand/Private/JournaledCache.cpp
--- a/Engine/Source/Runtime/Experimental/IoStoreOnDemand/Private/JournaledCache.cpp
+++ b/Engine/Source/Runtime/Experimental/IoStoreOnDemand/Private/JournaledCache.cpp
@@ -1651,6 +1651,7 @@ TUniquePtr<IIasCache> MakeIasCache(const TCHAR* RootPath, const FIasCacheConfig&
 		return TUniquePtr<IIasCache>(Cache);
 	}
 
+	delete Cache;
 	return nullptr;
 }
 
