Minor tidy.

#rnx

diff --git a/Engine/Source/Runtime/Experimental/IoStoreOnDemand/Private/JournaledCache.cpp b/Engine/Source/Runtime/Experimental/IoStoreOnDemand/Private/JournaledCache.cpp
--- a/Engine/Source/Runtime/Experimental/IoStoreOnDemand/Private/JournaledCache.cpp
+++ b/Engine/Source/Runtime/Experimental/IoStoreOnDemand/Private/JournaledCache.cpp
@@ -48,8 +48,11 @@ namespace UE::IO::IAS::JournaledCache
 
 // {{{1 misc ...................................................................
 
+////////////////////////////////////////////////////////////////////////////////
 static bool LoadCache(class FDiskCache&);
+using		EntryHandle = UPTRINT;
 
+////////////////////////////////////////////////////////////////////////////////
 #if !defined(IAS_HAS_WRITE_COMMIT_THRESHOLD)
 #	define IAS_HAS_WRITE_COMMIT_THRESHOLD 0
 #endif
@@ -107,9 +110,6 @@ private:
 using FReadAccess	= FAccessScope<false>;
 using FWriteAccess	= FAccessScope<true>;
 
-////////////////////////////////////////////////////////////////////////////////
-using EntryHandle = UPTRINT;
-
 
 
 // {{{1 mem-cache ..............................................................
