Added doctors visit to IAS cache and fail to initialize if it comes back from the doctors in poor health.

#rnx

diff --git a/Engine/Source/Runtime/Experimental/IoStoreOnDemand/Private/JournaledCache.cpp b/Engine/Source/Runtime/Experimental/IoStoreOnDemand/Private/JournaledCache.cpp
--- a/Engine/Source/Runtime/Experimental/IoStoreOnDemand/Private/JournaledCache.cpp
+++ b/Engine/Source/Runtime/Experimental/IoStoreOnDemand/Private/JournaledCache.cpp
@@ -52,6 +52,13 @@ namespace UE::IO::IAS::JournaledCache
 static bool LoadCache(class FDiskCache&);
 using		EntryHandle = UPTRINT;
 
+////////////////////////////////////////////////////////////////////////////////
+enum class EAilments
+{
+	NoJrnHandle		= 1 << 0,
+	NoDataHandle	= 1 << 1,
+};
+
 ////////////////////////////////////////////////////////////////////////////////
 #if !defined(IAS_HAS_WRITE_COMMIT_THRESHOLD)
 #	define IAS_HAS_WRITE_COMMIT_THRESHOLD 0
@@ -426,6 +433,7 @@ class FDiskJournal
 {
 public:
 							FDiskJournal(FStringView InRootPath, uint32 InMaxSize);
+	uint32					GetAilments() const;
 	void					Reset();
 	void					Drop();
 	int32					Flush();
@@ -458,6 +466,17 @@ FDiskJournal::FDiskJournal(FStringView InRootPath, uint32 InMaxSize)
 	OpenJrnFile();
 }
 
+////////////////////////////////////////////////////////////////////////////////
+uint32 FDiskJournal::GetAilments() const
+{
+	uint32 Ret = 0;
+	if (!JrnHandle.IsValid())
+	{
+		Ret |= uint32(EAilments::NoJrnHandle);
+	}
+	return Ret;
+}
+
 ////////////////////////////////////////////////////////////////////////////////
 void FDiskJournal::Reset()
 {
@@ -580,6 +599,7 @@ class FDiskCache
 {
 public:
 							FDiskCache(FString&& Path, uint64 InMaxDataSize, uint32 InJournalSize);
+	uint32					GetAilments() const;
 	void					Reset();
 	FDiskPhrase				OpenPhrase(uint32 DataSize);
 	void					ClosePhrase(FDiskPhrase&& Phrase);
@@ -633,6 +653,17 @@ FDiskCache::FDiskCache(FString&& Path, uint64 InMaxDataSize, uint32 InJournalSiz
 	FOnDemandIoBackendStats::Get()->OnCacheSetMaxBytes(MaxDataSize);
 }
 
+////////////////////////////////////////////////////////////////////////////////
+uint32 FDiskCache::GetAilments() const
+{
+	uint32 Ret = Journal.GetAilments();
+	if (!DataHandle.IsValid())
+	{
+		Ret |= uint32(EAilments::NoDataHandle);
+	}
+	return Ret;
+}
+
 ////////////////////////////////////////////////////////////////////////////////
 void FDiskCache::Reset()
 {
@@ -1138,6 +1169,7 @@ public:
 	};
 
 					FCache(FConfig&& Config);
+	uint32			GetAilments() const;
 	void			Reset();
 	bool			Load();
 	uint32			GetDemand() const;
@@ -1184,6 +1216,14 @@ FCache::FCache(FConfig&& Config)
 	}
 }
 
+////////////////////////////////////////////////////////////////////////////////
+uint32 FCache::GetAilments() const
+{
+	uint32 Ret = 0;
+	Ret |= DiskCache.GetAilments();
+	return Ret;
+}
+
 ////////////////////////////////////////////////////////////////////////////////
 void FCache::Reset()
 {
@@ -1495,6 +1535,13 @@ bool FJournaledCache::Initialize(const TCHAR* RootDir, const FIasCacheConfig& Co
 	static_cast<FIasCacheConfig&>(EventualConfig) = Config;
 	EventualConfig.Path = CachePath;
 	Cache = MakeUnique<FCacheInner>(MoveTemp(EventualConfig));
+
+	if (uint32 Ailments = Cache->GetAilments(); Ailments != 0)
+	{
+		UE_LOG(LogIas, Error, TEXT("JournaledCache: Error initialising inner cache '%x'"), Ailments);
+		return false;
+	}
+
 	Cache->Load();
 
 	// Thread setup
