Tests to verify path handling and ill-health does what one expects

#rnx

diff --git a/Engine/Source/Runtime/Experimental/IoStoreOnDemand/Private/JournaledCache.cpp b/Engine/Source/Runtime/Experimental/IoStoreOnDemand/Private/JournaledCache.cpp
--- a/Engine/Source/Runtime/Experimental/IoStoreOnDemand/Private/JournaledCache.cpp
+++ b/Engine/Source/Runtime/Experimental/IoStoreOnDemand/Private/JournaledCache.cpp
@@ -1737,6 +1737,22 @@ struct FSupport
 		{
 			*(uint64*)(Working + i) = Mix();
 		}
+
+		CleanFs();
+	}
+
+	void CleanFs()
+	{
+		IFileManager& Ifm = IFileManager::Get();
+		if (Ifm.DirectoryExists(*TestDir))
+		{
+			// Windows doesn't delete directories immediately and subsequent make
+			// dirs can fail. So we rename first then delete.
+			FString TempDir = TestDir + "~";
+			check(Ifm.Move(*TempDir, *TestDir, false));
+			check(Ifm.DeleteDirectory(*TempDir, false, true));
+		}
+		check(Ifm.MakeDirectory(*TestDir, true));
 	}
 
 	auto DummyData(uint64 Size)
@@ -1750,6 +1766,7 @@ struct FSupport
 	const uint64		WorkingSize = 1_Mi;
 	TUniquePtr<uint8[]> WorkingScope = TUniquePtr<uint8[]>(new uint8[WorkingSize]);
 	uint8*				Working = WorkingScope.Get();
+	FString				TestDir = FPaths::ProjectPersistentDownloadDir() / TEXT("ias_cache_test");
 };
 
 ////////////////////////////////////////////////////////////////////////////////
@@ -1819,7 +1836,7 @@ static void CacheTests(FSupport& Support)
 	using namespace JournaledCache;
 
 	FCache::FConfig Config;
-	Config.Path = FPaths::ProjectPersistentDownloadDir() / TEXT("ias_cache_test");
+	Config.Path = Support.TestDir / "cache_tests";
 	Config.MemoryQuota = uint32(512_Ki);
 	Config.DiskQuota = 8_Mi;
 	Config.JournalQuota = uint32(7_Ki);
@@ -1956,10 +1973,64 @@ static void CacheTests(FSupport& Support)
 #endif // 0
 }
 
+////////////////////////////////////////////////////////////////////////////////
+static void MiscTests(FSupport& Support)
+{
+	using namespace JournaledCache;
+
+	FIasCacheConfig Config;
+	Config.Name = TEXT("misc");
+
+	{ // Benign
+		Support.CleanFs();
+		auto Jc = MakeIasCache(*Support.TestDir, Config);
+		check(Jc.IsValid());
+		Jc.Reset();
+	}
+
+	{ // Ensure path creation
+		Support.CleanFs();
+		for (int32 i : { 0, 1 })
+		{
+			const TCHAR* TestName = TEXT("m/i/s/c");
+			Config.Name = TestName + i;
+			auto Jc = MakeIasCache(*Support.TestDir, Config);
+			check(Jc.IsValid());
+			Jc.Reset();
+		}
+		Config.Name = TEXT("misc");
+	}
+
+	{ // Unable to create files
+		Config.Name = TEXT("Blocked");
+
+		for (int32 i : { 0, 1 })
+		{
+			Support.CleanFs();
+
+			FString Blocker = Support.TestDir;
+			Blocker /= GetCacheFsDir();
+			Blocker /= FString(Config.Name);
+			Blocker += GetCacheFsSuffix();
+			if (i == 1)
+			{
+				Blocker += GetCacheJrnSuffix();
+			}
+
+			IFileManager& Ifm = IFileManager::Get();
+			check(Ifm.MakeDirectory(*Blocker, true));
+
+			auto Jc = MakeIasCache(*Support.TestDir, Config);
+			check(!Jc.IsValid())
+		}
+	}
+}
+
 ////////////////////////////////////////////////////////////////////////////////
 IOSTOREONDEMAND_API void Tests()
 {
 	FSupport Support;
+	MiscTests(Support);
 	MemCacheTests(Support);
 	CacheTests(Support);
 }
