Log detail as to the reason why IAS' file cache is off.

#rnx

diff --git a/Engine/Source/Runtime/Experimental/IoStoreOnDemand/Private/IoStoreOnDemand.cpp b/Engine/Source/Runtime/Experimental/IoStoreOnDemand/Private/IoStoreOnDemand.cpp
--- a/Engine/Source/Runtime/Experimental/IoStoreOnDemand/Private/IoStoreOnDemand.cpp
+++ b/Engine/Source/Runtime/Experimental/IoStoreOnDemand/Private/IoStoreOnDemand.cpp
@@ -1568,14 +1568,16 @@ void FIoStoreOnDemandModule::StartupModule()
 	Endpoint.EndpointType = EOnDemandEndpointType::CDN;
 
 	TSharedPtr<IIasCache> Cache;
-	if (FIasCacheConfig Config = GetIasCacheConfig(CommandLine); Config.DiskQuota > 0)
+	FIasCacheConfig CacheConfig = GetIasCacheConfig(CommandLine);
+	if (CacheConfig.DiskQuota > 0)
 	{
 		FString CacheDir = FPaths::ProjectPersistentDownloadDir();
-		Cache = MakeShareable(MakeIasCache(*CacheDir, Config).Release());
+		Cache = MakeShareable(MakeIasCache(*CacheDir, CacheConfig).Release());
 	}
-	else
+	if (!Cache.IsValid())
 	{
-		UE_LOG(LogIas, Log, TEXT("File cache disabled. Streaming only."));
+		UE_LOG(LogIas, Log, TEXT("File cache disabled - streaming only (%s)"),
+			(CacheConfig.DiskQuota > 0) ? TEXT("init-fail") : TEXT("zero-quota"));
 	}
 
 	Backend = MakeOnDemandIoDispatcherBackend(Cache);
