Iris
- RPC:s are now correctly delayed if we are waiting for references pending async loading

diff --git a/Engine/Source/Runtime/Experimental/Iris/Core/Private/Iris/ReplicationSystem/NetBlob/NetBlob.cpp b/Engine/Source/Runtime/Experimental/Iris/Core/Private/Iris/ReplicationSystem/NetBlob/NetBlob.cpp
--- a/Engine/Source/Runtime/Experimental/Iris/Core/Private/Iris/ReplicationSystem/NetBlob/NetBlob.cpp
+++ b/Engine/Source/Runtime/Experimental/Iris/Core/Private/Iris/ReplicationSystem/NetBlob/NetBlob.cpp
@@ -88,14 +88,10 @@ void FNetBlob::Deserialize(FNetSerializationContext& Context)
 	DeserializeBlob(Context);
 }
 
-ENetObjectReferenceResolveResult FNetBlob::ResolveObjectReferences(FNetSerializationContext& Context) const
+void FNetBlob::CollectObjectReferences(FNetSerializationContext& Context, FNetReferenceCollector& Collector) const
 {
-	if (BlobDescriptor.IsValid() && QuantizedBlobState.IsValid())
-	{
-		return Private::FReplicationStateOperationsInternal::TryToResolveObjectReferences(Context, QuantizedBlobState.Get(), BlobDescriptor);
-	}
-	
-	return ENetObjectReferenceResolveResult::None;
+	const FNetSerializerChangeMaskParam InitStateChangeMaskInfo = { 0 };
+	Private::FReplicationStateOperationsInternal::CollectReferences(Context, Collector, InitStateChangeMaskInfo, QuantizedBlobState.Get(), BlobDescriptor);
 }
 
 void FNetBlob::SerializeBlob(FNetSerializationContext& Context) const
diff --git a/Engine/Source/Runtime/Experimental/Iris/Core/Private/Iris/ReplicationSystem/ObjectReferenceCache.cpp b/Engine/Source/Runtime/Experimental/Iris/Core/Private/Iris/ReplicationSystem/ObjectReferenceCache.cpp
--- a/Engine/Source/Runtime/Experimental/Iris/Core/Private/Iris/ReplicationSystem/ObjectReferenceCache.cpp
+++ b/Engine/Source/Runtime/Experimental/Iris/Core/Private/Iris/ReplicationSystem/ObjectReferenceCache.cpp
@@ -8,6 +8,7 @@
 #include "Iris/ReplicationSystem/ReplicationSystem.h"
 #include "Iris/ReplicationSystem/ReplicationSystemInternal.h"
 #include "Iris/ReplicationSystem/StringTokenStore.h"
+#include "Iris/ReplicationSystem/PendingBatchData.h"
 #include "Iris/Serialization/InternalNetSerializationContext.h"
 #include "Iris/Serialization/NetBitStreamReader.h"
 #include "Iris/Serialization/NetBitStreamUtil.h"
@@ -633,6 +634,35 @@ bool FObjectReferenceCache::IsNetRefHandleBroken(FNetRefHandle RefHandle, bool b
 	return CacheObjectPtr ? CacheObjectPtr->bIsBroken : bMustBeRegistered;
 }
 
+bool FObjectReferenceCache::IsNetRefHandlePending(FNetRefHandle NetRefHandle, const FPendingBatches& PendingBatches) const
+{
+	// Check Outer chain
+	while (NetRefHandle.IsValid())
+	{
+		// Need lambda to figure this one out
+		if (PendingBatches.Find(NetRefHandle))
+		{
+			return true;
+		}
+
+		const FCachedNetObjectReference* CacheObjectPtr = ReferenceHandleToCachedReference.Find(NetRefHandle);
+
+		if (!CacheObjectPtr)
+		{
+			return false;
+		}
+
+		if (CacheObjectPtr->bIsPending)
+		{
+			return true;
+		}
+
+		NetRefHandle = CacheObjectPtr->OuterNetRefHandle;
+	}
+
+	return false;
+}
+
 // $IRIS: $TODO: Most of the logic comes from GUIDCache::GetObjectFromNetGUID so we want to keep this in sync
 UObject* FObjectReferenceCache::ResolveObjectReferenceHandleInternal(FNetRefHandle RefHandle, const FNetObjectResolveContext& ResolveContext, bool& bOutMustBeMapped)
 {
diff --git a/Engine/Source/Runtime/Experimental/Iris/Core/Private/Iris/ReplicationSystem/ObjectReferenceCache.h b/Engine/Source/Runtime/Experimental/Iris/Core/Private/Iris/ReplicationSystem/ObjectReferenceCache.h
--- a/Engine/Source/Runtime/Experimental/Iris/Core/Private/Iris/ReplicationSystem/ObjectReferenceCache.h
+++ b/Engine/Source/Runtime/Experimental/Iris/Core/Private/Iris/ReplicationSystem/ObjectReferenceCache.h
@@ -19,6 +19,7 @@ namespace UE::Net
 	{
 		class FNetRefHandleManager;
 		class FNetExportContext;
+		struct FPendingBatches;
 	}
 }
 
@@ -60,6 +61,9 @@ public:
 	// Returns true of this NetRefHandle is marked as broken
 	bool IsNetRefHandleBroken(FNetRefHandle Handle, bool bMustBeRegistered) const;
 
+	// Returns true of the provided NetRefHandle or one of its outers is pending async loading.
+	bool IsNetRefHandlePending(FNetRefHandle NetRefHandle, const FPendingBatches& PendingBatches) const;
+
 	// Find replicated outer
 	FNetObjectReference GetReplicatedOuter(const FNetObjectReference& Reference) const;
 
diff --git a/Engine/Source/Runtime/Experimental/Iris/Core/Private/Iris/ReplicationSystem/ObjectReferenceCacheFwd.h b/Engine/Source/Runtime/Experimental/Iris/Core/Private/Iris/ReplicationSystem/ObjectReferenceCacheFwd.h
--- a/Engine/Source/Runtime/Experimental/Iris/Core/Private/Iris/ReplicationSystem/ObjectReferenceCacheFwd.h
+++ b/Engine/Source/Runtime/Experimental/Iris/Core/Private/Iris/ReplicationSystem/ObjectReferenceCacheFwd.h
@@ -22,5 +22,6 @@ namespace UE::Net
 	namespace Private
 	{
 		class FNetExportContext;
+		class FNetPendingBatches;
 	}
 }
diff --git a/Engine/Source/Runtime/Experimental/Iris/Core/Private/Iris/ReplicationSystem/PendingBatchData.h b/Engine/Source/Runtime/Experimental/Iris/Core/Private/Iris/ReplicationSystem/PendingBatchData.h
new file mode 100644
--- /dev/null
+++ b/Engine/Source/Runtime/Experimental/Iris/Core/Private/Iris/ReplicationSystem/PendingBatchData.h
@@ -0,0 +1,53 @@
+// Copyright Epic Games, Inc. All Rights Reserved.
+
+#pragma once
+
+#include "Containers/Array.h"
+#include "Iris/ReplicationSystem/NetRefHandle.h"
+
+namespace UE::Net::Private
+{
+// Queued data chunk
+struct FQueuedDataChunk
+{
+	FQueuedDataChunk()
+	: StorageOffset(0U)
+	, NumBits(0U)
+	, bHasBatchOwnerData(0U)
+	, bIsEndReplicationChunk(0U)
+	{
+	}
+
+	uint32 StorageOffset;
+	uint32 NumBits : 30;
+	uint32 bHasBatchOwnerData : 1;
+	uint32 bIsEndReplicationChunk : 1;
+};
+
+// Struct to contain storage and required data for queued batches pending must be mapped references
+struct FPendingBatchData
+{
+	// We use a single array to store the actual data, it will grow if required.
+	TArray<uint32, TInlineAllocator<32>> DataChunkStorage;		
+	TArray<FQueuedDataChunk, TInlineAllocator<4>> QueuedDataChunks;
+
+	// Must be mapped references pending resolve
+	TArray<FNetRefHandle, TInlineAllocator<4>> PendingMustBeMappedReferences;
+
+	// Resolved references for which we have are holding on to references to avoid GC
+	TArray<FNetRefHandle, TInlineAllocator<4>> ResolvedReferences;
+
+	// Batch owner with queued data chunks
+	FNetRefHandle Handle;
+};
+
+struct FPendingBatches
+{
+	FPendingBatchData* Find(FNetRefHandle NetRefHandle) { return PendingBatches.FindByPredicate([&NetRefHandle](const FPendingBatchData& Entry) { return Entry.Handle == NetRefHandle; });}
+	const FPendingBatchData* Find(FNetRefHandle NetRefHandle) const { return PendingBatches.FindByPredicate([&NetRefHandle](const FPendingBatchData& Entry) { return Entry.Handle == NetRefHandle; });}
+	bool GetHasPendingBatches() const { return !PendingBatches.IsEmpty(); }
+
+	TArray<FPendingBatchData> PendingBatches;
+};
+
+}
diff --git a/Engine/Source/Runtime/Experimental/Iris/Core/Private/Iris/ReplicationSystem/ReplicationOperationsInternal.cpp b/Engine/Source/Runtime/Experimental/Iris/Core/Private/Iris/ReplicationSystem/ReplicationOperationsInternal.cpp
--- a/Engine/Source/Runtime/Experimental/Iris/Core/Private/Iris/ReplicationSystem/ReplicationOperationsInternal.cpp
+++ b/Engine/Source/Runtime/Experimental/Iris/Core/Private/Iris/ReplicationSystem/ReplicationOperationsInternal.cpp
@@ -365,36 +365,6 @@ void FReplicationStateOperationsInternal::CollectReferencesWithMask(FNetSerializ
 	}
 }
 
-ENetObjectReferenceResolveResult FReplicationStateOperationsInternal::TryToResolveObjectReferences(FNetSerializationContext& Context, uint8* RESTRICT InternalBuffer, const FReplicationStateDescriptor* Descriptor)
-{
-	FNetSerializationContext LocalContext;
-	FNetReferenceCollector Collector;
-
-	if (Descriptor->IsInitState() || Descriptor->MemberChangeMaskDescriptors == nullptr)
-	{
-		const FNetSerializerChangeMaskParam InitStateChangeMaskInfo = { 0 };
-		FReplicationStateOperationsInternal::CollectReferences(LocalContext, Collector, InitStateChangeMaskInfo, InternalBuffer, Descriptor);
-	}
-	else
-	{
-		FReplicationStateOperationsInternal::CollectReferencesWithMask(LocalContext, Collector, 0, InternalBuffer, Descriptor);
-	}
-
-	// Process references
-	const FInternalNetSerializationContext* InternalContext = Context.GetInternalContext();
-	const FNetObjectResolveContext& ResolveContext = InternalContext->ResolveContext;
-	FObjectReferenceCache& ObjectReferenceCache = *InternalContext->ObjectReferenceCache;
-
-	ENetObjectReferenceResolveResult Result = ENetObjectReferenceResolveResult::None;
-	for (const FNetReferenceCollector::FReferenceInfo& Info : MakeArrayView(Collector.GetCollectedReferences()))
-	{
-		UObject* ResolvedObject;
-		Result |= ObjectReferenceCache.ResolveObjectReference(Info.Reference, ResolveContext, ResolvedObject);
-	}
-
-	return Result;
-};
-
 void FReplicationProtocolOperationsInternal::CloneDynamicState(FNetSerializationContext& Context, uint8* RESTRICT DstObjectStateBuffer, const uint8* RESTRICT SrcObjectStateBuffer, const FReplicationProtocol* Protocol)
 {
 	const FReplicationStateDescriptor** ReplicationStateDescriptors = Protocol->ReplicationStateDescriptors;
diff --git a/Engine/Source/Runtime/Experimental/Iris/Core/Private/Iris/ReplicationSystem/ReplicationOperationsInternal.h b/Engine/Source/Runtime/Experimental/Iris/Core/Private/Iris/ReplicationSystem/ReplicationOperationsInternal.h
--- a/Engine/Source/Runtime/Experimental/Iris/Core/Private/Iris/ReplicationSystem/ReplicationOperationsInternal.h
+++ b/Engine/Source/Runtime/Experimental/Iris/Core/Private/Iris/ReplicationSystem/ReplicationOperationsInternal.h
@@ -61,9 +61,6 @@ struct FReplicationStateOperationsInternal
 
 	/** Collect references from a state based on the provided changemask information */
 	static IRISCORE_API void CollectReferencesWithMask(FNetSerializationContext& Context, FNetReferenceCollector& Collector, const uint32 ChangeMaskOffset, const uint8* RESTRICT SrcInternalBuffer,  const FReplicationStateDescriptor* Descriptor);
-
-	/** Try to resolve the references in the state described by the descriptor, returns true if all references can be resolved */
-	static IRISCORE_API ENetObjectReferenceResolveResult TryToResolveObjectReferences(FNetSerializationContext& Context, uint8* RESTRICT InternalBuffer, const FReplicationStateDescriptor* Descriptor);
 };
 
 struct FReplicationProtocolOperationsInternal
diff --git a/Engine/Source/Runtime/Experimental/Iris/Core/Private/Iris/ReplicationSystem/ReplicationReader.cpp b/Engine/Source/Runtime/Experimental/Iris/Core/Private/Iris/ReplicationSystem/ReplicationReader.cpp
--- a/Engine/Source/Runtime/Experimental/Iris/Core/Private/Iris/ReplicationSystem/ReplicationReader.cpp
+++ b/Engine/Source/Runtime/Experimental/Iris/Core/Private/Iris/ReplicationSystem/ReplicationReader.cpp
@@ -221,7 +221,7 @@ void FReplicationReader::Init(const FReplicationParameters& InParameters)
 
 void FReplicationReader::Deinit()
 {
-	for (FPendingBatchData& PendingBatchData : PendingBatches)
+	for (FPendingBatchData& PendingBatchData : PendingBatches.PendingBatches)
 	{
 		UE_LOG_REPLICATIONREADER_WARNING(TEXT("FReplicationReader::Deinit NetHandle %s has %d unprocessed data batches"), *PendingBatchData.Handle.ToString(), PendingBatchData.QueuedDataChunks.Num());
 
@@ -273,7 +273,7 @@ uint32 FReplicationReader::ReadObjectsPendingDestroy(FNetSerializationContext& C
 	
 	if (!Context.HasErrorOrOverflow())
 	{
-		const bool bHasPendingBatches = !PendingBatches.IsEmpty();
+		const bool bHasPendingBatches = PendingBatches.GetHasPendingBatches();
 	
 		for (uint32 It = 0; It < ObjectsToRead; ++It)
 		{
@@ -291,7 +291,7 @@ uint32 FReplicationReader::ReadObjectsPendingDestroy(FNetSerializationContext& C
 				break;
 			}
 
-			if (FPendingBatchData* PendingBatchData = bHasPendingBatches ? PendingBatches.FindByPredicate([&SubObjectRootOrHandle](const FPendingBatchData& Entry) { return Entry.Handle == SubObjectRootOrHandle; }) : nullptr)
+			if (FPendingBatchData* PendingBatchData = bHasPendingBatches ? PendingBatches.Find(SubObjectRootOrHandle) : nullptr)
 			{
 				EnqueueEndReplication(PendingBatchData, bShouldDestroyInstance, IncompleteHandle);
 				continue;
@@ -457,9 +457,9 @@ void FReplicationReader::DeserializeObjectStateDelta(FNetSerializationContext& C
 	}
 }
 
-FReplicationReader::FPendingBatchData* FReplicationReader::UpdateUnresolvedMustBeMappedReferences(FNetRefHandle InHandle, TArray<FNetRefHandle>& MustBeMappedReferences)
+FPendingBatchData* FReplicationReader::UpdateUnresolvedMustBeMappedReferences(FNetRefHandle InHandle, TArray<FNetRefHandle>& MustBeMappedReferences)
 {
-	FPendingBatchData* PendingBatch = PendingBatches.FindByPredicate([&InHandle](const FPendingBatchData& Entry) { return Entry.Handle == InHandle; });
+	FPendingBatchData* PendingBatch = PendingBatches.Find(InHandle);
 	// If we already have a pending batch we append any new must be mapped references to it.
 	if (PendingBatch)
 	{
@@ -497,7 +497,7 @@ FReplicationReader::FPendingBatchData* FReplicationReader::UpdateUnresolvedMustB
 		// We must create a new batch
 		if (!Batch)
 		{
-			Batch = &PendingBatches.AddDefaulted_GetRef();
+			Batch = &PendingBatches.PendingBatches.AddDefaulted_GetRef();
 			Batch->Handle = InHandle;
 		}
 
@@ -662,7 +662,7 @@ uint32 FReplicationReader::ReadObjectBatch(FNetSerializationContext& Context)
 		return 0U;
 	}
 
-	FPendingBatchData* PendingBatch = PendingBatches.FindByPredicate([&IncompleteHandle](const FPendingBatchData& Entry) { return Entry.Handle == IncompleteHandle; });
+	FPendingBatchData* PendingBatch = PendingBatches.Find(IncompleteHandle);
 
 	// This object has pending must be mapped references that must be resolved before we can process the data.
 	FPendingBatchData* PendingBatchData = ObjectReferenceCache->ShouldAsyncLoad() ? UpdateUnresolvedMustBeMappedReferences(IncompleteHandle, TempMustBeMappedReferences) : nullptr;
@@ -1750,9 +1750,9 @@ bool FReplicationReader::EnqueueEndReplication(FPendingBatchData* PendingBatchDa
 
 void FReplicationReader::ProcessQueuedBatches()
 {
-	UE_NET_TRACE_FRAME_STATSCOUNTER(Parameters.ReplicationSystem->GetId(), ReplicationReader.PendingQueuedBatches, PendingBatches.Num(), ENetTraceVerbosity::Trace);
+	UE_NET_TRACE_FRAME_STATSCOUNTER(Parameters.ReplicationSystem->GetId(), ReplicationReader.PendingQueuedBatches, PendingBatches.PendingBatches.Num(), ENetTraceVerbosity::Trace);
 
-	if (PendingBatches.IsEmpty())
+	if (!PendingBatches.GetHasPendingBatches())
 	{
 		//Nothing to do.
 		return;
@@ -1772,9 +1772,9 @@ void FReplicationReader::ProcessQueuedBatches()
 	Context.SetInternalContext(&InternalContext);
 	Context.SetNetBlobReceiver(&ReplicationSystemInternal->GetNetBlobHandlerManager());
 
-	for (int BatchIt = 0; BatchIt < PendingBatches.Num(); )
+	for (int BatchIt = 0; BatchIt < PendingBatches.PendingBatches.Num(); )
 	{
-		FPendingBatchData& PendingBatchData = PendingBatches[BatchIt];
+		FPendingBatchData& PendingBatchData = PendingBatches.PendingBatches[BatchIt];
 
 		// Try to resolve remaining must be mapped references
 		TempMustBeMappedReferences.Reset();
@@ -1871,7 +1871,7 @@ void FReplicationReader::ProcessQueuedBatches()
 			}
 
 			// Not optimal, but we want to preserve the order if we can as there might be batches waiting for the same reference
-			PendingBatches.RemoveAt(BatchIt);
+			PendingBatches.PendingBatches.RemoveAt(BatchIt);
 		}
 		else
 		{
@@ -2027,16 +2027,32 @@ void FReplicationReader::ResolveAndDispatchAttachments(FNetSerializationContext&
 		{
 			while (const TRefCountPtr<FNetBlob>* Attachment = AttachmentQueue->PeekReliable())
 			{
-				// Delay attachments with unresolved references
+				// Delay attachments with unresolved pending references
 				if (bCanDelayAttachments)
 				{
-					// Only block reliable stream if we have unresolved references that must be mapped
-					const ENetObjectReferenceResolveResult ResolveResult = Attachment->GetReference()->ResolveObjectReferences(Context);
-					if (EnumHasAnyFlags(ResolveResult, ENetObjectReferenceResolveResult::HasUnresolvedMustBeMappedReferences))
+					bool bDelayRpc = false;
+
+					FNetReferenceCollector Collector;
+					Attachment->GetReference()->CollectObjectReferences(Context, Collector);
+
+					// Check status of references, as we already should have queued up any unmapped references at the batch level, it should be enough to only check if we have any unresolved references pending async load.
+					// NOTE: Behavior is slightly different between Iris and old replication system due to the fact that Iris processes incoming packet data prior to dispatching received stats and RPC:s, that means that 
+					// we expect to be able to resolve all dynamic references contained in the same data packet and does not delay the RPC until the next tick to solve that as the old system does. 
+					// The difference is that the old system might be able to resolve incoming dynamic references from later packets processed for the same tick, but as this is far from guaranteed we currently do not try to mimic this.
+					for (const FNetReferenceCollector::FReferenceInfo& Info : MakeArrayView(Collector.GetCollectedReferences()))
+					{
+						if (ObjectReferenceCache->IsNetRefHandlePending(Info.Reference.GetRefHandle(), PendingBatches))
+						{
+							bDelayRpc = true;
+							break;
+						}
+					}
+
+					if (bDelayRpc)
 					{
 						const FReplicationStateDescriptor* Descriptor = Attachment->GetReference()->GetReplicationStateDescriptor();
 
-						UE_LOG(LogIris, Warning, TEXT("Unable to resolve references in %s for InternalIndex %u"), (Descriptor != nullptr ? ToCStr(Descriptor->DebugName) : TEXT("N/A")), InternalIndex);
+						UE_LOG(LogIris, Verbose, TEXT("Delaying Attachment - %s for InternalIndex %u." ), (Descriptor != nullptr ? ToCStr(Descriptor->DebugName) : TEXT("N/A")), InternalIndex);
 						break;
 					}
 				}
diff --git a/Engine/Source/Runtime/Experimental/Iris/Core/Private/Iris/ReplicationSystem/ReplicationReader.h b/Engine/Source/Runtime/Experimental/Iris/Core/Private/Iris/ReplicationSystem/ReplicationReader.h
--- a/Engine/Source/Runtime/Experimental/Iris/Core/Private/Iris/ReplicationSystem/ReplicationReader.h
+++ b/Engine/Source/Runtime/Experimental/Iris/Core/Private/Iris/ReplicationSystem/ReplicationReader.h
@@ -8,6 +8,7 @@
 #include "Iris/ReplicationSystem/NetRefHandleManager.h"
 #include "Iris/ReplicationSystem/ObjectReferenceCache.h"
 #include "Iris/ReplicationSystem/ReplicationTypes.h"
+#include "Iris/ReplicationSystem/PendingBatchData.h"
 #include "Containers/Map.h"
 #include "Misc/MemStack.h"
 
@@ -116,39 +117,6 @@ private:
 		uint32 bDeferredEndReplication : 1;
 	};
 
-	// Queued data chunk
-	struct FQueuedDataChunk
-	{
-		FQueuedDataChunk()
-		: StorageOffset(0U)
-		, NumBits(0U)
-		, bHasBatchOwnerData(0U)
-		, bIsEndReplicationChunk(0U)
-		{
-		}
-
-		uint32 StorageOffset;
-		uint32 NumBits : 30;
-		uint32 bHasBatchOwnerData : 1;
-		uint32 bIsEndReplicationChunk : 1;
-	};
-
-	// Struct to contain storage and required data for queued batches pending must be mapped references
-	struct FPendingBatchData
-	{
-		// We use a single array to store the actual data, it will grow if required.
-		TArray<uint32, TInlineAllocator<32>> DataChunkStorage;		
-		TArray<FQueuedDataChunk, TInlineAllocator<4>> QueuedDataChunks;
-
-		// Must be mapped references pending resolve
-		TArray<FNetRefHandle, TInlineAllocator<4>> PendingMustBeMappedReferences;
-
-		// Resolved references for which we have are holding on to references to avoid GC
-		TArray<FNetRefHandle, TInlineAllocator<4>> ResolvedReferences;
-
-		// Batch owner with queued data chunks
-		FNetRefHandle Handle;
-	};
 
 	enum : uint32
 	{
@@ -228,7 +196,7 @@ private:
 	void DeserializeObjectStateDelta(FNetSerializationContext& Context, uint32 InternalIndex, FDispatchObjectInfo& Info, FReplicatedObjectInfo& ObjectInfo, const FNetRefHandleManager::FReplicatedObjectData& ObjectData, uint32& OutNewBaselineIndex);
 
 	// If async loading is enabled this function will verify if we can resolve all PendingMustBeMappedReferences
-	FReplicationReader::FPendingBatchData* UpdateUnresolvedMustBeMappedReferences(FNetRefHandle Handle, TArray<FNetRefHandle>& MustBeMappedReferences);
+	FPendingBatchData* UpdateUnresolvedMustBeMappedReferences(FNetRefHandle Handle, TArray<FNetRefHandle>& MustBeMappedReferences);
 
 	// If we are queuing data for a batch we must also defer calls to EndReplication
 	// This method writes this method in the form of a QueuedChunk
@@ -268,7 +236,7 @@ private:
 	TMultiMap<FNetRefHandle, uint32> ResolvedDynamicHandleToDependents;
 
 	// We do not expect to have many objects in this state
-	TArray<FPendingBatchData> PendingBatches;
+	FPendingBatches PendingBatches;
 
 	// We do not expect many objects to be broken
 	TArray<FNetRefHandle> BrokenObjects;
diff --git a/Engine/Source/Runtime/Experimental/Iris/Core/Public/Iris/ReplicationSystem/NetBlob/NetBlob.h b/Engine/Source/Runtime/Experimental/Iris/Core/Public/Iris/ReplicationSystem/NetBlob/NetBlob.h
--- a/Engine/Source/Runtime/Experimental/Iris/Core/Public/Iris/ReplicationSystem/NetBlob/NetBlob.h
+++ b/Engine/Source/Runtime/Experimental/Iris/Core/Public/Iris/ReplicationSystem/NetBlob/NetBlob.h
@@ -14,6 +14,7 @@ namespace UE::Net
 	class FNetBlob;
 	class FNetSerializationContext;
 	struct FReplicationStateDescriptor;
+	class FNetReferenceCollector;
 	namespace Private
 	{
 		class FNetBlobManager;
@@ -87,11 +88,8 @@ public:
 	/** Deserialize a blob that was serialized with Serialize. */
 	IRISCORE_API virtual void Deserialize(FNetSerializationContext& Context);
 
-	/**
-	 * Returns the status of the object reference resolving.
-	 * @see ENetObjectReferenceResolveResult
-	 */
-	IRISCORE_API ENetObjectReferenceResolveResult ResolveObjectReferences(FNetSerializationContext& Context) const;
+	/** Collect object references from quantized data */
+	IRISCORE_API void CollectObjectReferences(FNetSerializationContext& Context, FNetReferenceCollector& Collector) const;
 
 	/** Adds a reference. A blob is created with reference count zero. */
 	void AddRef() const { ++RefCount; }
