Iris
- Added missing conditional for valid descriptor in void FNetBlob::CollectObjectReferences

#preflignt n/a

diff --git a/Engine/Source/Runtime/Experimental/Iris/Core/Private/Iris/ReplicationSystem/NetBlob/NetBlob.cpp b/Engine/Source/Runtime/Experimental/Iris/Core/Private/Iris/ReplicationSystem/NetBlob/NetBlob.cpp
--- a/Engine/Source/Runtime/Experimental/Iris/Core/Private/Iris/ReplicationSystem/NetBlob/NetBlob.cpp
+++ b/Engine/Source/Runtime/Experimental/Iris/Core/Private/Iris/ReplicationSystem/NetBlob/NetBlob.cpp
@@ -90,8 +90,11 @@ void FNetBlob::Deserialize(FNetSerializationContext& Context)
 
 void FNetBlob::CollectObjectReferences(FNetSerializationContext& Context, FNetReferenceCollector& Collector) const
 {
-	const FNetSerializerChangeMaskParam InitStateChangeMaskInfo = { 0 };
-	Private::FReplicationStateOperationsInternal::CollectReferences(Context, Collector, InitStateChangeMaskInfo, QuantizedBlobState.Get(), BlobDescriptor);
+	if (BlobDescriptor.IsValid() && QuantizedBlobState.IsValid())
+	{
+		const FNetSerializerChangeMaskParam InitStateChangeMaskInfo = { 0 };
+		Private::FReplicationStateOperationsInternal::CollectReferences(Context, Collector, InitStateChangeMaskInfo, QuantizedBlobState.Get(), BlobDescriptor);
+	}
 }
 
 void FNetBlob::SerializeBlob(FNetSerializationContext& Context) const
