Add support to set a custom instance data slot on geometry collection ISM pool instances.

diff --git a/Engine/Source/Runtime/Experimental/GeometryCollectionEngine/Private/GeometryCollection/GeometryCollectionISMPoolComponent.cpp b/Engine/Source/Runtime/Experimental/GeometryCollectionEngine/Private/GeometryCollection/GeometryCollectionISMPoolComponent.cpp
--- a/Engine/Source/Runtime/Experimental/GeometryCollectionEngine/Private/GeometryCollection/GeometryCollectionISMPoolComponent.cpp
+++ b/Engine/Source/Runtime/Experimental/GeometryCollectionEngine/Private/GeometryCollection/GeometryCollectionISMPoolComponent.cpp
@@ -52,6 +52,14 @@ bool FGeometryCollectionMeshGroup::BatchUpdateInstancesTransforms(FGeometryColle
 	return false;
 }
 
+void FGeometryCollectionMeshGroup::BatchUpdateInstanceCustomData(FGeometryCollectionISMPool& ISMPool, int32 CustomFloatIndex, float CustomFloatValue)
+{
+	for (const FGeometryCollectionMeshInfo& MeshInfo : MeshInfos)
+	{
+		ISMPool.BatchUpdateInstanceCustomData(MeshInfo, CustomFloatIndex, CustomFloatValue);
+	}
+}
+
 void FGeometryCollectionMeshGroup::RemoveAllMeshes(FGeometryCollectionISMPool& ISMPool)
 {
 	for (const FGeometryCollectionMeshInfo& MeshInfo: MeshInfos)
@@ -261,6 +269,28 @@ bool FGeometryCollectionISMPool::BatchUpdateInstancesTransforms(FGeometryCollect
 	return false;
 }
 
+void FGeometryCollectionISMPool::BatchUpdateInstanceCustomData(FGeometryCollectionMeshInfo const& MeshInfo, int32 CustomFloatIndex, float CustomFloatValue)
+{
+	if (!ISMs.IsValidIndex(MeshInfo.ISMIndex))
+	{
+		return;
+	}
+	
+	FGeometryCollectionISM& ISM = ISMs[MeshInfo.ISMIndex];
+	if (!ensure(CustomFloatValue < ISM.MeshInstance.Desc.NumCustomDataFloats))
+	{
+		return;
+	}
+		
+	const FInstanceGroups::FInstanceGroupRange& InstanceGroup = ISM.InstanceGroups.GroupRanges[MeshInfo.InstanceGroupIndex];
+	for (int32 InstanceIndex = 0; InstanceIndex < InstanceGroup.Count; ++InstanceIndex)
+	{
+		const int32 RenderIndex = ISM.InstanceIndexToRenderIndex[InstanceGroup.Start + InstanceIndex];
+		const bool bLastUpdate = InstanceIndex == InstanceGroup.Count - 1;
+		ISM.ISMComponent->SetCustomDataValue(RenderIndex, CustomFloatIndex, CustomFloatValue, bLastUpdate);
+	}
+}
+
 void FGeometryCollectionISMPool::RemoveISM(const FGeometryCollectionMeshInfo& MeshInfo)
 {
 	if (ISMs.IsValidIndex(MeshInfo.ISMIndex))
@@ -435,6 +465,17 @@ bool UGeometryCollectionISMPoolComponent::BatchUpdateInstancesTransforms(FMeshGr
 	return false;
 }
 
+bool UGeometryCollectionISMPoolComponent::BatchUpdateInstanceCustomData(FMeshGroupId MeshGroupId, int32 CustomFloatIndex, float CustomFloatValue)
+{
+	if (FGeometryCollectionMeshGroup* MeshGroup = MeshGroups.Find(MeshGroupId))
+	{
+		MeshGroup->BatchUpdateInstanceCustomData(Pool, CustomFloatIndex, CustomFloatValue);
+		return true;
+	}
+	UE_LOG(LogChaos, Warning, TEXT("UGeometryCollectionISMPoolComponent : Trying to update instance with mesh group (%d) that not exists"), MeshGroupId);
+	return false;
+}
+
 void UGeometryCollectionISMPoolComponent::PreallocateMeshInstance(const FGeometryCollectionStaticMeshInstance& MeshInstance)
 {
 	// If we are recycling components with a free list then we don't expect to have zero instance components.
diff --git a/Engine/Source/Runtime/Experimental/GeometryCollectionEngine/Public/GeometryCollection/GeometryCollectionComponent.h b/Engine/Source/Runtime/Experimental/GeometryCollectionEngine/Public/GeometryCollection/GeometryCollectionComponent.h
--- a/Engine/Source/Runtime/Experimental/GeometryCollectionEngine/Public/GeometryCollection/GeometryCollectionComponent.h
+++ b/Engine/Source/Runtime/Experimental/GeometryCollectionEngine/Public/GeometryCollection/GeometryCollectionComponent.h
@@ -1115,6 +1115,8 @@ public:
 
 	GEOMETRYCOLLECTIONENGINE_API void SetUpdateNavigationInTick(const bool bUpdateInTick) { bUpdateNavigationInTick = bUpdateInTick; }
 
+	/** Get any custom renderer. Returns nullptr if none is set. */
+	GEOMETRYCOLLECTIONENGINE_API IGeometryCollectionExternalRenderInterface* GetCustomRenderer() { return CustomRenderer.GetInterface(); }
 	/** Enable or disable root proxy for custom rendering - this can be set at runtime */
 	GEOMETRYCOLLECTIONENGINE_API void EnableRootProxyForCustomRenderer(bool bEnable);
 
diff --git a/Engine/Source/Runtime/Experimental/GeometryCollectionEngine/Public/GeometryCollection/GeometryCollectionISMPoolComponent.h b/Engine/Source/Runtime/Experimental/GeometryCollectionEngine/Public/GeometryCollection/GeometryCollectionISMPoolComponent.h
--- a/Engine/Source/Runtime/Experimental/GeometryCollectionEngine/Public/GeometryCollection/GeometryCollectionISMPoolComponent.h
+++ b/Engine/Source/Runtime/Experimental/GeometryCollectionEngine/Public/GeometryCollection/GeometryCollectionISMPoolComponent.h
@@ -231,6 +231,7 @@ struct FGeometryCollectionMeshGroup
 	/** Update instance transforms for a group of instances. */
 	bool BatchUpdateInstancesTransforms(FGeometryCollectionISMPool& ISMPool, FMeshId MeshId, int32 StartInstanceIndex, const TArray<FTransform>& NewInstancesTransforms, bool bWorldSpace, bool bMarkRenderStateDirty, bool bTeleport);
 	bool BatchUpdateInstancesTransforms(FGeometryCollectionISMPool& ISMPool, FMeshId MeshId, int32 StartInstanceIndex, TArrayView<const FTransform> NewInstancesTransforms, bool bWorldSpace, bool bMarkRenderStateDirty, bool bTeleport);
+	void BatchUpdateInstanceCustomData(FGeometryCollectionISMPool& ISMPool, int32 CustomFloatIndex, float CustomFloatValue);
 
 	/** Remove all of our managed meshes and associated instances. */
 	void RemoveAllMeshes(FGeometryCollectionISMPool& ISMPool);
@@ -276,6 +277,7 @@ struct FGeometryCollectionISMPool
 	/** Update ISM contents. */
 	bool BatchUpdateInstancesTransforms(FGeometryCollectionMeshInfo& MeshInfo, int32 StartInstanceIndex, const TArray<FTransform>& NewInstancesTransforms, bool bWorldSpace, bool bMarkRenderStateDirty, bool bTeleport);
 	bool BatchUpdateInstancesTransforms(FGeometryCollectionMeshInfo& MeshInfo, int32 StartInstanceIndex, TArrayView<const FTransform> NewInstancesTransforms, bool bWorldSpace, bool bMarkRenderStateDirty, bool bTeleport);
+	void BatchUpdateInstanceCustomData(FGeometryCollectionMeshInfo const& MeshInfo, int32 CustomFloatIndex, float CustomFloatValue);
 
 	/** Clear all ISM components and associated data. */
 	void Clear();
@@ -327,12 +329,15 @@ public:
 	/** Add a static mesh for a mesh group */
 	GEOMETRYCOLLECTIONENGINE_API FMeshId AddMeshToGroup(FMeshGroupId MeshGroupId, const FGeometryCollectionStaticMeshInstance& MeshInstance, int32 InstanceCount, TArrayView<const float> CustomDataFloats);
 
-	/** Add a static mesh for a mesh group */
+	/** Update transforms for a mesh group */
 	GEOMETRYCOLLECTIONENGINE_API bool BatchUpdateInstancesTransforms(FMeshGroupId MeshGroupId, FMeshId MeshId, int32 StartInstanceIndex, TArrayView<const FTransform> NewInstancesTransforms, bool bWorldSpace = false, bool bMarkRenderStateDirty = false, bool bTeleport = false);
 
 	UE_DEPRECATED(5.3, "BatchUpdateInstancesTransforms Array parameter version is deprecated, use the TArrayView version instead")
 	GEOMETRYCOLLECTIONENGINE_API bool BatchUpdateInstancesTransforms(FMeshGroupId MeshGroupId, FMeshId MeshId, int32 StartInstanceIndex, const TArray<FTransform>& NewInstancesTransforms, bool bWorldSpace = false, bool bMarkRenderStateDirty = false, bool bTeleport = false);
 
+	/** Update a single slot of custom instance data for all instances in a mesh group */
+	GEOMETRYCOLLECTIONENGINE_API bool BatchUpdateInstanceCustomData(FMeshGroupId MeshGroupId, int32 CustomFloatIndex, float CustomFloatValue);
+
 	/** 
 	 * Preallocate an ISM in the pool. 
 	 * Doing this early for known mesh instance descriptions can reduce the component registration cost of AddMeshToGroup() for newly discovered mesh descriptions.
