Separate disk and memory locks. This way puts don't get held up by disk activity.

#rnx

diff --git a/Engine/Source/Runtime/Experimental/IoStoreOnDemand/Private/JournaledCache.cpp b/Engine/Source/Runtime/Experimental/IoStoreOnDemand/Private/JournaledCache.cpp
--- a/Engine/Source/Runtime/Experimental/IoStoreOnDemand/Private/JournaledCache.cpp
+++ b/Engine/Source/Runtime/Experimental/IoStoreOnDemand/Private/JournaledCache.cpp
@@ -1180,7 +1180,8 @@ public:
 	uint32			DebugVisit(void* Param, FDebugCacheEntry::Callback* Callback);
 
 private:
-	mutable FRWLock	Lock;
+	mutable FRWLock	FsLock;
+	mutable FRWLock	MemLock;
 	FMemCache		MemCache;
 	FDiskCache		DiskCache;
 	std::atomic_int	Demand;
@@ -1227,7 +1228,7 @@ uint32 FCache::GetAilments() const
 ////////////////////////////////////////////////////////////////////////////////
 void FCache::Reset()
 {
-	FWriteAccess _(Lock);
+	FWriteAccess _[] = { MemLock, FsLock };
 	MemCache.Reset();
 	DiskCache.Reset();
 }
@@ -1235,7 +1236,7 @@ void FCache::Reset()
 ////////////////////////////////////////////////////////////////////////////////
 bool FCache::Load()
 {
-	FWriteAccess _(Lock);
+	FWriteAccess _(FsLock);
 	return LoadCache(DiskCache);
 }
 
@@ -1249,20 +1250,24 @@ uint32 FCache::GetDemand() const
 FCache::FEntry FCache::Get(uint64 Key) const
 {
 	FEntry Ret;
-	Ret.Lock = FReadAccess(Lock);
 
+	Ret.Lock = FReadAccess(FsLock);
 	if (Ret.Handle = DiskCache.Get(Key); Ret.Handle)
 	{
 		Ret.HitType = EHit::Disk;
 		Ret.Owner = &DiskCache;
+		return Ret;
 	}
 		
-	else if (Ret.Handle = MemCache.Get(Key); Ret.Handle)
+	Ret.Lock = FReadAccess(MemLock);
+	if (Ret.Handle = MemCache.Get(Key); Ret.Handle)
 	{
 		Ret.HitType = EHit::Memory;
 		Ret.Owner = &MemCache;
+		return Ret;
 	}
 
+	Ret.Lock = FReadAccess();
 	return Ret;
 }
 
@@ -1270,7 +1275,7 @@ FCache::FEntry FCache::Get(uint64 Key) const
 bool FCache::Put(uint64 Key, FIoBuffer& Data)
 {
 	FIoBuffer Cloned = Data;
-	FWriteAccess _(Lock);
+	FWriteAccess _(MemLock);
 	bool Ok = MemCache.Put(Key, MoveTemp(Cloned));
 	if (Ok)
 	{
@@ -1283,7 +1288,7 @@ bool FCache::Put(uint64 Key, FIoBuffer& Data)
 ////////////////////////////////////////////////////////////////////////////////
 uint32 FCache::Flush()
 {
-	FWriteAccess _(Lock);
+	FWriteAccess _(FsLock);
 	return DiskCache.Flush();
 }
 
@@ -1293,9 +1298,15 @@ uint32 FCache::WriteMemToDisk(int32 Allowance)
 	TRACE_CPUPROFILER_EVENT_SCOPE(IasCache::Flush_MemCache);
 
 	FMemCache::PeelItems PeelItems;
-	FWriteAccess _(Lock);
-	int32 MemCacheSize = MemCache.Peel(Allowance, PeelItems);
+	int32 MemCacheSize;
+	{
+		FWriteAccess _(MemLock);
 
+		MemCacheSize = MemCache.Peel(Allowance, PeelItems);
+
+		uint32 NewDemand = MemCache.GetDemand();
+		Demand.store(NewDemand, std::memory_order_relaxed);
+	}
 	FDiskPhrase Phrase = DiskCache.OpenPhrase(MemCacheSize);
 	int32 PeelIndex = -1;
 	for (int32 i = 0, n = PeelItems.Num(); i < n; ++i)
@@ -1315,10 +1326,10 @@ uint32 FCache::WriteMemToDisk(int32 Allowance)
 		MemCacheSize = Phrase.GetDataSize();
 	}
 
-	DiskCache.ClosePhrase(MoveTemp(Phrase));
-
-	uint32 NewDemand = MemCache.GetDemand();
-	Demand.store(NewDemand, std::memory_order_relaxed);
+	{
+		FWriteAccess _(FsLock);
+		DiskCache.ClosePhrase(MoveTemp(Phrase));
+	}
 
 	return MemCacheSize;
 }
@@ -1326,7 +1337,7 @@ uint32 FCache::WriteMemToDisk(int32 Allowance)
 ////////////////////////////////////////////////////////////////////////////////
 uint32 FCache::DebugVisit(void* Param, FDebugCacheEntry::Callback* Callback)
 {
-	FReadAccess _(Lock);
+	FReadAccess _[] = { MemLock, FsLock };
 	uint32 Count = 0;
 	Count += MemCache.DebugVisit(Param, Callback);
 	Count += DiskCache.DebugVisit(Param, Callback);
