Fix missing breaking event
Now the Game Thread can run several time while the Physics thread run just once. The Game Thread has the responsibility to swap the event buffers when it finish to read them.
But the Physics Thread was still resetting the buffer. From now the Game Thread is swapping buffer then resting the new available buffer.
Now we ensure the Physics Thread is not overwritten its previous data.

diff --git a/Engine/Source/Runtime/Experimental/Chaos/Private/EventDefaults.cpp b/Engine/Source/Runtime/Experimental/Chaos/Private/EventDefaults.cpp
--- a/Engine/Source/Runtime/Experimental/Chaos/Private/EventDefaults.cpp
+++ b/Engine/Source/Runtime/Experimental/Chaos/Private/EventDefaults.cpp
@@ -83,13 +83,9 @@ namespace Chaos
 			//if (!Solver->GetEventFilters()->IsCollisionEventEnabled())
 			//	return;
 
-			FCollisionDataArray& AllCollisionsDataArray = CollisionEventData.CollisionData.AllCollisionsArray;
-			TMap<IPhysicsProxyBase*, TArray<int32>>& AllCollisionsIndicesByPhysicsProxy = CollisionEventData.PhysicsProxyToCollisionIndices.PhysicsProxyToIndicesMap;
-
 			if (bResetData)
 			{
-				AllCollisionsDataArray.Reset();
-				AllCollisionsIndicesByPhysicsProxy.Reset();
+				CollisionEventData.Reset();
 			}
 			CollisionEventData.CollisionData.TimeCreated = Solver->MTime;
 			CollisionEventData.PhysicsProxyToCollisionIndices.TimeCreated = Solver->MTime;
@@ -286,6 +282,9 @@ namespace Chaos
 								}
 							}
 						}, Chaos::SmallBatchSize);
+
+						FCollisionDataArray& AllCollisionsDataArray = CollisionEventData.CollisionData.AllCollisionsArray;
+						TMap<IPhysicsProxyBase*, TArray<int32>>& AllCollisionsIndicesByPhysicsProxy = CollisionEventData.PhysicsProxyToCollisionIndices.PhysicsProxyToIndicesMap;
 						for (int32 IdxCollision = 0; IdxCollision < NumValidCollisions; ++IdxCollision)
 						{
 							if (DupAllCollisionsDataArray[IdxCollision].Proxy1 != nullptr)
@@ -318,22 +317,21 @@ namespace Chaos
 
 			// #todo: This isn't working - SolverActor parameters are set on a solver but it is currently a different solver that is simulating!!
 			if (!Solver->GetEventFilters()->IsBreakingEventEnabled())
+			{
 				return;
-
-			FBreakingDataArray& FilteredBreakingDataArray = BreakingEventData.BreakingData.AllBreakingsArray;
-			TMap<IPhysicsProxyBase*, TArray<int32>>& FilteredBreakingIndicesByPhysicsProxy = BreakingEventData.PhysicsProxyToBreakingIndices.PhysicsProxyToIndicesMap;
+			}
 
 			if (bResetData)
 			{
-				FilteredBreakingDataArray.Reset();
-				FilteredBreakingIndicesByPhysicsProxy.Reset();
-				BreakingEventData.BreakingData.bHasGlobalEvent = false;
+				BreakingEventData.Reset();
 			}
 			BreakingEventData.BreakingData.TimeCreated = Solver->MTime;
 
 			const auto* Evolution = Solver->GetEvolution();
 			const FPBDRigidParticles& Particles = Evolution->GetParticles().GetDynamicParticles();
 			const TArray<FBreakingData>& AllClusterBreakings = Evolution->GetRigidClustering().GetAllClusterBreakings();
+			FBreakingDataArray& FilteredBreakingDataArray = BreakingEventData.BreakingData.AllBreakingsArray;
+			TMap<IPhysicsProxyBase*, TArray<int32>>& FilteredBreakingIndicesByPhysicsProxy = BreakingEventData.PhysicsProxyToBreakingIndices.PhysicsProxyToIndicesMap;
 	
 			if (AllClusterBreakings.Num() > 0)
 			{
@@ -371,18 +369,18 @@ namespace Chaos
 #if TODO_REIMPLEMENT_RIGID_CLUSTERING
 			const TMap<uint32, TUniquePtr<TArray<uint32>>>& ParentToChildrenMap = Evolution->GetRigidClustering().GetChildrenMap();
 #endif
-			FTrailingDataArray& AllTrailingsDataArray = TrailingEventData.TrailingData.AllTrailingsArray;
-			TMap<IPhysicsProxyBase*, TArray<int32>>& AllTrailingIndicesByPhysicsProxy = TrailingEventData.PhysicsProxyToTrailingIndices.PhysicsProxyToIndicesMap;
+
 
 			if (bResetData)
 			{
-				AllTrailingsDataArray.Reset();
-				AllTrailingIndicesByPhysicsProxy.Reset();
+				TrailingEventData.Reset();
 			}
 			TrailingEventData.TrailingData.TimeCreated = Solver->MTime;
 			TrailingEventData.PhysicsProxyToTrailingIndices.TimeCreated = Solver->MTime;
 
 			const TArray<TPBDRigidParticleHandle<Chaos::FReal, 3>*>& ActiveParticlesArray = Evolution->GetParticles().GetActiveParticlesArray();
+			FTrailingDataArray& AllTrailingsDataArray = TrailingEventData.TrailingData.AllTrailingsArray;
+			TMap<IPhysicsProxyBase*, TArray<int32>>& AllTrailingIndicesByPhysicsProxy = TrailingEventData.PhysicsProxyToTrailingIndices.PhysicsProxyToIndicesMap;
 
 			for (TPBDRigidParticleHandle<Chaos::FReal, 3>* ActiveParticle : ActiveParticlesArray)
 			{
@@ -460,8 +458,7 @@ namespace Chaos
 
 			const auto* Evolution = Solver->GetEvolution();
 
-			FSleepingDataArray& EventSleepDataArray = SleepingEventData.SleepingData;
-			EventSleepDataArray.Reset();
+			SleepingEventData.Reset();
 
 			Chaos::FPBDRigidsSolver* NonConstSolver = const_cast<Chaos::FPBDRigidsSolver*>(Solver);
 
@@ -471,6 +468,7 @@ namespace Chaos
 				&NonConstSolver->Particles.GetGeometryCollectionParticles()
 			};
 
+			FSleepingDataArray& EventSleepDataArray = SleepingEventData.SleepingData;
 			for (FPBDRigidParticles* ParticleArray : RelevantParticleArrays)
 			{
 				check(ParticleArray != nullptr);
@@ -507,18 +505,16 @@ namespace Chaos
 			{
 				check(Solver);
 				EnsureIsInPhysicsThreadContext();
-				
-				FRemovalDataArray& AllRemovalDataArray = RemovalEventData.RemovalData.AllRemovalArray;
-				TMap<IPhysicsProxyBase*, TArray<int32>>& AllRemovalIndicesByPhysicsProxy = RemovalEventData.PhysicsProxyToRemovalIndices.PhysicsProxyToIndicesMap;
 
 				if (bResetData)
 				{
-					AllRemovalDataArray.Reset();
-					AllRemovalIndicesByPhysicsProxy.Reset();
+					RemovalEventData.Reset();
 				}
 				RemovalEventData.RemovalData.TimeCreated = Solver->MTime;
 
 				const TArray<FRemovalData>& AllRemovalsArray = Solver->GetEvolution()->GetAllRemovals();
+				FRemovalDataArray& AllRemovalDataArray = RemovalEventData.RemovalData.AllRemovalArray;
+				TMap<IPhysicsProxyBase*, TArray<int32>>& AllRemovalIndicesByPhysicsProxy = RemovalEventData.PhysicsProxyToRemovalIndices.PhysicsProxyToIndicesMap;
 				
 				for (int32 Idx = 0; Idx < AllRemovalsArray.Num(); ++Idx)
 				{
@@ -552,7 +548,6 @@ namespace Chaos
 			if (bResetData)
 			{
 				CrumblingEventData.Reset();
-				CrumblingEventData.CrumblingData.bHasGlobalEvent = false;
 			}
 			CrumblingEventData.SetTimeCreated(Solver->MTime);
 
diff --git a/Engine/Source/Runtime/Experimental/Chaos/Private/EventManager.cpp b/Engine/Source/Runtime/Experimental/Chaos/Private/EventManager.cpp
--- a/Engine/Source/Runtime/Experimental/Chaos/Private/EventManager.cpp
+++ b/Engine/Source/Runtime/Experimental/Chaos/Private/EventManager.cpp
@@ -62,6 +62,7 @@ namespace Chaos
 		{
 			if (EventContainer)
 			{
+				EventContainer->ResetConsumerBuffer();
 				EventContainer->FlipBufferIfRequired();
 			}
 		}
diff --git a/Engine/Source/Runtime/Experimental/Chaos/Private/PBDRigidsSolver.cpp b/Engine/Source/Runtime/Experimental/Chaos/Private/PBDRigidsSolver.cpp
--- a/Engine/Source/Runtime/Experimental/Chaos/Private/PBDRigidsSolver.cpp
+++ b/Engine/Source/Runtime/Experimental/Chaos/Private/PBDRigidsSolver.cpp
@@ -516,8 +516,9 @@ namespace Chaos
 				SCOPE_CYCLE_COUNTER(STAT_EventDataGathering);
 				{
 					SCOPE_CYCLE_COUNTER(STAT_FillProducerData);
-					bool ResetData = (MSubStepInfo.Step == 0);
-					MSolver->GetEventManager()->FillProducerData(MSolver, ResetData);
+					// The Game Thread is now in charge to reset the producer buffer
+					constexpr bool bResetProducerData = false;
+					MSolver->GetEventManager()->FillProducerData(MSolver, bResetProducerData);
 					MSolver->GetEvolution()->ResetAllRemovals();
 				}
 			}
diff --git a/Engine/Source/Runtime/Experimental/Chaos/Public/EventManager.h b/Engine/Source/Runtime/Experimental/Chaos/Public/EventManager.h
--- a/Engine/Source/Runtime/Experimental/Chaos/Public/EventManager.h
+++ b/Engine/Source/Runtime/Experimental/Chaos/Public/EventManager.h
@@ -128,6 +128,11 @@ namespace Chaos
 		 */
 		virtual void FlipBufferIfRequired() = 0;
 
+		/**
+		 * Reset the consumer buffer, can be used just before flipping buffers to start with a clean buffer
+		 */
+		virtual void ResetConsumerBuffer() = 0;
+
 		/**
 		 * Dispatch events to the registered handlers
 		 */
@@ -264,6 +269,14 @@ namespace Chaos
 			EventBuffer->FlipProducer();
 		}
 
+		/**
+		 * Reset the consumer buffer, can be used just before flipping buffers to start with a clean buffer
+		 */
+		virtual void ResetConsumerBuffer()
+		{
+			EventBuffer->GetConsumerBufferMutable()->Reset();
+		}
+
 		/**
 		 * Dispatch events to the registered handlers
 		 */
diff --git a/Engine/Source/Runtime/Experimental/Chaos/Public/EventsData.h b/Engine/Source/Runtime/Experimental/Chaos/Public/EventsData.h
--- a/Engine/Source/Runtime/Experimental/Chaos/Public/EventsData.h
+++ b/Engine/Source/Runtime/Experimental/Chaos/Public/EventsData.h
@@ -64,6 +64,12 @@ namespace Chaos
 	{
 		FCollisionEventData() {}
 
+		void Reset()
+		{
+			PhysicsProxyToCollisionIndices.PhysicsProxyToIndicesMap.Reset();
+			CollisionData.Reset();
+		}
+
 		FAllCollisionData CollisionData;
 		FIndicesByPhysicsProxy PhysicsProxyToCollisionIndices;
 	};
@@ -80,6 +86,7 @@ namespace Chaos
 		void Reset()
 		{
 			AllBreakingsArray.Reset();
+			bHasGlobalEvent = false;
 		}
 
 		FBreakingDataArray AllBreakingsArray;
@@ -90,6 +97,12 @@ namespace Chaos
 	{
 		FBreakingEventData() {}
 
+		void Reset()
+		{
+			BreakingData.Reset();
+			PhysicsProxyToBreakingIndices.PhysicsProxyToIndicesMap.Reset();
+		}
+
 		FAllBreakingData BreakingData;
 		FIndicesByPhysicsProxy PhysicsProxyToBreakingIndices;
 	};
@@ -116,6 +129,12 @@ namespace Chaos
 	{
 		FTrailingEventData() {}
 
+		void Reset()
+		{
+			TrailingData.Reset();
+			PhysicsProxyToTrailingIndices.Reset();
+		}
+
 		FAllTrailingData TrailingData;
 		FIndicesByPhysicsProxy PhysicsProxyToTrailingIndices;
 	};
@@ -142,6 +161,12 @@ namespace Chaos
 	{
 		FRemovalEventData() {}
 
+		void Reset()
+		{
+			RemovalData.Reset();
+			PhysicsProxyToRemovalIndices.PhysicsProxyToIndicesMap.Reset();
+		}
+
 		FAllRemovalData RemovalData;
 		FIndicesByPhysicsProxy PhysicsProxyToRemovalIndices;
 	};
@@ -150,6 +175,11 @@ namespace Chaos
 	{
 		FSleepingEventData() {}
 
+		void Reset()
+		{
+			SleepingData.Reset();
+		}
+
 		FSleepingDataArray SleepingData;
 	};
 
@@ -163,6 +193,7 @@ namespace Chaos
 		void Reset()
 		{
 			AllCrumblingsArray.Reset();
+			bHasGlobalEvent = false;
 		}
 
 		FCrumblingDataArray AllCrumblingsArray;
