IAS - better error handling

- added support for disabling HTTP streaming when too many errors occurred during the last X request(s)
- added support for reconnecting to any available endpoint when HTTP is disabled to due to errors
- added config parameter for number of endpoints to test at startup

#rnx

diff --git a/Engine/Source/Runtime/Experimental/IoStoreOnDemand/Private/CoreHttp/LatencyTesting.cpp b/Engine/Source/Runtime/Experimental/IoStoreOnDemand/Private/CoreHttp/LatencyTesting.cpp
--- a/Engine/Source/Runtime/Experimental/IoStoreOnDemand/Private/CoreHttp/LatencyTesting.cpp
+++ b/Engine/Source/Runtime/Experimental/IoStoreOnDemand/Private/CoreHttp/LatencyTesting.cpp
@@ -9,8 +9,6 @@
 #include "Misc/StringBuilder.h"
 #include "Templates/Function.h"
 
-#if !UE_BUILD_SHIPPING
-
 namespace UE::IO::IAS::HTTP
 {
 
@@ -24,7 +22,10 @@ void LatencyTest(FStringView InUrl, FStringView InPath, TArrayView<int32> OutRes
 	FConnectionPool Pool(PoolParams);
 
 	TAnsiStringBuilder<256> AnsiPath;
-	AnsiPath << "/";
+	if (!InPath.StartsWith(TEXT('/')))
+	{
+		AnsiPath << '/';
+	}
 	AnsiPath << InPath;
 
 	FEventLoop Loop;
@@ -61,5 +62,3 @@ void LatencyTest(FStringView InUrl, FStringView InPath, TArrayView<int32> OutRes
 }
 
 } // namespace UE::IO::IAS::HTTP
-
-#endif // !UE_BUILD_SHIPPING
\ No newline at end of file
diff --git a/Engine/Source/Runtime/Experimental/IoStoreOnDemand/Private/CoreHttp/LatencyTesting.h b/Engine/Source/Runtime/Experimental/IoStoreOnDemand/Private/CoreHttp/LatencyTesting.h
--- a/Engine/Source/Runtime/Experimental/IoStoreOnDemand/Private/CoreHttp/LatencyTesting.h
+++ b/Engine/Source/Runtime/Experimental/IoStoreOnDemand/Private/CoreHttp/LatencyTesting.h
@@ -5,13 +5,9 @@
 #include "Containers/ArrayView.h"
 #include "Containers/StringFwd.h"
 
-#if !UE_BUILD_SHIPPING
-
 namespace UE::IO::IAS::HTTP
 {
 
 void LatencyTest(FStringView InUrl, FStringView InPath, TArrayView<int32> OutResults);
 
 } // namespace UE::IO::IAS::HTTP
-
-#endif // !UE_BUILD_SHIPPING
diff --git a/Engine/Source/Runtime/Experimental/IoStoreOnDemand/Private/IoStoreOnDemand.cpp b/Engine/Source/Runtime/Experimental/IoStoreOnDemand/Private/IoStoreOnDemand.cpp
--- a/Engine/Source/Runtime/Experimental/IoStoreOnDemand/Private/IoStoreOnDemand.cpp
+++ b/Engine/Source/Runtime/Experimental/IoStoreOnDemand/Private/IoStoreOnDemand.cpp
@@ -584,7 +584,8 @@ TIoStatusOr<FOnDemandToc> LoadTocFromUrl(const FString& ServiceURL, const FStrin
 				}
 			});
 
-		while (HttpClient->Tick());
+		const bool bBlock = true;
+		while (HttpClient->Tick(bBlock));
 
 		if (Toc.IsOk())
 		{
diff --git a/Engine/Source/Runtime/Experimental/IoStoreOnDemand/Private/OnDemandIoDispatcherBackend.cpp b/Engine/Source/Runtime/Experimental/IoStoreOnDemand/Private/OnDemandIoDispatcherBackend.cpp
--- a/Engine/Source/Runtime/Experimental/IoStoreOnDemand/Private/OnDemandIoDispatcherBackend.cpp
+++ b/Engine/Source/Runtime/Experimental/IoStoreOnDemand/Private/OnDemandIoDispatcherBackend.cpp
@@ -4,6 +4,7 @@
 
 #include "AnalyticsEventAttribute.h"
 #include "CancellationToken.h"
+#include "Containers/BitArray.h"
 #include "Containers/StringView.h"
 #include "CoreHttp/LatencyTesting.h"
 #include "DistributionEndpoints.h"
@@ -26,6 +27,7 @@
 #include "IO/IoStoreOnDemand.h"
 #include "IasCache.h"
 #include "Math/NumericLimits.h"
+#include "Misc/EnumClassFlags.h"
 #include "Misc/ScopeLock.h"
 #include "Misc/ScopeRWLock.h"
 #include "OnDemandHttpClient.h"
@@ -55,6 +57,34 @@ static FAutoConsoleVariableRef CVar_IasMaxHttpRetryCount(
 	TEXT("Max number of HTTP request retries before failing the I/O request.")
 );
 
+int32 GIasHttpHealthCheckWaitTime = 3000;
+static FAutoConsoleVariableRef CVar_IasHttpHealthCheckWaitTime(
+	TEXT("ias.HttpHealthCheckWaitTime"),
+	GIasHttpHealthCheckWaitTime,
+	TEXT("Number of milliseconds to wait before reconnecting to avaiable endpoint(s)")
+);
+
+int32 GIasMaxEndpointTestCountAtStartup = 1;
+static FAutoConsoleVariableRef CVar_IasMaxEndpointTestCountAtStartup(
+	TEXT("ias.MaxEndpointTestCountAtStartup"),
+	GIasMaxEndpointTestCountAtStartup,
+	TEXT("Number of endpoint(s) to test at startup")
+);
+
+int32 GIasHttpErrorSampleCount = 8;
+static FAutoConsoleVariableRef CVar_IasHttpErrorSampleCount(
+	TEXT("ias.HttpErrorSampleCount"),
+	GIasHttpErrorSampleCount,
+	TEXT("Number of samples for computing the moving average of failed HTTP requests")
+);
+
+float GIasHttpErrorHighWater = 0.5f;
+static FAutoConsoleVariableRef CVar_IasHttpErrorHighWater(
+	TEXT("ias.HttpErrorHighWater"),
+	GIasHttpErrorHighWater,
+	TEXT("High water mark when HTTP streaming will be disabled")
+);
+
 bool GIasHttpEnabled = true;
 static FAutoConsoleVariableRef CVar_IasHttpEnabled(
 	TEXT("ias.HttpEnabled"),
@@ -76,16 +106,56 @@ static TAutoConsoleVariable<bool> CVar_IoReportAnalytics(
 
 ///////////////////////////////////////////////////////////////////////////////
 #if !UE_BUILD_SHIPPING
-static void LatencyTest(FStringView InUrl, FStringView InPath)
+static void LatencyTest(FStringView Url, FStringView Path)
 {
 	int32 Results[4] = {};
-	UE::IO::IAS::HTTP::LatencyTest(InUrl, InPath, MakeArrayView(Results));
-
-	UE_LOG(LogIas, VeryVerbose, TEXT("HEAD latencies (ms); %d %d %d %d (%s)"),
-		Results[0], Results[1], Results[2], Results[3], InUrl.GetData());
+	UE::IO::IAS::HTTP::LatencyTest(Url, Path, MakeArrayView(Results));
+	UE_LOG(LogIas, Log, TEXT("Endpoint '%s' latency test (ms): %d %d %d %d"),
+		Url.GetData(), Results[0], Results[1], Results[2], Results[3]);
 }
 #endif // !UE_BUILD_SHIPPING
+///////////////////////////////////////////////////////////////////////////////
+static int32 LatencyTest(TConstArrayView<FString> Urls, FStringView Path, std::atomic_bool& bCancel)
+{
+	for (int32 Idx = 0; Idx < Urls.Num() && !bCancel.load(std::memory_order_relaxed); ++Idx)
+	{
+		int32 LatencyMs = -1;
+		UE::IO::IAS::HTTP::LatencyTest(Urls[Idx], Path, MakeArrayView(&LatencyMs, 1));
+		if (LatencyMs > 0)
+		{
+			return Idx;
+		}
+	}
 
+	return INDEX_NONE;
+}
+///////////////////////////////////////////////////////////////////////////////
+struct FBitWindow
+{
+	void Reset(uint32 Count)
+	{
+		Count = FMath::RoundUpToPowerOfTwo(Count);
+		Bits.SetNum(int32(Count), false);
+		Counter = 0;
+		Mask = Count - 1;
+	}
+
+	void Add(bool bValue)
+	{
+		const uint32 Idx = Counter++ & Mask;
+		Bits[Idx] = bValue;
+	}
+
+	float AvgSetBits() const
+	{
+		return float(Bits.CountSetBits()) / float(Bits.Num());
+	}
+
+private:
+	TBitArray<> Bits;
+	uint32 Counter = 0;
+	uint32 Mask = 0;
+};
 ///////////////////////////////////////////////////////////////////////////////
 FIoHash GetChunkKey(const FIoHash& ChunkHash, const FIoOffsetAndLength& Range)
 {
@@ -127,7 +197,7 @@ public:
 
 	struct FToc
 	{
-		FOnDemandEndpoint Endpoint;
+		FString TocPath;
 		TArray<FContainer> Containers;
 	};
 
@@ -151,9 +221,10 @@ public:
 	FOnDemandIoStore();
 	~FOnDemandIoStore();
 
-	void AddToc(const FOnDemandEndpoint& Ep, FOnDemandToc&& Toc);
+	void AddToc(const FString& TocPath, FOnDemandToc&& Toc);
 	TIoStatusOr<uint64> GetChunkSize(const FIoChunkId& ChunkId);
 	FChunkInfo GetChunkInfo(const FIoChunkId& ChunkId);
+	FString GetFirstTocPath() const;
 
 	TArray<FIoChunkId> GetAllChunkIds(bool bIncludeOptionalChunks);
 
@@ -177,19 +248,18 @@ FOnDemandIoStore::~FOnDemandIoStore()
 	FEncryptionKeyManager::Get().OnKeyAdded().RemoveAll(this);
 }
 
-void FOnDemandIoStore::AddToc(const FOnDemandEndpoint& Ep, FOnDemandToc&& Toc)
+void FOnDemandIoStore::AddToc(const FString& TocPath, FOnDemandToc&& Toc)
 {
-	check(Ep.IsValid());
-	UE_LOG(LogIas, Log, TEXT("Adding TOC '%s/%s'"), *Ep.ServiceUrl, *Ep.TocPath);
+	UE_LOG(LogIas, Log, TEXT("Adding TOC '%s'"), *TocPath);
 
 	TRACE_CPUPROFILER_EVENT_SCOPE(FOnDemandIoBackend::AddToc);
 
 	FString Prefix;
 	{
 		int32 Idx = INDEX_NONE;
-		if (Ep.TocPath.FindLastChar(TCHAR('/'), Idx))
+		if (TocPath.FindLastChar(TCHAR('/'), Idx))
 		{
-			Prefix = Ep.TocPath.Left(Idx);
+			Prefix = TocPath.Left(Idx);
 		}
 	}
 
@@ -197,7 +267,7 @@ void FOnDemandIoStore::AddToc(const FOnDemandEndpoint& Ep, FOnDemandToc&& Toc)
 		FWriteScopeLock _(Lock);
 
 		const FOnDemandTocHeader& Header = Toc.Header;
-		FToc* NewToc = new(Tocs) FToc{Ep};
+		FToc* NewToc = new(Tocs) FToc{TocPath};
 		NewToc->Containers.SetNum(Toc.Containers.Num()); // List of containers can never change
 
 		const FName CompressionFormat(Header.CompressionFormat);
@@ -329,6 +399,12 @@ void FOnDemandIoStore::OnEncryptionKeyAdded(const FGuid& Id, const FAES::FAESKey
 	AddDeferredContainers();
 }
 
+FString FOnDemandIoStore::GetFirstTocPath() const
+{
+	FReadScopeLock _(Lock);
+	return Tocs.Num() > 0 ? Tocs[0].TocPath : FString();
+}
+
 ///////////////////////////////////////////////////////////////////////////////
 template<typename T>
 class TThreadSafeIntrusiveQueue
@@ -658,7 +734,153 @@ static void LogIoResult(
 		CompressedOffset,
 		Priority);
 };
+///////////////////////////////////////////////////////////////////////////////
+struct FBackendStatus
+{
+	enum class EFlags : uint8
+	{
+		None						= 0,
+		CacheEnabled				= (1 << 0),
+		HttpEnabled					= (1 << 1),
+		HttpError					= (1 << 2),
+		HttpBulkOptionalDisabled	= (1 << 3),
+	};
+
+	bool IsHttpEnabled() const
+	{
+		return IsHttpEnabled(Flags.load(std::memory_order_relaxed));
+	}
+
+	bool IsHttpEnabled(EIoChunkType ChunkType) const
+	{
+		const uint8 CurrentFlags = Flags.load(std::memory_order_relaxed);
+		return IsHttpEnabled(CurrentFlags) &&
+			(ChunkType != EIoChunkType::OptionalBulkData ||
+				((CurrentFlags & uint8(EFlags::HttpBulkOptionalDisabled)) == 0 && GIasHttpOptionalBulkDataEnabled));
+	}
+
+	bool IsHttpError() const
+	{
+		return HasAnyFlags(EFlags::HttpError);
+	}
+
+	bool IsCacheEnabled() const
+	{
+		return HasAnyFlags(EFlags::CacheEnabled);
+	}
+
+	bool IsCacheWriteable() const
+	{
+		const uint8 CurrentFlags = Flags.load(std::memory_order_relaxed);
+		return (CurrentFlags & uint8(EFlags::CacheEnabled)) && IsHttpEnabled(CurrentFlags); 
+	}
+
+	bool IsCacheReadOnly() const
+	{
+		const uint8 CurrentFlags = Flags.load(std::memory_order_relaxed);
+		return (CurrentFlags & uint8(EFlags::CacheEnabled)) && !IsHttpEnabled(CurrentFlags);
+	}
+
+	void SetHttpEnabled(bool bEnabled)
+	{
+		AddOrRemoveFlags(EFlags::HttpEnabled, bEnabled, TEXT("HTTP streaming enabled"));
+		FGenericCrashContext::SetEngineData(TEXT("IAS.Enabled"), bEnabled ? TEXT("true") : TEXT("false"));
+	}
+
+	void SetHttpOptionalBulkEnabled(bool bEnabled)
+	{
+		AddOrRemoveFlags(EFlags::HttpBulkOptionalDisabled, bEnabled == false, TEXT("HTTP streaming of optional bulk data disabled"));
+	}
 
+	void SetCacheEnabled(bool bEnabled)
+	{
+		AddOrRemoveFlags(EFlags::CacheEnabled, bEnabled, TEXT("Cache enabled"));
+	}
+
+	void SetHttpError(bool bError)
+	{
+		AddOrRemoveFlags(EFlags::HttpError, bError, TEXT("HTTP streaming error"));
+	}
+
+private:
+	static bool IsHttpEnabled(uint8 FlagsToTest)
+	{
+		constexpr uint8 HttpFlags = uint8(EFlags::HttpEnabled) | uint8(EFlags::HttpError);
+		return ((FlagsToTest & HttpFlags) == uint8(EFlags::HttpEnabled)) && GIasHttpEnabled;
+	}
+
+	bool HasAnyFlags(uint8 Contains) const
+	{
+		return (Flags.load(std::memory_order_relaxed) & Contains) != 0;
+	}
+	
+	bool HasAnyFlags(EFlags Contains) const
+	{
+		return HasAnyFlags(uint8(Contains));
+	}
+
+	uint8 AddFlags(EFlags FlagsToAdd)
+	{
+		return Flags.fetch_or(uint8(FlagsToAdd));
+	}
+
+	uint8 RemoveFlags(EFlags FlagsToRemove)
+	{
+		return Flags.fetch_and(~uint8(FlagsToRemove));
+	}
+
+	uint8 AddOrRemoveFlags(EFlags FlagsToAddOrRemove, bool bValue)
+	{
+		return bValue ? AddFlags(FlagsToAddOrRemove) : RemoveFlags(FlagsToAddOrRemove);
+	}
+
+	void AddOrRemoveFlags(EFlags FlagsToAddOrRemove, bool bValue, const TCHAR* DebugText)
+	{
+		const uint8 PrevFlags = AddOrRemoveFlags(FlagsToAddOrRemove, bValue);
+		TStringBuilder<128> Sb;
+		Sb.Append(DebugText)
+			<< TEXT(" '");
+		Sb.Append(bValue ? TEXT("true") : TEXT("false"))
+			<< TEXT("', backend status '(")
+			<< EFlags(PrevFlags)
+			<< TEXT(") -> (")
+			<< EFlags(Flags.load(std::memory_order_relaxed))
+			<< TEXT(")'");
+		UE_LOG(LogIas, Log, TEXT("%s"), Sb.ToString());
+	}
+
+	friend FStringBuilderBase& operator<<(FStringBuilderBase& Sb, EFlags StatusFlags)
+	{
+		if (StatusFlags == EFlags::None)
+		{
+			Sb.Append(TEXT("None"));
+			return Sb;
+		}
+
+		bool bFirst = true;
+		auto AppendIf = [StatusFlags, &Sb, &bFirst](EFlags Contains, const TCHAR* Str)
+		{
+			if (uint8(StatusFlags) & uint8(Contains))
+			{
+				if (!bFirst)
+				{
+					Sb.AppendChar(TEXT('|'));
+				}
+				Sb.Append(Str);
+				bFirst = false;
+			}
+		};
+
+		AppendIf(EFlags::CacheEnabled, TEXT("CacheEnabled"));
+		AppendIf(EFlags::HttpEnabled, TEXT("HttpEnabled"));
+		AppendIf(EFlags::HttpError, TEXT("HttpError"));
+		AppendIf(EFlags::HttpBulkOptionalDisabled, TEXT("HttpBulkOptionalDisabled"));
+
+		return Sb;
+	}
+
+	std::atomic<uint8> Flags{0};
+};
 ///////////////////////////////////////////////////////////////////////////////
 class FOnDemandIoBackend final
 	: public FRunnable
@@ -667,6 +889,15 @@ class FOnDemandIoBackend final
 	using FIoRequestQueue = TThreadSafeIntrusiveQueue<FIoRequestImpl>;
 	using FChunkRequestQueue = TThreadSafeIntrusiveQueue<FChunkRequest>;
 
+	struct FAvailableEps
+	{
+		bool HasCurrent() const { return Current != INDEX_NONE; }
+		const FString& GetCurrent() const { return Urls[Current]; }
+
+		int32 Current = INDEX_NONE;
+		TArray<FString> Urls;
+	};
+
 	struct FBackendData
 	{
 		static void Attach(FIoRequestImpl* Request, const FIoHash& ChunkKey)
@@ -815,53 +1046,43 @@ public:
 	virtual bool DoesChunkExist(const FIoChunkId& ChunkId) const override;
 	virtual bool DoesChunkExist(const FIoChunkId& ChunkId, const FIoOffsetAndLength& ChunkRange) const override;
 	virtual TIoStatusOr<uint64> GetSizeForChunk(const FIoChunkId& ChunkId) const override;
+	virtual TIoStatusOr<uint64> GetSizeForChunk(const FIoChunkId& ChunkId, const FIoOffsetAndLength& ChunkRange, uint64& OutAvailable) const;
 	virtual FIoRequestImpl* GetCompletedRequests() override;
 	virtual TIoStatusOr<FIoMappedRegion> OpenMapped(const FIoChunkId& ChunkId, const FIoReadOptions& Options) override;
 
 	// I/O Http backend
 	virtual void Mount(const FOnDemandEndpoint& Endpoint) override;
-	virtual void SetBulkOptionalEnabled(bool bInEnabled) override;
-	virtual void SetEnabled(bool bInEnabled) override;
+	virtual void SetBulkOptionalEnabled(bool bEnabled) override;
+	virtual void SetEnabled(bool bEnabled) override;
 	virtual void ReportAnalytics(TArray<FAnalyticsEventAttribute>& OutAnalyticsArray) const override;
 
-#if IS_PROGRAM || WITH_EDITOR
-	virtual bool FlushDeferedEndPoints(double TimeOut = 0.0) override;
-	virtual TArray<FIoChunkId> GetAllChunkIds() override;
-#endif // IS_PROGRAM || WITH_EDITOR
-
 	// Runnable
 	virtual bool Init() override { return true; }
 	virtual void Stop() override { bStopRequested = true; }
 	virtual uint32 Run() override;
 
 private:
+	FString GetEndpointTestPath() const;
+	void ConditionallyStartBackendThread();
 	void CompleteRequest(FChunkRequest* ChunkRequest);
-	FIoStatus MountDeferredEndpoints(const FString& DistributionUrl, const TConstArrayView<FString>& ServiceUrls);
-	static TIoStatusOr<FOnDemandToc> GetToc(const FOnDemandEndpoint& Endpoint);
-	FIoStatus AddToc(const FOnDemandEndpoint& Endpoint);
-	bool IsHttpEnabled() const { return bHttpEnabled && GIasHttpEnabled; }
-	bool IsHttpEnabled(const FIoChunkId& ChunkId) const
-	{ 
-		return (ChunkId.GetChunkType() != EIoChunkType::OptionalBulkData || (bHttpOptionalBulkDataEnabled && GIasHttpOptionalBulkDataEnabled));
-	}
-	virtual TIoStatusOr<uint64> GetSizeForChunk(const FIoChunkId& ChunkId, const FIoOffsetAndLength& ChunkRange, uint64& OutAvailable) const;
+	void AddDeferredTocs();
+	void ProcessHttpRequests(FOnDemandHttpClient* HttpClient, FBitWindow& HttpErrors, int32 MaxConcurrentRequests);
 
 	TSharedPtr<IIasCache> Cache;
 	TUniquePtr<FOnDemandIoStore> IoStore;
 	TSharedPtr<const FIoDispatcherBackendContext> BackendContext;
 	TUniquePtr<FRunnableThread> BackendThread;
 	FEventRef TickBackendEvent;
-	FDistributionEndpoints DistributionEndpoints;
-	TArray<FOnDemandEndpoint> DeferredEndpoints;
+	TArray<FString> DeferredTocs;
 	FChunkRequests ChunkRequests;
 	FIoRequestQueue CompletedRequests;
 	FChunkRequestQueue HttpRequests;
-	TUniquePtr<FOnDemandHttpClient> HttpClient;
 	FOnDemandIoBackendStats Stats;
-	FRWLock Lock;
+	FBackendStatus BackendStatus;
+	FAvailableEps AvailableEps;
+	FString DistributionUrl;
+	mutable FRWLock Lock;
 	std::atomic_bool bStopRequested{false};
-	std::atomic_bool bHttpOptionalBulkDataEnabled{true};
-	std::atomic_bool bHttpEnabled{true};
 };
 
 ///////////////////////////////////////////////////////////////////////////////
@@ -869,6 +1090,8 @@ FOnDemandIoBackend::FOnDemandIoBackend(TSharedPtr<IIasCache> InCache)
 	: Cache(InCache)
 {
 	IoStore = MakeUnique<FOnDemandIoStore>();
+	BackendStatus.SetHttpEnabled(true);
+	BackendStatus.SetCacheEnabled(Cache.IsValid());
 }
 
 FOnDemandIoBackend::~FOnDemandIoBackend()
@@ -881,9 +1104,43 @@ void FOnDemandIoBackend::Initialize(TSharedRef<const FIoDispatcherBackendContext
 	TRACE_CPUPROFILER_EVENT_SCOPE(FOnDemandIoBackend::Initialize);
 	LLM_SCOPE_BYTAG(Ias);
 	UE_LOG(LogIas, Log, TEXT("Initializing on demand I/O dispatcher backend"));
-	FGenericCrashContext::SetEngineData(TEXT("IAS.Enabled"), bHttpEnabled.load(std::memory_order_relaxed) ? TEXT("true") : TEXT("false"));
 	BackendContext = Context;
-	DistributionEndpoints.ResolveDeferredEndpoints();
+
+	if (DistributionUrl.IsEmpty() == false)
+	{
+		TSharedPtr<FDistributionEndpoints> Resolver = MakeShared<FDistributionEndpoints>();
+		Resolver->ResolveEndpoints(DistributionUrl,
+			[this, Resolver](const FString& DistributionEp, TConstArrayView<FString> Eps)
+			{
+				UE_CLOG(Eps.IsEmpty(), LogIas, Warning, TEXT("Failed to resolve available endpoint(s) from '%s'"), *DistributionEp);
+				{
+					FWriteScopeLock _(Lock);
+					for (const FString& Ep : Eps)
+					{
+						AvailableEps.Urls.Add(Ep.Replace(TEXT("https"), TEXT("http")));
+					}
+				}
+
+				{
+					TConstArrayView<FString> Urls = AvailableEps.Urls;
+					const int32 MaxUrls = FMath::Min(GIasMaxEndpointTestCountAtStartup, AvailableEps.Urls.Num());
+					const FString TestPath = GetEndpointTestPath();
+					if (int32 Idx = LatencyTest(Urls.Left(MaxUrls), TestPath, bStopRequested); Idx != INDEX_NONE)
+					{
+						AvailableEps.Current = Idx;
+						UE_LOG(LogIas, Log, TEXT("Using endpoint '%s'"), *AvailableEps.GetCurrent());
+						AddDeferredTocs();
+					}
+					else
+					{
+						BackendStatus.SetHttpError(true);
+					}
+				}
+				ConditionallyStartBackendThread();
+			});
+		Resolver->ResolveDeferredEndpoints();
+		DistributionUrl.Reset();
+	}
 }
 
 void FOnDemandIoBackend::Shutdown()
@@ -901,6 +1158,29 @@ void FOnDemandIoBackend::Shutdown()
 	BackendContext.Reset();
 }
 
+FString FOnDemandIoBackend::GetEndpointTestPath() const
+{
+	FString TestPath = IoStore.IsValid() ? IoStore->GetFirstTocPath() : FString();
+	if (TestPath.IsEmpty())
+	{
+		FReadScopeLock _(Lock);
+		if (DeferredTocs.IsEmpty() == false)
+		{
+			TestPath = DeferredTocs[0];
+		}
+	}
+	return TestPath;
+}
+
+void FOnDemandIoBackend::ConditionallyStartBackendThread()
+{
+	FWriteScopeLock _(Lock);
+	if (BackendThread.IsValid() == false)
+	{
+		BackendThread.Reset(FRunnableThread::Create(this, TEXT("Ias.Http"), 0, TPri_AboveNormal));
+	}
+}
+
 void FOnDemandIoBackend::CompleteRequest(FChunkRequest* ChunkRequest)
 {
 	LLM_SCOPE_BYTAG(Ias);
@@ -921,7 +1201,7 @@ void FOnDemandIoBackend::CompleteRequest(FChunkRequest* ChunkRequest)
 	FIoChunkDecodingParams DecodingParams = ChunkRequest->Params.GetDecodingParams();
 
 	// Only cache chunks if HTTP streaming is enabled
-	bool bCanCache = Cache.IsValid() && IsHttpEnabled();
+	bool bCacheChunk = ChunkRequest->bCached == false && Chunk.GetSize() > 0;
 	FIoRequestImpl* NextRequest = ChunkRequest->DeqeueDispatcherRequests();
 	while (NextRequest)
 	{
@@ -952,7 +1232,7 @@ void FOnDemandIoBackend::CompleteRequest(FChunkRequest* ChunkRequest)
 		}
 		else
 		{
-			bCanCache = false;
+			bCacheChunk = false;
 			Request->SetFailed();
 
 			Stats.OnIoRequestError();
@@ -965,7 +1245,7 @@ void FOnDemandIoBackend::CompleteRequest(FChunkRequest* ChunkRequest)
 		BackendContext->WakeUpDispatcherThreadDelegate.Execute();
 	}
 
-	if (bCanCache && !ChunkRequest->bCached && Chunk.GetSize() > 0)
+	if (bCacheChunk && BackendStatus.IsCacheWriteable())
 	{
 		Cache->Put(ChunkRequest->Params.ChunkKey, Chunk);
 	}
@@ -985,11 +1265,10 @@ bool FOnDemandIoBackend::Resolve(FIoRequestImpl* Request)
 
 	FChunkRequestParams RequestParams = FChunkRequestParams::Create(Request, ChunkInfo);
 
-	if (IsHttpEnabled() == false || IsHttpEnabled(Request->ChunkId) == false)
+	if (BackendStatus.IsHttpEnabled(Request->ChunkId.GetChunkType()) == false)
 	{ 
-		// Allow reading from the cache only when HTTP streaming is disabled otherwise the chunk may be evicted
-		// before trying to read the cache entry.
-		if (IsHttpEnabled() || Cache.IsValid() == false || Cache->ContainsChunk(RequestParams.ChunkKey) == false)
+		// If the cache is not readonly the chunk may get evicted before the request is completed
+		if (BackendStatus.IsCacheReadOnly() == false || Cache->ContainsChunk(RequestParams.ChunkKey) == false)
 		{
 			return false;
 		}
@@ -1036,9 +1315,9 @@ bool FOnDemandIoBackend::Resolve(FIoRequestImpl* Request)
 			}
 		}
 
-		if (ChunkRequest->CancellationToken.IsCancelled() || IsHttpEnabled() == false)
+		if (ChunkRequest->CancellationToken.IsCancelled() || BackendStatus.IsHttpEnabled() == false)
 		{
-			UE_CLOG(IsHttpEnabled() == false, LogIas, Log, TEXT("Chunk was not found in the cache and HTTP is disabled"));
+			UE_CLOG(BackendStatus.IsHttpEnabled() == false, LogIas, Log, TEXT("Chunk was not found in the cache and HTTP is disabled"));
 			return CompleteRequest(ChunkRequest);
 		}
 
@@ -1101,14 +1380,15 @@ TIoStatusOr<uint64> FOnDemandIoBackend::GetSizeForChunk(const FIoChunkId& ChunkI
 	FIoOffsetAndLength RequestedRange(ChunkRange.GetOffset(), FMath::Min<uint64>(ChunkInfo.Entry->RawSize, ChunkRange.GetLength()));
 	OutAvailable = ChunkInfo.Entry->RawSize;
 
-	if (IsHttpEnabled() == false || IsHttpEnabled(ChunkId) == false)
+	if (BackendStatus.IsHttpEnabled(ChunkId.GetChunkType()) == false)
 	{
-		if (IsHttpEnabled() || Cache.IsValid() == false)
+		// If the cache is not readonly the chunk may get evicted before the request is resolved
+		if (BackendStatus.IsCacheReadOnly() == false)
 		{
 			return FIoStatus(EIoErrorCode::UnknownChunkID);
 		}
 
-		// When HTTP streaming is disabled, no cache entries will be evicted. 
+		check(Cache.IsValid());
 		const FChunkRequestParams RequestParams = FChunkRequestParams::Create(RequestedRange, ChunkInfo);
 		if (Cache->ContainsChunk(RequestParams.ChunkKey) == false)
 		{
@@ -1140,42 +1420,30 @@ TIoStatusOr<FIoMappedRegion> FOnDemandIoBackend::OpenMapped(const FIoChunkId& Ch
 	return FIoStatus::Unknown;
 }
 
-FIoStatus FOnDemandIoBackend::MountDeferredEndpoints(const FString& DistributionUrl, const TConstArrayView<FString>& ServiceUrls)
+void FOnDemandIoBackend::AddDeferredTocs()
 {
-	TArray<FOnDemandEndpoint, TInlineAllocator<4>> EndpointsToAdd;
+	TArray<FString> TocPaths;
 	{
 		FWriteScopeLock _(Lock);
-		for (auto It = DeferredEndpoints.CreateIterator(); It; ++It)
+		if (DeferredTocs.IsEmpty() || AvailableEps.HasCurrent() == false)
 		{
-			if (It->DistributionUrl.Compare(DistributionUrl, ESearchCase::IgnoreCase) == 0)
-			{
-				EndpointsToAdd.Add(*It);
-				It.RemoveCurrent();
-			}
+			return;
 		}
+		TocPaths = MoveTemp(DeferredTocs);
 	}
 
-	for (const FOnDemandEndpoint& Ep : EndpointsToAdd)
+	for (const FString& TocPath : TocPaths)
 	{
-		FIoStatus Status;
-		for (const FString& SerivceUrl : ServiceUrls)
+		TIoStatusOr<FOnDemandToc> Toc = LoadTocFromUrl(AvailableEps.GetCurrent(), TocPath, 1);
+		if (Toc.IsOk())
 		{
-			// Currently we don't need use secure sockets to fetch on demand content
-			FString UnsecureUrl = SerivceUrl.Replace(TEXT("https"), TEXT("http"));
-			Status = AddToc(FOnDemandEndpoint{Ep.EndpointType, DistributionUrl, UnsecureUrl, Ep.TocPath});
-			if (Status.IsOk())
-			{
-				break;
-			}
+			IoStore->AddToc(TocPath, Toc.ConsumeValueOrDie());
 		}
-
-		if (!Status.IsOk())
+		else
 		{
-			return Status;
+			UE_LOG(LogIas, Error, TEXT("Failed to add TOC '%s/%s', reason '%s'"), *AvailableEps.GetCurrent(), *TocPath, *Toc.Status().ToString());
 		}
 	}
-
-	return FIoStatus::Ok;
 }
 
 void FOnDemandIoBackend::Mount(const FOnDemandEndpoint& Endpoint)
@@ -1191,26 +1459,42 @@ void FOnDemandIoBackend::Mount(const FOnDemandEndpoint& Endpoint)
 
 	if (EnumHasAnyFlags(Endpoint.EndpointType, EOnDemandEndpointType::CDN))
 	{
-		if (Endpoint.ServiceUrl.IsEmpty())
 		{
+			FWriteScopeLock _(Lock);
+			if (Endpoint.ServiceUrl.IsEmpty())
 			{
-				FWriteScopeLock _(Lock);
-				DeferredEndpoints.Add(Endpoint);
-			}
-
-			DistributionEndpoints.ResolveEndpoints(Endpoint.DistributionUrl, [this](const FString& DistributionUrl, TConstArrayView<FString> SerivceUrls)
-			{
-				if (FIoStatus Status = MountDeferredEndpoints(DistributionUrl, SerivceUrls); !Status.IsOk())
+				if (AvailableEps.HasCurrent() == false)
 				{
-					UE_LOG(LogIas, Error, TEXT("Failed to add endpoint(s), reason '%s'"), *Status.ToString());
+					if (DistributionUrl.IsEmpty())
+					{
+						DistributionUrl = Endpoint.DistributionUrl;
+					}
+					DeferredTocs.Add(Endpoint.TocPath);
+					return;
 				}
-			});
+			}
+			else if (AvailableEps.Urls.IsEmpty())
+			{
+				AvailableEps.Urls.Add(Endpoint.ServiceUrl.Replace(TEXT("https"), TEXT("http")));
+				AvailableEps.Current = 0;
+			}
 		}
-		else if (FIoStatus Status = AddToc(Endpoint); !Status.IsOk())
+
+		check(AvailableEps.HasCurrent());
+		TIoStatusOr<FOnDemandToc> Toc = LoadTocFromUrl(AvailableEps.GetCurrent(), Endpoint.TocPath, 1);
+		if (Toc.IsOk())
 		{
-			UE_LOG(LogIas, Error, TEXT("Failed to add TOC '%s/%s', reason '%s'"),
-				*Endpoint.ServiceUrl, *Endpoint.TocPath, *Status.ToString());
+			IoStore->AddToc(Endpoint.TocPath, Toc.ConsumeValueOrDie());
 		}
+		else
+		{
+			UE_LOG(LogIas, Error, TEXT("Deferring TOC '%s/%s' due to '%s'"), *AvailableEps.GetCurrent(), *Endpoint.TocPath, *Toc.Status().ToString());
+			BackendStatus.SetHttpError(true);
+			FWriteScopeLock _(Lock);
+			DeferredTocs.Add(Endpoint.TocPath);
+		}
+
+		ConditionallyStartBackendThread();
 	}
 	else
 	{
@@ -1218,17 +1502,14 @@ void FOnDemandIoBackend::Mount(const FOnDemandEndpoint& Endpoint)
 	}
 }
 
-void FOnDemandIoBackend::SetBulkOptionalEnabled(bool bInEnabled)
+void FOnDemandIoBackend::SetBulkOptionalEnabled(bool bEnabled)
 {
-	UE_LOG(LogIas, Log, TEXT("HTTP optional bulk data streaming '%s'"), bInEnabled ? TEXT("Enabled") : TEXT("Disabled"));
-	bHttpOptionalBulkDataEnabled = bInEnabled;
+	BackendStatus.SetHttpOptionalBulkEnabled(bEnabled);
 }
 
-void FOnDemandIoBackend::SetEnabled(bool bInEnabled)
+void FOnDemandIoBackend::SetEnabled(bool bEnabled)
 {
-	UE_LOG(LogIas, Log, TEXT("HTTP streaming '%s'"), bInEnabled ? TEXT("Enabled") : TEXT("Disabled"));
-	bHttpEnabled = bInEnabled;
-	FGenericCrashContext::SetEngineData(TEXT("IAS.Enabled"), bHttpEnabled.load(std::memory_order_acquire) ? TEXT("true") : TEXT("false"));
+	BackendStatus.SetHttpEnabled(bEnabled);
 }
 
 void FOnDemandIoBackend::ReportAnalytics(TArray<FAnalyticsEventAttribute>& OutAnalyticsArray) const
@@ -1238,9 +1519,9 @@ void FOnDemandIoBackend::ReportAnalytics(TArray<FAnalyticsEventAttribute>& OutAn
 		return;
 	}
 
-	if (HttpClient.IsValid())
+	if (AvailableEps.HasCurrent())
 	{
-		FString CdnUrl = HttpClient->ServiceUrl();
+		FString CdnUrl = AvailableEps.GetCurrent();
 
 		// Strip the prefix from the url as some analytics systems may have trouble dealing with it
 		if (!CdnUrl.RemoveFromStart(TEXT("http://")))
@@ -1249,94 +1530,42 @@ void FOnDemandIoBackend::ReportAnalytics(TArray<FAnalyticsEventAttribute>& OutAn
 		}
 
 		AppendAnalyticsEventAttributeArray(OutAnalyticsArray, TEXT("IasCdnUrl"), MoveTemp(CdnUrl));
-
-		Stats.ReportAnalytics(OutAnalyticsArray);
-	}
-}
-
-#if IS_PROGRAM || WITH_EDITOR
-bool FOnDemandIoBackend::FlushDeferedEndPoints(double TimeOut)
-{
-	return DistributionEndpoints.Flush(TimeOut);
-}
-
-TArray<FIoChunkId> FOnDemandIoBackend::GetAllChunkIds()
-{
-	TRACE_CPUPROFILER_EVENT_SCOPE(FOnDemandIoBackend::GetAllChunkIds);
-
-	if (!IoStore.IsValid() || !IsHttpEnabled())
-	{
-		return TArray<FIoChunkId>();
-	}
-
-	const bool bAllowOptional = GIasHttpOptionalBulkDataEnabled && bHttpOptionalBulkDataEnabled;
-
-	return IoStore->GetAllChunkIds(bAllowOptional);
-}
-
-#endif // IS_PROGRAM || WITH_EDITOR
-
-FIoStatus FOnDemandIoBackend::AddToc(const FOnDemandEndpoint& Endpoint)
-{
-	TIoStatusOr<FOnDemandToc> Toc = LoadTocFromUrl(Endpoint.ServiceUrl, Endpoint.TocPath, GIasMaxHttpRetryCount);
-	if (!Toc.IsOk())
-	{
-		return FIoStatus(Toc.Status());
 	}
-
-	{
-		FWriteScopeLock _(Lock);
-		if (!HttpClient.IsValid())
-		{
-			HttpClient = MakeUnique<FOnDemandHttpClient>(Endpoint.ServiceUrl, GIasMaxHttpConnectionCount);
-			BackendThread.Reset(FRunnableThread::Create(this, TEXT("IoStoreOnDemand"), 0, TPri_AboveNormal));
-		}
-	}
-
-#if !UE_BUILD_SHIPPING
-	UE::Tasks::Launch(TEXT("IasLatencyTest"), [ServiceUrl=Endpoint.ServiceUrl, TocPath=Endpoint.TocPath] ()
-	{
-		LatencyTest(ServiceUrl, TocPath);
-	});
-#endif // !UE_BUILD_SHIPPING
-
-	IoStore->AddToc(Endpoint, Toc.ConsumeValueOrDie());
-
-	UE_CLOG(Endpoint.ServiceUrl !=  HttpClient->ServiceUrl(),
-		LogIas, Fatal, TEXT("Fetching on demand content from multiple endpoints are currently not supported"));
-
-	return FIoStatus::Ok;
 }
 
-uint32 FOnDemandIoBackend::Run()
+void FOnDemandIoBackend::ProcessHttpRequests(FOnDemandHttpClient* HttpClient, FBitWindow& HttpErrors, int32 MaxConcurrentRequests)
 {
-	LLM_SCOPE_BYTAG(Ias);
-
-	const int32 MaxConcurrentRequests = HttpClient->MaxConnectionCount();
-	FChunkRequest* NextChunkRequest = nullptr;
 	int32 NumConcurrentRequests = 0;
+	FChunkRequest* NextChunkRequest = HttpRequests.Dequeue();
 
-	while (!bStopRequested)
+	while (NextChunkRequest)
 	{
-		NextChunkRequest = HttpRequests.Dequeue();
 		while (NextChunkRequest)
 		{
-			while (NextChunkRequest)
 			{
-				{
-					TRACE_CPUPROFILER_EVENT_SCOPE(FOnDemandIoBackend::IssueHttpGet);
-					FChunkRequest* ChunkRequest = NextChunkRequest;
-					NextChunkRequest = ChunkRequest->NextRequest;
-					ChunkRequest->NextRequest = nullptr;
+				TRACE_CPUPROFILER_EVENT_SCOPE(FOnDemandIoBackend::IssueHttpGet);
+				FChunkRequest* ChunkRequest = NextChunkRequest;
+				NextChunkRequest = ChunkRequest->NextRequest;
+				ChunkRequest->NextRequest = nullptr;
 
-					Stats.OnHttpDequeue();
+				Stats.OnHttpDequeue();
 
+				if (BackendStatus.IsHttpEnabled() == false)
+				{
+					UE::Tasks::Launch(UE_SOURCE_LOCATION, [this, ChunkRequest]()
+					{
+						CompleteRequest(ChunkRequest);
+					});
+				}
+				else
+				{
+					check(HttpClient);
 					TAnsiStringBuilder<256> Url;
 					ChunkRequest->Params.GetUrl(Url);
 
 					NumConcurrentRequests++;
 					HttpClient->Get(Url.ToView(), ChunkRequest->Params.ChunkRange,
-						[this, ChunkRequest, &NumConcurrentRequests](TIoStatusOr<FIoBuffer> Status, uint64 DurationMs)
+						[this, &NextChunkRequest, ChunkRequest, &NumConcurrentRequests, &HttpErrors](TIoStatusOr<FIoBuffer> Status, uint64 DurationMs)
 						{
 							NumConcurrentRequests--;
 
@@ -1345,22 +1574,33 @@ uint32 FOnDemandIoBackend::Run()
 								if (++ChunkRequest->HttpRetryCount <= GIasMaxHttpRetryCount)
 								{
 									Stats.OnHttpRetry();
-									Stats.OnHttpEnqueue();
-
-									// Note there is no need to trigger TickBackendEvent as this callback will occur within FOnDemandIoBackend::Run
-									ChunkRequest->Priority = IoDispatcherPriority_High;
-									return HttpRequests.EnqueueByPriority(ChunkRequest);
+									ChunkRequest->NextRequest = NextChunkRequest;
+									NextChunkRequest = ChunkRequest;
+									return;
 								}
 							}
 
 							if (Status.IsOk())
 							{
+								HttpErrors.Add(false);
 								ChunkRequest->Chunk = Status.ConsumeValueOrDie();
 								Stats.OnHttpGet(ChunkRequest->Chunk.DataSize(), DurationMs);
 							}
 							else
 							{
 								Stats.OnHttpError();
+								HttpErrors.Add(true);
+
+								const float Average = HttpErrors.AvgSetBits();
+								const bool bAboveHighWaterMark = Average > GIasHttpErrorHighWater;
+								UE_LOG(LogIas, Warning, TEXT("%.2f%% the last %d HTTP requests failed"), Average * 100.0f, GIasHttpErrorSampleCount);
+
+								if (bAboveHighWaterMark)
+								{
+									BackendStatus.SetHttpError(true);
+									UE_LOG(LogIas, Warning, TEXT("HTTP streaming disabled due to high water mark of %.2f of the last %d requests reached"),
+										GIasHttpErrorHighWater * 100.0f, GIasHttpErrorSampleCount);
+								}
 							}
 
 							UE::Tasks::Launch(UE_SOURCE_LOCATION, [this, ChunkRequest]()
@@ -1369,41 +1609,80 @@ uint32 FOnDemandIoBackend::Run()
 							});
 						});
 				}
+			}
 
-				if (NumConcurrentRequests >= MaxConcurrentRequests)
+			if (NumConcurrentRequests >= MaxConcurrentRequests)
+			{
+				TRACE_CPUPROFILER_EVENT_SCOPE(FOnDemandIoBackend::TickHttpSaturated);
+				while (NumConcurrentRequests >= MaxConcurrentRequests)
 				{
-					TRACE_CPUPROFILER_EVENT_SCOPE(FOnDemandIoBackend::TickHttpSaturated);
-					while (NumConcurrentRequests >= MaxConcurrentRequests)
-					{
-						HttpClient->Tick(true);
-					}
+					HttpClient->Tick(true);
 				}
+			}
 
-				if (!NextChunkRequest)
-				{
-					NextChunkRequest = HttpRequests.Dequeue();
-				}
+			if (!NextChunkRequest)
+			{
+				NextChunkRequest = HttpRequests.Dequeue();
 			}
+		}
+
+		{
+			// Keep processing pending connections until all requests are completed or a new one is issued
+			TRACE_CPUPROFILER_EVENT_SCOPE(FOnDemandIoBackend::TickHttp);
+			while (HttpClient->Tick(true) && !NextChunkRequest)
+			{
+				NextChunkRequest = HttpRequests.Dequeue();
+			}
+		}
+	} 
+}
+
+uint32 FOnDemandIoBackend::Run()
+{
+	LLM_SCOPE_BYTAG(Ias);
 
+	FBitWindow HttpErrors;
+	HttpErrors.Reset(GIasHttpErrorSampleCount);
+	
+	TUniquePtr<FOnDemandHttpClient> HttpClient;
+	if (AvailableEps.HasCurrent())
+	{
+		HttpClient = MakeUnique<FOnDemandHttpClient>(AvailableEps.GetCurrent(), GIasMaxHttpConnectionCount);
+#if !UE_BUILD_SHIPPING
+		LatencyTest(AvailableEps.GetCurrent(), GetEndpointTestPath());
+#endif 
+	}
+
+	while (!bStopRequested)
+	{
+		// Process HTTP request(s) even if the client is invalid to ensure enqueued request(s) gets completed.
+		ProcessHttpRequests(HttpClient.Get(), HttpErrors, GIasMaxHttpConnectionCount);
+
+		if (!bStopRequested)
+		{
+			uint32 WaitTime = MAX_uint32;
+			if (BackendStatus.IsHttpError())
 			{
-				// Keep processing pending connections until all requests are completed or a new one is issued
-				TRACE_CPUPROFILER_EVENT_SCOPE(FOnDemandIoBackend::TickHttp);
-				while (HttpClient->Tick(true) && !NextChunkRequest)
+				WaitTime = GIasHttpHealthCheckWaitTime;
+				if (HttpClient.IsValid())
 				{
-					NextChunkRequest = HttpRequests.Dequeue();
+					HttpClient.Reset();
+					HttpErrors.Reset(GIasHttpErrorSampleCount);
+					AvailableEps.Current = INDEX_NONE;
 				}
 
-				// Tick can cause a new request to be added to HttpRequests, so we should try one last time
-				if (NextChunkRequest == nullptr)
+				UE_LOG(LogIas, Log, TEXT("Trying to reconnect to any available endpoint"));
+				const FString TestPath = GetEndpointTestPath(); 
+				if (int32 Idx = LatencyTest(AvailableEps.Urls, TestPath, bStopRequested); Idx != INDEX_NONE)
 				{
-					NextChunkRequest = HttpRequests.Dequeue();
+					AvailableEps.Current = Idx;
+					HttpClient = MakeUnique<FOnDemandHttpClient>(AvailableEps.GetCurrent(), GIasMaxHttpConnectionCount);
+					BackendStatus.SetHttpError(false);
+					UE_LOG(LogIas, Log, TEXT("Successfully reconnected to '%s'"), *AvailableEps.GetCurrent());
+					AddDeferredTocs();
 				}
 			}
-		}
-
-		if (!bStopRequested)
-		{
-			TickBackendEvent->Wait();
+			TickBackendEvent->Wait(WaitTime);
 		}
 	}
 
diff --git a/Engine/Source/Runtime/Experimental/IoStoreOnDemand/Private/OnDemandIoDispatcherBackend.h b/Engine/Source/Runtime/Experimental/IoStoreOnDemand/Private/OnDemandIoDispatcherBackend.h
--- a/Engine/Source/Runtime/Experimental/IoStoreOnDemand/Private/OnDemandIoDispatcherBackend.h
+++ b/Engine/Source/Runtime/Experimental/IoStoreOnDemand/Private/OnDemandIoDispatcherBackend.h
@@ -39,13 +39,7 @@ public:
 	virtual void Mount(const FOnDemandEndpoint& Endpoint) = 0;
 	virtual void SetBulkOptionalEnabled(bool bInEnabled) = 0;
 	virtual void SetEnabled(bool bInEnabled) = 0;
-
 	virtual void ReportAnalytics(TArray<FAnalyticsEventAttribute>& OutAnalyticsArray) const = 0;
-
-#if IS_PROGRAM || WITH_EDITOR
-	virtual bool FlushDeferedEndPoints(double TimeOut) = 0;
-	virtual TArray<FIoChunkId> GetAllChunkIds() = 0;
-#endif // IS_PROGRAM || WITH_EDITOR
 };
 
 TSharedPtr<IOnDemandIoDispatcherBackend> MakeOnDemandIoDispatcherBackend(TSharedPtr<IIasCache> Cache);
