Fix headless Chaos
Single buffer was broken because resetting the consumer rest the only buffer.
Reset is now applied only on double buffer.

diff --git a/Engine/Source/Runtime/Experimental/Chaos/Private/EventManager.cpp b/Engine/Source/Runtime/Experimental/Chaos/Private/EventManager.cpp
--- a/Engine/Source/Runtime/Experimental/Chaos/Private/EventManager.cpp
+++ b/Engine/Source/Runtime/Experimental/Chaos/Private/EventManager.cpp
@@ -62,7 +62,10 @@ namespace Chaos
 		{
 			if (EventContainer)
 			{
-				EventContainer->ResetConsumerBuffer();
+				if (BufferMode == EMultiBufferMode::Double)
+				{
+					EventContainer->ResetConsumerBuffer();
+				}
 				EventContainer->FlipBufferIfRequired();
 			}
 		}
