Geometry collection : Optimize CopyMatchingAttributesFrom method to reduce cache missed and allow for static arrays to be past as skip list avoiding allocation for the caller
this also makes sure we only resize and add the common groups instead of all of them

Note : only the geometry collection proxy is using this new method

diff --git a/Engine/Source/Runtime/Experimental/Chaos/Private/GeometryCollection/ManagedArrayCollection.cpp b/Engine/Source/Runtime/Experimental/Chaos/Private/GeometryCollection/ManagedArrayCollection.cpp
--- a/Engine/Source/Runtime/Experimental/Chaos/Private/GeometryCollection/ManagedArrayCollection.cpp
+++ b/Engine/Source/Runtime/Experimental/Chaos/Private/GeometryCollection/ManagedArrayCollection.cpp
@@ -452,6 +452,40 @@ void FManagedArrayCollection::SyncGroupSizeFrom(const FManagedArrayCollection& I
 	Resize(InCollection.GroupInfo[Group].Size, Group);
 }
 
+void FManagedArrayCollection::CopyMatchingAttributesFrom(const FManagedArrayCollection& FromCollection, const TArrayView<const FAttributeAndGroupId> SkipList)
+{
+	MatchOptionalDefaultAttributes(FromCollection);
+
+	// we only want to resize the groups that are in common 
+	for (const TPair<FName, FGroupInfo>& Pair: FromCollection.GroupInfo)
+	{
+		const FName& GroupName = Pair.Key;
+		if (HasGroup(GroupName))
+		{
+			Resize(Pair.Value.Size, GroupName);
+		}
+	}
+
+	for (TTuple<FKeyType, FValueType>& Entry : Map)
+	{
+		const FName& AttributeName = Entry.Key.Get<0>();
+		const FName& GroupName = Entry.Key.Get<1>();
+
+		if (SkipList.Contains(FAttributeAndGroupId{ AttributeName, GroupName }))
+		{
+			continue;
+		}
+		if (const FValueType* FromAttribute = FromCollection.Map.Find(MakeMapKey(AttributeName, GroupName)))
+		{
+			FValueType& ToAttribute = Entry.Value;
+			if (ToAttribute.ArrayType == FromAttribute->ArrayType)
+			{
+				ToAttribute.Value->Init(*FromAttribute->Value);
+			}
+		}
+	}
+}
+
 void FManagedArrayCollection::CopyMatchingAttributesFrom(
 	const FManagedArrayCollection& InCollection,
 	const TMap<FName, TSet<FName>>* SkipList)
@@ -486,7 +520,6 @@ void FManagedArrayCollection::CopyMatchingAttributesFrom(
 			}
 		}
 	}
-
 }
 
 void FManagedArrayCollection::CopyAttribute(const FManagedArrayCollection& InCollection, FName Name, FName Group)
diff --git a/Engine/Source/Runtime/Experimental/Chaos/Private/PhysicsProxy/GeometryCollectionPhysicsProxy.cpp b/Engine/Source/Runtime/Experimental/Chaos/Private/PhysicsProxy/GeometryCollectionPhysicsProxy.cpp
--- a/Engine/Source/Runtime/Experimental/Chaos/Private/PhysicsProxy/GeometryCollectionPhysicsProxy.cpp
+++ b/Engine/Source/Runtime/Experimental/Chaos/Private/PhysicsProxy/GeometryCollectionPhysicsProxy.cpp
@@ -717,11 +717,13 @@ void FGeometryCollectionPhysicsProxy::Initialize(Chaos::FPBDRigidsEvolutionBase
 	}
 
 	// Skip simplicials, as they're owned by unique pointers.
-	TMap<FName, TSet<FName>> SkipList;
-	TSet<FName>& TransformGroupSkipList = SkipList.Emplace(FTransformCollection::TransformGroup);
-	TransformGroupSkipList.Add(DynamicCollection.SimplicialsAttribute);
+	static const FAttributeAndGroupId SkipList[] =
+	{
+		{ FGeometryDynamicCollection::SimplicialsAttribute, FTransformCollection::TransformGroup },
+	};
+	static constexpr int32 SkipListSize = sizeof(SkipList) / sizeof(FAttributeAndGroupId);
 
-	PhysicsThreadCollection.CopyMatchingAttributesFrom(DynamicCollection, &SkipList);
+	PhysicsThreadCollection.CopyMatchingAttributesFrom(DynamicCollection, MakeArrayView(SkipList, SkipListSize));
 	PhysicsThreadCollection.CopyInitialVelocityAttributesFrom(DynamicCollection);
 
 	// Copy simplicials.
@@ -781,16 +783,18 @@ void FGeometryCollectionPhysicsProxy::InitializeDynamicCollection(FGeometryDynam
 	// This function will use the rest collection to populate the dynamic collection. 
 	//
 
-	TMap<FName, TSet<FName>> SkipList;
-	TSet<FName>& KeepFromDynamicCollection = SkipList.Emplace(FTransformCollection::TransformGroup);
-	KeepFromDynamicCollection.Add(FTransformCollection::TransformAttribute);
-	KeepFromDynamicCollection.Add(FTransformCollection::ParentAttribute);
-	KeepFromDynamicCollection.Add(FTransformCollection::ChildrenAttribute);
-	KeepFromDynamicCollection.Add(FGeometryCollection::SimulationTypeAttribute);
-	KeepFromDynamicCollection.Add(DynamicCollection.SimplicialsAttribute);
-	KeepFromDynamicCollection.Add(DynamicCollection.ActiveAttribute);
-	DynamicCollection.CopyMatchingAttributesFrom(RestCollection, &SkipList);
+	static const FAttributeAndGroupId SkipList[] =
+	{
+		{ FTransformCollection::TransformAttribute, FTransformCollection::TransformGroup },
+		{ FTransformCollection::ParentAttribute, FTransformCollection::TransformGroup },
+		{ FTransformCollection::ChildrenAttribute, FTransformCollection::TransformGroup },
+		{ FGeometryCollection::SimulationTypeAttribute, FTransformCollection::TransformGroup },
+		{ FGeometryDynamicCollection::SimplicialsAttribute, FTransformCollection::TransformGroup },
+		{ FGeometryDynamicCollection::ActiveAttribute, FTransformCollection::TransformGroup },
+	};
+	static const int32 SkipListSize = sizeof(SkipList) / sizeof(FAttributeAndGroupId);
 
+	DynamicCollection.CopyMatchingAttributesFrom(RestCollection, MakeArrayView(SkipList, SkipListSize));
 
 	// User defined initial velocities need to be populated. 
 	if (Params.InitialVelocityType == EInitialVelocityTypeEnum::Chaos_Initial_Velocity_User_Defined)
diff --git a/Engine/Source/Runtime/Experimental/Chaos/Public/GeometryCollection/ManagedArrayCollection.h b/Engine/Source/Runtime/Experimental/Chaos/Public/GeometryCollection/ManagedArrayCollection.h
--- a/Engine/Source/Runtime/Experimental/Chaos/Public/GeometryCollection/ManagedArrayCollection.h
+++ b/Engine/Source/Runtime/Experimental/Chaos/Public/GeometryCollection/ManagedArrayCollection.h
@@ -18,6 +18,17 @@ namespace Chaos
 	class FChaosArchive;
 }
 
+CHAOS_API struct FAttributeAndGroupId
+{
+	FName AttributeName;
+	FName GroupName;
+
+	bool operator==(const FAttributeAndGroupId& Other) const
+	{
+		return ((AttributeName == Other.AttributeName) && (GroupName == Other.GroupName));
+	}
+};
+
 /**
 * ManagedArrayCollection
 *
@@ -484,6 +495,17 @@ public:
 	*/
 	CHAOS_API void CopyMatchingAttributesFrom(const FManagedArrayCollection& InCollection, const TMap<FName, TSet<FName>>* SkipList=nullptr);
 
+	/**
+	* Copy attributes that match the input collection. This is a utility to easily sync collections
+	* This version is recommend to be used as it is more performant overall
+	*	- it only resize the necessary groups
+	*	- it takes a array view for the skip list
+	*	- it has a more contained logic that reduces the number of lookups for attributes in the maps
+	* @param InCollection - All groups from this collection found in the input will be sized accordingly
+	* @param SkipList - Group/Attrs to skip. Keys are group names, values are attributes in those groups.
+	*/
+	CHAOS_API void CopyMatchingAttributesFrom(const FManagedArrayCollection& FromCollection, const TArrayView<const FAttributeAndGroupId> SkipList);
+
 	/**
 	* Number of elements in a group
 	* @return the group size, and Zero(0) if group does not exist.
