Make the triple buffer more consistent with the double buffer
The consumer will be reseted just before being swapped

diff --git a/Engine/Source/Runtime/Experimental/Chaos/Private/EventManager.cpp b/Engine/Source/Runtime/Experimental/Chaos/Private/EventManager.cpp
--- a/Engine/Source/Runtime/Experimental/Chaos/Private/EventManager.cpp
+++ b/Engine/Source/Runtime/Experimental/Chaos/Private/EventManager.cpp
@@ -62,7 +62,7 @@ namespace Chaos
 		{
 			if (EventContainer)
 			{
-				if (BufferMode == EMultiBufferMode::Double)
+				if (BufferMode != EMultiBufferMode::Single)
 				{
 					EventContainer->ResetConsumerBuffer();
 				}
