Fix compilation issue on some platforms ( related to to CL 27642035 )

diff --git a/Engine/Source/Runtime/Experimental/Chaos/Public/GeometryCollection/ManagedArrayCollection.h b/Engine/Source/Runtime/Experimental/Chaos/Public/GeometryCollection/ManagedArrayCollection.h
--- a/Engine/Source/Runtime/Experimental/Chaos/Public/GeometryCollection/ManagedArrayCollection.h
+++ b/Engine/Source/Runtime/Experimental/Chaos/Public/GeometryCollection/ManagedArrayCollection.h
@@ -18,7 +18,7 @@ namespace Chaos
 	class FChaosArchive;
 }
 
-CHAOS_API struct FAttributeAndGroupId
+struct FAttributeAndGroupId
 {
 	FName AttributeName;
 	FName GroupName;
