Compute component space transform optimization
Store breadth first order indices for transform on the rest collection

diff --git a/Engine/Source/Runtime/Experimental/Chaos/Private/GeometryCollection/Facades/CollectionHierarchyFacade.cpp b/Engine/Source/Runtime/Experimental/Chaos/Private/GeometryCollection/Facades/CollectionHierarchyFacade.cpp
--- a/Engine/Source/Runtime/Experimental/Chaos/Private/GeometryCollection/Facades/CollectionHierarchyFacade.cpp
+++ b/Engine/Source/Runtime/Experimental/Chaos/Private/GeometryCollection/Facades/CollectionHierarchyFacade.cpp
@@ -160,6 +160,11 @@ namespace Chaos::Facades
 	}
 
 	TArray<int32> FCollectionHierarchyFacade::GetTransformArrayInDepthFirstOrder() const
+	{
+		return ComputeTransformIndicesInDepthFirstOrder();
+	}
+
+	TArray<int32> FCollectionHierarchyFacade::ComputeTransformIndicesInDepthFirstOrder() const
 	{
 		TArray<int32> OrderedTransforms;
 		if (ParentAttribute.IsValid() && ChildrenAttribute.IsValid())
@@ -223,4 +228,31 @@ namespace Chaos::Facades
 		}
 		return OrderedTransforms;
 	}
+
+	TArray<int32> FCollectionHierarchyFacade::ComputeTransformIndicesInBreadthFirstOrder() const
+	{
+		TArray<int32> BreadthFirstIndices;
+		if (ParentAttribute.IsValid() && ChildrenAttribute.IsValid())
+		{
+			const TManagedArray<int32>& Parent = ParentAttribute.Get();
+			const TManagedArray<TSet<int32>>& Children = ChildrenAttribute.Get();
+
+			// first add all the roots
+			BreadthFirstIndices = GetRootIndices();
+			BreadthFirstIndices.Reserve(Parent.Num());
+
+			for (int Index = 0; Index < BreadthFirstIndices.Num(); Index++)
+			{
+				const int32 TransformIndex = BreadthFirstIndices[Index];
+				if (Children[TransformIndex].Num())
+				{
+					for (const int32 ChildTransformIndex : Children[TransformIndex])
+					{
+						BreadthFirstIndices.Add(ChildTransformIndex);
+					}
+				}
+			}
+		}
+		return BreadthFirstIndices;
+	}
 }
\ No newline at end of file
diff --git a/Engine/Source/Runtime/Experimental/Chaos/Public/GeometryCollection/Facades/CollectionHierarchyFacade.h b/Engine/Source/Runtime/Experimental/Chaos/Public/GeometryCollection/Facades/CollectionHierarchyFacade.h
--- a/Engine/Source/Runtime/Experimental/Chaos/Public/GeometryCollection/Facades/CollectionHierarchyFacade.h
+++ b/Engine/Source/Runtime/Experimental/Chaos/Public/GeometryCollection/Facades/CollectionHierarchyFacade.h
@@ -70,6 +70,27 @@ namespace Chaos::Facades
 		*/
 		CHAOS_API TArray<int32> GetTransformArrayInDepthFirstOrder() const;
 
+		/**
+		* compute transform indices in a depth first order 
+		*/
+		CHAOS_API TArray<int32> ComputeTransformIndicesInDepthFirstOrder() const;
+
+		/**
+		* compute transform indices in a breadth first order ( root to leaves )
+		*/
+		CHAOS_API TArray<int32> ComputeTransformIndicesInBreadthFirstOrder() const;
+
+		/**
+		* generate a non serialize array that stores the indices of the transforms in a breath first order
+		*/
+		CHAOS_API void GenerateBreadthFirstOrderIndicesAttribute();
+
+		/**
+		* get the breadth first ordered transform indices from the stored attribute if any
+		* may return empty array if the attribute does not exists
+		*/
+		CHAOS_API const TArray<int32>& GetBreadthFirstOrderIndicesFromAttribute() const;
+
 	public:
 		/** Get the root indicies */
 		static CHAOS_API TArray<int32> GetRootIndices(const TManagedArrayAccessor<int32>& ParentAttribute);
diff --git a/Engine/Source/Runtime/Experimental/GeometryCollectionEngine/Private/GeometryCollection/GeometryCollectionComponent.cpp b/Engine/Source/Runtime/Experimental/GeometryCollectionEngine/Private/GeometryCollection/GeometryCollectionComponent.cpp
--- a/Engine/Source/Runtime/Experimental/GeometryCollectionEngine/Private/GeometryCollection/GeometryCollectionComponent.cpp
+++ b/Engine/Source/Runtime/Experimental/GeometryCollectionEngine/Private/GeometryCollection/GeometryCollectionComponent.cpp
@@ -4658,24 +4658,55 @@ static void CheckForNaNs(const TArray<FTransform>& Transforms)
 }
 #endif
 
+const TArray<FTransform>& UGeometryCollectionComponent::GetCurrentTransforms() const
+{
+	const bool bRestTransformsOverriden = (RestTransforms.Num() > 0);
+	if (!DynamicCollection && bRestTransformsOverriden)
+	{
+		return RestTransforms;
+	}
+	return GetTransformArray().GetConstArray();
+}
+
 void UGeometryCollectionComponent::CalculateGlobalMatrices()
 {
 	SCOPE_CYCLE_COUNTER(STAT_GCCUGlobalMatrices);
 
 	// If hierarchy topology has changed, the RestTransforms is invalidated.
+	// todo(chaos) should should be able to remove this : this should be normally be properly handled by the SetRestCollection function now
 	if (RestTransforms.Num() != GetTransformArray().Num())
 	{
 		RestTransforms.Empty();
 	}
 
-	const bool bRestTransformsOverriden = (RestTransforms.Num() > 0);
-	if (!DynamicCollection && bRestTransformsOverriden)
+	const TArray<FTransform>& CurrentTransforms = GetCurrentTransforms();
+
+	bool bFastPath = false;
+	if (RestCollection)
 	{
-		GeometryCollectionAlgo::GlobalMatrices(RestTransforms, GetParentArray(), ComponentSpaceTransforms);
+		const TArray<int32>& BreadthFirstTransformIndices = RestCollection->GetBreadthFirstTransformIndices();
+		bFastPath = (BreadthFirstTransformIndices.Num() == CurrentTransforms.Num());
+	}
+
+	if (bFastPath)
+	{
+		ComponentSpaceTransforms.SetNumUninitialized(CurrentTransforms.Num(), false);
+
+		const TArray<int32>& BreadthFirstTransformIndices = RestCollection->GetBreadthFirstTransformIndices();
+		const TArray<int32>& ParentArray = GetParentArray().GetConstArray();
+
+		for (int32 Index = 0; Index < BreadthFirstTransformIndices.Num(); Index++)
+		{
+			const int32 TransformIndex = BreadthFirstTransformIndices[Index];
+			const int32 ParentTransformIndex = ParentArray[TransformIndex];
+
+			const FTransform& ParentTransform = (ParentTransformIndex != INDEX_NONE) ? CurrentTransforms[ParentTransformIndex] : FTransform::Identity;
+			ComponentSpaceTransforms[TransformIndex] = CurrentTransforms[TransformIndex] * ParentTransform;
+		}
 	}
 	else
 	{
-		GeometryCollectionAlgo::GlobalMatrices(GetTransformArray(), GetParentArray(), ComponentSpaceTransforms);
+		GeometryCollectionAlgo::GlobalMatrices(GetCurrentTransforms(), GetParentArray(), ComponentSpaceTransforms);
 	}
 
 #if WITH_EDITOR
diff --git a/Engine/Source/Runtime/Experimental/GeometryCollectionEngine/Private/GeometryCollection/GeometryCollectionObject.cpp b/Engine/Source/Runtime/Experimental/GeometryCollectionEngine/Private/GeometryCollection/GeometryCollectionObject.cpp
--- a/Engine/Source/Runtime/Experimental/GeometryCollectionEngine/Private/GeometryCollection/GeometryCollectionObject.cpp
+++ b/Engine/Source/Runtime/Experimental/GeometryCollectionEngine/Private/GeometryCollection/GeometryCollectionObject.cpp
@@ -319,6 +319,16 @@ void UGeometryCollection::UpdateRootIndex()
 	}
 }
 
+void UGeometryCollection::CacheBreadthFirstTransformIndices()
+{
+	BreadthFirstTransformIndices.Reset();
+	if (GeometryCollection)
+	{
+		Chaos::Facades::FCollectionHierarchyFacade HierarchyFacade(*GeometryCollection);
+		BreadthFirstTransformIndices = HierarchyFacade.ComputeTransformIndicesInBreadthFirstOrder();
+	}
+}
+
 void UGeometryCollection::UpdateGeometryDependentProperties()
 {
 #if WITH_EDITOR
@@ -1141,6 +1151,9 @@ void UGeometryCollection::Serialize(FArchive& Ar)
 		UpdateRootIndex();
 	}
 
+	// Generate root to leave order lookup
+	CacheBreadthFirstTransformIndices();
+
 	if (Ar.IsLoading())
 	{
 		FillAutoInstanceMeshesInstancesIfNeeded();
@@ -1507,6 +1520,7 @@ void UGeometryCollection::InvalidateCollection()
 {
 	StateGuid = FGuid::NewGuid();
 	UpdateRootIndex();
+	CacheBreadthFirstTransformIndices();
 }
 
 #if WITH_EDITOR
diff --git a/Engine/Source/Runtime/Experimental/GeometryCollectionEngine/Public/GeometryCollection/GeometryCollectionComponent.h b/Engine/Source/Runtime/Experimental/GeometryCollectionEngine/Public/GeometryCollection/GeometryCollectionComponent.h
--- a/Engine/Source/Runtime/Experimental/GeometryCollectionEngine/Public/GeometryCollection/GeometryCollectionComponent.h
+++ b/Engine/Source/Runtime/Experimental/GeometryCollectionEngine/Public/GeometryCollection/GeometryCollectionComponent.h
@@ -1347,6 +1347,10 @@ protected:
 private:
 	GEOMETRYCOLLECTIONENGINE_API void ProcessRepDataOnPT();
 
+	// return the most actual transforms
+	// this can be the rest collection ones, the overriden RestTransforms or the dynamic collection ones
+	const TArray<FTransform>& GetCurrentTransforms() const;
+
 	bool bRenderStateDirty;
 	bool bEnableBoneSelection;
 	int ViewLevel;
diff --git a/Engine/Source/Runtime/Experimental/GeometryCollectionEngine/Public/GeometryCollection/GeometryCollectionObject.h b/Engine/Source/Runtime/Experimental/GeometryCollectionEngine/Public/GeometryCollection/GeometryCollectionObject.h
--- a/Engine/Source/Runtime/Experimental/GeometryCollectionEngine/Public/GeometryCollection/GeometryCollectionObject.h
+++ b/Engine/Source/Runtime/Experimental/GeometryCollectionEngine/Public/GeometryCollection/GeometryCollectionObject.h
@@ -812,6 +812,8 @@ public:
 	UPROPERTY(EditAnywhere, BlueprintReadWrite, Category = "Dataflow", DisplayName = "DataFlow Overrides")
 	TMap<FString, FString> Overrides;
 
+	GEOMETRYCOLLECTIONENGINE_API const TArray<int32>& GetBreadthFirstTransformIndices() const { return BreadthFirstTransformIndices; }
+
 private:
 #if WITH_EDITOR
 	GEOMETRYCOLLECTIONENGINE_API void CreateSimulationDataImp(bool bCopyFromDDC);
@@ -830,6 +832,8 @@ private:
 	// fill instanced mesh instance count from geometry collection data if not done yet 
 	GEOMETRYCOLLECTIONENGINE_API void FillAutoInstanceMeshesInstancesIfNeeded();
 
+	void CacheBreadthFirstTransformIndices();
+
 private:
 	/** Guid created on construction of this collection. It should be used to uniquely identify this collection */
 	UPROPERTY()
@@ -859,6 +863,10 @@ private:
 	UPROPERTY(VisibleAnywhere, Category = "Clustering")
 	int32 RootIndex = INDEX_NONE;
 
+	// cache transform indices in breadth-first order
+	UPROPERTY(VisibleAnywhere, Transient, Category = "Clustering")
+	TArray<int32> BreadthFirstTransformIndices;
+
 	// #todo(dmp): rename to be consistent BoneSelectedMaterialID?
 	// Legacy index of the bone selected material in the object's Materials array, or INDEX_NONE if it is not stored there.
 	// Note for new objects the bone selected material should not be stored in the Materials array, so this should be INDEX_NONE
