Geometry collection : fix component space transform computation using the wrong parent transform

diff --git a/Engine/Source/Runtime/Experimental/GeometryCollectionEngine/Private/GeometryCollection/GeometryCollectionComponent.cpp b/Engine/Source/Runtime/Experimental/GeometryCollectionEngine/Private/GeometryCollection/GeometryCollectionComponent.cpp
--- a/Engine/Source/Runtime/Experimental/GeometryCollectionEngine/Private/GeometryCollection/GeometryCollectionComponent.cpp
+++ b/Engine/Source/Runtime/Experimental/GeometryCollectionEngine/Private/GeometryCollection/GeometryCollectionComponent.cpp
@@ -4700,7 +4700,7 @@ void UGeometryCollectionComponent::CalculateGlobalMatrices()
 			const int32 TransformIndex = BreadthFirstTransformIndices[Index];
 			const int32 ParentTransformIndex = ParentArray[TransformIndex];
 
-			const FTransform& ParentTransform = (ParentTransformIndex != INDEX_NONE) ? CurrentTransforms[ParentTransformIndex] : FTransform::Identity;
+			const FTransform& ParentTransform = (ParentTransformIndex != INDEX_NONE) ? ComponentSpaceTransforms[ParentTransformIndex] : FTransform::Identity;
 			ComponentSpaceTransforms[TransformIndex] = CurrentTransforms[TransformIndex] * ParentTransform;
 		}
 	}
