Micro-optimization to physics initialization in SetMaterials; found a few unnecessary allocations in CPU samples so move the array to cut them in half.

diff --git a/Engine/Source/Runtime/Engine/Private/PhysicsEngine/Experimental/PhysInterface_Chaos.cpp b/Engine/Source/Runtime/Engine/Private/PhysicsEngine/Experimental/PhysInterface_Chaos.cpp
--- a/Engine/Source/Runtime/Engine/Private/PhysicsEngine/Experimental/PhysInterface_Chaos.cpp
+++ b/Engine/Source/Runtime/Engine/Private/PhysicsEngine/Experimental/PhysInterface_Chaos.cpp
@@ -816,9 +816,9 @@ void FPhysInterface_Chaos::SetMaterials(const FPhysicsShapeHandle& InShape, cons
 			}
 		}
 
-		InShape.Shape->SetMaterialMasks(NewMaterialMaskHandles);
-		InShape.Shape->SetMaterialMaskMaps(NewMaterialMaskMaps);
-		InShape.Shape->SetMaterialMaskMapMaterials(NewMaterialMaskMaterialHandles);
+		InShape.Shape->SetMaterialMasks(MoveTemp(NewMaterialMaskHandles));
+		InShape.Shape->SetMaterialMaskMaps(MoveTemp(NewMaterialMaskMaps));
+		InShape.Shape->SetMaterialMaskMapMaterials(MoveTemp(NewMaterialMaskMaterialHandles));
 	}
 }
 
diff --git a/Engine/Source/Runtime/Experimental/Chaos/Public/Chaos/ShapeInstance.h b/Engine/Source/Runtime/Experimental/Chaos/Public/Chaos/ShapeInstance.h
--- a/Engine/Source/Runtime/Experimental/Chaos/Public/Chaos/ShapeInstance.h
+++ b/Engine/Source/Runtime/Experimental/Chaos/Public/Chaos/ShapeInstance.h
@@ -120,12 +120,15 @@ namespace Chaos
 
 		CHAOS_API const TArray<FMaterialMaskHandle>& GetMaterialMasks() const;
 		CHAOS_API void SetMaterialMasks(const TArray<FMaterialMaskHandle>& InMaterialMasks);
+		CHAOS_API void SetMaterialMasks(TArray<FMaterialMaskHandle>&& InMaterialMasks);
 
 		CHAOS_API const TArray<uint32>& GetMaterialMaskMaps() const;
 		CHAOS_API void SetMaterialMaskMaps(const TArray<uint32>& InMaterialMaskMaps);
+		CHAOS_API void SetMaterialMaskMaps(TArray<uint32>&& InMaterialMaskMaps);
 
 		CHAOS_API const TArray<FMaterialHandle>& GetMaterialMaskMapMaterials() const;
 		CHAOS_API void SetMaterialMaskMapMaterials(const TArray<FMaterialHandle>& InMaterialMaskMapMaterials);
+		CHAOS_API void SetMaterialMaskMapMaterials(TArray<FMaterialHandle>&& InMaterialMaskMapMaterials);
 
 		CHAOS_API const FShapeDirtyFlags GetDirtyFlags() const;
 
@@ -1061,16 +1064,31 @@ namespace Chaos
 		DownCast([&InMaterialMasks](auto& ShapeInstance) { ShapeInstance.SetMaterialMasks(InMaterialMasks); });
 	}
 
+	inline void FPerShapeData::SetMaterialMasks(TArray<FMaterialMaskHandle>&& InMaterialMasks)
+	{
+		DownCast([&InMaterialMasks](auto& ShapeInstance) { ShapeInstance.SetMaterialMasks(MoveTemp(InMaterialMasks)); });
+	}
+
 	inline void FPerShapeData::SetMaterialMaskMaps(const TArray<uint32>& InMaterialMaskMaps)
 	{
 		DownCast([&InMaterialMaskMaps](auto& ShapeInstance) { ShapeInstance.SetMaterialMaskMaps(InMaterialMaskMaps); });
 	}
 
+	inline void FPerShapeData::SetMaterialMaskMaps(TArray<uint32>&& InMaterialMaskMaps)
+	{
+		DownCast([&InMaterialMaskMaps](auto& ShapeInstance) { ShapeInstance.SetMaterialMaskMaps(MoveTemp(InMaterialMaskMaps)); });
+	}
+
 	inline void FPerShapeData::SetMaterialMaskMapMaterials(const TArray<FMaterialHandle>& InMaterialMaskMapMaterials)
 	{
 		DownCast([&InMaterialMaskMapMaterials](auto& ShapeInstance) { ShapeInstance.SetMaterialMaskMapMaterials(InMaterialMaskMapMaterials); });
 	}
 
+	inline void FPerShapeData::SetMaterialMaskMapMaterials(TArray<FMaterialHandle>&& InMaterialMaskMapMaterials)
+	{
+		DownCast([&InMaterialMaskMapMaterials](auto& ShapeInstance) { ShapeInstance.SetMaterialMaskMapMaterials(MoveTemp(InMaterialMaskMapMaterials)); });
+	}
+
 	inline bool FPerShapeData::GetQueryEnabled() const
 	{
 		return DownCast([](auto& ShapeInstance) { return ShapeInstance.GetQueryEnabled(); });
