Iris - Respect poll frequency overrides.

diff --git a/Engine/Source/Runtime/Experimental/Iris/Core/Private/Iris/ReplicationSystem/ObjectReplicationBridge.cpp b/Engine/Source/Runtime/Experimental/Iris/Core/Private/Iris/ReplicationSystem/ObjectReplicationBridge.cpp
--- a/Engine/Source/Runtime/Experimental/Iris/Core/Private/Iris/ReplicationSystem/ObjectReplicationBridge.cpp
+++ b/Engine/Source/Runtime/Experimental/Iris/Core/Private/Iris/ReplicationSystem/ObjectReplicationBridge.cpp
@@ -315,7 +315,7 @@ UE::Net::FNetRefHandle UObjectReplicationBridge::BeginReplication(UObject* Insta
 			float PollFrequency = Params.PollFrequency;
 			FindOrCachePollFrequency(Instance->GetClass(), PollFrequency);
 			
-			uint8 PollFramePeriod = ConvertPollFrequencyIntoFrames(Params.PollFrequency);
+			uint8 PollFramePeriod = ConvertPollFrequencyIntoFrames(PollFrequency);
 			PollFrequencyLimiter->SetPollFramePeriod(InternalReplicationIndex, PollFramePeriod);
 
 			UE_LOG_OBJECTREPLICATIONBRIDGE(Verbose, TEXT("BeginReplication Created %s with ProtocolId:0x%" UINT64_x_FMT " for Object named %s"), *RefHandle.ToString(), ReplicationProtocol->ProtocolIdentifier, *Instance->GetName());
