Chaos - fix mesh collision face index

diff --git a/Engine/Source/Runtime/Experimental/Chaos/Private/Chaos/Collision/MeshContactGenerator.cpp b/Engine/Source/Runtime/Experimental/Chaos/Private/Chaos/Collision/MeshContactGenerator.cpp
--- a/Engine/Source/Runtime/Experimental/Chaos/Private/Chaos/Collision/MeshContactGenerator.cpp
+++ b/Engine/Source/Runtime/Experimental/Chaos/Private/Chaos/Collision/MeshContactGenerator.cpp
@@ -117,8 +117,11 @@ namespace Chaos::Private
 				VertexContactIndicesMap.Emplace(VertexID0, VertexID0, NewContactIndex, bIsFaceContact);
 			}
 
-			// Store and enable the contact
+			// Store the contact, set the face index
 			Contacts.Add(ContactPoint);
+			Contacts[NewContactIndex].FaceIndex = Triangle.TriangleIndex;
+
+			// Set the contact metadata and enable it
 			ContactDatas.AddDefaulted();
 			ContactDatas[NewContactIndex].SetEdgeOrVertexID({ VertexID0, VertexID1 });
 			ContactDatas[NewContactIndex].SetContactNormalDotTriangleNormal(ContactNormalDotTriangleNormal);
