Fix CIS - post-merge compile error in MeshContactGenerator.cpp

diff --git a/Engine/Source/Runtime/Experimental/Chaos/Private/Chaos/Collision/MeshContactGenerator.cpp b/Engine/Source/Runtime/Experimental/Chaos/Private/Chaos/Collision/MeshContactGenerator.cpp
--- a/Engine/Source/Runtime/Experimental/Chaos/Private/Chaos/Collision/MeshContactGenerator.cpp
+++ b/Engine/Source/Runtime/Experimental/Chaos/Private/Chaos/Collision/MeshContactGenerator.cpp
@@ -119,7 +119,7 @@ namespace Chaos::Private
 
 			// Store the contact, set the face index
 			Contacts.Add(ContactPoint);
-			Contacts[NewContactIndex].FaceIndex = Triangle.TriangleIndex;
+			Contacts[NewContactIndex].FaceIndex = Triangle.GetTriangleIndex();
 
 			// Set the contact metadata and enable it
 			ContactDatas.AddDefaulted();
