Chaos - remove normal impulse averaging in friction calculation

This was causing problems with large manifolds when a point has its normal impulse go to zero and another point take over after it has calculated a friction impulse. We can end up with multiple contacts applying friction impulses that are above the allowed limit, leading to a glitch in sliding movement.

We could probaly just detect the to-zero case, and explicitly undo friction at that point, but the averaging doesn't seem to be necessary anyway and it adds to the cost.

diff --git a/Engine/Source/Runtime/Experimental/Chaos/Public/Chaos/Collision/PBDCollisionSolver.h b/Engine/Source/Runtime/Experimental/Chaos/Public/Chaos/Collision/PBDCollisionSolver.h
--- a/Engine/Source/Runtime/Experimental/Chaos/Public/Chaos/Collision/PBDCollisionSolver.h
+++ b/Engine/Source/Runtime/Experimental/Chaos/Public/Chaos/Collision/PBDCollisionSolver.h
@@ -958,10 +958,7 @@ namespace Chaos
 			const FSolverReal Dt,
 			const FSolverReal MaxPushOut)
 		{
-			// Accumulate net pushout for friction limits below
 			bool bApplyFriction = false;
-			FSolverReal NumFrictionContacts = FSolverReal(0);	// NOTE: deliberately not an int
-			FSolverReal TotalPushOutNormal = FSolverReal(0);
 
 			// Apply the position correction along the normal and determine if we want to run friction on each point
 			for (int32 PointIndex = 0; PointIndex < NumManifoldPoints(); ++PointIndex)
@@ -982,28 +979,15 @@ namespace Chaos
 				}
 
 				// Friction gets updated for any point with a net normal correction or where we have previously had a normal correction and 
-				// already applied friction (in which case we may need to zero it)
-				SolverManifoldPoint.bApplyFriction = false;
-				if (SolverManifoldPoint.NetPushOutNormal > 0)
-				{
-					TotalPushOutNormal += SolverManifoldPoint.NetPushOutNormal;
-					NumFrictionContacts += FSolverReal(1);
-				}
-				if ((SolverManifoldPoint.NetPushOutNormal > 0) || (SolverManifoldPoint.NetPushOutTangentU != 0) || (SolverManifoldPoint.NetPushOutTangentV != 0))
-				{
-					bApplyFriction = true;
-					SolverManifoldPoint.bApplyFriction = ((SolverManifoldPoint.NetPushOutNormal > 0) || (SolverManifoldPoint.NetPushOutTangentU != 0) || (SolverManifoldPoint.NetPushOutTangentV != 0));
-				}
+				// already applied friction (in which case we may need to zero it if the normal correction gets removed)
+				const bool bHasAnyPushOut = (SolverManifoldPoint.NetPushOutNormal > 0) || (SolverManifoldPoint.NetPushOutTangentU != 0) || (SolverManifoldPoint.NetPushOutTangentV != 0);
+				SolverManifoldPoint.bApplyFriction = bHasAnyPushOut;
+				bApplyFriction = bApplyFriction || bHasAnyPushOut;
 			}
 
 			// Apply the tangential position correction if required
 			if (bApplyFriction)
 			{
-				// We clip the tangential correction at each contact to the friction cone, but we use to average impulse
-				// among all contacts as the clipping limit. This is not really correct but it is much more stable to 
-				// differences in contacts from tick to tick
-				// @todo(chaos): try a decaying maximum per contact point rather than an average (again - we had that once!)
-				const FSolverReal FrictionMaxPushOut = (NumFrictionContacts > FSolverReal(0)) ? TotalPushOutNormal / NumFrictionContacts : FSolverReal(0);
 				const FSolverReal FrictionStiffness = State.Stiffness * CVars::Chaos_PBDCollisionSolver_Position_StaticFrictionStiffness;
 
 				for (int32 PointIndex = 0; PointIndex < NumManifoldPoints(); ++PointIndex)
@@ -1014,6 +998,8 @@ namespace Chaos
 						FSolverReal ContactDeltaTangentU, ContactDeltaTangentV;
 						CalculateContactPositionErrorTangential(SolverManifoldPoint, ContactDeltaTangentU, ContactDeltaTangentV);
 
+						const FSolverReal FrictionMaxPushOut = SolverManifoldPoint.NetPushOutNormal;
+
 						ApplyPositionCorrectionTangential(
 							FrictionStiffness,
 							State.StaticFriction,
