Geometry collection : do not pay the cost of transform multiplication if transform as no parent

diff --git a/Engine/Source/Runtime/Experimental/GeometryCollectionEngine/Private/GeometryCollection/GeometryCollectionComponent.cpp b/Engine/Source/Runtime/Experimental/GeometryCollectionEngine/Private/GeometryCollection/GeometryCollectionComponent.cpp
--- a/Engine/Source/Runtime/Experimental/GeometryCollectionEngine/Private/GeometryCollection/GeometryCollectionComponent.cpp
+++ b/Engine/Source/Runtime/Experimental/GeometryCollectionEngine/Private/GeometryCollection/GeometryCollectionComponent.cpp
@@ -4700,8 +4700,15 @@ void UGeometryCollectionComponent::CalculateGlobalMatrices()
 			const int32 TransformIndex = BreadthFirstTransformIndices[Index];
 			const int32 ParentTransformIndex = ParentArray[TransformIndex];
 
-			const FTransform& ParentTransform = (ParentTransformIndex != INDEX_NONE) ? ComponentSpaceTransforms[ParentTransformIndex] : FTransform::Identity;
-			ComponentSpaceTransforms[TransformIndex] = CurrentTransforms[TransformIndex] * ParentTransform;
+			if (ParentTransformIndex == INDEX_NONE)
+			{
+				ComponentSpaceTransforms[TransformIndex] = CurrentTransforms[TransformIndex];
+			}
+			else
+			{
+				const FTransform& ParentTransform = ComponentSpaceTransforms[ParentTransformIndex];
+				ComponentSpaceTransforms[TransformIndex] = CurrentTransforms[TransformIndex] * ParentTransform;
+			}
 		}
 	}
 	else
