Chaos - fix occasional AABBTree crash in dev builds when editing GCs

diff --git a/Engine/Source/Runtime/Experimental/Chaos/Public/Chaos/AABBTree.h b/Engine/Source/Runtime/Experimental/Chaos/Public/Chaos/AABBTree.h
--- a/Engine/Source/Runtime/Experimental/Chaos/Public/Chaos/AABBTree.h
+++ b/Engine/Source/Runtime/Experimental/Chaos/Public/Chaos/AABBTree.h
@@ -1836,11 +1836,14 @@ public:
 #if CHAOS_DEBUG_NAME			
 			if constexpr (std::is_same_v<TPayloadType, FAccelerationStructureHandle>)
 			{
-				const TSharedPtr<FString, ESPMode::ThreadSafe>& DebugName = Payload.GetGeometryParticleHandle_PhysicsThread()->DebugName();
-				if (IsInPhysicsThreadContext())
+				if (const FGeometryParticleHandle* Particle = Payload.GetGeometryParticleHandle_PhysicsThread())
 				{
-					FString DebugStr = DebugName ? *DebugName : TEXT("No Name");
-					UE_LOG(LogChaos, Warning, TEXT("AABBTree encountered invalid bounds input : %s"), *DebugStr);
+					const TSharedPtr<FString, ESPMode::ThreadSafe>& DebugName = Particle->DebugName();
+					if (IsInPhysicsThreadContext())
+					{
+						FString DebugStr = DebugName ? *DebugName : TEXT("No Name");
+						UE_LOG(LogChaos, Warning, TEXT("AABBTree encountered invalid bounds input : %s"), *DebugStr);
+					}
 				}
 			}
 #endif
