Changed pool's host name view to a pointer to avoid confusion when used with external APIs

#rnx

diff --git a/Engine/Source/Runtime/Experimental/IoStoreOnDemand/Private/CoreHttp/Http.cpp b/Engine/Source/Runtime/Experimental/IoStoreOnDemand/Private/CoreHttp/Http.cpp
--- a/Engine/Source/Runtime/Experimental/IoStoreOnDemand/Private/CoreHttp/Http.cpp
+++ b/Engine/Source/Runtime/Experimental/IoStoreOnDemand/Private/CoreHttp/Http.cpp
@@ -651,7 +651,7 @@ class FSocketPool
 public:
 	enum class EDirection : uint8 { Send, Recv };
 
-					FSocketPool(FAnsiStringView InHostName, uint32 InPort, uint32 InMaxLeases);
+					FSocketPool(const ANSICHAR* InHostName, uint32 InPort, uint32 InMaxLeases);
 					~FSocketPool();
 	static uint32	GetAllocSize(uint32 MaxLeases);
 	bool			LeaseSocket(SocketType& Out);
@@ -665,7 +665,7 @@ public:
 	uint32			GetPort() const			{ return Port; }
 
 private:
-	FAnsiStringView	HostName;
+	const ANSICHAR*	HostName;
 	uint32			IpAddresses[4] = {};
 	int16			SendBufKb = -1;
 	int16			RecvBufKb = -1;
@@ -676,7 +676,7 @@ private:
 };
 
 ////////////////////////////////////////////////////////////////////////////////
-FSocketPool::FSocketPool(FAnsiStringView InHostName, uint32 InPort, uint32 InMaxLeases)
+FSocketPool::FSocketPool(const ANSICHAR* InHostName, uint32 InPort, uint32 InMaxLeases)
 : HostName(InHostName)
 , Port(uint16(InPort))
 , MaxLeases(uint8(InMaxLeases))
@@ -768,13 +768,11 @@ int32 FSocketPool::ResolveHostName()
 	addrinfo* Info = nullptr;
 	ON_SCOPE_EXIT { if (Info != nullptr) freeaddrinfo(Info); };
 
-	const FAnsiStringView& Host = GetHostName();
-
 	addrinfo Hints = {};
 	Hints.ai_family = AF_INET;
 	Hints.ai_socktype = SOCK_STREAM;
 	Hints.ai_protocol = IPPROTO_TCP;
-	auto Result = getaddrinfo(Host.GetData(), nullptr, &Hints, &Info);
+	auto Result = getaddrinfo(HostName, nullptr, &Hints, &Info);
 	if (uint32(Result) || Info == nullptr)
 	{
 		IpAddresses[0] = 2;
@@ -861,7 +859,7 @@ FConnectionPool::FConnectionPool(const FParams& Params)
 
 	// Init internal object
 	new (Internal) FSocketPool(
-		FAnsiStringView(HostDest, HostNameLen),
+		HostDest,
 		Params.Host.Port,
 		Params.ConnectionCount
 	);
@@ -2520,9 +2518,7 @@ FRequest FEventLoop::Request(
 
 	memcpy(HostNamePtr, HostName.GetData(), HostNameLength);
 	HostNamePtr[HostNameLength] = '\0';
-	HostName = FAnsiStringView(HostNamePtr, HostNameLength);
-
-	new (Pool) FSocketPool(HostName, Port, 1);
+	new (Pool) FSocketPool(HostNamePtr, Port, 1);
 
 	return Impl->Request(Method, Path, Activity);
 }
