Geometry collection : remove parent space transform from Pt to GT marschalled data
- they are now recomputed on the gamethread from XR of the particles
- this removes the need to interpolate those matrices
- this saves about half of the marschalled data

diff --git a/Engine/Source/Runtime/Experimental/Chaos/Private/PhysicsProxy/GeometryCollectionPhysicsProxy.cpp b/Engine/Source/Runtime/Experimental/Chaos/Private/PhysicsProxy/GeometryCollectionPhysicsProxy.cpp
--- a/Engine/Source/Runtime/Experimental/Chaos/Private/PhysicsProxy/GeometryCollectionPhysicsProxy.cpp
+++ b/Engine/Source/Runtime/Experimental/Chaos/Private/PhysicsProxy/GeometryCollectionPhysicsProxy.cpp
@@ -121,6 +121,12 @@ FAutoConsoleVariableRef CVarGeometryCollectionLocalInertiaDropOffDiagonalTerms(
 	GeometryCollectionLocalInertiaDropOffDiagonalTerms,
 	TEXT("When true, force diagonal inertia for GCs in their local space by simply dropping off-diagonal terms"));
 
+float GeometryCollectionTransformTolerance = 0.001f;
+FAutoConsoleVariableRef CVarGeometryCollectionTransformTolerance(
+	TEXT("p.GeometryCollection.TransformTolerance"),
+	GeometryCollectionTransformTolerance,
+	TEXT("Tolerance to detect if a transform has changed"));
+
 DEFINE_LOG_CATEGORY_STATIC(UGCC_LOG, Error, All);
 
 static const FSharedSimulationSizeSpecificData& GetSizeSpecificData(const TArray<FSharedSimulationSizeSpecificData>& SizeSpecificData, const FGeometryCollection& RestCollection, const int32 TransformIndex, const FBox& BoundingBox);
@@ -143,7 +149,6 @@ void FGeometryCollectionResults::Reset()
 	States.Reset();
 	Positions.Reset();
 	Velocities.Reset();
-	Transforms.Reset();
 
 	IsObjectDynamic = false;
 	IsObjectLoading = false;
@@ -3245,7 +3250,6 @@ void FGeometryCollectionPhysicsProxy::BufferPhysicsResults_Internal(Chaos::FPBDR
 				Results.SetState(EntryIndex, StateData);
 				Results.SetPositions(EntryIndex, PositionData);
 				Results.SetVelocities(EntryIndex, VelocityData);
-				Results.SetTransform(EntryIndex, ParentSpaceTransform);
 				PhysicsThreadCollection.Transform[TransformGroupIndex] = ParentSpaceTransform;
 				IsObjectDynamic = true;
 			}
@@ -3350,14 +3354,12 @@ public:
 		int32 TransformIndexIn,
 		const FGeometryCollectionResults& ResultsIn, 
 		const FGeometryCollectionResults& NextResultsIn,
-		const TManagedArray<FTransform>& TransformsIn,
 		Chaos::FPBDRigidParticle& GTParticleIn,
 		Chaos::FRealSingle AlphaIn
 	)
 		: TransformIndex(TransformIndexIn)
 		, Results(&ResultsIn)
 		, NextResults(&NextResultsIn)
-		, Transforms(TransformsIn)
 		, GTParticle(GTParticleIn)
 		, Alpha(AlphaIn)
 	{
@@ -3422,36 +3424,12 @@ public:
 		}
 	}
 
-	void GetTransform(FTransform& TransformOut, const FTransform& MassToLocal) const
-	{
-		if (NextEntryIndex == INDEX_NONE)
-		{
-			// not present in next results just use results
-			TransformOut = Results->GetTransform(EntryIndex);
-		}
-		else if (EntryIndex == INDEX_NONE)
-		{
-			// not present in current results, we need to interpolate from the actual GTparticle
-			FTransform InterpolatedTransform;
-			InterpolatedTransform.Blend(MassToLocal * Transforms[TransformIndex], MassToLocal * NextResults->GetTransform(NextEntryIndex), Alpha);
-			TransformOut = MassToLocal.Inverse() * InterpolatedTransform;
-		}
-		else
-		{
-			// both available 
-			FTransform InterpolatedTransform;
-			InterpolatedTransform.Blend(MassToLocal * Results->GetTransform(EntryIndex), MassToLocal * NextResults->GetTransform(NextEntryIndex), Alpha);
-			TransformOut = MassToLocal.Inverse() * InterpolatedTransform;
-		}
-	}
-
 private:
 	const int32 TransformIndex;
 	const FGeometryCollectionResults* Results;
 	FGeometryCollectionResults::FEntryIndex EntryIndex;
 	const FGeometryCollectionResults* NextResults;
 	FGeometryCollectionResults::FEntryIndex NextEntryIndex;
-	const TManagedArray<FTransform>& Transforms;
 	Chaos::FPBDRigidParticle& GTParticle;
 	Chaos::FRealSingle Alpha;
 
@@ -3547,9 +3525,6 @@ bool FGeometryCollectionPhysicsProxy::PullNonInterpolatableDataFromSinglePhysics
 					bIsCollectionDirty |= UpdateValue((*LinearVelocities)[TransformGroupIndex], VelocityData.ParticleV);
 					bIsCollectionDirty |= UpdateValue((*AngularVelocities)[TransformGroupIndex], VelocityData.ParticleW);
 				}
-
-				const FTransform& NewTransform = CurrentResults.GetTransform(EntryIndex);
-				bIsCollectionDirty |= UpdateTransform(GameThreadCollection.Transform[TransformGroupIndex], NewTransform);
 			}
 		}
 
@@ -3606,6 +3581,8 @@ bool FGeometryCollectionPhysicsProxy::PullFromPhysicsState(const Chaos::FDirtyGe
 		}
 		bIsCollectionDirty |= PullNonInterpolatableDataFromSinglePhysicsState(PullData, !bNeedInterpolation, NextPullData ? &NextPullData->Results().GetModifiedTransformIndices() : nullptr);
 
+		const TManagedArray<bool>* AnimationsActive = GameThreadCollection.FindAttribute<bool>(AnimateTransformAttributeName, FGeometryCollection::TransformGroup);
+
 		// second : interpolate-able ones
 		if (bNeedInterpolation)
 		{
@@ -3614,7 +3591,6 @@ bool FGeometryCollectionPhysicsProxy::PullFromPhysicsState(const Chaos::FDirtyGe
 
 			TManagedArray<FVector3f>* LinearVelocities = GameThreadCollection.FindAttributeTyped<FVector3f>(LinearVelocityAttributeName, FTransformCollection::TransformGroup);
 			TManagedArray<FVector3f>* AngularVelocities = GameThreadCollection.FindAttributeTyped<FVector3f>(AngularVelocitiesAttributeName, FTransformCollection::TransformGroup);
-			const TManagedArray<bool>* AnimationsActive = GameThreadCollection.FindAttribute<bool>(AnimateTransformAttributeName, FGeometryCollection::TransformGroup);
 
 			// for that case we cannot just go through the list of entries since Results and NextResults may have different number of entries that don't always match
 			// so we need to go through the transform indices and find the matching entries on both side 
@@ -3627,7 +3603,7 @@ bool FGeometryCollectionPhysicsProxy::PullFromPhysicsState(const Chaos::FDirtyGe
 				}
 
 				FParticle& GTParticle = *GTParticles[TransformGroupIndex];
-				const FResultInterpolator ResultInterpolator(TransformGroupIndex, PrevResults, NextResults, GameThreadCollection.Transform, GTParticle, *Alpha);
+				const FResultInterpolator ResultInterpolator(TransformGroupIndex, PrevResults, NextResults, GTParticle, *Alpha);
 				if (ResultInterpolator.HasNoEntry())
 				{
 					continue;
@@ -3652,15 +3628,61 @@ bool FGeometryCollectionPhysicsProxy::PullFromPhysicsState(const Chaos::FDirtyGe
 						bIsCollectionDirty |= UpdateValue((*LinearVelocities)[TransformGroupIndex], NewV);
 						bIsCollectionDirty |= UpdateValue((*AngularVelocities)[TransformGroupIndex], NewW);
 					}
-
-					FTransform NewTransform{ FTransform::Identity };
-					const FTransform& MassToLocal = GameThreadCollection.MassToLocal[TransformGroupIndex];
-					ResultInterpolator.GetTransform(NewTransform, MassToLocal);
-					bIsCollectionDirty |= UpdateTransform(GameThreadCollection.Transform[TransformGroupIndex], NewTransform);
 				}
 			}
-		} 
+		}
+
+		// update the parent transform now that we have all X and R computed
+		bool bHasDifferentTransforms = false;
+		if (bIsCollectionDirty)
+		{
+			const TBitArray<>& PrevResultsModifiedIndices = PullData.Results().GetModifiedTransformIndices();
+			const FGeometryCollectionResults& NextResults = NextPullData->Results();
 
+			for (int32 TransformIndex = 0; TransformIndex < NumTransforms; TransformIndex++)
+			{
+				const bool bAnimatingWhileDisabled = AnimationsActive ? (*AnimationsActive)[TransformIndex] : false;
+				const bool bActive = GameThreadCollection.Active[TransformIndex];
+				if (bActive || bAnimatingWhileDisabled)
+				{
+					if (GTParticles[TransformIndex] != nullptr)
+					{
+						bool bWasModified = PrevResultsModifiedIndices.IsValidIndex(TransformIndex) ? PrevResultsModifiedIndices[TransformIndex] : false;
+						if (NextPullData)
+						{
+							const TBitArray<>& NextResultsModifiedIndices = NextPullData->Results().GetModifiedTransformIndices();
+							bWasModified |= NextResultsModifiedIndices.IsValidIndex(TransformIndex) ? NextResultsModifiedIndices[TransformIndex] : false;
+						}
+						if (bWasModified)
+						{
+							const FParticle& GTParticle = *GTParticles[TransformIndex];
+							const FTransform& MassToLocal = GameThreadCollection.MassToLocal[TransformIndex];
+							const int32 ParentTransformIndex = GameThreadCollection.Parent[TransformIndex];
+
+							const FTransform& WorldTransform = MassToLocal.Inverse() * FTransform { GTParticle.R(), GTParticle.X() };
+
+							// by default parent is the component's world transform 
+							FTransform ParentWorldTransform = Parameters.WorldTransform;
+							if (ParentTransformIndex != INDEX_NONE)
+							{
+								if (const FParticle* GTParentParticle = (ParentTransformIndex != INDEX_NONE) ? GTParticles[ParentTransformIndex].Get() : nullptr)
+								{
+									const FTransform& ParentMassToLocal = GameThreadCollection.MassToLocal[TransformIndex];
+									ParentWorldTransform = ParentMassToLocal.Inverse() * FTransform { GTParentParticle->R(), GTParentParticle->X() };
+								}
+							}
+
+							const FTransform NewTransform = WorldTransform.GetRelativeTransform(ParentWorldTransform);
+							if (!NewTransform.Equals(GameThreadCollection.Transform[TransformIndex], GeometryCollectionTransformTolerance))
+							{
+								bHasDifferentTransforms = true;
+								GameThreadCollection.Transform[TransformIndex] = WorldTransform.GetRelativeTransform(ParentWorldTransform);
+							}
+						}
+					}
+				}
+			}
+		}
 
 #if WITH_EDITORONLY_DATA
 		// Damage is collected in full every time so no need to go through PullNonInterpolatableDataFromSinglePhysicsState.
diff --git a/Engine/Source/Runtime/Experimental/Chaos/Public/GeometryCollectionProxyData.h b/Engine/Source/Runtime/Experimental/Chaos/Public/GeometryCollectionProxyData.h
--- a/Engine/Source/Runtime/Experimental/Chaos/Public/GeometryCollectionProxyData.h
+++ b/Engine/Source/Runtime/Experimental/Chaos/Public/GeometryCollectionProxyData.h
@@ -352,7 +352,7 @@ class FGeometryCollectionResults: public FRefCountedObject
 public:
 	FGeometryCollectionResults();
 
-	int32 GetNumEntries() const { return Transforms.Num(); }
+	int32 GetNumEntries() const { return States.Num(); }
 
 	void Reset();
 
@@ -442,11 +442,6 @@ public:
 		return Velocities[EntryIndex];
 	}
 
-	inline const FTransform& GetTransform(FEntryIndex EntryIndex) const
-	{
-		return Transforms[EntryIndex];
-	}
-
 	inline void SetSolverDt(const Chaos::FReal SolverDtIn)
 	{
 		SolverDt = SolverDtIn;
@@ -464,7 +459,6 @@ public:
 		ensure(GetEntryIndexByTransformIndex(TransformIndex) == EntryIndex);
 		Positions.AddDefaulted();
 		Velocities.AddDefaulted();
-		Transforms.AddDefaulted();
 		return EntryIndex;
 	}
 
@@ -478,11 +472,6 @@ public:
 		Velocities[EntryIndex] = VelocityData;
 	}
 
-	inline void SetTransform(FEntryIndex EntryIndex, const FTransform& Transform)
-	{
-		Transforms[EntryIndex] = Transform;
-	}
-
 	inline const TBitArray<>& GetModifiedTransformIndices() const
 	{
 		return ModifiedTransformIndices;
@@ -498,7 +487,6 @@ private:
 	TArray<FStateData> States;
 	TArray<FPositionData> Positions;
 	TArray<FVelocityData> Velocities;
-	TArray<FTransform> Transforms;
 
 #if WITH_EDITORONLY_DATA
 	// use to display impulse statistics in editor
