Chaos - new initial overlap depenetration system

Currently disabled with a cvar.

diff --git a/Engine/Source/Runtime/Experimental/Chaos/Private/Chaos/ChaosDebugDraw.cpp b/Engine/Source/Runtime/Experimental/Chaos/Private/Chaos/ChaosDebugDraw.cpp
--- a/Engine/Source/Runtime/Experimental/Chaos/Private/Chaos/ChaosDebugDraw.cpp
+++ b/Engine/Source/Runtime/Experimental/Chaos/Private/Chaos/ChaosDebugDraw.cpp
@@ -1111,6 +1111,12 @@ namespace Chaos
 					// Manifold point
 					FDebugDrawQueue::GetInstance().DrawDebugCircle(WorldPointLocation, 0.5f * Settings.DrawScale * Settings.ContactWidth, 12, DiscColor, false, Duration, uint8(Settings.DrawPriority), Settings.LineThickness, Axes.GetUnitAxis(EAxis::Y), Axes.GetUnitAxis(EAxis::Z), false);
 
+					if (ManifoldPoint.InitialPhi != 0)
+					{
+						FColor C3 = FColor::Silver;
+						FDebugDrawQueue::GetInstance().DrawDebugCircle(WorldPlaneLocation + ManifoldPoint.InitialPhi * WorldPlaneNormal, 0.25f * Settings.DrawScale * Settings.ContactWidth, 12, C3, false, Duration, uint8(Settings.DrawPriority), Settings.LineThickness, Axes.GetUnitAxis(EAxis::Y), Axes.GetUnitAxis(EAxis::Z), false);
+					}
+
 					// Whether restored
 					if (Settings.ContactInfoWidth > 0)
 					{
diff --git a/Engine/Source/Runtime/Experimental/Chaos/Private/Chaos/Collision/PBDCollisionConstraint.cpp b/Engine/Source/Runtime/Experimental/Chaos/Private/Chaos/Collision/PBDCollisionConstraint.cpp
--- a/Engine/Source/Runtime/Experimental/Chaos/Private/Chaos/Collision/PBDCollisionConstraint.cpp
+++ b/Engine/Source/Runtime/Experimental/Chaos/Private/Chaos/Collision/PBDCollisionConstraint.cpp
@@ -72,7 +72,7 @@ namespace Chaos
 	};
 
 	// @todo(chaos): put these tolerances on cvars
-	// @todo(chaos): tune the tolerances used in FPBDCollisionConstraint::UpdateAndTryRestoreManifold
+	// @todo(chaos): tune the tolerances used in FPBDCollisionConstraint::TryRestoreManifold
 	FCollisionTolerances Chaos_Manifold_Tolerances;
 
 	FString FPBDCollisionConstraint::ToString() const
@@ -200,6 +200,7 @@ namespace Chaos
 		, ShapesType(EContactShapesType::Unknown)
 		, SavedManifoldPoints()
 		, ManifoldPoints()
+		, MinInitialPhi(0)
 		, CCDTimeOfImpact(0)
 		, SolverBodies{ nullptr, nullptr }
 		, CCDEnablePenetration(0)
@@ -238,6 +239,7 @@ namespace Chaos
 		, ShapesType(EContactShapesType::Unknown)
 		, SavedManifoldPoints()
 		, ManifoldPoints()
+		, MinInitialPhi(0)
 		, CCDTimeOfImpact(0)
 		, SolverBodies{ nullptr, nullptr }
 		, CCDEnablePenetration(0)
@@ -426,10 +428,16 @@ namespace Chaos
 		Flags.bIsCurrent = true;
 		Flags.bDisabled = false;
 
+		// If we have data from the previous tick, update the maximum allowed initial penetration
+		CalculateMinInitialPhi();
+
+		// Match new manifold points with data from previous tick
 		AssignSavedManifoldPoints();
 
+		// If any materials were changed, update the friction etc
 		UpdateMaterialProperties();
 
+		// Undo any temporary modifications from user callbacks
 		ResetModifications();
 
 		AccumulatedImpulse = FVec3(0);
@@ -760,19 +768,9 @@ namespace Chaos
 					ManifoldPoint.ContactPoint.ShapeContactPoints[1] = ShapeContactPoint1;
 					ManifoldPoint.ContactPoint.Phi = FRealSingle(ContactPhi);
 
-					// Set the friction anchors.
-					// In principle we should always have the same number of saved manifold points as manifold points 
-					// from the previous frame, but there are ways to break this. E.g., Contact modification resets
-					// friction if positions are changed.
-					if (ManifoldPointIndex < SavedManifoldPoints.Num())
-					{
-						FSavedManifoldPoint& SavedManifoldPoint = SavedManifoldPoints[ManifoldPointIndex];
-						ManifoldPoint.ShapeAnchorPoints[0] = SavedManifoldPoint.ShapeContactPoints[0];
-						ManifoldPoint.ShapeAnchorPoints[1] = SavedManifoldPoint.ShapeContactPoints[1];
-					}
-
 					ManifoldPoint.Flags.bWasRestored = true;
 					ManifoldPoint.Flags.bWasReplaced = false;
+					ManifoldPoint.Flags.bInitialContact = false;
 
 					if (ManifoldPoint.ContactPoint.Phi < GetPhi())
 					{
@@ -785,6 +783,7 @@ namespace Chaos
 					ManifoldPoint.Flags.bDisabled = true;
 					ManifoldPoint.Flags.bWasRestored = false;
 					ManifoldPoint.Flags.bWasReplaced = false;
+					ManifoldPoint.Flags.bInitialContact = false;
 					--NumActiveManifoldPoints;
 				}
 			}
@@ -1009,7 +1008,7 @@ namespace Chaos
 		return true;
 	}
 
-	FReal FPBDCollisionConstraint::CalculateSavedManifoldPointScore(const FSavedManifoldPoint& SavedManifoldPoint, const FManifoldPoint& ManifoldPoint, const FReal DistanceToleranceSq) const
+	FReal FPBDCollisionConstraint::CalculateSavedManifoldDistanceSq(const FSavedManifoldPoint& SavedManifoldPoint, const FManifoldPoint& ManifoldPoint, const FReal DistanceToleranceSq) const
 	{
 		// If we have a vertex-plane (or vertex-vertex) contact, we want to know if we have the same vertex(es).
 		// If we have and edge-edge contact, we want to know if we have the same edges.
@@ -1049,42 +1048,56 @@ namespace Chaos
 
 	int32 FPBDCollisionConstraint::FindSavedManifoldPoint(const int32 ManifoldPointIndex, int32* InOutAllowedSavedPointIndices, int32& InOutNumAllowedSavedPoints) const
 	{
-		int32 MatchIndex = INDEX_NONE;
-
 		if (bChaos_Manifold_EnableFrictionRestore)
 		{
 			const FManifoldPoint& ManifoldPoint = ManifoldPoints[ManifoldPointIndex];
 			if (!ManifoldPoint.Flags.bDisabled)
 			{
 				const FReal DistanceToleranceSq = FMath::Square(Chaos_Manifold_FrictionPositionTolerance);
-				FReal BestScore = DistanceToleranceSq;
+				const FReal ExactDistanceToleranceSq = 0.2f * DistanceToleranceSq;
+				FReal BestDistanceSq = DistanceToleranceSq;
 
+				int32 MatchAllowedPointIndex = INDEX_NONE;
 				for (int32 AllowedPointIndex = 0; AllowedPointIndex < InOutNumAllowedSavedPoints; ++AllowedPointIndex)
 				{
 					const int32 SavedPointIndex = InOutAllowedSavedPointIndices[AllowedPointIndex];
 					const FSavedManifoldPoint& SavedManifoldPoint = SavedManifoldPoints[SavedPointIndex];
 
-					const FReal Score = CalculateSavedManifoldPointScore(SavedManifoldPoint, ManifoldPoint, DistanceToleranceSq);
-					if (Score < BestScore)
+					const FReal DistanceSq = CalculateSavedManifoldDistanceSq(SavedManifoldPoint, ManifoldPoint, DistanceToleranceSq);
+
+					// If this is an exact match, take the point without searching further
+					if (DistanceSq < ExactDistanceToleranceSq)
 					{
-						BestScore = Score;
-						MatchIndex = SavedPointIndex;
+						BestDistanceSq = DistanceSq;
+						MatchAllowedPointIndex = AllowedPointIndex;
+						break;
+					}
 
-						// RemoveAtSwap
-						--InOutNumAllowedSavedPoints;
-						if (AllowedPointIndex < InOutNumAllowedSavedPoints)
-						{
-							InOutAllowedSavedPointIndices[AllowedPointIndex] = InOutAllowedSavedPointIndices[InOutNumAllowedSavedPoints];
-						}
+					// If this is a close match, we may use it if we can't find a better one
+					if (DistanceSq < BestDistanceSq)
+					{
+						BestDistanceSq = DistanceSq;
+						MatchAllowedPointIndex = AllowedPointIndex;
+					}
+				}
 
-						// Just take the first match we find that meets our tolerance
-						break;
+				// If we selected a point, remove it from the set
+				if (MatchAllowedPointIndex != INDEX_NONE)
+				{
+					// RemoveAtSwap
+					const int32 MatchIndex = InOutAllowedSavedPointIndices[MatchAllowedPointIndex];
+					--InOutNumAllowedSavedPoints;
+					if (MatchAllowedPointIndex < InOutNumAllowedSavedPoints)
+					{
+						InOutAllowedSavedPointIndices[MatchAllowedPointIndex] = InOutAllowedSavedPointIndices[InOutNumAllowedSavedPoints];
 					}
+
+					return MatchIndex;
 				}
 			}
 		}
 
-		return MatchIndex;
+		return INDEX_NONE;
 	}
 
 	void FPBDCollisionConstraint::AssignSavedManifoldPoints()
@@ -1099,16 +1112,37 @@ namespace Chaos
 		int32 NumAllowedPointIndices = AllowedPointIndices.Num();
 		for (int32 PointIndex = 0, PointEndIndex = ManifoldPoints.Num(); PointIndex < PointEndIndex; ++PointIndex)
 		{
-			const int32 SavedManifoldPointIndex = FindSavedManifoldPoint(PointIndex, AllowedPointIndices.GetData(), NumAllowedPointIndices);
-
-			if (SavedManifoldPointIndex != INDEX_NONE)
+			// @todo(chaos): ideally we would skip this step for restored manifold points but for now we need to 
+			// extract InitialPhi from the results
+			// @see UpdateCollisionSolverContactPointFromConstraint
+			//if (!ManifoldPoints[PointIndex].Flags.bWasRestored)
 			{
-				ManifoldPoints[PointIndex].Flags.bHasStaticFrictionAnchor = true;
-				ManifoldPoints[PointIndex].ShapeAnchorPoints[0] = SavedManifoldPoints[SavedManifoldPointIndex].ShapeContactPoints[0];
-				ManifoldPoints[PointIndex].ShapeAnchorPoints[1] = SavedManifoldPoints[SavedManifoldPointIndex].ShapeContactPoints[1];
+				// Assume this is a new contact
+				ManifoldPoints[PointIndex].Flags.bInitialContact = true;
+
+				// See if we have saved data for this contact
+				const int32 SavedManifoldPointIndex = FindSavedManifoldPoint(PointIndex, AllowedPointIndices.GetData(), NumAllowedPointIndices);
+				if (SavedManifoldPointIndex != INDEX_NONE)
+				{
+					ManifoldPoints[PointIndex].Flags.bHasStaticFrictionAnchor = true;
+					ManifoldPoints[PointIndex].ShapeAnchorPoints[0] = SavedManifoldPoints[SavedManifoldPointIndex].ShapeContactPoints[0];
+					ManifoldPoints[PointIndex].ShapeAnchorPoints[1] = SavedManifoldPoints[SavedManifoldPointIndex].ShapeContactPoints[1];
+					ManifoldPoints[PointIndex].InitialPhi = SavedManifoldPoints[SavedManifoldPointIndex].InitialPhi;
+					ManifoldPoints[PointIndex].Flags.bInitialContact = false;
+				}
+				// Nothing to do if no saved friction point because we already set the achor to the most recently detected contact point
+				// (See InitManifoldPoint, And TryRestoreManifold)
 			}
-			// Nothing to do if no saved friction point because we already set the achor to the most recently detected contact point
-			// (See InitManifoldPoint, And UpdateAndTryRestoreManifold)
+		}
+	}
+
+	void FPBDCollisionConstraint::CalculateMinInitialPhi()
+	{
+		MinInitialPhi = 0;
+
+		for (const FSavedManifoldPoint& SavedManifoldPoint : SavedManifoldPoints)
+		{
+			MinInitialPhi = FMath::Min(MinInitialPhi, SavedManifoldPoint.InitialPhi);
 		}
 	}
 
diff --git a/Engine/Source/Runtime/Experimental/Chaos/Private/Chaos/Collision/PBDCollisionContainerSolver.cpp b/Engine/Source/Runtime/Experimental/Chaos/Private/Chaos/Collision/PBDCollisionContainerSolver.cpp
--- a/Engine/Source/Runtime/Experimental/Chaos/Private/Chaos/Collision/PBDCollisionContainerSolver.cpp
+++ b/Engine/Source/Runtime/Experimental/Chaos/Private/Chaos/Collision/PBDCollisionContainerSolver.cpp
@@ -51,6 +51,10 @@ namespace Chaos
 		// @todo(chaos): to be tuned
 		FRealSingle Chaos_PBDCollisionSolver_JacobiRotationTolerance = 1.e-8f;
 		FAutoConsoleVariableRef CVarChaosPBDCollisionSolverJacobiRotationTolerance(TEXT("p.Chaos.PBDCollisionSolver.JacobiRotationTolerance"), Chaos_PBDCollisionSolver_JacobiRotationTolerance, TEXT(""));
+
+		// Whether to enable the new initial overlap depentration system
+		bool bChaos_Collision_EnableInitialDepenetration = false;
+		FAutoConsoleVariableRef CVarChaosCollisionEnableInitialDepentration(TEXT("p.Chaos.PBDCollisionSolver.EnableInitialDepenetration"), bChaos_Collision_EnableInitialDepenetration, TEXT(""));
 	}
 
 
@@ -63,13 +67,14 @@ namespace Chaos
 	void UpdateCollisionSolverContactPointFromConstraint(
 		Private::FPBDCollisionSolver& Solver, 
 		const int32 SolverPointIndex, 
-		const FPBDCollisionConstraint* Constraint, 
+		FPBDCollisionConstraint* Constraint, 
 		const int32 ConstraintPointIndex, 
 		const FRealSingle Dt, 
+		const FRealSingle MaxDepentrationVelocity,
 		const FConstraintSolverBody& Body0,
 		const FConstraintSolverBody& Body1)
 	{
-		const FManifoldPoint& ManifoldPoint = Constraint->GetManifoldPoint(ConstraintPointIndex);
+		FManifoldPoint& ManifoldPoint = Constraint->GetManifoldPoint(ConstraintPointIndex);
 
 		const FRealSingle Restitution = FRealSingle(Constraint->GetRestitution());
 		const FRealSingle RestitutionVelocityThreshold = FRealSingle(Constraint->GetRestitutionThreshold()) * Dt;
@@ -97,14 +102,17 @@ namespace Chaos
 		}
 		const FVec3f WorldContactTangentV = FVec3f::CrossProduct(WorldContactNormal, WorldContactTangentU);
 
-		// Calculate contact velocity if we will need it below (restitution and/or frist-contact friction)
-		const bool bNeedContactVelocity = (!ManifoldPoint.Flags.bHasStaticFrictionAnchor) || (Restitution > FRealSingle(0));
+
+		// Calculate contact velocity if we will need it below (restitution and/or first-contact for friction or initial depenetration)
+		const bool bNeedContactVelocity = (!ManifoldPoint.Flags.bHasStaticFrictionAnchor) || (Restitution > FRealSingle(0)) || ManifoldPoint.Flags.bInitialContact;
 		FVec3f ContactVel = FVec3(0);
+		FRealSingle ContactVelocityNormal = FRealSingle(0);
 		if (bNeedContactVelocity)
 		{
 			const FVec3f ContactVel0 = Body0.V() + FVec3f::CrossProduct(Body0.W(), WorldRelativeContact0);
 			const FVec3f ContactVel1 = Body1.V() + FVec3f::CrossProduct(Body1.W(), WorldRelativeContact1);
 			ContactVel = ContactVel0 - ContactVel1;
+			ContactVelocityNormal = FVec3f::DotProduct(ContactVel, WorldContactNormal);
 		}
 
 		// If we have contact data from a previous tick, use it to calculate the lateral position delta we need
@@ -125,9 +133,8 @@ namespace Chaos
 		}
 
 		// The contact point error we are trying to correct in this solver
-		const FRealSingle TargetPhi = ManifoldPoint.TargetPhi;
 		const FVec3f WorldContactDelta = FVec3f(WorldContact0 - WorldContact1);
-		const FRealSingle WorldContactDeltaNormal = FVec3f::DotProduct(WorldContactDelta, WorldContactNormal) - TargetPhi;
+		FRealSingle WorldContactDeltaNormal = FVec3f::DotProduct(WorldContactDelta, WorldContactNormal);
 		const FRealSingle WorldContactDeltaTangentU = FVec3f::DotProduct(WorldContactDelta + WorldFrictionDelta, WorldContactTangentU);
 		const FRealSingle WorldContactDeltaTangentV = FVec3f::DotProduct(WorldContactDelta + WorldFrictionDelta, WorldContactTangentV);
 
@@ -135,13 +142,61 @@ namespace Chaos
 		FRealSingle WorldContactTargetVelocityNormal = FRealSingle(0);
 		if (Restitution > FRealSingle(0))
 		{
-			const FRealSingle ContactVelocityNormal = FVec3f::DotProduct(ContactVel, WorldContactNormal);
 			if (ContactVelocityNormal < -RestitutionVelocityThreshold)
 			{
 				WorldContactTargetVelocityNormal = -Restitution * ContactVelocityNormal;
 			}
 		}
 
+		// Initial Phi for initial-overlap depenetration.
+		// If we have an initial contact, calculate the initial overlap. This will get saved in SetSolverResults
+		FRealSingle WorldContactInitialPhi = 0;
+		if (CVars::bChaos_Collision_EnableInitialDepenetration)
+		{
+			if (ManifoldPoint.Flags.bInitialContact)
+			{
+				// This is a new manifold point, capture current Phi as the initial Phi
+				WorldContactInitialPhi = FMath::Min(WorldContactDeltaNormal, FRealSingle(0));
+
+				// Reduce initial penetration depth by any movement towards each other this tick
+				WorldContactInitialPhi -= FMath::Min(ContactVelocityNormal * Dt, FRealSingle(0));
+
+				// If this is a new manifold point on a pre-existing manifold, we limit the initial depth.
+				// This is so that as we are depenetrating and new points are added to the manifold, we don't suddenly pop out, which would
+				// happen if we set InitialPhi=0. But also want to ensure that if we are nearly done handling initial overlap, we don't treat
+				// new deeper manifold points as full initial overlaps.
+				// NOTE: here we are checking IsInitialContact on the constraint, which is only set when we first make contact, as opposed 
+				// to the bInitialContact on the manifold point which is true for any new manifold point, regardless of the constraint age.
+				if (!Constraint->IsInitialContact())
+				{
+					WorldContactInitialPhi = FMath::Max(WorldContactInitialPhi, Constraint->GetMinInitialPhi());
+				}
+
+			}
+			else if (ManifoldPoint.InitialPhi < 0)
+			{
+				// This is a pre-existing manifold point with some initial penetration to resolve.
+				// If we are currently penetrating less than the inital overlap, reduce the initial overlap
+				// Also resolve initial overlap over time by reducing allowed penetration by MaxDepentrationVelocity
+				WorldContactInitialPhi = FMath::Max(ManifoldPoint.InitialPhi + MaxDepentrationVelocity * Dt, WorldContactDeltaNormal);
+			}
+
+			// InitialPhi is only for tracking penetration - cannot be positive
+			WorldContactInitialPhi = FMath::Min(WorldContactInitialPhi, FRealSingle(0));
+
+			// Apply initial penetration allowance to depth correction
+			WorldContactDeltaNormal -= WorldContactInitialPhi;
+
+			// @todo(chaos): InitialPhi should probably be updated prior to Gather, but this is where we calculate
+			// the depth and contact velocity so it's convenient for now. Not great that we have to write back to
+			// the constraint here though...
+			ManifoldPoint.InitialPhi = WorldContactInitialPhi;
+		}
+
+		// Adjust depth to account for target penetration from user
+		const FRealSingle TargetPhi = ManifoldPoint.TargetPhi;
+		WorldContactDeltaNormal -= TargetPhi;
+
 		Solver.InitManifoldPoint(
 			SolverPointIndex,
 			Dt,
@@ -158,9 +213,10 @@ namespace Chaos
 
 	void UpdateCollisionSolverManifoldFromConstraint(
 		Private::FPBDCollisionSolver& Solver, 
-		const FPBDCollisionConstraint* Constraint, 
+		FPBDCollisionConstraint* Constraint, 
 		const FSolverReal Dt, 
-		const int32 ConstraintPointBeginIndex, 
+		const FRealSingle MaxDepentrationVelocity,
+		const int32 ConstraintPointBeginIndex,
 		const int32 ConstraintPointEndIndex)
 	{
 		const FConstraintSolverBody& Body0 = Solver.SolverBody0();
@@ -177,7 +233,7 @@ namespace Chaos
 				{
 					// Transform the constraint contact data into world space for use by the solver
 					// We build this data directly into the solver's world-space contact data which looks a bit odd with "Init" called after but there you go
-					UpdateCollisionSolverContactPointFromConstraint(Solver, SolverManifoldPointIndex, Constraint, ConstraintManifoldPointIndex, Dt, Body0, Body1);
+					UpdateCollisionSolverContactPointFromConstraint(Solver, SolverManifoldPointIndex, Constraint, ConstraintManifoldPointIndex, Dt, MaxDepentrationVelocity, Body0, Body1);
 				}
 			}
 		}
@@ -187,7 +243,7 @@ namespace Chaos
 
 	void UpdateCollisionSolverFromConstraint(
 		Private::FPBDCollisionSolver& Solver, 
-		const FPBDCollisionConstraint* Constraint, 
+		FPBDCollisionConstraint* Constraint, 
 		const FSolverReal Dt, 
 		const FPBDCollisionSolverSettings& SolverSettings, 
 		bool& bOutPerIterationCollision)
@@ -233,7 +289,7 @@ namespace Chaos
 
 		bOutPerIterationCollision = (!Constraint->GetUseManifold() || Constraint->GetUseIncrementalCollisionDetection());
 
-		UpdateCollisionSolverManifoldFromConstraint(Solver, Constraint, Dt, 0, Constraint->NumManifoldPoints());
+		UpdateCollisionSolverManifoldFromConstraint(Solver, Constraint, Dt, SolverSettings.DepenetrationVelocity, 0, Constraint->NumManifoldPoints());
 	}
 
 	FORCEINLINE_DEBUGGABLE void UpdateCollisionConstraintFromSolver(FPBDCollisionConstraint* Constraint, const Private::FPBDCollisionSolver& Solver, const FSolverReal Dt)
@@ -686,7 +742,10 @@ namespace Chaos
 				Collisions::UpdateConstraint(*Constraint, CorrectedShapeWorldTransform0, CorrectedShapeWorldTransform1, Dt);
 
 				// Update the manifold based on the new or updated contacts
-				UpdateCollisionSolverManifoldFromConstraint(CollisionSolver, Constraint, Dt, BeginPointIndex, Constraint->NumManifoldPoints());
+				UpdateCollisionSolverManifoldFromConstraint(
+					CollisionSolver, Constraint, 
+					Dt, ConstraintContainer.GetSolverSettings().DepenetrationVelocity, 
+					BeginPointIndex, Constraint->NumManifoldPoints());
 
 				Constraint->SetSolverBodies(nullptr, nullptr);
 			}
diff --git a/Engine/Source/Runtime/Experimental/Chaos/Private/Chaos/Collision/ParticlePairMidPhase.cpp b/Engine/Source/Runtime/Experimental/Chaos/Private/Chaos/Collision/ParticlePairMidPhase.cpp
--- a/Engine/Source/Runtime/Experimental/Chaos/Private/Chaos/Collision/ParticlePairMidPhase.cpp
+++ b/Engine/Source/Runtime/Experimental/Chaos/Private/Chaos/Collision/ParticlePairMidPhase.cpp
@@ -477,6 +477,8 @@ namespace Chaos
 			// don't know in advance whether we will pass the Phi check (deferred narrow phase is used with RBAN)
 			if (Constraint->GetPhi() <= CullDistance || Context.GetSettings().bDeferNarrowPhase)
 			{
+				Constraint->SetIsInitialContact(!bWasUpdatedLastTick);
+
 				if (Context.GetAllocator()->ActivateConstraint(Constraint.Get()))
 				{
 					LastUsedEpoch = CurrentEpoch;
@@ -582,9 +584,14 @@ namespace Chaos
 
 			if (bShouldActivate)
 			{
+				const int32 CurrentEpoch = Context.GetAllocator()->GetCurrentEpoch();
+				const int32 LastEpoch = CurrentEpoch - 1;
+				const bool bWasUpdatedLastTick = IsUsedSince(LastEpoch);
+				Constraint->SetIsInitialContact(!bWasUpdatedLastTick);
+
 				if (Context.GetAllocator()->ActivateConstraint(Constraint.Get()))
 				{
-					LastUsedEpoch = Context.GetAllocator()->GetCurrentEpoch();
+					LastUsedEpoch = CurrentEpoch;
 				}
 			}
 
@@ -1734,9 +1741,14 @@ namespace Chaos
 		// If we have a valid contact, add it to the active list
 		// We also add it to the active list if collision detection is deferred because the data will be filled in later and we
 		// don't know in advance whether we will pass the Phi check (deferred narrow phase is used with RBAN)
-		const bool bIsActive = (Constraint->GetPhi() <= CullDistance || Context.GetSettings().bDeferNarrowPhase);
+		const bool bShouldActivate = (Constraint->GetPhi() <= CullDistance || Context.GetSettings().bDeferNarrowPhase);
+
+		if (bShouldActivate)
+		{
+			Constraint->SetIsInitialContact(!bWasUpdatedLastTick);
+		}
 
-		return bIsActive;
+		return bShouldActivate;
 	}
 
 
@@ -1750,9 +1762,6 @@ namespace Chaos
 
 		check(Constraint->GetCCDSweepEnabled());
 
-		Constraint->ResetManifold();
-		Constraint->ResetActiveManifoldContacts();
-
 		const FRigidTransform3 ShapeWorldTransform0 = Constraint->GetShapeRelativeTransform0() * FConstGenericParticleHandle(Constraint->GetParticle0())->GetTransformPQ();
 		const FRigidTransform3 ShapeWorldTransform1 = Constraint->GetShapeRelativeTransform1() * FConstGenericParticleHandle(Constraint->GetParticle1())->GetTransformPQ();
 
@@ -1808,6 +1817,14 @@ namespace Chaos
 			bShouldActivate = Constraint->GetPhi() < CullDistance;
 		}
 
+		if (bShouldActivate)
+		{
+			const int32 CurrentEpoch = Context.GetAllocator()->GetCurrentEpoch();
+			const int32 LastEpoch = CurrentEpoch - 1;
+			const bool bWasUpdatedLastTick = IsUsedSince(LastEpoch);
+			Constraint->SetIsInitialContact(!bWasUpdatedLastTick);
+		}
+
 		return bShouldActivate;
 	}
 
diff --git a/Engine/Source/Runtime/Experimental/Chaos/Private/PBDRigidsSolver.cpp b/Engine/Source/Runtime/Experimental/Chaos/Private/PBDRigidsSolver.cpp
--- a/Engine/Source/Runtime/Experimental/Chaos/Private/PBDRigidsSolver.cpp
+++ b/Engine/Source/Runtime/Experimental/Chaos/Private/PBDRigidsSolver.cpp
@@ -328,6 +328,9 @@ namespace Chaos
 		FRealSingle ChaosSolverMaxPushOutVelocity = -1.0f;
 		FAutoConsoleVariableRef CVarChaosSolverMaxPushOutVelocity(TEXT("p.Chaos.Solver.Collision.MaxPushOutVelocity"), ChaosSolverMaxPushOutVelocity, TEXT("Override max pushout velocity (if >= 0)"));
 
+		FRealSingle ChaosSolverDepenetrationVelocity = -1.0f;
+		FAutoConsoleVariableRef CVarChaosSolverInitialOverlapDepentrationVelocity(TEXT("p.Chaos.Solver.Collision.DepenetrationVelocity"), ChaosSolverDepenetrationVelocity, TEXT("Override initial overlap depenetration velocity (if >= 0)"));
+
 		int32 ChaosSolverCleanupCommandsOnDestruction = 1;
 		FAutoConsoleVariableRef CVarChaosSolverCleanupCommandsOnDestruction(TEXT("p.Chaos.Solver.CleanupCommandsOnDestruction"), ChaosSolverCleanupCommandsOnDestruction, TEXT("Whether or not to run internal command queue cleanup on solver destruction (0 = no cleanup, >0 = cleanup all commands)"));
 
@@ -1235,6 +1238,10 @@ namespace Chaos
 			{
 				SetCollisionMaxPushOutVelocity(ChaosSolverMaxPushOutVelocity);
 			}
+			if (ChaosSolverDepenetrationVelocity >= 0.0f)
+			{
+				SetCollisionDepenetrationVelocity(ChaosSolverDepenetrationVelocity);
+			}
 			if (ChaosSolverDeterministic >= 0)
 			{
 				UpdateIsDeterministic();
@@ -2625,6 +2632,7 @@ TRACE_COUNTER_SET(ChaosTraceCounter_##Name, Value)
 		SetProjectionIterations(InConfig.ProjectionIterations);
 		SetCollisionCullDistance(InConfig.CollisionCullDistance);
 		SetCollisionMaxPushOutVelocity(InConfig.CollisionMaxPushOutVelocity);
+		SetCollisionDepenetrationVelocity(InConfig.CollisionInitialOverlapDepenetrationVelocity);
 		SetGenerateCollisionData(InConfig.bGenerateCollisionData);
 		SetGenerateBreakingData(InConfig.bGenerateBreakData);
 		SetGenerateTrailingData(InConfig.bGenerateTrailingData);
diff --git a/Engine/Source/Runtime/Experimental/Chaos/Public/Chaos/Collision/ContactPoint.h b/Engine/Source/Runtime/Experimental/Chaos/Public/Chaos/Collision/ContactPoint.h
--- a/Engine/Source/Runtime/Experimental/Chaos/Public/Chaos/Collision/ContactPoint.h
+++ b/Engine/Source/Runtime/Experimental/Chaos/Public/Chaos/Collision/ContactPoint.h
@@ -108,6 +108,7 @@ namespace Chaos
 				uint8 bWasRestored : 1;						// Whether the point was retored from the previous frame due to lack of movement
 				uint8 bWasReplaced : 1;						// @todo(chaos): remove this
 				uint8 bHasStaticFrictionAnchor : 1;			// Whether our static friction anchor was recovered from a prior tick
+				uint8 bInitialContact : 1;					// Whether this is a new contact this tick for handling initial penetration
 			};
 			uint8 Bits;
 		};
@@ -116,6 +117,7 @@ namespace Chaos
 			: ContactPoint()
 			, Flags()
 			, TargetPhi(0)
+			, InitialPhi(0)
 			, ShapeAnchorPoints{ FVec3f(0), FVec3f(0) }
 			, InitialShapeContactPoints{ FVec3f(0), FVec3f(0) }
 		{}
@@ -124,6 +126,7 @@ namespace Chaos
 			: ContactPoint(InContactPoint)
 			, Flags()
 			, TargetPhi(0)
+			, InitialPhi(0)
 			, ShapeAnchorPoints{ FVec3f(0), FVec3f(0) }
 			, InitialShapeContactPoints{ FVec3f(0), FVec3f(0) }
 		{}
@@ -131,6 +134,7 @@ namespace Chaos
 		FContactPointf ContactPoint;			// Contact point results of low-level collision detection
 		FFlags Flags;							// Various flags
 		FRealSingle TargetPhi;					// Usually 0, but can be used to add padding or penetration (e.g., via a collision modifer)
+		FRealSingle InitialPhi;					// Non-resolved initial penetration
 		FVec3f ShapeAnchorPoints[2];			// When static friction holds, the contact points on each shape when static friction contact was made
 		FVec3f InitialShapeContactPoints[2];	// ShapeContactPoints when the constraint was first initialized. Used to track reusablility
 	};
@@ -165,17 +169,8 @@ namespace Chaos
 	class FSavedManifoldPoint
 	{
 	public:
-		FSavedManifoldPoint()
-		{
-		}
-
-		void Reset()
-		{
-			ShapeContactPoints[0] = FVec3f(0);
-			ShapeContactPoints[1] = FVec3f(0);
-		}
-
 		FVec3f ShapeContactPoints[2];			// Contact anchor points for friction
+		FRealSingle InitialPhi;					// Non-resolved initial penetration
 	};
 
 	class FManifoldPointResult
diff --git a/Engine/Source/Runtime/Experimental/Chaos/Public/Chaos/Collision/PBDCollisionConstraint.h b/Engine/Source/Runtime/Experimental/Chaos/Public/Chaos/Collision/PBDCollisionConstraint.h
--- a/Engine/Source/Runtime/Experimental/Chaos/Public/Chaos/Collision/PBDCollisionConstraint.h
+++ b/Engine/Source/Runtime/Experimental/Chaos/Public/Chaos/Collision/PBDCollisionConstraint.h
@@ -376,6 +376,11 @@ namespace Chaos
 		void SetIsProbe(bool bInProbe) { Flags.bIsProbe = bInProbe; }
 		bool GetIsProbe() const { return Flags.bIsProbe; }
 
+		// Is this considered an initial contact (i.e., contact was activated this tick, but not last)
+		void SetIsInitialContact(const bool bInIsInitialContact) { Flags.bInitialContact = bInIsInitialContact; }
+		bool IsInitialContact() const { return Flags.bInitialContact; }
+		FRealSingle GetMinInitialPhi() const { return MinInitialPhi; }
+
 		virtual bool SupportsSleeping() const override final { return true; }
 		CHAOS_API virtual bool IsSleeping() const override final;
 		CHAOS_API virtual void SetIsSleeping(const bool bInIsSleeping) override final;
@@ -600,7 +605,7 @@ namespace Chaos
 		}
 
 		// Set the transforms when we last ran collision detection. This also sets the bCanRestoreManifold flag which
-		// allows the use of UpdateAndTryRestoreManifold on the next tick.
+		// allows the use of TryRestoreManifold on the next tick.
 		void SetLastShapeWorldTransforms(const FRigidTransform3& InShapeWorldTransform0, const FRigidTransform3& InShapeWorldTransform1)
 		{
 			LastShapeWorldPositionDelta = InShapeWorldTransform0.GetTranslation() - InShapeWorldTransform1.GetTranslation();
@@ -672,7 +677,7 @@ namespace Chaos
 
 		FORCEINLINE void ResetSolverResults()
 		{
-			// NOTE: does not initalize any data. All properties will be written to in SetSolverResults
+			// NOTE: does not initalize any manifold point data. All properties will be written to in SetSolverResults
 			//SavedManifoldPoints.SetNum(ManifoldPoints.Num());
 			SavedManifoldPoints.Reset(ManifoldPoints.Num());
 			ManifoldPointResults.SetNum(ManifoldPoints.Num());
@@ -690,26 +695,20 @@ namespace Chaos
 		{
 			FManifoldPoint& ManifoldPoint = ManifoldPoints[ManifoldPointIndex];
 			FManifoldPointResult& ManifoldPointResult = ManifoldPointResults[ManifoldPointIndex];
+			FSavedManifoldPoint* SavedManifoldPoint = nullptr;
 
-			ManifoldPointResult.NetPushOut = NetPushOut;
-			ManifoldPointResult.NetImpulse = NetImpulse;
-			ManifoldPointResult.bIsValid = true;
-			ManifoldPointResult.bInsideStaticFrictionCone = false;
-
-			AccumulatedImpulse += NetImpulse + (NetPushOut / Dt);
-
-			// Save contact data for friction\
-			// NOTE: We also write the anchors back to the manifold point for use if the point gets restored next tick
-			// in which case it skips the saved point search etc
-			if (StaticFrictionRatio >= FReal(1.0f - UE_KINDA_SMALL_NUMBER))
+			// Save contact data for friction
+			FVec3f Anchor0, Anchor1;
+			bool bInsideStaticFrictionCone = false;
+			if (StaticFrictionRatio >= FRealSingle(1.0f - UE_KINDA_SMALL_NUMBER))
 			{
 				// StaticFrictionRatio ~= 1: Static friction held - we keep the same contacts points as-is for use next frame
-				FSavedManifoldPoint& SavedManifoldPoint = SavedManifoldPoints[SavedManifoldPoints.AddUninitialized()];
-				SavedManifoldPoint.ShapeContactPoints[0] = ManifoldPoint.ShapeAnchorPoints[0];
-				SavedManifoldPoint.ShapeContactPoints[1] = ManifoldPoint.ShapeAnchorPoints[1];
-				ManifoldPointResult.bInsideStaticFrictionCone = true;
+				SavedManifoldPoint = &SavedManifoldPoints[SavedManifoldPoints.AddUninitialized()];
+				Anchor0 = ManifoldPoint.ShapeAnchorPoints[0];
+				Anchor1 = ManifoldPoint.ShapeAnchorPoints[1];
+				bInsideStaticFrictionCone = true;
 			}
-			else if (StaticFrictionRatio < FReal(UE_KINDA_SMALL_NUMBER))
+			else if (StaticFrictionRatio < FRealSingle(UE_KINDA_SMALL_NUMBER))
 			{
 				// StaticFrictionRatio ~= 0: No friction (or no contact/impulse) - discard the friction anchors
 				// If we have a lot of manifold points, we don't store the previous position for these contacts because we assume
@@ -718,22 +717,32 @@ namespace Chaos
 				const int32 SmallNumManifoldPoints = 8;
 				if (ManifoldPoints.Num() < SmallNumManifoldPoints)
 				{
-					const FVec3 Anchor0 = ManifoldPoint.ContactPoint.ShapeContactPoints[0];
-					const FVec3 Anchor1 = ManifoldPoint.ContactPoint.ShapeContactPoints[1];
-					FSavedManifoldPoint& SavedManifoldPoint = SavedManifoldPoints[SavedManifoldPoints.AddUninitialized()];
-					SavedManifoldPoint.ShapeContactPoints[0] = Anchor0;
-					SavedManifoldPoint.ShapeContactPoints[1] = Anchor1;
+					SavedManifoldPoint = &SavedManifoldPoints[SavedManifoldPoints.AddUninitialized()];
+					Anchor0 = ManifoldPoint.ContactPoint.ShapeContactPoints[0];
+					Anchor1 = ManifoldPoint.ContactPoint.ShapeContactPoints[1];
 				}
 			}
 			else
 			{
 				// 0 < StaticFrictionRatio < 1: We exceeded the friction cone. Slide the friction anchor 
 				// toward the last-detected contact position so that it sits at the edge of the friction cone.
-				const FVec3 Anchor0 = FVec3::Lerp(ManifoldPoint.ContactPoint.ShapeContactPoints[0], ManifoldPoint.ShapeAnchorPoints[0], StaticFrictionRatio);
-				const FVec3 Anchor1 = FVec3::Lerp(ManifoldPoint.ContactPoint.ShapeContactPoints[1], ManifoldPoint.ShapeAnchorPoints[1], StaticFrictionRatio);
-				FSavedManifoldPoint& SavedManifoldPoint = SavedManifoldPoints[SavedManifoldPoints.AddUninitialized()];
-				SavedManifoldPoint.ShapeContactPoints[0] = Anchor0;
-				SavedManifoldPoint.ShapeContactPoints[1] = Anchor1;
+				SavedManifoldPoint = &SavedManifoldPoints[SavedManifoldPoints.AddUninitialized()];
+				Anchor0 = FVec3f::Lerp(ManifoldPoint.ContactPoint.ShapeContactPoints[0], ManifoldPoint.ShapeAnchorPoints[0], StaticFrictionRatio);
+				Anchor1 = FVec3f::Lerp(ManifoldPoint.ContactPoint.ShapeContactPoints[1], ManifoldPoint.ShapeAnchorPoints[1], StaticFrictionRatio);
+			}
+
+			AccumulatedImpulse += NetImpulse + (NetPushOut / Dt);
+
+			ManifoldPointResult.NetPushOut = NetPushOut;
+			ManifoldPointResult.NetImpulse = NetImpulse;
+			ManifoldPointResult.bIsValid = true;
+			ManifoldPointResult.bInsideStaticFrictionCone = bInsideStaticFrictionCone;
+
+			if (SavedManifoldPoint != nullptr)
+			{
+				SavedManifoldPoint->ShapeContactPoints[0] = Anchor0;
+				SavedManifoldPoint->ShapeContactPoints[1] = Anchor1;
+				SavedManifoldPoint->InitialPhi = ManifoldPoint.InitialPhi;
 			}
 		}
 
@@ -801,6 +810,7 @@ namespace Chaos
 
 		CHAOS_API int32 FindSavedManifoldPoint(const int32 ManifoldPointIndex, int32* InOutAllowedSavedPointIndices, int32& InOutNumAllowedSavedPoints) const;
 		CHAOS_API void AssignSavedManifoldPoints();
+		CHAOS_API void CalculateMinInitialPhi();
 
 		inline void InitManifoldPoint(const int32 ManifoldPointIndex, const FContactPoint& ContactPoint)
 		{
@@ -810,7 +820,8 @@ namespace Chaos
 			ManifoldPoint.ShapeAnchorPoints[1] = ContactPoint.ShapeContactPoints[1];
 			ManifoldPoint.InitialShapeContactPoints[0] = ContactPoint.ShapeContactPoints[0];
 			ManifoldPoint.InitialShapeContactPoints[1] = ContactPoint.ShapeContactPoints[1];
-			ManifoldPoint.TargetPhi = FReal(0);
+			ManifoldPoint.TargetPhi = FRealSingle(0);
+			ManifoldPoint.InitialPhi = FRealSingle(0);
 			ManifoldPoint.Flags.Reset();
 		}
 
@@ -832,7 +843,7 @@ namespace Chaos
 		CHAOS_API void UpdateMaterialPropertiesImpl();
 
 	private:
-		CHAOS_API FReal CalculateSavedManifoldPointScore(const FSavedManifoldPoint& SavedManifoldPoint, const FManifoldPoint& ManifoldPoint, const FReal DistanceToleranceSq) const;
+		CHAOS_API FReal CalculateSavedManifoldDistanceSq(const FSavedManifoldPoint& SavedManifoldPoint, const FManifoldPoint& ManifoldPoint, const FReal DistanceToleranceSq) const;
 
 		union FFlags
 		{
@@ -852,6 +863,7 @@ namespace Chaos
 				uint16 bCCDSweepEnabled: 1;				// If this is a CCD constraint, do we want to enable the sweep/rewind phase?
 				uint16 bModifierApplied : 1;			// Was a constraint modifier applied this tick
 				uint16 bMaterialSet : 1;				// Has the material been set (or does it need to be reset)
+				uint16 bInitialContact : 1;				// Is this contact considered an initial contact
 			};
 			uint16 Bits;
 		};
@@ -889,7 +901,7 @@ namespace Chaos
 		// The shape transforms at the current particle transforms
 		FRigidTransform3 ShapeWorldTransforms[2];
 
-		// Used by manifold point injection to see how many points were in the manifold before UpdateAndTryRestore
+		// Used by manifold point injection to see how many points were in the manifold before TryRestoreManifold
 		int32 ExpectedNumManifoldPoints;
 
 		// Relative transform the last time we ran the narrow phase
@@ -914,6 +926,12 @@ namespace Chaos
 		TManifoldPointArray<FManifoldPoint> ManifoldPoints;
 		TManifoldPointArray<FManifoldPointResult> ManifoldPointResults;
 
+		// The lowest InitialPhi value of all saved manifold points. This is used when we add new points to
+		// the manifold but are still in the depenetrating phase to limit the initial overlap. 
+		// We don't want the new point to cause a pop but also don't want all new contacts to be initial overlaps
+		// when we were already separated at some point in the past.
+		FRealSingle MinInitialPhi;
+
 		// Value in range [0,1] used to interpolate P between [X,P] that we will rollback to when solving at time of impact.
 		FRealSingle CCDTimeOfImpact;
 
diff --git a/Engine/Source/Runtime/Experimental/Chaos/Public/Chaos/Collision/PBDCollisionSolverSettings.h b/Engine/Source/Runtime/Experimental/Chaos/Public/Chaos/Collision/PBDCollisionSolverSettings.h
--- a/Engine/Source/Runtime/Experimental/Chaos/Public/Chaos/Collision/PBDCollisionSolverSettings.h
+++ b/Engine/Source/Runtime/Experimental/Chaos/Public/Chaos/Collision/PBDCollisionSolverSettings.h
@@ -15,6 +15,7 @@ namespace Chaos
 
 		FPBDCollisionSolverSettings()
 			: MaxPushOutVelocity(0)
+			, DepenetrationVelocity(0)
 			, NumPositionFrictionIterations(4)
 			, NumVelocityFrictionIterations(1)
 			, NumPositionShockPropagationIterations(3)
@@ -24,8 +25,12 @@ namespace Chaos
 
 		// Maximum speed at which two objects can depenetrate (actually, how much relative velocity can be added
 		// to a contact per frame when depentrating. Stacks and deep penetrations can lead to larger velocities)
+		// A value of zero means unlimited.
 		FReal MaxPushOutVelocity;
 
+		// The speed at which initially-overlapping objects depentrate
+		FRealSingle DepenetrationVelocity;
+
 		// How many of the position iterations should run static/dynamic friction
 		int32 NumPositionFrictionIterations;
 
diff --git a/Engine/Source/Runtime/Experimental/Chaos/Public/Chaos/PBDCollisionConstraints.h b/Engine/Source/Runtime/Experimental/Chaos/Public/Chaos/PBDCollisionConstraints.h
--- a/Engine/Source/Runtime/Experimental/Chaos/Public/Chaos/PBDCollisionConstraints.h
+++ b/Engine/Source/Runtime/Experimental/Chaos/Public/Chaos/PBDCollisionConstraints.h
@@ -233,7 +233,12 @@ public:
 
 	void SetMaxPushOutVelocity(const FReal InMaxPushOutVelocity)
 	{
-		SolverSettings.MaxPushOutVelocity = InMaxPushOutVelocity;
+		SolverSettings.MaxPushOutVelocity = FMath::Max(InMaxPushOutVelocity, FReal(0));
+	}
+
+	void SetDepenetrationVelocity(const FRealSingle InVel)
+	{
+		SolverSettings.DepenetrationVelocity = FMath::Max(InVel, FRealSingle(0));
 	}
 
 	void SetPositionFrictionIterations(const int32 InNumIterations)
diff --git a/Engine/Source/Runtime/Experimental/Chaos/Public/ChaosSolverConfiguration.h b/Engine/Source/Runtime/Experimental/Chaos/Public/ChaosSolverConfiguration.h
--- a/Engine/Source/Runtime/Experimental/Chaos/Public/ChaosSolverConfiguration.h
+++ b/Engine/Source/Runtime/Experimental/Chaos/Public/ChaosSolverConfiguration.h
@@ -63,6 +63,12 @@ struct FChaosSolverConfiguration
 	UPROPERTY(EditAnywhere, Category = "SolverConfiguration|Collision", meta = (ClampMin = "0.0"))
 	float CollisionMaxPushOutVelocity;
 
+	// If two bodies start off in overaapping each other, they will depentrate at this speed when they wake.
+	// NOTE: this is the setting for the new Initial Penetration Handling system. This is currently enabled with the 
+	// cvar p.Chaos.PBDCollisionSolver.EnableInitialDepenetration but will eventually be on by default.
+	UPROPERTY(EditAnywhere, Category = "SolverConfiguration|Collision", meta = (ClampMin = "0.0"))
+	float CollisionInitialOverlapDepenetrationVelocity;
+
 	UPROPERTY(EditAnywhere, Category = "SolverConfiguration|Clustering")
 	float ClusterConnectionFactor;
 
diff --git a/Engine/Source/Runtime/Experimental/Chaos/Public/PBDRigidsSolver.h b/Engine/Source/Runtime/Experimental/Chaos/Public/PBDRigidsSolver.h
--- a/Engine/Source/Runtime/Experimental/Chaos/Public/PBDRigidsSolver.h
+++ b/Engine/Source/Runtime/Experimental/Chaos/Public/PBDRigidsSolver.h
@@ -181,6 +181,7 @@ namespace Chaos
 		void SetCollisionCullDistance(const FReal InCullDistance) { GetEvolution()->GetCollisionConstraints().SetCullDistance(InCullDistance); }
 		void SetVelocityBoundsExpansion(const FReal BoundsVelocityMultiplier, const FReal MaxBoundsVelocityExpansion) { GetEvolution()->GetCollisionConstraints().SetVelocityBoundsExpansion(BoundsVelocityMultiplier, MaxBoundsVelocityExpansion); }
 		void SetCollisionMaxPushOutVelocity(const FReal InMaxPushOutVelocity) { GetEvolution()->GetCollisionConstraints().SetMaxPushOutVelocity(InMaxPushOutVelocity); }
+		void SetCollisionDepenetrationVelocity(const FRealSingle InVelocity) { GetEvolution()->GetCollisionConstraints().SetDepenetrationVelocity(InVelocity); }
 
 		/**/
 		void SetGenerateCollisionData(bool bDoGenerate) { GetEventFilters()->SetGenerateCollisionEvents(bDoGenerate); }
