Geometry collection : fix scale not being properly taken in account when computing the transforms from the GT particles
knock-on bug from recent optimization

diff --git a/Engine/Source/Runtime/Experimental/Chaos/Private/PhysicsProxy/GeometryCollectionPhysicsProxy.cpp b/Engine/Source/Runtime/Experimental/Chaos/Private/PhysicsProxy/GeometryCollectionPhysicsProxy.cpp
--- a/Engine/Source/Runtime/Experimental/Chaos/Private/PhysicsProxy/GeometryCollectionPhysicsProxy.cpp
+++ b/Engine/Source/Runtime/Experimental/Chaos/Private/PhysicsProxy/GeometryCollectionPhysicsProxy.cpp
@@ -3639,6 +3639,9 @@ bool FGeometryCollectionPhysicsProxy::PullFromPhysicsState(const Chaos::FDirtyGe
 			const TBitArray<>& PrevResultsModifiedIndices = PullData.Results().GetModifiedTransformIndices();
 			const FGeometryCollectionResults& NextResults = NextPullData->Results();
 
+			const bool bIsComponentTransformScaled = !Parameters.WorldTransform.GetScale3D().Equals(FVector::OneVector);
+			const FTransform ComponentScaleTransform(FQuat::Identity, FVector::ZeroVector, Parameters.WorldTransform.GetScale3D());
+
 			for (int32 TransformIndex = 0; TransformIndex < NumTransforms; TransformIndex++)
 			{
 				const bool bAnimatingWhileDisabled = AnimationsActive ? (*AnimationsActive)[TransformIndex] : false;
@@ -3672,11 +3675,15 @@ bool FGeometryCollectionPhysicsProxy::PullFromPhysicsState(const Chaos::FDirtyGe
 								}
 							}
 
-							const FTransform NewTransform = WorldTransform.GetRelativeTransform(ParentWorldTransform);
+							FTransform NewTransform = WorldTransform.GetRelativeTransform(ParentWorldTransform);
+							if (ParentTransformIndex == INDEX_NONE && bIsComponentTransformScaled)
+							{
+								NewTransform = MassToLocal.Inverse() * ComponentScaleTransform * MassToLocal * NewTransform;
+							}
 							if (!NewTransform.Equals(GameThreadCollection.Transform[TransformIndex], GeometryCollectionTransformTolerance))
 							{
 								bHasDifferentTransforms = true;
-								GameThreadCollection.Transform[TransformIndex] = WorldTransform.GetRelativeTransform(ParentWorldTransform);
+								GameThreadCollection.Transform[TransformIndex] = NewTransform;
 							}
 						}
 					}
