Iris
- Changed when we call OnSubObjectCreatedFromReplication for new subobjets to behave closer to legacy.
- It is now called after we have applied state for the owner but before we apply state for the subobject, this is to mimic serialization order from legacy system.

diff --git a/Engine/Source/Runtime/Engine/Private/Net/Iris/ReplicationSystem/ActorReplicationBridge.cpp b/Engine/Source/Runtime/Engine/Private/Net/Iris/ReplicationSystem/ActorReplicationBridge.cpp
--- a/Engine/Source/Runtime/Engine/Private/Net/Iris/ReplicationSystem/ActorReplicationBridge.cpp
+++ b/Engine/Source/Runtime/Engine/Private/Net/Iris/ReplicationSystem/ActorReplicationBridge.cpp
@@ -670,8 +670,8 @@ FObjectReplicationBridgeInstantiateResult UActorReplicationBridge::BeginInstanti
 				checkf(SubObj->IsIn(ObjOuter), TEXT("UActorReplicationBridge::BeginInstantiateFromRemote: Subobject is not in Outer. SubObject: %s, Actor: %s Outer: %s"), *SubObj->GetName(), *Owner->GetName(), *ObjOuter->GetName());
 				checkf(Cast<AActor>(SubObj) == nullptr, TEXT("UActorReplicationBridge::BeginInstantiateFromRemote: Subobject is an Actor. SubObject: %s, Actor: %s"), *SubObj->GetName(), *Owner->GetName());
 
-				// Notify actor that we created a component from replication
-				Owner->OnSubobjectCreatedFromReplication(SubObj);
+				// We must defer call OnSubObjectCreatedFromReplication after the state has been applied to the owning actor in order to behave like old system.
+				InstantiateResult.Flags |= EReplicationBridgeCreateNetRefHandleResultFlags::ShouldCallSubObjectCreatedFromReplication;
 
 				// Created objects may be destroyed.
 				InstantiateResult.Flags |= EReplicationBridgeCreateNetRefHandleResultFlags::AllowDestroyInstanceFromRemote;
@@ -743,6 +743,16 @@ void UActorReplicationBridge::EndInstantiateFromRemote(FNetRefHandle Handle)
 	}
 }
 
+void UActorReplicationBridge::OnSubObjectCreatedFromReplication(FNetRefHandle SubObjectHandle)
+{
+	AActor* Owner = Cast<AActor>(GetReplicatedObject(InternalGetSubObjectOwner(SubObjectHandle)));
+	UObject* SubObject = GetReplicatedObject(SubObjectHandle);
+	if (IsValid(Owner) && IsValid(SubObject))
+	{
+		Owner->OnSubobjectCreatedFromReplication(SubObject);
+	}
+}
+
 void UActorReplicationBridge::DestroyInstanceFromRemote(UObject* Instance, EReplicationBridgeDestroyInstanceReason DestroyReason, EReplicationBridgeDestroyInstanceFlags DestroyFlags)
 {
 	if (DestroyReason == EReplicationBridgeDestroyInstanceReason::DoNotDestroy)
diff --git a/Engine/Source/Runtime/Engine/Public/Net/Iris/ReplicationSystem/ActorReplicationBridge.h b/Engine/Source/Runtime/Engine/Public/Net/Iris/ReplicationSystem/ActorReplicationBridge.h
--- a/Engine/Source/Runtime/Engine/Public/Net/Iris/ReplicationSystem/ActorReplicationBridge.h
+++ b/Engine/Source/Runtime/Engine/Public/Net/Iris/ReplicationSystem/ActorReplicationBridge.h
@@ -86,6 +86,7 @@ protected:
 	virtual FCreationHeader* ReadCreationHeader(UE::Net::FNetSerializationContext& Context) override;
 	virtual FObjectReplicationBridgeInstantiateResult BeginInstantiateFromRemote(FNetRefHandle SubObjectOwnerNetHandle, const UE::Net::FNetObjectResolveContext& ResolveContext, const FCreationHeader* InHeader) override;
 	virtual bool OnInstantiatedFromRemote(UObject* Instance, const FCreationHeader* InHeader, uint32 ConnectionId) const override;
+	virtual void OnSubObjectCreatedFromReplication(FNetRefHandle SubObjectHandle) override;
 	virtual void EndInstantiateFromRemote(FNetRefHandle Handle) override;
 	virtual void DestroyInstanceFromRemote(UObject* Instance, EReplicationBridgeDestroyInstanceReason DestroyReason, EReplicationBridgeDestroyInstanceFlags DestroyFlags) override;
 	virtual void GetInitialDependencies(FNetRefHandle Handle, FNetDependencyInfoArray& OutDependencies) const override;
diff --git a/Engine/Source/Runtime/Experimental/Iris/Core/Private/Iris/ReplicationSystem/ObjectReplicationBridge.cpp b/Engine/Source/Runtime/Experimental/Iris/Core/Private/Iris/ReplicationSystem/ObjectReplicationBridge.cpp
--- a/Engine/Source/Runtime/Experimental/Iris/Core/Private/Iris/ReplicationSystem/ObjectReplicationBridge.cpp
+++ b/Engine/Source/Runtime/Experimental/Iris/Core/Private/Iris/ReplicationSystem/ObjectReplicationBridge.cpp
@@ -641,6 +641,11 @@ FReplicationBridgeCreateNetRefHandleResult UObjectReplicationBridge::CreateNetRe
 	return CreateResult;
 }
 
+void UObjectReplicationBridge::SubObjectCreatedFromReplication(FNetRefHandle Handle)
+{
+	OnSubObjectCreatedFromReplication(Handle);
+}
+
 void UObjectReplicationBridge::PostApplyInitialState(FNetRefHandle Handle)
 {
 	EndInstantiateFromRemote(Handle);
diff --git a/Engine/Source/Runtime/Experimental/Iris/Core/Private/Iris/ReplicationSystem/ReplicationBridge.cpp b/Engine/Source/Runtime/Experimental/Iris/Core/Private/Iris/ReplicationSystem/ReplicationBridge.cpp
--- a/Engine/Source/Runtime/Experimental/Iris/Core/Private/Iris/ReplicationSystem/ReplicationBridge.cpp
+++ b/Engine/Source/Runtime/Experimental/Iris/Core/Private/Iris/ReplicationSystem/ReplicationBridge.cpp
@@ -64,6 +64,10 @@ void UReplicationBridge::UpdateInstancesWorldLocation()
 {
 }
 
+void UReplicationBridge::SubObjectCreatedFromReplication(FNetRefHandle SubObjectRefHandle)
+{
+}
+
 void UReplicationBridge::PostApplyInitialState(FNetRefHandle Handle)
 {
 }
@@ -259,6 +263,11 @@ void UReplicationBridge::ReadAndExecuteDestructionInfoFromRemote(FReplicationBri
 	}
 }
 
+void UReplicationBridge::CallSubObjectCreatedFromReplication(FNetRefHandle SubObjectHandle)
+{
+	SubObjectCreatedFromReplication(SubObjectHandle);
+}
+
 void UReplicationBridge::CallPostApplyInitialState(FNetRefHandle Handle)
 {
 	PostApplyInitialState(Handle);
diff --git a/Engine/Source/Runtime/Experimental/Iris/Core/Private/Iris/ReplicationSystem/ReplicationReader.cpp b/Engine/Source/Runtime/Experimental/Iris/Core/Private/Iris/ReplicationSystem/ReplicationReader.cpp
--- a/Engine/Source/Runtime/Experimental/Iris/Core/Private/Iris/ReplicationSystem/ReplicationReader.cpp
+++ b/Engine/Source/Runtime/Experimental/Iris/Core/Private/Iris/ReplicationSystem/ReplicationReader.cpp
@@ -316,6 +316,7 @@ uint32 FReplicationReader::ReadObjectsPendingDestroy(FNetSerializationContext& C
 						Info.bIsInitialState = 0U;
 						Info.bHasState = 0U;
 						Info.bHasAttachments = 0U;
+						Info.bShouldCallSubObjectCreatedFromReplication = 0U;
 
 						// Mark for dispatch
 						ObjectsToDispatchArray->CommitPendingDispatchObjectInfo();
@@ -757,6 +758,7 @@ void FReplicationReader::ReadObjectInBatch(FNetSerializationContext& Context, FN
 	uint32 NewBaselineIndex = FDeltaCompressionBaselineManager::InvalidBaselineIndex;
 
 	const bool bIsInitialState = bHasState && Reader.ReadBool();
+	bool bShouldCallSubObjectCreatedFromReplication = false;
 	uint32 InternalIndex = ObjectIndexForOOBAttachment;
 
 	UE_NET_TRACE_OBJECT_SCOPE(IncompleteHandle, Reader, Context.GetTraceCollector(), ENetTraceVerbosity::Trace);
@@ -825,7 +827,8 @@ void FReplicationReader::ReadObjectInBatch(FNetSerializationContext& Context, FN
 		InternalIndex = NetRefHandleManager->GetInternalIndex(NetRefHandle);
 		FNetRefHandleManager::FReplicatedObjectData& ObjectData = NetRefHandleManager->GetReplicatedObjectDataNoCheck(InternalIndex);
 		ObjectData.bAllowDestroyInstanceFromRemote = EnumHasAnyFlags(CreateResult.Flags, EReplicationBridgeCreateNetRefHandleResultFlags::AllowDestroyInstanceFromRemote);
-
+		bShouldCallSubObjectCreatedFromReplication = EnumHasAnyFlags(CreateResult.Flags, EReplicationBridgeCreateNetRefHandleResultFlags::ShouldCallSubObjectCreatedFromReplication);
+		
 		FReplicatedObjectInfo& ObjectInfo = StartReplication(InternalIndex);
 
 		ObjectInfo.bIsDeltaCompressionEnabled = bIsDeltaCompressed;
@@ -867,6 +870,7 @@ void FReplicationReader::ReadObjectInBatch(FNetSerializationContext& Context, FN
 		Info.bDestroy = !!(ReplicatedDestroyHeaderFlags & (ReplicatedDestroyHeaderFlags_TearOff | ReplicatedDestroyHeaderFlags_DestroyInstance));
 		Info.bTearOff = !!(ReplicatedDestroyHeaderFlags & ReplicatedDestroyHeaderFlags_TearOff);
 		Info.bDeferredEndReplication = !!(ReplicatedDestroyHeaderFlags & (ReplicatedDestroyHeaderFlags_TearOff | ReplicatedDestroyHeaderFlags_EndReplication));
+		Info.bShouldCallSubObjectCreatedFromReplication = bShouldCallSubObjectCreatedFromReplication;
 
 		if (bHasState)
 		{
@@ -1439,6 +1443,16 @@ void FReplicationReader::DispatchStateData(FNetSerializationContext& Context)
 		PostDispatchObjectInfo.DequantizeAndApplyContext = nullptr;
 		PostDispatchObjectInfo.AttachmentDispatchedFlags = ENetObjectAttachmentDispatchFlags::None;
 
+		// For SubObjects we call must call this method after applying state data for the owner, in order to remain backwards compatible.
+		if (Info.bShouldCallSubObjectCreatedFromReplication)
+		{
+			const FNetRefHandleManager::FReplicatedObjectData& ObjectData = NetRefHandleManager->GetReplicatedObjectDataNoCheck(Info.InternalIndex);
+			if (ObjectData.SubObjectRootIndex != FNetRefHandleManager::InvalidInternalIndex)
+			{
+				ReplicationBridge->CallSubObjectCreatedFromReplication(ObjectData.RefHandle);
+			}
+		}
+
 		// If we have any object references we want to update any unresolved ones, including previously unresolved references
 		if (Info.bHasState)
 		{
diff --git a/Engine/Source/Runtime/Experimental/Iris/Core/Private/Iris/ReplicationSystem/ReplicationReader.h b/Engine/Source/Runtime/Experimental/Iris/Core/Private/Iris/ReplicationSystem/ReplicationReader.h
--- a/Engine/Source/Runtime/Experimental/Iris/Core/Private/Iris/ReplicationSystem/ReplicationReader.h
+++ b/Engine/Source/Runtime/Experimental/Iris/Core/Private/Iris/ReplicationSystem/ReplicationReader.h
@@ -115,6 +115,7 @@ private:
 		uint32 bDestroy : 1;
 		uint32 bTearOff : 1;
 		uint32 bDeferredEndReplication : 1;
+		uint32 bShouldCallSubObjectCreatedFromReplication : 1;
 	};
 
 
diff --git a/Engine/Source/Runtime/Experimental/Iris/Core/Public/Iris/ReplicationSystem/ObjectReplicationBridge.h b/Engine/Source/Runtime/Experimental/Iris/Core/Public/Iris/ReplicationSystem/ObjectReplicationBridge.h
--- a/Engine/Source/Runtime/Experimental/Iris/Core/Public/Iris/ReplicationSystem/ObjectReplicationBridge.h
+++ b/Engine/Source/Runtime/Experimental/Iris/Core/Public/Iris/ReplicationSystem/ObjectReplicationBridge.h
@@ -164,6 +164,7 @@ protected:
 	IRISCORE_API virtual void PruneStaleObjects() override;	
 	IRISCORE_API virtual bool WriteNetRefHandleCreationInfo(FReplicationBridgeSerializationContext& Context, FNetRefHandle Handle) override;
 	IRISCORE_API virtual FReplicationBridgeCreateNetRefHandleResult CreateNetRefHandleFromRemote(FNetRefHandle SubObjectOwnerNetHandle, FNetRefHandle WantedNetHandle, FReplicationBridgeSerializationContext& Context) override;
+	IRISCORE_API virtual void SubObjectCreatedFromReplication(FNetRefHandle SubObjectHandle) override;
 	IRISCORE_API virtual void PostApplyInitialState(FNetRefHandle Handle) override;
 	IRISCORE_API virtual void DetachInstanceFromRemote(FNetRefHandle Handle, EReplicationBridgeDestroyInstanceReason DestroyReason, EReplicationBridgeDestroyInstanceFlags DestroyFlags) override;
 	IRISCORE_API virtual void DetachInstance(FNetRefHandle Handle) override;
@@ -188,6 +189,9 @@ protected:
 	/** Invoked before we start applying state data to instance on remote end. */
 	virtual bool OnInstantiatedFromRemote(UObject* Instance, const FCreationHeader* InHeader, uint32 ConnectionId) const { return true; }
 
+	/** Invoked for new replicated SubObjects after state has been applied to owner */
+	virtual void OnSubObjectCreatedFromReplication(FNetRefHandle SubObjectHandle) {};
+
 	/** Invoked after remote NetHandle has been created and initial state is applied. */
 	virtual void EndInstantiateFromRemote(FNetRefHandle Handle) {};
 
diff --git a/Engine/Source/Runtime/Experimental/Iris/Core/Public/Iris/ReplicationSystem/ReplicationBridge.h b/Engine/Source/Runtime/Experimental/Iris/Core/Public/Iris/ReplicationSystem/ReplicationBridge.h
--- a/Engine/Source/Runtime/Experimental/Iris/Core/Public/Iris/ReplicationSystem/ReplicationBridge.h
+++ b/Engine/Source/Runtime/Experimental/Iris/Core/Public/Iris/ReplicationSystem/ReplicationBridge.h
@@ -82,6 +82,7 @@ enum class EReplicationBridgeCreateNetRefHandleResultFlags : unsigned
 	None = 0U,
 	/** Whether the instance may be destroyed due to the remote peer requesting the object to be destroyed. If not then the object itself must not be destroyed. */
 	AllowDestroyInstanceFromRemote = 1U << 0U,
+	ShouldCallSubObjectCreatedFromReplication = AllowDestroyInstanceFromRemote << 1U,
 };
 ENUM_CLASS_FLAGS(EReplicationBridgeCreateNetRefHandleResultFlags);
 
@@ -181,6 +182,9 @@ protected:
 	/** Read data required to instantiate NetObject from bitstream. */
 	IRISCORE_API virtual FReplicationBridgeCreateNetRefHandleResult CreateNetRefHandleFromRemote(FNetRefHandle SubObjectOwnerNetHandle, FNetRefHandle WantedNetHandle, FReplicationBridgeSerializationContext& Context);
 
+	/** Invoked right before we apply the state for a new received subobject but after we have applied state for owning/root object in order to behave like old replication system */
+	IRISCORE_API virtual void SubObjectCreatedFromReplication(FNetRefHandle SubObjectRefHandle);
+
 	/** Invoke after we have applied the initial state for an object.*/
 	IRISCORE_API virtual void PostApplyInitialState(FNetRefHandle Handle);
 
@@ -261,6 +265,7 @@ private:
 	void CallPreSendUpdateSingleHandle(FNetRefHandle Handle);
 	void CallUpdateInstancesWorldLocation();
 	bool CallWriteNetRefHandleCreationInfo(FReplicationBridgeSerializationContext& Context, FNetRefHandle Handle);
+	void CallSubObjectCreatedFromReplication(FNetRefHandle SubObjectHandle);
 	void CallPostApplyInitialState(FNetRefHandle Handle);
 	void CallPruneStaleObjects();
 	void CallGetInitialDependencies(FNetRefHandle Handle, FNetDependencyInfoArray& OutDependencies) const;
