Add check prior to calling the solver GetNet impulse functions - results in a substantial speedup for scenes with mnay joints.

diff --git a/Engine/Source/Runtime/Experimental/Chaos/Private/Chaos/Joint/PBDJointContainerSolver.cpp b/Engine/Source/Runtime/Experimental/Chaos/Private/Chaos/Joint/PBDJointContainerSolver.cpp
--- a/Engine/Source/Runtime/Experimental/Chaos/Private/Chaos/Joint/PBDJointContainerSolver.cpp
+++ b/Engine/Source/Runtime/Experimental/Chaos/Private/Chaos/Joint/PBDJointContainerSolver.cpp
@@ -299,7 +299,8 @@ namespace Chaos
 				Solver.ApplyConstraints(Dt, IterationStiffness, Settings, JointSettings);
 
 				// @todo(ccaulfield): We should be clamping the impulse at this point. Maybe move breaking to the solver
-				if (GetJointShouldBreak(JointSettings, Dt, Solver.GetNetLinearImpulse(), Solver.GetNetAngularImpulse()))
+				if ((JointSettings.LinearBreakForce != FLT_MAX || JointSettings.AngularBreakTorque != FLT_MAX) &&
+					GetJointShouldBreak(JointSettings, Dt, Solver.GetNetLinearImpulse(), Solver.GetNetAngularImpulse()))
 				{
 					Solver.SetIsBroken(true);
 				}
