Small updates to journaled cache tests.

#rnx

diff --git a/Engine/Source/Runtime/Experimental/IoStoreOnDemand/Private/JournaledCache.cpp b/Engine/Source/Runtime/Experimental/IoStoreOnDemand/Private/JournaledCache.cpp
--- a/Engine/Source/Runtime/Experimental/IoStoreOnDemand/Private/JournaledCache.cpp
+++ b/Engine/Source/Runtime/Experimental/IoStoreOnDemand/Private/JournaledCache.cpp
@@ -1846,13 +1846,19 @@ static void CacheTests(FSupport& Support)
 {
 	using namespace JournaledCache;
 
-	FCache::FConfig Config;
-	Config.Path = Support.TestDir / "cache_tests";
-	Config.MemoryQuota = uint32(512_Ki);
-	Config.DiskQuota = 8_Mi;
-	Config.JournalQuota = uint32(7_Ki);
-	Config.DropCache = true;
-	FCache Cache(MoveTemp(Config));
+	FCache* Cache = nullptr;
+	auto NewCache = [&Cache, &Support] (bool Drop=false)
+	{
+		delete Cache;
+		FCache::FConfig TestConfig;
+		TestConfig.Path = Support.TestDir / "cache_tests";
+		TestConfig.MemoryQuota = uint32(512_Ki);
+		TestConfig.DiskQuota = 8_Mi;
+		TestConfig.JournalQuota = uint32(7_Ki);
+		TestConfig.DropCache = Drop;
+		FCache::FConfig Config = TestConfig;
+		Cache = new FCache(MoveTemp(Config));
+	};
 
 	auto PrimePuts = [&] (int64 PutMax) {
 		TMap<uint64, FIoBuffer> Ret;
@@ -1865,7 +1871,7 @@ static void CacheTests(FSupport& Support)
 			}
 			FIoBuffer Data = Support.DummyData(Size);
 			uint64 Key = KeyGen(Data);
-			Cache.Put(Key, Data);
+			Cache->Put(Key, Data);
 			Ret.Add(Key, Data);
 		}
 		return Ret;
@@ -1874,17 +1880,17 @@ static void CacheTests(FSupport& Support)
 	uint32 WriteAllowance;
 
 	// no-op
+	NewCache(true);
 	WriteAllowance = uint32(1_Ki);
-	Cache.WriteMemToDisk(WriteAllowance);
-	Cache.Flush();
-	Cache.Reset();
+	check(Cache->WriteMemToDisk(WriteAllowance) == 0);
+	check(Cache->Flush() == 0);
 
+	NewCache(true);
 	WriteAllowance = uint32(512_Ki);
 	PrimePuts(WriteAllowance);
-	Cache.WriteMemToDisk(0);
-	Cache.WriteMemToDisk(WriteAllowance);
-	Cache.Flush();
-	Cache.Reset();
+	check(Cache->WriteMemToDisk(0) == 0);
+	check(Cache->WriteMemToDisk(WriteAllowance) > 0);
+	check(Cache->Flush() > 0);
 
 #if 0
 	auto Validate = [&] () {
