Removed the Reset() path.

Too fragile, and used only by my tests which have changed to create a new cache object instead of trying to reuse a single one (for hardiness).

#rnx

diff --git a/Engine/Source/Runtime/Experimental/IoStoreOnDemand/Private/JournaledCache.cpp b/Engine/Source/Runtime/Experimental/IoStoreOnDemand/Private/JournaledCache.cpp
--- a/Engine/Source/Runtime/Experimental/IoStoreOnDemand/Private/JournaledCache.cpp
+++ b/Engine/Source/Runtime/Experimental/IoStoreOnDemand/Private/JournaledCache.cpp
@@ -135,7 +135,6 @@ public:
 	using PeelItems = ItemArray;
 
 					FMemCache(uint32 InMaxSize=64 << 10);
-	void			Reset();
 	uint32			GetDemand() const;
 	uint32			GetCount() const	{ return Items.Num(); }
 	uint32			GetUsed() const		{ return UsedSize; }
@@ -161,13 +160,6 @@ FMemCache::FMemCache(uint32 InMaxSize)
 {
 }
 
-////////////////////////////////////////////////////////////////////////////////
-void FMemCache::Reset()
-{
-	Items.Reset();
-	UsedSize = 0;
-}
-
 ////////////////////////////////////////////////////////////////////////////////
 uint32 FMemCache::GetDemand() const
 {
@@ -434,7 +426,6 @@ class FDiskJournal
 public:
 							FDiskJournal(FStringView InRootPath, uint32 InMaxSize);
 	uint32					GetAilments() const;
-	void					Reset();
 	void					Drop();
 	int32					Flush();
 	FDiskPhrase				OpenPhrase(uint32 DataSize);
@@ -477,14 +468,6 @@ uint32 FDiskJournal::GetAilments() const
 	return Ret;
 }
 
-////////////////////////////////////////////////////////////////////////////////
-void FDiskJournal::Reset()
-{
-	// Marker = 0; // We'll just lets this roll along in its own little world
-	Cursor = 0;
-	Entries.Reset();
-}
-
 ////////////////////////////////////////////////////////////////////////////////
 void FDiskJournal::Drop()
 {
@@ -496,7 +479,9 @@ void FDiskJournal::Drop()
 	IPlatformFile& Ipf = IPlatformFile::GetPlatformPhysical();
 	Ipf.DeleteFile(*JrnPath);
 
-	Reset();
+	Cursor = 0;
+	Entries.Reset();
+
 	OpenJrnFile();
 }
 
@@ -600,7 +585,6 @@ class FDiskCache
 public:
 							FDiskCache(FString&& Path, uint64 InMaxDataSize, uint32 InJournalSize);
 	uint32					GetAilments() const;
-	void					Reset();
 	FDiskPhrase				OpenPhrase(uint32 DataSize);
 	void					ClosePhrase(FDiskPhrase&& Phrase);
 	EntryHandle				Get(uint64 Key) const;
@@ -627,13 +611,11 @@ private:
 	void					Prune(uint64 DataBase, uint32 Size);
 	FString					BinPath;
 	FDataMap				DataMap;
-	uint64					MappedBytes;
-
+	uint64					MappedBytes = 0;
 	uint64					MaxDataSize;
-	uint64					DataCursor;
+	uint64					DataCursor = 0;
 	TUniquePtr<IFileHandle>	DataHandle;
-	uint32					OverRemoval;
-
+	uint32					OverRemoval = 0;
 	FDiskJournal			Journal;
 };
 
@@ -647,7 +629,6 @@ FDiskCache::FDiskCache(FString&& Path, uint64 InMaxDataSize, uint32 InJournalSiz
 	InJournalSize = Journal.GetMaxSize();
 	MaxDataSize = (MaxDataSize - InJournalSize) & ~((1ull << 20) - 1);
 
-	Reset();
 	OpenDataFile();
 	
 	FOnDemandIoBackendStats::Get()->OnCacheSetMaxBytes(MaxDataSize);
@@ -664,16 +645,6 @@ uint32 FDiskCache::GetAilments() const
 	return Ret;
 }
 
-////////////////////////////////////////////////////////////////////////////////
-void FDiskCache::Reset()
-{
-	DataMap.Reset();
-	Journal.Reset();
-	MappedBytes = 0;
-	DataCursor = 0;
-	OverRemoval = 0;
-}
-
 ////////////////////////////////////////////////////////////////////////////////
 void FDiskCache::OpenDataFile()
 {
@@ -850,7 +821,11 @@ void FDiskCache::Drop()
 
 	Journal.Drop();
 
-	Reset();
+	DataCursor = 0;
+	OverRemoval = 0;
+	MappedBytes = 0;
+	DataMap.Reset();
+
 	OpenDataFile();
 }
 
@@ -1170,7 +1145,6 @@ public:
 
 					FCache(FConfig&& Config);
 	uint32			GetAilments() const;
-	void			Reset();
 	bool			Load();
 	uint32			GetDemand() const;
 	FEntry			Get(uint64 Key) const;
@@ -1225,14 +1199,6 @@ uint32 FCache::GetAilments() const
 	return Ret;
 }
 
-////////////////////////////////////////////////////////////////////////////////
-void FCache::Reset()
-{
-	FWriteAccess _[] = { MemLock, FsLock };
-	MemCache.Reset();
-	DiskCache.Reset();
-}
-
 ////////////////////////////////////////////////////////////////////////////////
 bool FCache::Load()
 {
