Chaos - fix uninitialized setting in FChaosSolverConfiguration

diff --git a/Engine/Source/Runtime/Experimental/Chaos/Private/ChaosSolverConfiguration.cpp b/Engine/Source/Runtime/Experimental/Chaos/Private/ChaosSolverConfiguration.cpp
--- a/Engine/Source/Runtime/Experimental/Chaos/Private/ChaosSolverConfiguration.cpp
+++ b/Engine/Source/Runtime/Experimental/Chaos/Private/ChaosSolverConfiguration.cpp
@@ -13,6 +13,7 @@ FChaosSolverConfiguration::FChaosSolverConfiguration()
 	, CollisionMarginMax(Chaos::FPBDRigidsEvolutionGBF::DefaultCollisionMarginMax)
 	, CollisionCullDistance(Chaos::FPBDRigidsEvolutionGBF::DefaultCollisionCullDistance)
 	, CollisionMaxPushOutVelocity(Chaos::FPBDRigidsEvolutionGBF::DefaultCollisionMaxPushOutVelocity)
+	, CollisionInitialOverlapDepenetrationVelocity(Chaos::FPBDRigidsEvolutionGBF::DefaultCollisionDepenetrationVelocity)
 	, ClusterConnectionFactor(1.0f)
 	, ClusterUnionConnectionType(EClusterUnionMethod::DelaunayTriangulation)
 	, bGenerateCollisionData(false)
diff --git a/Engine/Source/Runtime/Experimental/Chaos/Public/Chaos/PBDRigidsEvolutionGBF.h b/Engine/Source/Runtime/Experimental/Chaos/Public/Chaos/PBDRigidsEvolutionGBF.h
--- a/Engine/Source/Runtime/Experimental/Chaos/Public/Chaos/PBDRigidsEvolutionGBF.h
+++ b/Engine/Source/Runtime/Experimental/Chaos/Public/Chaos/PBDRigidsEvolutionGBF.h
@@ -56,7 +56,7 @@ namespace Chaos
 		using FExternalForces = FPerParticleExternalForces;
 		using FJointConstraints = FPBDJointConstraints;
 
-		// Default iteration counts
+		// Default settings for FChaosSolverConfiguration
 		static constexpr int32 DefaultNumPositionIterations = 8;
 		static constexpr int32 DefaultNumVelocityIterations = 2;
 		static constexpr int32 DefaultNumProjectionIterations = 1;
@@ -64,6 +64,7 @@ namespace Chaos
 		static constexpr FRealSingle DefaultCollisionMarginMax = 10.0f;
 		static constexpr FRealSingle DefaultCollisionCullDistance = 3.0f;
 		static constexpr FRealSingle DefaultCollisionMaxPushOutVelocity = 1000.0f;
+		static constexpr FRealSingle DefaultCollisionDepenetrationVelocity = 0.0f;
 		static constexpr int32 DefaultRestitutionThreshold = 1000;
 
 		CHAOS_API FPBDRigidsEvolutionGBF(
