Make the GeometryCollection Dataflow node for Flatten skip all non-clusters in the selection, rather than converting them to select the parent cluster. (This allows one to e.g. select all level 1 nodes and flatten without risk of flattening level 0 if there are leaves on level 1.)

+ apply the same change to the Flatten tool in fracture mode for consistency + because it also makes sense there

diff --git a/Engine/Plugins/Experimental/ChaosEditor/Source/FractureEditor/Private/FractureToolClustering.cpp b/Engine/Plugins/Experimental/ChaosEditor/Source/FractureEditor/Private/FractureToolClustering.cpp
--- a/Engine/Plugins/Experimental/ChaosEditor/Source/FractureEditor/Private/FractureToolClustering.cpp
+++ b/Engine/Plugins/Experimental/ChaosEditor/Source/FractureEditor/Private/FractureToolClustering.cpp
@@ -55,7 +55,8 @@ void UFractureToolFlattenAll::Execute(TWeakPtr<FFractureEditorModeToolkit> InToo
 			}
 			const TManagedArray<int32>& Levels = Context.GetGeometryCollection()->GetAttribute<int32>("Level", FGeometryCollection::TransformGroup);
 
-			Context.ConvertSelectionToClusterNodes();
+			Context.Sanitize();
+			Context.FilterSelectionBySimulationType(FGeometryCollection::FST_Clustered);
 
 			int32 MaxClusterLevel = -1;
 			for (int32 ClusterIndex : Context.GetSelection())
diff --git a/Engine/Plugins/Experimental/ChaosEditor/Source/FractureEditor/Private/FractureToolContext.cpp b/Engine/Plugins/Experimental/ChaosEditor/Source/FractureEditor/Private/FractureToolContext.cpp
--- a/Engine/Plugins/Experimental/ChaosEditor/Source/FractureEditor/Private/FractureToolContext.cpp
+++ b/Engine/Plugins/Experimental/ChaosEditor/Source/FractureEditor/Private/FractureToolContext.cpp
@@ -5,6 +5,8 @@
 #include "GeometryCollection/GeometryCollectionClusteringUtility.h"
 #include "GeometryCollection/GeometryCollectionUtility.h"
 
+#include "GeometryCollection/Facades/CollectionTransformSelectionFacade.h"
+
 #include "Algo/RemoveIf.h"
 
 
@@ -302,6 +304,12 @@ void FFractureToolContext::ConvertSelectionToClusterNodes()
 	Sanitize();
 }
 
+void FFractureToolContext::FilterSelectionBySimulationType(FGeometryCollection::ESimulationTypes KeepSimulationType)
+{
+	GeometryCollection::Facades::FCollectionTransformSelectionFacade SelectionFacade(*GeometryCollection);
+	SelectionFacade.FilterSelectionBySimulationType(SelectedBones, KeepSimulationType);
+}
+
 
 bool FFractureToolContext::HasSelectedAncestor(int32 Index) const
 {
diff --git a/Engine/Plugins/Experimental/ChaosEditor/Source/FractureEditor/Private/FractureToolContext.h b/Engine/Plugins/Experimental/ChaosEditor/Source/FractureEditor/Private/FractureToolContext.h
--- a/Engine/Plugins/Experimental/ChaosEditor/Source/FractureEditor/Private/FractureToolContext.h
+++ b/Engine/Plugins/Experimental/ChaosEditor/Source/FractureEditor/Private/FractureToolContext.h
@@ -24,6 +24,7 @@ public:
 	void ConvertSelectionToRigidNodes();
 	void ConvertSelectionToEmbeddedGeometryNodes();
 	void ConvertSelectionToClusterNodes();
+	void FilterSelectionBySimulationType(FGeometryCollection::ESimulationTypes KeepSimulationType);
 	void ConvertSelectionToLevel(int32 TargetLevel, bool bAllowSelectChildren, bool bAllowSelectParents);
 	void RemoveChildrenOfSelectedNodesFromSelection();
 	void ConvertEmbeddedSelectionToParents();
diff --git a/Engine/Plugins/Experimental/GeometryCollectionPlugin/Source/GeometryCollectionNodes/Private/Dataflow/GeometryCollectionClusteringNodes.cpp b/Engine/Plugins/Experimental/GeometryCollectionPlugin/Source/GeometryCollectionNodes/Private/Dataflow/GeometryCollectionClusteringNodes.cpp
--- a/Engine/Plugins/Experimental/GeometryCollectionPlugin/Source/GeometryCollectionNodes/Private/Dataflow/GeometryCollectionClusteringNodes.cpp
+++ b/Engine/Plugins/Experimental/GeometryCollectionPlugin/Source/GeometryCollectionNodes/Private/Dataflow/GeometryCollectionClusteringNodes.cpp
@@ -114,7 +114,8 @@ void FClusterFlattenDataflowNode::Evaluate(Dataflow::FContext& Context, const FD
 					const FDataflowTransformSelection& InTransformSelection = GetValue(Context, &OptionalTransformSelection);
 					ToFlatten = InTransformSelection.AsArray();
 					GeometryCollection::Facades::FCollectionTransformSelectionFacade SelectionFacade(*GeomCollection);
-					SelectionFacade.ConvertSelectionToClusterNodes(ToFlatten, false);
+					SelectionFacade.Sanitize(ToFlatten);
+					SelectionFacade.FilterSelectionBySimulationType(ToFlatten, FGeometryCollection::FST_Clustered);
 				}
 				else
 				{
diff --git a/Engine/Source/Runtime/Experimental/Chaos/Private/GeometryCollection/Facades/CollectionTransformSelectionFacade.cpp b/Engine/Source/Runtime/Experimental/Chaos/Private/GeometryCollection/Facades/CollectionTransformSelectionFacade.cpp
--- a/Engine/Source/Runtime/Experimental/Chaos/Private/GeometryCollection/Facades/CollectionTransformSelectionFacade.cpp
+++ b/Engine/Source/Runtime/Experimental/Chaos/Private/GeometryCollection/Facades/CollectionTransformSelectionFacade.cpp
@@ -258,6 +258,16 @@ namespace GeometryCollection::Facades
 		}
 	}
 
+	void FCollectionTransformSelectionFacade::FilterSelectionBySimulationType(TArray<int32>& InOutSelection, FGeometryCollection::ESimulationTypes KeepSimulationType) const
+	{
+		const TManagedArray<int32>& SimulationType = SimulationTypeAttribute.Get();
+
+		InOutSelection.SetNum(Algo::RemoveIf(InOutSelection, [&](int32 BoneIdx)
+			{
+				return SimulationType[BoneIdx] != KeepSimulationType;
+			}));
+	}
+
 	void FCollectionTransformSelectionFacade::ConvertSelectionToClusterNodes(TArray<int32>& InOutSelection, bool bLeaveRigidRoots) const
 	{
 		const TManagedArray<int32>& Parents = ParentAttribute.Get();
diff --git a/Engine/Source/Runtime/Experimental/Chaos/Public/GeometryCollection/Facades/CollectionTransformSelectionFacade.h b/Engine/Source/Runtime/Experimental/Chaos/Public/GeometryCollection/Facades/CollectionTransformSelectionFacade.h
--- a/Engine/Source/Runtime/Experimental/Chaos/Public/GeometryCollection/Facades/CollectionTransformSelectionFacade.h
+++ b/Engine/Source/Runtime/Experimental/Chaos/Public/GeometryCollection/Facades/CollectionTransformSelectionFacade.h
@@ -7,6 +7,7 @@
 #include "GeometryCollection/ManagedArrayAccessor.h"
 #include "GeometryCollection/ManagedArrayCollection.h"
 #include "GeometryCollection/Facades/CollectionSelectionFacade.h"
+#include "GeometryCollection/GeometryCollection.h"
 #include "Chaos/Triangle.h"
 
 class FGeometryCollection;
@@ -64,6 +65,9 @@ namespace GeometryCollection::Facades
 		CHAOS_API void ConvertSelectionToRigidNodes(const int32 Index, TArray<int32>& InSelection) const;
 		CHAOS_API void ConvertSelectionToRigidNodes(TArray<int32>& InSelection) const;
 
+		/** Keep only the bones in the selection that match the input SimulationType */
+		CHAOS_API void FilterSelectionBySimulationType(TArray<int32>& InOutSelection, FGeometryCollection::ESimulationTypes KeepSimulationType) const;
+
 		/** Replace any Embedded or Rigid transforms in InOutSelection with their parent cluster. If there is just a rigid root node selected, optionally keep it if bLeaveRigidRoot is true. */
 		CHAOS_API void ConvertSelectionToClusterNodes(TArray<int32>& InOutSelection, bool bLeaveRigidRoots = true) const;
 
