Fix potential crash if there's no next result when GC proxy is pulling data from physics state

diff --git a/Engine/Source/Runtime/Experimental/Chaos/Private/PhysicsProxy/GeometryCollectionPhysicsProxy.cpp b/Engine/Source/Runtime/Experimental/Chaos/Private/PhysicsProxy/GeometryCollectionPhysicsProxy.cpp
--- a/Engine/Source/Runtime/Experimental/Chaos/Private/PhysicsProxy/GeometryCollectionPhysicsProxy.cpp
+++ b/Engine/Source/Runtime/Experimental/Chaos/Private/PhysicsProxy/GeometryCollectionPhysicsProxy.cpp
@@ -3637,7 +3637,6 @@ bool FGeometryCollectionPhysicsProxy::PullFromPhysicsState(const Chaos::FDirtyGe
 		if (bIsCollectionDirty)
 		{
 			const TBitArray<>& PrevResultsModifiedIndices = PullData.Results().GetModifiedTransformIndices();
-			const FGeometryCollectionResults& NextResults = NextPullData->Results();
 
 			const bool bIsComponentTransformScaled = !Parameters.WorldTransform.GetScale3D().Equals(FVector::OneVector);
 			const FTransform ComponentScaleTransform(FQuat::Identity, FVector::ZeroVector, Parameters.WorldTransform.GetScale3D());
