Chaos : RemoveElementFrom and UpdateElementIn optimization
* RemoveElementFrom : make sure we do not try to remove from all the other acceleration structures if payload was removed from the first one
* UpdateElementIn : avoid allocation when trying to remove from other structure if the particle does not exists anymore

diff --git a/Engine/Source/Runtime/Experimental/Chaos/Public/Chaos/SpatialAccelerationCollection.h b/Engine/Source/Runtime/Experimental/Chaos/Public/Chaos/SpatialAccelerationCollection.h
--- a/Engine/Source/Runtime/Experimental/Chaos/Public/Chaos/SpatialAccelerationCollection.h
+++ b/Engine/Source/Runtime/Experimental/Chaos/Public/Chaos/SpatialAccelerationCollection.h
@@ -591,20 +591,21 @@ public:
 	{
 		const uint16 UseBucket = ((1 << SpatialIdx.Bucket) & this->ActiveBucketsMask) ? SpatialIdx.Bucket : 0;
 		bool bSuccess = Buckets[UseBucket].Objects[SpatialIdx.InnerIdx].Acceleration->RemoveElement(Payload);
-		//ensure(Success); // Debug check to see if SpatialIdx was what the caller expected it to be
-		// Make sure that we remove this Payload even if the SpatialIdx is wrong
-		VisitAllSpatialIndices(
-			[this, &Payload, &SpatialIdx, &bSuccess](FSpatialAccelerationIdx Idx)
-			{
-				if (!(Idx == SpatialIdx))
+		if (!bSuccess)
+		{
+			// Make sure that we remove this Payload even if the SpatialIdx is wrong
+			VisitAllSpatialIndices(
+				[this, &Payload, &SpatialIdx, &bSuccess](FSpatialAccelerationIdx Idx)
 				{
-					const uint16 Buckt = ((1 << Idx.Bucket) & this->ActiveBucketsMask) ? Idx.Bucket : 0;
-					const bool bRemoved = Buckets[Buckt].Objects[Idx.InnerIdx].Acceleration->RemoveElement(Payload);
-					check(!bSuccess || !bRemoved); // make sure we only remove from one acceleration structure only 
-					bSuccess |= bRemoved;
+					if (!(Idx == SpatialIdx))
+					{
+						const uint16 Buckt = ((1 << Idx.Bucket) & this->ActiveBucketsMask) ? Idx.Bucket : 0;
+						const bool bRemoved = Buckets[Buckt].Objects[Idx.InnerIdx].Acceleration->RemoveElement(Payload);
+						bSuccess |= bRemoved;
+					}
 				}
-			}
-		);
+			);
+		}
 		return bSuccess;
 	}
 
@@ -616,15 +617,16 @@ public:
 		// In case spatial index changed, remove this element in all other substructures
 		if (!bElementExisted)
 		{
-			for (FSpatialAccelerationIdx Idx : GetAllSpatialIndices())
-			{
-				if (!(Idx == SpatialIdx))
+			VisitAllSpatialIndices(
+				[this, &Payload, &SpatialIdx, &bElementExisted](FSpatialAccelerationIdx Idx)
 				{
-					const uint16 Buckt = ((1 << Idx.Bucket) & this->ActiveBucketsMask) ? Idx.Bucket : 0;
-					const bool Removed = Buckets[Buckt].Objects[Idx.InnerIdx].Acceleration->RemoveElement(Payload);
-					bElementExisted = bElementExisted || Removed;
-				}
-			}
+					if (!(Idx == SpatialIdx))
+					{
+						const uint16 Buckt = ((1 << Idx.Bucket) & this->ActiveBucketsMask) ? Idx.Bucket : 0;
+						const bool Removed = Buckets[Buckt].Objects[Idx.InnerIdx].Acceleration->RemoveElement(Payload);
+						bElementExisted = bElementExisted || Removed;
+					}
+				});
 		}
 		return bElementExisted;
 	}
