Chaos - add manifold point number limit to collision constraint (cvar, default unlimited)
- add manifold point limit to per-thread collision solver (cvar, default unlimited)
- make ScratchBuffer robust to OOM allocation failure

diff --git a/Engine/Source/Runtime/Experimental/Chaos/Private/Chaos/Collision/PBDCollisionConstraint.cpp b/Engine/Source/Runtime/Experimental/Chaos/Private/Chaos/Collision/PBDCollisionConstraint.cpp
--- a/Engine/Source/Runtime/Experimental/Chaos/Private/Chaos/Collision/PBDCollisionConstraint.cpp
+++ b/Engine/Source/Runtime/Experimental/Chaos/Private/Chaos/Collision/PBDCollisionConstraint.cpp
@@ -53,6 +53,11 @@ namespace Chaos
 	bool bChaos_Collision_EnableBoundsChecks = true;
 	FAutoConsoleVariableRef CVarChaos_Collision_EnableBoundsChecks(TEXT("p.Chaos.Collision.EnableBoundsChecks"), bChaos_Collision_EnableBoundsChecks, TEXT(""));
 
+	// Maximum number of manifold points per contact (-1 for unlimited)
+	int32 Chaos_Collision_MaxManifoldPoints = -1;
+	FAutoConsoleVariableRef CVarChaos_Collision_MaxManifoldPoints(TEXT("p.Chaos.Collision.MaxManifoldPoints"), Chaos_Collision_MaxManifoldPoints, TEXT(""));
+
+
 	struct FCollisionTolerances
 	{
 		// Multiplied by the contact margin to produce a distance within which contacts are considered to be the same point
@@ -1205,4 +1210,9 @@ namespace Chaos
 			return NoRestingDependency;
 		}
 	}
+
+	void FPBDCollisionConstraint::LogOneShotManifoldError(const int32 MaxManifoldPoints, const TArrayView<const FContactPoint>& ContactPoints)
+	{
+		UE_LOG(LogChaos, Error, TEXT("FPBDCollisionConstraint: exceeded per-constraint manifold point limit at %d / %d. Particles %s -:- %s"), ContactPoints.Num(), MaxManifoldPoints, *GetParticle0()->GetDebugName(), *GetParticle1()->GetDebugName());
+	}
 }
\ No newline at end of file
diff --git a/Engine/Source/Runtime/Experimental/Chaos/Private/Chaos/Collision/PBDCollisionContainerSolver.cpp b/Engine/Source/Runtime/Experimental/Chaos/Private/Chaos/Collision/PBDCollisionContainerSolver.cpp
--- a/Engine/Source/Runtime/Experimental/Chaos/Private/Chaos/Collision/PBDCollisionContainerSolver.cpp
+++ b/Engine/Source/Runtime/Experimental/Chaos/Private/Chaos/Collision/PBDCollisionContainerSolver.cpp
@@ -55,6 +55,10 @@ namespace Chaos
 		// Whether to enable the new initial overlap depentration system
 		bool bChaos_Collision_EnableInitialDepenetration = false;
 		FAutoConsoleVariableRef CVarChaosCollisionEnableInitialDepentration(TEXT("p.Chaos.PBDCollisionSolver.EnableInitialDepenetration"), bChaos_Collision_EnableInitialDepenetration, TEXT(""));
+
+		// The maximum number of constraints we will attempt to solve (-1 for unlimited)
+		int32 Chaos_Collision_MaxSolverManifoldPoints = -1;
+		FAutoConsoleVariableRef CVarChaosCollisionMaxSolverManifoldPoints(TEXT("p.Chaos.PBDCollisionSolver.MaxManifoldPoints"), Chaos_Collision_MaxSolverManifoldPoints, TEXT(""));
 	}
 
 
@@ -405,30 +409,58 @@ namespace Chaos
 	{
 		NumCollisionSolverManifoldPoints = 0;
 		MaxCollisionSolverManifoldPoints = 0;
+		NumCollisionSolvers = 0;
 		CollisionSolvers = nullptr;
 		CollisionSolverManifoldPoints = nullptr;
 
+		// We have one solver per constraint, unless we have too manu constraints...
+		NumCollisionSolvers = CollisionConstraints.Num();
+
 		// Count the manifold points
 		// @todo(chaos): can we avoid this?
-		const int32 NumCollisionsConstraints = CollisionConstraints.Num();
-		for (int32 ConstraintIndex = 0; ConstraintIndex < NumCollisionsConstraints; ++ConstraintIndex)
+		const int32 ManifoldPointLimit = CVars::Chaos_Collision_MaxSolverManifoldPoints;
+		for (int32 ConstraintIndex = 0; ConstraintIndex < CollisionConstraints.Num(); ++ConstraintIndex)
+		{
+			const int32 MaxManifoldPoints = CalculateConstraintMaxManifoldPoints(GetConstraint(ConstraintIndex));
+			
+			// Drop some of the constraints if we exceed some tunable maximum. This is purely to prevent massive 
+			// slowdowns or excessive scratch allocations when too many collisions are generated.
+			if ((ManifoldPointLimit >= 0) && (MaxCollisionSolverManifoldPoints + MaxManifoldPoints > ManifoldPointLimit))
+			{
+				// At this point something is assumed to have gone wrong, so this is an error condition.
+				UE_LOG(LogChaos, Error, TEXT("FPBDCollisionContainerSolver: exceeded solver manifold point limit %d at constraint %d of %d. This and remaining constraints will not be solved"), ManifoldPointLimit, ConstraintIndex, CollisionConstraints.Num());
+				NumCollisionSolvers = ConstraintIndex;
+				break;
+			}
+
+			MaxCollisionSolverManifoldPoints += MaxManifoldPoints;
+		}
+
+		// If we have no manifold points, there's no point creating any solvers
+		if (MaxCollisionSolverManifoldPoints == 0)
 		{
-			MaxCollisionSolverManifoldPoints += CalculateConstraintMaxManifoldPoints(GetConstraint(ConstraintIndex));
+			NumCollisionSolvers = 0;
 		}
 
 		// Set up the solver buffers
-		if (NumCollisionsConstraints > 0)
+		if (NumCollisionSolvers > 0)
 		{
 			// Resize the scratch buffer (up to 25% slack)
 			constexpr size_t AlignedSolverSize = Align(sizeof(Private::FPBDCollisionSolver), alignof(Private::FPBDCollisionSolver));
 			constexpr size_t AlignedPointSize = Align(sizeof(Private::FPBDCollisionSolverManifoldPoint), alignof(Private::FPBDCollisionSolverManifoldPoint));
-			const size_t ScratchSize = NumCollisionsConstraints * AlignedSolverSize + MaxCollisionSolverManifoldPoints * AlignedPointSize;
+			const size_t ScratchSize = NumCollisionSolvers * AlignedSolverSize + MaxCollisionSolverManifoldPoints * AlignedPointSize;
 			const size_t ScratchBufferSize = CalculateCollisionBufferNum(ScratchSize, Scratch.BufferSize());
 			Scratch.Reset(ScratchBufferSize);
+			
+			if (Scratch.BufferSize() == 0)
+			{
+				UE_LOG(LogChaos, Error, TEXT("FPBDCollisionContainerSolver: failed to allocate scratch buffer of size %lld bytes. NumCollisions=%d, NumManifoldPoints=%d. Collisions will be lost."), ScratchBufferSize, NumCollisionSolvers, MaxCollisionSolverManifoldPoints);
+				NumCollisionSolvers = 0;
+				return;
+			}
 
 			// Allocate scratch space for the collision solvers and manifold points
-			// NOTE: scratch return a valid pointer even for 0 size so CollisionSolverManifoldPoints is always a valid pointer in the scratch (but may have no space)
-			CollisionSolvers = Scratch.AllocArray<Private::FPBDCollisionSolver>(NumCollisionsConstraints);
+			CollisionSolvers = Scratch.AllocArray<Private::FPBDCollisionSolver>(NumCollisionSolvers);
 			CollisionSolverManifoldPoints = Scratch.AllocArray<Private::FPBDCollisionSolverManifoldPoint>(MaxCollisionSolverManifoldPoints);
 		}
 	}
@@ -468,10 +500,10 @@ namespace Chaos
 		// All constarints are now added. We can allocate the solver buffers.
 		PrepareSolverBuffer();
 
-		// Make sure have a valid manifold point buffer if we have constraints (it may be zero size, but we want the pointer to be valid)
-		check((CollisionSolverManifoldPoints != nullptr) || (GetNumConstraints() == 0));
+		// Make sure have a valid manifold point buffer if we have constraints
+		check((CollisionSolverManifoldPoints != nullptr) || (NumSolvers() == 0));
 
-		for (int32 ConstraintIndex = 0, ConstraintEndIndex = GetNumConstraints(); ConstraintIndex < ConstraintEndIndex; ++ConstraintIndex)
+		for (int32 ConstraintIndex = 0, ConstraintEndIndex = NumSolvers(); ConstraintIndex < ConstraintEndIndex; ++ConstraintIndex)
 		{
 			Private::FPBDCollisionSolver& CollisionSolver = GetSolver(ConstraintIndex);
 			FPBDCollisionConstraint* Constraint = GetConstraint(ConstraintIndex);
@@ -508,12 +540,16 @@ namespace Chaos
 		}
 	}
 
-	void FPBDCollisionContainerSolver::GatherInput(const FReal InDt, const int32 BeginIndex, const int32 EndIndex)
+	void FPBDCollisionContainerSolver::GatherInput(const FReal InDt, const int32 ConstraintBeginIndex, const int32 ConstraintEndIndex)
 	{
 		// NOTE: may be called in parallel. Should not change the container or any elements outside of [BeginIndex, EndIndex)
 
-		check(BeginIndex >= 0);
-		check(EndIndex <= NumSolvers());
+		check(ConstraintBeginIndex >= 0);
+		check(ConstraintEndIndex <= CollisionConstraints.Num());
+
+		// Handle the case where we dropped some constraints because there were too  many
+		const int32 BeginIndex = FMath::Min(ConstraintBeginIndex, NumCollisionSolvers);
+		const int32 EndIndex = FMath::Min(ConstraintEndIndex, NumCollisionSolvers);
 
 		const FSolverReal Dt = FSolverReal(InDt);
 
@@ -550,12 +586,16 @@ namespace Chaos
 		ScatterOutput(Dt, 0, NumSolvers());
 	}
 
-	void FPBDCollisionContainerSolver::ScatterOutput(const FReal InDt, const int32 BeginIndex, const int32 EndIndex)
+	void FPBDCollisionContainerSolver::ScatterOutput(const FReal InDt, const int32 ConstraintBeginIndex, const int32 ConstraintEndIndex)
 	{
 		// NOTE: may be called in parallel. Should not change the container or any elements outside of [BeginIndex, EndIndex)
 
-		check(BeginIndex >= 0);
-		check(EndIndex <= NumSolvers());
+		check(ConstraintBeginIndex >= 0);
+		check(ConstraintEndIndex <= CollisionConstraints.Num());
+
+		// Handle the case where we dropped some constraints because there were too  many
+		const int32 BeginIndex = FMath::Min(ConstraintBeginIndex, NumCollisionSolvers);
+		const int32 EndIndex = FMath::Min(ConstraintEndIndex, NumCollisionSolvers);
 
 		const FSolverReal Dt = FSolverReal(InDt);
 
diff --git a/Engine/Source/Runtime/Experimental/Chaos/Private/Chaos/GeometryParticles.cpp b/Engine/Source/Runtime/Experimental/Chaos/Private/Chaos/GeometryParticles.cpp
--- a/Engine/Source/Runtime/Experimental/Chaos/Private/Chaos/GeometryParticles.cpp
+++ b/Engine/Source/Runtime/Experimental/Chaos/Private/Chaos/GeometryParticles.cpp
@@ -13,6 +13,16 @@ namespace Chaos
 {
 	extern void UpdateShapesArrayFromGeometry(FShapeInstanceArray& ShapesArray, const FImplicitObjectPtr& Geometry, const FRigidTransform3& ActorTM);
 
+	namespace Private
+	{
+		// The name shown for particoles that have not had their DebugName set
+#if CHAOS_DEBUG_NAME
+		FString EmptyParticleName = TEXT("<NotNamed>");
+#else
+		FString EmptyParticleName = TEXT("<NotSupported>");
+#endif
+	}
+
 	FShapeOrShapesArray::FShapeOrShapesArray(const FGeometryParticleHandle* Particle)
 	{
 		if (Particle)
diff --git a/Engine/Source/Runtime/Experimental/Chaos/Public/Chaos/Collision/PBDCollisionConstraint.h b/Engine/Source/Runtime/Experimental/Chaos/Public/Chaos/Collision/PBDCollisionConstraint.h
--- a/Engine/Source/Runtime/Experimental/Chaos/Public/Chaos/Collision/PBDCollisionConstraint.h
+++ b/Engine/Source/Runtime/Experimental/Chaos/Public/Chaos/Collision/PBDCollisionConstraint.h
@@ -39,6 +39,8 @@ namespace Chaos
 
 	CHAOS_API bool ContactConstraintSortPredicate(const FPBDCollisionConstraint& L, const FPBDCollisionConstraint& R);
 
+	extern int32 Chaos_Collision_MaxManifoldPoints;
+
 	/*
 	 * @brief Material properties for a collision constraint
 	*/
@@ -556,7 +558,16 @@ namespace Chaos
 
 			FReal MinPhi = TNumericLimits<FReal>::Max();
 			
-			const int32 NumContacts = ContactPoints.Num();
+			int32 NumContacts = ContactPoints.Num();
+
+			// If we have too many manifold points, clip the array. This is considered and error but
+			// is important for avoiding OOM or massive slowdowns when something goes wrong
+			const int32 MaxManifoldPoints = Chaos_Collision_MaxManifoldPoints;
+			if ((MaxManifoldPoints >= 0) && (NumContacts > MaxManifoldPoints))
+			{
+				NumContacts = MaxManifoldPoints;
+				LogOneShotManifoldError(MaxManifoldPoints, ContactPoints);
+			}
 			
 			ManifoldPoints.Reserve(NumContacts);
 
@@ -845,6 +856,8 @@ namespace Chaos
 	private:
 		CHAOS_API FReal CalculateSavedManifoldDistanceSq(const FSavedManifoldPoint& SavedManifoldPoint, const FManifoldPoint& ManifoldPoint, const FReal DistanceToleranceSq) const;
 
+		CHAOS_API void LogOneShotManifoldError(const int32 MaxManifoldPoints, const TArrayView<const FContactPoint>& ContactPoints);
+
 		union FFlags
 		{
 			FFlags() : Bits(0) {}
diff --git a/Engine/Source/Runtime/Experimental/Chaos/Public/Chaos/Collision/PBDCollisionContainerSolver.h b/Engine/Source/Runtime/Experimental/Chaos/Public/Chaos/Collision/PBDCollisionContainerSolver.h
--- a/Engine/Source/Runtime/Experimental/Chaos/Public/Chaos/Collision/PBDCollisionContainerSolver.h
+++ b/Engine/Source/Runtime/Experimental/Chaos/Public/Chaos/Collision/PBDCollisionContainerSolver.h
@@ -27,7 +27,7 @@ namespace Chaos
 		FPBDCollisionContainerSolver(const FPBDCollisionConstraints& InConstraintContainer, const int32 InPriority);
 		~FPBDCollisionContainerSolver();
 
-		int32 NumSolvers() const { return CollisionConstraints.Num(); }
+		int32 NumSolvers() const { return NumCollisionSolvers; }
 
 		virtual void Reset(const int32 InMaxCollisions) override final;
 
@@ -80,6 +80,7 @@ namespace Chaos
 
 		// The start of the solver array in the scratch buffer
 		Private::FPBDCollisionSolver* CollisionSolvers;
+		int32 NumCollisionSolvers;
 
 		// The start of the manifold points array in the scratch buffer
 		Private::FPBDCollisionSolverManifoldPoint* CollisionSolverManifoldPoints;
diff --git a/Engine/Source/Runtime/Experimental/Chaos/Public/Chaos/Framework/ScratchBuffer.h b/Engine/Source/Runtime/Experimental/Chaos/Public/Chaos/Framework/ScratchBuffer.h
--- a/Engine/Source/Runtime/Experimental/Chaos/Public/Chaos/Framework/ScratchBuffer.h
+++ b/Engine/Source/Runtime/Experimental/Chaos/Public/Chaos/Framework/ScratchBuffer.h
@@ -111,12 +111,17 @@ namespace Chaos
 				if (InMaxBytes != BufferSize())
 				{
 					DestroyBuffer();
+
 					if (InMaxBytes > 0)
 					{
 						// Over-allocate and store a sentinel after the block
 						BufferBegin = new uint8[InMaxBytes + sizeof(SentinelValue)];
 					}
-					BufferEnd = BufferBegin + InMaxBytes;
+
+					if (BufferBegin != nullptr)
+					{
+						BufferEnd = BufferBegin + InMaxBytes;
+					}
 
 					InitSentinel();
 				}
@@ -126,15 +131,16 @@ namespace Chaos
 
 			void DestroyBuffer()
 			{
+				CheckSentinel();
+
 				if (BufferBegin != nullptr)
 				{
-					CheckSentinel();
-
 					delete[] BufferBegin;
-					BufferBegin = nullptr;
-					BufferEnd = nullptr;
-					BufferNext = nullptr;
 				}
+
+				BufferBegin = nullptr;
+				BufferEnd = nullptr;
+				BufferNext = nullptr;
 			}
 
 			size_t* Sentinel()
diff --git a/Engine/Source/Runtime/Experimental/Chaos/Public/Chaos/GeometryParticles.h b/Engine/Source/Runtime/Experimental/Chaos/Public/Chaos/GeometryParticles.h
--- a/Engine/Source/Runtime/Experimental/Chaos/Public/Chaos/GeometryParticles.h
+++ b/Engine/Source/Runtime/Experimental/Chaos/Public/Chaos/GeometryParticles.h
@@ -41,6 +41,8 @@ namespace Chaos
 	namespace Private
 	{
 		class FPBDIslandParticle;
+
+		CHAOS_API extern FString EmptyParticleName;
 	}
 
 	/**
@@ -432,6 +434,17 @@ namespace Chaos
 		}
 #endif
 
+		const FString& GetDebugName(const int32 Index) const
+		{
+#if CHAOS_DEBUG_NAME
+			if (MDebugName[Index].IsValid())
+			{
+				return *(MDebugName[Index].Get());
+			}
+#endif
+			return Private::EmptyParticleName;
+		}
+
 		const TAABB<T, d>& WorldSpaceInflatedBounds(const int32 Index) const
 		{
 			return MWorldSpaceInflatedBounds[Index];
diff --git a/Engine/Source/Runtime/Experimental/Chaos/Public/Chaos/ParticleHandle.h b/Engine/Source/Runtime/Experimental/Chaos/Public/Chaos/ParticleHandle.h
--- a/Engine/Source/Runtime/Experimental/Chaos/Public/Chaos/ParticleHandle.h
+++ b/Engine/Source/Runtime/Experimental/Chaos/Public/Chaos/ParticleHandle.h
@@ -621,7 +621,8 @@ public:
 	const TSharedPtr<FString, ESPMode::ThreadSafe>& DebugName() const { return GeometryParticles->DebugName(ParticleIdx); }
 	void SetDebugName(const TSharedPtr<FString, ESPMode::ThreadSafe>& InDebugName) { GeometryParticles->DebugName(ParticleIdx) = InDebugName; }
 #endif
-	
+	const FString& GetDebugName() const { return GeometryParticles->GetDebugName(ParticleIdx); }
+
 	EObjectStateType ObjectState() const;
 
 	TGeometryParticle<T, d>* GTGeometryParticle() const { return GeometryParticles->GTGeometryParticle(ParticleIdx); }
@@ -2213,6 +2214,12 @@ public:
 	}
 #endif
 
+	const FString& GetDebugName() const
+	{
+		return MHandle->GetDebugName();
+	}
+
+
 	FORCEINLINE bool IsInConstraintGraph() const
 	{
 		return MHandle->IsInConstraintGraph();
