Chaos - reduce mesh contact edge normal threshold to 5 deg (from 10) to reduce edge collision artefacts

diff --git a/Engine/Source/Runtime/Experimental/Chaos/Private/Chaos/CollisionOneShotManifolds.cpp b/Engine/Source/Runtime/Experimental/Chaos/Private/Chaos/CollisionOneShotManifolds.cpp
--- a/Engine/Source/Runtime/Experimental/Chaos/Private/Chaos/CollisionOneShotManifolds.cpp
+++ b/Engine/Source/Runtime/Experimental/Chaos/Private/Chaos/CollisionOneShotManifolds.cpp
@@ -88,7 +88,7 @@ namespace Chaos
 	bool bChaos_Collision_EnableLargeMeshManifolds = 1;
 	FAutoConsoleVariableRef CVarChaos_Collision_EnableLargeMeshManifolds(TEXT("p.Chaos.Collision.EnableLargeMeshManifolds"), bChaos_Collision_EnableLargeMeshManifolds, TEXT("Whether to allow large mesh manifolds for collisions against meshes (required for good behaviour)"));
 
-	FRealSingle Chaos_Collision_MeshContactNormalThreshold = 0.98f;	// ~10deg
+	FRealSingle Chaos_Collision_MeshContactNormalThreshold = 0.996f;	// ~5deg
 	FAutoConsoleVariableRef CVarChaos_Collision_MeshContactNormalThreshold(TEXT("p.Chaos.Collision.MeshContactNormalThreshold"), Chaos_Collision_MeshContactNormalThreshold, TEXT("Treat contact with a dot product between the normal and the triangle face greater than this as face collisions"));
 
 	bool bChaos_Collision_MeshManifoldSortByDistance = false;
