Geometry collection : Optimization : Make sure attribute accessor Num() do not always lookup the group Map

diff --git a/Engine/Source/Runtime/Experimental/Chaos/Public/GeometryCollection/ManagedArrayAccessor.h b/Engine/Source/Runtime/Experimental/Chaos/Public/GeometryCollection/ManagedArrayAccessor.h
--- a/Engine/Source/Runtime/Experimental/Chaos/Public/GeometryCollection/ManagedArrayAccessor.h
+++ b/Engine/Source/Runtime/Experimental/Chaos/Public/GeometryCollection/ManagedArrayAccessor.h
@@ -171,7 +171,9 @@ public:
 
 	int32 Num() const
 	{
-		return ConstCollection.NumElements(Group);
+		return ConstAttributeArray
+			? ConstAttributeArray->Num()
+			: ConstCollection.NumElements(Group); // more expensive to fetch
 	}
 
 private:
