geometry collection : remove MassToLocal from dynamic collections since it is never modified from the asset
- using the RestCollection immutable MassToLocal instead
- this save a 2x96 bytes per bone per geometry collection instance

diff --git a/Engine/Plugins/Experimental/ChaosCaching/Source/ChaosCaching/Private/Chaos/Adapters/GeometryCollectionComponentCacheAdapter.cpp b/Engine/Plugins/Experimental/ChaosCaching/Source/ChaosCaching/Private/Chaos/Adapters/GeometryCollectionComponentCacheAdapter.cpp
--- a/Engine/Plugins/Experimental/ChaosCaching/Source/ChaosCaching/Private/Chaos/Adapters/GeometryCollectionComponentCacheAdapter.cpp
+++ b/Engine/Plugins/Experimental/ChaosCaching/Source/ChaosCaching/Private/Chaos/Adapters/GeometryCollectionComponentCacheAdapter.cpp
@@ -87,7 +87,7 @@ namespace Chaos
 		}
 
 		FGeometryDynamicCollection&            Collection       = Proxy->GetPhysicsCollection();
-		const TManagedArray<FTransform>&       MassToLocal		= Collection.MassToLocal;
+		const TManagedArray<FTransform>&       MassToLocal      = RestCollection->GetAttribute<FTransform>("MassToLocal", FGeometryCollection::TransformGroup);
 		const TManagedArray<int32>&            Parents			= RestCollection->Parent;
 		const TManagedArray<TSet<int32>>&	   Children         = RestCollection->Children;
 		const TArray<FBreakingData>&		   Breaks           = Solver->GetEvolution()->GetRigidClustering().GetAllClusterBreakings();
@@ -262,7 +262,7 @@ namespace Chaos
 		}
 
 		FGeometryDynamicCollection&      Collection       = Proxy->GetPhysicsCollection();
-		const TManagedArray<FTransform>& MassToLocal	  = Collection.MassToLocal;
+		const TManagedArray<FTransform>& MassToLocal      = RestCollection->GetAttribute<FTransform>("MassToLocal", FGeometryCollection::TransformGroup);
 		TArray<FClusterParticle*>        Particles        = Proxy->GetParticles();
 
 		FCacheEvaluationContext Context(TickRecord);
diff --git a/Engine/Source/Programs/HeadlessChaos/Private/GeometryCollection/GeometryCollectionTestCollisionResolution.cpp b/Engine/Source/Programs/HeadlessChaos/Private/GeometryCollection/GeometryCollectionTestCollisionResolution.cpp
--- a/Engine/Source/Programs/HeadlessChaos/Private/GeometryCollection/GeometryCollectionTestCollisionResolution.cpp
+++ b/Engine/Source/Programs/HeadlessChaos/Private/GeometryCollection/GeometryCollectionTestCollisionResolution.cpp
@@ -28,9 +28,14 @@ namespace GeometryCollectionTest
 			Params.NestedTransforms = { FTransform::Identity, FTransform::Identity, FTransform::Identity };
 			
 			Collection = TNewSimulationObject<GeometryType::GeometryCollectionWithSingleRigid>::Init(Params)->template As<FGeometryCollectionWrapper>();
+
+			const TManagedArray<FTransform>& MassToLocal = Collection->RestCollection->GetAttribute<FTransform>("MassToLocal", FGeometryCollection::TransformGroup);
+
 			EXPECT_EQ(Collection->DynamicCollection->Parent[0], 1); // is a child of index one
-			EXPECT_TRUE(Collection->DynamicCollection->MassToLocal[0].Equals(FTransform::Identity)); // we are not testing MassToLocal in this test
+			EXPECT_TRUE(MassToLocal[0].Equals(FTransform::Identity)); // we are not testing MassToLocal in this test
+			
 			
+
 			UnitTest.AddSimulationObject(Collection);
 		}
 
@@ -72,8 +77,11 @@ namespace GeometryCollectionTest
 			FVector Scale(Radius);
 			Params.GeomTransform.SetScale3D(Scale); // Sphere radius
 			Collection = TNewSimulationObject<GeometryType::GeometryCollectionWithSingleRigid>::Init(Params)->template As<FGeometryCollectionWrapper>();
+
+			const TManagedArray<FTransform>& MassToLocal = Collection->RestCollection->GetAttribute<FTransform>("MassToLocal", FGeometryCollection::TransformGroup);
+
 			EXPECT_EQ(Collection->DynamicCollection->Parent[0], 1); // is a child of index one
-			EXPECT_TRUE(Collection->DynamicCollection->MassToLocal[0].Equals(FTransform::Identity)); // we are not testing MassToLocal in this test
+			EXPECT_TRUE(MassToLocal[0].Equals(FTransform::Identity)); // we are not testing MassToLocal in this test
 
 			UnitTest.AddSimulationObject(Collection);
 		}
@@ -283,8 +291,10 @@ namespace GeometryCollectionTest
 
 			Collection = TNewSimulationObject<GeometryType::GeometryCollectionWithSingleRigid>::Init(Params)->template As<FGeometryCollectionWrapper>();
 			
+			const TManagedArray<FTransform>& MassToLocal = Collection->RestCollection->GetAttribute<FTransform>("MassToLocal", FGeometryCollection::TransformGroup);
+
 			EXPECT_EQ(Collection->DynamicCollection->Parent[0], 1); // is a child of index one
-			EXPECT_TRUE(Collection->DynamicCollection->MassToLocal[0].Equals(FTransform::Identity)); // we are not testing MassToLocal in this test
+			EXPECT_TRUE(MassToLocal[0].Equals(FTransform::Identity)); // we are not testing MassToLocal in this test
 
 			UnitTest.AddSimulationObject(Collection);
 		}
@@ -301,8 +311,11 @@ namespace GeometryCollectionTest
 			Params.CollisionType = ECollisionTypeEnum::Chaos_Volumetric;
 			Params.RootTransform = FTransform(GlobalRotation, GlobalTranslation); Params.NestedTransforms = { FTransform::Identity, FTransform::Identity, FTransform::Identity };
 			CollectionStaticSphere = TNewSimulationObject<GeometryType::GeometryCollectionWithSingleRigid>::Init(Params)->template As<FGeometryCollectionWrapper>();
+			
+			const TManagedArray<FTransform>& MassToLocal = Collection->RestCollection->GetAttribute<FTransform>("MassToLocal", FGeometryCollection::TransformGroup);
+			
 			EXPECT_EQ(CollectionStaticSphere->DynamicCollection->Parent[0], 1); // is a child of index one
-			EXPECT_TRUE(CollectionStaticSphere->DynamicCollection->MassToLocal[0].Equals(FTransform::Identity)); // we are not testing MassToLocal in this test
+			EXPECT_TRUE(MassToLocal[0].Equals(FTransform::Identity)); // we are not testing MassToLocal in this test
 
 			UnitTest.AddSimulationObject(CollectionStaticSphere);
 		}
@@ -502,8 +515,11 @@ namespace GeometryCollectionTest
 			FVector TetraHedronScale(Scale);
 			Params.GeomTransform.SetScale3D(TetraHedronScale); // Tetrahedron dimensions
 			Collection = TNewSimulationObject<GeometryType::GeometryCollectionWithSingleRigid>::Init(Params)->template As<FGeometryCollectionWrapper>();
+
+			const TManagedArray<FTransform>& MassToLocal = Collection->RestCollection->GetAttribute<FTransform>("MassToLocal", FGeometryCollection::TransformGroup);
+
 			EXPECT_EQ(Collection->DynamicCollection->Parent[0], -1); // is a child of index one
-			EXPECT_NEAR((Collection->DynamicCollection->MassToLocal[0].GetTranslation()-FVector(0,0,Scale + 10)).Size(),0,KINDA_SMALL_NUMBER);
+			EXPECT_NEAR((MassToLocal[0].GetTranslation()-FVector(0,0,Scale + 10)).Size(),0,KINDA_SMALL_NUMBER);
 
 			UnitTest.AddSimulationObject(Collection);
 		}
diff --git a/Engine/Source/Runtime/Experimental/Chaos/Private/GeometryCollectionProxyData.cpp b/Engine/Source/Runtime/Experimental/Chaos/Private/GeometryCollectionProxyData.cpp
--- a/Engine/Source/Runtime/Experimental/Chaos/Private/GeometryCollectionProxyData.cpp
+++ b/Engine/Source/Runtime/Experimental/Chaos/Private/GeometryCollectionProxyData.cpp
@@ -55,7 +55,6 @@ FGeometryDynamicCollection::FGeometryDynamicCollection()
 	AddExternalAttribute<bool>(FGeometryDynamicCollection::ActiveAttribute, FTransformCollection::TransformGroup, Active);
 	AddExternalAttribute<int32>(FGeometryDynamicCollection::DynamicStateAttribute, FTransformCollection::TransformGroup, DynamicState);
 	AddExternalAttribute(ImplicitsAttribute, FTransformCollection::TransformGroup, Implicits);
-	AddExternalAttribute("MassToLocal", FTransformCollection::TransformGroup, MassToLocal);
 	AddExternalAttribute(SimplicialsAttribute, FTransformCollection::TransformGroup, Simplicials);
 	AddExternalAttribute(SimulatableParticlesAttribute, FGeometryCollection::TransformGroup, SimulatableParticles);
 
diff --git a/Engine/Source/Runtime/Experimental/Chaos/Private/PhysicsProxy/GeometryCollectionPhysicsProxy.cpp b/Engine/Source/Runtime/Experimental/Chaos/Private/PhysicsProxy/GeometryCollectionPhysicsProxy.cpp
--- a/Engine/Source/Runtime/Experimental/Chaos/Private/PhysicsProxy/GeometryCollectionPhysicsProxy.cpp
+++ b/Engine/Source/Runtime/Experimental/Chaos/Private/PhysicsProxy/GeometryCollectionPhysicsProxy.cpp
@@ -65,6 +65,9 @@ namespace
 	const FName AngularVelocitiesAttributeName = "AngularVelocity";
 	const FName InternalClusterParentTypeArrayAttributeName = "InternalClusterParentTypeArray";
 	const FName AnimateTransformAttributeName = "AnimateTransformAttribute";
+	const FName MassToLocalAttributeName = "MassToLocal";
+	const FName MassAttributeName = "Mass";
+	const FName InertiaTensorAttributeName = "InertiaTensor";
 }
 
 namespace Chaos{
@@ -553,7 +556,8 @@ void FGeometryCollectionPhysicsProxy::Initialize(Chaos::FPBDRigidsEvolutionBase
 	PhysicsThreadCollectionAnchoringFacade.CopyAnchoredAttribute(DynamicCollectionAnchoringFacade);
 	
 	const FVector Scale = Parameters.WorldTransform.GetScale3D();
-	const TManagedArray<float>& Mass = Parameters.RestCollection->GetAttribute<float>("Mass", FTransformCollection::TransformGroup);
+	const TManagedArray<float>& Mass = Parameters.RestCollection->GetAttribute<float>(MassAttributeName, FTransformCollection::TransformGroup);
+	const TManagedArray<FTransform>& MassToLocal = Parameters.RestCollection->GetAttribute<FTransform>(MassToLocalAttributeName, FTransformCollection::TransformGroup);
 
 	TArray<int32> ChildrenToCheckForParentFix;
 	if(ensure(NumTransforms == GameThreadCollection.Implicits.Num() && NumTransforms == GTParticles.Num())) // Implicits are in the transform group so this invariant should always hold
@@ -575,7 +579,7 @@ void FGeometryCollectionPhysicsProxy::Initialize(Chaos::FPBDRigidsEvolutionBase
 				// constraint will query the game thread particle position/rotation for the geometry collection to compute its
 				// reference frame. If that position/rotation does not match up with the physics thread's position/rotation,
 				// the geometry collection particle will have an added velocity computed by the joint constraint solver.
-				const FTransform& T = GameThreadCollection.MassToLocal[Index] * GameThreadCollection.Transform[Index] * Parameters.WorldTransform;
+				const FTransform& T = MassToLocal[Index] * GameThreadCollection.Transform[Index] * Parameters.WorldTransform;
 				P->SetX(T.GetTranslation(), false);
 				P->SetR(T.GetRotation(), false);
 				P->SetM(ScaledMass);
@@ -641,7 +645,7 @@ void FGeometryCollectionPhysicsProxy::Initialize(Chaos::FPBDRigidsEvolutionBase
 				{
 					if (GameThreadCollection.Implicits[ParentToFixIndex] == nullptr)
 					{
-						const Chaos::FRigidTransform3 ParentShapeTransform =  GameThreadCollection.MassToLocal[ParentToFixIndex];
+						const Chaos::FRigidTransform3 ParentShapeTransform =  MassToLocal[ParentToFixIndex];
 				
 						// let's make sure all our children have an implicit defined, other wise, postpone to next iteration 
 						bool bAllChildrenHaveCollision = true;
@@ -664,7 +668,7 @@ void FGeometryCollectionPhysicsProxy::Initialize(Chaos::FPBDRigidsEvolutionBase
 								const Chaos::FImplicitObjectPtr& ChildImplicit = GameThreadCollection.Implicits[ChildIndex];
 								if (ChildImplicit)
 								{
-									const Chaos::FRigidTransform3 ChildShapeTransform = GameThreadCollection.MassToLocal[ChildIndex] * GameThreadCollection.Transform[ChildIndex];
+									const Chaos::FRigidTransform3 ChildShapeTransform = MassToLocal[ChildIndex] * GameThreadCollection.Transform[ChildIndex];
 									const Chaos::FRigidTransform3 RelativeShapeTransform = ChildShapeTransform.GetRelativeTransform(ParentShapeTransform);
 
 									Chaos::FImplicitObjectPtr TransformedChildImplicit = MakeTransformImplicitObject(*ChildImplicit, RelativeShapeTransform);
@@ -1037,13 +1041,13 @@ void FGeometryCollectionPhysicsProxy::InitializeBodiesPT(Chaos::FPBDRigidsSolver
 		const TManagedArray<int32>& BoneMap = RestCollection->BoneMap;
 		const TManagedArray<int32>& SimulationType = RestCollection->SimulationType;
 		const TManagedArray<FVector3f>& Vertex = RestCollection->Vertex;
-		const TManagedArray<float>& Mass = RestCollection->GetAttribute<float>("Mass", FTransformCollection::TransformGroup);
-		const TManagedArray<FVector3f>& InertiaTensor = RestCollection->GetAttribute<FVector3f>("InertiaTensor", FTransformCollection::TransformGroup);
+		const TManagedArray<float>& Mass = RestCollection->GetAttribute<float>(MassAttributeName, FTransformCollection::TransformGroup);
+		const TManagedArray<FVector3f>& InertiaTensor = RestCollection->GetAttribute<FVector3f>(InertiaTensorAttributeName, FTransformCollection::TransformGroup);
 
 		const int32 NumTransforms = DynamicCollection.NumElements(FTransformCollection::TransformGroup);
 		const TManagedArray<int32>& DynamicState = DynamicCollection.DynamicState;
 		const TManagedArray<bool>& SimulatableParticles = DynamicCollection.SimulatableParticles;
-		const TManagedArray<FTransform>& MassToLocal = DynamicCollection.MassToLocal;
+		const TManagedArray<FTransform>& MassToLocal = RestCollection->GetAttribute<FTransform>(MassToLocalAttributeName, FTransformCollection::TransformGroup);
 		const TManagedArray<Chaos::FImplicitObjectPtr>& Implicits = DynamicCollection.Implicits;
 		const TManagedArray<TUniquePtr<FCollisionStructureManager::FSimplicial>>& Simplicials = DynamicCollection.Simplicials;
 		const TManagedArray<TSet<int32>>& Children = DynamicCollection.Children;
@@ -1617,7 +1621,7 @@ Chaos::TPBDGeometryCollectionParticleHandle<Chaos::FReal, 3>* FGeometryCollectio
 	TManagedArray<int32>& ParentIndex = DynamicCollection.Parent;
 	TManagedArray<TSet<int32>>& Children = DynamicCollection.Children;
 	TManagedArray<FTransform>& Transform = DynamicCollection.Transform;
-	TManagedArray<FTransform>& MassToLocal = DynamicCollection.MassToLocal;
+	const TManagedArray<FTransform>& MassToLocal = Parameters.RestCollection->GetAttribute<FTransform>(MassToLocalAttributeName, FTransformCollection::TransformGroup);
 	const TManagedArray<Chaos::FImplicitObjectPtr>& Implicits = DynamicCollection.Implicits;
 	const TManagedArray<TUniquePtr<FCollisionStructureManager::FSimplicial>>& Simplicials = DynamicCollection.Simplicials;
 	Chaos::Facades::FCollectionAnchoringFacade AnchoringFacade(DynamicCollection);
@@ -1737,7 +1741,7 @@ FGeometryCollectionPhysicsProxy::BuildClusters_Internal(
 	TManagedArray<int32>& ParentIndex = DynamicCollection.Parent;
 	TManagedArray<TSet<int32>>& Children = DynamicCollection.Children;
 	TManagedArray<FTransform>& Transform = DynamicCollection.Transform;
-	TManagedArray<FTransform>& MassToLocal = DynamicCollection.MassToLocal;
+	const TManagedArray<FTransform>& MassToLocal = Parameters.RestCollection->GetAttribute<FTransform>(MassToLocalAttributeName, FTransformCollection::TransformGroup);
 	//TManagedArray<TSharedPtr<FCollisionStructureManager::FSimplicial> >& Simplicials = DynamicCollection.Simplicials;
 	TManagedArray<Chaos::FImplicitObjectPtr>& Implicits = DynamicCollection.Implicits;
 
@@ -1798,8 +1802,8 @@ FGeometryCollectionPhysicsProxy::BuildClusters_Internal(
 	}
 
 	check(Parameters.RestCollection);
-	const TManagedArray<float>& Mass = Parameters.RestCollection->GetAttribute<float>("Mass", FTransformCollection::TransformGroup);
-	const TManagedArray<FVector3f>& InertiaTensor = Parameters.RestCollection->GetAttribute<FVector3f>("InertiaTensor", FTransformCollection::TransformGroup);
+	const TManagedArray<float>& Mass = Parameters.RestCollection->GetAttribute<float>(MassAttributeName, FTransformCollection::TransformGroup);
+	const TManagedArray<FVector3f>& InertiaTensor = Parameters.RestCollection->GetAttribute<FVector3f>(InertiaTensorAttributeName, FTransformCollection::TransformGroup);
 
 	const float ScaledMass = AdjustMassForScale(Mass[CollectionClusterIndex]);
 	const Chaos::FVec3f ScaledInertia = AdjustInertiaForScale((Chaos::FVec3f)InertiaTensor[CollectionClusterIndex]);
@@ -2907,6 +2911,8 @@ void FGeometryCollectionPhysicsProxy::PushToPhysicsState()
 			TArray<FPBDRigidParticleHandle*> DeferredClusterUnionParticleUpdates;
 			TArray<FTransform> DeferredClusterUnionChildToParentUpdates;
 
+			const TManagedArray<FTransform>& MassToLocal = Parameters.RestCollection->GetAttribute<FTransform>(MassToLocalAttributeName, FTransformCollection::TransformGroup);
+
 			for (int32 TransformGroupIndex = 0; TransformGroupIndex < NumTransforms; ++TransformGroupIndex)
 			{
 				if (Chaos::FPBDRigidClusteredParticleHandle* Handle = SolverParticleHandles[TransformGroupIndex])
@@ -2949,7 +2955,7 @@ void FGeometryCollectionPhysicsProxy::PushToPhysicsState()
 						if (ClusterUnionIndex != INDEX_NONE)
 						{
 							const FTransform ParentWorldTransform{ ParentHandle->R(), ParentHandle->X() };
-							const FTransform NewWorldTransform = PhysicsThreadCollection.MassToLocal[TransformGroupIndex] * PhysicsThreadCollection.Transform[TransformGroupIndex] * Parameters.WorldTransform;
+							const FTransform NewWorldTransform = MassToLocal[TransformGroupIndex] * PhysicsThreadCollection.Transform[TransformGroupIndex] * Parameters.WorldTransform;
 							const FTransform RelativeTransform = NewWorldTransform.GetRelativeTransform(ParentWorldTransform);
 
 							DeferredClusterUnionParticleUpdates.Add(Handle);
@@ -3111,6 +3117,8 @@ void FGeometryCollectionPhysicsProxy::BufferPhysicsResults_Internal(Chaos::FPBDR
 	{ 
 		SCOPE_CYCLE_COUNTER(STAT_CalcParticleToWorld);
 
+		const TManagedArray<FTransform>& MassToLocalArray = Parameters.RestCollection->GetAttribute<FTransform>(MassToLocalAttributeName, FTransformCollection::TransformGroup);
+
 		for (int32 TransformGroupIndex = 0; TransformGroupIndex < NumTransforms; ++TransformGroupIndex)
 		{
 			Chaos::FPBDRigidClusteredParticleHandle* Handle = SolverParticleHandles[TransformGroupIndex];
@@ -3237,7 +3245,7 @@ void FGeometryCollectionPhysicsProxy::BufferPhysicsResults_Internal(Chaos::FPBDR
 				if (!ClusterParent || ClusterParent->InternalCluster())
 				{
 					const FRigidTransform3 ParticleToWorld(PositionData.ParticleX, PositionData.ParticleR);
-					const FTransform MassToLocal = PhysicsThreadCollection.MassToLocal[TransformGroupIndex];
+					const FTransform MassToLocal = MassToLocalArray[TransformGroupIndex];
 					ParentSpaceTransform = MassToLocal.GetRelativeTransformReverse(ParticleToWorld).GetRelativeTransform(ActorToWorld);
 					ParentSpaceTransform.NormalizeRotation();
 					if (IsActorScaled)
@@ -3250,7 +3258,8 @@ void FGeometryCollectionPhysicsProxy::BufferPhysicsResults_Internal(Chaos::FPBDR
 				Results.SetState(EntryIndex, StateData);
 				Results.SetPositions(EntryIndex, PositionData);
 				Results.SetVelocities(EntryIndex, VelocityData);
-				PhysicsThreadCollection.Transform[TransformGroupIndex] = ParentSpaceTransform;
+				// todo(chaos) : we shoudl eventually get rid of the transform in the Physics collection
+				PhysicsThreadCollection.Transform[TransformGroupIndex] = ParentSpaceTransform; 
 				IsObjectDynamic = true;
 			}
 
@@ -3582,6 +3591,7 @@ bool FGeometryCollectionPhysicsProxy::PullFromPhysicsState(const Chaos::FDirtyGe
 		bIsCollectionDirty |= PullNonInterpolatableDataFromSinglePhysicsState(PullData, !bNeedInterpolation, NextPullData ? &NextPullData->Results().GetModifiedTransformIndices() : nullptr);
 
 		const TManagedArray<bool>* AnimationsActive = GameThreadCollection.FindAttribute<bool>(AnimateTransformAttributeName, FGeometryCollection::TransformGroup);
+		const TManagedArray<FTransform>& MassToLocal = Parameters.RestCollection->GetAttribute<FTransform>(MassToLocalAttributeName, FTransformCollection::TransformGroup);
 
 		// second : interpolate-able ones
 		if (bNeedInterpolation)
@@ -3658,10 +3668,10 @@ bool FGeometryCollectionPhysicsProxy::PullFromPhysicsState(const Chaos::FDirtyGe
 						if (bWasModified)
 						{
 							const FParticle& GTParticle = *GTParticles[TransformIndex];
-							const FTransform& MassToLocal = GameThreadCollection.MassToLocal[TransformIndex];
+							const FTransform& ParticleMassToLocal = MassToLocal[TransformIndex];
 							const int32 ParentTransformIndex = GameThreadCollection.Parent[TransformIndex];
 
-							const FTransform& WorldTransform = MassToLocal.Inverse() * FTransform { GTParticle.R(), GTParticle.X() };
+							const FTransform& WorldTransform = ParticleMassToLocal.Inverse() * FTransform { GTParticle.R(), GTParticle.X() };
 
 							// by default parent is the component's world transform 
 							FTransform ParentWorldTransform = Parameters.WorldTransform;
@@ -3669,7 +3679,7 @@ bool FGeometryCollectionPhysicsProxy::PullFromPhysicsState(const Chaos::FDirtyGe
 							{
 								if (const FParticle* GTParentParticle = (ParentTransformIndex != INDEX_NONE) ? GTParticles[ParentTransformIndex].Get() : nullptr)
 								{
-									const FTransform& ParentMassToLocal = GameThreadCollection.MassToLocal[TransformIndex];
+									const FTransform& ParentMassToLocal = MassToLocal[TransformIndex];
 									ParentWorldTransform = ParentMassToLocal.Inverse() * FTransform { GTParentParticle->R(), GTParentParticle->X() };
 								}
 							}
@@ -3677,7 +3687,7 @@ bool FGeometryCollectionPhysicsProxy::PullFromPhysicsState(const Chaos::FDirtyGe
 							FTransform NewTransform = WorldTransform.GetRelativeTransform(ParentWorldTransform);
 							if (ParentTransformIndex == INDEX_NONE && bIsComponentTransformScaled)
 							{
-								NewTransform = MassToLocal.Inverse() * ComponentScaleTransform * MassToLocal * NewTransform;
+								NewTransform = ParticleMassToLocal.Inverse() * ComponentScaleTransform * ParticleMassToLocal * NewTransform;
 							}
 							if (!NewTransform.Equals(GameThreadCollection.Transform[TransformIndex], GeometryCollectionTransformTolerance))
 							{
@@ -3848,7 +3858,7 @@ static Chaos::FImplicitObjectPtr CreateImplicitGeometry(
 	Chaos::FImplicitObjectPtr NewImplicit;
 	if (SizeSpecificData.CollisionShapesData.Num())
 	{
-		const TManagedArray<FTransform>& CollectionMassToLocal = RestCollection.GetAttribute<FTransform>("MassToLocal", FTransformCollection::TransformGroup);
+		const TManagedArray<FTransform>& CollectionMassToLocal = RestCollection.GetAttribute<FTransform>(MassToLocalAttributeName, FTransformCollection::TransformGroup);
 		const TManagedArray<TSet<int32>>* TransformToConvexIndices = RestCollection.FindAttribute<TSet<int32>>("TransformToConvexIndices", FTransformCollection::TransformGroup);
 		const TManagedArray<Chaos::FConvexPtr>* ConvexGeometry = RestCollection.FindAttribute<Chaos::FConvexPtr>(FGeometryCollection::ConvexHullAttribute, FGeometryCollection::ConvexGroup);
 
@@ -4173,8 +4183,8 @@ void FGeometryCollectionPhysicsProxy::InitializeSharedCollisionStructures(
 	const TManagedArray<TSet<int32>>& Children = RestCollection.Children;
 	const TManagedArray<int32>& SimulationType = RestCollection.SimulationType;
 	TManagedArray<bool>& CollectionSimulatableParticles = RestCollection.ModifyAttribute<bool>(FGeometryCollection::SimulatableParticlesAttribute, FTransformCollection::TransformGroup);
-	TManagedArray<FVector3f>& CollectionInertiaTensor = RestCollection.AddAttribute<FVector3f>(TEXT("InertiaTensor"), FTransformCollection::TransformGroup);
-	TManagedArray<FRealSingle>& CollectionMass = RestCollection.AddAttribute<FRealSingle>(TEXT("Mass"), FTransformCollection::TransformGroup);
+	TManagedArray<FVector3f>& CollectionInertiaTensor = RestCollection.AddAttribute<FVector3f>(InertiaTensorAttributeName, FTransformCollection::TransformGroup);
+	TManagedArray<FRealSingle>& CollectionMass = RestCollection.AddAttribute<FRealSingle>(MassAttributeName, FTransformCollection::TransformGroup);
 	TManagedArray<TUniquePtr<FSimplicial>>& CollectionSimplicials =	RestCollection.AddAttribute<TUniquePtr<FSimplicial>>(FGeometryDynamicCollection::SimplicialsAttribute, FTransformCollection::TransformGroup);
 
 	TManagedArray<int32>& Levels = RestCollection.AddAttribute<int32>(TEXT("Level"), FTransformCollection::TransformGroup);
@@ -4190,7 +4200,7 @@ void FGeometryCollectionPhysicsProxy::InitializeSharedCollisionStructures(
 
 
 	// @todo(chaos_transforms) : do we still use this?
-	TManagedArray<FTransform>& CollectionMassToLocal = RestCollection.AddAttribute<FTransform>(TEXT("MassToLocal"), FTransformCollection::TransformGroup);
+	TManagedArray<FTransform>& CollectionMassToLocal = RestCollection.AddAttribute<FTransform>(MassToLocalAttributeName, FTransformCollection::TransformGroup);
 	FTransform IdentityXf(FQuat::Identity, FVector(0));
 	IdentityXf.NormalizeRotation();
 	CollectionMassToLocal.Fill(IdentityXf);
diff --git a/Engine/Source/Runtime/Experimental/Chaos/Public/GeometryCollectionProxyData.h b/Engine/Source/Runtime/Experimental/Chaos/Public/GeometryCollectionProxyData.h
--- a/Engine/Source/Runtime/Experimental/Chaos/Public/GeometryCollectionProxyData.h
+++ b/Engine/Source/Runtime/Experimental/Chaos/Public/GeometryCollectionProxyData.h
@@ -79,7 +79,6 @@ public:
 	
 	TManagedArray<int32> DynamicState;
 	TManagedArray<Chaos::FImplicitObjectPtr> Implicits;
-	TManagedArray<FTransform> MassToLocal;
 	TManagedArray<TUniquePtr<FCollisionStructureManager::FSimplicial>> Simplicials;
 	TManagedArray<bool> SimulatableParticles;
 
diff --git a/Engine/Source/Runtime/Experimental/GeometryCollectionEngine/Private/GeometryCollection/GeometryCollectionComponent.cpp b/Engine/Source/Runtime/Experimental/GeometryCollectionEngine/Private/GeometryCollection/GeometryCollectionComponent.cpp
--- a/Engine/Source/Runtime/Experimental/GeometryCollectionEngine/Private/GeometryCollection/GeometryCollectionComponent.cpp
+++ b/Engine/Source/Runtime/Experimental/GeometryCollectionEngine/Private/GeometryCollection/GeometryCollectionComponent.cpp
@@ -3306,6 +3306,8 @@ void UGeometryCollectionComponent::UpdateNavigationDataIfNeeded(bool bDynamicCol
 
 void UGeometryCollectionComponent::UpdateRemovalIfNeeded()
 {
+	static FName MassToLocalAttributeName = "MassToLocal";
+
 	// if removal is enabled, update the dynamic collection transform based on the decay 
 	// todo: we could optimize this using a list of transform to update from when we update the decay values
 	if (DynamicCollection && bAllowRemovalOnBreak && bAllowRemovalOnSleep)
@@ -3317,6 +3319,8 @@ void UGeometryCollectionComponent::UpdateRemovalIfNeeded()
 
 			const FTransform InverseComponentTransform = (RestCollection->bScaleOnRemoval) ? GetComponentTransform().Inverse() : FTransform::Identity;
 
+			const TManagedArray<FTransform>& MassToLocal = RestCollection->GetGeometryCollection()->GetAttribute<FTransform>(MassToLocalAttributeName, FGeometryCollection::TransformGroup);
+
 			const int32 NumTransforms = DecayFacade.GetDecayAttributeSize();
 			for (int32 TransformIndex = 0; TransformIndex < NumTransforms; ++TransformIndex)
 			{
@@ -3342,7 +3346,7 @@ void UGeometryCollectionComponent::UpdateRemovalIfNeeded()
 
 						const FQuat LocalRotation = (InverseComponentTransform * ComponentSpaceTransforms[TransformIndex].Inverse()).GetRotation();
 						const FVector LocalDown = LocalRotation.RotateVector(FVector(0.f, 0.f, ShrinkRadius));
-						const FVector CenterOfMass = DynamicCollection->MassToLocal[TransformIndex].GetTranslation();
+						const FVector CenterOfMass = MassToLocal[TransformIndex].GetTranslation();
 						const FVector ScaleCenter = LocalDown + CenterOfMass;
 						const FTransform ScaleTransform(FQuat::Identity, ScaleCenter * FVector::FReal(1.f - Scale), FVector(Scale));
 						DynamicCollection->Transform[TransformIndex] = ScaleTransform * DynamicCollection->Transform[TransformIndex];
