Chaos - don't expand Cull Distance for CCD contacts. It is not used by the sweep and results in unnecessary contacts being generated after the CCD rewind

diff --git a/Engine/Source/Runtime/Experimental/Chaos/Private/Chaos/Collision/ParticlePairMidPhase.cpp b/Engine/Source/Runtime/Experimental/Chaos/Private/Chaos/Collision/ParticlePairMidPhase.cpp
--- a/Engine/Source/Runtime/Experimental/Chaos/Private/Chaos/Collision/ParticlePairMidPhase.cpp
+++ b/Engine/Source/Runtime/Experimental/Chaos/Private/Chaos/Collision/ParticlePairMidPhase.cpp
@@ -870,30 +870,26 @@ namespace Chaos
 			// If CCD is enabled, did we move far enough to require a sweep?
 			Flags.bUseSweep = Flags.bIsCCD && ShouldEnableCCD(Dt);
 
-			// Extend cull distance based on velocity
+			// We increase CullDistance based on velocity (up to a limit for perf with large velocities).
+			// NOTE: This somewhat matches the bounds expansion in FPBDRigidsEvolutionGBF::Integrate
 			// NOTE: We use PreV here which is the velocity after collisions from the previous tick because we want the 
 			// velocity without gravity from this tick applied. This is mainly so that we get the same CullDistance from 
 			// one tick to the next, even if one of the particles goes to sleep, and therefore its velocity is now zero 
 			// because gravity is no longer applied. Also see FPBDIslandManager::PropagateIslandSleep for other issues 
 			// related to velocity and sleeping...
-			const FReal VMax0 = FConstGenericParticleHandle(GetParticle0())->PreV().GetAbsMax();
-			const FReal VMax1 = FConstGenericParticleHandle(GetParticle1())->PreV().GetAbsMax();
+			// NOTE: we used to extend the cull distance for CCD objects, but this is no longer required. The sweep and
+			// rewind phase does not use CullDistance, and once rewound we are using normal collision detection where
+			// an expanded CullDistance doesn't help and makes perf worse.
+			FConstGenericParticleHandle P0 = GetParticle0();
+			FConstGenericParticleHandle P1 = GetParticle1();
+			const FReal VMax0 = P0->PreV().GetAbsMax();
+			const FReal VMax1 = P1->PreV().GetAbsMax();
 			const FReal VMaxDt = FMath::Max(VMax0, VMax1) * Dt;
-			if (!Flags.bUseSweep)
-			{
-				// Normal (non sweep) mode: we increase CullDistance based on velocity up to a limit
-				// NOTE: This somewhat matches the bounds expansion in FPBDRigidsEvolutionGBF::Integrate
-				const FReal VelocityBoundsMultiplier = Context.GetSettings().BoundsVelocityInflation;
-				const FReal MaxVelocityBoundsExpansion = Context.GetSettings().MaxVelocityBoundsExpansion;
-				if ((VelocityBoundsMultiplier > 0) && (MaxVelocityBoundsExpansion > 0))
-				{
-					CullDistance += FMath::Min(VelocityBoundsMultiplier * VMaxDt, MaxVelocityBoundsExpansion);
-				}
-			}
-			else
+			const FReal VelocityBoundsMultiplier = Context.GetSettings().BoundsVelocityInflation;
+			const FReal MaxVelocityBoundsExpansion = Context.GetSettings().MaxVelocityBoundsExpansion;
+			if ((VelocityBoundsMultiplier > 0) && (MaxVelocityBoundsExpansion > 0))
 			{
-				// CCD (sweept) mode: we increase CullDistance based on velocity, with no limits
-				CullDistance += VMaxDt;
+				CullDistance += FMath::Min(VelocityBoundsMultiplier * VMaxDt, MaxVelocityBoundsExpansion);
 			}
 
 			// Run collision detection on all potentially colliding shape pairs
