Chaos - fix collision point matching tolerance and make it tunable

diff --git a/Engine/Source/Runtime/Experimental/Chaos/Private/Chaos/Collision/PBDCollisionConstraint.cpp b/Engine/Source/Runtime/Experimental/Chaos/Private/Chaos/Collision/PBDCollisionConstraint.cpp
--- a/Engine/Source/Runtime/Experimental/Chaos/Private/Chaos/Collision/PBDCollisionConstraint.cpp
+++ b/Engine/Source/Runtime/Experimental/Chaos/Private/Chaos/Collision/PBDCollisionConstraint.cpp
@@ -32,8 +32,11 @@ namespace Chaos
 	FAutoConsoleVariableRef CVarChaos_Manifold_MatchPositionTolerance(TEXT("p.Chaos.Collision.Manifold.MatchPositionTolerance"), Chaos_Manifold_MatchPositionTolerance, TEXT("A tolerance as a fraction of object size used to determine if two contact points are the same"));
 	FAutoConsoleVariableRef CVarChaos_Manifold_MatchNormalTolerance(TEXT("p.Chaos.Collision.Manifold.MatchNormalTolerance"), Chaos_Manifold_MatchNormalTolerance, TEXT("A tolerance on the normal dot product used to determine if two contact points are the same"));
 
-	FRealSingle Chaos_Manifold_FrictionPositionTolerance = 1.0f;	// Distance a shape-relative contact point can move and still be considered the same point
-	FAutoConsoleVariableRef CVarChaos_Manifold_FrictionPositionTolerance(TEXT("p.Chaos.Collision.Manifold.FrictionPositionTolerance"), Chaos_Manifold_FrictionPositionTolerance, TEXT(""));
+	FRealSingle Chaos_Manifold_FrictionExactPositionTolerance = 0.2f;	// Distance a shape-relative contact point can move and we can still be sure it's the same point (don't look for better options)
+	FRealSingle Chaos_Manifold_FrictionNearPositionTolerance = 1.0f;	// Distance a shape-relative contact point can move and still be considered the same point, if no better option exists
+	FAutoConsoleVariableRef CVarChaos_Manifold_FrictionExactPositionTolerance(TEXT("p.Chaos.Collision.Manifold.FrictionExactPositionTolerance"), Chaos_Manifold_FrictionExactPositionTolerance, TEXT(""));
+	FAutoConsoleVariableRef CVarChaos_Manifold_FrictionNearPositionTolerance(TEXT("p.Chaos.Collision.Manifold.FrictionNearPositionTolerance"), Chaos_Manifold_FrictionNearPositionTolerance, TEXT(""));
+
 
 	FRealSingle Chaos_GBFCharacteristicTimeRatio = 1.0f;
 	FAutoConsoleVariableRef CVarChaos_GBFCharacteristicTimeRatio(TEXT("p.Chaos.Collision.GBFCharacteristicTimeRatio"), Chaos_GBFCharacteristicTimeRatio, TEXT("The ratio between characteristic time and Dt"));
@@ -1013,42 +1016,31 @@ namespace Chaos
 		return true;
 	}
 
-	FReal FPBDCollisionConstraint::CalculateSavedManifoldDistanceSq(const FSavedManifoldPoint& SavedManifoldPoint, const FManifoldPoint& ManifoldPoint, const FReal DistanceToleranceSq) const
+	FRealSingle FPBDCollisionConstraint::CalculateSavedManifoldPointDistanceSq(const FSavedManifoldPoint& SavedManifoldPoint, const FManifoldPoint& ManifoldPoint) const
 	{
 		// If we have a vertex-plane (or vertex-vertex) contact, we want to know if we have the same vertex(es).
 		// If we have and edge-edge contact, we want to know if we have the same edges.
 		// But we don't know what type of contact we have, so for now...
 		// If the contact point is in the same spot on one of the bodies, assume it is the same contact
 		// @todo(chaos) - collision detection should provide the contact point types (vertex/edge/plane)
-		FReal DP0Sq = TNumericLimits<FReal>::Max();
-		FReal DP1Sq = TNumericLimits<FReal>::Max();
-		const FVec3 DP0 = ManifoldPoint.ContactPoint.ShapeContactPoints[0] - SavedManifoldPoint.ShapeContactPoints[0];
-		const FVec3 DP1 = ManifoldPoint.ContactPoint.ShapeContactPoints[1] - SavedManifoldPoint.ShapeContactPoints[1];
+		const FVec3f DP0 = ManifoldPoint.ContactPoint.ShapeContactPoints[0] - SavedManifoldPoint.ShapeContactPoints[0];
+		const FVec3f DP1 = ManifoldPoint.ContactPoint.ShapeContactPoints[1] - SavedManifoldPoint.ShapeContactPoints[1];
 
 		// When only one shape is quadratic, we only look at the quadratic contact point so we don't identify
 		// a sphere spinning on the spot as a stationary contact
 		// @todo(chaos): handle quadratic shapes better with static friction
 		if (IsQuadratic0() && !IsQuadratic1())
 		{
-			DP0Sq = DP0.SizeSquared();
+			return DP0.SizeSquared();
 		}
 		else if (IsQuadratic1() && !IsQuadratic0())
 		{
-			DP1Sq = DP1.SizeSquared();
+			return DP1.SizeSquared();
 		}
 		else
 		{
-			DP0Sq = DP0.SizeSquared();
-			DP1Sq = DP1.SizeSquared();
-		}
-
-		const FReal MinDPSq = FMath::Min(DP0Sq, DP1Sq);
-		if (MinDPSq < DistanceToleranceSq)
-		{
-			return MinDPSq;
+			return FMath::Min(DP0.SizeSquared(), DP1.SizeSquared());
 		}
-
-		return TNumericLimits<FReal>::Max();
 	}
 
 	int32 FPBDCollisionConstraint::FindSavedManifoldPoint(const int32 ManifoldPointIndex, int32* InOutAllowedSavedPointIndices, int32& InOutNumAllowedSavedPoints) const
@@ -1058,9 +1050,11 @@ namespace Chaos
 			const FManifoldPoint& ManifoldPoint = ManifoldPoints[ManifoldPointIndex];
 			if (!ManifoldPoint.Flags.bDisabled)
 			{
-				const FReal DistanceToleranceSq = FMath::Square(Chaos_Manifold_FrictionPositionTolerance);
-				const FReal ExactDistanceToleranceSq = 0.2f * DistanceToleranceSq;
-				FReal BestDistanceSq = DistanceToleranceSq;
+				const FRealSingle ExactDistanceToleranceSq = FMath::Square(Chaos_Manifold_FrictionExactPositionTolerance);
+				const FRealSingle NearDistanceToleranceSq = FMath::Square(Chaos_Manifold_FrictionNearPositionTolerance);
+
+				// Ignore popints farther than NearDistanceTolerance
+				FRealSingle BestDistanceSq = NearDistanceToleranceSq;
 
 				int32 MatchAllowedPointIndex = INDEX_NONE;
 				for (int32 AllowedPointIndex = 0; AllowedPointIndex < InOutNumAllowedSavedPoints; ++AllowedPointIndex)
@@ -1068,7 +1062,7 @@ namespace Chaos
 					const int32 SavedPointIndex = InOutAllowedSavedPointIndices[AllowedPointIndex];
 					const FSavedManifoldPoint& SavedManifoldPoint = SavedManifoldPoints[SavedPointIndex];
 
-					const FReal DistanceSq = CalculateSavedManifoldDistanceSq(SavedManifoldPoint, ManifoldPoint, DistanceToleranceSq);
+					const FRealSingle DistanceSq = CalculateSavedManifoldPointDistanceSq(SavedManifoldPoint, ManifoldPoint);
 
 					// If this is an exact match, take the point without searching further
 					if (DistanceSq < ExactDistanceToleranceSq)
diff --git a/Engine/Source/Runtime/Experimental/Chaos/Public/Chaos/Collision/PBDCollisionConstraint.h b/Engine/Source/Runtime/Experimental/Chaos/Public/Chaos/Collision/PBDCollisionConstraint.h
--- a/Engine/Source/Runtime/Experimental/Chaos/Public/Chaos/Collision/PBDCollisionConstraint.h
+++ b/Engine/Source/Runtime/Experimental/Chaos/Public/Chaos/Collision/PBDCollisionConstraint.h
@@ -854,7 +854,7 @@ namespace Chaos
 		CHAOS_API void UpdateMaterialPropertiesImpl();
 
 	private:
-		CHAOS_API FReal CalculateSavedManifoldDistanceSq(const FSavedManifoldPoint& SavedManifoldPoint, const FManifoldPoint& ManifoldPoint, const FReal DistanceToleranceSq) const;
+		CHAOS_API FRealSingle CalculateSavedManifoldPointDistanceSq(const FSavedManifoldPoint& SavedManifoldPoint, const FManifoldPoint& ManifoldPoint) const;
 
 		CHAOS_API void LogOneShotManifoldError(const int32 MaxManifoldPoints, const TArrayView<const FContactPoint>& ContactPoints);
 
