Add extra stats to track post physics sync callback in geometry collection as well as fully decayed event broadcast

diff --git a/Engine/Source/Runtime/Experimental/Chaos/Private/ChaosStats.cpp b/Engine/Source/Runtime/Experimental/Chaos/Private/ChaosStats.cpp
--- a/Engine/Source/Runtime/Experimental/Chaos/Private/ChaosStats.cpp
+++ b/Engine/Source/Runtime/Experimental/Chaos/Private/ChaosStats.cpp
@@ -24,6 +24,8 @@ DEFINE_STAT(STAT_GCSweep);
 DEFINE_STAT(STAT_GCCUpdateBounds);
 DEFINE_STAT(STAT_GCUpdateFilterData);
 DEFINE_STAT(STAT_GCCUGlobalMatrices);
+DEFINE_STAT(STAT_GCPostPhysicsSync);
+DEFINE_STAT(STAT_GCFullyDecayedBroadcast);
 DEFINE_STAT(STAT_GCInitDynamicData);
 DEFINE_STAT(STAT_GCTotalTransforms);
 DEFINE_STAT(STAT_GCChangedTransforms);
diff --git a/Engine/Source/Runtime/Experimental/Chaos/Public/ChaosStats.h b/Engine/Source/Runtime/Experimental/Chaos/Public/ChaosStats.h
--- a/Engine/Source/Runtime/Experimental/Chaos/Public/ChaosStats.h
+++ b/Engine/Source/Runtime/Experimental/Chaos/Public/ChaosStats.h
@@ -42,6 +42,8 @@ DECLARE_CYCLE_STAT_EXTERN(TEXT("Geometry Collection Sweep"), STAT_GCSweep, STATG
 DECLARE_CYCLE_STAT_EXTERN(TEXT("Geoemtry Collection Update Filter Data"), STAT_GCUpdateFilterData, STATGROUP_Chaos, CHAOS_API)
 DECLARE_CYCLE_STAT_EXTERN(TEXT("Geometry Collection Component UpdateBounds"), STAT_GCCUpdateBounds, STATGROUP_Chaos, CHAOS_API);
 DECLARE_CYCLE_STAT_EXTERN(TEXT("Geometry Collection Component CalculateGlobalMatrices"), STAT_GCCUGlobalMatrices, STATGROUP_Chaos, CHAOS_API);
+DECLARE_CYCLE_STAT_EXTERN(TEXT("Geometry Collection Component OnPostPhysicsSync"), STAT_GCPostPhysicsSync, STATGROUP_Chaos, CHAOS_API);
+DECLARE_CYCLE_STAT_EXTERN(TEXT("Geometry Collection Component FullyDecayed event Broadcast"), STAT_GCFullyDecayedBroadcast, STATGROUP_Chaos, CHAOS_API);
 DECLARE_CYCLE_STAT_EXTERN(TEXT("Geometry Collection InitDynamicData"), STAT_GCInitDynamicData, STATGROUP_Chaos, CHAOS_API);
 DECLARE_DWORD_COUNTER_STAT_EXTERN(TEXT("Geometry Collection Total Transforms"), STAT_GCTotalTransforms, STATGROUP_Chaos, CHAOS_API);
 DECLARE_DWORD_COUNTER_STAT_EXTERN(TEXT("Geometry Collection Changed Transforms"), STAT_GCChangedTransforms, STATGROUP_Chaos, CHAOS_API);
diff --git a/Engine/Source/Runtime/Experimental/GeometryCollectionEngine/Private/GeometryCollection/GeometryCollectionComponent.cpp b/Engine/Source/Runtime/Experimental/GeometryCollectionEngine/Private/GeometryCollection/GeometryCollectionComponent.cpp
--- a/Engine/Source/Runtime/Experimental/GeometryCollectionEngine/Private/GeometryCollection/GeometryCollectionComponent.cpp
+++ b/Engine/Source/Runtime/Experimental/GeometryCollectionEngine/Private/GeometryCollection/GeometryCollectionComponent.cpp
@@ -2750,7 +2750,11 @@ void UGeometryCollectionComponent::CheckFullyDecayed()
 			if (bFullyDecayed)
 			{
 				bAlreadyFullyDecayed = true;
-				OnFullyDecayedEvent.Broadcast();
+				if (OnFullyDecayedEvent.IsBound())
+				{
+					SCOPE_CYCLE_COUNTER(STAT_GCFullyDecayedBroadcast);
+					OnFullyDecayedEvent.Broadcast();
+				}
 			}
 		}
 	}
@@ -3209,6 +3213,8 @@ void UGeometryCollectionComponent::RegisterAndInitializePhysicsProxy()
 
 void UGeometryCollectionComponent::OnPostPhysicsSync()
 {
+	SCOPE_CYCLE_COUNTER(STAT_GCPostPhysicsSync);
+
 	UpdateAttachedChildrenTransform();
 
 	if (GetIsReplicated() && GetNetMode() != ENetMode::NM_Client)
@@ -3319,7 +3325,11 @@ void UGeometryCollectionComponent::UpdateRemovalIfNeeded()
 
 			const FTransform InverseComponentTransform = (RestCollection->bScaleOnRemoval) ? GetComponentTransform().Inverse() : FTransform::Identity;
 
-			const TManagedArray<FTransform>& MassToLocal = RestCollection->GetGeometryCollection()->GetAttribute<FTransform>(MassToLocalAttributeName, FGeometryCollection::TransformGroup);
+			const TManagedArray<FTransform>* MassToLocal = nullptr;
+			if (RestCollection->bScaleOnRemoval)
+			{
+				MassToLocal = RestCollection->GetGeometryCollection()->FindAttribute<FTransform>(MassToLocalAttributeName, FGeometryCollection::TransformGroup);
+			}
 
 			const int32 NumTransforms = DecayFacade.GetDecayAttributeSize();
 			for (int32 TransformIndex = 0; TransformIndex < NumTransforms; ++TransformIndex)
@@ -3334,7 +3344,7 @@ void UGeometryCollectionComponent::UpdateRemovalIfNeeded()
 						DynamicCollection->Transform[TransformIndex].SetScale3D(FVector::ZeroVector);
 					}
 					// do not try to get this condition out of the loop as this may cause some optimizer related issues
-					else if (RestCollection->bScaleOnRemoval) 
+					else if (RestCollection->bScaleOnRemoval && MassToLocal)
 					{
 						float ShrinkRadius = 0.0f;
 						UE::Math::TSphere<double> AccumulatedSphere;
@@ -3346,7 +3356,7 @@ void UGeometryCollectionComponent::UpdateRemovalIfNeeded()
 
 						const FQuat LocalRotation = (InverseComponentTransform * ComponentSpaceTransforms[TransformIndex].Inverse()).GetRotation();
 						const FVector LocalDown = LocalRotation.RotateVector(FVector(0.f, 0.f, ShrinkRadius));
-						const FVector CenterOfMass = MassToLocal[TransformIndex].GetTranslation();
+						const FVector CenterOfMass = (*MassToLocal)[TransformIndex].GetTranslation();
 						const FVector ScaleCenter = LocalDown + CenterOfMass;
 						const FTransform ScaleTransform(FQuat::Identity, ScaleCenter * FVector::FReal(1.f - Scale), FVector(Scale));
 						DynamicCollection->Transform[TransformIndex] = ScaleTransform * DynamicCollection->Transform[TransformIndex];
