Fix Materialize() to correctly return failures. Disk-based materialise fails earlier.

#rnx

diff --git a/Engine/Source/Runtime/Experimental/IoStoreOnDemand/Private/JournaledCache.cpp b/Engine/Source/Runtime/Experimental/IoStoreOnDemand/Private/JournaledCache.cpp
--- a/Engine/Source/Runtime/Experimental/IoStoreOnDemand/Private/JournaledCache.cpp
+++ b/Engine/Source/Runtime/Experimental/IoStoreOnDemand/Private/JournaledCache.cpp
@@ -711,6 +711,11 @@ bool FDiskCache::Materialize(EntryHandle Handle, FIoBuffer& Out, uint32 Offset)
 {
 	TRACE_CPUPROFILER_EVENT_SCOPE(IasCache::Materialize_Disk);
 
+	if (!DataHandle.IsValid())
+	{
+		return false;
+	}
+
 	const FMapEntry& Entry = *(FMapEntry*)Handle;
 
 	uint32 ReadSize = uint32(Entry.Size) - Offset;
@@ -722,11 +727,6 @@ bool FDiskCache::Materialize(EntryHandle Handle, FIoBuffer& Out, uint32 Offset)
 
 	ReadSize = FMath::Min<uint32>(uint32(Out.GetSize()), ReadSize);
 
-	if (!DataHandle.IsValid())
-	{
-		return false;
-	}
-
 	DataHandle->Seek(Entry.DataCursor + Offset);
 	return DataHandle->Read(Out.GetData(), ReadSize);
 }
@@ -1169,13 +1169,19 @@ bool FCache::FEntry::Materialize(FIoBuffer& Out, uint32 Offset)
 		return false;
 	}
 
+	bool Ret = false;
 	switch (HitType)
 	{
-	case EHit::Memory:	((const FMemCache*)Owner)->Materialize(Handle, Out, Offset); break;
-	case EHit::Disk:	((const FDiskCache*)Owner)->Materialize(Handle, Out, Offset); break;
+	case EHit::Memory:	Ret = ((const FMemCache*)Owner)->Materialize(Handle, Out, Offset); break;
+	case EHit::Disk:	Ret = ((const FDiskCache*)Owner)->Materialize(Handle, Out, Offset); break;
 	case EHit::None:	return false;
 	}
 
+	if (!Ret)
+	{
+		return false;
+	}
+
 	FOnDemandIoBackendStats::Get()->OnCacheGet(Out.GetSize());
 	return true;
 }
