Added API to abandon an IAS cache.

#rnx

diff --git a/Engine/Source/Runtime/Experimental/IoStoreOnDemand/Private/IasCache.h b/Engine/Source/Runtime/Experimental/IoStoreOnDemand/Private/IasCache.h
--- a/Engine/Source/Runtime/Experimental/IoStoreOnDemand/Private/IasCache.h
+++ b/Engine/Source/Runtime/Experimental/IoStoreOnDemand/Private/IasCache.h
@@ -21,6 +21,12 @@ class IIasCache
 public:
 	virtual ~IIasCache() = default;
 
+	/** Deletes the IAS object, dropping all data persisted to disk and releasing
+	OS resources. As this also deletes the object thus any unique pointers should
+	be released prior to abandonment; TUniquePtr->Release()->Abandon(). Be sure
+	to cancel and collect any Get() tasks beforehand. */
+	virtual void Abandon() = 0;
+
 	/** Returns whether the specified cache key is present in the cache. */
 	virtual bool ContainsChunk(const FIoHash& Key) const = 0;
 
diff --git a/Engine/Source/Runtime/Experimental/IoStoreOnDemand/Private/JournaledCache.cpp b/Engine/Source/Runtime/Experimental/IoStoreOnDemand/Private/JournaledCache.cpp
--- a/Engine/Source/Runtime/Experimental/IoStoreOnDemand/Private/JournaledCache.cpp
+++ b/Engine/Source/Runtime/Experimental/IoStoreOnDemand/Private/JournaledCache.cpp
@@ -1146,6 +1146,7 @@ public:
 					FCache(FConfig&& Config);
 	uint32			GetAilments() const;
 	bool			Load();
+	void			Drop();
 	uint32			GetDemand() const;
 	FEntry			Get(uint64 Key) const;
 	bool			Put(uint64 Key, FIoBuffer& Data);
@@ -1212,6 +1213,13 @@ bool FCache::Load()
 	return LoadCache(DiskCache);
 }
 
+////////////////////////////////////////////////////////////////////////////////
+void FCache::Drop()
+{
+	FWriteAccess _(FsLock);
+	DiskCache.Drop();
+}
+
 ////////////////////////////////////////////////////////////////////////////////
 uint32 FCache::GetDemand() const
 {
@@ -1468,6 +1476,7 @@ public:
 
 								FJournaledCache() = default;
 	bool						Initialize(const TCHAR* RootDir, const FIasCacheConfig& Config);
+	virtual void				Abandon() override;
 	virtual bool				ContainsChunk(const FIoHash& Key) const override;
 	virtual GetRetType			Get(const FIoHash& Key, const FIoReadOptions& Options, const FIoCancellationToken* CancellationToken) override;
 	virtual FIoStatus			Put(const FIoHash& Key, FIoBuffer& Data) override;
@@ -1538,6 +1547,16 @@ bool FJournaledCache::Initialize(const TCHAR* RootDir, const FIasCacheConfig& Co
 	return true;
 }
 
+////////////////////////////////////////////////////////////////////////////////
+void FJournaledCache::Abandon()
+{
+	LLM_SCOPE_BYTAG(Ias);
+
+	Thread.Reset();
+	Cache->Drop();
+	delete this;
+}
+
 ////////////////////////////////////////////////////////////////////////////////
 void FJournaledCache::StartThread()
 {
