Added mechanism for getting output from vehicles modules to game thread.

diff --git a/Engine/Source/Runtime/Experimental/ChaosVehicles/ChaosVehiclesCore/Private/SimModule/EngineModule.cpp b/Engine/Source/Runtime/Experimental/ChaosVehicles/ChaosVehiclesCore/Private/SimModule/EngineModule.cpp
--- a/Engine/Source/Runtime/Experimental/ChaosVehicles/ChaosVehiclesCore/Private/SimModule/EngineModule.cpp
+++ b/Engine/Source/Runtime/Experimental/ChaosVehicles/ChaosVehiclesCore/Private/SimModule/EngineModule.cpp
@@ -71,6 +71,29 @@ namespace Chaos
 		return Setup().TorqueCurve.GetValue(RPM, Setup().MaxRPM, Setup().MaxTorque);
 	}
 
+	void FEngineOutputData::FillOutputState(const ISimulationModuleBase* SimModule)
+	{
+		check(SimModule->GetSimType() == eSimType::Engine);
+		if (const FEngineSimModule* Sim = static_cast<const FEngineSimModule*>(SimModule))
+		{
+			RPM = Sim->GetRPM();
+		}
+	}
+
+	void FEngineOutputData::Lerp(const FSimOutputData& InCurrent, const FSimOutputData& InNext, float Alpha)
+	{
+		const FEngineOutputData& Current = static_cast<const FEngineOutputData&>(InCurrent);
+		const FEngineOutputData& Next = static_cast<const FEngineOutputData&>(InNext);
+
+		RPM = FMath::Lerp(Current.RPM, Next.RPM, Alpha);
+	}
+
+	FString FEngineOutputData::ToString()
+	{
+		return  FString::Printf(TEXT("RPM=%3.3f")
+			, RPM);
+	}
+
 } // namespace Chaos
 
 #if VEHICLE_DEBUGGING_ENABLED
diff --git a/Engine/Source/Runtime/Experimental/ChaosVehicles/ChaosVehiclesCore/Private/SimModule/SuspensionModule.cpp b/Engine/Source/Runtime/Experimental/ChaosVehicles/ChaosVehiclesCore/Private/SimModule/SuspensionModule.cpp
--- a/Engine/Source/Runtime/Experimental/ChaosVehicles/ChaosVehiclesCore/Private/SimModule/SuspensionModule.cpp
+++ b/Engine/Source/Runtime/Experimental/ChaosVehicles/ChaosVehiclesCore/Private/SimModule/SuspensionModule.cpp
@@ -16,6 +16,7 @@ namespace Chaos
 		: TSimModuleSettings<FSuspensionSettings>(Settings)
 		, SpringDisplacement(0.f)
 		, LastDisplacement(0.f)
+		, SpringSpeed(0.f)
 		, WheelSimTreeIndex(INVALID_IDX)
 	{
 		AccessSetup().MaxLength = FMath::Abs(Settings.MaxRaise + Settings.MaxDrop);
@@ -51,7 +52,7 @@ namespace Chaos
 		if (SpringDisplacement > 0)
 		{
 			float Damping = (SpringDisplacement < LastDisplacement) ? Setup().CompressionDamping : Setup().ReboundDamping;
-			float SpringSpeed = (LastDisplacement - SpringDisplacement) / DeltaTime;
+			SpringSpeed = (LastDisplacement - SpringDisplacement) / DeltaTime;
 
 			float StiffnessForce = SpringDisplacement * Setup().SpringRate;
 			float DampingForce = SpringSpeed * Damping;
@@ -96,6 +97,7 @@ namespace Chaos
 
 	void FSuspensionSimModuleDatas::FillSimState(ISimulationModuleBase* SimModule)
 	{
+		check(SimModule->GetSimType() == eSimType::Suspension);
 		if (FSuspensionSimModule* Sim = static_cast<FSuspensionSimModule*>(SimModule))
 		{
 			Sim->SpringDisplacement = SpringDisplacement;
@@ -105,6 +107,7 @@ namespace Chaos
 
 	void FSuspensionSimModuleDatas::FillNetState(const ISimulationModuleBase* SimModule)
 	{
+		check(SimModule->GetSimType() == eSimType::Suspension);
 		if (const FSuspensionSimModule* Sim = static_cast<const FSuspensionSimModule*>(SimModule))
 		{
 			SpringDisplacement = Sim->SpringDisplacement;
@@ -129,6 +132,30 @@ namespace Chaos
 	}
 #endif
 
+	void FSuspensionOutputData::FillOutputState(const ISimulationModuleBase* SimModule)
+	{
+		check(SimModule->GetSimType() == eSimType::Suspension);
+		if (const FSuspensionSimModule* Sim = static_cast<const FSuspensionSimModule*>(SimModule))
+		{
+			SpringDisplacement = Sim->SpringDisplacement;
+			SpringSpeed = Sim->SpringSpeed;
+		}
+	}
+
+	void FSuspensionOutputData::Lerp(const FSimOutputData& InCurrent, const FSimOutputData& InNext, float Alpha)
+	{
+		const FSuspensionOutputData& Current = static_cast<const FSuspensionOutputData&>(InCurrent);
+		const FSuspensionOutputData& Next = static_cast<const FSuspensionOutputData&>(InNext);
+
+		SpringDisplacement = FMath::Lerp(Current.SpringDisplacement, Next.SpringDisplacement, Alpha);
+		SpringSpeed = FMath::Lerp(Current.SpringSpeed, Next.SpringSpeed, Alpha);
+	}
+
+	FString FSuspensionOutputData::ToString()
+	{
+		return FString::Printf(TEXT("SpringDisplacement=%3.3f, SpringSpeed=%3.3f"), SpringDisplacement, SpringSpeed);
+	}
+
 } // namespace Chaos
 
 #if VEHICLE_DEBUGGING_ENABLED
diff --git a/Engine/Source/Runtime/Experimental/ChaosVehicles/ChaosVehiclesCore/Private/SimModule/TransmissionModule.cpp b/Engine/Source/Runtime/Experimental/ChaosVehicles/ChaosVehiclesCore/Private/SimModule/TransmissionModule.cpp
--- a/Engine/Source/Runtime/Experimental/ChaosVehicles/ChaosVehiclesCore/Private/SimModule/TransmissionModule.cpp
+++ b/Engine/Source/Runtime/Experimental/ChaosVehicles/ChaosVehiclesCore/Private/SimModule/TransmissionModule.cpp
@@ -135,6 +135,7 @@ namespace Chaos
 
 	void FTransmissionSimModuleDatas::FillSimState(ISimulationModuleBase* SimModule)
 	{
+		check(SimModule->GetSimType() == eSimType::Transmission);
 		if(FTransmissionSimModule* Sim = static_cast<FTransmissionSimModule*>(SimModule))
 		{
 			Sim->CurrentGear = CurrentGear;
@@ -145,6 +146,7 @@ namespace Chaos
 
 	void FTransmissionSimModuleDatas::FillNetState(const ISimulationModuleBase* SimModule)
 	{
+		check(SimModule->GetSimType() == eSimType::Transmission);
 		if (const FTransmissionSimModule* Sim = static_cast<const FTransmissionSimModule*>(SimModule))
 		{
 			CurrentGear = Sim->CurrentGear;
@@ -171,6 +173,28 @@ namespace Chaos
 	}
 #endif
 
+	void FTransmissionOutputData::FillOutputState(const ISimulationModuleBase* SimModule)
+	{
+		check(SimModule->GetSimType() == eSimType::Transmission);
+		if (const FTransmissionSimModule* Sim = static_cast<const FTransmissionSimModule*>(SimModule))
+		{
+			CurrentGear = Sim->CurrentGear;
+		}
+	}
+
+	void FTransmissionOutputData::Lerp(const FSimOutputData& InCurrent, const FSimOutputData& InNext, float Alpha)
+	{
+		const FTransmissionOutputData& Current = static_cast<const FTransmissionOutputData&>(InCurrent);
+		const FTransmissionOutputData& Next = static_cast<const FTransmissionOutputData&>(InNext);
+
+		CurrentGear = Current.CurrentGear;
+	}
+
+	FString FTransmissionOutputData::ToString()
+	{
+		return  FString::Printf(TEXT("CurrentGear=%d")
+			, CurrentGear);
+	}
 
 } // namespace Chaos
 
diff --git a/Engine/Source/Runtime/Experimental/ChaosVehicles/ChaosVehiclesCore/Private/SimModule/WheelModule.cpp b/Engine/Source/Runtime/Experimental/ChaosVehicles/ChaosVehiclesCore/Private/SimModule/WheelModule.cpp
--- a/Engine/Source/Runtime/Experimental/ChaosVehicles/ChaosVehiclesCore/Private/SimModule/WheelModule.cpp
+++ b/Engine/Source/Runtime/Experimental/ChaosVehicles/ChaosVehiclesCore/Private/SimModule/WheelModule.cpp
@@ -43,7 +43,7 @@ namespace Chaos
 			BrakeTorque = Setup().HandbrakeTorque;
 		}
 
-		bool bTouchingGround = ForceIntoSurface > SMALL_NUMBER;
+		bTouchingGround = ForceIntoSurface > SMALL_NUMBER;
 
 		if (bTouchingGround)
 		{
@@ -57,7 +57,7 @@ namespace Chaos
 			float TorqueFromGroundInteraction = Delta * Setup().WheelInertia / DeltaTime; // torque from wheels moving over terrain
 
 			// X is longitudinal direction, Y is lateral
-			float SlipAngle = FVehicleUtility::CalculateSlipAngle(LocalWheelVelocity.Y, LocalWheelVelocity.X);
+			SlipAngle = FVehicleUtility::CalculateSlipAngle(LocalWheelVelocity.Y, LocalWheelVelocity.X);
 
 			float AppliedLinearDriveForce = DriveTorque / Re;
 			float AppliedLinearBrakeForce = FMath::Abs(BrakeTorque) / Re;
@@ -226,6 +226,36 @@ namespace Chaos
 	}
 
 
+	inline void FWheelOutputData::FillOutputState(const ISimulationModuleBase* SimModule)
+	{
+		check(SimModule->GetSimType() == eSimType::Wheel);
+		if (const FWheelSimModule* Sim = static_cast<const FWheelSimModule*>(SimModule))
+		{
+			bTouchingGround = Sim->bTouchingGround;
+			ForceIntoSurface = Sim->ForceIntoSurface;
+			SlipAngle = Sim->SlipAngle;
+			RPM = Sim->GetRPM();
+		}
+	}
+
+	void FWheelOutputData::Lerp(const FSimOutputData& InCurrent, const FSimOutputData& InNext, float Alpha)
+	{
+		const FWheelOutputData& Current = static_cast<const FWheelOutputData&>(InCurrent);
+		const FWheelOutputData& Next = static_cast<const FWheelOutputData&>(InNext);
+
+		bTouchingGround = Current.bTouchingGround;
+		ForceIntoSurface = FMath::Lerp(Current.ForceIntoSurface, Next.ForceIntoSurface, Alpha);
+		SlipAngle = FMath::Lerp(Current.SlipAngle, Next.SlipAngle, Alpha);
+		RPM = FMath::Lerp(Current.RPM, Next.RPM, Alpha);
+	}
+
+	FString FWheelOutputData::ToString()
+	{
+		return  FString::Printf(TEXT("bTouchingGround=%d, ForceIntoSurface=%3.3f, SlipAngle=%3.3f, RPM=%3.3f")
+			, bTouchingGround, ForceIntoSurface, SlipAngle, RPM);
+	}
+
+
 } // namespace Chaos
 
 
diff --git a/Engine/Source/Runtime/Experimental/ChaosVehicles/ChaosVehiclesCore/Public/SimModule/ClutchModule.h b/Engine/Source/Runtime/Experimental/ChaosVehicles/ChaosVehiclesCore/Public/SimModule/ClutchModule.h
--- a/Engine/Source/Runtime/Experimental/ChaosVehicles/ChaosVehiclesCore/Public/SimModule/ClutchModule.h
+++ b/Engine/Source/Runtime/Experimental/ChaosVehicles/ChaosVehiclesCore/Public/SimModule/ClutchModule.h
@@ -21,6 +21,32 @@ namespace Chaos
 		float ClutchStrength;
 	};
 
+
+	struct CHAOSVEHICLESCORE_API FClutchSimModuleDatas : public FTorqueSimModuleDatas
+	{
+#if !(UE_BUILD_SHIPPING || UE_BUILD_TEST)
+		FClutchSimModuleDatas(int NodeArrayIndex, const FString& InDebugString) : FTorqueSimModuleDatas(NodeArrayIndex, InDebugString) {}
+#else
+		FClutchSimModuleDatas(int NodeArrayIndex) : FModuleNetData(NodeArrayIndex) {}
+#endif
+
+		virtual eSimType GetType() override { return eSimType::Clutch; }
+
+		virtual void FillSimState(ISimulationModuleBase* SimModule) override
+		{
+			check(SimModule->GetSimType() == eSimType::Clutch);
+			FTorqueSimModuleDatas::FillSimState(SimModule);
+		}
+
+		virtual void FillNetState(const ISimulationModuleBase* SimModule) override
+		{
+			check(SimModule->GetSimType() == eSimType::Clutch);
+			FTorqueSimModuleDatas::FillNetState(SimModule);
+		}
+
+	};
+
+
 	/// <summary>
 	/// 
 	/// a vehicle component that transmits torque from one source to another through a clutch system, i.e. connect an engine to a transmission
@@ -36,6 +62,16 @@ namespace Chaos
 
 		FClutchSimModule(const FClutchSettings& Settings);
 
+		virtual TSharedPtr<FModuleNetData> GenerateNetData(int SimArrayIndex) const
+		{
+			return MakeShared<FClutchSimModuleDatas>(
+				SimArrayIndex
+#if !(UE_BUILD_SHIPPING || UE_BUILD_TEST)
+				, GetDebugName()
+#endif			
+			);
+		}
+
 		virtual eSimType GetSimType() const { return eSimType::Clutch; }
 
 		virtual const FString GetDebugName() const { return TEXT("Clutch"); }
diff --git a/Engine/Source/Runtime/Experimental/ChaosVehicles/ChaosVehiclesCore/Public/SimModule/EngineModule.h b/Engine/Source/Runtime/Experimental/ChaosVehicles/ChaosVehiclesCore/Public/SimModule/EngineModule.h
--- a/Engine/Source/Runtime/Experimental/ChaosVehicles/ChaosVehiclesCore/Public/SimModule/EngineModule.h
+++ b/Engine/Source/Runtime/Experimental/ChaosVehicles/ChaosVehiclesCore/Public/SimModule/EngineModule.h
@@ -9,6 +9,42 @@ namespace Chaos
 	struct FAllInputs;
 	class FSimModuleTree;
 
+	struct CHAOSVEHICLESCORE_API FEngineSimModuleDatas : public FTorqueSimModuleDatas
+	{
+#if !(UE_BUILD_SHIPPING || UE_BUILD_TEST)
+		FEngineSimModuleDatas(int NodeArrayIndex, const FString& InDebugString) : FTorqueSimModuleDatas(NodeArrayIndex, InDebugString) {}
+#else
+		FEngineSimModuleDatas(int NodeArrayIndex) : FModuleNetData(NodeArrayIndex) {}
+#endif
+
+		virtual eSimType GetType() override { return eSimType::Engine; }
+
+		virtual void FillSimState(ISimulationModuleBase* SimModule) override
+		{
+			check(SimModule->GetSimType() == eSimType::Engine);
+			FTorqueSimModuleDatas::FillSimState(SimModule);
+		}
+
+		virtual void FillNetState(const ISimulationModuleBase* SimModule) override
+		{
+			check(SimModule->GetSimType() == eSimType::Engine);
+			FTorqueSimModuleDatas::FillNetState(SimModule);
+		}
+
+	};
+
+	struct CHAOSVEHICLESCORE_API FEngineOutputData : public FSimOutputData
+	{
+		virtual FSimOutputData* MakeNewData() override { return FEngineOutputData::MakeNew(); }
+		static FSimOutputData* MakeNew() { return new FEngineOutputData(); }
+
+		FORCEINLINE virtual void FillOutputState(const ISimulationModuleBase* SimModule) override;
+		FORCEINLINE virtual void Lerp(const FSimOutputData& InCurrent, const FSimOutputData& InNext, float Alpha) override;
+		FORCEINLINE virtual FString ToString() override;
+
+		float RPM;
+	};
+
 	struct CHAOSVEHICLESCORE_API FEngineSettings
 	{
 		FEngineSettings()
@@ -54,6 +90,20 @@ namespace Chaos
 
 		virtual ~FEngineSimModule() {}
 
+		virtual TSharedPtr<FModuleNetData> GenerateNetData(int SimArrayIndex) const
+		{
+			return MakeShared<FEngineSimModuleDatas>(
+				SimArrayIndex
+#if !(UE_BUILD_SHIPPING || UE_BUILD_TEST)
+				, GetDebugName()
+#endif			
+			);
+		}
+		virtual FSimOutputData* GenerateOutputData() const override
+		{
+			return FEngineOutputData::MakeNew();
+		}
+
 		virtual eSimType GetSimType() const { return eSimType::Engine; }
 
 		virtual const FString GetDebugName() const { return TEXT("Engine"); }
diff --git a/Engine/Source/Runtime/Experimental/ChaosVehicles/ChaosVehiclesCore/Public/SimModule/SimulationModuleBase.h b/Engine/Source/Runtime/Experimental/ChaosVehicles/ChaosVehiclesCore/Public/SimModule/SimulationModuleBase.h
--- a/Engine/Source/Runtime/Experimental/ChaosVehicles/ChaosVehiclesCore/Public/SimModule/SimulationModuleBase.h
+++ b/Engine/Source/Runtime/Experimental/ChaosVehicles/ChaosVehiclesCore/Public/SimModule/SimulationModuleBase.h
@@ -22,6 +22,7 @@ namespace Chaos
 {
 	class FSimModuleTree;
 	struct FModuleNetData;
+	struct FSimOutputData;
 	class FClusterUnionPhysicsProxy;
 
 	struct CHAOSVEHICLESCORE_API FControlInputs
@@ -295,6 +296,8 @@ namespace Chaos
 		// this is the replication datas
 		virtual TSharedPtr<FModuleNetData> GenerateNetData(int NodeArrayIndex) const = 0;
 
+		virtual FSimOutputData* GenerateOutputData() const { return nullptr; }
+
 		//void SetClusterParticle(FPBDRigidClusteredParticleHandle* ParticleIn) { ClusterParticle = ParticleIn; }
 		Chaos::FPBDRigidClusteredParticleHandle* GetClusterParticle(Chaos::FClusterUnionPhysicsProxy* Proxy);
 
@@ -353,6 +356,18 @@ namespace Chaos
 
 	using FModuleNetDataArray = TArray<TSharedPtr<FModuleNetData>>;
 
+	struct CHAOSVEHICLESCORE_API FSimOutputData
+	{
+		FSimOutputData() = default;
+		virtual ~FSimOutputData() {}
+
+		virtual FSimOutputData* MakeNewData() = 0;
+		virtual void FillOutputState(const ISimulationModuleBase* SimModule) = 0;
+		virtual void Lerp(const FSimOutputData& InCurrent, const FSimOutputData& InNext, float Alpha) = 0;
+		virtual FString ToString() { return FString(); }
+	};
+
+
 } // namespace Chaos
 
 
diff --git a/Engine/Source/Runtime/Experimental/ChaosVehicles/ChaosVehiclesCore/Public/SimModule/SuspensionModule.h b/Engine/Source/Runtime/Experimental/ChaosVehicles/ChaosVehiclesCore/Public/SimModule/SuspensionModule.h
--- a/Engine/Source/Runtime/Experimental/ChaosVehicles/ChaosVehiclesCore/Public/SimModule/SuspensionModule.h
+++ b/Engine/Source/Runtime/Experimental/ChaosVehicles/ChaosVehiclesCore/Public/SimModule/SuspensionModule.h
@@ -40,7 +40,19 @@ namespace Chaos
 		float LastDisplacement = 0.0f;
 	};
 
+	struct CHAOSVEHICLESCORE_API FSuspensionOutputData : public FSimOutputData
+	{
+		virtual FSimOutputData* MakeNewData() override { return FSuspensionOutputData::MakeNew(); }
+		static FSimOutputData* MakeNew() { return new FSuspensionOutputData(); }
+
+		virtual void FillOutputState(const ISimulationModuleBase* SimModule) override;
+		virtual void Lerp(const FSimOutputData& InCurrent, const FSimOutputData& InNext, float Alpha) override;
 
+		virtual FString ToString() override;
+
+		float SpringDisplacement;
+		float SpringSpeed;
+	};
 
 	struct CHAOSVEHICLESCORE_API FSuspensionSettings
 	{
@@ -98,6 +110,7 @@ namespace Chaos
 	class CHAOSVEHICLESCORE_API FSuspensionSimModule : public ISimulationModuleBase, public TSimModuleSettings<FSuspensionSettings>
 	{
 		friend FSuspensionSimModuleDatas;
+		friend FSuspensionOutputData;
 
 	public:
 
@@ -114,6 +127,11 @@ namespace Chaos
 			);
 		}
 
+		virtual FSimOutputData* GenerateOutputData() const override
+		{
+			return FSuspensionOutputData::MakeNew();
+		}
+
 		virtual eSimType GetSimType() const { return eSimType::Suspension; }
 
 		virtual const FString GetDebugName() const { return TEXT("Suspension"); }
@@ -135,6 +153,7 @@ namespace Chaos
 
 		float SpringDisplacement;
 		float LastDisplacement;
+		float SpringSpeed;
 		int WheelSimTreeIndex;
 	};
 
diff --git a/Engine/Source/Runtime/Experimental/ChaosVehicles/ChaosVehiclesCore/Public/SimModule/TorqueSimModule.h b/Engine/Source/Runtime/Experimental/ChaosVehicles/ChaosVehiclesCore/Public/SimModule/TorqueSimModule.h
--- a/Engine/Source/Runtime/Experimental/ChaosVehicles/ChaosVehiclesCore/Public/SimModule/TorqueSimModule.h
+++ b/Engine/Source/Runtime/Experimental/ChaosVehicles/ChaosVehiclesCore/Public/SimModule/TorqueSimModule.h
@@ -63,15 +63,6 @@ public:
 	 * Is Module of a specific type - used for casting
 	 */
 	virtual bool IsBehaviourType(eSimModuleTypeFlags InType) const { return (InType & TorqueBased); }
-	virtual TSharedPtr<FModuleNetData> GenerateNetData(int SimArrayIndex) const
-	{
-		return MakeShared<FTorqueSimModuleDatas>(
-			SimArrayIndex
-#if !(UE_BUILD_SHIPPING || UE_BUILD_TEST)
-			, GetDebugName()
-#endif			
-		);
-	}
 
 	virtual eSimType GetSimType() const { return eSimType::TorqueSim; }
 
diff --git a/Engine/Source/Runtime/Experimental/ChaosVehicles/ChaosVehiclesCore/Public/SimModule/TransmissionModule.h b/Engine/Source/Runtime/Experimental/ChaosVehicles/ChaosVehiclesCore/Public/SimModule/TransmissionModule.h
--- a/Engine/Source/Runtime/Experimental/ChaosVehicles/ChaosVehiclesCore/Public/SimModule/TransmissionModule.h
+++ b/Engine/Source/Runtime/Experimental/ChaosVehicles/ChaosVehiclesCore/Public/SimModule/TransmissionModule.h
@@ -43,6 +43,17 @@ namespace Chaos
 		float CurrentGearChangeTime = 0.0f;
 	};
 
+	struct CHAOSVEHICLESCORE_API FTransmissionOutputData : public FSimOutputData
+	{
+		virtual FSimOutputData* MakeNewData() override { return FTransmissionOutputData::MakeNew(); }
+		static FSimOutputData* MakeNew() { return new FTransmissionOutputData(); }
+
+		FORCEINLINE virtual void FillOutputState(const ISimulationModuleBase* SimModule) override;
+		FORCEINLINE virtual void Lerp(const FSimOutputData& InCurrent, const FSimOutputData& InNext, float Alpha) override;
+		FORCEINLINE virtual FString ToString() override;
+
+		int32 CurrentGear;
+	};
 
 	struct CHAOSVEHICLESCORE_API FTransmissionSettings
 	{
@@ -87,6 +98,7 @@ namespace Chaos
 	class CHAOSVEHICLESCORE_API FTransmissionSimModule : public FTorqueSimModule, public TSimModuleSettings<FTransmissionSettings>
 	{
 		friend FTransmissionSimModuleDatas;
+		friend FTransmissionOutputData;
 
 	public:
 
@@ -109,6 +121,11 @@ namespace Chaos
 			);
 		}
 
+		virtual FSimOutputData* GenerateOutputData() const override
+		{
+			return FTransmissionOutputData::MakeNew();
+		}
+
 		virtual eSimType GetSimType() const { return eSimType::Transmission; }
 
 		virtual const FString GetDebugName() const { return TEXT("Transmission"); }
diff --git a/Engine/Source/Runtime/Experimental/ChaosVehicles/ChaosVehiclesCore/Public/SimModule/WheelModule.h b/Engine/Source/Runtime/Experimental/ChaosVehicles/ChaosVehiclesCore/Public/SimModule/WheelModule.h
--- a/Engine/Source/Runtime/Experimental/ChaosVehicles/ChaosVehiclesCore/Public/SimModule/WheelModule.h
+++ b/Engine/Source/Runtime/Experimental/ChaosVehicles/ChaosVehiclesCore/Public/SimModule/WheelModule.h
@@ -16,6 +16,48 @@ namespace Chaos
 		Y	// Y forward
 	};
 
+	struct CHAOSVEHICLESCORE_API FWheelSimModuleDatas : public FTorqueSimModuleDatas
+	{
+#if !(UE_BUILD_SHIPPING || UE_BUILD_TEST)
+		FWheelSimModuleDatas(int NodeArrayIndex, const FString& InDebugString) : FTorqueSimModuleDatas(NodeArrayIndex, InDebugString) {}
+#else
+		FWheelSimModuleDatas(int NodeArrayIndex) : FModuleNetData(NodeArrayIndex) {}
+#endif
+
+		virtual eSimType GetType() override { return eSimType::Wheel; }
+
+		virtual void FillSimState(ISimulationModuleBase* SimModule) override
+		{
+			check(SimModule->GetSimType() == eSimType::Wheel);
+			FTorqueSimModuleDatas::FillSimState(SimModule);
+		}
+
+		virtual void FillNetState(const ISimulationModuleBase* SimModule) override
+		{
+			check(SimModule->GetSimType() == eSimType::Wheel);
+			FTorqueSimModuleDatas::FillNetState(SimModule);
+		}
+
+	};
+
+	struct CHAOSVEHICLESCORE_API FWheelOutputData : public FSimOutputData
+	{
+		virtual FSimOutputData* MakeNewData() override { return FWheelOutputData::MakeNew(); }
+		static FSimOutputData* MakeNew() { return new FWheelOutputData(); }
+
+		FORCEINLINE virtual void FillOutputState(const ISimulationModuleBase* SimModule) override;
+		FORCEINLINE virtual void Lerp(const FSimOutputData& InCurrent, const FSimOutputData& InNext, float Alpha) override;
+		FORCEINLINE virtual FString ToString() override;
+
+		bool bTouchingGround;
+		float ForceIntoSurface;
+		float SlipAngle;
+		float RPM;
+
+		//HitLocation
+		//PhysMaterial
+	};
+
 	struct CHAOSVEHICLESCORE_API FWheelSettings
 	{
 		FWheelSettings()
@@ -68,10 +110,27 @@ namespace Chaos
 
 	class CHAOSVEHICLESCORE_API FWheelSimModule : public FTorqueSimModule, public TSimModuleSettings<FWheelSettings>
 	{
+		friend FWheelOutputData;
+
 	public:
 
 		FWheelSimModule(const FWheelSettings& Settings);
 
+		virtual TSharedPtr<FModuleNetData> GenerateNetData(int SimArrayIndex) const
+		{
+			return MakeShared<FWheelSimModuleDatas>(
+				SimArrayIndex
+#if !(UE_BUILD_SHIPPING || UE_BUILD_TEST)
+				, GetDebugName()
+#endif			
+			);
+		}
+
+		virtual FSimOutputData* GenerateOutputData() const override
+		{
+			return FWheelOutputData::MakeNew();
+		}
+
 		virtual eSimType GetSimType() const { return eSimType::Wheel; }
 
 		virtual const FString GetDebugName() const { return TEXT("Wheel"); }
@@ -104,6 +163,9 @@ namespace Chaos
 		float MassPerWheel;
 		float SteerAngleDegrees;
 
+		// for output
+		bool bTouchingGround;
+		float SlipAngle;
 	};
 
 
