Chaos: Fixed cluster union material identification. Previously we were only checking for per-particle materials which are almost always null so we set a null material and used a default material profile. That default led to very small linear and angular sleep thresholds for cluster unions (linear was 0.00001mm). Now the cluster union will pick up the first valid material from its children and apply that to itself.

diff --git a/Engine/Source/Runtime/Experimental/Chaos/Private/Chaos/ClusterUnionManager.cpp b/Engine/Source/Runtime/Experimental/Chaos/Private/Chaos/ClusterUnionManager.cpp
--- a/Engine/Source/Runtime/Experimental/Chaos/Private/Chaos/ClusterUnionManager.cpp
+++ b/Engine/Source/Runtime/Experimental/Chaos/Private/Chaos/ClusterUnionManager.cpp
@@ -586,6 +586,7 @@ namespace Chaos
 		// For all the particles that have been added to the cluster we need to set their parent proxy to the
 		// cluster's proxy if it exists. We need the proxy type check because for non-cluster union proxy backed unions,
 		// the cluster union's proxy will be the proxy of the most recently added particle.
+		FMaterialHandle MatHandle;
 		if (OldProxy && OldProxyType == EPhysicsProxyType::ClusterUnionProxy)
 		{
 			for (FPBDRigidParticleHandle* Particle : FinalParticlesToAdd)
@@ -593,6 +594,21 @@ namespace Chaos
 				if (Particle && Particle->PhysicsProxy())
 				{
 					Particle->PhysicsProxy()->SetParentProxy(OldProxy);
+
+					if(!MatHandle.InnerHandle.IsValid())
+					{
+						const FShapesArray& Shapes = Particle->ShapesArray();
+						for(const TUniquePtr<FPerShapeData>& ShapeData : Shapes)
+						{
+							const FMaterialData& MaterialData = ShapeData->GetMaterialData();
+
+							if(MaterialData.Materials.Num() > 0)
+							{
+								MatHandle = MaterialData.Materials[0];
+								break;
+							}
+						}
+					}
 				}
 			}
 		}
@@ -617,7 +633,13 @@ namespace Chaos
 				MEvolution.SetParticleObjectState(Cluster->InternalCluster, Chaos::EObjectStateType::Sleeping);
 			}
 
-			MEvolution.SetPhysicsMaterial(Cluster->InternalCluster, MEvolution.GetPhysicsMaterial(FinalParticlesToAdd[0]));
+			if(ensure(MatHandle.InnerHandle.IsValid()))
+			{
+				for(const TUniquePtr<FPerShapeData>& ShapeData : Cluster->InternalCluster->ShapesArray())
+				{
+					ShapeData->SetMaterial(MatHandle);
+				}
+			}
 		}
 	}
 
