Iris - RPC call forwarding

diff --git a/Engine/Source/Runtime/Engine/Classes/Engine/NetDriver.h b/Engine/Source/Runtime/Engine/Classes/Engine/NetDriver.h
--- a/Engine/Source/Runtime/Engine/Classes/Engine/NetDriver.h
+++ b/Engine/Source/Runtime/Engine/Classes/Engine/NetDriver.h
@@ -2095,6 +2095,8 @@ protected:
 	void NotifyGameInstanceUpdated();
 
 private:
+	// Only for ForwardRemoteFunction
+	friend FObjectReplicator;
 
 	ENGINE_API virtual ECreateReplicationChangelistMgrFlags GetCreateReplicationChangelistMgrFlags() const;
 
@@ -2129,6 +2131,9 @@ private:
 
 	void InitNetTraceId();
 
+	/** Called from RPC processing code to forward RPC to other NetDrivers if ShouldForwardFunction returns true. */
+	void ForwardRemoteFunction(UObject* RootObject, UObject* SubObject, UFunction* Function, void* Parms);
+
 	UPROPERTY(transient)
 	TObjectPtr<UReplicationDriver> ReplicationDriver;
 
@@ -2175,7 +2180,6 @@ private:
 	/** Stat tracking for the total number of out of order packets that were duplicates */
 	int32 TotalOutOfOrderPacketsDuplicate = 0;
 
-
 	/** Cached value for UEngine.GlobalNetTravelCount, at the time of NetDriver initialization */
 	uint32 CachedGlobalNetTravelCount = 0;
 };
diff --git a/Engine/Source/Runtime/Engine/Private/DataReplication.cpp b/Engine/Source/Runtime/Engine/Private/DataReplication.cpp
--- a/Engine/Source/Runtime/Engine/Private/DataReplication.cpp
+++ b/Engine/Source/Runtime/Engine/Private/DataReplication.cpp
@@ -31,7 +31,6 @@ LLM_DEFINE_TAG(NetObjReplicator, NAME_None, TEXT("Networking"), GET_STATFNAME(ST
 
 DECLARE_CYCLE_STAT(TEXT("Custom Delta Property Rep Time"), STAT_NetReplicateCustomDeltaPropTime, STATGROUP_Game);
 DECLARE_CYCLE_STAT(TEXT("ReceiveRPC"), STAT_NetReceiveRPC, STATGROUP_Game);
-DECLARE_CYCLE_STAT(TEXT("ReceiveRPC_ProcessRemoteFunction"), STAT_NetReceiveRPC_ProcessRemoteFunction, STATGROUP_Game);
 
 static TAutoConsoleVariable<int32> CVarMaxRPCPerNetUpdate(
 	TEXT("net.MaxRPCPerNetUpdate"),
@@ -1353,29 +1352,14 @@ bool FObjectReplicator::ReceivedRPC(FNetBitReader& Reader, const FReplicationFla
 		else
 		{
 			AActor* OwningActor = OwningChannel->Actor;
+			UObject* const SubObject = Object != OwningChannel->Actor ? Object : nullptr;
 
-			if (Connection->Driver->ShouldForwardFunction(OwningActor, Function, Parms))
-			{
-				FWorldContext* const Context = GEngine->GetWorldContextFromWorld(Connection->Driver->GetWorld());
-				if (Context != nullptr)
-				{
-					UObject* const SubObject = Object != OwningChannel->Actor ? Object : nullptr;
-
-					for (FNamedNetDriver& Driver : Context->ActiveNetDrivers)
-					{
-						if (Driver.NetDriver != nullptr && (Driver.NetDriver != Connection->Driver) && Driver.NetDriver->ShouldReplicateFunction(OwningActor, Function))
-						{
-							SCOPE_CYCLE_COUNTER(STAT_NetReceiveRPC_ProcessRemoteFunction);
-							Driver.NetDriver->ProcessRemoteFunction(OwningActor, Function, Parms, nullptr, nullptr, SubObject);
-						}
-					}
-				}
-			}
+			// Forward the function call.
+			Connection->Driver->ForwardRemoteFunction(OwningActor, SubObject, Function, Parms);
 
 			// Reset errors from replay driver
 			RPC_ResetLastFailedReason();
 
-
 			{
 				UE::Net::FScopedNetContextRPC CallingRPC;
 				// Call the function.
diff --git a/Engine/Source/Runtime/Engine/Private/NetDriver.cpp b/Engine/Source/Runtime/Engine/Private/NetDriver.cpp
--- a/Engine/Source/Runtime/Engine/Private/NetDriver.cpp
+++ b/Engine/Source/Runtime/Engine/Private/NetDriver.cpp
@@ -141,6 +141,7 @@ DECLARE_CYCLE_STAT(TEXT("Process Prioritized Actors Time"), STAT_NetProcessPrior
 DECLARE_CYCLE_STAT(TEXT("NetDriver TickFlush"), STAT_NetTickFlush, STATGROUP_Game);
 DECLARE_CYCLE_STAT(TEXT("NetDriver TickFlush GatherStats"), STAT_NetTickFlushGatherStats, STATGROUP_Game);
 DECLARE_CYCLE_STAT(TEXT("NetDriver TickFlush GatherStatsPerfCounters"), STAT_NetTickFlushGatherStatsPerfCounters, STATGROUP_Game);
+DECLARE_CYCLE_STAT(TEXT("ReceiveRPC_ProcessRemoteFunction"), STAT_NetReceiveRPC_ProcessRemoteFunction, STATGROUP_Game);
 
 DEFINE_LOG_CATEGORY_STATIC(LogNetSyncLoads, Log, All);
 
@@ -6367,6 +6368,7 @@ void UNetDriver::CreateReplicationSystem(bool bInitAsClient)
 		Params.ReplicationBridge = ReplicationBridge;
 		Params.bIsServer = !bInitAsClient;
 		Params.bAllowObjectReplication = !bInitAsClient;
+		Params.ForwardNetRPCCallDelegate.BindUObject(this, &UNetDriver::ForwardRemoteFunction);
 
 		UE::Net::Private::ApplyReplicationSystemConfig(ReplicationSystemConfig, Params);
 
@@ -6926,6 +6928,32 @@ bool UNetDriver::ShouldForwardFunction(AActor* Actor, UFunction* Function, void*
 	return !IsServer();
 }
 
+void UNetDriver::ForwardRemoteFunction(UObject* RootObject, UObject* SubObject, UFunction* Function, void* Parms)
+{
+	AActor* OwningActor = Cast<AActor>(RootObject);
+	if (!ensure(OwningActor != nullptr))
+	{
+		return;
+	}
+
+	if (!ShouldForwardFunction(OwningActor, Function, Parms))
+	{
+		return;
+	}
+
+	if (FWorldContext* Context = GEngine->GetWorldContextFromWorld(GetWorld()))
+	{
+		for (FNamedNetDriver& Driver : Context->ActiveNetDrivers)
+		{
+			if (Driver.NetDriver != nullptr && Driver.NetDriver != this && Driver.NetDriver->ShouldReplicateFunction(OwningActor, Function))
+			{
+				SCOPE_CYCLE_COUNTER(STAT_NetReceiveRPC_ProcessRemoteFunction);
+				Driver.NetDriver->ProcessRemoteFunction(OwningActor, Function, Parms, static_cast<FOutParmRec*>(nullptr), static_cast<FFrame*>(nullptr), SubObject);
+			}
+		}
+	}
+}
+
 bool UNetDriver::ShouldReplicateActor(AActor* Actor) const
 {
 	return (Actor && Actor->GetNetDriverName() == NetDriverName);
diff --git a/Engine/Source/Runtime/Experimental/Iris/Core/Private/Iris/ReplicationSystem/NetBlob/NetRPC.cpp b/Engine/Source/Runtime/Experimental/Iris/Core/Private/Iris/ReplicationSystem/NetBlob/NetRPC.cpp
--- a/Engine/Source/Runtime/Experimental/Iris/Core/Private/Iris/ReplicationSystem/NetBlob/NetRPC.cpp
+++ b/Engine/Source/Runtime/Experimental/Iris/Core/Private/Iris/ReplicationSystem/NetBlob/NetRPC.cpp
@@ -5,6 +5,7 @@
 #include "Iris/Core/BitTwiddling.h"
 #include "Net/Core/Trace/NetTrace.h"
 #include "Net/Core/Misc/NetContext.h"
+#include "Iris/ReplicationSystem/NetBlob/NetRPCHandler.h"
 #include "Iris/ReplicationSystem/ObjectReferenceCache.h"
 #include "Iris/ReplicationSystem/ObjectReplicationBridge.h"
 #include "Iris/Serialization/NetReferenceCollector.h"
@@ -12,6 +13,7 @@
 #include "Iris/ReplicationSystem/ReplicationOperationsInternal.h"
 #include "Iris/ReplicationSystem/ReplicationProtocol.h"
 #include "Iris/ReplicationSystem/ReplicationSystem.h"
+#include "Iris/ReplicationSystem/ReplicationSystemInternal.h"
 #include "Iris/Serialization/InternalNetSerializationContext.h"
 #include "Iris/Serialization/NetBitStreamReader.h"
 #include "Iris/Serialization/NetBitStreamWriter.h"
@@ -30,7 +32,8 @@ namespace UE::Net::Private
 static bool NetRPC_GetFunctionLocator(const UReplicationSystem* ReplicationSystem, const FNetObjectReference& ObjectReference, const UFunction* Function, FNetRPC::FFunctionLocator& OutFunctionLocator, const FReplicationStateMemberFunctionDescriptor*& OutFunctionDescriptor);
 
 static UObject* NetRPC_GetObject(FNetSerializationContext& Context, const FNetObjectReference& ObjectReference);
-static UObject* NetRPC_GetObject(FNetSerializationContext& Context, const FNetObjectReference& RootObjectReference, const FNetObjectReference& SubObjectReference);							                
+static UObject* NetRPC_GetObject(FNetSerializationContext& Context, const FNetObjectReference& RootObjectReference, const FNetObjectReference& SubObjectReference);
+static UObject* NetRPC_GetRootObject(FNetSerializationContext& Context, const FNetObjectReference& ObjectReference);
 
 static bool NetRPC_GetFunctionAndObject(FNetSerializationContext& Context, const FNetObjectReference& RootObjectReference, const FNetObjectReference& SubObjectReferece, const FNetRPC::FFunctionLocator& FunctionLocator, const FReplicationStateMemberFunctionDescriptor*& OutFunctionDescriptor, TWeakObjectPtr<UObject>& OutObject);
 
@@ -418,12 +421,14 @@ FNetRPC* FNetRPC::Create(UReplicationSystem* ReplicationSystem, const FNetBlobCr
 	return NetRPC;
 }
 
-void FNetRPC::CallFunction(FNetSerializationContext& Context)
+void FNetRPC::CallFunction(FNetRPCCallContext& CallContext)
 {
 #if UE_NET_IRIS_CSV_STATS
 	CSV_SCOPED_TIMING_STAT_EXCLUSIVE(HandleRPC);
 #endif
 
+	FNetSerializationContext& Context = CallContext.GetNetSerializationContext();
+
 	const UReplicationSystem* ReplicationSystem = Context.GetInternalContext()->ReplicationSystem;
 
 	// Check whether we are ok with calling the function
@@ -483,26 +488,24 @@ void FNetRPC::CallFunction(FNetSerializationContext& Context)
 		bool bLogRpc = true;
 
 		// Suppress spammy engine RPCs. This could be made a configable list in the future.
-		if (Function->GetName().Contains(TEXT("ServerUpdateCamera"))) bLogRpc = false;
-		if (Function->GetName().Contains(TEXT("ClientAckGoodMove"))) bLogRpc = false;
-		if (Function->GetName().Contains(TEXT("ServerMove"))) bLogRpc = false;
-		
-		if (bLogRpc)
+		const FString& FunctionName = Function->GetName();
+		if (   FunctionName.Contains(TEXT("ServerUpdateCamera"))
+			|| FunctionName.Contains(TEXT("ClientAckGoodMove"))
+			|| FunctionName.Contains(TEXT("ServerMove")))
 		{
-			UE_LOG(LogIrisRpc, Verbose, TEXT("Calling %s RPC function %s for %s : %s"), (Function->FunctionFlags & FUNC_NetReliable ? TEXT("reliable") : TEXT("unreliable")), ToCStr(Function->GetName()), *NetObjectReference.ToString(), *Object->GetFullName());
+			bLogRpc = false;
 		}
+		
+		UE_CLOG(bLogRpc, LogIrisRpc, Verbose, TEXT("Calling %hs RPC function %s for %s : %s"), (Function->FunctionFlags & FUNC_NetReliable ? "reliable" : "unreliable"), ToCStr(Function->GetName()), *NetObjectReference.ToString(), *Object->GetFullName());
 	}
-	// Call the function
-	if (Function->ParmsSize == 0)
-	{
-		UE::Net::FScopedNetContextRPC CallingRPC;
-		Object->ProcessEvent(const_cast<UFunction*>(Function), nullptr);
-	}
-	else
+
+	// Setup function parameters
+	uint8* FunctionParameters = nullptr;
+	if (Function->ParmsSize > 0)
 	{
 		check(BlobDescriptor.IsValid());
 
-		uint8* FunctionParameters = static_cast<uint8*>(FMemory_Alloca(Function->ParmsSize));
+		FunctionParameters = static_cast<uint8*>(FMemory_Alloca(Function->ParmsSize));
 		FMemory::Memzero(FunctionParameters, Function->ParmsSize);
 		if (!EnumHasAnyFlags(BlobDescriptor->Traits, EReplicationStateTraits::IsSourceTriviallyConstructible))
 		{
@@ -529,10 +532,25 @@ void FNetRPC::CallFunction(FNetSerializationContext& Context)
 		}
 		
 		FReplicationStateOperations::Dequantize(Context, FunctionParameters, QuantizedBlobState.Get(), BlobDescriptor);
-		{
-			UE::Net::FScopedNetContextRPC CallingRPC;
-			Object->ProcessEvent(const_cast<UFunction*>(Function), FunctionParameters);
-		}
+	}
+
+	// Forward function
+	if (const FForwardNetRPCCallMulticastDelegate& Delegate = CallContext.GetForwardNetRPCCallDelegate(); Delegate.IsBound())
+	{
+		UObject* RootObject = NetRPC_GetRootObject(Context, NetObjectReference);
+		UObject* SubObject = (Object != RootObject ? Object : static_cast<UObject*>(nullptr));
+		Delegate.Broadcast(RootObject, SubObject, const_cast<UFunction*>(Function), FunctionParameters);
+	}
+
+	// Call function
+	{
+		UE::Net::FScopedNetContextRPC CallingRPC;
+		Object->ProcessEvent(const_cast<UFunction*>(Function), FunctionParameters);
+	}
+
+	// Deinitialize function parameters
+	if (FunctionParameters != nullptr)
+	{
 
 		if (!EnumHasAnyFlags(BlobDescriptor->Traits, EReplicationStateTraits::IsSourceTriviallyDestructible))
 		{
@@ -597,6 +615,29 @@ static UObject* NetRPC_GetObject(FNetSerializationContext& Context, const FNetOb
 	return InternalContext->ObjectReferenceCache->ResolveObjectReference(ObjectReference, InternalContext->ResolveContext);
 }
 
+static UObject* NetRPC_GetRootObject(FNetSerializationContext& Context, const FNetObjectReference& ObjectReference)
+{
+	FInternalNetSerializationContext* InternalContext = Context.GetInternalContext();
+	const FNetRefHandleManager& NetRefHandleManager = InternalContext->ReplicationSystem->GetReplicationSystemInternal()->GetNetRefHandleManager();
+	
+	FNetRefHandle OwnerRefHandle = ObjectReference.GetRefHandle();
+	const FInternalNetRefIndex InternalIndex = NetRefHandleManager.GetInternalIndex(OwnerRefHandle);
+	if (!ensureMsgf(InternalIndex != FNetRefHandleManager::InvalidInternalIndex, TEXT("Unable to find InternalIndex for object reference %s"), ToCStr(ObjectReference.ToString())))
+	{
+		return nullptr;
+	}
+
+	const FNetRefHandleManager::FReplicatedObjectData& ObjectData = NetRefHandleManager.GetReplicatedObjectDataNoCheck(InternalIndex);
+	if (ObjectData.SubObjectRootIndex != FNetRefHandleManager::InvalidInternalIndex)
+	{
+		const FNetRefHandleManager::FReplicatedObjectData& RootObjectData = NetRefHandleManager.GetReplicatedObjectDataNoCheck(ObjectData.SubObjectRootIndex);
+		OwnerRefHandle = RootObjectData.RefHandle;
+	}
+
+	UObject* RootObject = InternalContext->ObjectReferenceCache->ResolveObjectReferenceHandle(OwnerRefHandle, InternalContext->ResolveContext);
+	return RootObject;
+}
+
 static bool NetRPC_GetFunctionAndObject(FNetSerializationContext& Context, const FNetObjectReference& ObjectReference, const FNetObjectReference& SubObjectReference, const FNetRPC::FFunctionLocator& FunctionLocator, const FReplicationStateMemberFunctionDescriptor*& OutFunctionDescriptor, TWeakObjectPtr<UObject>& OutObject)
 {
 	const UReplicationSystem* ReplicationSystem = Context.GetInternalContext()->ReplicationSystem;
diff --git a/Engine/Source/Runtime/Experimental/Iris/Core/Private/Iris/ReplicationSystem/NetBlob/NetRPC.h b/Engine/Source/Runtime/Experimental/Iris/Core/Private/Iris/ReplicationSystem/NetBlob/NetRPC.h
--- a/Engine/Source/Runtime/Experimental/Iris/Core/Private/Iris/ReplicationSystem/NetBlob/NetRPC.h
+++ b/Engine/Source/Runtime/Experimental/Iris/Core/Private/Iris/ReplicationSystem/NetBlob/NetRPC.h
@@ -4,6 +4,11 @@
 #include "Iris/ReplicationSystem/NetBlob/NetBlob.h"
 #include "UObject/WeakObjectPtr.h"
 
+namespace UE::Net
+{
+	class FNetRPCCallContext;
+}
+
 namespace UE::Net::Private
 {
 
@@ -22,7 +27,7 @@ public:
 
 	static FNetRPC* Create(UReplicationSystem* ReplicationSystem, const FNetBlobCreationInfo& CreationInfo, const FNetObjectReference& ObjectReference, const UFunction* Function, const void* FunctionParameters);
 
-	void CallFunction(FNetSerializationContext& Context);
+	void CallFunction(FNetRPCCallContext& Context);
 
 	void SetFunctionLocator(const FFunctionLocator& InfFunctionLocator) { FunctionLocator = InfFunctionLocator; }
 	const FFunctionLocator& GetFunctionLocator() const { return FunctionLocator; }
diff --git a/Engine/Source/Runtime/Experimental/Iris/Core/Private/Iris/ReplicationSystem/NetBlob/NetRPCHandler.cpp b/Engine/Source/Runtime/Experimental/Iris/Core/Private/Iris/ReplicationSystem/NetBlob/NetRPCHandler.cpp
--- a/Engine/Source/Runtime/Experimental/Iris/Core/Private/Iris/ReplicationSystem/NetBlob/NetRPCHandler.cpp
+++ b/Engine/Source/Runtime/Experimental/Iris/Core/Private/Iris/ReplicationSystem/NetBlob/NetRPCHandler.cpp
@@ -2,9 +2,10 @@
 
 #include "Iris/ReplicationSystem/NetBlob/NetRPCHandler.h"
 #include "Iris/ReplicationSystem/NetBlob/NetRPC.h"
+#include "Iris/ReplicationSystem/ReplicationSystem.h"
+#include "Iris/ReplicationSystem/ReplicationSystemInternal.h"
 
 UNetRPCHandler::UNetRPCHandler()
-: ReplicationSystem(nullptr)
 {
 }
 
@@ -32,8 +33,10 @@ TRefCountPtr<UE::Net::FNetBlob> UNetRPCHandler::CreateNetBlob(const FNetBlobCrea
 	return RPC;
 }
 
-void UNetRPCHandler::OnNetBlobReceived(UE::Net::FNetSerializationContext& Context, const TRefCountPtr<FNetBlob>& NetBlob)
+void UNetRPCHandler::OnNetBlobReceived(UE::Net::FNetSerializationContext& NetContext, const TRefCountPtr<FNetBlob>& NetBlob)
 {
+	const UE::Net::FForwardNetRPCCallMulticastDelegate& ForwardNetRPCCallDelegate = ReplicationSystem->GetReplicationSystemInternal()->GetForwardNetRPCCallMulticastDelegate();
+	UE::Net::FNetRPCCallContext CallContext(NetContext, ForwardNetRPCCallDelegate);
 	FNetRPC* RPC = static_cast<FNetRPC*>(NetBlob.GetReference());
-	RPC->CallFunction(Context);
+	RPC->CallFunction(CallContext);
 }
diff --git a/Engine/Source/Runtime/Experimental/Iris/Core/Private/Iris/ReplicationSystem/NetBlob/NetRPCHandler.h b/Engine/Source/Runtime/Experimental/Iris/Core/Private/Iris/ReplicationSystem/NetBlob/NetRPCHandler.h
--- a/Engine/Source/Runtime/Experimental/Iris/Core/Private/Iris/ReplicationSystem/NetBlob/NetRPCHandler.h
+++ b/Engine/Source/Runtime/Experimental/Iris/Core/Private/Iris/ReplicationSystem/NetBlob/NetRPCHandler.h
@@ -3,8 +3,27 @@
 #pragma once
 #include "Iris/ReplicationSystem/NetBlob/NetRPC.h"
 #include "Iris/ReplicationSystem/NetBlob/NetBlobHandler.h"
+#include "Iris/ReplicationSystem/ReplicationSystemTypes.h"
 #include "NetRPCHandler.generated.h"
 
+namespace UE::Net
+{
+
+class FNetRPCCallContext
+{
+public:
+	FNetRPCCallContext(FNetSerializationContext& NetContext, const FForwardNetRPCCallMulticastDelegate& ForwardNetRPCCallDelegate);
+
+	FNetSerializationContext& GetNetSerializationContext();
+	const FForwardNetRPCCallMulticastDelegate& GetForwardNetRPCCallDelegate() const;
+
+private:
+	FNetSerializationContext& NetContext;
+	const FForwardNetRPCCallMulticastDelegate& ForwardNetRPCCallDelegate;
+};
+
+}
+
 UCLASS(transient, MinimalAPI)
 class UNetRPCHandler final : public UNetBlobHandler
 {
@@ -24,5 +43,26 @@ private:
 	virtual TRefCountPtr<FNetBlob> CreateNetBlob(const FNetBlobCreationInfo&) const override;
 	virtual void OnNetBlobReceived(UE::Net::FNetSerializationContext& Context, const TRefCountPtr<FNetBlob>& NetBlob) override;
 
-	UReplicationSystem* ReplicationSystem;
+	UReplicationSystem* ReplicationSystem = nullptr;
 };
+
+namespace UE::Net
+{
+
+inline FNetRPCCallContext::FNetRPCCallContext(FNetSerializationContext& InNetContext, const FForwardNetRPCCallMulticastDelegate& InForwardNetRPCCallDelegate)
+: NetContext(InNetContext)
+, ForwardNetRPCCallDelegate(InForwardNetRPCCallDelegate)
+{
+}
+
+inline FNetSerializationContext& FNetRPCCallContext::GetNetSerializationContext()
+{
+	return NetContext;
+}
+
+inline const FForwardNetRPCCallMulticastDelegate& FNetRPCCallContext::GetForwardNetRPCCallDelegate() const
+{
+	return ForwardNetRPCCallDelegate;
+}
+
+}
diff --git a/Engine/Source/Runtime/Experimental/Iris/Core/Private/Iris/ReplicationSystem/ReplicationSystem.cpp b/Engine/Source/Runtime/Experimental/Iris/Core/Private/Iris/ReplicationSystem/ReplicationSystem.cpp
--- a/Engine/Source/Runtime/Experimental/Iris/Core/Private/Iris/ReplicationSystem/ReplicationSystem.cpp
+++ b/Engine/Source/Runtime/Experimental/Iris/Core/Private/Iris/ReplicationSystem/ReplicationSystem.cpp
@@ -219,6 +219,11 @@ public:
 			InitParams.bSendAttachmentsWithObject = ReplicationSystem->IsServer();
 			BlobManager.Init(InitParams);
 		}
+
+		if (Params.ForwardNetRPCCallDelegate.IsBound())
+		{
+			ReplicationSystemInternal.GetForwardNetRPCCallMulticastDelegate().Add(Params.ForwardNetRPCCallDelegate);
+		}
 	}
 
 	void Deinit()
diff --git a/Engine/Source/Runtime/Experimental/Iris/Core/Private/Iris/ReplicationSystem/ReplicationSystemInternal.h b/Engine/Source/Runtime/Experimental/Iris/Core/Private/Iris/ReplicationSystem/ReplicationSystemInternal.h
--- a/Engine/Source/Runtime/Experimental/Iris/Core/Private/Iris/ReplicationSystem/ReplicationSystemInternal.h
+++ b/Engine/Source/Runtime/Experimental/Iris/Core/Private/Iris/ReplicationSystem/ReplicationSystemInternal.h
@@ -103,6 +103,11 @@ public:
 		return SendStats;
 	}
 
+	FForwardNetRPCCallMulticastDelegate& GetForwardNetRPCCallMulticastDelegate()
+	{
+		return ForwardNetRPCCallMulticastDelegate;
+	}
+
 private:
 	FReplicationProtocolManager ReplicationProtocolManager;
 	FNetRefHandleManager NetRefHandleManager;
@@ -126,6 +131,7 @@ private:
 	FDeltaCompressionBaselineManager DeltaCompressionBaselineManager;
 	FDeltaCompressionBaselineInvalidationTracker DeltaCompressionBaselineInvalidationTracker;
 	FNetSendStats SendStats;
+	FForwardNetRPCCallMulticastDelegate ForwardNetRPCCallMulticastDelegate;
 	uint32 Id;
 };
 
diff --git a/Engine/Source/Runtime/Experimental/Iris/Core/Public/Iris/ReplicationSystem/ReplicationSystem.h b/Engine/Source/Runtime/Experimental/Iris/Core/Public/Iris/ReplicationSystem/ReplicationSystem.h
--- a/Engine/Source/Runtime/Experimental/Iris/Core/Public/Iris/ReplicationSystem/ReplicationSystem.h
+++ b/Engine/Source/Runtime/Experimental/Iris/Core/Public/Iris/ReplicationSystem/ReplicationSystem.h
@@ -9,6 +9,7 @@
 #include "Iris/IrisConfig.h"
 #include "Iris/ReplicationSystem/NetRefHandle.h"
 #include "Iris/ReplicationSystem/NetObjectGroupHandle.h"
+#include "Iris/ReplicationSystem/ReplicationSystemTypes.h"
 #include "Net/Core/NetHandle/NetHandle.h"
 
 #include "ReplicationSystem.generated.h"
@@ -64,6 +65,7 @@ public:
 		uint32 MaxNetObjectGroupCount = 2048U;
 		bool bIsServer = false;
 		bool bAllowObjectReplication = false;
+		UE::Net::FForwardNetRPCCallDelegate ForwardNetRPCCallDelegate;
 	};
 
 	/** @return The unique ID of the ReplicationSystem. */
@@ -533,7 +535,6 @@ public:
 
 public:
 	// For internal use and not exported.
-
 	UE::Net::Private::FReplicationSystemInternal* GetReplicationSystemInternal();
 	const UE::Net::Private::FReplicationSystemInternal* GetReplicationSystemInternal() const;
 
diff --git a/Engine/Source/Runtime/Experimental/Iris/Core/Public/Iris/ReplicationSystem/ReplicationSystemTypes.h b/Engine/Source/Runtime/Experimental/Iris/Core/Public/Iris/ReplicationSystem/ReplicationSystemTypes.h
--- a/Engine/Source/Runtime/Experimental/Iris/Core/Public/Iris/ReplicationSystem/ReplicationSystemTypes.h
+++ b/Engine/Source/Runtime/Experimental/Iris/Core/Public/Iris/ReplicationSystem/ReplicationSystemTypes.h
@@ -3,6 +3,7 @@
 #pragma once
 
 #include "CoreTypes.h"
+#include "Delegates/Delegate.h"
 
 namespace UE::Net
 {
@@ -25,4 +26,7 @@ enum class EDependentObjectSchedulingHint : uint8
 	ScheduleBeforeParentIfInitialState,
 };
 
+using FForwardNetRPCCallDelegate = TDelegate<void(UObject* RootObject,UObject* SubObject, UFunction* Function, void* Params)>;
+using FForwardNetRPCCallMulticastDelegate = TMulticastDelegate<typename FForwardNetRPCCallDelegate::TFuncType>;
+
 }
