GC ISM Pool : do not set the relative transform of the component if it is not necessary
this prevent the "Moveable" related  warnings in PIE

diff --git a/Engine/Source/Runtime/Experimental/GeometryCollectionEngine/Private/GeometryCollection/GeometryCollectionISMPoolComponent.cpp b/Engine/Source/Runtime/Experimental/GeometryCollectionEngine/Private/GeometryCollection/GeometryCollectionISMPoolComponent.cpp
--- a/Engine/Source/Runtime/Experimental/GeometryCollectionEngine/Private/GeometryCollection/GeometryCollectionISMPoolComponent.cpp
+++ b/Engine/Source/Runtime/Experimental/GeometryCollectionEngine/Private/GeometryCollection/GeometryCollectionISMPoolComponent.cpp
@@ -120,7 +120,11 @@ void FGeometryCollectionISM::InitISM(const FGeometryCollectionStaticMeshInstance
 	// Instead of reverse culling we put the mirror in the component transform so that PRIMITIVE_SCENE_DATA_FLAG_DETERMINANT_SIGN will be set for use by materials.
 	//ISMComponent->SetReverseCulling(bReverseCulling);
 	const FVector Scale = bReverseCulling ? FVector(-1, 1, 1) : FVector(1, 1, 1);
-	ISMComponent->SetRelativeTransform(FTransform(FQuat::Identity, MeshInstance.Desc.Position, Scale));
+	const FTransform NewRelativeTransform(FQuat::Identity, MeshInstance.Desc.Position, Scale);
+	if (!ISMComponent->GetRelativeTransform().Equals(NewRelativeTransform))
+	{
+		ISMComponent->SetRelativeTransform(FTransform(FQuat::Identity, MeshInstance.Desc.Position, Scale));
+	}
 
 	ISMComponent->SetMobility((MeshInstance.Desc.Flags & FISMComponentDescription::StaticMobility) != 0 ? EComponentMobility::Static : EComponentMobility::Stationary);
 	ISMComponent->SetCachedMaxDrawDistance(MeshInstance.Desc.EndCullDistance);
