Geometry collection optimization : Add a tolerance for detecting if the GT particles have changed their X and R

diff --git a/Engine/Source/Runtime/Experimental/Chaos/Private/PhysicsProxy/GeometryCollectionPhysicsProxy.cpp b/Engine/Source/Runtime/Experimental/Chaos/Private/PhysicsProxy/GeometryCollectionPhysicsProxy.cpp
--- a/Engine/Source/Runtime/Experimental/Chaos/Private/PhysicsProxy/GeometryCollectionPhysicsProxy.cpp
+++ b/Engine/Source/Runtime/Experimental/Chaos/Private/PhysicsProxy/GeometryCollectionPhysicsProxy.cpp
@@ -130,6 +130,12 @@ FAutoConsoleVariableRef CVarGeometryCollectionTransformTolerance(
 	GeometryCollectionTransformTolerance,
 	TEXT("Tolerance to detect if a transform has changed"));
 
+float GeometryCollectionXRUpdateTolerance = UE_KINDA_SMALL_NUMBER;
+FAutoConsoleVariableRef CVarGeometryCollectionXRUpdateTolerance(
+	TEXT("p.GeometryCollection.XRUpdateTolerance"),
+	GeometryCollectionXRUpdateTolerance,
+	TEXT("Tolerance to detect if a XR has changed has changed when syncing PT to GT"));
+
 DEFINE_LOG_CATEGORY_STATIC(UGCC_LOG, Error, All);
 
 static const FSharedSimulationSizeSpecificData& GetSizeSpecificData(const TArray<FSharedSimulationSizeSpecificData>& SizeSpecificData, const FGeometryCollection& RestCollection, const int32 TransformIndex, const FBox& BoundingBox);
@@ -3299,14 +3305,14 @@ static inline bool UpdateGTParticleXR(Chaos::FPBDRigidParticle& GTParticle, cons
 	GC_PHYSICSPROXY_CHECK_FOR_NAN(NewX);
 
 	const Chaos::FVec3 OldX = GTParticle.X();
-	const bool bNeedUpdateX = (NewX != OldX);
+	const bool bNeedUpdateX = (!NewX.Equals(OldX, GeometryCollectionXRUpdateTolerance));
 	if (bNeedUpdateX)
 	{
 		GTParticle.SetX(NewX, false);
 	}
 
 	const Chaos::FRotation3 OldR = GTParticle.R();
-	const bool bNeedUpdateR = (NewR != OldR);
+	const bool bNeedUpdateR = (!NewR.Equals(OldR, GeometryCollectionXRUpdateTolerance));
 	if (bNeedUpdateR)
 	{
 		GTParticle.SetR(NewR, false);
