Replaced TChooseClass with std::conditional_t.

diff --git a/Engine/Plugins/Experimental/LearningAgents/Source/Learning/Public/MultiArray.h b/Engine/Plugins/Experimental/LearningAgents/Source/Learning/Public/MultiArray.h
--- a/Engine/Plugins/Experimental/LearningAgents/Source/Learning/Public/MultiArray.h
+++ b/Engine/Plugins/Experimental/LearningAgents/Source/Learning/Public/MultiArray.h
@@ -68,16 +68,16 @@ public:
 	static constexpr bool bIsRestrict = bInIsRestrict;
 	using SizeType = typename InAllocatorType::SizeType;
 
-	using PointerType = typename TChooseClass<
+	using PointerType = std::conditional_t<
 		bIsRestrict,
 		ElementType* RESTRICT,
-		ElementType*>::Result;
+		ElementType*>;
 
-	using ElementAllocatorType = typename TChooseClass<
+	using ElementAllocatorType = std::conditional_t<
 		AllocatorType::NeedsElementType,
 		typename AllocatorType::template ForElementType<ElementType>,
 		typename AllocatorType::ForAnyElementType
-	>::Result;
+	>;
 
 	static_assert(DimNum > 1, "TMultiArray requires a positive, non-zero number of dimensions");
 	static_assert(TIsSigned<SizeType>::Value, "TMultiArray only supports signed index types");
@@ -833,16 +833,16 @@ public:
 	static constexpr bool bIsRestrict = bInIsRestrict;
 	using SizeType = typename InAllocatorType::SizeType;
 
-	using PointerType = typename TChooseClass<
+	using PointerType = std::conditional_t<
 		bIsRestrict,
 		ElementType* RESTRICT,
-		ElementType*>::Result;
+		ElementType*>;
 
-	using ElementAllocatorType = typename TChooseClass<
+	using ElementAllocatorType = std::conditional_t<
 		AllocatorType::NeedsElementType,
 		typename AllocatorType::template ForElementType<ElementType>,
 		typename AllocatorType::ForAnyElementType
-	>::Result;
+	>;
 
 	static_assert(TIsSigned<SizeType>::Value, "TMultiArray only supports signed index types");
 
diff --git a/Engine/Plugins/Experimental/LearningAgents/Source/Learning/Public/MultiArrayView.h b/Engine/Plugins/Experimental/LearningAgents/Source/Learning/Public/MultiArrayView.h
--- a/Engine/Plugins/Experimental/LearningAgents/Source/Learning/Public/MultiArrayView.h
+++ b/Engine/Plugins/Experimental/LearningAgents/Source/Learning/Public/MultiArrayView.h
@@ -180,10 +180,10 @@ public:
 	static constexpr bool bIsRestrict = bInIsRestrict;
 	using SizeType = InSizeType;
 
-	using PointerType = typename TChooseClass<
+	using PointerType = std::conditional_t<
 		bIsRestrict,
 		ElementType* RESTRICT,
-		ElementType*>::Result;
+		ElementType*>;
 
 	static_assert(DimNum > 1, "TMultiArrayView requires a positive, non-zero number of dimensions");
 	static_assert(TIsSigned<SizeType>::Value, "TMultiArrayView only supports signed index types");
@@ -612,10 +612,10 @@ public:
 	static constexpr bool bIsRestrict = bInIsRestrict;
 	using SizeType = InSizeType;
 
-	using PointerType = typename TChooseClass<
+	using PointerType = std::conditional_t<
 		bIsRestrict,
 		ElementType* RESTRICT,
-		ElementType*>::Result;
+		ElementType*>;
 
 	static_assert(TIsSigned<SizeType>::Value, "TMultiArrayView only supports signed index types");
 
@@ -654,11 +654,11 @@ public:
 		>::Type
 	>
 		FORCEINLINE TMultiArrayView(OtherRangeType&& Other)
-		: DataPtr(TChooseClass<
+		: DataPtr(std::conditional_t<
 			TIsCompatibleRangeType<OtherRangeType>::Value,
 			TIsCompatibleRangeType<OtherRangeType>,
 			TIsReinterpretableRangeType<OtherRangeType>
-		>::Result::GetData(Forward<OtherRangeType>(Other)))
+		>::GetData(Forward<OtherRangeType>(Other)))
 	{
 		const auto InCount = GetNum(Forward<OtherRangeType>(Other));
 		check((InCount >= 0) && ((sizeof(InCount) < sizeof(SizeType)) || (InCount <= static_cast<decltype(InCount)>(TNumericLimits<SizeType>::Max()))));
diff --git a/Engine/Plugins/Experimental/StructUtils/Source/StructUtils/Public/InstancedStructContainer.h b/Engine/Plugins/Experimental/StructUtils/Source/StructUtils/Public/InstancedStructContainer.h
--- a/Engine/Plugins/Experimental/StructUtils/Source/StructUtils/Public/InstancedStructContainer.h
+++ b/Engine/Plugins/Experimental/StructUtils/Source/StructUtils/Public/InstancedStructContainer.h
@@ -147,7 +147,7 @@ public:
 	template<typename T>
 	struct TIterator
 	{
-		using StructViewType = typename TChooseClass<TIsConst<T>::Value, FConstStructView, FStructView>::Result;
+		using StructViewType = std::conditional_t<TIsConst<T>::Value, FConstStructView, FStructView>;
 
 		/** @return struct view (or const structview) to the item. */
 		StructViewType operator*()
diff --git a/Engine/Plugins/Online/OnlineFramework/Source/Party/Public/Party/PartyTypes.h b/Engine/Plugins/Online/OnlineFramework/Source/Party/Public/Party/PartyTypes.h
--- a/Engine/Plugins/Online/OnlineFramework/Source/Party/Public/Party/PartyTypes.h
+++ b/Engine/Plugins/Online/OnlineFramework/Source/Party/Public/Party/PartyTypes.h
@@ -501,7 +501,7 @@ protected:
 public:	\
 	/** If the property is a POD or ptr type, we'll work with it by copy. Otherwise, by const ref */	\
 	using Mutable##PropertyName##Type = std::remove_const_t<PropertyType>;	\
-	using PropertyName##ArgType = typename TChooseClass<TOr<TIsPODType<PropertyType>, TIsPointer<PropertyType>>::Value, PropertyType, const Mutable##PropertyName##Type&>::Result;	\
+	using PropertyName##ArgType = std::conditional_t<TOr<TIsPODType<PropertyType>, TIsPointer<PropertyType>>::Value, PropertyType, const Mutable##PropertyName##Type&>;	\
 	\
 private:	\
 	/** Bummer to have two signatures, but cases that want both the old and new values are much rarer, so most don't want to bother with a handler that takes an extra unused param */	\
diff --git a/Engine/Plugins/Runtime/ReplicationSystemTestPlugin/Source/Private/Tests/Serialization/TestEnumNetSerializer.cpp b/Engine/Plugins/Runtime/ReplicationSystemTestPlugin/Source/Private/Tests/Serialization/TestEnumNetSerializer.cpp
--- a/Engine/Plugins/Runtime/ReplicationSystemTestPlugin/Source/Private/Tests/Serialization/TestEnumNetSerializer.cpp
+++ b/Engine/Plugins/Runtime/ReplicationSystemTestPlugin/Source/Private/Tests/Serialization/TestEnumNetSerializer.cpp
@@ -114,7 +114,7 @@ void FTestEnumIntNetSerializer<SerializerConfig, SourceType, EnumType>::SetUp()
 		// NumEnums actually also contain the generated _MAX enum value which might not even be a valid value by the backed type. Skip it!
 		const int32 EnumValueCount = Enum->NumEnums() - 1;
 
-		using LargeIntegerType = typename TChooseClass<TIsSigned<SourceType>::Value, int64, uint64>::Result;
+		using LargeIntegerType = std::conditional_t<TIsSigned<SourceType>::Value, int64, uint64>;
 
 		// Valid values
 		TArray<SourceType> TempValues;
diff --git a/Engine/Source/Runtime/Core/Private/GenericPlatform/GenericWidePlatformString.cpp b/Engine/Source/Runtime/Core/Private/GenericPlatform/GenericWidePlatformString.cpp
--- a/Engine/Source/Runtime/Core/Private/GenericPlatform/GenericWidePlatformString.cpp
+++ b/Engine/Source/Runtime/Core/Private/GenericPlatform/GenericWidePlatformString.cpp
@@ -892,7 +892,7 @@ int32 FGenericWidePlatformString::GetVarArgs( WIDECHAR* Dest, SIZE_T DestSize, c
 			case 'S':
 			{
 				// The %S format represents a string which is the opposite of %s - wide if TCHAR is narrow, or narrow if TCHAR is wide
-				using OtherCharType = TChooseClass<std::is_same_v<TCHAR, ANSICHAR>, WIDECHAR, ANSICHAR>::Result;
+				using OtherCharType = std::conditional_t<std::is_same_v<TCHAR, ANSICHAR>, WIDECHAR, ANSICHAR>;
 
 				ProcessStringArg<OtherCharType>(DestIter, Src, FieldLen, PrecisionLen, ArgPtr);
 				if (!DestIter)
diff --git a/Engine/Source/Runtime/Core/Public/Containers/Array.h b/Engine/Source/Runtime/Core/Public/Containers/Array.h
--- a/Engine/Source/Runtime/Core/Public/Containers/Array.h
+++ b/Engine/Source/Runtime/Core/Public/Containers/Array.h
@@ -393,11 +393,11 @@ public:
 	UE_DEPRECATED(5.0, "TArray::Allocator type is deprecated, please use TArray::AllocatorType instead.")
 	typedef InAllocatorType Allocator;
 
-	typedef typename TChooseClass<
+	using ElementAllocatorType = std::conditional_t<
 		AllocatorType::NeedsElementType,
 		typename AllocatorType::template ForElementType<ElementType>,
 		typename AllocatorType::ForAnyElementType
-	>::Result ElementAllocatorType;
+	>;
 
 	static_assert(TIsSigned<SizeType>::Value, "TArray only supports signed index types");
 
diff --git a/Engine/Source/Runtime/Core/Public/Containers/ArrayView.h b/Engine/Source/Runtime/Core/Public/Containers/ArrayView.h
--- a/Engine/Source/Runtime/Core/Public/Containers/ArrayView.h
+++ b/Engine/Source/Runtime/Core/Public/Containers/ArrayView.h
@@ -192,11 +192,11 @@ public:
 		std::enable_if_t<TIsTArrayView_V<std::decay_t<OtherRangeType>> && !std::is_same_v<std::decay_t<OtherRangeType>, TArrayView>>* = nullptr
 	>
 	FORCEINLINE TArrayView(OtherRangeType&& Other)
-		: DataPtr(TChooseClass<
+		: DataPtr(std::conditional_t<
 						TIsCompatibleRangeType<OtherRangeType>::Value,
 						TIsCompatibleRangeType<OtherRangeType>,
 						TIsReinterpretableRangeType<OtherRangeType>
-					>::Result::GetData(Forward<OtherRangeType>(Other)))
+					>::GetData(Forward<OtherRangeType>(Other)))
 	{
 		const auto InCount = GetNum(Forward<OtherRangeType>(Other));
 		using InCountType = decltype(InCount);
@@ -224,11 +224,11 @@ public:
 		std::enable_if_t<!TIsTArrayView_V<std::decay_t<OtherRangeType>>>* = nullptr
 	>
 	FORCEINLINE TArrayView(OtherRangeType&& Other UE_LIFETIMEBOUND)
-		: DataPtr(TChooseClass<
+		: DataPtr(std::conditional_t<
 			TIsCompatibleRangeType<OtherRangeType>::Value,
 			TIsCompatibleRangeType<OtherRangeType>,
 			TIsReinterpretableRangeType<OtherRangeType>
-		>::Result::GetData(Forward<OtherRangeType>(Other)))
+		>::GetData(Forward<OtherRangeType>(Other)))
 	{
 		const auto InCount = GetNum(Forward<OtherRangeType>(Other));
 		using InCountType = decltype(InCount);
diff --git a/Engine/Source/Runtime/Core/Public/Containers/BitArray.h b/Engine/Source/Runtime/Core/Public/Containers/BitArray.h
--- a/Engine/Source/Runtime/Core/Public/Containers/BitArray.h
+++ b/Engine/Source/Runtime/Core/Public/Containers/BitArray.h
@@ -2042,7 +2042,7 @@ using TConstDualEitherSetBitIterator = TConstDualSetBitIterator<Allocator, Other
 template <typename Allocator, typename InDerivedType>
 class TScriptBitArray
 {
-	using DerivedType = typename TChooseClass<std::is_void_v<InDerivedType>, TScriptBitArray, InDerivedType>::Result;
+	using DerivedType = std::conditional_t<std::is_void_v<InDerivedType>, TScriptBitArray, InDerivedType>;
 
 public:
 	/**
diff --git a/Engine/Source/Runtime/Core/Public/Containers/Deque.h b/Engine/Source/Runtime/Core/Public/Containers/Deque.h
--- a/Engine/Source/Runtime/Core/Public/Containers/Deque.h
+++ b/Engine/Source/Runtime/Core/Public/Containers/Deque.h
@@ -96,10 +96,10 @@ public:
 	using SizeType = typename InAllocatorType::SizeType;
 	using ElementType = InElementType;
 
-	using ElementAllocatorType = typename TChooseClass<
+	using ElementAllocatorType = std::conditional_t<
 		AllocatorType::NeedsElementType,
 		typename AllocatorType::template ForElementType<ElementType>,
-		typename AllocatorType::ForAnyElementType>::Result;
+		typename AllocatorType::ForAnyElementType>;
 
 	using ConstIteratorType = UE::Deque::Private::TIterator<const ElementType, SizeType>;
 	using IteratorType = UE::Deque::Private::TIterator<ElementType, SizeType>;
diff --git a/Engine/Source/Runtime/Core/Public/Containers/HashTable.h b/Engine/Source/Runtime/Core/Public/Containers/HashTable.h
--- a/Engine/Source/Runtime/Core/Public/Containers/HashTable.h
+++ b/Engine/Source/Runtime/Core/Public/Containers/HashTable.h
@@ -412,11 +412,11 @@ class THashTable
 public:
 	using Allocator = InAllocator;
 
-	using ElementAllocatorType = typename TChooseClass<
+	using ElementAllocatorType = std::conditional_t<
 		Allocator::NeedsElementType,
 		typename Allocator::template ForElementType<uint32>,
 		typename Allocator::ForAnyElementType
-	>::Result;
+	>;
 
 	explicit THashTable(uint32 InHashSize = 1024, uint32 InIndexSize = 0);
 	THashTable(const THashTable& Other) = delete;
diff --git a/Engine/Source/Runtime/Core/Public/Containers/Map.h b/Engine/Source/Runtime/Core/Public/Containers/Map.h
--- a/Engine/Source/Runtime/Core/Public/Containers/Map.h
+++ b/Engine/Source/Runtime/Core/Public/Containers/Map.h
@@ -775,16 +775,16 @@ protected:
 	class TBaseIterator
 	{
 	public:
-		typedef typename TChooseClass<
+		typedef std::conditional_t<
 			bConst,
-			typename TChooseClass<bRangedFor, typename ElementSetType::TRangedForConstIterator, typename ElementSetType::TConstIterator>::Result,
-			typename TChooseClass<bRangedFor, typename ElementSetType::TRangedForIterator, typename ElementSetType::TIterator     >::Result
-		>::Result PairItType;
+			std::conditional_t<bRangedFor, typename ElementSetType::TRangedForConstIterator, typename ElementSetType::TConstIterator>,
+			std::conditional_t<bRangedFor, typename ElementSetType::TRangedForIterator, typename ElementSetType::TIterator     >
+		> PairItType;
 	private:
-		typedef typename TChooseClass<bConst, const TMapBase, TMapBase>::Result MapType;
-		typedef typename TChooseClass<bConst, const KeyType, KeyType>::Result ItKeyType;
-		typedef typename TChooseClass<bConst, const ValueType, ValueType>::Result ItValueType;
-		typedef typename TChooseClass<bConst, const typename ElementSetType::ElementType, typename ElementSetType::ElementType>::Result PairType;
+		typedef std::conditional_t<bConst, const TMapBase, TMapBase> MapType;
+		typedef std::conditional_t<bConst, const KeyType, KeyType> ItKeyType;
+		typedef std::conditional_t<bConst, const ValueType, ValueType> ItValueType;
+		typedef std::conditional_t<bConst, const typename ElementSetType::ElementType, typename ElementSetType::ElementType> PairType;
 
 	public:
 		FORCEINLINE TBaseIterator(const PairItType& InElementIt)
@@ -832,9 +832,9 @@ protected:
 	class TBaseKeyIterator
 	{
 	private:
-		typedef typename TChooseClass<bConst, typename ElementSetType::TConstKeyIterator, typename ElementSetType::TKeyIterator>::Result SetItType;
-		typedef typename TChooseClass<bConst, const KeyType, KeyType>::Result ItKeyType;
-		typedef typename TChooseClass<bConst, const ValueType, ValueType>::Result ItValueType;
+		typedef std::conditional_t<bConst, typename ElementSetType::TConstKeyIterator, typename ElementSetType::TKeyIterator> SetItType;
+		typedef std::conditional_t<bConst, const KeyType, KeyType> ItKeyType;
+		typedef std::conditional_t<bConst, const ValueType, ValueType> ItValueType;
 
 	public:
 		/** Initialization constructor. */
@@ -1693,7 +1693,7 @@ struct FScriptMapLayout
 template <typename AllocatorType, typename InDerivedType>
 class TScriptMap
 {
-	using DerivedType = typename TChooseClass<std::is_void_v<InDerivedType>, TScriptMap, InDerivedType>::Result;
+	using DerivedType = std::conditional_t<std::is_void_v<InDerivedType>, TScriptMap, InDerivedType>;
 
 public:
 	static FScriptMapLayout GetScriptLayout(int32 KeySize, int32 KeyAlignment, int32 ValueSize, int32 ValueAlignment)
diff --git a/Engine/Source/Runtime/Core/Public/Containers/PagedArray.h b/Engine/Source/Runtime/Core/Public/Containers/PagedArray.h
--- a/Engine/Source/Runtime/Core/Public/Containers/PagedArray.h
+++ b/Engine/Source/Runtime/Core/Public/Containers/PagedArray.h
@@ -167,10 +167,10 @@ public:
 	using AllocatorType = InAllocatorType;
 	using SizeType = typename InAllocatorType::SizeType;
 	using ElementType = InElementType;
-	using ElementAllocatorType = typename TChooseClass<
+	using ElementAllocatorType = std::conditional_t<
 		AllocatorType::NeedsElementType,
 		typename AllocatorType::template ForElementType<ElementType>,
-		typename AllocatorType::ForAnyElementType>::Result;
+		typename AllocatorType::ForAnyElementType>;
 
 private:
 	static constexpr SizeType GetPageIndex(SizeType Index)
diff --git a/Engine/Source/Runtime/Core/Public/Containers/Set.h b/Engine/Source/Runtime/Core/Public/Containers/Set.h
--- a/Engine/Source/Runtime/Core/Public/Containers/Set.h
+++ b/Engine/Source/Runtime/Core/Public/Containers/Set.h
@@ -1569,14 +1569,14 @@ private:
 	private:
 		friend class TSet;
 
-		typedef typename TChooseClass<bConst,const ElementType,ElementType>::Result ItElementType;
+		typedef std::conditional_t<bConst,const ElementType,ElementType> ItElementType;
 
 	public:
-		typedef typename TChooseClass<
+		typedef std::conditional_t<
 			bConst,
-			typename TChooseClass<bRangedFor, typename ElementArrayType::TRangedForConstIterator, typename ElementArrayType::TConstIterator>::Result,
-			typename TChooseClass<bRangedFor, typename ElementArrayType::TRangedForIterator,      typename ElementArrayType::TIterator     >::Result
-		>::Result ElementItType;
+			std::conditional_t<bRangedFor, typename ElementArrayType::TRangedForConstIterator, typename ElementArrayType::TConstIterator>,
+			std::conditional_t<bRangedFor, typename ElementArrayType::TRangedForIterator,      typename ElementArrayType::TIterator     >
+		> ElementItType;
 
 		FORCEINLINE TBaseIterator(const ElementItType& InElementIt)
 			: ElementIt(InElementIt)
@@ -1626,8 +1626,8 @@ private:
 	class TBaseKeyIterator
 	{
 	private:
-		typedef typename TChooseClass<bConst,const TSet,TSet>::Result SetType;
-		typedef typename TChooseClass<bConst,const ElementType,ElementType>::Result ItElementType;
+		typedef std::conditional_t<bConst, const TSet, TSet> SetType;
+		typedef std::conditional_t<bConst,const ElementType,ElementType> ItElementType;
 		typedef typename TTypeTraits<typename KeyFuncs::KeyType>::ConstPointerType ReferenceOrValueType;
 
 	public:
@@ -1846,7 +1846,7 @@ struct FScriptSetLayout
 template <typename Allocator, typename InDerivedType>
 class TScriptSet
 {
-	using DerivedType = typename TChooseClass<std::is_void_v<InDerivedType>, TScriptSet, InDerivedType>::Result;
+	using DerivedType = std::conditional_t<std::is_void_v<InDerivedType>, TScriptSet, InDerivedType>;
 
 public:
 	static FScriptSetLayout GetScriptLayout(int32 ElementSize, int32 ElementAlignment)
diff --git a/Engine/Source/Runtime/Core/Public/Containers/SortedMap.h b/Engine/Source/Runtime/Core/Public/Containers/SortedMap.h
--- a/Engine/Source/Runtime/Core/Public/Containers/SortedMap.h
+++ b/Engine/Source/Runtime/Core/Public/Containers/SortedMap.h
@@ -565,12 +565,12 @@ private:
 	class TBaseIterator
 	{
 	public:
-		typedef typename TChooseClass<bConst,typename ElementArrayType::TConstIterator,typename ElementArrayType::TIterator>::Result PairItType;
+		typedef std::conditional_t<bConst,typename ElementArrayType::TConstIterator,typename ElementArrayType::TIterator> PairItType;
 	private:
-		typedef typename TChooseClass<bConst,const TSortedMap,TSortedMap>::Result MapType;
-		typedef typename TChooseClass<bConst,const KeyType,KeyType>::Result ItKeyType;
-		typedef typename TChooseClass<bConst,const ValueType,ValueType>::Result ItValueType;
-		typedef typename TChooseClass<bConst,const typename ElementArrayType::ElementType, typename ElementArrayType::ElementType>::Result PairType;
+		typedef std::conditional_t<bConst,const TSortedMap,TSortedMap> MapType;
+		typedef std::conditional_t<bConst,const KeyType,KeyType> ItKeyType;
+		typedef std::conditional_t<bConst,const ValueType,ValueType> ItValueType;
+		typedef std::conditional_t<bConst,const typename ElementArrayType::ElementType, typename ElementArrayType::ElementType> PairType;
 
 	protected:
 		FORCEINLINE TBaseIterator(const PairItType& InElementIt)
@@ -615,13 +615,13 @@ private:
 	{
 		// Once we add reverse iterator to TArray, this class and TBaseIterator could be merged with a template parameter for forward vs reverse.
 	private:
-		typedef typename TChooseClass<bConst, const TSortedMap, TSortedMap>::Result MapType;
-		typedef typename TChooseClass<bConst, const KeyType, KeyType>::Result ItKeyType;
-		typedef typename TChooseClass<bConst, const ValueType, ValueType>::Result ItValueType;
+		typedef std::conditional_t<bConst, const TSortedMap, TSortedMap> MapType;
+		typedef std::conditional_t<bConst, const KeyType, KeyType> ItKeyType;
+		typedef std::conditional_t<bConst, const ValueType, ValueType> ItValueType;
 		typedef typename ElementArrayType::SizeType SizeType;
 
 	public:
-		typedef typename TChooseClass<bConst, const typename ElementArrayType::ElementType, typename ElementArrayType::ElementType>::Result PairType;
+		typedef std::conditional_t<bConst, const typename ElementArrayType::ElementType, typename ElementArrayType::ElementType> PairType;
 
 	protected:
 		FORCEINLINE TBaseReverseIterator(PairType* InData, SizeType InNum)
diff --git a/Engine/Source/Runtime/Core/Public/Containers/SparseArray.h b/Engine/Source/Runtime/Core/Public/Containers/SparseArray.h
--- a/Engine/Source/Runtime/Core/Public/Containers/SparseArray.h
+++ b/Engine/Source/Runtime/Core/Public/Containers/SparseArray.h
@@ -877,8 +877,8 @@ private:
 		typedef TConstSetBitIterator<typename Allocator::BitArrayAllocator> BitArrayItType;
 
 	private:
-		typedef typename TChooseClass<bConst,const TSparseArray,TSparseArray>::Result ArrayType;
-		typedef typename TChooseClass<bConst,const ElementType,ElementType>::Result ItElementType;
+		typedef std::conditional_t<bConst,const TSparseArray,TSparseArray> ArrayType;
+		typedef std::conditional_t<bConst,const ElementType,ElementType> ItElementType;
 
 	public:
 		explicit TBaseIterator(ArrayType& InArray, const BitArrayItType& InBitArrayIt)
@@ -1270,7 +1270,7 @@ struct FScriptSparseArrayLayout
 template <typename AllocatorType, typename InDerivedType>
 class TScriptSparseArray
 {
-	using DerivedType = typename TChooseClass<std::is_void_v<InDerivedType>, TScriptSparseArray, InDerivedType>::Result;
+	using DerivedType = std::conditional_t<std::is_void_v<InDerivedType>, TScriptSparseArray, InDerivedType>;
 
 public:
 	static FScriptSparseArrayLayout GetScriptLayout(int32 ElementSize, int32 ElementAlignment)
diff --git a/Engine/Source/Runtime/Core/Public/Memory/MemoryView.h b/Engine/Source/Runtime/Core/Public/Memory/MemoryView.h
--- a/Engine/Source/Runtime/Core/Public/Memory/MemoryView.h
+++ b/Engine/Source/Runtime/Core/Public/Memory/MemoryView.h
@@ -30,7 +30,7 @@ class TMemoryView
 {
 	static_assert(std::is_void_v<DataType>, "DataType must be cv-qualified void");
 
-	using ByteType = typename TChooseClass<TIsConst<DataType>::Value, const uint8, uint8>::Result;
+	using ByteType = std::conditional_t<TIsConst<DataType>::Value, const uint8, uint8>;
 
 public:
 	/** Construct an empty view. */
@@ -296,6 +296,6 @@ UE_NODISCARD constexpr inline auto MakeMemoryView(ContainerType&& Container)
 {
 	using ElementType = typename TRemovePointer<decltype(GetData(DeclVal<ContainerType>()))>::Type;
 	constexpr bool bIsConst = TIsConst<ElementType>::Value;
-	using DataType = typename TChooseClass<bIsConst, const void, void>::Result;
+	using DataType = std::conditional_t<bIsConst, const void, void>;
 	return TMemoryView<DataType>(GetData(Container), GetNum(Container) * sizeof(ElementType));
 }
diff --git a/Engine/Source/Runtime/Core/Public/Templates/Function.h b/Engine/Source/Runtime/Core/Public/Templates/Function.h
--- a/Engine/Source/Runtime/Core/Public/Templates/Function.h
+++ b/Engine/Source/Runtime/Core/Public/Templates/Function.h
@@ -182,7 +182,7 @@ namespace UE::Core::Private::Function
 	template <typename T, bool bOnHeap>
 	struct TFunction_OwnedObject : public
 #if TFUNCTION_USES_INLINE_STORAGE
-		TChooseClass<bOnHeap, IFunction_OwnedObject_OnHeap<T>, IFunction_OwnedObject_Inline<T>>::Result
+		std::conditional_t<bOnHeap, IFunction_OwnedObject_OnHeap<T>, IFunction_OwnedObject_Inline<T>>
 #else
 		IFunction_OwnedObject_OnHeap<T>
 #endif
diff --git a/Engine/Source/Runtime/Core/Tests/Math/GuardedIntTest.cpp b/Engine/Source/Runtime/Core/Tests/Math/GuardedIntTest.cpp
--- a/Engine/Source/Runtime/Core/Tests/Math/GuardedIntTest.cpp
+++ b/Engine/Source/Runtime/Core/Tests/Math/GuardedIntTest.cpp
@@ -16,7 +16,7 @@ bool DoGuardedIntTestForType(FAutomationTestBase& Test)
 	static_assert(std::is_same_v<IntType, int8> || std::is_same_v<IntType, uint8>);
 
 	constexpr bool bIsSigned = std::is_signed_v<IntType>;
-	using IntType32 = typename TChooseClass<bIsSigned, int32, uint32>::Result;
+	using IntType32 = std::conditional_t<bIsSigned, int32, uint32>;
 	constexpr IntType32 IntTypeMin = std::numeric_limits<IntType>::min();
 	constexpr IntType32 IntTypeMax = std::numeric_limits<IntType>::max();
 
diff --git a/Engine/Source/Runtime/CoreUObject/Private/UObject/EnumProperty.cpp b/Engine/Source/Runtime/CoreUObject/Private/UObject/EnumProperty.cpp
--- a/Engine/Source/Runtime/CoreUObject/Private/UObject/EnumProperty.cpp
+++ b/Engine/Source/Runtime/CoreUObject/Private/UObject/EnumProperty.cpp
@@ -17,7 +17,7 @@ namespace UEEnumProperty_Private
 	template <typename OldIntType>
 	void ConvertIntValueToEnumProperty(OldIntType OldValue, FEnumProperty* EnumProp, FNumericProperty* UnderlyingProp, UEnum* Enum, void* Obj)
 	{
-		using LargeIntType = typename TChooseClass<TIsSigned<OldIntType>::Value, int64, uint64>::Result;
+		using LargeIntType = std::conditional_t<TIsSigned<OldIntType>::Value, int64, uint64>;
 
 		LargeIntType NewValue = OldValue;
 		if (!UnderlyingProp->CanHoldValue(NewValue) || !Enum->IsValidEnumValue(NewValue))
diff --git a/Engine/Source/Runtime/Engine/Private/Collision/SceneQuery.cpp b/Engine/Source/Runtime/Engine/Private/Collision/SceneQuery.cpp
--- a/Engine/Source/Runtime/Engine/Private/Collision/SceneQuery.cpp
+++ b/Engine/Source/Runtime/Engine/Private/Collision/SceneQuery.cpp
@@ -165,8 +165,8 @@ struct TSQTraits
 	static const ESingleMultiOrTest SingleMultiOrTest = InSingleMultiOrTest;
 	static const ESweepOrRay GeometryQuery = InGeometryQuery;
 	using THitType = InHitType;
-	using TOutHits = typename TChooseClass<InSingleMultiOrTest == ESingleMultiOrTest::Multi, TArray<FHitResult>, FHitResult>::Result;
-	using THitBuffer = typename TChooseClass<InSingleMultiOrTest == ESingleMultiOrTest::Multi, FDynamicHitBuffer<InHitType>, FSingleHitBuffer<InHitType>>::Result;
+	using TOutHits = std::conditional_t<InSingleMultiOrTest == ESingleMultiOrTest::Multi, TArray<FHitResult>, FHitResult>;
+	using THitBuffer = std::conditional_t<InSingleMultiOrTest == ESingleMultiOrTest::Multi, FDynamicHitBuffer<InHitType>, FSingleHitBuffer<InHitType>>;
 
 	// GetNumHits - multi
 	template <ESingleMultiOrTest T = SingleMultiOrTest>
diff --git a/Engine/Source/Runtime/Engine/Public/WorldPartition/ActorDescContainerCollection.h b/Engine/Source/Runtime/Engine/Public/WorldPartition/ActorDescContainerCollection.h
--- a/Engine/Source/Runtime/Engine/Public/WorldPartition/ActorDescContainerCollection.h
+++ b/Engine/Source/Runtime/Engine/Public/WorldPartition/ActorDescContainerCollection.h
@@ -92,11 +92,11 @@ public:
 	protected:
 		typedef UActorDescContainer ContainerType;
 		typedef TArray<ActorDescContPtrType> ContainerCollectionType;
-		typedef typename TChooseClass<bConst, typename ContainerCollectionType::TConstIterator, typename ContainerCollectionType::TIterator>::Result ContainerIteratorType;
-		typedef typename TChooseClass<bConst, typename ContainerType::TConstIterator<ActorType>, typename ContainerType::TIterator<ActorType>>::Result ActDescIteratorType;
+		typedef std::conditional_t<bConst, typename ContainerCollectionType::TConstIterator, typename ContainerCollectionType::TIterator> ContainerIteratorType;
+		typedef std::conditional_t<bConst, typename ContainerType::TConstIterator<ActorType>, typename ContainerType::TIterator<ActorType>> ActDescIteratorType;
 
 		typedef typename FWorldPartitionActorDescType<ActorType>::Type ValueType;
-		typedef typename TChooseClass<bConst, const ValueType*, ValueType*>::Result ReturnType;
+		typedef std::conditional_t<bConst, const ValueType*, ValueType*> ReturnType;
 
 	public:
 		template<class T>
diff --git a/Engine/Source/Runtime/Engine/Public/WorldPartition/ActorDescList.h b/Engine/Source/Runtime/Engine/Public/WorldPartition/ActorDescList.h
--- a/Engine/Source/Runtime/Engine/Public/WorldPartition/ActorDescList.h
+++ b/Engine/Source/Runtime/Engine/Public/WorldPartition/ActorDescList.h
@@ -50,9 +50,9 @@ public:
 	protected:
 		using MapType = FGuidActorDescMap;
 		using ValueType = typename FWorldPartitionActorDescType<ActorType>::Type;
-		using IteratorType = typename TChooseClass<bConst, MapType::TConstIterator, MapType::TIterator>::Result;
-		using ListType = typename TChooseClass<bConst, const FActorDescList*, FActorDescList*>::Result;
-		using ReturnType = typename TChooseClass<bConst, const ValueType*, ValueType*>::Result;
+		using IteratorType = std::conditional_t<bConst, MapType::TConstIterator, MapType::TIterator>;
+		using ListType = std::conditional_t<bConst, const FActorDescList*, FActorDescList*>;
+		using ReturnType = std::conditional_t<bConst, const ValueType*, ValueType*>;
 
 	public:
 		TBaseIterator(ListType InActorDescList, UClass* InActorClass)
diff --git a/Engine/Source/Runtime/Experimental/Chaos/Public/Chaos/ImplicitObjectScaled.h b/Engine/Source/Runtime/Experimental/Chaos/Public/Chaos/ImplicitObjectScaled.h
--- a/Engine/Source/Runtime/Experimental/Chaos/Public/Chaos/ImplicitObjectScaled.h
+++ b/Engine/Source/Runtime/Experimental/Chaos/Public/Chaos/ImplicitObjectScaled.h
@@ -421,7 +421,7 @@ public:
 	using ObjectType = TRefCountPtr<TConcrete>;
 	using FImplicitObject::GetTypeName;
 	
-	using ObjectTypeDeprecated = typename TChooseClass<bInstanced, TSerializablePtr<TConcrete>, TUniquePtr<TConcrete>>::Result;
+	using ObjectTypeDeprecated = std::conditional_t<bInstanced, TSerializablePtr<TConcrete>, TUniquePtr<TConcrete>>;
 
 	UE_DEPRECATED(5.4, "Constructor no longer used anymore")
 	TImplicitObjectScaled(ObjectTypeDeprecated Object, const TSharedPtr<TConcrete, ESPMode::ThreadSafe>& SharedPtrForRefCount, const FVec3& Scale, FReal InMargin = 0)
diff --git a/Engine/Source/Runtime/Experimental/Chaos/Public/Chaos/ParticleHandle.h b/Engine/Source/Runtime/Experimental/Chaos/Public/Chaos/ParticleHandle.h
--- a/Engine/Source/Runtime/Experimental/Chaos/Public/Chaos/ParticleHandle.h
+++ b/Engine/Source/Runtime/Experimental/Chaos/Public/Chaos/ParticleHandle.h
@@ -814,7 +814,7 @@ protected:
 	friend TGeometryParticleHandles<T, d>;
 	
 	struct FInvalidFromTransient {};
-	typename TChooseClass<bPersistent, int32, FInvalidFromTransient>::Result HandleIdx;	//Index into the handles array. This is useful for binding external attributes. Note the index can change
+	std::conditional_t<bPersistent, int32, FInvalidFromTransient> HandleIdx;	//Index into the handles array. This is useful for binding external attributes. Note the index can change
 };
 
 template<>
diff --git a/Engine/Source/Runtime/Experimental/Iris/Core/Public/Iris/Serialization/BitPacking.h b/Engine/Source/Runtime/Experimental/Iris/Core/Public/Iris/Serialization/BitPacking.h
--- a/Engine/Source/Runtime/Experimental/Iris/Core/Public/Iris/Serialization/BitPacking.h
+++ b/Engine/Source/Runtime/Experimental/Iris/Core/Public/Iris/Serialization/BitPacking.h
@@ -114,14 +114,14 @@ template<typename T, typename TEnableIf<TIsSigned<T>::Value && sizeof(T) <= 8U,
 inline void SerializeIntDelta(FNetBitStreamWriter& Writer, const T Value, const T PrevValue, const uint8* SmallBitCountTable, const uint32 SmallBitCountTableEntryCount, uint8 LargeBitCount)
 {
 	// Careful casting to properly represent negative numbers in the larger type as we want the delta between small negative numbers and small positive numbers to be small.
-	using SignedType = typename TChooseClass<sizeof(T) == 8, int64, int32>::Result;
+	using SignedType = std::conditional_t<sizeof(T) == 8, int64, int32>;
 	return SerializeIntDelta(Writer, SignedType(Value), SignedType(PrevValue), SmallBitCountTable, SmallBitCountTableEntryCount, LargeBitCount);
 }
 
 template<typename T, typename TEnableIf<TIsSigned<T>::Value && sizeof(T) <= 8U, int32>::Type U = -1>
 inline void DeserializeIntDelta(FNetBitStreamReader& Reader, T& OutValue, const T PrevValue, const uint8* SmallBitCountTable, const uint32 SmallBitCountTableEntryCount, uint8 LargeBitCount)
 {
-	using SignedType = typename TChooseClass<sizeof(T) == 8, int64, int32>::Result;
+	using SignedType = std::conditional_t<sizeof(T) == 8, int64, int32>;
 
 	SignedType Value;
 	DeserializeIntDelta(Reader, Value, SignedType(PrevValue), SmallBitCountTable, SmallBitCountTableEntryCount, LargeBitCount);
@@ -132,14 +132,14 @@ template<typename T, typename TEnableIf<!TIsSigned<T>::Value && TIsIntegral<T>::
 inline void SerializeUintDelta(FNetBitStreamWriter& Writer, const T Value, const T PrevValue, const uint8* SmallBitCountTable, const uint32 SmallBitCountTableEntryCount, uint8 LargeBitCount)
 {
 	// Careful casting to allow the delta between small positive numbers and large positive numbers to be small.
-	using UnsignedType = typename TChooseClass<sizeof(T) == 8, uint64, uint32>::Result;
+	using UnsignedType = std::conditional_t<sizeof(T) == 8, uint64, uint32>;
 	return SerializeUintDelta(Writer, UnsignedType(Value), UnsignedType(PrevValue), SmallBitCountTable, SmallBitCountTableEntryCount, LargeBitCount);
 }
 
 template<typename T, typename TEnableIf<!TIsSigned<T>::Value && TIsIntegral<T>::Value && sizeof(T) <= 8U, uint32>::Type U = 1U>
 inline void DeserializeUintDelta(FNetBitStreamReader& Reader, T& OutValue, const T PrevValue, const uint8* SmallBitCountTable, const uint32 SmallBitCountTableEntryCount, uint8 LargeBitCount)
 {
-	using UnsignedType = typename TChooseClass<sizeof(T) == 8, uint64, uint32>::Result;
+	using UnsignedType = std::conditional_t<sizeof(T) == 8, uint64, uint32>;
 
 	UnsignedType Value;
 	DeserializeUintDelta(Reader, Value, UnsignedType(PrevValue), SmallBitCountTable, SmallBitCountTableEntryCount, LargeBitCount);
diff --git a/Engine/Source/Runtime/Experimental/Iris/Core/Public/Iris/Serialization/NetSerializerBuilder.inl b/Engine/Source/Runtime/Experimental/Iris/Core/Public/Iris/Serialization/NetSerializerBuilder.inl
--- a/Engine/Source/Runtime/Experimental/Iris/Core/Public/Iris/Serialization/NetSerializerBuilder.inl
+++ b/Engine/Source/Runtime/Experimental/Iris/Core/Public/Iris/Serialization/NetSerializerBuilder.inl
@@ -431,7 +431,7 @@ public:
 	static constexpr uint32 GetQuantizedTypeSize() { return sizeof(typename NetSerializerImpl::QuantizedType); }
 
 	template<typename T = void, typename U = typename TEnableIf<!HasQuantizedType && HasSourceType, T>::Type, bool V = true, char W = 0>
-	static constexpr uint32 GetQuantizedTypeSize() { return std::is_same_v<void, typename NetSerializerImpl::SourceType> ? uint32(0) : sizeof(typename TChooseClass<std::is_same_v<void, typename NetSerializerImpl::SourceType>, uint8, typename NetSerializerImpl::SourceType>::Result); }
+	static constexpr uint32 GetQuantizedTypeSize() { return std::is_same_v<void, typename NetSerializerImpl::SourceType> ? uint32(0) : sizeof(std::conditional_t<std::is_same_v<void, typename NetSerializerImpl::SourceType>, uint8, typename NetSerializerImpl::SourceType>); }
 
 	template<typename T = void, typename U = typename TEnableIf<!(HasSourceType || HasQuantizedType), T>::Type, char V = 0>
 	static constexpr uint32 GetQuantizedTypeSize() { return 0; }
@@ -440,7 +440,7 @@ public:
 	static constexpr uint32 GetQuantizedTypeAlignment() { return alignof(typename NetSerializerImpl::QuantizedType); }
 
 	template<typename T = void, typename U = typename TEnableIf<!HasQuantizedType && HasSourceType, T>::Type, bool V = true, char W = 0>
-	static constexpr uint32 GetQuantizedTypeAlignment() { return alignof(typename TChooseClass<std::is_same_v<void, typename NetSerializerImpl::SourceType>, uint8, typename NetSerializerImpl::SourceType>::Result); }
+	static constexpr uint32 GetQuantizedTypeAlignment() { return alignof(std::conditional_t<std::is_same_v<void, typename NetSerializerImpl::SourceType>, uint8, typename NetSerializerImpl::SourceType>); }
 
 	template<typename T = void, typename U = typename TEnableIf<!(HasSourceType || HasQuantizedType), T>::Type, char V = 0>
 	static constexpr uint32 GetQuantizedTypeAlignment() { return 1; }
diff --git a/Engine/Source/Runtime/MeshDescription/Public/MeshAttributeArray.h b/Engine/Source/Runtime/MeshDescription/Public/MeshAttributeArray.h
--- a/Engine/Source/Runtime/MeshDescription/Public/MeshAttributeArray.h
+++ b/Engine/Source/Runtime/MeshDescription/Public/MeshAttributeArray.h
@@ -731,7 +731,7 @@ template <typename T>
 struct TMeshAttributesRefTypeBase
 {
 	using AttributeType = T;
-	using RealAttributeType = typename TChooseClass<TIsDerivedFrom<AttributeType, FElementID>::Value, int32, AttributeType>::Result;
+	using RealAttributeType = std::conditional_t<TIsDerivedFrom<AttributeType, FElementID>::Value, int32, AttributeType>;
 };
 
 template <typename T>
diff --git a/Engine/Source/Runtime/MovieScene/Public/EntitySystem/MovieSceneComponentPtr.h b/Engine/Source/Runtime/MovieScene/Public/EntitySystem/MovieSceneComponentPtr.h
--- a/Engine/Source/Runtime/MovieScene/Public/EntitySystem/MovieSceneComponentPtr.h
+++ b/Engine/Source/Runtime/MovieScene/Public/EntitySystem/MovieSceneComponentPtr.h
@@ -15,7 +15,7 @@ template<typename T>
 struct TComponentPtr
 {
 	/** Value type is either a T or const T& for read-only types, or T& for writeable types */
-	using ValueType = typename TChooseClass<TIsConst<T>::Value, typename TCallTraits<T>::ParamType, T&>::Result;
+	using ValueType = std::conditional_t<TIsConst<T>::Value, typename TCallTraits<T>::ParamType, T&>;
 
 	TComponentPtr()
 		: ComponentPtr(nullptr)
diff --git a/Engine/Source/Runtime/MovieScene/Public/EntitySystem/MovieSceneOverlappingEntityTracker.h b/Engine/Source/Runtime/MovieScene/Public/EntitySystem/MovieSceneOverlappingEntityTracker.h
--- a/Engine/Source/Runtime/MovieScene/Public/EntitySystem/MovieSceneOverlappingEntityTracker.h
+++ b/Engine/Source/Runtime/MovieScene/Public/EntitySystem/MovieSceneOverlappingEntityTracker.h
@@ -664,11 +664,11 @@ protected:
 
 
 template<typename OutputType, typename... KeyType>
-using TOverlappingEntityTracker = typename TChooseClass<
+using TOverlappingEntityTracker = std::conditional_t<
 	(THasAddReferencedObjectForComponent<KeyType>::Value || ...) || THasAddReferencedObjectForComponent<OutputType>::Value,
 	TOverlappingEntityTracker_WithGarbage<OutputType, KeyType...>,
 	TOverlappingEntityTracker_NoGarbage<OutputType, KeyType...>
->::Result;
+>;
 
 } // namespace MovieScene
 } // namespace UE
diff --git a/Engine/Source/Runtime/Net/Core/Public/Net/Core/Trace/Private/NetTraceInternal.h b/Engine/Source/Runtime/Net/Core/Public/Net/Core/Trace/Private/NetTraceInternal.h
--- a/Engine/Source/Runtime/Net/Core/Public/Net/Core/Trace/Private/NetTraceInternal.h
+++ b/Engine/Source/Runtime/Net/Core/Public/Net/Core/Trace/Private/NetTraceInternal.h
@@ -274,7 +274,7 @@ struct FNetTrace
 	template<uint32 Verbosity, typename StreamType>
 	struct TChooseTraceEventScope
 	{
-		typedef typename TChooseClass<FNetTrace::GetNetTraceVerbosityEnabled(Verbosity), FNetTraceEventScope<StreamType>, FNetTraceNullEventScope<StreamType> >::Result Type;
+		typedef std::conditional_t<FNetTrace::GetNetTraceVerbosityEnabled(Verbosity), FNetTraceEventScope<StreamType>, FNetTraceNullEventScope<StreamType> > Type;
 	};
 };
 
diff --git a/Engine/Source/Runtime/Online/Experimental/EventLoop/Public/EventLoop/EventLoopManagedStorage.h b/Engine/Source/Runtime/Online/Experimental/EventLoop/Public/EventLoop/EventLoopManagedStorage.h
--- a/Engine/Source/Runtime/Online/Experimental/EventLoop/Public/EventLoop/EventLoopManagedStorage.h
+++ b/Engine/Source/Runtime/Online/Experimental/EventLoop/Public/EventLoop/EventLoopManagedStorage.h
@@ -405,8 +405,8 @@ public:
 	{
 	private:
 		using TInternalIterator = typename FStorageType::TRangedForIterator;
-		using StorageIteratorType = typename TChooseClass<bConst, typename FStorageType::TRangedForConstIterator, typename FStorageType::TRangedForIterator>::Result;
-		using InternalElementType = typename TChooseClass<bConst, const ElementType, ElementType>::Result;
+		using StorageIteratorType = std::conditional_t<bConst, typename FStorageType::TRangedForConstIterator, typename FStorageType::TRangedForIterator>;
+		using InternalElementType = std::conditional_t<bConst, const ElementType, ElementType>;
 		using ItElementType = TPair<FInternalHandle, InternalElementType&>;
 
 	public:
diff --git a/Engine/Source/Runtime/PhysicsCore/Public/SQVisitor.h b/Engine/Source/Runtime/PhysicsCore/Public/SQVisitor.h
--- a/Engine/Source/Runtime/PhysicsCore/Public/SQVisitor.h
+++ b/Engine/Source/Runtime/PhysicsCore/Public/SQVisitor.h
@@ -431,7 +431,7 @@ private:
 template <typename QueryGeometryType, typename TPayload, typename THitType, bool bGTData = true>
 struct TBPVisitor : public Chaos::ISpatialVisitor<TPayload, Chaos::FReal>
 {
-	using TGeometryType = typename TChooseClass<bGTData, Chaos::FGeometryParticle, Chaos::FGeometryParticleHandle>::Result;
+	using TGeometryType = std::conditional_t<bGTData, Chaos::FGeometryParticle, Chaos::FGeometryParticleHandle>;
 
 	TBPVisitor(const FTransform& InWorldTM, ChaosInterface::FSQHitBuffer<THitType>& InHitBuffer,
 		const ChaosInterface::FQueryFilterData& InQueryFilterData, ICollisionQueryFilterCallbackBase& InQueryCallback, const QueryGeometryType& InQueryGeom, const ChaosInterface::FQueryDebugParams& InDebugParams)
diff --git a/Engine/Source/Runtime/SlateCore/Private/FastUpdate/SlateInvalidationWidgetList.cpp b/Engine/Source/Runtime/SlateCore/Private/FastUpdate/SlateInvalidationWidgetList.cpp
--- a/Engine/Source/Runtime/SlateCore/Private/FastUpdate/SlateInvalidationWidgetList.cpp
+++ b/Engine/Source/Runtime/SlateCore/Private/FastUpdate/SlateInvalidationWidgetList.cpp
@@ -744,7 +744,7 @@ namespace InvalidationList
 	template<typename TSlateInvalidationWidgetList, typename Predicate>
 	void ForEachChildren(TSlateInvalidationWidgetList& Self, const typename TSlateInvalidationWidgetList::InvalidationWidgetType& InvalidationWidget, FSlateInvalidationWidgetIndex WidgetIndex, Predicate InPredicate)
 	{
-		using SlateInvalidationWidgetType = typename TChooseClass<TIsConst<TSlateInvalidationWidgetList>::Value, const typename TSlateInvalidationWidgetList::InvalidationWidgetType, typename TSlateInvalidationWidgetList::InvalidationWidgetType>::Result;
+		using SlateInvalidationWidgetType = std::conditional_t<TIsConst<TSlateInvalidationWidgetList>::Value, const typename TSlateInvalidationWidgetList::InvalidationWidgetType, typename TSlateInvalidationWidgetList::InvalidationWidgetType>;
 		if (InvalidationWidget.LeafMostChildIndex != WidgetIndex)
 		{
 			FSlateInvalidationWidgetIndex CurrentWidgetIndex = Self.IncrementIndex(WidgetIndex);
