Fixed unity build.

diff --git a/Engine/Source/Runtime/Experimental/Iris/Core/Public/Iris/Serialization/BitPacking.h b/Engine/Source/Runtime/Experimental/Iris/Core/Public/Iris/Serialization/BitPacking.h
--- a/Engine/Source/Runtime/Experimental/Iris/Core/Public/Iris/Serialization/BitPacking.h
+++ b/Engine/Source/Runtime/Experimental/Iris/Core/Public/Iris/Serialization/BitPacking.h
@@ -2,9 +2,10 @@
 
 #pragma once
 
+#include <type_traits>
+
 #include "Iris/Serialization/NetBitStreamReader.h"
 #include "Iris/Serialization/NetBitStreamWriter.h"
-#include "Templates/ChooseClass.h"
 #include "Templates/EnableIf.h"
 #include "Templates/IsIntegral.h"
 #include "Templates/IsSigned.h"
