Geometry collection : update code to support masked materials for nanite code path

diff --git a/Engine/Source/Runtime/Experimental/GeometryCollectionEngine/Private/GeometryCollection/GeometryCollectionSceneProxy.cpp b/Engine/Source/Runtime/Experimental/GeometryCollectionEngine/Private/GeometryCollection/GeometryCollectionSceneProxy.cpp
--- a/Engine/Source/Runtime/Experimental/GeometryCollectionEngine/Private/GeometryCollection/GeometryCollectionSceneProxy.cpp
+++ b/Engine/Source/Runtime/Experimental/GeometryCollectionEngine/Private/GeometryCollection/GeometryCollectionSceneProxy.cpp
@@ -1060,7 +1060,7 @@ FNaniteGeometryCollectionSceneProxy::FNaniteGeometryCollectionSceneProxy(UGeomet
 		UMaterialInterface* MaterialInterface = bValidMeshSection ? Component->GetMaterial(MeshSection.MaterialID) : nullptr;
 
 		// TODO: PROG_RASTER (Implement programmable raster support)
-		const bool bInvalidMaterial = !MaterialInterface || !IsOpaqueBlendMode(*MaterialInterface) || MaterialInterface->GetShadingModels().HasShadingModel(MSM_SingleLayerWater);
+		const bool bInvalidMaterial = !MaterialInterface || !IsOpaqueOrMaskedBlendMode(*MaterialInterface) || MaterialInterface->GetShadingModels().HasShadingModel(MSM_SingleLayerWater);
 		if (bInvalidMaterial)
 		{
 			bHasMaterialErrors = true;
