Adding mechanism for getting output from vehicles modules to game thread.
Option to combine local actors into cluster union
Linear and angular damping option on vehicle cluster union

diff --git a/Engine/Source/Runtime/Experimental/ChaosVehicles/ChaosVehiclesCore/Private/SimModule/EngineModule.cpp b/Engine/Source/Runtime/Experimental/ChaosVehicles/ChaosVehiclesCore/Private/SimModule/EngineModule.cpp
--- a/Engine/Source/Runtime/Experimental/ChaosVehicles/ChaosVehiclesCore/Private/SimModule/EngineModule.cpp
+++ b/Engine/Source/Runtime/Experimental/ChaosVehicles/ChaosVehiclesCore/Private/SimModule/EngineModule.cpp
@@ -20,7 +20,7 @@ namespace Chaos
 
 		// TODO: Engine braking effect
 		DriveTorque = GetEngineTorque(Inputs.ControlInputs.Throttle, GetRPM());
-		
+
 		float BrakeTorque = 0.f;
 		TransmitTorque(VehicleModuleSystem, DriveTorque, BrakeTorque);
 		IntegrateAngularVelocity(DeltaTime, Setup().EngineInertia);
@@ -71,6 +71,29 @@ namespace Chaos
 		return Setup().TorqueCurve.GetValue(RPM, Setup().MaxRPM, Setup().MaxTorque);
 	}
 
+	void FEngineOutputData::FillOutputState(const ISimulationModuleBase* SimModule)
+	{
+		check(SimModule->GetSimType() == eSimType::Engine);
+		if (const FEngineSimModule* Sim = static_cast<const FEngineSimModule*>(SimModule))
+		{
+			RPM = Sim->GetRPM();
+		}
+	}
+
+	void FEngineOutputData::Lerp(const FSimOutputData& InCurrent, const FSimOutputData& InNext, float Alpha)
+	{
+		const FEngineOutputData& Current = static_cast<const FEngineOutputData&>(InCurrent);
+		const FEngineOutputData& Next = static_cast<const FEngineOutputData&>(InNext);
+
+		RPM = FMath::Lerp(Current.RPM, Next.RPM, Alpha);
+	}
+
+	FString FEngineOutputData::ToString()
+	{
+		return  FString::Printf(TEXT("RPM=%3.3f")
+			, RPM);
+	}
+
 } // namespace Chaos
 
 #if VEHICLE_DEBUGGING_ENABLED
diff --git a/Engine/Source/Runtime/Experimental/ChaosVehicles/ChaosVehiclesCore/Private/SimModule/SuspensionModule.cpp b/Engine/Source/Runtime/Experimental/ChaosVehicles/ChaosVehiclesCore/Private/SimModule/SuspensionModule.cpp
--- a/Engine/Source/Runtime/Experimental/ChaosVehicles/ChaosVehiclesCore/Private/SimModule/SuspensionModule.cpp
+++ b/Engine/Source/Runtime/Experimental/ChaosVehicles/ChaosVehiclesCore/Private/SimModule/SuspensionModule.cpp
@@ -16,6 +16,7 @@ namespace Chaos
 		: TSimModuleSettings<FSuspensionSettings>(Settings)
 		, SpringDisplacement(0.f)
 		, LastDisplacement(0.f)
+		, SpringSpeed(0.f)
 		, WheelSimTreeIndex(INVALID_IDX)
 	{
 		AccessSetup().MaxLength = FMath::Abs(Settings.MaxRaise + Settings.MaxDrop);
@@ -51,7 +52,7 @@ namespace Chaos
 		if (SpringDisplacement > 0)
 		{
 			float Damping = (SpringDisplacement < LastDisplacement) ? Setup().CompressionDamping : Setup().ReboundDamping;
-			float SpringSpeed = (LastDisplacement - SpringDisplacement) / DeltaTime;
+			SpringSpeed = (LastDisplacement - SpringDisplacement) / DeltaTime;
 
 			float StiffnessForce = SpringDisplacement * Setup().SpringRate;
 			float DampingForce = SpringSpeed * Damping;
@@ -75,7 +76,7 @@ namespace Chaos
 
 				Wheel->SetForceIntoSurface(ForceIntoSurface);
 			}
-		
+
 		}
 	}
 
@@ -96,6 +97,7 @@ namespace Chaos
 
 	void FSuspensionSimModuleDatas::FillSimState(ISimulationModuleBase* SimModule)
 	{
+		check(SimModule->GetSimType() == eSimType::Suspension);
 		if (FSuspensionSimModule* Sim = static_cast<FSuspensionSimModule*>(SimModule))
 		{
 			Sim->SpringDisplacement = SpringDisplacement;
@@ -105,6 +107,7 @@ namespace Chaos
 
 	void FSuspensionSimModuleDatas::FillNetState(const ISimulationModuleBase* SimModule)
 	{
+		check(SimModule->GetSimType() == eSimType::Suspension);
 		if (const FSuspensionSimModule* Sim = static_cast<const FSuspensionSimModule*>(SimModule))
 		{
 			SpringDisplacement = Sim->SpringDisplacement;
@@ -129,8 +132,32 @@ namespace Chaos
 	}
 #endif
 
+	void FSuspensionOutputData::FillOutputState(const ISimulationModuleBase* SimModule)
+	{
+		check(SimModule->GetSimType() == eSimType::Suspension);
+		if (const FSuspensionSimModule* Sim = static_cast<const FSuspensionSimModule*>(SimModule))
+		{
+			SpringDisplacement = Sim->SpringDisplacement;
+			SpringSpeed = Sim->SpringSpeed;
+		}
+	}
+
+	void FSuspensionOutputData::Lerp(const FSimOutputData& InCurrent, const FSimOutputData& InNext, float Alpha)
+	{
+		const FSuspensionOutputData& Current = static_cast<const FSuspensionOutputData&>(InCurrent);
+		const FSuspensionOutputData& Next = static_cast<const FSuspensionOutputData&>(InNext);
+
+		SpringDisplacement = FMath::Lerp(Current.SpringDisplacement, Next.SpringDisplacement, Alpha);
+		SpringSpeed = FMath::Lerp(Current.SpringSpeed, Next.SpringSpeed, Alpha);
+	}
+
+	FString FSuspensionOutputData::ToString()
+	{
+		return FString::Printf(TEXT("SpringDisplacement=%3.3f, SpringSpeed=%3.3f"), SpringDisplacement, SpringSpeed);
+	}
+
 } // namespace Chaos
 
 #if VEHICLE_DEBUGGING_ENABLED
 UE_ENABLE_OPTIMIZATION
-#endif
\ No newline at end of file
+#endif
diff --git a/Engine/Source/Runtime/Experimental/ChaosVehicles/ChaosVehiclesCore/Private/SimModule/TransmissionModule.cpp b/Engine/Source/Runtime/Experimental/ChaosVehicles/ChaosVehiclesCore/Private/SimModule/TransmissionModule.cpp
--- a/Engine/Source/Runtime/Experimental/ChaosVehicles/ChaosVehiclesCore/Private/SimModule/TransmissionModule.cpp
+++ b/Engine/Source/Runtime/Experimental/ChaosVehicles/ChaosVehiclesCore/Private/SimModule/TransmissionModule.cpp
@@ -135,7 +135,8 @@ namespace Chaos
 
 	void FTransmissionSimModuleDatas::FillSimState(ISimulationModuleBase* SimModule)
 	{
-		if(FTransmissionSimModule* Sim = static_cast<FTransmissionSimModule*>(SimModule))
+		check(SimModule->GetSimType() == eSimType::Transmission);
+		if (FTransmissionSimModule* Sim = static_cast<FTransmissionSimModule*>(SimModule))
 		{
 			Sim->CurrentGear = CurrentGear;
 			Sim->TargetGear = TargetGear;
@@ -145,6 +146,7 @@ namespace Chaos
 
 	void FTransmissionSimModuleDatas::FillNetState(const ISimulationModuleBase* SimModule)
 	{
+		check(SimModule->GetSimType() == eSimType::Transmission);
 		if (const FTransmissionSimModule* Sim = static_cast<const FTransmissionSimModule*>(SimModule))
 		{
 			CurrentGear = Sim->CurrentGear;
@@ -171,6 +173,28 @@ namespace Chaos
 	}
 #endif
 
+	void FTransmissionOutputData::FillOutputState(const ISimulationModuleBase* SimModule)
+	{
+		check(SimModule->GetSimType() == eSimType::Transmission);
+		if (const FTransmissionSimModule* Sim = static_cast<const FTransmissionSimModule*>(SimModule))
+		{
+			CurrentGear = Sim->CurrentGear;
+		}
+	}
+
+	void FTransmissionOutputData::Lerp(const FSimOutputData& InCurrent, const FSimOutputData& InNext, float Alpha)
+	{
+		const FTransmissionOutputData& Current = static_cast<const FTransmissionOutputData&>(InCurrent);
+		const FTransmissionOutputData& Next = static_cast<const FTransmissionOutputData&>(InNext);
+
+		CurrentGear = Current.CurrentGear;
+	}
+
+	FString FTransmissionOutputData::ToString()
+	{
+		return  FString::Printf(TEXT("CurrentGear=%d")
+			, CurrentGear);
+	}
 
 } // namespace Chaos
 
diff --git a/Engine/Source/Runtime/Experimental/ChaosVehicles/ChaosVehiclesCore/Private/SimModule/WheelModule.cpp b/Engine/Source/Runtime/Experimental/ChaosVehicles/ChaosVehiclesCore/Private/SimModule/WheelModule.cpp
--- a/Engine/Source/Runtime/Experimental/ChaosVehicles/ChaosVehiclesCore/Private/SimModule/WheelModule.cpp
+++ b/Engine/Source/Runtime/Experimental/ChaosVehicles/ChaosVehiclesCore/Private/SimModule/WheelModule.cpp
@@ -18,14 +18,14 @@ namespace Chaos
 		, SurfaceFriction(1.0f)
 		, SuspensionSimTreeIndex(INVALID_IDX)
 		, ForceFromFriction(FVector::ZeroVector)
-		, MassPerWheel(500.0f*0.25f)
+		, MassPerWheel(500.0f * 0.25f)
 		, SteerAngleDegrees(0.0f)
 	{
 
 	}
 
 	void FWheelSimModule::Simulate(float DeltaTime, const FAllInputs& Inputs, FSimModuleTree& VehicleModuleSystem)
-	{	
+	{
 		float Re = Setup().Radius;
 		float K = 0.4f;
 		float TorqueScaling = 1.0f;
@@ -43,7 +43,7 @@ namespace Chaos
 			BrakeTorque = Setup().HandbrakeTorque;
 		}
 
-		bool bTouchingGround = ForceIntoSurface > SMALL_NUMBER;
+		bTouchingGround = ForceIntoSurface > SMALL_NUMBER;
 
 		if (bTouchingGround)
 		{
@@ -57,7 +57,7 @@ namespace Chaos
 			float TorqueFromGroundInteraction = Delta * Setup().WheelInertia / DeltaTime; // torque from wheels moving over terrain
 
 			// X is longitudinal direction, Y is lateral
-			float SlipAngle = FVehicleUtility::CalculateSlipAngle(LocalWheelVelocity.Y, LocalWheelVelocity.X);
+			SlipAngle = FVehicleUtility::CalculateSlipAngle(LocalWheelVelocity.Y, LocalWheelVelocity.X);
 
 			float AppliedLinearDriveForce = DriveTorque / Re;
 			float AppliedLinearBrakeForce = FMath::Abs(BrakeTorque) / Re;
@@ -179,7 +179,7 @@ namespace Chaos
 				ForceFromFriction = -ForceFromFriction;
 			}
 
- 			AddLocalForce(SteeringRotator.RotateVector(ForceFromFriction));
+			AddLocalForce(SteeringRotator.RotateVector(ForceFromFriction));
 			TransmitTorque(VehicleModuleSystem, DriveTorque, BrakeTorque);
 
 			DriveTorque -= AvailableGrip;
@@ -226,6 +226,36 @@ namespace Chaos
 	}
 
 
+	inline void FWheelOutputData::FillOutputState(const ISimulationModuleBase* SimModule)
+	{
+		check(SimModule->GetSimType() == eSimType::Wheel);
+		if (const FWheelSimModule* Sim = static_cast<const FWheelSimModule*>(SimModule))
+		{
+			bTouchingGround = Sim->bTouchingGround;
+			ForceIntoSurface = Sim->ForceIntoSurface;
+			SlipAngle = Sim->SlipAngle;
+			RPM = Sim->GetRPM();
+		}
+	}
+
+	void FWheelOutputData::Lerp(const FSimOutputData& InCurrent, const FSimOutputData& InNext, float Alpha)
+	{
+		const FWheelOutputData& Current = static_cast<const FWheelOutputData&>(InCurrent);
+		const FWheelOutputData& Next = static_cast<const FWheelOutputData&>(InNext);
+
+		bTouchingGround = Current.bTouchingGround;
+		ForceIntoSurface = FMath::Lerp(Current.ForceIntoSurface, Next.ForceIntoSurface, Alpha);
+		SlipAngle = FMath::Lerp(Current.SlipAngle, Next.SlipAngle, Alpha);
+		RPM = FMath::Lerp(Current.RPM, Next.RPM, Alpha);
+	}
+
+	FString FWheelOutputData::ToString()
+	{
+		return  FString::Printf(TEXT("bTouchingGround=%d, ForceIntoSurface=%3.3f, SlipAngle=%3.3f, RPM=%3.3f")
+			, bTouchingGround, ForceIntoSurface, SlipAngle, RPM);
+	}
+
+
 } // namespace Chaos
 
 
diff --git a/Engine/Source/Runtime/Experimental/ChaosVehicles/ChaosVehiclesCore/Public/SimModule/ClutchModule.h b/Engine/Source/Runtime/Experimental/ChaosVehicles/ChaosVehiclesCore/Public/SimModule/ClutchModule.h
--- a/Engine/Source/Runtime/Experimental/ChaosVehicles/ChaosVehiclesCore/Public/SimModule/ClutchModule.h
+++ b/Engine/Source/Runtime/Experimental/ChaosVehicles/ChaosVehiclesCore/Public/SimModule/ClutchModule.h
@@ -21,6 +21,32 @@ namespace Chaos
 		float ClutchStrength;
 	};
 
+
+	struct CHAOSVEHICLESCORE_API FClutchSimModuleDatas : public FTorqueSimModuleDatas
+	{
+#if !(UE_BUILD_SHIPPING || UE_BUILD_TEST)
+		FClutchSimModuleDatas(int NodeArrayIndex, const FString& InDebugString) : FTorqueSimModuleDatas(NodeArrayIndex, InDebugString) {}
+#else
+		FClutchSimModuleDatas(int NodeArrayIndex) : FTorqueSimModuleDatas(NodeArrayIndex) {}
+#endif
+
+		virtual eSimType GetType() override { return eSimType::Clutch; }
+
+		virtual void FillSimState(ISimulationModuleBase* SimModule) override
+		{
+			check(SimModule->GetSimType() == eSimType::Clutch);
+			FTorqueSimModuleDatas::FillSimState(SimModule);
+		}
+
+		virtual void FillNetState(const ISimulationModuleBase* SimModule) override
+		{
+			check(SimModule->GetSimType() == eSimType::Clutch);
+			FTorqueSimModuleDatas::FillNetState(SimModule);
+		}
+
+	};
+
+
 	/// <summary>
 	/// 
 	/// a vehicle component that transmits torque from one source to another through a clutch system, i.e. connect an engine to a transmission
@@ -36,6 +62,16 @@ namespace Chaos
 
 		FClutchSimModule(const FClutchSettings& Settings);
 
+		virtual TSharedPtr<FModuleNetData> GenerateNetData(int SimArrayIndex) const
+		{
+			return MakeShared<FClutchSimModuleDatas>(
+				SimArrayIndex
+#if !(UE_BUILD_SHIPPING || UE_BUILD_TEST)
+				, GetDebugName()
+#endif			
+			);
+		}
+
 		virtual eSimType GetSimType() const { return eSimType::Clutch; }
 
 		virtual const FString GetDebugName() const { return TEXT("Clutch"); }
diff --git a/Engine/Source/Runtime/Experimental/ChaosVehicles/ChaosVehiclesCore/Public/SimModule/EngineModule.h b/Engine/Source/Runtime/Experimental/ChaosVehicles/ChaosVehiclesCore/Public/SimModule/EngineModule.h
--- a/Engine/Source/Runtime/Experimental/ChaosVehicles/ChaosVehiclesCore/Public/SimModule/EngineModule.h
+++ b/Engine/Source/Runtime/Experimental/ChaosVehicles/ChaosVehiclesCore/Public/SimModule/EngineModule.h
@@ -3,12 +3,49 @@
 #pragma once
 
 #include "SimModule/TorqueSimModule.h"
+#include "SimModule/SimulationModuleBase.h"
 
 namespace Chaos
 {
 	struct FAllInputs;
 	class FSimModuleTree;
 
+	struct CHAOSVEHICLESCORE_API FEngineSimModuleDatas : public FTorqueSimModuleDatas
+	{
+#if !(UE_BUILD_SHIPPING || UE_BUILD_TEST)
+		FEngineSimModuleDatas(int NodeArrayIndex, const FString& InDebugString) : FTorqueSimModuleDatas(NodeArrayIndex, InDebugString) {}
+#else
+		FEngineSimModuleDatas(int NodeArrayIndex) : FTorqueSimModuleDatas(NodeArrayIndex) {}
+#endif
+
+		virtual eSimType GetType() override { return eSimType::Engine; }
+
+		virtual void FillSimState(ISimulationModuleBase* SimModule) override
+		{
+			check(SimModule->GetSimType() == eSimType::Engine);
+			FTorqueSimModuleDatas::FillSimState(SimModule);
+		}
+
+		virtual void FillNetState(const ISimulationModuleBase* SimModule) override
+		{
+			check(SimModule->GetSimType() == eSimType::Engine);
+			FTorqueSimModuleDatas::FillNetState(SimModule);
+		}
+
+	};
+
+	struct CHAOSVEHICLESCORE_API FEngineOutputData : public FSimOutputData
+	{
+		virtual FSimOutputData* MakeNewData() override { return FEngineOutputData::MakeNew(); }
+		static FSimOutputData* MakeNew() { return new FEngineOutputData(); }
+
+		virtual void FillOutputState(const ISimulationModuleBase* SimModule) override;
+		virtual void Lerp(const FSimOutputData& InCurrent, const FSimOutputData& InNext, float Alpha) override;
+		virtual FString ToString() override;
+
+		float RPM;
+	};
+
 	struct CHAOSVEHICLESCORE_API FEngineSettings
 	{
 		FEngineSettings()
@@ -54,6 +91,20 @@ namespace Chaos
 
 		virtual ~FEngineSimModule() {}
 
+		virtual TSharedPtr<FModuleNetData> GenerateNetData(int SimArrayIndex) const
+		{
+			return MakeShared<FEngineSimModuleDatas>(
+				SimArrayIndex
+#if !(UE_BUILD_SHIPPING || UE_BUILD_TEST)
+				, GetDebugName()
+#endif			
+			);
+		}
+		virtual FSimOutputData* GenerateOutputData() const override
+		{
+			return FEngineOutputData::MakeNew();
+		}
+
 		virtual eSimType GetSimType() const { return eSimType::Engine; }
 
 		virtual const FString GetDebugName() const { return TEXT("Engine"); }
diff --git a/Engine/Source/Runtime/Experimental/ChaosVehicles/ChaosVehiclesCore/Public/SimModule/SimulationModuleBase.h b/Engine/Source/Runtime/Experimental/ChaosVehicles/ChaosVehiclesCore/Public/SimModule/SimulationModuleBase.h
--- a/Engine/Source/Runtime/Experimental/ChaosVehicles/ChaosVehiclesCore/Public/SimModule/SimulationModuleBase.h
+++ b/Engine/Source/Runtime/Experimental/ChaosVehicles/ChaosVehiclesCore/Public/SimModule/SimulationModuleBase.h
@@ -22,6 +22,7 @@ namespace Chaos
 {
 	class FSimModuleTree;
 	struct FModuleNetData;
+	struct FSimOutputData;
 	class FClusterUnionPhysicsProxy;
 
 	struct CHAOSVEHICLESCORE_API FControlInputs
@@ -112,10 +113,10 @@ namespace Chaos
 
 	enum eSimModuleTypeFlags
 	{
-		NonFunctional	= (1 << 0),	// bitmask 1,2,4,8
-		Raycast			= (1 << 1),	// requires raycast data
-		TorqueBased		= (1 << 2),	// performs torque calculations
-		Velocity		= (1 << 3),	// requires velocity data
+		NonFunctional = (1 << 0),	// bitmask 1,2,4,8
+		Raycast = (1 << 1),	// requires raycast data
+		TorqueBased = (1 << 2),	// performs torque calculations
+		Velocity = (1 << 3),	// requires velocity data
 	};
 
 	enum eSimType
@@ -257,7 +258,7 @@ namespace Chaos
 		/**
 		 * Set the COM relative transform of module when it is clustered, so relative to parent COM
 		 */
-		void SetClusteredTransform(const FTransform& TransformIn) { ClusteredCOMRelativeTransform  = TransformIn; }
+		void SetClusteredTransform(const FTransform& TransformIn) { ClusteredCOMRelativeTransform = TransformIn; }
 		const FTransform& GetClusteredTransform() const { return ClusteredCOMRelativeTransform; }
 
 		void SetInitialParticleTransform(const FTransform& TransformIn) { InitialParticleTransform = TransformIn; }
@@ -295,6 +296,8 @@ namespace Chaos
 		// this is the replication datas
 		virtual TSharedPtr<FModuleNetData> GenerateNetData(int NodeArrayIndex) const = 0;
 
+		virtual FSimOutputData* GenerateOutputData() const { return nullptr; }
+
 		//void SetClusterParticle(FPBDRigidClusteredParticleHandle* ParticleIn) { ClusterParticle = ParticleIn; }
 		Chaos::FPBDRigidClusteredParticleHandle* GetClusterParticle(Chaos::FClusterUnionPhysicsProxy* Proxy);
 
@@ -353,6 +356,18 @@ namespace Chaos
 
 	using FModuleNetDataArray = TArray<TSharedPtr<FModuleNetData>>;
 
+	struct CHAOSVEHICLESCORE_API FSimOutputData
+	{
+		FSimOutputData() = default;
+		virtual ~FSimOutputData() {}
+
+		virtual FSimOutputData* MakeNewData() = 0;
+		virtual void FillOutputState(const ISimulationModuleBase* SimModule) = 0;
+		virtual void Lerp(const FSimOutputData& InCurrent, const FSimOutputData& InNext, float Alpha) = 0;
+		virtual FString ToString() { return FString(); }
+	};
+
+
 } // namespace Chaos
 
 
diff --git a/Engine/Source/Runtime/Experimental/ChaosVehicles/ChaosVehiclesCore/Public/SimModule/SuspensionModule.h b/Engine/Source/Runtime/Experimental/ChaosVehicles/ChaosVehiclesCore/Public/SimModule/SuspensionModule.h
--- a/Engine/Source/Runtime/Experimental/ChaosVehicles/ChaosVehiclesCore/Public/SimModule/SuspensionModule.h
+++ b/Engine/Source/Runtime/Experimental/ChaosVehicles/ChaosVehiclesCore/Public/SimModule/SuspensionModule.h
@@ -40,7 +40,19 @@ namespace Chaos
 		float LastDisplacement = 0.0f;
 	};
 
+	struct CHAOSVEHICLESCORE_API FSuspensionOutputData : public FSimOutputData
+	{
+		virtual FSimOutputData* MakeNewData() override { return FSuspensionOutputData::MakeNew(); }
+		static FSimOutputData* MakeNew() { return new FSuspensionOutputData(); }
+
+		virtual void FillOutputState(const ISimulationModuleBase* SimModule) override;
+		virtual void Lerp(const FSimOutputData& InCurrent, const FSimOutputData& InNext, float Alpha) override;
 
+		virtual FString ToString() override;
+
+		float SpringDisplacement;
+		float SpringSpeed;
+	};
 
 	struct CHAOSVEHICLESCORE_API FSuspensionSettings
 	{
@@ -54,14 +66,14 @@ namespace Chaos
 			, SpringPreload(0.5f)
 			, CompressionDamping(0.9f)
 			, ReboundDamping(0.9f)
-		//	, SwaybarEffect(0.5f)
-		//	, DampingRatio(0.3f)
+			//	, SwaybarEffect(0.5f)
+			//	, DampingRatio(0.3f)
 		{
 
 		}
 
 		FVector SuspensionAxis;		// local axis, direction of suspension force raycast traces
-		FVector RestOffset;	
+		FVector RestOffset;
 		float MaxRaise;				// distance [cm]
 		float MaxDrop;				// distance [cm]
 		float MaxLength;			// distance [cm]
@@ -71,9 +83,9 @@ namespace Chaos
 		float CompressionDamping;	// limit compression speed
 		float ReboundDamping;		// limit rebound speed
 
-	//	float Swaybar;				// Anti-roll bar
+		//	float Swaybar;				// Anti-roll bar
 
-	//	float DampingRatio;			// value between (0-no damping) and (1-critical damping)
+		//	float DampingRatio;			// value between (0-no damping) and (1-critical damping)
 	};
 
 	/** Suspension world ray/shape trace start and end positions */
@@ -98,6 +110,7 @@ namespace Chaos
 	class CHAOSVEHICLESCORE_API FSuspensionSimModule : public ISimulationModuleBase, public TSimModuleSettings<FSuspensionSettings>
 	{
 		friend FSuspensionSimModuleDatas;
+		friend FSuspensionOutputData;
 
 	public:
 
@@ -114,6 +127,11 @@ namespace Chaos
 			);
 		}
 
+		virtual FSimOutputData* GenerateOutputData() const override
+		{
+			return FSuspensionOutputData::MakeNew();
+		}
+
 		virtual eSimType GetSimType() const { return eSimType::Suspension; }
 
 		virtual const FString GetDebugName() const { return TEXT("Suspension"); }
@@ -135,6 +153,7 @@ namespace Chaos
 
 		float SpringDisplacement;
 		float LastDisplacement;
+		float SpringSpeed;
 		int WheelSimTreeIndex;
 	};
 
diff --git a/Engine/Source/Runtime/Experimental/ChaosVehicles/ChaosVehiclesCore/Public/SimModule/TorqueSimModule.h b/Engine/Source/Runtime/Experimental/ChaosVehicles/ChaosVehiclesCore/Public/SimModule/TorqueSimModule.h
--- a/Engine/Source/Runtime/Experimental/ChaosVehicles/ChaosVehiclesCore/Public/SimModule/TorqueSimModule.h
+++ b/Engine/Source/Runtime/Experimental/ChaosVehicles/ChaosVehiclesCore/Public/SimModule/TorqueSimModule.h
@@ -17,7 +17,7 @@ namespace Chaos
 		FTorqueSimModuleDatas(int NodeArrayIndex) : FModuleNetData(NodeArrayIndex) {}
 #endif
 		virtual eSimType GetType() override { return eSimType::TorqueSim; }
-		
+
 		virtual void FillSimState(ISimulationModuleBase* SimModule) override;
 
 		virtual void FillNetState(const ISimulationModuleBase* SimModule) override;
@@ -45,87 +45,78 @@ namespace Chaos
 	};
 
 
-class CHAOSVEHICLESCORE_API FTorqueSimModule : public ISimulationModuleBase
-{
-	friend FTorqueSimModuleDatas;
-
-public:
-	FTorqueSimModule()
-		: DriveTorque(0.f)
-		, LoadTorque(0.f)
-		, BrakingTorque(0.f)
-		, AngularVelocity(0.f)
-		, AngularPosition(0.0f)
+	class CHAOSVEHICLESCORE_API FTorqueSimModule : public ISimulationModuleBase
 	{
-	}
+		friend FTorqueSimModuleDatas;
+
+	public:
+		FTorqueSimModule()
+			: DriveTorque(0.f)
+			, LoadTorque(0.f)
+			, BrakingTorque(0.f)
+			, AngularVelocity(0.f)
+			, AngularPosition(0.0f)
+		{
+		}
 
-	/**
-	 * Is Module of a specific type - used for casting
-	 */
-	virtual bool IsBehaviourType(eSimModuleTypeFlags InType) const { return (InType & TorqueBased); }
-	virtual TSharedPtr<FModuleNetData> GenerateNetData(int SimArrayIndex) const
-	{
-		return MakeShared<FTorqueSimModuleDatas>(
-			SimArrayIndex
-#if !(UE_BUILD_SHIPPING || UE_BUILD_TEST)
-			, GetDebugName()
-#endif			
-		);
-	}
+		/**
+		 * Is Module of a specific type - used for casting
+		 */
+		virtual bool IsBehaviourType(eSimModuleTypeFlags InType) const { return (InType & TorqueBased); }
 
-	virtual eSimType GetSimType() const { return eSimType::TorqueSim; }
+		virtual eSimType GetSimType() const { return eSimType::TorqueSim; }
 
-	void SetDriveTorque(float TorqueIn) { DriveTorque = TorqueIn; }
-	float GetDriveTorque() const { return DriveTorque; }
+		void SetDriveTorque(float TorqueIn) { DriveTorque = TorqueIn; }
+		float GetDriveTorque() const { return DriveTorque; }
 
-	void SetLoadTorque(float TorqueIn) { LoadTorque = TorqueIn; }
-	float GetLoadTorque() const { return LoadTorque; }
+		void SetLoadTorque(float TorqueIn) { LoadTorque = TorqueIn; }
+		float GetLoadTorque() const { return LoadTorque; }
 
-	void SetBrakingTorque(float TorqueIn) { BrakingTorque = TorqueIn; }
-	float GetBrakingTorque() const { return BrakingTorque; }
+		void SetBrakingTorque(float TorqueIn) { BrakingTorque = TorqueIn; }
+		float GetBrakingTorque() const { return BrakingTorque; }
 
-	void SetAngularVelocity(float AngularVelocityIn) { AngularVelocity = AngularVelocityIn; }
-	float GetAngularVelocity() const { return AngularVelocity; }
+		void SetAngularVelocity(float AngularVelocityIn) { AngularVelocity = AngularVelocityIn; }
+		float GetAngularVelocity() const { return AngularVelocity; }
 
-	void AddAngularVelocity(float AngularVelocityIn) { AngularVelocity += AngularVelocityIn; }
+		void AddAngularVelocity(float AngularVelocityIn) { AngularVelocity += AngularVelocityIn; }
 
-	void SetAngularPosition(float AngularPositionIn) { AngularPosition = AngularPositionIn; }
-	float GetAngularPosition() const { return AngularPosition; }
+		void SetAngularPosition(float AngularPositionIn) { AngularPosition = AngularPositionIn; }
+		float GetAngularPosition() const { return AngularPosition; }
 
-	void SetRPM(float InRPM) { AngularVelocity = RPMToOmega(InRPM); }
-	float GetRPM() const { return OmegaToRPM(AngularVelocity); }
+		void SetRPM(float InRPM) { AngularVelocity = RPMToOmega(InRPM); }
+		float GetRPM() const { return OmegaToRPM(AngularVelocity); }
 
-	/**
-	 * Transmit torque between this module and its Parent and Children. DriveTorque passed down to children, LoadTorque passed from child to parent
-	 */
-	void TransmitTorque(const FSimModuleTree& BlockSystem, float PushedTorque, float BrakeTorque = 0.f, float GearingRatio = 1.0f, float ClutchSlip = 1.0f);
+		/**
+		 * Transmit torque between this module and its Parent and Children. DriveTorque passed down to children, LoadTorque passed from child to parent
+		 */
+		void TransmitTorque(const FSimModuleTree& BlockSystem, float PushedTorque, float BrakeTorque = 0.f, float GearingRatio = 1.0f, float ClutchSlip = 1.0f);
 
-	/**
-	 * Integrate angular velocity using specified DeltaTime & Inertia value, Note the inertia should be the combined inertia of all the connected pieces otherwise things will rotate at different rates
-	 */
-	void IntegrateAngularVelocity(float DeltaTime, float Inertia, float MaxRotationVel=MAX_FLT);
+		/**
+		 * Integrate angular velocity using specified DeltaTime & Inertia value, Note the inertia should be the combined inertia of all the connected pieces otherwise things will rotate at different rates
+		 */
+		void IntegrateAngularVelocity(float DeltaTime, float Inertia, float MaxRotationVel = MAX_FLT);
 
-	/**
-	 * Cast an ISimulationModuleBase to a FTorqueSimModule if compatible class
-	 */
-	static FTorqueSimModule* CastToTorqueInterface(ISimulationModuleBase* SimModule)
-	{
-		if (SimModule && SimModule->IsBehaviourType(eSimModuleTypeFlags::TorqueBased))
+		/**
+		 * Cast an ISimulationModuleBase to a FTorqueSimModule if compatible class
+		 */
+		static FTorqueSimModule* CastToTorqueInterface(ISimulationModuleBase* SimModule)
 		{
-			return static_cast<FTorqueSimModule*>(SimModule);
-		}
+			if (SimModule && SimModule->IsBehaviourType(eSimModuleTypeFlags::TorqueBased))
+			{
+				return static_cast<FTorqueSimModule*>(SimModule);
+			}
 
-		return nullptr;
-	}
+			return nullptr;
+		}
 
-protected:
+	protected:
 
-	float DriveTorque;
-	float LoadTorque;
-	float BrakingTorque;
-	float AngularVelocity;
-	float AngularPosition;
-};
+		float DriveTorque;
+		float LoadTorque;
+		float BrakingTorque;
+		float AngularVelocity;
+		float AngularPosition;
+	};
 
 
-} // namespace Chaos
\ No newline at end of file
+} // namespace Chaos
diff --git a/Engine/Source/Runtime/Experimental/ChaosVehicles/ChaosVehiclesCore/Public/SimModule/TransmissionModule.h b/Engine/Source/Runtime/Experimental/ChaosVehicles/ChaosVehiclesCore/Public/SimModule/TransmissionModule.h
--- a/Engine/Source/Runtime/Experimental/ChaosVehicles/ChaosVehiclesCore/Public/SimModule/TransmissionModule.h
+++ b/Engine/Source/Runtime/Experimental/ChaosVehicles/ChaosVehiclesCore/Public/SimModule/TransmissionModule.h
@@ -10,7 +10,6 @@ namespace Chaos
 	struct FAllInputs;
 	class FSimModuleTree;
 
-
 	struct CHAOSVEHICLESCORE_API FTransmissionSimModuleDatas : public FModuleNetData
 	{
 #if !(UE_BUILD_SHIPPING || UE_BUILD_TEST)
@@ -43,6 +42,17 @@ namespace Chaos
 		float CurrentGearChangeTime = 0.0f;
 	};
 
+	struct CHAOSVEHICLESCORE_API FTransmissionOutputData : public FSimOutputData
+	{
+		virtual FSimOutputData* MakeNewData() override { return FTransmissionOutputData::MakeNew(); }
+		static FSimOutputData* MakeNew() { return new FTransmissionOutputData(); }
+
+		virtual void FillOutputState(const ISimulationModuleBase* SimModule) override;
+		virtual void Lerp(const FSimOutputData& InCurrent, const FSimOutputData& InNext, float Alpha) override;
+		virtual FString ToString() override;
+
+		int32 CurrentGear;
+	};
 
 	struct CHAOSVEHICLESCORE_API FTransmissionSettings
 	{
@@ -87,6 +97,7 @@ namespace Chaos
 	class CHAOSVEHICLESCORE_API FTransmissionSimModule : public FTorqueSimModule, public TSimModuleSettings<FTransmissionSettings>
 	{
 		friend FTransmissionSimModuleDatas;
+		friend FTransmissionOutputData;
 
 	public:
 
@@ -109,6 +120,11 @@ namespace Chaos
 			);
 		}
 
+		virtual FSimOutputData* GenerateOutputData() const override
+		{
+			return FTransmissionOutputData::MakeNew();
+		}
+
 		virtual eSimType GetSimType() const { return eSimType::Transmission; }
 
 		virtual const FString GetDebugName() const { return TEXT("Transmission"); }
diff --git a/Engine/Source/Runtime/Experimental/ChaosVehicles/ChaosVehiclesCore/Public/SimModule/WheelModule.h b/Engine/Source/Runtime/Experimental/ChaosVehicles/ChaosVehiclesCore/Public/SimModule/WheelModule.h
--- a/Engine/Source/Runtime/Experimental/ChaosVehicles/ChaosVehiclesCore/Public/SimModule/WheelModule.h
+++ b/Engine/Source/Runtime/Experimental/ChaosVehicles/ChaosVehiclesCore/Public/SimModule/WheelModule.h
@@ -16,6 +16,48 @@ namespace Chaos
 		Y	// Y forward
 	};
 
+	struct CHAOSVEHICLESCORE_API FWheelSimModuleDatas : public FTorqueSimModuleDatas
+	{
+#if !(UE_BUILD_SHIPPING || UE_BUILD_TEST)
+		FWheelSimModuleDatas(int NodeArrayIndex, const FString& InDebugString) : FTorqueSimModuleDatas(NodeArrayIndex, InDebugString) {}
+#else
+		FWheelSimModuleDatas(int NodeArrayIndex) : FTorqueSimModuleDatas(NodeArrayIndex) {}
+#endif
+
+		virtual eSimType GetType() override { return eSimType::Wheel; }
+
+		virtual void FillSimState(ISimulationModuleBase* SimModule) override
+		{
+			check(SimModule->GetSimType() == eSimType::Wheel);
+			FTorqueSimModuleDatas::FillSimState(SimModule);
+		}
+
+		virtual void FillNetState(const ISimulationModuleBase* SimModule) override
+		{
+			check(SimModule->GetSimType() == eSimType::Wheel);
+			FTorqueSimModuleDatas::FillNetState(SimModule);
+		}
+
+	};
+
+	struct CHAOSVEHICLESCORE_API FWheelOutputData : public FSimOutputData
+	{
+		virtual FSimOutputData* MakeNewData() override { return FWheelOutputData::MakeNew(); }
+		static FSimOutputData* MakeNew() { return new FWheelOutputData(); }
+
+		virtual void FillOutputState(const ISimulationModuleBase* SimModule) override;
+		virtual void Lerp(const FSimOutputData& InCurrent, const FSimOutputData& InNext, float Alpha) override;
+		virtual FString ToString() override;
+
+		bool bTouchingGround;
+		float ForceIntoSurface;
+		float SlipAngle;
+		float RPM;
+
+		//HitLocation
+		//PhysMaterial
+	};
+
 	struct CHAOSVEHICLESCORE_API FWheelSettings
 	{
 		FWheelSettings()
@@ -68,10 +110,27 @@ namespace Chaos
 
 	class CHAOSVEHICLESCORE_API FWheelSimModule : public FTorqueSimModule, public TSimModuleSettings<FWheelSettings>
 	{
+		friend FWheelOutputData;
+
 	public:
 
 		FWheelSimModule(const FWheelSettings& Settings);
 
+		virtual TSharedPtr<FModuleNetData> GenerateNetData(int SimArrayIndex) const
+		{
+			return MakeShared<FWheelSimModuleDatas>(
+				SimArrayIndex
+#if !(UE_BUILD_SHIPPING || UE_BUILD_TEST)
+				, GetDebugName()
+#endif			
+			);
+		}
+
+		virtual FSimOutputData* GenerateOutputData() const override
+		{
+			return FWheelOutputData::MakeNew();
+		}
+
 		virtual eSimType GetSimType() const { return eSimType::Wheel; }
 
 		virtual const FString GetDebugName() const { return TEXT("Wheel"); }
@@ -82,7 +141,7 @@ namespace Chaos
 
 		virtual void Animate(Chaos::FClusterUnionPhysicsProxy* Proxy) override;
 
-		virtual bool IsBehaviourType(eSimModuleTypeFlags InType) const override { return (InType & TorqueBased)||(InType & Velocity); }
+		virtual bool IsBehaviourType(eSimModuleTypeFlags InType) const override { return (InType & TorqueBased) || (InType & Velocity); }
 
 		void SetSuspensionSimTreeIndex(int IndexIn) { SuspensionSimTreeIndex = IndexIn; }
 		int GetSuspensionSimTreeIndex() const { return SuspensionSimTreeIndex; }
@@ -104,6 +163,9 @@ namespace Chaos
 		float MassPerWheel;
 		float SteerAngleDegrees;
 
+		// for output
+		bool bTouchingGround;
+		float SlipAngle;
 	};
 
 
