IAS - support for deleting local cache

#rnx

diff --git a/Engine/Source/Runtime/Experimental/IoStoreOnDemand/Private/IoStoreOnDemand.cpp b/Engine/Source/Runtime/Experimental/IoStoreOnDemand/Private/IoStoreOnDemand.cpp
--- a/Engine/Source/Runtime/Experimental/IoStoreOnDemand/Private/IoStoreOnDemand.cpp
+++ b/Engine/Source/Runtime/Experimental/IoStoreOnDemand/Private/IoStoreOnDemand.cpp
@@ -1502,6 +1502,14 @@ void FIoStoreOnDemandModule::SetEnabled(bool bInEnabled)
 	}
 }
 
+void FIoStoreOnDemandModule::AbandonCache()
+{
+	if (Backend.IsValid())
+	{
+		Backend->AbandonCache();
+	}
+}
+
 void FIoStoreOnDemandModule::ReportAnalytics(TArray<FAnalyticsEventAttribute>& OutAnalyticsArray) const
 {
 	if (Backend.IsValid())
@@ -1571,12 +1579,12 @@ void FIoStoreOnDemandModule::StartupModule()
 
 	Endpoint.EndpointType = EOnDemandEndpointType::CDN;
 
-	TSharedPtr<IIasCache> Cache;
+	TUniquePtr<IIasCache> Cache;
 	FIasCacheConfig CacheConfig = GetIasCacheConfig(CommandLine);
 	if (CacheConfig.DiskQuota > 0)
 	{
 		FString CacheDir = FPaths::ProjectPersistentDownloadDir();
-		Cache = MakeShareable(MakeIasCache(*CacheDir, CacheConfig).Release());
+		Cache = MakeIasCache(*CacheDir, CacheConfig);
 	}
 	if (!Cache.IsValid())
 	{
@@ -1584,7 +1592,7 @@ void FIoStoreOnDemandModule::StartupModule()
 			(CacheConfig.DiskQuota > 0) ? TEXT("init-fail") : TEXT("zero-quota"));
 	}
 
-	Backend = MakeOnDemandIoDispatcherBackend(Cache);
+	Backend = MakeOnDemandIoDispatcherBackend(MoveTemp(Cache));
 	Backend->Mount(Endpoint);
 	int32 BackendPriority = -10;
 #if !UE_BUILD_SHIPPING
diff --git a/Engine/Source/Runtime/Experimental/IoStoreOnDemand/Private/OnDemandIoDispatcherBackend.cpp b/Engine/Source/Runtime/Experimental/IoStoreOnDemand/Private/OnDemandIoDispatcherBackend.cpp
--- a/Engine/Source/Runtime/Experimental/IoStoreOnDemand/Private/OnDemandIoDispatcherBackend.cpp
+++ b/Engine/Source/Runtime/Experimental/IoStoreOnDemand/Private/OnDemandIoDispatcherBackend.cpp
@@ -104,6 +104,17 @@ static TAutoConsoleVariable<bool> CVar_IoReportAnalytics(
 	true,
 	TEXT("Enables reporting statics to the analytics system"));
 
+#if !UE_BUILD_SHIPPING
+static FAutoConsoleCommand CVar_IasAbandonCache(
+	TEXT("Ias.AbandonCache"),
+	TEXT("Abandon the local file cache"),
+	FConsoleCommandDelegate::CreateLambda([]()
+	{
+		FIoStoreOnDemandModule& Module = FModuleManager::Get().GetModuleChecked<FIoStoreOnDemandModule>("IoStoreOnDemand");
+		Module.AbandonCache();
+	})
+);
+#endif //!UE_BUILD_SHIPPING
 ///////////////////////////////////////////////////////////////////////////////
 #if !UE_BUILD_SHIPPING
 static void LatencyTest(FStringView Url, FStringView Path)
@@ -755,6 +766,7 @@ struct FBackendStatus
 		HttpEnabled					= (1 << 1),
 		HttpError					= (1 << 2),
 		HttpBulkOptionalDisabled	= (1 << 3),
+		AbandonCache				= (1 << 4),
 	};
 
 	bool IsHttpEnabled() const
@@ -792,6 +804,11 @@ struct FBackendStatus
 		return (CurrentFlags & uint8(EFlags::CacheEnabled)) && !IsHttpEnabled(CurrentFlags);
 	}
 
+	bool ShouldAbandonCache() const
+	{
+		return HasAnyFlags(EFlags::AbandonCache);
+	}
+
 	void SetHttpEnabled(bool bEnabled)
 	{
 		AddOrRemoveFlags(EFlags::HttpEnabled, bEnabled, TEXT("HTTP streaming enabled"));
@@ -813,6 +830,11 @@ struct FBackendStatus
 		AddOrRemoveFlags(EFlags::HttpError, bError, TEXT("HTTP streaming error"));
 	}
 
+	void SetAbandonCache(bool bAbandon)
+	{
+		AddOrRemoveFlags(EFlags::AbandonCache, bAbandon, TEXT("Abandon cache"));
+	}
+
 private:
 	static bool IsHttpEnabled(uint8 FlagsToTest)
 	{
@@ -1045,7 +1067,7 @@ class FOnDemandIoBackend final
 	};
 public:
 
-	FOnDemandIoBackend(TSharedPtr<IIasCache> Cache);
+	FOnDemandIoBackend(TUniquePtr<IIasCache>&& InCache);
 	virtual ~FOnDemandIoBackend();
 
 	// I/O dispatcher backend
@@ -1065,6 +1087,7 @@ public:
 	virtual void Mount(const FOnDemandEndpoint& Endpoint) override;
 	virtual void SetBulkOptionalEnabled(bool bEnabled) override;
 	virtual void SetEnabled(bool bEnabled) override;
+	virtual void AbandonCache() override;
 	virtual void ReportAnalytics(TArray<FAnalyticsEventAttribute>& OutAnalyticsArray) const override;
 
 	// Runnable
@@ -1079,7 +1102,7 @@ private:
 	void AddDeferredTocs();
 	void ProcessHttpRequests(FOnDemandHttpClient* HttpClient, FBitWindow& HttpErrors, int32 MaxConcurrentRequests);
 
-	TSharedPtr<IIasCache> Cache;
+	TUniquePtr<IIasCache> Cache;
 	TUniquePtr<FOnDemandIoStore> IoStore;
 	TSharedPtr<const FIoDispatcherBackendContext> BackendContext;
 	TUniquePtr<FRunnableThread> BackendThread;
@@ -1093,12 +1116,13 @@ private:
 	FAvailableEps AvailableEps;
 	FString DistributionUrl;
 	mutable FRWLock Lock;
+	std::atomic_uint32_t InflightCacheRequestCount{0};
 	std::atomic_bool bStopRequested{false};
 };
 
 ///////////////////////////////////////////////////////////////////////////////
-FOnDemandIoBackend::FOnDemandIoBackend(TSharedPtr<IIasCache> InCache)
-	: Cache(InCache)
+FOnDemandIoBackend::FOnDemandIoBackend(TUniquePtr<IIasCache>&& InCache)
+	: Cache(MoveTemp(InCache))
 {
 	IoStore = MakeUnique<FOnDemandIoStore>();
 	BackendStatus.SetHttpEnabled(true);
@@ -1262,6 +1286,11 @@ void FOnDemandIoBackend::CompleteRequest(FChunkRequest* ChunkRequest)
 	}
 
 	ChunkRequests.Release(ChunkRequest);
+
+	if (BackendStatus.ShouldAbandonCache() && InflightCacheRequestCount.load(std::memory_order_relaxed) == 0)
+	{
+		TickBackendEvent->Trigger();
+	}
 }
 
 bool FOnDemandIoBackend::Resolve(FIoRequestImpl* Request)
@@ -1304,6 +1333,7 @@ bool FOnDemandIoBackend::Resolve(FIoRequestImpl* Request)
 	if (Cache.IsValid())
 	{
 		//TODO: Pass priority to cache
+		InflightCacheRequestCount.fetch_add(1, std::memory_order_relaxed);
 		ChunkRequest->CacheTask = Cache->Get(ChunkRequest->Params.ChunkKey, FIoReadOptions(), &ChunkRequest->CancellationToken);
 	}
 
@@ -1314,6 +1344,7 @@ bool FOnDemandIoBackend::Resolve(FIoRequestImpl* Request)
 		TRACE_CPUPROFILER_EVENT_SCOPE(FOnDemandIoBackend::CompleteOrEnqueueHttpRequest);
 		if (ChunkRequest->CacheTask.IsValid())
 		{
+			InflightCacheRequestCount.fetch_sub(1, std::memory_order_relaxed);
 			if (TIoStatusOr<FIoBuffer> Status = ChunkRequest->CacheTask.GetResult(); Status.IsOk())
 			{
 				ChunkRequest->Chunk = Status.ConsumeValueOrDie();
@@ -1523,6 +1554,12 @@ void FOnDemandIoBackend::SetEnabled(bool bEnabled)
 	BackendStatus.SetHttpEnabled(bEnabled);
 }
 
+void FOnDemandIoBackend::AbandonCache()
+{
+	BackendStatus.SetCacheEnabled(false);
+	BackendStatus.SetAbandonCache(true);
+}
+
 void FOnDemandIoBackend::ReportAnalytics(TArray<FAnalyticsEventAttribute>& OutAnalyticsArray) const
 {
 	if (!CVar_IoReportAnalytics.GetValueOnAnyThread())
@@ -1695,6 +1732,16 @@ uint32 FOnDemandIoBackend::Run()
 					AddDeferredTocs();
 				}
 			}
+			if (BackendStatus.ShouldAbandonCache())
+			{
+				BackendStatus.SetAbandonCache(false);
+				check(BackendStatus.IsCacheEnabled() == false);
+				if (Cache.IsValid())
+				{
+					UE_LOG(LogIas, Log, TEXT("Abandoning cache, local file cache is no longer available"));
+					Cache.Release()->Abandon(); // Will delete its self
+				}
+			}
 			TickBackendEvent->Wait(WaitTime);
 		}
 	}
@@ -1702,9 +1749,9 @@ uint32 FOnDemandIoBackend::Run()
 	return 0;
 }
 
-TSharedPtr<IOnDemandIoDispatcherBackend> MakeOnDemandIoDispatcherBackend(TSharedPtr<IIasCache> Cache)
+TSharedPtr<IOnDemandIoDispatcherBackend> MakeOnDemandIoDispatcherBackend(TUniquePtr<IIasCache>&& Cache)
 {
-	return MakeShareable<IOnDemandIoDispatcherBackend>(new FOnDemandIoBackend(Cache));
+	return MakeShareable<IOnDemandIoDispatcherBackend>(new FOnDemandIoBackend(MoveTemp(Cache)));
 }
 
 } // namespace UE::IO::IAS
diff --git a/Engine/Source/Runtime/Experimental/IoStoreOnDemand/Private/OnDemandIoDispatcherBackend.h b/Engine/Source/Runtime/Experimental/IoStoreOnDemand/Private/OnDemandIoDispatcherBackend.h
--- a/Engine/Source/Runtime/Experimental/IoStoreOnDemand/Private/OnDemandIoDispatcherBackend.h
+++ b/Engine/Source/Runtime/Experimental/IoStoreOnDemand/Private/OnDemandIoDispatcherBackend.h
@@ -3,7 +3,7 @@
 #pragma once
 
 #include "IO/IoDispatcherBackend.h"
-#include "Templates/SharedPointer.h"
+#include "Templates/UniquePtr.h"
 
 class IIasCache;
 struct FAnalyticsEventAttribute;
@@ -39,9 +39,10 @@ public:
 	virtual void Mount(const FOnDemandEndpoint& Endpoint) = 0;
 	virtual void SetBulkOptionalEnabled(bool bInEnabled) = 0;
 	virtual void SetEnabled(bool bInEnabled) = 0;
+	virtual void AbandonCache() = 0;
 	virtual void ReportAnalytics(TArray<FAnalyticsEventAttribute>& OutAnalyticsArray) const = 0;
 };
 
-TSharedPtr<IOnDemandIoDispatcherBackend> MakeOnDemandIoDispatcherBackend(TSharedPtr<IIasCache> Cache);
+TSharedPtr<IOnDemandIoDispatcherBackend> MakeOnDemandIoDispatcherBackend(TUniquePtr<IIasCache>&& Cache);
 
 } // namespace UE::IO::IAS
diff --git a/Engine/Source/Runtime/Experimental/IoStoreOnDemand/Public/IO/IoStoreOnDemand.h b/Engine/Source/Runtime/Experimental/IoStoreOnDemand/Public/IO/IoStoreOnDemand.h
--- a/Engine/Source/Runtime/Experimental/IoStoreOnDemand/Public/IO/IoStoreOnDemand.h
+++ b/Engine/Source/Runtime/Experimental/IoStoreOnDemand/Public/IO/IoStoreOnDemand.h
@@ -203,6 +203,7 @@ private:
 public:
 	UE_API void SetBulkOptionalEnabled(bool bInEnabled);
 	UE_API void SetEnabled(bool bInEnabled);
+	UE_API void AbandonCache();
 
 	UE_API void ReportAnalytics(TArray<FAnalyticsEventAttribute>& OutAnalyticsArray) const;
 
