Fix missing include in on demand I/O backend

#rnx

diff --git a/Engine/Source/Runtime/Experimental/IoStoreOnDemand/Private/OnDemandIoDispatcherBackend.cpp b/Engine/Source/Runtime/Experimental/IoStoreOnDemand/Private/OnDemandIoDispatcherBackend.cpp
--- a/Engine/Source/Runtime/Experimental/IoStoreOnDemand/Private/OnDemandIoDispatcherBackend.cpp
+++ b/Engine/Source/Runtime/Experimental/IoStoreOnDemand/Private/OnDemandIoDispatcherBackend.cpp
@@ -39,6 +39,10 @@
 
 #include <atomic>
 
+#if !UE_BUILD_SHIPPING
+#include "Modules/ModuleManager.h"
+#endif 
+
 namespace UE::IO::IAS
 {
 
