Ias: Don't decrement GHttpPendingCount OnHttpRetry to avoid negative numbers (fallout from 27635427 which stopped putting retry items back in the queue).

#rnx

diff --git a/Engine/Source/Runtime/Experimental/IoStoreOnDemand/Private/Statistics.cpp b/Engine/Source/Runtime/Experimental/IoStoreOnDemand/Private/Statistics.cpp
--- a/Engine/Source/Runtime/Experimental/IoStoreOnDemand/Private/Statistics.cpp
+++ b/Engine/Source/Runtime/Experimental/IoStoreOnDemand/Private/Statistics.cpp
@@ -457,7 +457,6 @@ void FOnDemandIoBackendStats::OnHttpGet(uint64 SizeBytes, uint64 DurationMs)
 
 void FOnDemandIoBackendStats::OnHttpRetry()
 {
-	GHttpPendingCount.Add(-1);
 	GHttpInflightCount.Add(-1);
 	GHttpRetryCount.Add(1);
 }
