Ias: Only block on http when the client is sufficiently saturated (2*NumConcurrentConnections).
When unsatured, wake up every 17 ms to check for new requests to be issued.

#rnx

diff --git a/Engine/Source/Runtime/Experimental/IoStoreOnDemand/Private/OnDemandHttpClient.cpp b/Engine/Source/Runtime/Experimental/IoStoreOnDemand/Private/OnDemandHttpClient.cpp
--- a/Engine/Source/Runtime/Experimental/IoStoreOnDemand/Private/OnDemandHttpClient.cpp
+++ b/Engine/Source/Runtime/Experimental/IoStoreOnDemand/Private/OnDemandHttpClient.cpp
@@ -18,7 +18,7 @@ static FAutoConsoleVariableRef CVar_GIasHttpRateLimitKiBPerSecond(
 	TEXT("Http throttle limit in KiBPerSecond")
 );
 
-int32 GIasHttpPollTimeoutMs = 0;
+int32 GIasHttpPollTimeoutMs = 17;
 static FAutoConsoleVariableRef CVar_GIasHttpPollTimeoutMs(
 	TEXT("ias.HttpPollTimeoutMs"),
 	GIasHttpPollTimeoutMs,
diff --git a/Engine/Source/Runtime/Experimental/IoStoreOnDemand/Private/OnDemandIoDispatcherBackend.cpp b/Engine/Source/Runtime/Experimental/IoStoreOnDemand/Private/OnDemandIoDispatcherBackend.cpp
--- a/Engine/Source/Runtime/Experimental/IoStoreOnDemand/Private/OnDemandIoDispatcherBackend.cpp
+++ b/Engine/Source/Runtime/Experimental/IoStoreOnDemand/Private/OnDemandIoDispatcherBackend.cpp
@@ -1670,7 +1670,7 @@ void FOnDemandIoBackend::ProcessHttpRequests(FOnDemandHttpClient* HttpClient, FB
 				TRACE_CPUPROFILER_EVENT_SCOPE(FOnDemandIoBackend::TickHttpSaturated);
 				while (NumConcurrentRequests >= MaxConcurrentRequests)
 				{
-					HttpClient->Tick(true);
+					HttpClient->Tick(/*Block*/true);
 				}
 			}
 
@@ -1683,7 +1683,7 @@ void FOnDemandIoBackend::ProcessHttpRequests(FOnDemandHttpClient* HttpClient, FB
 		{
 			// Keep processing pending connections until all requests are completed or a new one is issued
 			TRACE_CPUPROFILER_EVENT_SCOPE(FOnDemandIoBackend::TickHttp);
-			while (HttpClient->Tick(true) && !NextChunkRequest)
+			while (!NextChunkRequest && HttpClient->Tick(/*Block*/false))
 			{
 				NextChunkRequest = HttpRequests.Dequeue();
 			}
@@ -1710,7 +1710,7 @@ uint32 FOnDemandIoBackend::Run()
 	while (!bStopRequested)
 	{
 		// Process HTTP request(s) even if the client is invalid to ensure enqueued request(s) gets completed.
-		ProcessHttpRequests(HttpClient.Get(), HttpErrors, GIasMaxHttpConnectionCount);
+		ProcessHttpRequests(HttpClient.Get(), HttpErrors, FMath::Min(2 * GIasMaxHttpConnectionCount, 64));
 
 		if (!bStopRequested)
 		{
