Ias: Make sure to never overwrite the NextChunkRequest pointer when ticking the http client.
The pointer could have been set in the tick loop for retrying.
Fix for cl 27827566.

#rnx

diff --git a/Engine/Source/Runtime/Experimental/IoStoreOnDemand/Private/OnDemandIoDispatcherBackend.cpp b/Engine/Source/Runtime/Experimental/IoStoreOnDemand/Private/OnDemandIoDispatcherBackend.cpp
--- a/Engine/Source/Runtime/Experimental/IoStoreOnDemand/Private/OnDemandIoDispatcherBackend.cpp
+++ b/Engine/Source/Runtime/Experimental/IoStoreOnDemand/Private/OnDemandIoDispatcherBackend.cpp
@@ -1683,9 +1683,16 @@ void FOnDemandIoBackend::ProcessHttpRequests(FOnDemandHttpClient* HttpClient, FB
 		{
 			// Keep processing pending connections until all requests are completed or a new one is issued
 			TRACE_CPUPROFILER_EVENT_SCOPE(FOnDemandIoBackend::TickHttp);
-			while (!NextChunkRequest && HttpClient->Tick(/*Block*/false))
+			while (HttpClient->Tick(/*Block*/false))
 			{
-				NextChunkRequest = HttpRequests.Dequeue();
+				if (!NextChunkRequest)
+				{
+					NextChunkRequest = HttpRequests.Dequeue();
+				}
+				if (NextChunkRequest)
+				{
+					break;
+				}
 			}
 		}
 	} 
