incremental addition to/removal from cluster union implicit object geometry on both the GT and PT

- removed the cvar (ClusterUnion.IncrementalUnionBuild) to disable incrementally building the implicit object union on the GT - only incremental building is supported now
- Deprecate UClusterUnionComponent::ForceRebuildGTParticleGeometry now that it's no longer used internally
- Added the EUpdateClusterUnionPropertiesFlags::ForceGenerateGeometry and EUpdateClusterUnionPropertiesFlags::IncrementalGenerateGeometry flags to mark when we want to update the cluster union geometry
- Replace Chaos::UpdateShapesDatas with Chaos::ModifyAdditionOfChildrenToClusterUnionGeometry, Chaos::RemoveParticlesFromClusterUnionGeometry, Chaos::TransferClusterUnionShapeData to provide context agnostic ways of modifying cluster union geometry.
- Introduce FClusterUnionManager::FlushIncrementalGeometryOperations which processes the pending geometry updates and applies them to the cluster union implicit object union
- Added TGeometryParticleHandle::MergeGeometry, TGeometryParticleHandle::ShapesArray (non-const), TGeometryParticleHandle::RemoveShape, and TGeometryParticleHandle::RemoveShapeAtIndex to parallel their GT equivalents.
- Added a stat scope for UpdateShapesArrayFromGeometryImpl and made sure it's modifying the shape's shape index.

diff --git a/Engine/Source/Runtime/Engine/Classes/PhysicsEngine/ClusterUnionComponent.h b/Engine/Source/Runtime/Engine/Classes/PhysicsEngine/ClusterUnionComponent.h
--- a/Engine/Source/Runtime/Engine/Classes/PhysicsEngine/ClusterUnionComponent.h
+++ b/Engine/Source/Runtime/Engine/Classes/PhysicsEngine/ClusterUnionComponent.h
@@ -182,6 +182,7 @@ public:
 
 	// Force a rebuild of the GT geometry. This needs to happen immediately when we add/remove on the GT so that the SQ is up to date
 	// and doesn't need to wait for the next OnSyncBodies.
+	UE_DEPRECATED(5.4, "Cluster unions now incrementally update their geometry. A force rebuild should never be needed - this function will no longer book-keep properly either.")
 	ENGINE_API void ForceRebuildGTParticleGeometry();
 
 	ENGINE_API const FSpatialAcceleration* GetSpatialAcceleration() const;
@@ -276,10 +277,10 @@ private:
 	ENGINE_API bool IsAuthority() const;
 
 	// Merge all the physics objects geometries into the cluster union
-	ENGINE_API void AddGTParticleGeometry(const TArray<Chaos::FPhysicsObjectHandle>& PhysicsObjects) const;
+	ENGINE_API void AddGTParticleGeometry(const TArray<Chaos::FPhysicsObjectHandle>& PhysicsObjects);
 
 	// Remove all the physics objects geometries from the cluster union
-	ENGINE_API void RemoveGTParticleGeometry(const TSet<Chaos::FPhysicsObjectHandle>& PhysicsObjects) const;
+	ENGINE_API void RemoveGTParticleGeometry(const TSet<Chaos::FPhysicsObjectHandle>& PhysicsObjects);
 
 	//~ Begin UActorComponent Interface
 public:
diff --git a/Engine/Source/Runtime/Engine/Private/PhysicsEngine/ClusterUnionComponent.cpp b/Engine/Source/Runtime/Engine/Private/PhysicsEngine/ClusterUnionComponent.cpp
--- a/Engine/Source/Runtime/Engine/Private/PhysicsEngine/ClusterUnionComponent.cpp
+++ b/Engine/Source/Runtime/Engine/Private/PhysicsEngine/ClusterUnionComponent.cpp
@@ -25,9 +25,6 @@ namespace
 	bool bUseClusterUnionAccelerationStructure = true;
 	FAutoConsoleVariableRef CVarUseClusterUnionAccelerationStructure(TEXT("ClusterUnion.UseAccelerationStructure"), bUseClusterUnionAccelerationStructure, TEXT("Whether component level sweeps and overlaps against cluster unions should use an acceleration structure instead."));
 
-	bool bIncrementalUnionBuild = false;
-	FAutoConsoleVariableRef CVarIncrementalUnionBuild(TEXT("ClusterUnion.IncrementalUnionBuild"), bIncrementalUnionBuild , TEXT("Cvar to incrementally build the union."));
-	
 	template<typename PayloadType>
 	struct TClusterUnionAABBTreeStorageTraits
 	{
@@ -194,22 +191,15 @@ void UClusterUnionComponent::AddComponentToCluster(UPrimitiveComponent* InCompon
 	// This technically is giving us a false sense of security because when we get this notification, the primitive component will already have no physics proxy on it so we can't actually grab physics objects/particles from it.
 	InComponent->OnComponentPhysicsStateChanged.AddUniqueDynamic(this, &UClusterUnionComponent::HandleComponentPhysicsStateChangePostAddIntoClusterUnion);
 
+	if (bRebuildGeometry)
+	{
+		AddGTParticleGeometry(Objects);
+	}
+
 	PendingComponentSync.Add(InComponent, PendingData);
 
 	PhysicsProxy->AddPhysicsObjects_External(Objects);
 	ForceSetChildToParent(InComponent, PendingData.BoneIds, ChildToParents);
-
-	if (bRebuildGeometry)
-	{
-		if(!bIncrementalUnionBuild)
-		{
-			ForceRebuildGTParticleGeometry();
-		}
-		else
-		{
-			AddGTParticleGeometry(Objects);
-		}
-	}
 }
 
 void UClusterUnionComponent::RemoveComponentFromCluster(UPrimitiveComponent* InComponent)
@@ -280,20 +270,12 @@ void UClusterUnionComponent::RemoveComponentFromCluster(UPrimitiveComponent* InC
 			}
 			ComponentData->CachedAccelerationPayloads.Reset();
 		}
-	}
 
-	// If PhysicsObjectsToRemove is empty, it either means that BoneIds is empty OR the component's physics state is already destroyed.
-	// In the case of the latter, we rely on the physics thread to cleanup the cluster union manager properly and to sync back.
-	if (!PhysicsObjectsToRemove.IsEmpty())
-	{
-		PhysicsProxy->RemovePhysicsObjects_External(PhysicsObjectsToRemove);
-		
-		if(!bIncrementalUnionBuild)
-        {
-			ForceRebuildGTParticleGeometry();
-		}
-		else
+		// If PhysicsObjectsToRemove is empty, it either means that BoneIds is empty OR the component's physics state is already destroyed.
+		// In the case of the latter, we rely on the physics thread to cleanup the cluster union manager properly and to sync back.
+		if (!PhysicsObjectsToRemove.IsEmpty())
 		{
+			PhysicsProxy->RemovePhysicsObjects_External(PhysicsObjectsToRemove);
 			RemoveGTParticleGeometry(PhysicsObjectsToRemove);
 		}
 	}
@@ -378,6 +360,14 @@ void UClusterUnionComponent::RemoveComponentBonesFromCluster(UPrimitiveComponent
 			PhysicsObjectsToRemove = TSet<Chaos::FPhysicsObjectHandle>{ GetAllPhysicsObjectsById(InComponent, RemoveBoneIdsSet.Array()) };
 		}
 
+		// If PhysicsObjectsToRemove is empty, it either means that BoneIds is empty OR the component's physics state is already destroyed.
+		// In the case of the latter, we rely on the physics thread to cleanup the cluster union manager properly and to sync back.
+		if (!PhysicsObjectsToRemove.IsEmpty())
+		{
+			PhysicsProxy->RemovePhysicsObjects_External(PhysicsObjectsToRemove);
+			RemoveGTParticleGeometry(PhysicsObjectsToRemove);
+		}
+
 		if (ComponentData->BoneIds.IsEmpty())
 		{
 			// We need to mark the replicated proxy as pending deletion.
@@ -395,18 +385,11 @@ void UClusterUnionComponent::RemoveComponentBonesFromCluster(UPrimitiveComponent
 
 			PerComponentData.Remove(InComponent);
 		}
-	}
 
-	// If PhysicsObjectsToRemove is empty, it either means that BoneIds is empty OR the component's physics state is already destroyed.
-	// In the case of the latter, we rely on the physics thread to cleanup the cluster union manager properly and to sync back.
-	if (!PhysicsObjectsToRemove.IsEmpty())
-	{
-		PhysicsProxy->RemovePhysicsObjects_External(PhysicsObjectsToRemove);
-		ForceRebuildGTParticleGeometry();
 	}
 }
 
-void UClusterUnionComponent::AddGTParticleGeometry(const TArray<Chaos::FPhysicsObjectHandle>& PhysicsObjects) const
+void UClusterUnionComponent::AddGTParticleGeometry(const TArray<Chaos::FPhysicsObjectHandle>& PhysicsObjects)
 {
 	if (!PhysicsProxy)
 	{
@@ -420,8 +403,15 @@ void UClusterUnionComponent::AddGTParticleGeometry(const TArray<Chaos::FPhysicsO
 	TArray<Chaos::FPBDRigidParticle*> ShapeParticles;
 	
 	const FTransform ClusterWorldTM = GetComponentTransform();
-	for (Chaos::FPBDRigidParticle* RigidParticle : Interface->GetAllRigidParticles(PhysicsObjects))
+	for (Chaos::FPhysicsObjectHandle PhysicsObject : PhysicsObjects)
 	{
+		const int32 BoneId = Chaos::FPhysicsObjectInterface::GetId(PhysicsObject);
+		if (BoneId == INDEX_NONE)
+		{
+			continue;
+		}
+
+		Chaos::FPBDRigidParticle* RigidParticle = Interface->GetRigidParticle(PhysicsObject);
 		if (RigidParticle && RigidParticle->GetGeometry())
 		{
 			const FTransform ChildWorldTM{ RigidParticle->R(), RigidParticle->X() };
@@ -432,7 +422,8 @@ void UClusterUnionComponent::AddGTParticleGeometry(const TArray<Chaos::FPhysicsO
 	}
 	if(!ImplicitObjects.IsEmpty() && PhysicsProxy->GetParticle_External())
 	{
-		if(PhysicsProxy->GetParticle_External()->GetGeometry() == nullptr)
+		const Chaos::FImplicitObjectRef ExistingGeometry = PhysicsProxy->GetParticle_External()->GetGeometry();
+		if(ExistingGeometry == nullptr || ExistingGeometry->GetType() != Chaos::ImplicitObjectType::Union)
 		{
 			Chaos::FImplicitObjectUnion* NewGeometry = ImplicitObjects.IsEmpty() ? new Chaos::FImplicitObjectUnionClustered() : new Chaos::FImplicitObjectUnion(MoveTemp(ImplicitObjects));
 			NewGeometry->SetAllowBVH(true);
@@ -446,7 +437,7 @@ void UClusterUnionComponent::AddGTParticleGeometry(const TArray<Chaos::FPhysicsO
 	}
 }
 
-void UClusterUnionComponent::RemoveGTParticleGeometry(const TSet<Chaos::FPhysicsObjectHandle>& PhysicsObjects) const
+void UClusterUnionComponent::RemoveGTParticleGeometry(const TSet<Chaos::FPhysicsObjectHandle>& PhysicsObjects)
 {
 	if (!PhysicsProxy)
 	{
@@ -464,11 +455,19 @@ void UClusterUnionComponent::RemoveGTParticleGeometry(const TSet<Chaos::FPhysics
 	FLockedReadPhysicsObjectExternalInterface Interface = FPhysicsObjectExternalInterface::LockRead(GetWorld()->GetPhysicsScene());
 
 	TArray<Chaos::FPBDRigidParticle*> ShapeParticles;
-	for (Chaos::FPBDRigidParticle* RigidParticle : Interface->GetAllRigidParticles(ArrayObjects))
+	for (Chaos::FPhysicsObjectHandle PhysicsObject : PhysicsObjects)
 	{
+		const int32 BoneId = Chaos::FPhysicsObjectInterface::GetId(PhysicsObject);
+		if (BoneId == INDEX_NONE)
+		{
+			continue;
+		}
+
+		Chaos::FPBDRigidParticle* RigidParticle = Interface->GetRigidParticle(PhysicsObject);
 		if (RigidParticle && RigidParticle->GetGeometry())
 		{
 			ShapeParticles.Add(RigidParticle);
+
 		}
 	}
 	if(!ShapeParticles.IsEmpty() && PhysicsProxy->GetParticle_External())
@@ -477,7 +476,6 @@ void UClusterUnionComponent::RemoveGTParticleGeometry(const TSet<Chaos::FPhysics
 	}
 }
 
-
 void UClusterUnionComponent::ForceRebuildGTParticleGeometry()
 {
 	if (!PhysicsProxy)
@@ -824,19 +822,36 @@ void UClusterUnionComponent::SyncClusterUnionFromProxy()
 	}
 	
 	const Chaos::FClusterUnionSyncedData& FullData = PhysicsProxy->GetSyncedData_External();
+
+	TArray<Chaos::FPBDRigidParticle*> ChildParticles;
+	ChildParticles.Reserve(FullData.ChildParticles.Num());
+
 	// Note that at the UClusterUnionComponent level we really only want to be dealing with components.
 	// Hence why we need to modify each of the particles that we synced from the game thread into a
 	// component + bone id combination for identification. 
 	TMap<TObjectKey<UPrimitiveComponent>, TMap<int32, FTransform>> MappedData;
-	for (const Chaos::FClusterUnionChildData& ChildData : FullData.ChildParticles)
+
 	{
-		// Using the scene's proxy to component mapping let's us detect a component physics state was destroyed.
-		if (UPrimitiveComponent* Component = Scene->GetOwningComponent<UPrimitiveComponent>(ChildData.Proxy))
+		FLockedReadPhysicsObjectExternalInterface Interface = FPhysicsObjectExternalInterface::LockRead(GetWorld()->GetPhysicsScene());
+		for (const Chaos::FClusterUnionChildData& ChildData : FullData.ChildParticles)
 		{
-			MappedData.FindOrAdd(Component).Add(ChildData.BoneId, ChildData.ChildToParent);
+			// Using the scene's proxy to component mapping let's us detect a component physics state was destroyed.
+			if (UPrimitiveComponent* Component = Scene->GetOwningComponent<UPrimitiveComponent>(ChildData.Proxy))
+			{
+				MappedData.FindOrAdd(Component).Add(ChildData.BoneId, ChildData.ChildToParent);
+	
+				Chaos::FPhysicsObjectHandle Handle = Component->GetPhysicsObjectById(ChildData.BoneId);
+				ChildParticles.Add(Interface->GetRigidParticle(Handle));
+			}
+			else
+			{
+				ChildParticles.Add(nullptr);
+			}
 		}
 	}
 
+	PhysicsProxy->ForceSetGeometryChildParticles_External(MoveTemp(ChildParticles));
+
 	// We need to handle any additions, deletions, and modifications to any child in the cluster union here.
 	// If a component lives in MappedData but not in PerComponentData, new component!
 	// If a component lives in both, then it's a modified component.
diff --git a/Engine/Source/Runtime/Experimental/Chaos/Private/Chaos/ClusterUnionManager.cpp b/Engine/Source/Runtime/Experimental/Chaos/Private/Chaos/ClusterUnionManager.cpp
--- a/Engine/Source/Runtime/Experimental/Chaos/Private/Chaos/ClusterUnionManager.cpp
+++ b/Engine/Source/Runtime/Experimental/Chaos/Private/Chaos/ClusterUnionManager.cpp
@@ -241,6 +241,31 @@ namespace Chaos
 		}
 	}
 
+	const TArray<FPBDRigidParticleHandle*>& FClusterUnion::GetPendingGeometryOperationParticles(EClusterUnionGeometryOperation Op) const
+	{
+		static TArray<FPBDRigidParticleHandle*> EmptyArray;
+		if (const TArray<FPBDRigidParticleHandle*>* Results = PendingGeometryOperations.Find(Op))
+		{
+			return *Results;
+		}
+		return EmptyArray;
+	}
+
+	void FClusterUnion::AddPendingGeometryOperation(EClusterUnionGeometryOperation Op, FPBDRigidParticleHandle* Particle)
+	{
+		PendingGeometryOperations.FindOrAdd(Op).AddUnique(Particle);
+	}
+
+	void FClusterUnion::ClearAllPendingGeometryOperations()
+	{
+		PendingGeometryOperations.Empty();
+	}
+
+	void FClusterUnion::ClearPendingGeometryOperations(EClusterUnionGeometryOperation Op)
+	{
+		PendingGeometryOperations.Remove(Op);
+	}
+
 	FClusterUnionManager::FClusterUnionManager(FRigidClustering& InClustering, FPBDRigidsEvolutionGBF& InEvolution)
 		: MClustering(InClustering)
 		, MEvolution(InEvolution)
@@ -257,7 +282,7 @@ namespace Chaos
 		FClusterUnion NewUnion;
 		NewUnion.InternalIndex = NewIndex;
 		NewUnion.ExplicitIndex = ClusterUnionParameters.ExplicitIndex;
-		NewUnion.Geometry = ForceRecreateClusterUnionGeometry(NewUnion);
+		NewUnion.Geometry = MakeImplicitObjectPtr<FImplicitObjectUnionClustered>();
 		NewUnion.InternalCluster = MClustering.CreateClusterParticle(-NewIndex, {}, Parameters, NewUnion.Geometry, nullptr, ClusterUnionParameters.UniqueIndex);
 		NewUnion.Parameters = Parameters;
 		NewUnion.ClusterUnionParameters = ClusterUnionParameters;
@@ -562,6 +587,7 @@ namespace Chaos
 			if (!bIsNewCluster)
 			{
 				Cluster->PendingConnectivityOperations.Add({ Particle, EClusterUnionConnectivityOperation::Add });
+				Cluster->AddPendingGeometryOperation(EClusterUnionGeometryOperation::Add, Particle);
 			}
 
 			if (!Cluster->ChildProperties.Contains(Particle))
@@ -619,7 +645,9 @@ namespace Chaos
 			Cluster->InternalCluster->SetIsAnchored(true);
 		}
 
-		RequestDeferredClusterPropertiesUpdate(ClusterIndex, bIsNewCluster ? EUpdateClusterUnionPropertiesFlags::All : EUpdateClusterUnionPropertiesFlags::IncrementalGenerateConnectionGraph);
+		constexpr EUpdateClusterUnionPropertiesFlags DefaultIncrementalFlags = EUpdateClusterUnionPropertiesFlags::IncrementalGenerateConnectionGraph
+			| EUpdateClusterUnionPropertiesFlags::IncrementalGenerateGeometry;
+		RequestDeferredClusterPropertiesUpdate(ClusterIndex, bIsNewCluster ? EUpdateClusterUnionPropertiesFlags::All : DefaultIncrementalFlags);
 
 		if (OldProxy)
 		{
@@ -731,6 +759,7 @@ namespace Chaos
 				}
 
 				Cluster->PendingConnectivityOperations.Add({ Handle, EClusterUnionConnectivityOperation::Remove });
+				Cluster->AddPendingGeometryOperation(EClusterUnionGeometryOperation::Remove, Handle);
 
 				MEvolution.GetParticles().MarkTransientDirtyParticle(Handle);
 			}
@@ -741,7 +770,11 @@ namespace Chaos
 		{
 			const int32 ParticleIndex = ParticleIndicesToRemove[Index];
 			Cluster->ChildProperties.Remove(Cluster->ChildParticles[ParticleIndex]);
-			Cluster->ChildParticles.RemoveAtSwap(ParticleIndex);
+
+			// This can't be RemoveAtSwap otherwise there will be a mismatch between the index of a particle
+			// and the index of its corresponding shape in the cluster union's shape array. There is currently
+			// an assumption that the two will always match each other.
+			Cluster->ChildParticles.RemoveAt(ParticleIndex);
 		}
 
 		MClustering.RemoveParticlesFromCluster(Cluster->InternalCluster, ParticleSet.Array());
@@ -750,7 +783,9 @@ namespace Chaos
 		// This gets changed because we go through an internal initialization route when we update the cluster union particle's properties.
 		Cluster->InternalCluster->SetPhysicsProxy(OldProxy);
 
-		constexpr EUpdateClusterUnionPropertiesFlags RemoveUpdateFlags = EUpdateClusterUnionPropertiesFlags::IncrementalGenerateConnectionGraph | EUpdateClusterUnionPropertiesFlags::UpdateKinematicProperties;
+		constexpr EUpdateClusterUnionPropertiesFlags RemoveUpdateFlags = EUpdateClusterUnionPropertiesFlags::IncrementalGenerateConnectionGraph
+			| EUpdateClusterUnionPropertiesFlags::UpdateKinematicProperties
+			| EUpdateClusterUnionPropertiesFlags::IncrementalGenerateGeometry;
 		switch (UpdateClusterPropertiesTiming)
 		{
 		case EClusterUnionOperationTiming::Immediate:
@@ -827,13 +862,30 @@ namespace Chaos
 		// The recreation of the geometry must happen after the call to UpdateClusterMassProperties.
 		// Creating the geometry requires knowing the relative frame between the parent cluster and the child clusters. The
 		// parent transform is not set properly for a new empty cluster until UpdateClusterMassProperties is called for the first time.
-		ClusterUnion.Geometry = ForceRecreateClusterUnionGeometry(ClusterUnion);
-		UpdateGeometry(ClusterUnion.InternalCluster, FullChildrenSet, MClustering.GetChildrenMap(), ClusterUnion.Geometry, ClusterUnion.Parameters);
-
-		const int32 NumSimpleShapes = ClusterUnion.Geometry->AsA<FImplicitObjectUnion>()->GetConvexes().Num();
-		UpdateShapesDatas(ClusterUnion.ChildParticles, ClusterUnion.InternalCluster->ShapesArray(),
-			ClusterUnion.ClusterUnionParameters.ActorId, ClusterUnion.ClusterUnionParameters.ComponentId, NumSimpleShapes);
+		if (EnumHasAnyFlags(Flags, EUpdateClusterUnionPropertiesFlags::ForceGenerateGeometry))
+		{
+			ClusterUnion.InternalCluster->SetGeometry(MakeImplicitObjectPtr<FImplicitObjectUnionClustered>());
+
+			ModifyAdditionOfChildrenToClusterUnionGeometry(
+				ClusterUnion.InternalCluster,
+				ClusterUnion.ChildParticles,
+				ClusterUnion.ClusterUnionParameters.ActorId,
+				ClusterUnion.ClusterUnionParameters.ComponentId,
+				[this, &ClusterUnion, &FullChildrenSet]()
+				{
+					ClusterUnion.Geometry = ForceRecreateClusterUnionGeometry(ClusterUnion);
+					UpdateGeometry(ClusterUnion.InternalCluster, FullChildrenSet, MClustering.GetChildrenMap(), ClusterUnion.Geometry, ClusterUnion.Parameters);
+				}
+			);
 
+			ClusterUnion.GeometryChildParticles = ClusterUnion.ChildParticles;
+			ClusterUnion.ClearAllPendingGeometryOperations();
+		}
+		else if (EnumHasAnyFlags(Flags, EUpdateClusterUnionPropertiesFlags::IncrementalGenerateGeometry))
+		{
+			FlushIncrementalGeometryOperations(ClusterUnion);
+		}
+		
 		if (ClusterUnion.bGenerateConnectivityEdges)
 		{
 			if(EnumHasAnyFlags(Flags, EUpdateClusterUnionPropertiesFlags::ForceGenerateConnectionGraph))
@@ -877,6 +929,96 @@ namespace Chaos
 		ClusterUnion.PendingConnectivityOperations.Empty();
 	}
 
+	DECLARE_CYCLE_STAT(TEXT("FClusterUnionManager::FlushIncrementalGeometryOperations"), STAT_FlushIncrementalGeometryOperations, STATGROUP_Chaos);
+	void FClusterUnionManager::FlushIncrementalGeometryOperations(FClusterUnion& ClusterUnion)
+	{
+		SCOPE_CYCLE_COUNTER(STAT_FlushIncrementalGeometryOperations);
+		check(ClusterUnion.Geometry != nullptr && ClusterUnion.Geometry->GetType() == ImplicitObjectType::Union);
+
+		const TArray<FPBDRigidParticleHandle*>& PendingGeometryAdditions = ClusterUnion.GetPendingGeometryOperationParticles(EClusterUnionGeometryOperation::Add);
+		if (!PendingGeometryAdditions.IsEmpty())
+		{
+			ModifyAdditionOfChildrenToClusterUnionGeometry(
+				ClusterUnion.InternalCluster,
+				PendingGeometryAdditions,
+				ClusterUnion.ClusterUnionParameters.ActorId,
+				ClusterUnion.ClusterUnionParameters.ComponentId,
+				[this, &ClusterUnion, &PendingGeometryAdditions]()
+				{
+					const FRigidTransform3 ClusterWorldTM(ClusterUnion.InternalCluster->X(), ClusterUnion.InternalCluster->R());
+
+					TArray<Chaos::FImplicitObjectPtr> Objects;
+					Objects.Reserve(PendingGeometryAdditions.Num());
+					ClusterUnion.GeometryChildParticles.Reserve(ClusterUnion.GeometryChildParticles.Num() + PendingGeometryAdditions.Num());
+
+					for (FPBDRigidParticleHandle* Child : PendingGeometryAdditions)
+					{
+						const FRigidTransform3 Frame = GetParticleRigidFrameInClusterUnion(Child, ClusterWorldTM);
+						if (Child->GetGeometry())
+						{
+							ClusterUnion.GeometryChildParticles.Add(Child);
+							Objects.Add(Chaos::FImplicitObjectPtr(CreateTransformGeometryForClusterUnion<EThreadContext::Internal>(Child, Frame)));
+						}
+					}
+
+					if (!Objects.IsEmpty())
+					{
+						ClusterUnion.InternalCluster->MergeGeometry(MoveTemp(Objects));
+					}
+				}
+			);
+
+			ClusterUnion.ClearPendingGeometryOperations(EClusterUnionGeometryOperation::Add);
+		}
+
+		const TArray<FPBDRigidParticleHandle*>& PendingGeometryRemovals = ClusterUnion.GetPendingGeometryOperationParticles(EClusterUnionGeometryOperation::Remove);
+		if (!PendingGeometryRemovals.IsEmpty())
+		{
+			RemoveParticlesFromClusterUnionGeometry(ClusterUnion.InternalCluster, PendingGeometryRemovals, ClusterUnion.GeometryChildParticles);
+			ClusterUnion.ClearPendingGeometryOperations(EClusterUnionGeometryOperation::Remove);
+		}
+
+		const TArray<FPBDRigidParticleHandle*>& PendingGeometryRefresh = ClusterUnion.GetPendingGeometryOperationParticles(EClusterUnionGeometryOperation::Refresh);
+		if (!PendingGeometryRefresh.IsEmpty())
+		{
+			// For each particle we need to find the corresponding shape.
+			// Note that by the time we get to handling PendimgGeometryRefresh, we can once again
+			// make the assumption that the number of children particles = number of shapes.
+			const FShapesArray& ShapesArray = ClusterUnion.InternalCluster->ShapesArray();
+			check(ClusterUnion.ChildParticles.Num() == ShapesArray.Num());
+
+			for (FPBDRigidParticleHandle* Particle : PendingGeometryRefresh)
+			{
+				if (!Particle)
+				{
+					continue;
+				}
+
+				const int32 Index = ClusterUnion.ChildParticles.Find(Particle);
+				if (Index == INDEX_NONE || !ShapesArray.IsValidIndex(Index))
+				{
+					continue;
+				}
+
+				const TUniquePtr<Chaos::FPerShapeData>& TemplateShape = Particle->ShapesArray()[0];
+				if (!TemplateShape)
+				{
+					continue;
+				}
+
+				TransferClusterUnionShapeData(
+					ShapesArray[Index],
+					Particle,
+					TemplateShape,
+					ClusterUnion.ClusterUnionParameters.ActorId,
+					ClusterUnion.ClusterUnionParameters.ComponentId
+				);
+			}
+
+			ClusterUnion.ClearPendingGeometryOperations(EClusterUnionGeometryOperation::Refresh);
+		}
+	}
+
 	DECLARE_CYCLE_STAT(TEXT("FClusterUnionManager::GetOrCreateClusterUnionIndexFromExplicitIndex"), STAT_GetOrCreateClusterUnionIndexFromExplicitIndex, STATGROUP_Chaos);
 	FClusterUnionIndex FClusterUnionManager::GetOrCreateClusterUnionIndexFromExplicitIndex(FClusterUnionExplicitIndex InIndex)
 	{
diff --git a/Engine/Source/Runtime/Experimental/Chaos/Private/Chaos/GeometryParticles.cpp b/Engine/Source/Runtime/Experimental/Chaos/Private/Chaos/GeometryParticles.cpp
--- a/Engine/Source/Runtime/Experimental/Chaos/Private/Chaos/GeometryParticles.cpp
+++ b/Engine/Source/Runtime/Experimental/Chaos/Private/Chaos/GeometryParticles.cpp
@@ -49,6 +49,13 @@ namespace Chaos
 		bIsSingleShape = true;
 	}
 
+	template <typename T, int d, EGeometryParticlesSimType SimType>
+	void TGeometryParticlesImp<T, d, SimType>::RemoveShapeAtIndex(const int32 ParticleIndex, const int32 Index)
+	{
+		check(MShapesArray.IsValidIndex(ParticleIndex));
+		MShapesArray[ParticleIndex].RemoveAt(Index);
+	}
+
 	template <typename T, int d, EGeometryParticlesSimType SimType>
 	void TGeometryParticlesImp<T, d, SimType>::UpdateShapesArray(const int32 Index)
 	{
diff --git a/Engine/Source/Runtime/Experimental/Chaos/Private/Chaos/ParticleHandle.cpp b/Engine/Source/Runtime/Experimental/Chaos/Private/Chaos/ParticleHandle.cpp
--- a/Engine/Source/Runtime/Experimental/Chaos/Private/Chaos/ParticleHandle.cpp
+++ b/Engine/Source/Runtime/Experimental/Chaos/Private/Chaos/ParticleHandle.cpp
@@ -83,6 +83,26 @@ namespace Chaos
 		}
 	}
 
+	template <typename T, int d, bool bPersistent>
+	void Chaos::TGeometryParticleHandleImp<T, d, bPersistent>::MergeGeometry(TArray<Chaos::FImplicitObjectPtr>&& Objects)
+	{
+		if (Objects.IsEmpty())
+		{
+			return;
+		}
+
+		const FImplicitObjectRef CurrentGeometry = GetGeometry();
+		ensure(CurrentGeometry != nullptr);
+		if (ensure(CurrentGeometry->GetType() == FImplicitObjectUnion::StaticType()))
+		{
+			FImplicitObjectUnion& Union = CurrentGeometry->GetObjectChecked<FImplicitObjectUnion>();
+			Union.Combine(Objects);
+
+			// Needed to update the shapes array.
+			SetGeometry(GeometryParticles->GetGeometry(ParticleIdx));
+		}
+	}
+
 	template <typename T, int d>
 	void Chaos::TGeometryParticle<T, d>::RemoveShape(FPerShapeData* InShape, bool bWakeTouching)
 	{
@@ -94,24 +114,89 @@ namespace Chaos
 		{
 			if (InShape == MShapesArray[Index].Get())
 			{
-				MShapesArray.RemoveAt(Index);
 				FoundIndex = Index;
 				break;
 			}
 		}
 
+		RemoveShapeAtIndex(FoundIndex);
+	}
+
+	template <typename T, int d>
+	void Chaos::TGeometryParticle<T, d>::RemoveShapeAtIndex(int32 InIndex)
+	{
+		// NOTE: only intended use is to remove objects from inside a FImplicitObjectUnion
+		CHAOS_ENSURE(MNonFrequentData.Read().GetGeometry()->GetType() == FImplicitObjectUnion::StaticType());
+
+		if (InIndex == INDEX_NONE)
+		{
+			return;
+		}
+
+		check(MShapesArray.IsValidIndex(InIndex));
+		MShapesArray.RemoveAt(InIndex);
+
 		if (MNonFrequentData.Read().GetGeometry()->GetType() == FImplicitObjectUnion::StaticType())
 		{
 			// if we are currently a union then remove geometry from this union
-			ModifyGeometry([FoundIndex](FImplicitObject& GeomToModify)
+			ModifyGeometry([InIndex](FImplicitObject& GeomToModify)
 			{
 				if (FImplicitObjectUnion* Union = GeomToModify.template GetObject<FImplicitObjectUnion>())
 				{
-					Union->RemoveAt(FoundIndex);
+					Union->RemoveAt(InIndex);
 				}
 			});
 		}
+	}
+
+	template <typename T, int d, bool bPersistent>
+	void Chaos::TGeometryParticleHandleImp<T, d, bPersistent>::RemoveShape(FPerShapeData* InShape)
+	{
+		// NOTE: only intended use is to remove objects from inside a FImplicitObjectUnion
+		const FImplicitObjectRef CurrentGeometry = GetGeometry();
+		ensure(CurrentGeometry != nullptr);
+		if (ensure(CurrentGeometry->GetType() == FImplicitObjectUnion::StaticType()))
+		{
+			const FShapesArray& CurrentShapesArray = ShapesArray();
+
+			int32 FoundIndex = INDEX_NONE;
+			for (int32 Index = 0; Index < CurrentShapesArray.Num(); Index++)
+			{
+				if (InShape == CurrentShapesArray[Index].Get())
+				{
+					FoundIndex = Index;
+					break;
+				}
+			}
+
+			RemoveShapeAtIndex(FoundIndex);
+		}
+	}
+
+	template <typename T, int d, bool bPersistent>
+	void Chaos::TGeometryParticleHandleImp<T, d, bPersistent>::RemoveShapeAtIndex(int32 InIndex)
+	{
+		if (InIndex == INDEX_NONE)
+		{
+			return;
+		}
+
+		// NOTE: only intended use is to remove objects from inside a FImplicitObjectUnion
+		const FImplicitObjectRef CurrentGeometry = GetGeometry();
+		ensure(CurrentGeometry != nullptr);
+		if (!ensure(CurrentGeometry->GetType() == FImplicitObjectUnion::StaticType()))
+		{
+			return;
+		}
+
+		GeometryParticles->RemoveShapeAtIndex(ParticleIdx, InIndex);
+
+		// if we are currently a union then remove geometry from this union
+		FImplicitObjectUnion& Union = CurrentGeometry->GetObjectChecked<FImplicitObjectUnion>();
+		Union.RemoveAt(InIndex);
 
+		// Needed to update the shapes array.
+		SetGeometry(GeometryParticles->GetGeometry(ParticleIdx));
 	}
 
 	template <typename T, int d>
diff --git a/Engine/Source/Runtime/Experimental/Chaos/Private/Chaos/ShapeInstance.cpp b/Engine/Source/Runtime/Experimental/Chaos/Private/Chaos/ShapeInstance.cpp
--- a/Engine/Source/Runtime/Experimental/Chaos/Private/Chaos/ShapeInstance.cpp
+++ b/Engine/Source/Runtime/Experimental/Chaos/Private/Chaos/ShapeInstance.cpp
@@ -8,6 +8,7 @@ namespace Chaos
 {
 	static_assert(sizeof(FShapeInstanceProxy) <= 192, "FShapeInstanceProxy was optimized to fit into 192b bin of MB3 to prevent excess memory waste");
 
+	DECLARE_CYCLE_STAT(TEXT("UpdateShapesArrayFromGeometryImpl"), STAT_UpdateShapesArrayFromGeometryImpl, STATGROUP_Chaos);
 	// Create or reuse the shapes in the shapes array and populate them with the Geometry.
 	// If we have a Union it will be unpacked into the ShapesArray.
 	// On the Physics Thread we set bAllowCachedLeafInfo which caches the shapes world space state to optimize collision detection,
@@ -20,6 +21,7 @@ namespace Chaos
 		const FRigidTransform3& ActorTM,
 		IPhysicsProxyBase* Proxy)
 	{
+		SCOPE_CYCLE_COUNTER(STAT_UpdateShapesArrayFromGeometryImpl);
 		// TShapesArrayType = TArray<TUniquePtr<T>>
 		using FShapePtrType = typename TShapesArrayType::ElementType;	// TUniquePtr<T>
 		using FShapeType = typename FShapePtrType::ElementType;			// T
@@ -27,7 +29,7 @@ namespace Chaos
 		if (Geometry)
 		{
 			const int32 OldShapeNum = ShapesArray.Num();
-			if (const auto* Union = Geometry->template GetObject<FImplicitObjectUnion>())
+			if (const FImplicitObjectUnion* Union = Geometry->template GetObject<FImplicitObjectUnion>())
 			{
 				const int32 NumImplicits = Union->NumImplicits();
 				ShapesArray.Reserve(NumImplicits);
@@ -46,6 +48,7 @@ namespace Chaos
 					{
 						// Update geometry pointer if it changed
 						FShapeType::UpdateGeometry(ShapesArray[ShapeIndex], ShapeGeometry);
+						ShapesArray[ShapeIndex]->ModifyShapeIndex(ShapeIndex);
 					}
 				}
 			}
diff --git a/Engine/Source/Runtime/Experimental/Chaos/Private/PhysicsProxy/ClusterUnionPhysicsProxy.cpp b/Engine/Source/Runtime/Experimental/Chaos/Private/PhysicsProxy/ClusterUnionPhysicsProxy.cpp
--- a/Engine/Source/Runtime/Experimental/Chaos/Private/PhysicsProxy/ClusterUnionPhysicsProxy.cpp
+++ b/Engine/Source/Runtime/Experimental/Chaos/Private/PhysicsProxy/ClusterUnionPhysicsProxy.cpp
@@ -43,6 +43,7 @@ namespace Chaos
 			BufferData.V = Particle->V();
 			BufferData.W = Particle->W();
 			BufferData.ObjectState = Particle->ObjectState();
+			BufferData.Geometry = nullptr;
 
 			const FShapesArray& ShapeArray = Particle->ShapesArray();
 			BufferData.CollisionData.Empty(ShapeArray.Num());
@@ -59,11 +60,6 @@ namespace Chaos
 			if constexpr (std::is_base_of_v<FClusterUnionPhysicsProxy::FInternalParticle, TParticle>)
 			{
 				BufferData.bIsAnchored = Particle->IsAnchored();
-				BufferData.Geometry = Particle->GetGeometry();
-			}
-			else if constexpr (std::is_base_of_v<FClusterUnionPhysicsProxy::FExternalParticle, TParticle>)
-			{
-				BufferData.Geometry = Particle->GetGeometry();
 			}
 		}
 	}
@@ -263,26 +259,10 @@ namespace Chaos
 			}
 		);
 	}
-	void FClusterUnionPhysicsProxy::UpdateShapes_External(const TArray<FPBDRigidParticle*>& ShapeParticles)
-	{
-		FImplicitObjectUnion* ImplicitUnion = Particle_External->GetGeometry()->AsA<FImplicitObjectUnion>();
-		const int32 NumSimpleShapes = ImplicitUnion ? ImplicitUnion->GetConvexes().Num() : 0;
-		
-		UpdateShapesDatas(ShapeParticles, Particle_External->ShapesArray(),
-			InitData.ActorId, InitData.ComponentId, NumSimpleShapes);
-	}
-	
-	void FClusterUnionPhysicsProxy::RemoveShapes_External(const TArray<FPBDRigidParticle*>& ShapeParticles) const
+
+	void FClusterUnionPhysicsProxy::RemoveShapes_External(const TArray<FPBDRigidParticle*>& ShapeParticles)
     {
-		if (!ensure(Particle_External))
-		{
-			return;
-		}
-		for(FPBDRigidParticle* ShapeParticle : ShapeParticles)
-		{
-			const TUniquePtr<Chaos::FPerShapeData>& TemplateShape = ShapeParticle->ShapesArray()[0];
-			Particle_External->RemoveShape(TemplateShape.Get(), false);
-		}
+		RemoveParticlesFromClusterUnionGeometry(Particle_External.Get(), ShapeParticles, GeometryChildParticles_External);
     }
 	
 	void FClusterUnionPhysicsProxy::MergeGeometry_External(TArray<Chaos::FImplicitObjectPtr>&& ImplicitGeometries, const TArray<FPBDRigidParticle*>& ShapeParticles)
@@ -293,10 +273,18 @@ namespace Chaos
 		}
 		
 		// In the cases where this is necessary, the SQ should have a valid state - it's just the geometry itself that isn't valid.
-		Particle_External->MergeGeometry(MoveTemp(ImplicitGeometries));
+		ModifyAdditionOfChildrenToClusterUnionGeometry(
+			Particle_External.Get(),
+			ShapeParticles,
+			InitData.ActorId,
+			InitData.ComponentId,
+			[this, ImplicitGeometries=MoveTemp(ImplicitGeometries)]() mutable
+			{
+				Particle_External->MergeGeometry(MoveTemp(ImplicitGeometries));
+			}
+		);
 
-		// Update shapes datas (Collision/Query/Sim)
-		UpdateShapes_External(ShapeParticles);
+		GeometryChildParticles_External.Append(ShapeParticles);
 	}
 
 	void FClusterUnionPhysicsProxy::SetGeometry_External(const Chaos::FImplicitObjectPtr& Geometry, const TArray<FPBDRigidParticle*>& ShapeParticles)
@@ -307,10 +295,18 @@ namespace Chaos
 		}
 
 		// In the cases where this is necessary, the SQ should have a valid state - it's just the geometry itself that isn't valid.
-		Particle_External->SetGeometry(Geometry);
+		ModifyAdditionOfChildrenToClusterUnionGeometry(
+			Particle_External.Get(),
+			ShapeParticles,
+			InitData.ActorId,
+			InitData.ComponentId,
+			[this, &Geometry]()
+			{
+				Particle_External->SetGeometry(Geometry);
+			}
+		);
 
-		// Update shapes datas (Collision/Query/Sim)
-		UpdateShapes_External(ShapeParticles);
+		GeometryChildParticles_External = ShapeParticles;
 	}
 
 	DECLARE_CYCLE_STAT(TEXT("FClusterUnionPhysicsProxy::PushToPhysicsState"), STAT_ClusterUnionPhysicsProxyPushToPhysicsState, STATGROUP_Chaos);
@@ -413,7 +409,10 @@ namespace Chaos
 			SyncedData_External.ChildParticles.Add(ConvertedData);
 		}
 
-		Particle_External->SetGeometry(CurrentPullData.Geometry);
+		if (CurrentPullData.Geometry)
+		{
+			Particle_External->SetGeometry(CurrentPullData.Geometry);
+		}
 		Particle_External->SetObjectState(CurrentPullData.ObjectState, true, /*bInvalidate=*/false);
 
 		const FShapesArray& ShapeArray = Particle_External->ShapesArray();
@@ -547,7 +546,11 @@ namespace Chaos
 				}
 			}
 
-			BufferData.Geometry = ClusterUnion->Geometry;
+			// TODO: Probably only want to do this when something about the geometry changes.
+			if (!ClusterUnion->ChildParticles.IsEmpty())
+			{
+				BufferData.Geometry = ClusterUnion->Geometry->DeepCopyGeometry();
+			}
 		}
 	}
 
diff --git a/Engine/Source/Runtime/Experimental/Chaos/Private/PhysicsProxy/GeometryCollectionPhysicsProxy.cpp b/Engine/Source/Runtime/Experimental/Chaos/Private/PhysicsProxy/GeometryCollectionPhysicsProxy.cpp
--- a/Engine/Source/Runtime/Experimental/Chaos/Private/PhysicsProxy/GeometryCollectionPhysicsProxy.cpp
+++ b/Engine/Source/Runtime/Experimental/Chaos/Private/PhysicsProxy/GeometryCollectionPhysicsProxy.cpp
@@ -3798,10 +3798,10 @@ void FGeometryCollectionPhysicsProxy::UpdatePerParticleFilterData_External(const
 
 					if (Chaos::FPBDRigidClusteredParticleHandle* Particle = SolverParticleHandles[Data.ParticleIndex])
 					{
-						Chaos::FClusterUnionIndex ClusterUnionIndex = ClusterUnionManager.FindClusterUnionIndexFromParticle(Particle);
-						if (ClusterUnionIndex != INDEX_NONE)
+						if (Chaos::FClusterUnion* ClusterUnion = ClusterUnionManager.FindClusterUnionFromParticle(Particle))
 						{
-							ClusterUnionManager.RequestDeferredClusterPropertiesUpdate(ClusterUnionIndex, Chaos::EUpdateClusterUnionPropertiesFlags::None);
+							ClusterUnion->AddPendingGeometryOperation(Chaos::EClusterUnionGeometryOperation::Refresh, Particle);
+							ClusterUnionManager.RequestDeferredClusterPropertiesUpdate(ClusterUnion->InternalIndex, Chaos::EUpdateClusterUnionPropertiesFlags::IncrementalGenerateGeometry);
 						}
 					}
 				}
diff --git a/Engine/Source/Runtime/Experimental/Chaos/Public/Chaos/ClusterUnionManager.h b/Engine/Source/Runtime/Experimental/Chaos/Public/Chaos/ClusterUnionManager.h
--- a/Engine/Source/Runtime/Experimental/Chaos/Public/Chaos/ClusterUnionManager.h
+++ b/Engine/Source/Runtime/Experimental/Chaos/Public/Chaos/ClusterUnionManager.h
@@ -40,7 +40,9 @@ namespace Chaos
 		ForceGenerateConnectionGraph = 1 << 1,
 		IncrementalGenerateConnectionGraph = 1 << 2,
 		UpdateKinematicProperties = 1 << 3,
-		All = RecomputeMassOrientation | ForceGenerateConnectionGraph | UpdateKinematicProperties
+		ForceGenerateGeometry = 1 << 4,
+		IncrementalGenerateGeometry = 1 << 5,
+		All = RecomputeMassOrientation | ForceGenerateConnectionGraph | UpdateKinematicProperties | ForceGenerateGeometry
 	};
 	ENUM_CLASS_FLAGS(EUpdateClusterUnionPropertiesFlags);
 
@@ -50,6 +52,13 @@ namespace Chaos
 		Remove
 	};
 
+	enum class EClusterUnionGeometryOperation : int8
+	{
+		Add,
+		Remove,
+		Refresh
+	};
+
 	struct FClusterUnionCreationParameters
 	{
 		FClusterUnionExplicitIndex ExplicitIndex = INDEX_NONE;
@@ -75,6 +84,10 @@ namespace Chaos
 		// All the particles that belong to this cluster.
 		TArray<FPBDRigidParticleHandle*> ChildParticles;
 
+		// All the particles that belong to this cluster that are currently in the implicit object geometry. There is a delay between removal from ChildParticles and its subsequent
+		// removal from GeometryChildParticles.
+		TArray<FPBDRigidParticleHandle*> GeometryChildParticles;
+
 		// Additional properties that are only relevant for particles that have been added into a cluster union.
 		TMap<FPBDRigidParticleHandle*, FClusterUnionParticleProperties> ChildProperties;
 
@@ -104,6 +117,16 @@ namespace Chaos
 
 		// Pending particles that need to be added into the connectivity graph.
 		TArray<TPair<FPBDRigidParticleHandle*, EClusterUnionConnectivityOperation>> PendingConnectivityOperations;
+
+		const TArray<FPBDRigidParticleHandle*>& GetPendingGeometryOperationParticles(EClusterUnionGeometryOperation Op) const;
+		void AddPendingGeometryOperation(EClusterUnionGeometryOperation Op, FPBDRigidParticleHandle* Particle);
+		void ClearAllPendingGeometryOperations();
+		void ClearPendingGeometryOperations(EClusterUnionGeometryOperation Op);
+
+	private:
+		// Pending particles that need to be added into/removed from the cluster union's geometry.
+		// These need to be arrays to be able to maintain ordering.
+		TMap<EClusterUnionGeometryOperation, TArray<FPBDRigidParticleHandle*>> PendingGeometryOperations;
 	};
 
 	struct FClusterUnionChildToParentUpdate
@@ -246,61 +269,168 @@ namespace Chaos
 
 		// Flush the cluster union's incremental connectivity operations
 		CHAOS_API void FlushIncrementalConnectivityGraphOperations(FClusterUnion& ClusterUnion);
+
+		// Flush the cluster union's incremental geometry operations.
+		void FlushIncrementalGeometryOperations(FClusterUnion& ClusterUnion);
 	};
 
-	// Update all the shapes datas including the simplified ones
-	template<typename ParticleType>
-	void UpdateShapesDatas(const TArray<ParticleType*>& ShapesParticles,
-		const FShapesArray& ShapesArray, const int32 ActorId, const int32 ComponentID, const int32 NumSimpleShapes)
+	template<typename TParticle>
+	void TransferClusterUnionShapeData(const TUniquePtr<Chaos::FPerShapeData>& ShapeData, TParticle* TemplateParticle, const TUniquePtr<Chaos::FPerShapeData>& TemplateShape, int32 ActorId, int32 ComponentId)
 	{
-		auto TransferShapeData = [&ActorId, &ComponentID](const TUniquePtr<Chaos::FPerShapeData>& ShapeData, const TUniquePtr<Chaos::FPerShapeData>& TemplateShape)
+		if (ShapeData && TemplateShape)
 		{
-			if (ShapeData && TemplateShape)
 			{
-				{
-					FCollisionData Data = TemplateShape->GetCollisionData();
-					Data.UserData = nullptr;
-					ShapeData->SetCollisionData(Data);
-				}
+				FCollisionData Data = TemplateShape->GetCollisionData();
+				Data.UserData = nullptr;
+				ShapeData->SetCollisionData(Data);
+			}
 
-				{
-					FCollisionFilterData Data = TemplateShape->GetQueryData();
-					Data.Word0 = ActorId;
-					ShapeData->SetQueryData(Data);
-				}
+			{
+				FCollisionFilterData Data = TemplateShape->GetQueryData();
+				Data.Word0 = ActorId;
+				ShapeData->SetQueryData(Data);
+			}
 
-				{
-					FCollisionFilterData Data = TemplateShape->GetSimData();
-					Data.Word0 = 0;
-					Data.Word2 = ComponentID;
-					ShapeData->SetSimData(Data);
-				}
+			{
+				FCollisionFilterData Data = TemplateShape->GetSimData();
+				Data.Word0 = 0;
+				Data.Word2 = ComponentId;
+				ShapeData->SetSimData(Data);
 			}
-		};
-		if (!ShapesParticles.IsEmpty() && (ShapesParticles.Num() <= (ShapesArray.Num()-NumSimpleShapes)))
+		}
+	}
+
+	/**
+	 * Wraps any additions to the cluster union geometry. We assume that the input lambda Func
+	 * will modify the cluster union particle's geometry *and* shapes with the input particles.
+	 * This function will make sure the newly added shapes are setup properly with the expected
+     * properties (sim data, query data, user data etc.).
+	 */
+	template<typename TClusterParticle, typename TParticle, typename TLambda>
+	void ModifyAdditionOfChildrenToClusterUnionGeometry(TClusterParticle* ClusterParticle, const TArray<TParticle*>& Particles, int32 ActorId, int32 ComponentId, TLambda&& Func)
+	{
+		if (!ClusterParticle || Particles.IsEmpty())
 		{
-			const int32 ShapesOffset = ShapesArray.Num() - ShapesParticles.Num();
-			for (int32 ParticleIndex = 0; ParticleIndex < ShapesParticles.Num(); ++ParticleIndex)
+			return;
+		}
+
+		FImplicitObjectUnion* ImplicitUnion = ClusterParticle->GetGeometry() ? ClusterParticle->GetGeometry()->template AsA<FImplicitObjectUnion>() : nullptr;
+		const int32 OldNumSimpleShapes = ImplicitUnion ? ImplicitUnion->GetConvexes().Num() : 0;
+		const int32 OldNumChildShapes = ImplicitUnion ? ImplicitUnion->GetNumRootObjects() : 0;
+
+		Func();
+
+		// We must have a union at this point since the contract is that Func is adding a shape
+		// which means we must be in a union already.
+		check(ClusterParticle->GetGeometry() != nullptr);
+		ImplicitUnion = ClusterParticle->GetGeometry()->template AsA<FImplicitObjectUnion>();
+		check(ImplicitUnion != nullptr);
+
+		const int32 NewNumSimpleShapes = ImplicitUnion->GetConvexes().Num();
+		const int32 NewNumChildShapes = ImplicitUnion->GetNumRootObjects();
+		const FShapesArray& ShapeArray = ClusterParticle->ShapesArray();
+		check(ShapeArray.Num() == (NewNumChildShapes + NewNumSimpleShapes));
+		check(Particles.Num() == (NewNumChildShapes - OldNumChildShapes));
+	
+		for (int32 Index = 0; Index < Particles.Num(); ++Index)
+		{
+			TParticle* Particle = Particles[Index];
+			if (!Particle)
+			{
+				continue;
+			}
+
+			const TUniquePtr<Chaos::FPerShapeData>& TemplateShape = Particle->ShapesArray()[0];
+			if (!TemplateShape)
+			{
+				continue;
+			}
+
+			const int32 ShapeIndex = NewNumSimpleShapes + OldNumChildShapes + Index;
+			if (const TUniquePtr<Chaos::FPerShapeData>& ShapeData = ShapeArray[ShapeIndex])
 			{
-				const TUniquePtr<Chaos::FPerShapeData>& ShapeData = ShapesArray[ShapesOffset+ParticleIndex];
-				const TUniquePtr<Chaos::FPerShapeData>& TemplateShape = ShapesParticles[ParticleIndex]->ShapesArray()[0];
-				TransferShapeData(ShapeData, TemplateShape);
-				
-				if(ShapeData && NumSimpleShapes > 0)
+				TransferClusterUnionShapeData(ShapeData, Particle, TemplateShape, ActorId, ComponentId);
+
+				if (NewNumSimpleShapes > 0)
 				{
 					ShapeData->SetSimEnabled(false);
 				}
 			}
-			for (int32 ShapeIndex = 0; ShapeIndex < NumSimpleShapes; ++ShapeIndex)
+		}
+
+		for (int32 Index = OldNumSimpleShapes; Index < NewNumSimpleShapes; ++Index)
+		{
+			const TUniquePtr<Chaos::FPerShapeData>& ShapeData = ShapeArray[Index];
+			const TUniquePtr<Chaos::FPerShapeData>& TemplateShape = Particles[0]->ShapesArray()[0];
+
+			if (ShapeData && TemplateShape)
 			{
-				const TUniquePtr<Chaos::FPerShapeData>& ShapeData = ShapesArray[ShapeIndex];
-				const TUniquePtr<Chaos::FPerShapeData>& TemplateShape = ShapesParticles[0]->ShapesArray()[0];
-				TransferShapeData(ShapeData, TemplateShape);
-				
+				TransferClusterUnionShapeData<TParticle>(ShapeData, nullptr, TemplateShape, ActorId, ComponentId);
 				ShapeData->SetQueryEnabled(false);
 			}
 		}
 	}
-	
 
+	/**
+	 * Utility function to remove particles from a cluster union's geometry incrementally.
+	 */
+	template<typename TClusterParticle, typename TParticle>
+	void RemoveParticlesFromClusterUnionGeometry(TClusterParticle* ClusterParticle, const TArray<TParticle*>& ShapeParticles, TArray<TParticle*>& AllChildParticles)
+	{
+		if (!ensure(ClusterParticle))
+		{
+			return;
+		}
+
+		const Chaos::FShapesArray& ShapesArray = ClusterParticle->ShapesArray();
+		const Chaos::FImplicitObjectRef Geometry = ClusterParticle->GetGeometry();
+		check(Geometry != nullptr);
+
+		const Chaos::FImplicitObjectUnion& GeometryUnion = Geometry->GetObjectChecked<Chaos::FImplicitObjectUnion>();
+		const int32 NumSimpleShapes = GeometryUnion.GetConvexes().Num();
+		check(AllChildParticles.Num() == (ShapesArray.Num() - NumSimpleShapes));
+		
+		TArray<int32> ShapeIndicesToRemove;
+		ShapeIndicesToRemove.Reserve(ShapeParticles.Num());
+		for(TParticle* ShapeParticle : ShapeParticles)
+		{
+			check(ShapeParticle != nullptr);
+			const int32 Index = AllChildParticles.Find(ShapeParticle);
+			if (Index != INDEX_NONE)
+			{
+				ShapeIndicesToRemove.Add(Index + NumSimpleShapes);
+			}
+		}
+
+		// To make sure we preserve the index as we remove shapes one-by-one, we need to order
+		// shape removal based on their index.
+		ShapeIndicesToRemove.Sort(
+			[](const int32 A, const int32 B)
+			{
+				// This way we guarantee descending order.
+				return A > B;
+			}
+		);
+
+		int32 LastIndex = INDEX_NONE;
+		// TODO: Can we batch this to remove all shapes at once?
+		// We're going to need to do that so we don't do multiple iterations through the shapes array.
+		for (int32 ToRemoveIndex : ShapeIndicesToRemove)
+		{
+			// This is a good cheap way of preventing us from trying to remove the same shape (index)
+			// multiple times.
+			if (ToRemoveIndex == LastIndex)
+			{
+				// Probably a good indication of a failure if we end up trying to remove the same shape more than once.
+				check(false);
+				continue;
+			}
+
+			ClusterParticle->RemoveShapeAtIndex(ToRemoveIndex);
+			AllChildParticles.RemoveAt(ToRemoveIndex - NumSimpleShapes);
+			LastIndex = ToRemoveIndex;
+		}
+
+		check(AllChildParticles.Num() == (ShapesArray.Num() - NumSimpleShapes));
+	}
 }
diff --git a/Engine/Source/Runtime/Experimental/Chaos/Public/Chaos/GeometryParticles.h b/Engine/Source/Runtime/Experimental/Chaos/Public/Chaos/GeometryParticles.h
--- a/Engine/Source/Runtime/Experimental/Chaos/Public/Chaos/GeometryParticles.h
+++ b/Engine/Source/Runtime/Experimental/Chaos/Public/Chaos/GeometryParticles.h
@@ -252,6 +252,7 @@ namespace Chaos
 		bool& HasCollision(const int32 Index) { return MHasCollision[Index]; }
 
 		const FShapesArray& ShapesArray(const int32 Index) const { return reinterpret_cast<const FShapesArray&>(MShapesArray[Index]); }
+		void RemoveShapeAtIndex(const int32 ParticleIndex, const int32 Index);
 
 		const FShapeInstanceArray& ShapeInstances(const int32 Index) const { return MShapesArray[Index]; }
 
diff --git a/Engine/Source/Runtime/Experimental/Chaos/Public/Chaos/PBDRigidClusteredParticles.h b/Engine/Source/Runtime/Experimental/Chaos/Public/Chaos/PBDRigidClusteredParticles.h
--- a/Engine/Source/Runtime/Experimental/Chaos/Public/Chaos/PBDRigidClusteredParticles.h
+++ b/Engine/Source/Runtime/Experimental/Chaos/Public/Chaos/PBDRigidClusteredParticles.h
@@ -7,7 +7,6 @@
 
 namespace Chaos
 {
-	class FImplicitObjectUnionClustered;
 
 /** 
  * Used within the clustering system to describe the clustering hierarchy. The ClusterId
diff --git a/Engine/Source/Runtime/Experimental/Chaos/Public/Chaos/ParticleHandle.h b/Engine/Source/Runtime/Experimental/Chaos/Public/Chaos/ParticleHandle.h
--- a/Engine/Source/Runtime/Experimental/Chaos/Public/Chaos/ParticleHandle.h
+++ b/Engine/Source/Runtime/Experimental/Chaos/Public/Chaos/ParticleHandle.h
@@ -556,6 +556,7 @@ public:
 
 	const FImplicitObjectRef GetGeometry() const { return GeometryParticles->GetGeometry(ParticleIdx).GetReference(); }
 	void SetGeometry(const FImplicitObjectPtr& InGeometry) { GeometryParticles->SetGeometry(ParticleIdx, InGeometry); }
+	void MergeGeometry(TArray<Chaos::FImplicitObjectPtr>&& Objects);
 
 	UE_DEPRECATED(5.4, "Use GetGeometry instead")
 	TSerializablePtr<FImplicitObject> Geometry() const { check(false); return TSerializablePtr<FImplicitObject>(); }
@@ -579,8 +580,12 @@ public:
 	void SetDynamicGeometry(TUniquePtr<FImplicitObject>&& Unique) { check(false); }
 
 	const FShapesArray& ShapesArray() const { return GeometryParticles->ShapesArray(ParticleIdx); }
+
 	const FShapeInstanceArray& ShapeInstances() const { return GeometryParticles->ShapeInstances(ParticleIdx); }
 
+	void RemoveShape(FPerShapeData* InShape);
+	void RemoveShapeAtIndex(int32 InIndex);
+
 	const TAABB<T, d>& LocalBounds() const { return GeometryParticles->LocalBounds(ParticleIdx); }
 	void SetLocalBounds(const TAABB<T, d>& NewBounds) { GeometryParticles->LocalBounds(ParticleIdx) = NewBounds; }
 
@@ -2529,6 +2534,7 @@ public:
 	TSerializablePtr<FImplicitObject> Geometry() const { check(false); return TSerializablePtr<FImplicitObject>(); }
 
 	CHAOS_API void RemoveShape(FPerShapeData* InShape, bool bWakeTouching);
+	CHAOS_API void RemoveShapeAtIndex(int32 InIndex);
 
 	void* UserData() const { return MUserData; }
 	void SetUserData(void* InUserData)
diff --git a/Engine/Source/Runtime/Experimental/Chaos/Public/PhysicsProxy/ClusterUnionPhysicsProxy.h b/Engine/Source/Runtime/Experimental/Chaos/Public/PhysicsProxy/ClusterUnionPhysicsProxy.h
--- a/Engine/Source/Runtime/Experimental/Chaos/Public/PhysicsProxy/ClusterUnionPhysicsProxy.h
+++ b/Engine/Source/Runtime/Experimental/Chaos/Public/PhysicsProxy/ClusterUnionPhysicsProxy.h
@@ -81,7 +81,7 @@ namespace Chaos
 		CHAOS_API void MergeGeometry_External(TArray<Chaos::FImplicitObjectPtr>&& ImplicitGeometries, const TArray<FPBDRigidParticle*>& ShapeParticles);
 
 		// Remove GT shapes from the existing unions
-		CHAOS_API void RemoveShapes_External(const TArray<FPBDRigidParticle*>& ShapeParticles) const;
+		CHAOS_API void RemoveShapes_External(const TArray<FPBDRigidParticle*>& ShapeParticles);
 		
 		UE_DEPRECATED(5.4, "Please use SetGeometry_External instead")
 		CHAOS_API void SetSharedGeometry_External(const TSharedPtr<Chaos::FImplicitObject, ESPMode::ThreadSafe>& Geometry, const TArray<FPBDRigidParticle*>& ShapeParticles)
@@ -132,9 +132,9 @@ namespace Chaos
 
 		FClusterUnionIndex GetClusterUnionIndex() const { return ClusterUnionIndex; }
 
+		void ForceSetGeometryChildParticles_External(TArray<FExternalParticle*>&& InParticles) { GeometryChildParticles_External = InParticles; }
+
 	private:
-		void UpdateShapes_External(const TArray<FPBDRigidParticle*>& ShapeParticles);
-		
 		bool bIsInitializedOnPhysicsThread = false;
 		FClusterCreationParameters ClusterParameters;
 		const FClusterUnionInitData InitData;
@@ -148,6 +148,11 @@ namespace Chaos
 
 		FProxyInterpolationBase InterpolationData;
 
+		// An array of a particles that exist in the external implicit object union.
+		// Note that this array should only be used for book-keeping. It is generally
+		// unsafe to try and access the particles within this array.
+		TArray<FExternalParticle*> GeometryChildParticles_External;
+
 		//~ Begin TPhysicsProxy Interface
 	public:
 		bool IsSimulating() const { return true; }
