Chaos - enable new initial depenetration handling system by default, with zero depentration velocity

diff --git a/Engine/Source/Runtime/Experimental/Chaos/Private/Chaos/Collision/PBDCollisionContainerSolver.cpp b/Engine/Source/Runtime/Experimental/Chaos/Private/Chaos/Collision/PBDCollisionContainerSolver.cpp
--- a/Engine/Source/Runtime/Experimental/Chaos/Private/Chaos/Collision/PBDCollisionContainerSolver.cpp
+++ b/Engine/Source/Runtime/Experimental/Chaos/Private/Chaos/Collision/PBDCollisionContainerSolver.cpp
@@ -53,7 +53,7 @@ namespace Chaos
 		FAutoConsoleVariableRef CVarChaosPBDCollisionSolverJacobiRotationTolerance(TEXT("p.Chaos.PBDCollisionSolver.JacobiRotationTolerance"), Chaos_PBDCollisionSolver_JacobiRotationTolerance, TEXT(""));
 
 		// Whether to enable the new initial overlap depentration system
-		bool bChaos_Collision_EnableInitialDepenetration = false;
+		bool bChaos_Collision_EnableInitialDepenetration = true;
 		FAutoConsoleVariableRef CVarChaosCollisionEnableInitialDepentration(TEXT("p.Chaos.PBDCollisionSolver.EnableInitialDepenetration"), bChaos_Collision_EnableInitialDepenetration, TEXT(""));
 
 		// The maximum number of constraints we will attempt to solve (-1 for unlimited)
