geometry collections need to reload collision profiles when the primitive component body instance collision settings change

diff --git a/Engine/Source/Runtime/Experimental/GeometryCollectionEngine/Private/GeometryCollection/GeometryCollectionComponent.cpp b/Engine/Source/Runtime/Experimental/GeometryCollectionEngine/Private/GeometryCollection/GeometryCollectionComponent.cpp
--- a/Engine/Source/Runtime/Experimental/GeometryCollectionEngine/Private/GeometryCollection/GeometryCollectionComponent.cpp
+++ b/Engine/Source/Runtime/Experimental/GeometryCollectionEngine/Private/GeometryCollection/GeometryCollectionComponent.cpp
@@ -5834,6 +5834,13 @@ void UGeometryCollectionComponent::SetEnableDamageFromCollision(bool bValue)
 	}
 }
 
+void UGeometryCollectionComponent::OnComponentCollisionSettingsChanged(bool bUpdateOverlaps)
+{
+	Super::OnComponentCollisionSettingsChanged(bUpdateOverlaps);
+
+	BuildInitialFilterData();
+	LoadCollisionProfiles();
+}
 
 #if WITH_EDITOR
 
diff --git a/Engine/Source/Runtime/Experimental/GeometryCollectionEngine/Public/GeometryCollection/GeometryCollectionComponent.h b/Engine/Source/Runtime/Experimental/GeometryCollectionEngine/Public/GeometryCollection/GeometryCollectionComponent.h
--- a/Engine/Source/Runtime/Experimental/GeometryCollectionEngine/Public/GeometryCollection/GeometryCollectionComponent.h
+++ b/Engine/Source/Runtime/Experimental/GeometryCollectionEngine/Public/GeometryCollection/GeometryCollectionComponent.h
@@ -505,6 +505,7 @@ public:
 
 
 	//~ Begin UPrimitiveComponent Interface.
+public:
 	GEOMETRYCOLLECTIONENGINE_API virtual FPrimitiveSceneProxy* CreateSceneProxy() override;
 	GEOMETRYCOLLECTIONENGINE_API virtual void OnRegister() override;
 	GEOMETRYCOLLECTIONENGINE_API virtual void OnUnregister() override;
@@ -520,10 +521,13 @@ public:
 	GEOMETRYCOLLECTIONENGINE_API virtual void AddRadialImpulse(FVector Origin, float Radius, float Strength, enum ERadialImpulseFalloff Falloff, bool bVelChange = false) override;
 	GEOMETRYCOLLECTIONENGINE_API virtual void AddTorqueInRadians(FVector Torque, FName BoneName = NAME_None, bool bAccelChange = false) override;
 	GEOMETRYCOLLECTIONENGINE_API virtual void PostLoad() override;
+protected:
+	GEOMETRYCOLLECTIONENGINE_API virtual void OnComponentCollisionSettingsChanged(bool bUpdateOverlaps=true) override;
 	//~ End UPrimitiveComponent Interface.
 
 
 	//~ Begin UMeshComponent Interface.	
+public:
 	GEOMETRYCOLLECTIONENGINE_API virtual int32 GetNumMaterials() const override;
 	GEOMETRYCOLLECTIONENGINE_API virtual UMaterialInterface* GetMaterial(int32 MaterialIndex) const override;
 	GEOMETRYCOLLECTIONENGINE_API virtual void GetUsedMaterials(TArray<UMaterialInterface*>& OutMaterials, bool bGetDebugMaterials = false) const override;
