Fix apply radial impulse to work with moved cluster unions
- Make sure that disabled particle are properly using the ChildToParent transform to compute their position

diff --git a/Engine/Source/Runtime/Experimental/Chaos/Private/Chaos/PhysicsObjectInterface.cpp b/Engine/Source/Runtime/Experimental/Chaos/Private/Chaos/PhysicsObjectInterface.cpp
--- a/Engine/Source/Runtime/Experimental/Chaos/Private/Chaos/PhysicsObjectInterface.cpp
+++ b/Engine/Source/Runtime/Experimental/Chaos/Private/Chaos/PhysicsObjectInterface.cpp
@@ -88,7 +88,20 @@ namespace
 
 		if (ParticleHandle)
 		{
-			const FVec3 ParticleToOrigin = ParticleHandle->X() - Origin;
+			FVec3 ParticleX = ParticleHandle->X();
+			if (ParticleHandle->Disabled())
+			{
+				if (FPBDRigidClusteredParticleHandle* ClusteredParticle = ParticleHandle->CastToClustered())
+				{
+					if (const FPBDRigidClusteredParticleHandle* ClusteredParent = ClusteredParticle->Parent())
+					{
+						const FTransform ParentTransform(ClusteredParent->R(), ClusteredParent->X());
+						ParticleX = ParentTransform.TransformPosition(ClusteredParticle->ChildToParent().GetTranslation());
+					}
+				}
+			}
+
+			const FVec3 ParticleToOrigin = ParticleX - Origin;
 			const float DistanceSquared = (float)ParticleToOrigin.SizeSquared();
 			const float RadiusSquared = Radius * Radius;
 
