fix nullptr static analysis check in Chaos::TGeometryParticleHandleImp<T, d, bPersistent>::MergeGeometry

diff --git a/Engine/Source/Runtime/Experimental/Chaos/Private/Chaos/ParticleHandle.cpp b/Engine/Source/Runtime/Experimental/Chaos/Private/Chaos/ParticleHandle.cpp
--- a/Engine/Source/Runtime/Experimental/Chaos/Private/Chaos/ParticleHandle.cpp
+++ b/Engine/Source/Runtime/Experimental/Chaos/Private/Chaos/ParticleHandle.cpp
@@ -92,14 +92,16 @@ namespace Chaos
 		}
 
 		const FImplicitObjectRef CurrentGeometry = GetGeometry();
-		ensure(CurrentGeometry != nullptr);
-		if (ensure(CurrentGeometry->GetType() == FImplicitObjectUnion::StaticType()))
+		if (ensure(CurrentGeometry != nullptr))
 		{
-			FImplicitObjectUnion& Union = CurrentGeometry->GetObjectChecked<FImplicitObjectUnion>();
-			Union.Combine(Objects);
+			if (ensure(CurrentGeometry->GetType() == FImplicitObjectUnion::StaticType()))
+			{
+				FImplicitObjectUnion& Union = CurrentGeometry->GetObjectChecked<FImplicitObjectUnion>();
+				Union.Combine(Objects);
 
-			// Needed to update the shapes array.
-			SetGeometry(GeometryParticles->GetGeometry(ParticleIdx));
+				// Needed to update the shapes array.
+				SetGeometry(GeometryParticles->GetGeometry(ParticleIdx));
+			}
 		}
 	}
 
