Fix crash in the event manager

diff --git a/Engine/Source/Runtime/Experimental/Chaos/Private/EventManager.cpp b/Engine/Source/Runtime/Experimental/Chaos/Private/EventManager.cpp
--- a/Engine/Source/Runtime/Experimental/Chaos/Private/EventManager.cpp
+++ b/Engine/Source/Runtime/Experimental/Chaos/Private/EventManager.cpp
@@ -39,6 +39,10 @@ namespace Chaos
 
 	void FEventManager::FillProducerData(const Chaos::FPBDRigidsSolver* Solver, bool bResetData)
 	{
+		if (BufferMode == EMultiBufferMode::Double)
+		{
+			ResourceLock.ReadLock();
+		}
 		ContainerLock.ReadLock();
 		for (FEventContainerBasePtr EventContainer : EventContainers)
 		{
@@ -48,6 +52,10 @@ namespace Chaos
 			}
 		}
 		ContainerLock.ReadUnlock();
+		if (BufferMode == EMultiBufferMode::Double)
+		{
+			ResourceLock.ReadUnlock();
+		}
 	}
 
 	void FEventManager::FlipBuffersIfRequired()
diff --git a/Engine/Source/Runtime/Experimental/Chaos/Public/EventManager.h b/Engine/Source/Runtime/Experimental/Chaos/Public/EventManager.h
--- a/Engine/Source/Runtime/Experimental/Chaos/Public/EventManager.h
+++ b/Engine/Source/Runtime/Experimental/Chaos/Public/EventManager.h
@@ -469,12 +469,20 @@ namespace Chaos
 		template<typename PayloadType>
 		void AddEvent(const EEventType& EventType, TFunction<void(PayloadType& EventData)> InFunction)
 		{
+			if (BufferMode == EMultiBufferMode::Double)
+			{
+				ResourceLock.ReadLock();
+			}
 			ContainerLock.ReadLock();
 			if (TEventContainer<PayloadType>* EventContainer = StaticCast<TEventContainer<PayloadType>*>(EventContainers[FEventID(EventType)]))
 			{
 				EventContainer->AddEvent(InFunction);
 			}
 			ContainerLock.ReadUnlock();
+			if (BufferMode == EMultiBufferMode::Double)
+			{
+				ResourceLock.ReadUnlock();
+			}
 		}
 
 	private:
