geometry collections need to reset their rest transforms just in case they were stored accidentally

diff --git a/Engine/Source/Runtime/Experimental/GeometryCollectionEngine/Private/GeometryCollection/GeometryCollectionComponent.cpp b/Engine/Source/Runtime/Experimental/GeometryCollectionEngine/Private/GeometryCollection/GeometryCollectionComponent.cpp
--- a/Engine/Source/Runtime/Experimental/GeometryCollectionEngine/Private/GeometryCollection/GeometryCollectionComponent.cpp
+++ b/Engine/Source/Runtime/Experimental/GeometryCollectionEngine/Private/GeometryCollection/GeometryCollectionComponent.cpp
@@ -2926,6 +2926,12 @@ void UGeometryCollectionComponent::ResetDynamicCollection()
 		DynamicCollection.Reset();
 	}
 
+	const int32 NumRestCollectionTransforms = RestCollection ? RestCollection->GetGeometryCollection()->Transform.Num() : 0;
+	if (RestTransforms.Num() != NumRestCollectionTransforms)
+	{
+		RestTransforms.Reset();
+	}
+
 	// if Rest transform have been overriden uses them to initialize the dynamic collection transforms
 	if (RestTransforms.Num() > 0)
 	{
