more nullptr checks for static analysis

diff --git a/Engine/Source/Runtime/Experimental/Chaos/Private/Chaos/ParticleHandle.cpp b/Engine/Source/Runtime/Experimental/Chaos/Private/Chaos/ParticleHandle.cpp
--- a/Engine/Source/Runtime/Experimental/Chaos/Private/Chaos/ParticleHandle.cpp
+++ b/Engine/Source/Runtime/Experimental/Chaos/Private/Chaos/ParticleHandle.cpp
@@ -156,22 +156,24 @@ namespace Chaos
 	{
 		// NOTE: only intended use is to remove objects from inside a FImplicitObjectUnion
 		const FImplicitObjectRef CurrentGeometry = GetGeometry();
-		ensure(CurrentGeometry != nullptr);
-		if (ensure(CurrentGeometry->GetType() == FImplicitObjectUnion::StaticType()))
+		if (ensure(CurrentGeometry != nullptr))
 		{
-			const FShapesArray& CurrentShapesArray = ShapesArray();
-
-			int32 FoundIndex = INDEX_NONE;
-			for (int32 Index = 0; Index < CurrentShapesArray.Num(); Index++)
+			if (ensure(CurrentGeometry->GetType() == FImplicitObjectUnion::StaticType()))
 			{
-				if (InShape == CurrentShapesArray[Index].Get())
+				const FShapesArray& CurrentShapesArray = ShapesArray();
+
+				int32 FoundIndex = INDEX_NONE;
+				for (int32 Index = 0; Index < CurrentShapesArray.Num(); Index++)
 				{
-					FoundIndex = Index;
-					break;
+					if (InShape == CurrentShapesArray[Index].Get())
+					{
+						FoundIndex = Index;
+						break;
+					}
 				}
-			}
 
-			RemoveShapeAtIndex(FoundIndex);
+				RemoveShapeAtIndex(FoundIndex);
+			}
 		}
 	}
 
@@ -185,7 +187,11 @@ namespace Chaos
 
 		// NOTE: only intended use is to remove objects from inside a FImplicitObjectUnion
 		const FImplicitObjectRef CurrentGeometry = GetGeometry();
-		ensure(CurrentGeometry != nullptr);
+		if (!ensure(CurrentGeometry != nullptr))
+		{
+			return;
+		}
+
 		if (!ensure(CurrentGeometry->GetType() == FImplicitObjectUnion::StaticType()))
 		{
 			return;
