Make Component Space transform to be computed on demand
- Hide the array beghind an API that expose root transform and all transform tracking a dirty state

diff --git a/Engine/Source/Runtime/Experimental/GeometryCollectionEngine/Private/GeometryCollection/GeometryCollectionComponent.cpp b/Engine/Source/Runtime/Experimental/GeometryCollectionEngine/Private/GeometryCollection/GeometryCollectionComponent.cpp
--- a/Engine/Source/Runtime/Experimental/GeometryCollectionEngine/Private/GeometryCollection/GeometryCollectionComponent.cpp
+++ b/Engine/Source/Runtime/Experimental/GeometryCollectionEngine/Private/GeometryCollection/GeometryCollectionComponent.cpp
@@ -375,6 +375,7 @@ UGeometryCollectionComponent::UGeometryCollectionComponent(const FObjectInitiali
 	, NavmeshInvalidationTimeSliceIndex(0)
 	, IsObjectDynamic(false)
 	, IsObjectLoading(true)
+	, ComponentSpaceTransforms(this)
 	, PhysicsProxy(nullptr)
 #if WITH_EDITOR && WITH_EDITORONLY_DATA
 	, EditorActor(nullptr)
@@ -648,41 +649,8 @@ FBox UGeometryCollectionComponent::ComputeBounds(const FTransform& LocalToWorldW
 	FBox BoundingBox(ForceInit);
 	if (RestCollection)
 	{
-		//Hold on to reference so it doesn't get GC'ed
-		auto GeometryCollectionPtr = RestCollection->GetGeometryCollection();
-
-		const int32 NumElements = GeometryCollectionPtr->NumElements(FGeometryCollection::TransformGroup);
-		if (NumElements == 0 || ComponentSpaceTransforms.Num() != RestCollection->NumElements(FGeometryCollection::TransformGroup))
-		{
-			// #todo(dmp): we could do the bbox transform in parallel with a bit of reformulating		
-			// #todo(dmp):  there are some cases where the calcbounds function is called before the component
-			// has set the global matrices cache while in the editor.  This is a somewhat weak guard against this
-			// to default to just calculating tmp global matrices.  This should be removed or modified somehow
-			// such that we always cache the global matrices and this method always does the correct behavior
-
-			const TManagedArray<FTransform>& Transforms = GetTransformArray();
-			const TManagedArray<int32>& ParentIndices = GetParentArray();
-			if (!ensure(Transforms.Num() == NumElements))
-			{
-				return FBox(ForceInitToZero);
-			}
-			
-			TArray<FTransform> TmpComponentSpaceTransforms;
-			GeometryCollectionAlgo::GlobalMatrices(Transforms, ParentIndices, TmpComponentSpaceTransforms);
-			if (TmpComponentSpaceTransforms.Num() == 0)
-			{
-				BoundingBox = FBox(ForceInitToZero);
-			}
-			else
-			{
-				UpdateGlobalMatricesWithExplodedVectors(TmpComponentSpaceTransforms, *GeometryCollectionPtr);
-				BoundingBox = ComputeBoundsFromComponentSpaceTransforms(LocalToWorldWithScale, TmpComponentSpaceTransforms);
-			}
-		}
-		else
-		{
-			BoundingBox = ComputeBoundsFromComponentSpaceTransforms(LocalToWorldWithScale, ComponentSpaceTransforms);
-		}
+		const TArray<FTransform>& CompSpaceTransforms = ComponentSpaceTransforms.RequestAllTransforms();
+		BoundingBox = ComputeBoundsFromComponentSpaceTransforms(LocalToWorldWithScale, CompSpaceTransforms);
 	}
 	return BoundingBox;
 }
@@ -1263,6 +1231,8 @@ void UGeometryCollectionComponent::RefreshEmbeddedGeometry()
 	EmbeddedInstanceIndex.Init(INDEX_NONE, RestCollection->GetGeometryCollection()->NumElements(FGeometryCollection::TransformGroup));
 #endif
 
+	const TArray<FTransform>& CompSpaceTransforms = ComponentSpaceTransforms.RequestAllTransforms();
+
 	for (int32 ExemplarIndex = 0; ExemplarIndex < ExemplarCount; ++ExemplarIndex)
 	{		
 #if WITH_EDITOR
@@ -1280,7 +1250,7 @@ void UGeometryCollectionComponent::RefreshEmbeddedGeometry()
 			{
 				if (!HideArray || !(*HideArray)[Idx])
 				{ 
-					InstanceTransforms.Add(ComponentSpaceTransforms[Idx]);
+					InstanceTransforms.Add(CompSpaceTransforms[Idx]);
 #if WITH_EDITOR
 					int32 InstanceIndex = EmbeddedBoneMaps[ExemplarIndex].Add(Idx);
 					EmbeddedInstanceIndex[Idx] = InstanceIndex;
@@ -1355,9 +1325,9 @@ void UGeometryCollectionComponent::SetRestState(TArray<FTransform>&& InRestTrans
 	if (SceneProxy)
 	{
 		FGeometryCollectionDynamicData* DynamicData = GDynamicDataPool.Allocate();
-		DynamicData->SetPrevTransforms(ComponentSpaceTransforms);
-		CalculateGlobalMatrices();
-		DynamicData->SetTransforms(ComponentSpaceTransforms);
+		DynamicData->SetPrevTransforms(ComponentSpaceTransforms.RequestAllTransforms());
+		ComponentSpaceTransforms.MarkDirty();
+		DynamicData->SetTransforms(ComponentSpaceTransforms.RequestAllTransforms());
 		DynamicData->IsDynamic = true;
 
 #if WITH_EDITOR
@@ -1388,7 +1358,8 @@ void UGeometryCollectionComponent::SetRestState(TArray<FTransform>&& InRestTrans
 	}
 	else
 	{
-		CalculateGlobalMatrices();
+		// only need to mark it dirty and let whoever needs it to compute it on demand
+		ComponentSpaceTransforms.MarkDirty();
 	}
 
 	RefreshEmbeddedGeometry();
@@ -2414,41 +2385,31 @@ FGeometryCollectionDynamicData* UGeometryCollectionComponent::InitDynamicData(bo
 		DynamicData->IsDynamic = true;
 		DynamicData->IsLoading = GetIsObjectLoading();
 
-		// If we have no transforms stored in the dynamic data, then assign both prev and current to the same global matrices
-		if (ComponentSpaceTransforms.Num() == 0)
-		{
-			// Copy global matrices over to DynamicData
-			CalculateGlobalMatrices();
-
-			DynamicData->SetAllTransforms(ComponentSpaceTransforms);
-		}
-		else
-		{
-			// Copy existing global matrices into prev transforms
-			DynamicData->SetPrevTransforms(ComponentSpaceTransforms);
+		const TArray<FTransform>& CompSpaceTransforms = ComponentSpaceTransforms.RequestAllTransforms();
 
-			// Copy global matrices over to DynamicData
-			CalculateGlobalMatrices();
+		// If we have no transforms stored in the dynamic data, then assign both prev and current to the same global matrices
+		// Copy existing global matrices into prev transforms
+		DynamicData->PrevTransforms = DynamicData->Transforms;
 
-			bool bComputeChanges = true;
+		// Copy global matrices over to DynamicData
+		bool bComputeChanges = true;
 
-			// if the number of matrices has changed between frames, then sync previous to current
-			if (ComponentSpaceTransforms.Num() != DynamicData->PrevTransforms.Num())
-			{
-				DynamicData->SetPrevTransforms(ComponentSpaceTransforms);
-				DynamicData->ChangedCount = ComponentSpaceTransforms.Num();
-				bComputeChanges = false; // Optimization to just force all transforms as changed and skip comparison
-			}
+		// if the number of matrices has changed between frames, then sync previous to current
+		if (CompSpaceTransforms.Num() != DynamicData->PrevTransforms.Num())
+		{
+			DynamicData->SetPrevTransforms(CompSpaceTransforms);
+			DynamicData->ChangedCount = CompSpaceTransforms.Num();
+			bComputeChanges = false; // Optimization to just force all transforms as changed and skip comparison
+		}
 
-			DynamicData->SetTransforms(ComponentSpaceTransforms);
+		DynamicData->SetTransforms(CompSpaceTransforms);
 
-			// The number of transforms for current and previous should match now
-			check(DynamicData->PrevTransforms.Num() == DynamicData->Transforms.Num());
+		// The number of transforms for current and previous should match now
+		check(DynamicData->PrevTransforms.Num() == DynamicData->Transforms.Num());
 
-			if (bComputeChanges)
-			{
-				DynamicData->DetermineChanges();
-			}
+		if (bComputeChanges)
+		{
+			DynamicData->DetermineChanges();
 		}
 	}
 
@@ -2577,13 +2538,14 @@ FTransform UGeometryCollectionComponent::GetSocketTransform(FName InSocketName,
 				const int32 TransformIndex = Collection->BoneName.Find(InSocketName.ToString());
 				if (TransformIndex != INDEX_NONE)
 				{
-					if (ComponentSpaceTransforms.IsValidIndex(TransformIndex))
+					const TArray<FTransform>& CompSpaceTransforms = ComponentSpaceTransforms.RequestAllTransforms();
+					if (CompSpaceTransforms.IsValidIndex(TransformIndex))
 					{
-						SocketComponentSpaceTransform = ComponentSpaceTransforms[TransformIndex];
+						SocketComponentSpaceTransform = CompSpaceTransforms[TransformIndex];
 						const int32 ParentTransformIndex = Collection->Parent[TransformIndex];
-						if (ComponentSpaceTransforms.IsValidIndex(ParentTransformIndex))
+						if (CompSpaceTransforms.IsValidIndex(ParentTransformIndex))
 						{
-							ParentComponentSpaceTransform = ComponentSpaceTransforms[ParentTransformIndex];
+							ParentComponentSpaceTransform = CompSpaceTransforms[ParentTransformIndex];
 						}
 						bFoundSocket = true;
 					}
@@ -2932,16 +2894,14 @@ void UGeometryCollectionComponent::ResetDynamicCollection()
 		RestTransforms.Reset();
 	}
 
-	// if Rest transform have been overriden uses them to initialize the dynamic collection transforms
+	// if Reset transform have been overriden uses them to initialize the dynamic collection transforms
 	if (RestTransforms.Num() > 0)
 	{
 		SetInitialTransforms(RestTransforms);
 	}
 
-	if (RestCollection)
-	{
-		CalculateGlobalMatrices();
-	}
+	ComponentSpaceTransforms.Reset(GetCurrentTransforms().Num(), GetRootIndex());
+
 	UpdateCachedBounds();
 }
 
@@ -3221,6 +3181,12 @@ void UGeometryCollectionComponent::OnPostPhysicsSync()
 {
 	SCOPE_CYCLE_COUNTER(STAT_GCPostPhysicsSync);
 
+	// dirty the transform if the collection is
+	if (DynamicCollection && DynamicCollection->IsDirty())
+	{
+		ComponentSpaceTransforms.MarkDirty();
+	}
+
 	UpdateAttachedChildrenTransform();
 
 	if (GetIsReplicated() && GetNetMode() != ENetMode::NM_Client)
@@ -3332,9 +3298,11 @@ void UGeometryCollectionComponent::UpdateRemovalIfNeeded()
 			const FTransform InverseComponentTransform = (RestCollection->bScaleOnRemoval) ? GetComponentTransform().Inverse() : FTransform::Identity;
 
 			const TManagedArray<FTransform>* MassToLocal = nullptr;
+			const TArray<FTransform>* CompSpaceTransform = nullptr;
 			if (RestCollection->bScaleOnRemoval)
 			{
 				MassToLocal = RestCollection->GetGeometryCollection()->FindAttribute<FTransform>(MassToLocalAttributeName, FGeometryCollection::TransformGroup);
+				CompSpaceTransform = &ComponentSpaceTransforms.RequestAllTransforms();
 			}
 
 			const int32 NumTransforms = DecayFacade.GetDecayAttributeSize();
@@ -3350,7 +3318,7 @@ void UGeometryCollectionComponent::UpdateRemovalIfNeeded()
 						DynamicCollection->Transform[TransformIndex].SetScale3D(FVector::ZeroVector);
 					}
 					// do not try to get this condition out of the loop as this may cause some optimizer related issues
-					else if (RestCollection->bScaleOnRemoval && MassToLocal)
+					else if (RestCollection->bScaleOnRemoval && MassToLocal && CompSpaceTransform)
 					{
 						float ShrinkRadius = 0.0f;
 						UE::Math::TSphere<double> AccumulatedSphere;
@@ -3360,7 +3328,7 @@ void UGeometryCollectionComponent::UpdateRemovalIfNeeded()
 							ShrinkRadius = -AccumulatedSphere.W;
 						}
 
-						const FQuat LocalRotation = (InverseComponentTransform * ComponentSpaceTransforms[TransformIndex].Inverse()).GetRotation();
+						const FQuat LocalRotation = (InverseComponentTransform * (*CompSpaceTransform)[TransformIndex].Inverse()).GetRotation();
 						const FVector LocalDown = LocalRotation.RotateVector(FVector(0.f, 0.f, ShrinkRadius));
 						const FVector CenterOfMass = (*MassToLocal)[TransformIndex].GetTranslation();
 						const FVector ScaleCenter = LocalDown + CenterOfMass;
@@ -4705,32 +4673,57 @@ const TArray<FTransform>& UGeometryCollectionComponent::GetCurrentTransforms() c
 	return GetTransformArray().GetConstArray();
 }
 
-void UGeometryCollectionComponent::CalculateGlobalMatrices()
+const FTransform& UGeometryCollectionComponent::FComponentSpaceTransforms::RequestRootTransform() const
+{
+	if (!Component)
+	{
+		ensure(false); // we should not be able to reach here 
+		return FTransform::Identity;
+	}
+
+	if (Transforms.IsValidIndex(RootIndex))
+	{
+		if (bIsRootDirty)
+		{
+			Transforms[RootIndex] = Component->GetCurrentTransforms()[RootIndex];
+			bIsRootDirty = false;
+		}
+
+		return Transforms[RootIndex];
+	}
+
+	return FTransform::Identity;
+}
+
+const TArray<FTransform>& UGeometryCollectionComponent::FComponentSpaceTransforms::RequestAllTransforms() const
 {
 	SCOPE_CYCLE_COUNTER(STAT_GCCUGlobalMatrices);
 
-	// If hierarchy topology has changed, the RestTransforms is invalidated.
-	// todo(chaos) should should be able to remove this : this should be normally be properly handled by the SetRestCollection function now
-	if (RestTransforms.Num() != GetTransformArray().Num())
+	static TArray<FTransform> EmptyArray;
+	if (!Component)
 	{
-		RestTransforms.Empty();
+		ensure(false); // we should not be able to reach here 
+		return EmptyArray;
 	}
 
-	const TArray<FTransform>& CurrentTransforms = GetCurrentTransforms();
+	if (!bIsDirty)
+	{
+		return Transforms;
+	}
+
+	const TArray<FTransform>& CurrentTransforms = Component->GetCurrentTransforms();
 
 	bool bFastPath = false;
-	if (RestCollection)
+	if (Component->RestCollection)
 	{
-		const TArray<int32>& BreadthFirstTransformIndices = RestCollection->GetBreadthFirstTransformIndices();
+		const TArray<int32>& BreadthFirstTransformIndices = Component->RestCollection->GetBreadthFirstTransformIndices();
 		bFastPath = (BreadthFirstTransformIndices.Num() == CurrentTransforms.Num());
 	}
 
 	if (bFastPath)
 	{
-		ComponentSpaceTransforms.SetNumUninitialized(CurrentTransforms.Num(), false);
-
-		const TArray<int32>& BreadthFirstTransformIndices = RestCollection->GetBreadthFirstTransformIndices();
-		const TArray<int32>& ParentArray = GetParentArray().GetConstArray();
+		const TArray<int32>& BreadthFirstTransformIndices = Component->RestCollection->GetBreadthFirstTransformIndices();
+		const TArray<int32>& ParentArray = Component->GetParentArray().GetConstArray();
 
 		for (int32 Index = 0; Index < BreadthFirstTransformIndices.Num(); Index++)
 		{
@@ -4739,27 +4732,31 @@ void UGeometryCollectionComponent::CalculateGlobalMatrices()
 
 			if (ParentTransformIndex == INDEX_NONE)
 			{
-				ComponentSpaceTransforms[TransformIndex] = CurrentTransforms[TransformIndex];
+				Transforms[TransformIndex] = CurrentTransforms[TransformIndex];
 			}
 			else
 			{
-				const FTransform& ParentTransform = ComponentSpaceTransforms[ParentTransformIndex];
-				ComponentSpaceTransforms[TransformIndex] = CurrentTransforms[TransformIndex] * ParentTransform;
+				const FTransform& ParentTransform = Transforms[ParentTransformIndex];
+				Transforms[TransformIndex] = CurrentTransforms[TransformIndex] * ParentTransform;
 			}
 		}
 	}
 	else
 	{
-		GeometryCollectionAlgo::GlobalMatrices(GetCurrentTransforms(), GetParentArray(), ComponentSpaceTransforms);
+		GeometryCollectionAlgo::GlobalMatrices(Component->GetCurrentTransforms(), Component->GetParentArray(), Transforms);
 	}
 
 #if WITH_EDITOR
-	UpdateGlobalMatricesWithExplodedVectors(ComponentSpaceTransforms, *(RestCollection->GetGeometryCollection()));
+	UpdateGlobalMatricesWithExplodedVectors(Transforms, *(Component->RestCollection->GetGeometryCollection()));
 #endif
 
 #if GEOMETRY_COLLECTION_CHECK_FOR_NANS_IN_TRANSFORMS
-	CheckForNaNs(ComponentSpaceTransforms);
+	CheckForNaNs(Transforms);
 #endif
+
+	bIsDirty = false;
+	bIsRootDirty = false;
+	return Transforms;
 }
 
 int32 UGeometryCollectionComponent::GetNumMaterials() const
@@ -4985,37 +4982,25 @@ void UGeometryCollectionComponent::RefreshCustomRenderer()
 		{
 			if (RestCollection != nullptr)
 			{
-				const FTransform ComponentTransform = GetComponentTransform();
-				const int32 RootIndex = GetRootIndex();
-
-				const bool bIsBroken = DynamicCollection ? !DynamicCollection->Active[RootIndex] : false;
-				const bool bRenderRootProxy = bEnableRootProxyForCustomRenderer && !bIsBroken;
-
 				if (bUpdateRenderer)
 				{
-					RendererInterface->UpdateState(*RestCollection, ComponentTransform, !bRenderRootProxy, !bHiddenInGame);
-				}
+					const FTransform ComponentTransform = GetComponentTransform();
+					const int32 RootIndex = GetRootIndex();
 
-				if (bRenderRootProxy)
-				{
-					// even thoiugh we only need the root transform, we need to compute all the Component space transform 
-					// because the calclBounds relies on it
-					// @todo(chaos) : if this is a performance problem , we'll need to optimize the transform computation
-					CalculateGlobalMatrices();
+					const bool bIsBroken = DynamicCollection ? !DynamicCollection->Active[RootIndex] : false;
+					const bool bRenderRootProxy = bEnableRootProxyForCustomRenderer && !bIsBroken;
+
+					RendererInterface->UpdateState(*RestCollection, ComponentTransform, !bRenderRootProxy, !bHiddenInGame);
 
-					if (bUpdateRenderer)
+					if (bRenderRootProxy)
 					{
-						RendererInterface->UpdateRootTransform(*RestCollection, ComponentSpaceTransforms[RootIndex]);
+						const FTransform& CompSpaceRootTransform = ComponentSpaceTransforms.RequestRootTransform();
+						RendererInterface->UpdateRootTransform(*RestCollection, CompSpaceRootTransform);
 					}
-				}
-				else
-				{
-					// render all individual pieces
-					CalculateGlobalMatrices();
-
-					if (bUpdateRenderer)
+					else
 					{
-						RendererInterface->UpdateTransforms(*RestCollection, ComponentSpaceTransforms);
+						const TArray<FTransform>& CompSpaceTransforms = ComponentSpaceTransforms.RequestAllTransforms();
+						RendererInterface->UpdateTransforms(*RestCollection, CompSpaceTransforms);
 					}
 				}
 			}
@@ -5443,11 +5428,8 @@ FTransform UGeometryCollectionComponent::GetRootCurrentTransform() const
 	FTransform RootInitialTransform{ FTransform::Identity };
 	if (RestCollection)
 	{
-		const int32 RootIndex = RestCollection->GetRootIndex();
-		if (ComponentSpaceTransforms.IsValidIndex(RootIndex))
-		{
-			RootInitialTransform = ComponentSpaceTransforms[RootIndex] * GetComponentTransform();
-		}
+		const FTransform& CompSpaceRootTransform = ComponentSpaceTransforms.RequestRootTransform();
+		RootInitialTransform = CompSpaceRootTransform * GetComponentTransform();
 	}
 	return RootInitialTransform;
 }
@@ -5525,9 +5507,10 @@ TArray<FMatrix> UGeometryCollectionComponent::ComputeGlobalMatricesFromComponent
 	TArray<FMatrix> ComponentSpaceMatrices;
 	ComponentSpaceMatrices.SetNumUninitialized(ComponentSpaceTransforms.Num());
 
-	for (int32 TransformIndex = 0; TransformIndex < ComponentSpaceTransforms.Num(); TransformIndex++)
+	const TArray<FTransform>& CompSpaceTransforms = ComponentSpaceTransforms.RequestAllTransforms();
+	for (int32 TransformIndex = 0; TransformIndex < CompSpaceTransforms.Num(); TransformIndex++)
 	{
-		ComponentSpaceMatrices[TransformIndex] = ComponentSpaceTransforms[TransformIndex].ToMatrixWithScale();
+		ComponentSpaceMatrices[TransformIndex] = CompSpaceTransforms[TransformIndex].ToMatrixWithScale();
 	}
 	return ComponentSpaceMatrices;
 }
diff --git a/Engine/Source/Runtime/Experimental/GeometryCollectionEngine/Public/GeometryCollection/GeometryCollectionComponent.h b/Engine/Source/Runtime/Experimental/GeometryCollectionEngine/Public/GeometryCollection/GeometryCollectionComponent.h
--- a/Engine/Source/Runtime/Experimental/GeometryCollectionEngine/Public/GeometryCollection/GeometryCollectionComponent.h
+++ b/Engine/Source/Runtime/Experimental/GeometryCollectionEngine/Public/GeometryCollection/GeometryCollectionComponent.h
@@ -1112,7 +1112,7 @@ public:
 	UE_DEPRECATED(5.3, "Use GetComponentSpaceTransforms instead")
 	TArray<FMatrix> GetGlobalMatrices() { return ComputeGlobalMatricesFromComponentSpaceTransforms(); }
 
-	const TArray<FTransform> GetComponentSpaceTransforms() { return ComponentSpaceTransforms; }
+	GEOMETRYCOLLECTIONENGINE_API const TArray<FTransform>& GetComponentSpaceTransforms() { return ComponentSpaceTransforms.RequestAllTransforms(); }
 
 	GEOMETRYCOLLECTIONENGINE_API const FGeometryDynamicCollection* GetDynamicCollection() const;
 	GEOMETRYCOLLECTIONENGINE_API FGeometryDynamicCollection* GetDynamicCollection();  // TEMP HACK?
@@ -1247,7 +1247,7 @@ protected:
 
 	GEOMETRYCOLLECTIONENGINE_API Chaos::FPhysicsSolver* GetSolver(const UGeometryCollectionComponent& GeometryCollectionComponent);
 	GEOMETRYCOLLECTIONENGINE_API void CalculateLocalBounds();
-	GEOMETRYCOLLECTIONENGINE_API void CalculateGlobalMatrices();
+	//GEOMETRYCOLLECTIONENGINE_API void CalculateGlobalMatrices();
 	
 	UE_DEPRECATED(5.3, "Use ComputeBoundsFromComponentSpaceTransforms instead")
 	GEOMETRYCOLLECTIONENGINE_API FBox ComputeBoundsFromGlobalMatrices(const FMatrix& LocalToWorldWithScale, const TArray<FMatrix>& GlobalMatricesArray) const;
@@ -1380,7 +1380,49 @@ private:
 	UE_DEPRECATED(5.3, "Use ComponentSpaceTransforms instead")
 	TArray<FMatrix> GlobalMatrices;
 
-	TArray<FTransform> ComponentSpaceTransforms;
+	struct FComponentSpaceTransforms
+	{
+	public:
+		FComponentSpaceTransforms(const UGeometryCollectionComponent* InComponent = nullptr)
+			: RootIndex(INDEX_NONE)
+			, Component(InComponent)
+		{
+			bIsRootDirty = 1;
+			bIsDirty = 1;
+		}
+		
+		void Reset(int32 NumTransforms, int32 InRootIndex)
+		{
+			Transforms.SetNumUninitialized(NumTransforms);
+			RootIndex = InRootIndex;
+		}
+
+		void MarkDirty()
+		{
+			bIsRootDirty = true;
+			bIsDirty = true;
+		}
+
+		int32 Num() const { return Transforms.Num(); }
+
+		SIZE_T GetAllocatedSize() const { return Transforms.GetAllocatedSize(); }
+
+		// request all transform to be update
+		// this will trigger an update if it is still marked dirty
+		const TArray<FTransform>& RequestAllTransforms() const;
+
+		// request the root transform, this may compute it if it is still marked dirty
+		const FTransform& RequestRootTransform() const;
+
+	private:
+		int32 RootIndex;
+		mutable uint8 bIsRootDirty : 1;
+		mutable uint8 bIsDirty : 1;
+		mutable TArray<FTransform> Transforms;
+		const UGeometryCollectionComponent* Component;
+	};
+
+	FComponentSpaceTransforms ComponentSpaceTransforms;
 
 	FBox LocalBounds;
 
