Attempt to fix Arm64 build ( from CL 27902529)

diff --git a/Engine/Source/Runtime/Experimental/GeometryCollectionEngine/Private/GeometryCollection/GeometryCollectionComponent.cpp b/Engine/Source/Runtime/Experimental/GeometryCollectionEngine/Private/GeometryCollection/GeometryCollectionComponent.cpp
--- a/Engine/Source/Runtime/Experimental/GeometryCollectionEngine/Private/GeometryCollection/GeometryCollectionComponent.cpp
+++ b/Engine/Source/Runtime/Experimental/GeometryCollectionEngine/Private/GeometryCollection/GeometryCollectionComponent.cpp
@@ -1459,6 +1459,11 @@ static void DispatchGeometryCollectionCrumblingEvent(const FChaosCrumblingEvent&
 	}
 }
 
+const TArray<FTransform>& UGeometryCollectionComponent::GetComponentSpaceTransforms()
+{ 
+	return ComponentSpaceTransforms.RequestAllTransforms();
+}
+
 const FGeometryDynamicCollection* UGeometryCollectionComponent::GetDynamicCollection() const
 {
 	return DynamicCollection.Get();
diff --git a/Engine/Source/Runtime/Experimental/GeometryCollectionEngine/Public/GeometryCollection/GeometryCollectionComponent.h b/Engine/Source/Runtime/Experimental/GeometryCollectionEngine/Public/GeometryCollection/GeometryCollectionComponent.h
--- a/Engine/Source/Runtime/Experimental/GeometryCollectionEngine/Public/GeometryCollection/GeometryCollectionComponent.h
+++ b/Engine/Source/Runtime/Experimental/GeometryCollectionEngine/Public/GeometryCollection/GeometryCollectionComponent.h
@@ -1112,7 +1112,7 @@ public:
 	UE_DEPRECATED(5.3, "Use GetComponentSpaceTransforms instead")
 	TArray<FMatrix> GetGlobalMatrices() { return ComputeGlobalMatricesFromComponentSpaceTransforms(); }
 
-	GEOMETRYCOLLECTIONENGINE_API const TArray<FTransform>& GetComponentSpaceTransforms() { return ComponentSpaceTransforms.RequestAllTransforms(); }
+	GEOMETRYCOLLECTIONENGINE_API const TArray<FTransform>& GetComponentSpaceTransforms();
 
 	GEOMETRYCOLLECTIONENGINE_API const FGeometryDynamicCollection* GetDynamicCollection() const;
 	GEOMETRYCOLLECTIONENGINE_API FGeometryDynamicCollection* GetDynamicCollection();  // TEMP HACK?
