Apply RadialImpulse : Add cvar to control the impulse distribution between children for clustered  particles

diff --git a/Engine/Source/Runtime/Experimental/Chaos/Private/Chaos/PhysicsObjectInterface.cpp b/Engine/Source/Runtime/Experimental/Chaos/Private/Chaos/PhysicsObjectInterface.cpp
--- a/Engine/Source/Runtime/Experimental/Chaos/Private/Chaos/PhysicsObjectInterface.cpp
+++ b/Engine/Source/Runtime/Experimental/Chaos/Private/Chaos/PhysicsObjectInterface.cpp
@@ -19,6 +19,14 @@ namespace PhysicsObjectInterfaceCVars
 		TEXT("Modify the strain by this factor"),
 		ECVF_Default
 	);
+
+	static bool RadialImpulseDistributeToChildren = true;
+	static FAutoConsoleVariableRef CVarRadialImpulseDistributeToChildren(
+		TEXT("Chaos.Debug.RadialImpulseDistributeToChildren"),
+		RadialImpulseDistributeToChildren,
+		TEXT("When one and applied to a geometry collection cluster, the impulse will be divided equally betweemn all the children"),
+		ECVF_Default
+	);
 }
 
 namespace
@@ -922,17 +930,21 @@ namespace Chaos
 									continue;
 								}
 
-								for (FPBDRigidParticleHandle* ChildHandle : *ChildrenHandles)
+								float VelocityRatio = 1.0f;
+								if (PhysicsObjectInterfaceCVars::RadialImpulseDistributeToChildren)
 								{
-									float VelocityRatio = 1.0f / static_cast<float>(ChildrenHandles->Num());
+									VelocityRatio = 1.0f / static_cast<float>(ChildrenHandles->Num());
+								}
 
-									float FalloffStrength = AddRadialImpulseHelper(ChildHandle, Origin, Radius, Strength, Falloff, VelocityRatio, bInvalidate);
+								for (FPBDRigidParticleHandle* ChildHandle : *ChildrenHandles)
+								{
+									const float FalloffStrength = AddRadialImpulseHelper(ChildHandle, Origin, Radius, Strength, Falloff, VelocityRatio, bInvalidate);
 
 									//to do: remove cvar when material system is in place and densities are updated
-									FalloffStrength = PhysicsObjectInterfaceCVars::StrainModifier * FalloffStrength;
-									if (FalloffStrength > 0)
+									const float StrainToApply = PhysicsObjectInterfaceCVars::StrainModifier * FalloffStrength;
+									if (StrainToApply > 0)
 									{
-										Clustering.SetExternalStrain(ChildHandle->CastToClustered(), FalloffStrength);
+										Clustering.SetExternalStrain(ChildHandle->CastToClustered(), StrainToApply);
 									}
 								}
 							}
