Geometry collection : Fix few crashes resulsting from recent changes

diff --git a/Engine/Source/Runtime/Experimental/GeometryCollectionEngine/Private/GeometryCollection/GeometryCollectionComponent.cpp b/Engine/Source/Runtime/Experimental/GeometryCollectionEngine/Private/GeometryCollection/GeometryCollectionComponent.cpp
--- a/Engine/Source/Runtime/Experimental/GeometryCollectionEngine/Private/GeometryCollection/GeometryCollectionComponent.cpp
+++ b/Engine/Source/Runtime/Experimental/GeometryCollectionEngine/Private/GeometryCollection/GeometryCollectionComponent.cpp
@@ -650,7 +650,10 @@ FBox UGeometryCollectionComponent::ComputeBounds(const FTransform& LocalToWorldW
 	if (RestCollection)
 	{
 		const TArray<FTransform>& CompSpaceTransforms = ComponentSpaceTransforms.RequestAllTransforms();
-		BoundingBox = ComputeBoundsFromComponentSpaceTransforms(LocalToWorldWithScale, CompSpaceTransforms);
+		if (CompSpaceTransforms.Num() > 0)
+		{
+			BoundingBox = ComputeBoundsFromComponentSpaceTransforms(LocalToWorldWithScale, CompSpaceTransforms);
+		}
 	}
 	return BoundingBox;
 }
@@ -2905,7 +2908,7 @@ void UGeometryCollectionComponent::ResetDynamicCollection()
 		SetInitialTransforms(RestTransforms);
 	}
 
-	ComponentSpaceTransforms.Reset(GetCurrentTransforms().Num(), GetRootIndex());
+	ComponentSpaceTransforms.Reset(NumRestCollectionTransforms, GetRootIndex());
 
 	UpdateCachedBounds();
 }
@@ -4705,9 +4708,8 @@ const TArray<FTransform>& UGeometryCollectionComponent::FComponentSpaceTransform
 	SCOPE_CYCLE_COUNTER(STAT_GCCUGlobalMatrices);
 
 	static TArray<FTransform> EmptyArray;
-	if (!Component)
+	if (!Component || Transforms.Num() == 0)
 	{
-		ensure(false); // we should not be able to reach here 
 		return EmptyArray;
 	}
 
