Geometry collection : fix rest transform not properly reflected in the viewport when set in a construction script

diff --git a/Engine/Source/Runtime/Experimental/GeometryCollectionEngine/Private/GeometryCollection/GeometryCollectionComponent.cpp b/Engine/Source/Runtime/Experimental/GeometryCollectionEngine/Private/GeometryCollection/GeometryCollectionComponent.cpp
--- a/Engine/Source/Runtime/Experimental/GeometryCollectionEngine/Private/GeometryCollection/GeometryCollectionComponent.cpp
+++ b/Engine/Source/Runtime/Experimental/GeometryCollectionEngine/Private/GeometryCollection/GeometryCollectionComponent.cpp
@@ -2371,7 +2371,9 @@ void UGeometryCollectionComponent::GetHiddenTransforms(TArray<bool>& OutHiddenTr
 void UGeometryCollectionComponent::GetRestTransforms(TArray<FMatrix44f>& OutRestTransforms) const
 {
 	TArray<FMatrix> RestMatrices;
-	GeometryCollectionAlgo::GlobalMatrices(RestCollection->GetGeometryCollection()->Transform, RestCollection->GetGeometryCollection()->Parent, RestMatrices);
+	const TArray<FTransform>& RestCollectionTransforms = RestCollection->GetGeometryCollection()->Transform.GetConstArray();
+	const TArray<FTransform>& RestTransformToUse = (RestTransforms.Num() == RestCollectionTransforms.Num()) ? RestTransforms : RestCollectionTransforms;
+	GeometryCollectionAlgo::GlobalMatrices(RestTransformToUse, RestCollection->GetGeometryCollection()->Parent, RestMatrices);
 #if WITH_EDITOR
 	UpdateGlobalMatricesWithExplodedVectors(RestMatrices, *(RestCollection->GetGeometryCollection()));
 #endif
diff --git a/Engine/Source/Runtime/Experimental/GeometryCollectionEngine/Public/GeometryCollection/GeometryCollectionComponent.h b/Engine/Source/Runtime/Experimental/GeometryCollectionEngine/Public/GeometryCollection/GeometryCollectionComponent.h
--- a/Engine/Source/Runtime/Experimental/GeometryCollectionEngine/Public/GeometryCollection/GeometryCollectionComponent.h
+++ b/Engine/Source/Runtime/Experimental/GeometryCollectionEngine/Public/GeometryCollection/GeometryCollectionComponent.h
@@ -1395,6 +1395,7 @@ private:
 		{
 			Transforms.SetNumUninitialized(NumTransforms);
 			RootIndex = InRootIndex;
+			MarkDirty();
 		}
 
 		void MarkDirty()
