Chaos: Alter material fetch during cluster creation to not call GetMaterialData directly as that causes a material array to be allocated which we do not need in this case and would be slower and more memory intensive.
Chaos: Fixed ensure in cluster unions that have children with no materials.

diff --git a/Engine/Source/Runtime/Experimental/Chaos/Private/Chaos/ClusterUnionManager.cpp b/Engine/Source/Runtime/Experimental/Chaos/Private/Chaos/ClusterUnionManager.cpp
--- a/Engine/Source/Runtime/Experimental/Chaos/Private/Chaos/ClusterUnionManager.cpp
+++ b/Engine/Source/Runtime/Experimental/Chaos/Private/Chaos/ClusterUnionManager.cpp
@@ -626,11 +626,10 @@ namespace Chaos
 						const FShapesArray& Shapes = Particle->ShapesArray();
 						for(const TUniquePtr<FPerShapeData>& ShapeData : Shapes)
 						{
-							const FMaterialData& MaterialData = ShapeData->GetMaterialData();
-
-							if(MaterialData.Materials.Num() > 0)
+							const int32 NumShapeMaterials = ShapeData->NumMaterials();
+							if(NumShapeMaterials > 0)
 							{
-								MatHandle = MaterialData.Materials[0];
+								MatHandle = ShapeData->GetMaterial(0);
 								break;
 							}
 						}
@@ -661,7 +660,7 @@ namespace Chaos
 				MEvolution.SetParticleObjectState(Cluster->InternalCluster, Chaos::EObjectStateType::Sleeping);
 			}
 
-			if(ensure(MatHandle.InnerHandle.IsValid()))
+			if(MatHandle.InnerHandle.IsValid())
 			{
 				for(const TUniquePtr<FPerShapeData>& ShapeData : Cluster->InternalCluster->ShapesArray())
 				{
