Geometry collection : fix assert to not trigger is nanite material is masked

diff --git a/Engine/Source/Runtime/Experimental/GeometryCollectionEngine/Private/GeometryCollection/GeometryCollectionSceneProxy.cpp b/Engine/Source/Runtime/Experimental/GeometryCollectionEngine/Private/GeometryCollection/GeometryCollectionSceneProxy.cpp
--- a/Engine/Source/Runtime/Experimental/GeometryCollectionEngine/Private/GeometryCollection/GeometryCollectionSceneProxy.cpp
+++ b/Engine/Source/Runtime/Experimental/GeometryCollectionEngine/Private/GeometryCollection/GeometryCollectionSceneProxy.cpp
@@ -1088,7 +1088,7 @@ FNaniteGeometryCollectionSceneProxy::FNaniteGeometryCollectionSceneProxy(UGeomet
 		check(MaterialInterface != nullptr);
 
 		// Should always be opaque blend mode here.
-		check(IsOpaqueBlendMode(*MaterialInterface));
+		check(IsOpaqueOrMaskedBlendMode(*MaterialInterface));
 
 		MaterialSections[SectionIndex].ShadingMaterialProxy = MaterialInterface->GetRenderProxy();
 		MaterialSections[SectionIndex].RasterMaterialProxy  = MaterialInterface->GetRenderProxy(); // TODO: PROG_RASTER (Implement programmable raster support)
