geometry collection fully decayed event needs to actually wait for particles to decay

diff --git a/Engine/Source/Runtime/Experimental/GeometryCollectionEngine/Private/GeometryCollection/GeometryCollectionComponent.cpp b/Engine/Source/Runtime/Experimental/GeometryCollectionEngine/Private/GeometryCollection/GeometryCollectionComponent.cpp
--- a/Engine/Source/Runtime/Experimental/GeometryCollectionEngine/Private/GeometryCollection/GeometryCollectionComponent.cpp
+++ b/Engine/Source/Runtime/Experimental/GeometryCollectionEngine/Private/GeometryCollection/GeometryCollectionComponent.cpp
@@ -1,4 +1,5 @@
 // Copyright Epic Games, Inc. All Rights Reserved.
+#pragma optimize("", off)
 #include "GeometryCollection/GeometryCollectionComponent.h"
 
 #include "AI/Navigation/NavCollisionBase.h"
@@ -2690,12 +2691,14 @@ void UGeometryCollectionComponent::CheckFullyDecayed()
 
 	if (DynamicCollection && PhysicsProxy)
 	{
-		bool bFullyDecayed = true;
 		FGeometryCollectionDynamicStateFacade DynamicStateFacade(*DynamicCollection);
+		FGeometryCollectionDecayDynamicFacade DecayFacade(*DynamicCollection);
 
-		if (DynamicStateFacade.IsValid())
+		if (DynamicStateFacade.IsValid() && DecayFacade.IsValid())
 		{
+			bool bFullyDecayed = true;
 			const int32 NumTransforms = DynamicCollection->NumElements(FGeometryCollection::TransformGroup);
+
 			for (int32 TransformIdx = 0; TransformIdx < NumTransforms; ++TransformIdx)
 			{
 				// If we didn't create a particle for this transform, we shouldn't consider this particle.
@@ -2717,6 +2720,12 @@ void UGeometryCollectionComponent::CheckFullyDecayed()
 					bFullyDecayed = false;
 					break;
 				}
+
+				if (DecayFacade.GetDecay(TransformIdx) < 1.0f)
+				{
+					bFullyDecayed = false;
+					break;
+				}
 			}
 
 			if (bFullyDecayed)
