geometry collection break cluster shouldn't require an explicit removal from the parent cluster union

diff --git a/Engine/Source/Runtime/Experimental/Chaos/Private/PhysicsProxy/GeometryCollectionPhysicsProxy.cpp b/Engine/Source/Runtime/Experimental/Chaos/Private/PhysicsProxy/GeometryCollectionPhysicsProxy.cpp
--- a/Engine/Source/Runtime/Experimental/Chaos/Private/PhysicsProxy/GeometryCollectionPhysicsProxy.cpp
+++ b/Engine/Source/Runtime/Experimental/Chaos/Private/PhysicsProxy/GeometryCollectionPhysicsProxy.cpp
@@ -1,5 +1,4 @@
 // Copyright Epic Games, Inc. All Rights Reserved.
-
 #include "PhysicsProxy/GeometryCollectionPhysicsProxy.h"
 #include "PhysicsProxy/FieldSystemProxyHelper.h"
 
@@ -2179,8 +2178,6 @@ void FGeometryCollectionPhysicsProxy::BreakClusters_External(TArray<FGeometryCol
 			{
 				if (Chaos::FPBDRigidClusteredParticleHandle* ClusteredHandle = FindClusteredParticleHandleByItemIndex_Internal(ItemIndex))
 				{
-					// Pre-emptively remove the particle from cluster unions it might be in.
-					Clustering.GetClusterUnionManager().HandleRemoveOperationWithClusterLookup({ ClusteredHandle }, Chaos::EClusterUnionOperationTiming::Immediate);
 					Clustering.BreakCluster(ClusteredHandle);
 				}
 			}
