Add Cvar to allow creation of physics state in non-PIE editor mode

diff --git a/Engine/Source/Runtime/Experimental/GeometryCollectionEngine/Private/GeometryCollection/GeometryCollectionComponent.cpp b/Engine/Source/Runtime/Experimental/GeometryCollectionEngine/Private/GeometryCollection/GeometryCollectionComponent.cpp
--- a/Engine/Source/Runtime/Experimental/GeometryCollectionEngine/Private/GeometryCollection/GeometryCollectionComponent.cpp
+++ b/Engine/Source/Runtime/Experimental/GeometryCollectionEngine/Private/GeometryCollection/GeometryCollectionComponent.cpp
@@ -127,6 +127,11 @@ FAutoConsoleVariableRef CVarGeometryCollectionRemovalTimerMultiplier(TEXT("p.Cha
 // temporary cvar , should be removed when the root event is no longer no necessary
 bool GeometryCollectionEmitRootBreakingEvent = false;
 FAutoConsoleVariableRef CVarGeometryCollectionEmitRootBreakingEvent(TEXT("p.Chaos.GC.EmitRootBreakingEvent"), GeometryCollectionEmitRootBreakingEvent, TEXT("When true send a breaking event when root is breaking"));
+
+bool GeometryCollectionCreatePhysicsStateInEditor = false;
+FAutoConsoleVariableRef CVarGeometryCollectionCreatePhysicsStateInEditor(TEXT("p.Chaos.GC.CreatePhysicsStateInEditor"), GeometryCollectionCreatePhysicsStateInEditor, TEXT("when on , physics state for a GC will be create in editor ( non PIE )"));
+
+
 DEFINE_LOG_CATEGORY_STATIC(UGCC_LOG, Error, All);
 
 extern FGeometryCollectionDynamicDataPool GDynamicDataPool;
@@ -2845,7 +2850,7 @@ void UGeometryCollectionComponent::ResetDynamicCollection()
 	bCreateDynamicCollection = false;
 	if (UWorld* World = GetWorld())
 	{
-		if(World->IsGameWorld())
+		if(World->IsGameWorld() || GeometryCollectionCreatePhysicsStateInEditor)
 		{
 			bCreateDynamicCollection = true;
 		}
@@ -2945,7 +2950,7 @@ void UGeometryCollectionComponent::OnCreatePhysicsState()
 			RestCollectionMutable->CreateSimulationDataIfNeeded();
 		}
 #endif
-		const bool bValidWorld = GetWorld() && (GetWorld()->IsGameWorld() || GetWorld()->IsPreviewWorld());
+		const bool bValidWorld = GetWorld() && (GetWorld()->IsGameWorld() || GetWorld()->IsPreviewWorld() || GeometryCollectionCreatePhysicsStateInEditor);
 		const bool bValidCollection = DynamicCollection && DynamicCollection->Transform.Num() > 0;
 		if (bValidWorld && bValidCollection)
 		{
