various cluster union changes to help ensure that less pieces of cluster unions will be left floating around
- moved cluster union connectivity checks to happen when pieces get removed from the cluster union rather than only on break; this is more robust as there are pathways to removing particles from a cluster union without breaking
- let gameplay override whether or not a certain particle is a 'main' particle just in case the GT knows when things are "destroyed" but that didn't get translated into a break on the PT

diff --git a/Engine/Source/Runtime/Engine/Classes/PhysicsEngine/ClusterUnionComponent.h b/Engine/Source/Runtime/Engine/Classes/PhysicsEngine/ClusterUnionComponent.h
--- a/Engine/Source/Runtime/Engine/Classes/PhysicsEngine/ClusterUnionComponent.h
+++ b/Engine/Source/Runtime/Engine/Classes/PhysicsEngine/ClusterUnionComponent.h
@@ -187,6 +187,9 @@ public:
 
 	ENGINE_API const FSpatialAcceleration* GetSpatialAcceleration() const;
 
+	ENGINE_API TArray<int32> GetAddedBoneIdsForComponent(UPrimitiveComponent* Component) const;
+	ENGINE_API void ChangeIfComponentBonesAreMainParticle(UPrimitiveComponent* Component, const TArray<int32>& BoneIds, bool bIsMain);
+
 	friend class UClusterUnionReplicatedProxyComponent;
 	friend class UModularVehicleBaseComponent;
 protected:
@@ -195,8 +198,6 @@ protected:
 	UFUNCTION()
 	ENGINE_API void ForceSetChildToParent(UPrimitiveComponent* InComponent, const TArray<int32>& BoneIds, const TArray<FTransform>& ChildToParent);
 
-	ENGINE_API TArray<int32> GetAddedBoneIdsForComponent(UPrimitiveComponent* Component) const;
-
 	ENGINE_API void BroadcastComponentAddedEvents(UPrimitiveComponent* ChangedComponent, const TSet<int32>& BoneIds, bool bIsNew);
 	ENGINE_API void BroadcastComponentRemovedEvents(UPrimitiveComponent* ChangedComponent);
 
diff --git a/Engine/Source/Runtime/Engine/Private/PhysicsEngine/ClusterUnionComponent.cpp b/Engine/Source/Runtime/Engine/Private/PhysicsEngine/ClusterUnionComponent.cpp
--- a/Engine/Source/Runtime/Engine/Private/PhysicsEngine/ClusterUnionComponent.cpp
+++ b/Engine/Source/Runtime/Engine/Private/PhysicsEngine/ClusterUnionComponent.cpp
@@ -1555,3 +1555,19 @@ bool UClusterUnionComponent::DoCustomNavigableGeometryExport(FNavigableGeometryE
 	);
 	return bHasData;
 }
+
+void UClusterUnionComponent::ChangeIfComponentBonesAreMainParticle(UPrimitiveComponent* Component, const TArray<int32>& BoneIds, bool bIsMain)
+{
+	if (!PhysicsProxy || BoneIds.IsEmpty() || !Component)
+	{
+		return;
+	}
+
+	TArray<Chaos::FPhysicsObjectHandle> PhysicsObjects = GetAllPhysicsObjectsById(Component, BoneIds);
+	if (PhysicsObjects.IsEmpty())
+	{
+		return;
+	}
+
+	PhysicsProxy->ChangeMainParticleStatus_External(PhysicsObjects, bIsMain);
+}
diff --git a/Engine/Source/Runtime/Experimental/Chaos/Private/Chaos/ClusterUnionManager.cpp b/Engine/Source/Runtime/Experimental/Chaos/Private/Chaos/ClusterUnionManager.cpp
--- a/Engine/Source/Runtime/Experimental/Chaos/Private/Chaos/ClusterUnionManager.cpp
+++ b/Engine/Source/Runtime/Experimental/Chaos/Private/Chaos/ClusterUnionManager.cpp
@@ -420,6 +420,7 @@ namespace Chaos
 
 		if (!PendingClusterIndexOperations.IsEmpty())
 		{
+			TArray<FClusterUnionIndex> ClustersRequiringConnectivityChecks;
 			for (TPair<FClusterUnionIndex, FClusterOpMap>& OpMap : PendingClusterIndexOperations)
 			{
 				// Is this sort necessary? Better to be safe than sorry. Since we need to guarantee that the UpdateChildToParent happens after add.
@@ -430,6 +431,7 @@ namespace Chaos
 					}
 				);
 
+
 				for (TPair<EClusterUnionOperation, TArray<FPBDRigidParticleHandle*>>& Op : OpMap.Value)
 				{
 					switch (Op.Key)
@@ -439,6 +441,7 @@ namespace Chaos
 						HandleAddOperation(OpMap.Key, Op.Value, Op.Key == EClusterUnionOperation::AddReleased);
 						break;
 					case EClusterUnionOperation::Remove:
+						ClustersRequiringConnectivityChecks.Add(OpMap.Key);
 						HandleRemoveOperation(OpMap.Key, Op.Value, EClusterUnionOperationTiming::Defer);
 						break;
 					case EClusterUnionOperation::UpdateChildToParent:
@@ -449,6 +452,18 @@ namespace Chaos
 			}
 			PendingClusterIndexOperations.Empty();
 
+			// Needs to be after the emptying of the cluster index operations because otherwise it might remove the removals
+			// added in by the connectivity check.
+			for (FClusterUnionIndex UnionIndex : ClustersRequiringConnectivityChecks)
+			{
+				if (FClusterUnion* ClusterUnion = FindClusterUnion(UnionIndex))
+				{
+					// TODO: Can probably argue that cluster union connectivity should be moved into the
+					// cluster union manager instead?
+					MClustering.HandleConnectivityOnReleaseClusterParticle(ClusterUnion->InternalCluster, false);
+				}
+			}
+
 			if (FRewindData* RewindData = MEvolution.GetRewindData())
 			{
 				// Temp, since rewind doesn't handle altered cluster unions we tell it to never rewind past this frame for now, only relevant when resimulation is used.
diff --git a/Engine/Source/Runtime/Experimental/Chaos/Private/Chaos/PBDRigidClustering.cpp b/Engine/Source/Runtime/Experimental/Chaos/Private/Chaos/PBDRigidClustering.cpp
--- a/Engine/Source/Runtime/Experimental/Chaos/Private/Chaos/PBDRigidClustering.cpp
+++ b/Engine/Source/Runtime/Experimental/Chaos/Private/Chaos/PBDRigidClustering.cpp
@@ -1099,18 +1099,6 @@ namespace Chaos
 					RemoveNodeConnections(Child);
 				}
 				ActivatedChildren.Append(HandleConnectivityOnReleaseClusterParticle(ClusteredParticle, bCreateNewClusters));
-
-				// Every removal from a cluster union can trigger connectivity changes as well.
-				if (!bIsClusterUnion && bCheckForInterclusterEdgesOnRelease)
-				{
-					for (FClusterUnionIndex UnionIndex : ClusterUnionsToConsiderForConnectivity)
-					{
-						if (FClusterUnion* ClusterUnion = ClusterUnionManager.FindClusterUnion(UnionIndex))
-						{
-							ActivatedChildren.Append(HandleConnectivityOnReleaseClusterParticle(ClusterUnion->InternalCluster, false));
-						}
-					}
-				}
 			}
 
 			for (FPBDRigidParticleHandle* Child : ActivatedChildren)
diff --git a/Engine/Source/Runtime/Experimental/Chaos/Private/PhysicsProxy/ClusterUnionPhysicsProxy.cpp b/Engine/Source/Runtime/Experimental/Chaos/Private/PhysicsProxy/ClusterUnionPhysicsProxy.cpp
--- a/Engine/Source/Runtime/Experimental/Chaos/Private/PhysicsProxy/ClusterUnionPhysicsProxy.cpp
+++ b/Engine/Source/Runtime/Experimental/Chaos/Private/PhysicsProxy/ClusterUnionPhysicsProxy.cpp
@@ -685,4 +685,34 @@ namespace Chaos
 			}
 		);
 	}
+
+	void FClusterUnionPhysicsProxy::ChangeMainParticleStatus_External(const TArray<FPhysicsObjectHandle>& Objects, bool bIsMain)
+	{
+		if (!Solver)
+		{
+			return;
+		}
+
+		Solver->EnqueueCommandImmediate(
+			[this, Objects, bIsMain]() mutable
+			{
+				FReadPhysicsObjectInterface_Internal Interface = FPhysicsObjectInternalInterface::GetRead();
+				TArray<FPBDRigidParticleHandle*> Particles = Interface.GetAllRigidParticles(Objects);
+
+				FPBDRigidsEvolutionGBF& Evolution = *static_cast<FPBDRigidsSolver*>(Solver)->GetEvolution();
+				FClusterUnionManager& ClusterUnionManager = Evolution.GetRigidClustering().GetClusterUnionManager();
+
+				if (FClusterUnion* Union = ClusterUnionManager.FindClusterUnion(ClusterUnionIndex))
+				{
+					for (FPBDRigidParticleHandle* Particle : Particles)
+					{
+						if (FClusterUnionParticleProperties* Props = Union->ChildProperties.Find(Particle))
+						{
+							Props->bIsAuxiliaryParticle = !bIsMain;
+						}
+					}
+				}
+			}
+		);
+	}
 }
\ No newline at end of file
diff --git a/Engine/Source/Runtime/Experimental/Chaos/Public/PhysicsProxy/ClusterUnionPhysicsProxy.h b/Engine/Source/Runtime/Experimental/Chaos/Public/PhysicsProxy/ClusterUnionPhysicsProxy.h
--- a/Engine/Source/Runtime/Experimental/Chaos/Public/PhysicsProxy/ClusterUnionPhysicsProxy.h
+++ b/Engine/Source/Runtime/Experimental/Chaos/Public/PhysicsProxy/ClusterUnionPhysicsProxy.h
@@ -113,6 +113,7 @@ namespace Chaos
 		CHAOS_API void SetAngularVelocity_External(const FVector& W);
 		CHAOS_API void SetChildToParent_External(FPhysicsObjectHandle Child, const FTransform& RelativeTransform, bool bLock);
 		CHAOS_API void BulkSetChildToParent_External(const TArray<FPhysicsObjectHandle>& Objects, const TArray<FTransform>& Transforms, bool bLock);
+		CHAOS_API void ChangeMainParticleStatus_External(const TArray<FPhysicsObjectHandle>& Objects, bool bIsMain);
 
 		//
 		// These functions take care of marshaling data back and forth between the game thread
