add ability for cluster unions to either set or inherit their childrens' gravity group index

diff --git a/Engine/Source/Runtime/Engine/Classes/PhysicsEngine/ClusterUnionComponent.h b/Engine/Source/Runtime/Engine/Classes/PhysicsEngine/ClusterUnionComponent.h
--- a/Engine/Source/Runtime/Engine/Classes/PhysicsEngine/ClusterUnionComponent.h
+++ b/Engine/Source/Runtime/Engine/Classes/PhysicsEngine/ClusterUnionComponent.h
@@ -209,6 +209,11 @@ private:
 	// be specified in the editor and never change.
 	UPROPERTY(EditAnywhere, Category = "Cluster Union")
 	TArray<FComponentReference> ClusteredComponentsReferences;
+
+	// If set to a value not equal to -1, will manually set the cluster union's gravity group
+	// instead of automatically inheriting it from its children particles.
+	UPROPERTY(EditAnywhere, Category = "Cluster Union")
+	int32 GravityGroupIndexOverride;
 	
 	// We need to keep track of the mapping of primitive components to physics objects.
 	// This way we know the right physics objects to pass when removing the component (because
diff --git a/Engine/Source/Runtime/Engine/Private/PhysicsEngine/ClusterUnionComponent.cpp b/Engine/Source/Runtime/Engine/Private/PhysicsEngine/ClusterUnionComponent.cpp
--- a/Engine/Source/Runtime/Engine/Private/PhysicsEngine/ClusterUnionComponent.cpp
+++ b/Engine/Source/Runtime/Engine/Private/PhysicsEngine/ClusterUnionComponent.cpp
@@ -74,6 +74,7 @@ UClusterUnionComponent::UClusterUnionComponent(const FObjectInitializer& ObjectI
 #if WITH_EDITORONLY_DATA
 	bVisualizeComponent = true;
 #endif
+	GravityGroupIndexOverride = INDEX_NONE;
 }
 
 FPhysScene_Chaos* UClusterUnionComponent::GetChaosScene() const
@@ -616,6 +617,7 @@ void UClusterUnionComponent::OnCreatePhysicsState()
 	InitData.ActorId = GetOwner()->GetUniqueID();
 	InitData.ComponentId = GetUniqueID();
 	InitData.InitialTransform = GetComponentTransform();
+	InitData.GravityGroupOverride = GravityGroupIndexOverride;
 
 	// Client needs to be set to unbreakable so the server is authoritative.
 	InitData.bUnbreakable = !bHasAuthority;
diff --git a/Engine/Source/Runtime/Experimental/Chaos/Private/Chaos/ClusterUnionManager.cpp b/Engine/Source/Runtime/Experimental/Chaos/Private/Chaos/ClusterUnionManager.cpp
--- a/Engine/Source/Runtime/Experimental/Chaos/Private/Chaos/ClusterUnionManager.cpp
+++ b/Engine/Source/Runtime/Experimental/Chaos/Private/Chaos/ClusterUnionManager.cpp
@@ -287,6 +287,11 @@ namespace Chaos
 		NewUnion.Parameters = Parameters;
 		NewUnion.ClusterUnionParameters = ClusterUnionParameters;
 
+		if (ClusterUnionParameters.GravityGroupOverride != INDEX_NONE)
+		{
+			NewUnion.InternalCluster->SetGravityGroupIndex(ClusterUnionParameters.GravityGroupOverride);
+		}
+
 		// Some parameters aren't relevant after creation.
 		NewUnion.ClusterUnionParameters.UniqueIndex = nullptr;
 
@@ -682,6 +687,11 @@ namespace Chaos
 					ShapeData->SetMaterial(MatHandle);
 				}
 			}
+
+			if (!Cluster->IsGravityOverrideSet() && !FinalParticlesToAdd.IsEmpty())
+			{
+				Cluster->InternalCluster->SetGravityGroupIndex(FinalParticlesToAdd[0]->GravityGroupIndex());
+			}
 		}
 	}
 
diff --git a/Engine/Source/Runtime/Experimental/Chaos/Private/PhysicsProxy/ClusterUnionPhysicsProxy.cpp b/Engine/Source/Runtime/Experimental/Chaos/Private/PhysicsProxy/ClusterUnionPhysicsProxy.cpp
--- a/Engine/Source/Runtime/Experimental/Chaos/Private/PhysicsProxy/ClusterUnionPhysicsProxy.cpp
+++ b/Engine/Source/Runtime/Experimental/Chaos/Private/PhysicsProxy/ClusterUnionPhysicsProxy.cpp
@@ -113,6 +113,7 @@ namespace Chaos
 		ClusterUnionParameters.UniqueIndex = &UniqueIndex;
 		ClusterUnionParameters.ActorId = InitData.ActorId;
 		ClusterUnionParameters.ComponentId = InitData.ComponentId;
+		ClusterUnionParameters.GravityGroupOverride = InitData.GravityGroupOverride;
 
 		ClusterUnionIndex = ClusterUnionManager.CreateNewClusterUnion(ClusterParameters, ClusterUnionParameters);
 		if (FClusterUnion* ClusterUnion = ClusterUnionManager.FindClusterUnion(ClusterUnionIndex); ensure(ClusterUnion != nullptr))
diff --git a/Engine/Source/Runtime/Experimental/Chaos/Public/Chaos/ClusterUnionManager.h b/Engine/Source/Runtime/Experimental/Chaos/Public/Chaos/ClusterUnionManager.h
--- a/Engine/Source/Runtime/Experimental/Chaos/Public/Chaos/ClusterUnionManager.h
+++ b/Engine/Source/Runtime/Experimental/Chaos/Public/Chaos/ClusterUnionManager.h
@@ -65,6 +65,7 @@ namespace Chaos
 		const FUniqueIdx* UniqueIndex = nullptr;
 		uint32 ActorId = 0;
 		uint32 ComponentId = 0;
+		int32 GravityGroupOverride = INDEX_NONE;
 	};
 
 	struct FClusterUnionParticleProperties
@@ -118,6 +119,8 @@ namespace Chaos
 		// Pending particles that need to be added into the connectivity graph.
 		TArray<TPair<FPBDRigidParticleHandle*, EClusterUnionConnectivityOperation>> PendingConnectivityOperations;
 
+		bool IsGravityOverrideSet() const { return ClusterUnionParameters.GravityGroupOverride != INDEX_NONE; }
+
 		const TArray<FPBDRigidParticleHandle*>& GetPendingGeometryOperationParticles(EClusterUnionGeometryOperation Op) const;
 		void AddPendingGeometryOperation(EClusterUnionGeometryOperation Op, FPBDRigidParticleHandle* Particle);
 		void ClearAllPendingGeometryOperations();
diff --git a/Engine/Source/Runtime/Experimental/Chaos/Public/PhysicsProxy/ClusterUnionPhysicsProxy.h b/Engine/Source/Runtime/Experimental/Chaos/Public/PhysicsProxy/ClusterUnionPhysicsProxy.h
--- a/Engine/Source/Runtime/Experimental/Chaos/Public/PhysicsProxy/ClusterUnionPhysicsProxy.h
+++ b/Engine/Source/Runtime/Experimental/Chaos/Public/PhysicsProxy/ClusterUnionPhysicsProxy.h
@@ -34,6 +34,7 @@ namespace Chaos
 		bool bCheckConnectivity = true;
 		bool bUnbreakable = false;
 		bool bGenerateConnectivityEdges = true;
+		int32 GravityGroupOverride = INDEX_NONE;
 	};
 
 	/**
