geometry collection : make the following properties to be editable at runtime ( even if the phsyics state is already created )
- DamageModel
- DamagePropagationdata
- UseMaterialDamageModifiers
- MaterialOverrideMassScaleMultiplier
- GravityGroupIndex

diff --git a/Engine/Source/Runtime/Experimental/Chaos/Private/PhysicsProxy/GeometryCollectionPhysicsProxy.cpp b/Engine/Source/Runtime/Experimental/Chaos/Private/PhysicsProxy/GeometryCollectionPhysicsProxy.cpp
--- a/Engine/Source/Runtime/Experimental/Chaos/Private/PhysicsProxy/GeometryCollectionPhysicsProxy.cpp
+++ b/Engine/Source/Runtime/Experimental/Chaos/Private/PhysicsProxy/GeometryCollectionPhysicsProxy.cpp
@@ -463,9 +463,11 @@ FGeometryCollectionPhysicsProxy::FGeometryCollectionPhysicsProxy(
 
 	, GameThreadCollection(GameThreadCollectionIn)
 	, GameThreadPerFrameData(SimulationParameters)
+	, MaterialOverrideMassScaleMultiplierChange(0)
 	, bIsPhysicsThreadWorldTransformDirty(false)
 	, bIsCollisionFilterDataDirty(false)
 	, bIsDamageThresholdDataDirty(false)
+	, bIsGravityGroupIndexDirty(false)
 	, CollectorGuid(InCollectorGuid)
 {
 	// We rely on a guarded buffer.
@@ -2808,8 +2810,54 @@ void FGeometryCollectionPhysicsProxy::SetDamageThresholds_External(const TArray<
 	SetProxyDirty_External();
 }
 
+void FGeometryCollectionPhysicsProxy::SetDamagePropagationData_External(bool bEnabled, float BreakDamagePropagationFactor, float ShockDamagePropagationFactor)
+{
+	check(IsInGameThread());
+	FGeometryCollectioPerFrameData::FDamagePropagationData Data;
+	Data.bEnabled = bEnabled;
+	Data.BreakDamagePropagationFactor = BreakDamagePropagationFactor;
+	Data.ShockDamagePropagationFactor = ShockDamagePropagationFactor;
+	GameThreadPerFrameData.DamagePropagationData = Data;
+	SetProxyDirty_External();
+}
+
+void FGeometryCollectionPhysicsProxy::SetDamageModel_External(EDamageModelTypeEnum DamageModel)
+{
+	check(IsInGameThread());
+	GameThreadPerFrameData.DamageModel = DamageModel;
+	SetProxyDirty_External();
+}
+
+void FGeometryCollectionPhysicsProxy::SetUseMaterialDamageModifiers_External(bool bUseMaterialDamageModifiers)
+{
+	check(IsInGameThread());
+	GameThreadPerFrameData.bUseMaterialDamageModifiers = bUseMaterialDamageModifiers;
+	SetProxyDirty_External();
+}
+
+void FGeometryCollectionPhysicsProxy::SetMaterialOverrideMassScaleMultiplier_External(float InMultiplier)
+{
+	check(IsInGameThread());
+	GameThreadPerFrameData.MaterialOverrideMassScaleMultiplier = InMultiplier;
+	SetProxyDirty_External();
+}
+
+void FGeometryCollectionPhysicsProxy::SetGravityGroupIndex_External(int32 GravityGroupIndex)
+{
+	check(IsInGameThread());
+	GameThreadPerFrameData.GravityGroupIndex = static_cast<uint8>(GravityGroupIndex); // we actually only support a handful gravity groups 
+	SetProxyDirty_External();
+}
+
 void FGeometryCollectionPhysicsProxy::PushStateOnGameThread(Chaos::FPBDRigidsSolver* InSolver)
 {
+	// rest all the dirty flags
+	bIsPhysicsThreadWorldTransformDirty = false;
+	bIsCollisionFilterDataDirty = false;
+	bIsDamageThresholdDataDirty = false;
+	bIsGravityGroupIndexDirty = false;
+	MaterialOverrideMassScaleMultiplierChange = 0;
+
 	// CONTEXT: GAMETHREAD
 	// this is running on GAMETHREAD before the PhysicsThread code runs for this frame
 	bIsPhysicsThreadWorldTransformDirty = GameThreadPerFrameData.GetIsWorldTransformDirty();
@@ -2845,11 +2893,49 @@ void FGeometryCollectionPhysicsProxy::PushStateOnGameThread(Chaos::FPBDRigidsSol
 		GameThreadPerFrameData.ResetIsDamageSettingsDataDirty();
 	}
 
-	bIsDamageThresholdDataDirty = GameThreadPerFrameData.GetIsDamageThresholdDataDirty();
-	if (bIsDamageThresholdDataDirty)
+	if (GameThreadPerFrameData.GetIsDamageThresholdDataDirty())
 	{
 		Parameters.DamageThreshold = GameThreadPerFrameData.GetDamageThresholds();
 		GameThreadPerFrameData.ResetIsDamageThresholdDataDirty();
+		bIsDamageThresholdDataDirty = true;
+	}
+
+	if (GameThreadPerFrameData.DamageModel.IsSet())
+	{
+		Parameters.DamageModel = GameThreadPerFrameData.DamageModel.GetValue();
+		GameThreadPerFrameData.DamageModel.Reset();
+		bIsDamageThresholdDataDirty = true;
+	}
+
+	if (GameThreadPerFrameData.bUseMaterialDamageModifiers.IsSet())
+	{
+		Parameters.bUseMaterialDamageModifiers = GameThreadPerFrameData.bUseMaterialDamageModifiers.GetValue();
+		GameThreadPerFrameData.bUseMaterialDamageModifiers.Reset();
+		bIsDamageThresholdDataDirty = true;
+	}
+
+	if (GameThreadPerFrameData.MaterialOverrideMassScaleMultiplier.IsSet())
+	{
+		const float NewValue = GameThreadPerFrameData.MaterialOverrideMassScaleMultiplier.GetValue();
+		MaterialOverrideMassScaleMultiplierChange = NewValue / Parameters.MaterialOverrideMassScaleMultiplier;
+		Parameters.MaterialOverrideMassScaleMultiplier = NewValue;
+		GameThreadPerFrameData.MaterialOverrideMassScaleMultiplier.Reset();
+	}
+
+	if (GameThreadPerFrameData.GravityGroupIndex.IsSet())
+	{
+		Parameters.GravityGroupIndex = GameThreadPerFrameData.GravityGroupIndex.GetValue();
+		GameThreadPerFrameData.GravityGroupIndex.Reset();
+		bIsGravityGroupIndexDirty = true;
+	}
+
+	if (GameThreadPerFrameData.DamagePropagationData.IsSet())
+	{
+		const FGeometryCollectioPerFrameData::FDamagePropagationData& Data = GameThreadPerFrameData.DamagePropagationData.GetValue();
+		Parameters.bUseDamagePropagation = Data.bEnabled;
+		Parameters.BreakDamagePropagationFactor = Data.BreakDamagePropagationFactor;
+		Parameters.ShockDamagePropagationFactor = Data.ShockDamagePropagationFactor;
+		GameThreadPerFrameData.DamagePropagationData.Reset();
 	}
 }
 
@@ -2982,6 +3068,36 @@ void FGeometryCollectionPhysicsProxy::PushToPhysicsState()
 	{
 		UpdateDamageThreshold_Internal();
 	}
+
+	if (bIsGravityGroupIndexDirty)
+	{
+		for (int32 TransformIndex = 0; TransformIndex < SolverParticleHandles.Num(); ++TransformIndex)
+		{
+			if (Chaos::FPBDRigidClusteredParticleHandle* Handle = SolverParticleHandles[TransformIndex])
+			{
+				Handle->SetGravityGroupIndex(Parameters.GravityGroupIndex);
+			}
+		}
+	}
+
+	if (MaterialOverrideMassScaleMultiplierChange > 0)
+	{
+		for (int32 TransformIndex = 0; TransformIndex < SolverParticleHandles.Num(); ++TransformIndex)
+		{
+			if (Chaos::FPBDRigidClusteredParticleHandle* Handle = SolverParticleHandles[TransformIndex])
+			{
+				const Chaos::FReal NewM = Handle->M() * MaterialOverrideMassScaleMultiplierChange;
+				const Chaos::FVec3 NewI = Handle->I() * MaterialOverrideMassScaleMultiplierChange;
+
+				Handle->SetM(NewM);
+				Handle->SetI(NewI);
+				const Chaos::FReal InvM = (NewM > 0.0f) ? 1.0f / NewM : 0.0f;
+				const Chaos::FVec3 InvI = (NewM > 0.0f) ? Chaos::FVec3(NewI).Reciprocal() : Chaos::FVec3::ZeroVector;
+				Handle->SetInvM(InvM);
+				Handle->SetInvI(InvI);
+			}
+		}
+	}
 }
 
 void FGeometryCollectionPhysicsProxy::SetClusteredParticleKinematicTarget_Internal(Chaos::FPBDRigidClusteredParticleHandle* Handle, const FTransform& NewWorldTransform)
diff --git a/Engine/Source/Runtime/Experimental/Chaos/Public/GeometryCollectionProxyData.h b/Engine/Source/Runtime/Experimental/Chaos/Public/GeometryCollectionProxyData.h
--- a/Engine/Source/Runtime/Experimental/Chaos/Public/GeometryCollectionProxyData.h
+++ b/Engine/Source/Runtime/Experimental/Chaos/Public/GeometryCollectionProxyData.h
@@ -168,6 +168,14 @@ private:
 
 class FGeometryCollectioPerFrameData
 {
+public:
+	struct FDamagePropagationData
+	{
+		bool bEnabled = false;
+		float BreakDamagePropagationFactor = 0.0f;
+		float ShockDamagePropagationFactor = 0.0f;
+	};
+
 public:
 	FGeometryCollectioPerFrameData()
 		: bIsWorldTransformDirty(false)
@@ -310,7 +318,6 @@ public:
 	bool GetIsDamageThresholdDataDirty() const { return bIsDamageThresholdDataDirty; }
 	void ResetIsDamageThresholdDataDirty() { bIsDamageThresholdDataDirty = false; }
 
-
 private:
 	uint16 bIsWorldTransformDirty : 1;
 	uint16 bIsCollisionFilterDataDirty : 1;
@@ -340,6 +347,13 @@ private:
 	/** updated when bIsCollisionFilterDataDirty is set */
 	FCollisionFilterData SimFilter;
 	FCollisionFilterData QueryFilter;
+
+public:
+	TOptional<EDamageModelTypeEnum> DamageModel;
+	TOptional<FDamagePropagationData> DamagePropagationData;
+	TOptional<bool> bUseMaterialDamageModifiers;
+	TOptional<float> MaterialOverrideMassScaleMultiplier;
+	TOptional<uint8> GravityGroupIndex;
 };
 
 /**
diff --git a/Engine/Source/Runtime/Experimental/Chaos/Public/PhysicsProxy/GeometryCollectionPhysicsProxy.h b/Engine/Source/Runtime/Experimental/Chaos/Public/PhysicsProxy/GeometryCollectionPhysicsProxy.h
--- a/Engine/Source/Runtime/Experimental/Chaos/Public/PhysicsProxy/GeometryCollectionPhysicsProxy.h
+++ b/Engine/Source/Runtime/Experimental/Chaos/Public/PhysicsProxy/GeometryCollectionPhysicsProxy.h
@@ -458,7 +458,11 @@ public:
 	CHAOS_API void UpdatePerParticleFilterData_External(const TArray<FParticleCollisionFilterData>& Data);
 	
 	CHAOS_API void SetDamageThresholds_External(const TArray<float>& DamageThresholds);
-
+	CHAOS_API void SetDamagePropagationData_External(bool bEnabled, float BreakDamagePropagationFactor, float ShockDamagePropagationFactor);
+	CHAOS_API void SetDamageModel_External(EDamageModelTypeEnum DamageModel);
+	CHAOS_API void SetUseMaterialDamageModifiers_External(bool bUseMaterialDamageModifiers);
+	CHAOS_API void SetMaterialOverrideMassScaleMultiplier_External(float InMultiplier);
+	CHAOS_API void SetGravityGroupIndex_External(int32 GravityGroupIndex);
 	/** 
 	 * Traverses the parents of TransformGroupIdx counting number of levels,
 	 * and sets levels array value for TransformGroupIdx and its parents if not yet initialized.
@@ -630,10 +634,13 @@ private:
 	FGeometryDynamicCollection& GameThreadCollection;
 
 	// this data flows from Game thread to physics thread
+	// todo: replace this with proper structure copies between the two threads 
 	FGeometryCollectioPerFrameData GameThreadPerFrameData;
-	bool bIsPhysicsThreadWorldTransformDirty;
-	bool bIsCollisionFilterDataDirty;
-	bool bIsDamageThresholdDataDirty;
+	float MaterialOverrideMassScaleMultiplierChange;
+	uint8 bIsPhysicsThreadWorldTransformDirty : 1;
+	uint8 bIsCollisionFilterDataDirty: 1;
+	uint8 bIsDamageThresholdDataDirty: 1;
+	uint8 bIsGravityGroupIndexDirty: 1;
 
 	// Currently this is using triple buffers for game-physics and 
 	// physics-game thread communication, but not for any reason other than this 
diff --git a/Engine/Source/Runtime/Experimental/GeometryCollectionEngine/Private/GeometryCollection/GeometryCollectionComponent.cpp b/Engine/Source/Runtime/Experimental/GeometryCollectionEngine/Private/GeometryCollection/GeometryCollectionComponent.cpp
--- a/Engine/Source/Runtime/Experimental/GeometryCollectionEngine/Private/GeometryCollection/GeometryCollectionComponent.cpp
+++ b/Engine/Source/Runtime/Experimental/GeometryCollectionEngine/Private/GeometryCollection/GeometryCollectionComponent.cpp
@@ -696,6 +696,24 @@ int32 UGeometryCollectionComponent::GetNumElements(FName Group) const
 	return Size > 0 ? Size : DynamicCollection->NumElements(Group);	//if not, maybe dynamic has the group
 }
 
+void UGeometryCollectionComponent::SetDamageModel(EDamageModelTypeEnum InDamageModel)
+{
+	DamageModel = InDamageModel;
+	if (PhysicsProxy)
+	{
+		PhysicsProxy->SetDamageModel_External(DamageModel);
+	}
+}
+
+void UGeometryCollectionComponent::SetGravityGroupIndex(int32 InGravityGroupIndex)
+{
+	GravityGroupIndex = InGravityGroupIndex;
+	if (PhysicsProxy)
+	{
+		PhysicsProxy->SetGravityGroupIndex_External(GravityGroupIndex);
+	}
+}
+
 void UGeometryCollectionComponent::SetDamageThreshold(const TArray<float>& InDamageThreshold)
 {
 	// NOTE: Should only call this during construction, not during runtime
@@ -708,6 +726,36 @@ void UGeometryCollectionComponent::SetDamageThreshold(const TArray<float>& InDam
 	}
 }
 
+void UGeometryCollectionComponent::SetDamagePropagationData(const FGeometryCollectionDamagePropagationData& InDamagePropagationData)
+{
+	DamagePropagationData = InDamagePropagationData;
+
+	if (PhysicsProxy)
+	{
+		PhysicsProxy->SetDamagePropagationData_External(DamagePropagationData.bEnabled, DamagePropagationData.BreakDamagePropagationFactor, DamagePropagationData.ShockDamagePropagationFactor);
+	}
+}
+
+void UGeometryCollectionComponent::SetUseMaterialDamageModifiers(bool bInUseMaterialDamageModifiers)
+{
+	bUseMaterialDamageModifiers = bInUseMaterialDamageModifiers;
+
+	if (PhysicsProxy)
+	{
+		PhysicsProxy->SetUseMaterialDamageModifiers_External(bUseMaterialDamageModifiers);
+	}
+}
+
+void UGeometryCollectionComponent::SetDensityFromPhysicsMaterial(bool bInDensityFromPhysicsMaterial)
+{
+	bDensityFromPhysicsMaterial = bInDensityFromPhysicsMaterial;
+
+	if (PhysicsProxy)
+	{
+		PhysicsProxy->SetMaterialOverrideMassScaleMultiplier_External(ComputeMassScaleRelativeToAsset());
+	}
+}
+
 void UGeometryCollectionComponent::UpdateCachedBounds()
 {
 	ComponentSpaceBounds = ComputeBounds(FTransform::Identity);
diff --git a/Engine/Source/Runtime/Experimental/GeometryCollectionEngine/Public/GeometryCollection/GeometryCollectionComponent.h b/Engine/Source/Runtime/Experimental/GeometryCollectionEngine/Public/GeometryCollection/GeometryCollectionComponent.h
--- a/Engine/Source/Runtime/Experimental/GeometryCollectionEngine/Public/GeometryCollection/GeometryCollectionComponent.h
+++ b/Engine/Source/Runtime/Experimental/GeometryCollectionEngine/Public/GeometryCollection/GeometryCollectionComponent.h
@@ -790,13 +790,19 @@ public:
 	UPROPERTY(EditAnywhere, BlueprintReadWrite, Category = "ChaosPhysics|General")
 	EObjectStateTypeEnum ObjectType;
 
-	UPROPERTY(EditAnywhere, BlueprintReadWrite, Category = "ChaosPhysics|General")
+	UPROPERTY(EditAnywhere, BlueprintReadWrite, BlueprintSetter = SetGravityGroupIndex, Category = "ChaosPhysics|General")
 	int32 GravityGroupIndex;
 
+	UFUNCTION(BlueprintCallable, BlueprintInternalUseOnly)
+	GEOMETRYCOLLECTIONENGINE_API void SetGravityGroupIndex(int32 InGravityGroupIndex);
+
 	/** when true, density will be used to compute mass using the assigned physics material */
-	UPROPERTY(EditAnywhere, BlueprintReadWrite, Category = "ChaosPhysics|General")
+	UPROPERTY(EditAnywhere, BlueprintReadWrite, BlueprintSetter = SetDensityFromPhysicsMaterial, Category = "ChaosPhysics|General")
 	bool bDensityFromPhysicsMaterial;
 
+	UFUNCTION(BlueprintCallable, BlueprintInternalUseOnly)
+	GEOMETRYCOLLECTIONENGINE_API void SetDensityFromPhysicsMaterial(bool bInDensityFromPhysicsMaterial);
+
 	/** If ForceMotionBlur is on, motion blur will always be active, even if the GeometryCollection is at rest. */
 	UPROPERTY(EditAnywhere, BlueprintReadWrite, Category = "ChaosPhysics|General")
 	bool bForceMotionBlur;
@@ -823,9 +829,12 @@ public:
 	int32 MaxSimulatedLevel;
 
 	/** Damage model to use for evaluating destruction. */
-	UPROPERTY(EditAnywhere, BlueprintReadWrite, Category = "ChaosPhysics|Damage")
+	UPROPERTY(EditAnywhere, BlueprintReadWrite, BlueprintSetter = SetDamageModel, Category = "ChaosPhysics|Damage")
 	EDamageModelTypeEnum DamageModel;
 
+	UFUNCTION(BlueprintCallable, BlueprintInternalUseOnly)
+	GEOMETRYCOLLECTIONENGINE_API void SetDamageModel(EDamageModelTypeEnum InDamageModel);
+
 	UPROPERTY(EditAnywhere, BlueprintReadWrite, BlueprintGetter=GetDamageThreshold, BlueprintSetter=SetDamageThreshold, Category = "ChaosPhysics|Damage", meta = (EditCondition = "!bUseSizeSpecificDamageThreshold && DamageModel == EDamageModelTypeEnum::Chaos_Damage_Model_UserDefined_Damage_Threshold"))
 	TArray<float> DamageThreshold;
 
@@ -840,13 +849,19 @@ public:
 	bool bUseSizeSpecificDamageThreshold;
 
 	/** When on , use the modifiers on the material to adjust the user defined damage threshold values */
-	UPROPERTY(EditAnywhere, BlueprintReadWrite, Category = "ChaosPhysics|Damage", meta = (EditCondition = "DamageModel == EDamageModelTypeEnum::Chaos_Damage_Model_UserDefined_Damage_Threshold"))
+	UPROPERTY(EditAnywhere, BlueprintReadWrite, BlueprintSetter = SetUseMaterialDamageModifiers, Category = "ChaosPhysics|Damage", meta = (EditCondition = "DamageModel == EDamageModelTypeEnum::Chaos_Damage_Model_UserDefined_Damage_Threshold"))
 	bool bUseMaterialDamageModifiers;
 
+	UFUNCTION(BlueprintCallable, BlueprintInternalUseOnly)
+	GEOMETRYCOLLECTIONENGINE_API void SetUseMaterialDamageModifiers(bool bInUseMaterialDamageModifiers);
+
 	/** Data about how damage propagation shoudl behave. */
-	UPROPERTY(EditAnywhere, BlueprintReadWrite, Category = "ChaosPhysics|Damage")
+	UPROPERTY(EditAnywhere, BlueprintReadWrite, BlueprintSetter = SetDamagePropagationData, Category = "ChaosPhysics|Damage")
 	FGeometryCollectionDamagePropagationData DamagePropagationData;
 
+	UFUNCTION(BlueprintCallable, BlueprintInternalUseOnly)
+	GEOMETRYCOLLECTIONENGINE_API void SetDamagePropagationData(const FGeometryCollectionDamagePropagationData& InDamagePropagationData);
+
 	/** Whether or not collisions against this geometry collection will apply strain which could cause the geometry collection to fracture. */
 	UPROPERTY(EditAnywhere, BlueprintReadWrite, BlueprintSetter=SetEnableDamageFromCollision, Category = "ChaosPhysics|Damage")
 	bool bEnableDamageFromCollision;
