Reduce memory usage in Geometry Collection
Relative Transform in Dynamic Collection are now in single precision. This change implicates an API change. The Managed Array of transforms being in public, and being changed make some backward compatibility breakage. This array shouldn't have be used much.

diff --git a/Engine/Plugins/Experimental/ChaosNiagara/Source/ChaosNiagara/Private/NiagaraDataInterfaceGeometryCollection.cpp b/Engine/Plugins/Experimental/ChaosNiagara/Source/ChaosNiagara/Private/NiagaraDataInterfaceGeometryCollection.cpp
--- a/Engine/Plugins/Experimental/ChaosNiagara/Source/ChaosNiagara/Private/NiagaraDataInterfaceGeometryCollection.cpp
+++ b/Engine/Plugins/Experimental/ChaosNiagara/Source/ChaosNiagara/Private/NiagaraDataInterfaceGeometryCollection.cpp
@@ -102,7 +102,6 @@ void FNDIGeometryCollectionData::Init(UNiagaraDataInterfaceGeometryCollection* I
 				Collection = Interface->GeometryCollectionActor->GetGeometryCollectionComponent()->RestCollection->GetGeometryCollection();
 			const TManagedArray<FBox>& BoundingBoxes = Collection->BoundingBox;
 			const TManagedArray<int32>& TransformIndex = Collection->TransformIndex;
-			const TManagedArray<FTransform>& Transforms = Interface->GeometryCollectionActor->GetGeometryCollectionComponent()->GetTransformArray();
 			const TManagedArray<TSet<int32>>& Children = Collection->Children;
 			const TManagedArray<int32>& TransformIndexArray = Collection->TransformIndex;
 
@@ -177,7 +176,6 @@ void FNDIGeometryCollectionData::Update(UNiagaraDataInterfaceGeometryCollection*
 				Collection = Interface->GeometryCollectionActor->GetGeometryCollectionComponent()->RestCollection->GetGeometryCollection();
 			const TManagedArray<FBox>& BoundingBoxes = Collection->BoundingBox;
 			const TManagedArray<int32>& TransformIndexArray = Collection->TransformIndex;
-			const TManagedArray<FTransform>& Transforms = Interface->GeometryCollectionActor->GetGeometryCollectionComponent()->GetTransformArray();
 			const TArray<FTransform>& ComponentSpaceTransforms = Interface->GeometryCollectionActor->GetGeometryCollectionComponent()->GetComponentSpaceTransforms();
 			const TManagedArray<TSet<int32>>& Children = Collection->Children;
 			
@@ -194,11 +192,6 @@ void FNDIGeometryCollectionData::Update(UNiagaraDataInterfaceGeometryCollection*
 				}
 			}
 
-			if (ComponentSpaceTransforms.Num() != Transforms.Num())
-			{
-				return;
-			}
-
 			if (NumPieces != AssetArrays->BoundsBuffer.Num())
 			{
 				Init(Interface, SystemInstance);
diff --git a/Engine/Source/Programs/HeadlessChaos/Private/GeometryCollection/GeometryCollectionTestClustering.cpp b/Engine/Source/Programs/HeadlessChaos/Private/GeometryCollection/GeometryCollectionTestClustering.cpp
--- a/Engine/Source/Programs/HeadlessChaos/Private/GeometryCollection/GeometryCollectionTestClustering.cpp
+++ b/Engine/Source/Programs/HeadlessChaos/Private/GeometryCollection/GeometryCollectionTestClustering.cpp
@@ -92,7 +92,7 @@ namespace GeometryCollectionTest
 		UnitTest.AddSimulationObject(Collection);
 		UnitTest.Initialize();
 
-		TManagedArray<FTransform>& Transform = Collection->DynamicCollection->Transform;
+		TManagedArray<FTransform3f>& Transform = Collection->DynamicCollection->Transform;
 		FReal StartingRigidDistance = (Transform[1].GetTranslation() - Transform[0].GetTranslation()).Size(), CurrentRigidDistance = 0.f;
 
 		TManagedArray<bool>& Active = Collection->DynamicCollection->Active;
@@ -167,7 +167,7 @@ namespace GeometryCollectionTest
 		UnitTest.AddSimulationObject(Collection);
 		UnitTest.Initialize();		
 
-		TManagedArray<FTransform>& Transform = Collection->DynamicCollection->Transform;
+		TManagedArray<FTransform3f>& Transform = Collection->DynamicCollection->Transform;
 		FReal StartingRigidDistance = (Transform[1].GetTranslation() - Transform[0].GetTranslation()).Size(), CurrentRigidDistance = 0.f;
 
 		UnitTest.Advance();
@@ -282,7 +282,7 @@ namespace GeometryCollectionTest
 		UnitTest.AddSimulationObject(Collection);
 		UnitTest.Initialize();
 
-		TManagedArray<FTransform>& Transform = Collection->DynamicCollection->Transform;
+		TManagedArray<FTransform3f>& Transform = Collection->DynamicCollection->Transform;
 		FReal StartingRigidDistance = (Transform[1].GetTranslation() - Transform[0].GetTranslation()).Size(), CurrentRigidDistance = 0.f;
 
 		UnitTest.Advance();
@@ -406,7 +406,7 @@ namespace GeometryCollectionTest
 
 		Collection->PhysObject->SetCollisionParticlesPerObjectFraction(1.0);
 
-		TManagedArray<FTransform>& Transform = Collection->DynamicCollection->Transform;
+		TManagedArray<FTransform3f>& Transform = Collection->DynamicCollection->Transform;
 		FReal StartingRigidDistance = (Transform[1].GetTranslation() - Transform[0].GetTranslation()).Size(), CurrentRigidDistance = 0.f;
 
 		UnitTest.Advance();
@@ -507,7 +507,7 @@ namespace GeometryCollectionTest
 		UnitTest.AddSimulationObject(Collection);
 		UnitTest.Initialize();
 		
-		TManagedArray<FTransform>& Transform = Collection->DynamicCollection->Transform;
+		TManagedArray<FTransform3f>& Transform = Collection->DynamicCollection->Transform;
 		FReal StartingRigidDistance = (Transform[1].GetTranslation() - Transform[0].GetTranslation()).Size(), CurrentRigidDistance = 0.f;
 
 		UnitTest.Advance();
@@ -699,7 +699,7 @@ namespace GeometryCollectionTest
 		UnitTest.AddSimulationObject(Collection);
 		UnitTest.Initialize();
 
-		TManagedArray<FTransform>& Transform = Collection->DynamicCollection->Transform;
+		TManagedArray<FTransform3f>& Transform = Collection->DynamicCollection->Transform;
 		FReal StartingRigidDistance = (Transform[1].GetTranslation() - Transform[0].GetTranslation()).Size(), CurrentRigidDistance = 0.f;
 
 		TArray<bool> Conditions = { false,false,false,false };
@@ -860,7 +860,7 @@ namespace GeometryCollectionTest
 		UnitTest.AddSimulationObject(Collection);
 		UnitTest.Initialize();
 
-		TManagedArray<FTransform>& Transform = Collection->DynamicCollection->Transform;
+		TManagedArray<FTransform3f>& Transform = Collection->DynamicCollection->Transform;
 		FReal StartingRigidDistance = (Transform[1].GetTranslation() - Transform[0].GetTranslation()).Size();
 		FReal CurrentRigidDistance = 0;
 
@@ -1159,7 +1159,7 @@ namespace GeometryCollectionTest
 		UnitTest.AddSimulationObject(Collection);
 		UnitTest.Initialize();
 
-		TManagedArray<FTransform>& Transform = Collection->DynamicCollection->Transform;
+		TManagedArray<FTransform3f>& Transform = Collection->DynamicCollection->Transform;
 		FReal StartingRigidDistance = (Transform[1].GetTranslation() - Transform[0].GetTranslation()).Size();
 		FReal CurrentRigidDistance = 0.f;
 
@@ -1446,8 +1446,8 @@ namespace GeometryCollectionTest
 
 
 		TArray<FReal> Distances;
-		TManagedArray<FTransform>& Transform = DynamicCollection->Transform;
-		TManagedArray<FTransform>& Transform2 = DynamicCollection2->Transform;
+		TManagedArray<FTransform3f>& Transform = DynamicCollection->Transform;
+		TManagedArray<FTransform3f>& Transform2 = DynamicCollection2->Transform;
 
 		auto& Clustering = UnitTest.Solver->GetEvolution()->GetRigidClustering();
 		const auto& ClusterMap = Clustering.GetChildrenMap();
@@ -1553,8 +1553,8 @@ namespace GeometryCollectionTest
 
 
 		TArray<FReal> Distances;
-		TManagedArray<FTransform>& Transform = DynamicCollection->Transform;
-		TManagedArray<FTransform>& Transform2 = DynamicCollection2->Transform;
+		TManagedArray<FTransform3f>& Transform = DynamicCollection->Transform;
+		TManagedArray<FTransform3f>& Transform2 = DynamicCollection2->Transform;
 
 
 		auto& Clustering = UnitTest.Solver->GetEvolution()->GetRigidClustering();
@@ -1657,8 +1657,8 @@ namespace GeometryCollectionTest
 		TSharedPtr<FGeometryDynamicCollection> DynamicCollection2 = Collection2->DynamicCollection;
 
 		TArray<FReal> Distances;
-		TManagedArray<FTransform>& Transform = DynamicCollection->Transform;
-		TManagedArray<FTransform>& Transform2 = DynamicCollection2->Transform;
+		TManagedArray<FTransform3f>& Transform = DynamicCollection->Transform;
+		TManagedArray<FTransform3f>& Transform2 = DynamicCollection2->Transform;
 
 		auto& Clustering = UnitTest.Solver->GetEvolution()->GetRigidClustering();
 		const auto& ClusterMap = Clustering.GetChildrenMap();
diff --git a/Engine/Source/Programs/HeadlessChaos/Private/GeometryCollection/GeometryCollectionTestCollisionResolution.cpp b/Engine/Source/Programs/HeadlessChaos/Private/GeometryCollection/GeometryCollectionTestCollisionResolution.cpp
--- a/Engine/Source/Programs/HeadlessChaos/Private/GeometryCollection/GeometryCollectionTestCollisionResolution.cpp
+++ b/Engine/Source/Programs/HeadlessChaos/Private/GeometryCollection/GeometryCollectionTestCollisionResolution.cpp
@@ -167,13 +167,13 @@ namespace GeometryCollectionTest
 			SimplicialSphereCollection->DynamicCollection->Transform[0].GetTranslation().Z,
 			ImplicitSphereCollection->DynamicCollection->Transform[0].GetTranslation().Z + 3);
 
-		const FVector FirstX = SimplicialSphereCollection->DynamicCollection->Transform[0].GetTranslation();
+		const FVector FirstX = FVector(SimplicialSphereCollection->DynamicCollection->Transform[0].GetTranslation());
 		FVector PrevX = FirstX;
 		for (int i = 0; i < 10; i++)
 		{
 			UnitTest.Advance();
 
-			const FVector& CurrX = SimplicialSphereCollection->DynamicCollection->Transform[0].GetTranslation();
+			const FVector CurrX = FVector(SimplicialSphereCollection->DynamicCollection->Transform[0].GetTranslation());
 			EXPECT_NE(CurrX.Z, FirstX.Z); // moved since init
 			EXPECT_GE(PrevX.Z - CurrX.Z, -KINDA_SMALL_NUMBER); // falling in -Z, or stopped
 			EXPECT_LE(FMath::Abs(CurrX.X), KINDA_SMALL_NUMBER); // straight down
@@ -184,7 +184,7 @@ namespace GeometryCollectionTest
 		{
 			// We expect the simplical sphere to drop by 0.1 in Z and come to rest
 			// on top of the implicit sphere.
-			const FVector& CurrX = SimplicialSphereCollection->DynamicCollection->Transform[0].GetTranslation();
+			const FVector CurrX = FVector(SimplicialSphereCollection->DynamicCollection->Transform[0].GetTranslation());
 			EXPECT_LE(CurrX.Z - 2.0, 0.2); // Relative large fudge factor accounts for aliasing?
 		}
 	}
@@ -245,13 +245,13 @@ namespace GeometryCollectionTest
 			BoxCollection0->DynamicCollection->Transform[0].GetTranslation().Z,
 			BoxCollection1->DynamicCollection->Transform[0].GetTranslation().Z + 3);
 
-		const FVector FirstX = BoxCollection0->DynamicCollection->Transform[0].GetTranslation();
+		const FVector FirstX = FVector(BoxCollection0->DynamicCollection->Transform[0].GetTranslation());
 		FVector PrevX = FirstX;
 		for (int i = 0; i < 10; i++)
 		{
 			UnitTest.Advance();
 
-			const FVector& CurrX = BoxCollection0->DynamicCollection->Transform[0].GetTranslation();
+			const FVector& CurrX = FVector(BoxCollection0->DynamicCollection->Transform[0].GetTranslation());
 			EXPECT_NE(CurrX.Z, FirstX.Z); // moved since init
 			EXPECT_LE(CurrX.Z, PrevX.Z); // falling in -Z, or stopped
 			EXPECT_LE(FMath::Abs(CurrX.X), KINDA_SMALL_NUMBER); // No deflection
@@ -262,7 +262,7 @@ namespace GeometryCollectionTest
 		{
 			// We expect the simplical sphere to drop by 0.1 in Z and come to rest
 			// on top of the implicit sphere.
-			const FVector& CurrX = BoxCollection0->DynamicCollection->Transform[0].GetTranslation();
+			const FVector& CurrX = FVector(BoxCollection0->DynamicCollection->Transform[0].GetTranslation());
 			EXPECT_LE(CurrX.Z - 2.0, 0.2); // Relative large fudge factor accounts for aliasing?
 		}
 
@@ -393,13 +393,13 @@ namespace GeometryCollectionTest
 		UnitTest.Initialize();
 		EXPECT_EQ(SimplicialSphereCollection->DynamicCollection->Transform[0].GetTranslation().Z, ImplicitSphereCollection->DynamicCollection->Transform[0].GetTranslation().Z + 2.0f * Radius + 1.0f);
 
-		const FVector FirstX = SimplicialSphereCollection->DynamicCollection->Transform[0].GetTranslation();
+		const FVector FirstX = FVector(SimplicialSphereCollection->DynamicCollection->Transform[0].GetTranslation());
 		FVector PrevX = FirstX;
 		for (int i = 0; i < 10; i++)
 		{
 			UnitTest.Advance();
 
-			const FVector& CurrX = SimplicialSphereCollection->DynamicCollection->Transform[0].GetTranslation();
+			const FVector& CurrX = FVector(SimplicialSphereCollection->DynamicCollection->Transform[0].GetTranslation());
 			EXPECT_NE(CurrX.Z, FirstX.Z); // moved since init
 			EXPECT_LE(FMath::Abs(CurrX.X), 0.1f); // straight down
 			EXPECT_LE(FMath::Abs(CurrX.Y), 0.1f); // straight down
@@ -409,7 +409,7 @@ namespace GeometryCollectionTest
 		{
 			// We expect the simplical sphere to drop by 0.1 in Z and come to rest
 			// on top of the implicit sphere.
-			const FVector& CurrX = SimplicialSphereCollection->DynamicCollection->Transform[0].GetTranslation();
+			const FVector& CurrX = FVector(SimplicialSphereCollection->DynamicCollection->Transform[0].GetTranslation());
 			EXPECT_LE(FMath::Abs(CurrX.Z - 2.0f * Radius), 0.1 * Radius);
 		}
 	}
@@ -470,13 +470,13 @@ namespace GeometryCollectionTest
 			BoxCollection0->DynamicCollection->Transform[0].GetTranslation().Z,
 			BoxCollection1->DynamicCollection->Transform[0].GetTranslation().Z + Length + 2.0f);
 
-		const FVector FirstX = BoxCollection0->DynamicCollection->Transform[0].GetTranslation();
+		const FVector FirstX = FVector(BoxCollection0->DynamicCollection->Transform[0].GetTranslation());
 		FVector PrevX = FirstX;
 		for (int i = 0; i < 10; i++)
 		{
 			UnitTest.Advance();
 
-			const FVector& CurrX = BoxCollection0->DynamicCollection->Transform[0].GetTranslation();
+			const FVector& CurrX = FVector(BoxCollection0->DynamicCollection->Transform[0].GetTranslation());
 			EXPECT_NE(CurrX.Z, FirstX.Z); // moved since init
 			EXPECT_LE(CurrX.Z, PrevX.Z); // falling in -Z, or stopped
 			EXPECT_LE(FMath::Abs(CurrX.X), KINDA_SMALL_NUMBER); // straight down
@@ -487,7 +487,7 @@ namespace GeometryCollectionTest
 		{
 			// We expect the simplical cube to drop in Z direction and come to rest
 			// on top of the implicit cube.
-			const FVector& CurrX = BoxCollection0->DynamicCollection->Transform[0].GetTranslation();
+			const FVector& CurrX = FVector(BoxCollection0->DynamicCollection->Transform[0].GetTranslation());
 			EXPECT_LE(FMath::Abs(CurrX.Z - Length), 0.2 * Length); // Relative large fudge factor accounts for spatial aliasing and contact location averaging.
 		}
 	}
diff --git a/Engine/Source/Programs/HeadlessChaos/Private/GeometryCollection/GeometryCollectionTestSimulation.cpp b/Engine/Source/Programs/HeadlessChaos/Private/GeometryCollection/GeometryCollectionTestSimulation.cpp
--- a/Engine/Source/Programs/HeadlessChaos/Private/GeometryCollection/GeometryCollectionTestSimulation.cpp
+++ b/Engine/Source/Programs/HeadlessChaos/Private/GeometryCollection/GeometryCollectionTestSimulation.cpp
@@ -145,7 +145,7 @@ namespace GeometryCollectionTest
 			UnitTest.Advance();
 
 		{
-			TManagedArray<FTransform>& Transform = Collection->DynamicCollection->Transform;
+			TManagedArray<FTransform3f>& Transform = Collection->DynamicCollection->Transform;
 			EXPECT_EQ(Transform.Num(), 1);
 			//UE_LOG(LogTest, Verbose, TEXT("Position : (%3.5f,%3.5f,%3.5f)"), Transform[0].GetTranslation().X, Transform[0].GetTranslation().Y, Transform[0].GetTranslation().Z);
 			EXPECT_EQ(Transform[0].GetTranslation().Z, 0.f);
@@ -288,7 +288,7 @@ namespace GeometryCollectionTest
 
 		// testing...
 		auto GCParticles = Collection->PhysObject->GetSolverParticleHandles();
-		TManagedArray<FTransform>& Transform = Collection->DynamicCollection->Transform;
+		TManagedArray<FTransform3f>& Transform = Collection->DynamicCollection->Transform;
 		for (int Frame = 1; Frame < 200; Frame++)
 		{
 			if (Frame == 1)
@@ -304,11 +304,11 @@ namespace GeometryCollectionTest
 				GCParticles[3]->SetCollisionGroup(3);
 				GCParticles[4]->SetCollisionGroup(-1);
 
-				EXPECT_TRUE(Transform[0].GetRotation() == FQuat::Identity); // Can use defaulted zero rotation to indicate that the
-				EXPECT_TRUE(Transform[1].GetRotation() == FQuat::Identity); // rigid has not been affected. Should we though??
-				EXPECT_TRUE(Transform[2].GetRotation() == FQuat::Identity);
-				EXPECT_TRUE(Transform[3].GetRotation() == FQuat::Identity);
-				EXPECT_TRUE(Transform[4].GetRotation() == FQuat::Identity);
+				EXPECT_TRUE(Transform[0].GetRotation() == FQuat4f::Identity); // Can use defaulted zero rotation to indicate that the
+				EXPECT_TRUE(Transform[1].GetRotation() == FQuat4f::Identity); // rigid has not been affected. Should we though??
+				EXPECT_TRUE(Transform[2].GetRotation() == FQuat4f::Identity);
+				EXPECT_TRUE(Transform[3].GetRotation() == FQuat4f::Identity);
+				EXPECT_TRUE(Transform[4].GetRotation() == FQuat4f::Identity);
 			}
 
 			if (Frame == 100)
@@ -316,11 +316,11 @@ namespace GeometryCollectionTest
 				EXPECT_NEAR(Transform[0].GetTranslation().Z, 50.0f, 1.0f);
 				EXPECT_NEAR(Transform[1].GetTranslation().Z, 150.0f, 1.0f);
 				EXPECT_NEAR(Transform[2].GetTranslation().Z, 250.0f, 1.0f);
-				EXPECT_FALSE(Transform[0].GetRotation() == FQuat::Identity);
-				EXPECT_FALSE(Transform[1].GetRotation() == FQuat::Identity);
-				EXPECT_FALSE(Transform[2].GetRotation() == FQuat::Identity);
-				EXPECT_FALSE(Transform[3].GetRotation() == FQuat::Identity);
-				EXPECT_TRUE(Transform[4].GetRotation() == FQuat::Identity);
+				EXPECT_FALSE(Transform[0].GetRotation() == FQuat4f::Identity);
+				EXPECT_FALSE(Transform[1].GetRotation() == FQuat4f::Identity);
+				EXPECT_FALSE(Transform[2].GetRotation() == FQuat4f::Identity);
+				EXPECT_FALSE(Transform[3].GetRotation() == FQuat4f::Identity);
+				EXPECT_TRUE(Transform[4].GetRotation() == FQuat4f::Identity);
 			}
 			UnitTest.Advance();
 		}
@@ -329,8 +329,8 @@ namespace GeometryCollectionTest
 		EXPECT_NEAR(Transform[1].GetTranslation().Z, 150.0f, 1.0f);
 		EXPECT_NEAR(Transform[2].GetTranslation().Z, 250.0f, 1.0f);
 		EXPECT_NEAR(Transform[3].GetTranslation().Z, 150.0f, 1.0f);
-		EXPECT_FALSE(Transform[3].GetRotation() == FQuat::Identity);
-		EXPECT_TRUE(Transform[4].GetRotation() == FQuat::Identity); // Phased through everything, good.
+		EXPECT_FALSE(Transform[3].GetRotation() == FQuat4f::Identity);
+		EXPECT_TRUE(Transform[4].GetRotation() == FQuat4f::Identity); // Phased through everything, good.
 		EXPECT_LT(Transform[4].GetTranslation().Z, -100.0f);
 
 		GCParticles[0]->SetCollisionGroup(-1);
diff --git a/Engine/Source/Programs/HeadlessChaos/Private/GeometryCollection/GeometryCollectionTestSimulationField.cpp b/Engine/Source/Programs/HeadlessChaos/Private/GeometryCollection/GeometryCollectionTestSimulationField.cpp
--- a/Engine/Source/Programs/HeadlessChaos/Private/GeometryCollection/GeometryCollectionTestSimulationField.cpp
+++ b/Engine/Source/Programs/HeadlessChaos/Private/GeometryCollection/GeometryCollectionTestSimulationField.cpp
@@ -90,9 +90,9 @@ namespace GeometryCollectionTest
 		EXPECT_EQ(DynamicState[0], (int32)EObjectStateTypeEnum::Chaos_Object_Kinematic);
 
 		// simulated
-		const TManagedArray<FTransform>& Transform = Collection->DynamicCollection->Transform;
+		const TManagedArray<FTransform3f>& Transform = Collection->DynamicCollection->Transform;
 		EXPECT_EQ(Transform.Num(), 1);
-		const FVector Translation1 = Transform[0].GetTranslation();
+		const FVector Translation1 = FVector(Transform[0].GetTranslation());
 		EXPECT_NEAR((Translation0 - Translation1).Size(), 0.f, KINDA_SMALL_NUMBER);
 		EXPECT_NEAR(Transform[0].GetTranslation().Z, 1.f, KINDA_SMALL_NUMBER);
 
@@ -110,7 +110,7 @@ namespace GeometryCollectionTest
 		}
 		EXPECT_EQ(DynamicState[0], (int32)EObjectStateTypeEnum::Chaos_Object_Dynamic);
 
-		const FVector Translation2 = Transform[0].GetTranslation();
+		const FVector Translation2 = FVector(Transform[0].GetTranslation());
 		EXPECT_NE(Translation1, Translation2);
 		EXPECT_LE(Transform[0].GetTranslation().Z, 0.f);
 	}
@@ -135,9 +135,9 @@ namespace GeometryCollectionTest
 		}
 
 		// simulated
-		const TManagedArray<FTransform>& Transform = Collection->DynamicCollection->Transform;
+		const TManagedArray<FTransform3f>& Transform = Collection->DynamicCollection->Transform;
 		EXPECT_EQ(Transform.Num(), 1);
-		const FVector Translation1 = Transform[0].GetTranslation();
+		const FVector Translation1 = FVector(Transform[0].GetTranslation());
 		EXPECT_NEAR((Translation0 - Translation1).Size(), 0.f, KINDA_SMALL_NUMBER);
 		EXPECT_NEAR(Transform[0].GetTranslation().Z, 1.f, KINDA_SMALL_NUMBER);
 
@@ -155,7 +155,7 @@ namespace GeometryCollectionTest
 			UnitTest.Advance();
 		}
 
-		const FVector Translation2 = Transform[0].GetTranslation();
+		const FVector Translation2 = FVector(Transform[0].GetTranslation());
 		EXPECT_NE(Translation1, Translation2);
 		EXPECT_LE(Transform[0].GetTranslation().Z, 0.f);
 	}
@@ -184,7 +184,7 @@ namespace GeometryCollectionTest
 
 		UnitTest.Initialize();
 
-		TManagedArray<FTransform>& Transform = Collection->DynamicCollection->Transform;
+		TManagedArray<FTransform3f>& Transform = Collection->DynamicCollection->Transform;
 		TManagedArray<int32>& DynamicState = Collection->DynamicCollection->DynamicState;
 
 		FReal PreviousY = 0.f;
@@ -241,7 +241,7 @@ namespace GeometryCollectionTest
 
 		UnitTest.Initialize();
 
-		TManagedArray<FTransform>& Transform = Collection->DynamicCollection->Transform;
+		TManagedArray<FTransform3f>& Transform = Collection->DynamicCollection->Transform;
 		for (int Frame = 0; Frame < 10; Frame++)
 		{
 			// Set everything inside the r=5.0 sphere to dynamic
@@ -288,7 +288,7 @@ namespace GeometryCollectionTest
 
 		UnitTest.Initialize();
 
-		TManagedArray<FTransform>& Transform = Collection->DynamicCollection->Transform;
+		TManagedArray<FTransform3f>& Transform = Collection->DynamicCollection->Transform;
 		FReal PreviousY = 0.0;
 		for (int Frame = 0; Frame < 10; Frame++)
 		{
@@ -336,7 +336,7 @@ namespace GeometryCollectionTest
 
 		UnitTest.Initialize();
 
-		TManagedArray<FTransform>& Transform = Collection->DynamicCollection->Transform;
+		TManagedArray<FTransform3f>& Transform = Collection->DynamicCollection->Transform;
 		FReal PreviousY = 0.0;
 		for (int Frame = 0; Frame < 10; Frame++)
 		{
@@ -385,7 +385,7 @@ namespace GeometryCollectionTest
 
 		UnitTest.Initialize();
 
-		TManagedArray<FTransform>& Transform = Collection->DynamicCollection->Transform;
+		TManagedArray<FTransform3f>& Transform = Collection->DynamicCollection->Transform;
 		TManagedArray<bool>& Active = Collection->DynamicCollection->Active;
 		auto& Particles = UnitTest.Solver->GetParticles().GetGeometryCollectionParticles();
 		for (int Frame = 0; Frame < 20; Frame++)
@@ -431,7 +431,7 @@ namespace GeometryCollectionTest
 		UnitTest.Advance();
 
 		FReal PreviousX = 0.0;
-		TManagedArray<FTransform>& Transform = Collection->DynamicCollection->Transform;
+		TManagedArray<FTransform3f>& Transform = Collection->DynamicCollection->Transform;
 		for (int Frame = 1; Frame < 10; Frame++)
 		{
 			UnitTest.Solver->GetPerSolverField().AddTransientCommand({ GetFieldPhysicsName(EFieldPhysicsType::Field_LinearVelocity), VectorField->NewCopy() });
@@ -973,7 +973,7 @@ namespace GeometryCollectionTest
 		UnitTest.Solver->GetPerSolverField().AddTransientCommand({ TargetNameForce, VectorFieldBackward->NewCopy() });
 		UnitTest.Advance();
 
-		TManagedArray<FTransform>& Transform = Collection->DynamicCollection->Transform;
+		TManagedArray<FTransform3f>& Transform = Collection->DynamicCollection->Transform;
 		for (int Frame = 1; Frame < 10; Frame++)
 		{
 			UnitTest.Solver->GetPerSolverField().AddTransientCommand({ GetFieldPhysicsName(EFieldPhysicsType::Field_LinearVelocity), VectorFieldRight->NewCopy() });
@@ -1020,7 +1020,7 @@ namespace GeometryCollectionTest
 
 		TArray<Chaos::TPBDRigidClusteredParticleHandle<FReal, 3>*>& ParticleHandles = Collection->PhysObject->GetSolverParticleHandles();
 		Chaos::TVector<float, 3> CurrV = ParticleHandles[0]->V();
-		TManagedArray<FTransform>& Transform = Collection->DynamicCollection->Transform;
+		TManagedArray<FTransform3f>& Transform = Collection->DynamicCollection->Transform;
 
 		for (int Frame = 1; Frame < 10; Frame++)
 		{
@@ -1036,7 +1036,7 @@ namespace GeometryCollectionTest
 			// Still falling?
 			EXPECT_LT(Transform[0].GetTranslation().Z, LastLocation.Z);
 
-			LastLocation = Transform[0].GetTranslation();
+			LastLocation = FVector(Transform[0].GetTranslation());
 
 		}
 	}
@@ -1064,7 +1064,7 @@ namespace GeometryCollectionTest
 
 		UnitTest.Initialize();
 
-		TManagedArray<FTransform>& Transform = Collection->DynamicCollection->Transform;
+		TManagedArray<FTransform3f>& Transform = Collection->DynamicCollection->Transform;
 		FReal PreviousHeight = Transform[0].GetTranslation().Z;
 		FReal PreviousX = Transform[0].GetRotation().Euler().X;
 		for (int Frame = 0; Frame < 10; Frame++)
diff --git a/Engine/Source/Programs/HeadlessChaos/Private/GeometryCollection/GeometryCollectionTestSimulationSolver.cpp b/Engine/Source/Programs/HeadlessChaos/Private/GeometryCollection/GeometryCollectionTestSimulationSolver.cpp
--- a/Engine/Source/Programs/HeadlessChaos/Private/GeometryCollection/GeometryCollectionTestSimulationSolver.cpp
+++ b/Engine/Source/Programs/HeadlessChaos/Private/GeometryCollection/GeometryCollectionTestSimulationSolver.cpp
@@ -62,7 +62,7 @@ using namespace ChaosTest;
 			EXPECT_LT(FMath::Abs(RestTransform[0].GetTranslation().Z), SMALL_THRESHOLD);
 
 			// simulated
-			TManagedArray<FTransform>& Transform = Collection->DynamicCollection->Transform;
+			TManagedArray<FTransform3f>& Transform = Collection->DynamicCollection->Transform;
 			EXPECT_EQ(Transform.Num(), 1);
 			EXPECT_LT(FMath::Abs(Transform[0].GetTranslation().Z), SMALL_THRESHOLD);
 		}
@@ -98,8 +98,8 @@ using namespace ChaosTest;
 		UnitTest.AddSimulationObject(Collection);
 		UnitTest.Initialize();
 
-		FVector StartingClusterPosition;
-		TManagedArray<FTransform>* Transform;
+		FVector3f StartingClusterPosition;
+		TManagedArray<FTransform3f>* Transform;
 		FReal StartingRigidDistance;
 
 		UnitTest.Solver->RegisterSimOneShotCallback([&]()
diff --git a/Engine/Source/Runtime/Experimental/Chaos/Private/GeometryCollection/GeometryCollectionAlgo.cpp b/Engine/Source/Runtime/Experimental/Chaos/Private/GeometryCollection/GeometryCollectionAlgo.cpp
--- a/Engine/Source/Runtime/Experimental/Chaos/Private/GeometryCollection/GeometryCollectionAlgo.cpp
+++ b/Engine/Source/Runtime/Experimental/Chaos/Private/GeometryCollection/GeometryCollectionAlgo.cpp
@@ -318,7 +318,8 @@ namespace GeometryCollectionAlgo
 		return true;
 	}
 
-	void GlobalMatricesHelper(const int32 Index, const TManagedArray<int32>& Parents, const TManagedArray<FTransform>& Transform, TArray<bool>& IsTransformComputed, const TManagedArray<FTransform>* UniformScale, TArray<FMatrix>& OutGlobalTransforms)
+	template<typename TransformType>
+	void GlobalMatricesHelper(const int32 Index, const TManagedArray<int32>& Parents, const TManagedArray<TransformType>& Transform, TArray<bool>& IsTransformComputed, const TManagedArray<FTransform>* UniformScale, TArray<FMatrix>& OutGlobalTransforms)
 	{
 		if (IsTransformComputed[Index])
 		{
@@ -335,7 +336,7 @@ namespace GeometryCollectionAlgo
 		{
 			const int32 ProcessIndex = ToProcess.Pop(false);
 			const int32 ParentIndex = Parents[ProcessIndex];
-			FMatrix Result = Transform[ProcessIndex].ToMatrixWithScale();
+			FMatrix Result = FTransform(Transform[ProcessIndex]).ToMatrixWithScale();
 			if (ParentIndex != FGeometryCollection::Invalid)
 			{
 				Result *= OutGlobalTransforms[ParentIndex];
@@ -355,7 +356,8 @@ namespace GeometryCollectionAlgo
 	}
 
 	// #note: this version outputs FTransforms to support functionality for getting global matrices for an array of indices.
-	void GlobalMatricesHelper(const int32 Index, const TManagedArray<int32>& Parents, const TManagedArray<FTransform>& Transform, TArray<bool>& IsTransformComputed, const TManagedArray<FTransform>* UniformScale, TArray<FTransform>& OutGlobalTransforms)
+	template<typename TransformType>
+	void GlobalMatricesHelper(const int32 Index, const TManagedArray<int32>& Parents, const TManagedArray<TransformType>& Transform, TArray<bool>& IsTransformComputed, const TManagedArray<FTransform>* UniformScale, TArray<FTransform>& OutGlobalTransforms)
 	{
 		if (IsTransformComputed[Index])
 		{
@@ -372,7 +374,7 @@ namespace GeometryCollectionAlgo
 		{
 			const int32 ProcessIndex = ToProcess.Pop(false);
 			const int32 ParentIndex = Parents[ProcessIndex];
-			FTransform Result = Transform[ProcessIndex];
+			FTransform Result = FTransform(Transform[ProcessIndex]);
 			if (ParentIndex != FGeometryCollection::Invalid)
 			{
 				Result *= OutGlobalTransforms[ParentIndex];
@@ -391,7 +393,8 @@ namespace GeometryCollectionAlgo
 		}
 	}
 
-	FTransform GlobalMatricesHelperForIndices(const int32 Index, const TManagedArray<int32>& Parents, const TManagedArray<FTransform>& Transform, TArray<bool>& IsTransformComputed, const TManagedArray<FTransform>* UniformScale, TArray<FTransform>& TransformCache)
+	template<typename TransformType>
+	FTransform GlobalMatricesHelperForIndices(const int32 Index, const TManagedArray<int32>& Parents, const TManagedArray<TransformType>& Transform, TArray<bool>& IsTransformComputed, const TManagedArray<FTransform>* UniformScale, TArray<FTransform>& TransformCache)
 	{
 		if (IsTransformComputed[Index])
 		{
@@ -408,7 +411,7 @@ namespace GeometryCollectionAlgo
 		{
 			const int32 ProcessIndex = ToProcess.Pop(false);
 			const int32 ParentIndex = Parents[ProcessIndex];
-			FTransform Result = Transform[ProcessIndex];
+			FTransform Result = FTransform(Transform[ProcessIndex]);
 			if (ParentIndex != FGeometryCollection::Invalid)
 			{
 				Result *= TransformCache[ParentIndex];
@@ -430,9 +433,8 @@ namespace GeometryCollectionAlgo
 	}
 
 
-
-
-	FTransform GlobalMatrix(const TManagedArray<FTransform>& RelativeTransforms, const TManagedArray<int32>& Parents, int32 Index)
+	template<class TransformType>
+	FTransform GlobalMatrixTemplate(const TManagedArray<TransformType>& RelativeTransforms, const TManagedArray<int32>& Parents, int32 Index)
 	{
 		FTransform Transform = FTransform::Identity;
 
@@ -440,29 +442,25 @@ namespace GeometryCollectionAlgo
 		{
 			while (Index != FGeometryCollection::Invalid)
 			{
-				Transform = Transform * RelativeTransforms[Index];
+				Transform = Transform * FTransform(RelativeTransforms[Index]);
 				Index = Parents[Index];
 			}
 		}
 		return Transform;
 	}
 
-	void GlobalMatrices(const TManagedArray<FTransform>& RelativeTransforms, const TManagedArray<int32>& Parents, const TArray<int32>& Indices, const TManagedArray<FTransform>& UniformScale, TArray<FTransform>& OutGlobalTransforms)
+	FTransform GlobalMatrix(const TManagedArray<FTransform>& RelativeTransforms, const TManagedArray<int32>& Parents, int32 Index)
 	{
-		TArray<bool> IsTransformComputed;
-		IsTransformComputed.AddDefaulted(RelativeTransforms.Num());
-
-		TArray<FTransform> TransformCache;
-		TransformCache.SetNumUninitialized(RelativeTransforms.Num(), false);
+		return GlobalMatrixTemplate(RelativeTransforms, Parents, Index);
+	}
 
-		OutGlobalTransforms.SetNumUninitialized(Indices.Num(), false);
-		for (int Idx = 0; Idx < Indices.Num(); Idx++)
-		{
-			OutGlobalTransforms[Idx] = GlobalMatricesHelperForIndices(Indices[Idx], Parents, RelativeTransforms, IsTransformComputed, &UniformScale, TransformCache);
-		}
+	FTransform GlobalMatrix(const TManagedArray<FTransform3f>& RelativeTransforms, const TManagedArray<int32>& Parents, int32 Index)
+	{
+		return GlobalMatrixTemplate(RelativeTransforms, Parents, Index);
 	}
 
-	void GlobalMatrices(const TManagedArray<FTransform>& RelativeTransforms, const TManagedArray<int32>& Parents, const TArray<int32>& Indices, TArray<FTransform>& OutGlobalTransforms)
+	template<class TransformType>
+	void GlobalMatricesTemplate(const TManagedArray<TransformType>& RelativeTransforms, const TManagedArray<int32>& Parents, const TArray<int32>& Indices, TArray<FTransform>& OutGlobalTransforms)
 	{
 		TArray<bool> IsTransformComputed;
 		IsTransformComputed.AddDefaulted(RelativeTransforms.Num());
@@ -477,6 +475,16 @@ namespace GeometryCollectionAlgo
 		}
 	}
 
+	void GlobalMatrices(const TManagedArray<FTransform>& RelativeTransforms, const TManagedArray<int32>& Parents, const TArray<int32>& Indices, TArray<FTransform>& Transforms)
+	{
+		GlobalMatricesTemplate<FTransform>(RelativeTransforms, Parents, Indices, Transforms);
+	}
+
+	void GlobalMatrices(const TManagedArray<FTransform3f>& RelativeTransforms, const TManagedArray<int32>& Parents, const TArray<int32>& Indices, TArray<FTransform>& Transforms)
+	{
+		GlobalMatricesTemplate<FTransform3f>(RelativeTransforms, Parents, Indices, Transforms);
+	}
+
 	
 	void GlobalMatricesFromRoot(const int32 ParentTransformIndex, const TManagedArray<FTransform>& RelativeTransforms, const TManagedArray<TSet<int32>>& Children, TArray<FMatrix>& Transforms)
 	{
@@ -491,11 +499,11 @@ namespace GeometryCollectionAlgo
 		}
 	}
 
-	template<typename MatrixType>
-	void GlobalMatrices(const TManagedArray<FTransform>& RelativeTransforms, const TManagedArray<int32>& Parents, const TManagedArray<FTransform>& UniformScale, TArray<MatrixType>& OutGlobalTransforms)
+	template<typename MatrixType, typename TransformType>
+	void GlobalMatrices(const TManagedArray<TransformType>& RelativeTransforms, const TManagedArray<int32>& Parents, const TManagedArray<FTransform>& UniformScale, TArray<MatrixType>& OutGlobalTransforms)
 	{
 		int32 NumTransforms = RelativeTransforms.Num();
-		
+
 		TArray<bool> IsTransformComputed;
 		IsTransformComputed.AddDefaulted(NumTransforms);
 
@@ -503,12 +511,12 @@ namespace GeometryCollectionAlgo
 
 		for (int BoneIdx = 0; BoneIdx < NumTransforms; ++BoneIdx)
 		{
-			GlobalMatricesHelper(BoneIdx, Parents, RelativeTransforms, IsTransformComputed, &UniformScale, OutGlobalTransforms);
+			GlobalMatricesHelper<TransformType>(BoneIdx, Parents, RelativeTransforms, IsTransformComputed, &UniformScale, OutGlobalTransforms);
 		}
 	}
 
-	template<typename MatrixType>
-	void GlobalMatrices(const TManagedArray<FTransform>& RelativeTransforms, const TManagedArray<int32>& Parents, TArray<MatrixType>& OutGlobalTransforms)
+	template<typename MatrixType, typename TransformType>
+	void GlobalMatrices(const TManagedArray<TransformType>& RelativeTransforms, const TManagedArray<int32>& Parents, TArray<MatrixType>& OutGlobalTransforms)
 	{
 		int32 NumTransforms = RelativeTransforms.Num();
 
@@ -519,13 +527,20 @@ namespace GeometryCollectionAlgo
 
 		for (int BoneIdx = 0; BoneIdx < NumTransforms; ++BoneIdx)
 		{
-			GlobalMatricesHelper(BoneIdx, Parents, RelativeTransforms, IsTransformComputed, nullptr, OutGlobalTransforms);
+			GlobalMatricesHelper<TransformType>(BoneIdx, Parents, RelativeTransforms, IsTransformComputed, nullptr, OutGlobalTransforms);
 		}
 	}
-	template void CHAOS_API GlobalMatrices<FTransform>(const TManagedArray<FTransform>&, const TManagedArray<int32>&, const TManagedArray<FTransform>&, TArray<FTransform>&);
-	template void CHAOS_API GlobalMatrices<FTransform>(const TManagedArray<FTransform>&, const TManagedArray<int32>&, TArray<FTransform>&);
-	template void CHAOS_API GlobalMatrices<FMatrix>(const TManagedArray<FTransform>&, const TManagedArray<int32>&, const TManagedArray<FTransform>&, TArray<FMatrix>&);
-	template void CHAOS_API GlobalMatrices<FMatrix>(const TManagedArray<FTransform>&, const TManagedArray<int32>&, TArray<FMatrix>&);
+
+	template void CHAOS_API GlobalMatrices<FTransform, FTransform>(const TManagedArray<FTransform>&, const TManagedArray<int32>&, const TManagedArray<FTransform>&, TArray<FTransform>&);
+	template void CHAOS_API GlobalMatrices<FTransform, FTransform>(const TManagedArray<FTransform>&, const TManagedArray<int32>&, TArray<FTransform>&);
+	template void CHAOS_API GlobalMatrices<FMatrix, FTransform>(const TManagedArray<FTransform>&, const TManagedArray<int32>&, const TManagedArray<FTransform>&, TArray<FMatrix>&);
+	template void CHAOS_API GlobalMatrices<FMatrix, FTransform>(const TManagedArray<FTransform>&, const TManagedArray<int32>&, TArray<FMatrix>&);
+
+	template void CHAOS_API GlobalMatrices<FTransform, FTransform3f>(const TManagedArray<FTransform3f>&, const TManagedArray<int32>&, const TManagedArray<FTransform>&, TArray<FTransform>&);
+	template void CHAOS_API GlobalMatrices<FTransform, FTransform3f>(const TManagedArray<FTransform3f>&, const TManagedArray<int32>&, TArray<FTransform>&);
+	template void CHAOS_API GlobalMatrices<FMatrix, FTransform3f>(const TManagedArray<FTransform3f>&, const TManagedArray<int32>&, const TManagedArray<FTransform>&, TArray<FMatrix>&);
+	template void CHAOS_API GlobalMatrices<FMatrix, FTransform3f>(const TManagedArray<FTransform3f>&, const TManagedArray<int32>&, TArray<FMatrix>&);
+
 
 	void FloodForOverlappedPairs(int Level, int32 BoneIndex, TMap<int32, int32> &BoneToGroup, const TManagedArray<int32>& Levels, const TMap<int32, FBox>& BoundingBoxes, TSet<TTuple<int32, int32>>& OutOverlappedPairs)
 	{
@@ -1553,4 +1568,3 @@ namespace GeometryCollectionAlgo
 
 }
 
-
diff --git a/Engine/Source/Runtime/Experimental/Chaos/Private/GeometryCollectionProxyData.cpp b/Engine/Source/Runtime/Experimental/Chaos/Private/GeometryCollectionProxyData.cpp
--- a/Engine/Source/Runtime/Experimental/Chaos/Private/GeometryCollectionProxyData.cpp
+++ b/Engine/Source/Runtime/Experimental/Chaos/Private/GeometryCollectionProxyData.cpp
@@ -24,7 +24,7 @@ void FTransformDynamicCollection::Construct()
 	FManagedArrayCollection::FConstructionParameters TransformDependency(FTransformCollection::TransformGroup);
 
 	// Transform Group
-	AddExternalAttribute<FTransform>(FTransformCollection::TransformAttribute, FTransformCollection::TransformGroup, Transform);
+	AddExternalAttribute<FTransform3f>(FTransformCollection::TransformAttribute, FTransformCollection::TransformGroup, Transform);
 	AddExternalAttribute<int32>(FTransformCollection::ParentAttribute, FTransformCollection::TransformGroup, Parent);
 	AddExternalAttribute<TSet<int32>>(FTransformCollection::ChildrenAttribute, FTransformCollection::TransformGroup, Children);
 	AddExternalAttribute<int32>(FGeometryCollection::SimulationTypeAttribute, FTransformCollection::TransformGroup, SimulationType);
diff --git a/Engine/Source/Runtime/Experimental/Chaos/Private/PhysicsProxy/GeometryCollectionPhysicsProxy.cpp b/Engine/Source/Runtime/Experimental/Chaos/Private/PhysicsProxy/GeometryCollectionPhysicsProxy.cpp
--- a/Engine/Source/Runtime/Experimental/Chaos/Private/PhysicsProxy/GeometryCollectionPhysicsProxy.cpp
+++ b/Engine/Source/Runtime/Experimental/Chaos/Private/PhysicsProxy/GeometryCollectionPhysicsProxy.cpp
@@ -586,7 +586,7 @@ void FGeometryCollectionPhysicsProxy::Initialize(Chaos::FPBDRigidsEvolutionBase
 				// constraint will query the game thread particle position/rotation for the geometry collection to compute its
 				// reference frame. If that position/rotation does not match up with the physics thread's position/rotation,
 				// the geometry collection particle will have an added velocity computed by the joint constraint solver.
-				const FTransform& T = MassToLocal[Index] * GameThreadCollection.Transform[Index] * Parameters.WorldTransform;
+				const FTransform& T = MassToLocal[Index] * FTransform(GameThreadCollection.Transform[Index]) * Parameters.WorldTransform;
 				P->SetX(T.GetTranslation(), false);
 				P->SetR(T.GetRotation(), false);
 				P->SetM(ScaledMass);
@@ -675,7 +675,7 @@ void FGeometryCollectionPhysicsProxy::Initialize(Chaos::FPBDRigidsEvolutionBase
 								const Chaos::FImplicitObjectPtr& ChildImplicit = GameThreadCollection.Implicits[ChildIndex];
 								if (ChildImplicit)
 								{
-									const Chaos::FRigidTransform3 ChildShapeTransform = MassToLocal[ChildIndex] * GameThreadCollection.Transform[ChildIndex];
+									const Chaos::FRigidTransform3 ChildShapeTransform = MassToLocal[ChildIndex] * FTransform(GameThreadCollection.Transform[ChildIndex]);
 									const Chaos::FRigidTransform3 RelativeShapeTransform = ChildShapeTransform.GetRelativeTransform(ParentShapeTransform);
 
 									Chaos::FImplicitObjectPtr TransformedChildImplicit = MakeTransformImplicitObject(*ChildImplicit, RelativeShapeTransform);
@@ -1627,7 +1627,7 @@ Chaos::TPBDGeometryCollectionParticleHandle<Chaos::FReal, 3>* FGeometryCollectio
 	TManagedArray<int32>& DynamicState = DynamicCollection.DynamicState;
 	TManagedArray<int32>& ParentIndex = DynamicCollection.Parent;
 	TManagedArray<TSet<int32>>& Children = DynamicCollection.Children;
-	TManagedArray<FTransform>& Transform = DynamicCollection.Transform;
+	TManagedArray<FTransform3f>& Transform = DynamicCollection.Transform;
 	const TManagedArray<FTransform>& MassToLocal = Parameters.RestCollection->GetAttribute<FTransform>(MassToLocalAttributeName, FTransformCollection::TransformGroup);
 	const TManagedArray<Chaos::FImplicitObjectPtr>& Implicits = DynamicCollection.Implicits;
 	const TManagedArray<TUniquePtr<FCollisionStructureManager::FSimplicial>>& Simplicials = DynamicCollection.Simplicials;
@@ -1747,7 +1747,7 @@ FGeometryCollectionPhysicsProxy::BuildClusters_Internal(
 	TManagedArray<int32>& DynamicState = DynamicCollection.DynamicState;
 	TManagedArray<int32>& ParentIndex = DynamicCollection.Parent;
 	TManagedArray<TSet<int32>>& Children = DynamicCollection.Children;
-	TManagedArray<FTransform>& Transform = DynamicCollection.Transform;
+	TManagedArray<FTransform3f>& Transform = DynamicCollection.Transform;
 	const TManagedArray<FTransform>& MassToLocal = Parameters.RestCollection->GetAttribute<FTransform>(MassToLocalAttributeName, FTransformCollection::TransformGroup);
 	//TManagedArray<TSharedPtr<FCollisionStructureManager::FSimplicial> >& Simplicials = DynamicCollection.Simplicials;
 	TManagedArray<Chaos::FImplicitObjectPtr>& Implicits = DynamicCollection.Implicits;
@@ -3044,7 +3044,7 @@ void FGeometryCollectionPhysicsProxy::PushToPhysicsState()
 						if (ClusterUnionIndex != INDEX_NONE)
 						{
 							const FTransform ParentWorldTransform{ ParentHandle->R(), ParentHandle->X() };
-							const FTransform NewWorldTransform = MassToLocal[TransformGroupIndex] * PhysicsThreadCollection.Transform[TransformGroupIndex] * Parameters.WorldTransform;
+							const FTransform NewWorldTransform = MassToLocal[TransformGroupIndex] * FTransform(PhysicsThreadCollection.Transform[TransformGroupIndex]) * Parameters.WorldTransform;
 							const FTransform RelativeTransform = NewWorldTransform.GetRelativeTransform(ParentWorldTransform);
 
 							DeferredClusterUnionParticleUpdates.Add(Handle);
@@ -3358,7 +3358,7 @@ void FGeometryCollectionPhysicsProxy::BufferPhysicsResults_Internal(Chaos::FPBDR
 			if (bHasChanged)
 			{
 				// default is what we have on the collection
-				FTransform ParentSpaceTransform = PhysicsThreadCollection.Transform[TransformGroupIndex];
+				FTransform ParentSpaceTransform = FTransform(PhysicsThreadCollection.Transform[TransformGroupIndex]);
 
 				// recompute parent space transform if there's no more parent or if the parent is an internal cluster 
 				if (!ClusterParent || ClusterParent->InternalCluster())
@@ -3378,7 +3378,7 @@ void FGeometryCollectionPhysicsProxy::BufferPhysicsResults_Internal(Chaos::FPBDR
 				Results.SetPositions(EntryIndex, PositionData);
 				Results.SetVelocities(EntryIndex, VelocityData);
 				// todo(chaos) : we shoudl eventually get rid of the transform in the Physics collection
-				PhysicsThreadCollection.Transform[TransformGroupIndex] = ParentSpaceTransform; 
+				PhysicsThreadCollection.Transform[TransformGroupIndex] = FTransform3f(ParentSpaceTransform);
 				IsObjectDynamic = true;
 			}
 
@@ -3473,7 +3473,7 @@ static inline bool UpdateTransform(FTransform& ValueInOut, const FTransform& New
 }
 
 // this helper class helps getting the right interpolated values )
-// it also handles the case where entries for a specific transfrom index may be available in one or the other the other results 
+// it also handles the case where entries for a specific transform index may be available in one or the other the other results 
 // and properly fallback to interpolate 
 struct FResultInterpolator
 {
@@ -3808,10 +3808,10 @@ bool FGeometryCollectionPhysicsProxy::PullFromPhysicsState(const Chaos::FDirtyGe
 							{
 								NewTransform = ParticleMassToLocal.Inverse() * ComponentScaleTransform * ParticleMassToLocal * NewTransform;
 							}
-							if (!NewTransform.Equals(GameThreadCollection.Transform[TransformIndex], GeometryCollectionTransformTolerance))
+							if (!NewTransform.Equals(FTransform(GameThreadCollection.Transform[TransformIndex]), GeometryCollectionTransformTolerance))
 							{
 								bHasDifferentTransforms = true;
-								GameThreadCollection.Transform[TransformIndex] = NewTransform;
+								GameThreadCollection.Transform[TransformIndex] = FTransform3f(NewTransform);
 							}
 						}
 					}
diff --git a/Engine/Source/Runtime/Experimental/Chaos/Public/GeometryCollection/GeometryCollectionAlgo.h b/Engine/Source/Runtime/Experimental/Chaos/Public/GeometryCollection/GeometryCollectionAlgo.h
--- a/Engine/Source/Runtime/Experimental/Chaos/Public/GeometryCollection/GeometryCollectionAlgo.h
+++ b/Engine/Source/Runtime/Experimental/Chaos/Public/GeometryCollection/GeometryCollectionAlgo.h
@@ -133,12 +133,14 @@ namespace GeometryCollectionAlgo
 	*  Global Matrices of the specified index.
 	*/
 	FTransform CHAOS_API GlobalMatrix(const TManagedArray<FTransform>& RelativeTransforms, const TManagedArray<int32>& Parents, int32 Index);
+	FTransform CHAOS_API GlobalMatrix(const TManagedArray<FTransform3f>& RelativeTransforms, const TManagedArray<int32>& Parents, int32 Index);
 
 
 	/*
 	*  Global Matrices of the collection based on list of indices
 	*/
 	void CHAOS_API GlobalMatrices(const TManagedArray<FTransform>& RelativeTransforms, const TManagedArray<int32>& Parents, const TArray<int32>& Indices, TArray<FTransform>& Transforms);
+	void CHAOS_API GlobalMatrices(const TManagedArray<FTransform3f>& RelativeTransforms, const TManagedArray<int32>& Parents, const TArray<int32>& Indices, TArray<FTransform>& Transforms);
 
 	/*
 	 *  Recursively traverse from a root node down
@@ -148,11 +150,12 @@ namespace GeometryCollectionAlgo
 	/*
 	*  Global Matrices of the collection, transforms will be resized to fit
 	*/
-	template<typename MatrixType>
-	void CHAOS_API GlobalMatrices(const TManagedArray<FTransform>& RelativeTransforms, const TManagedArray<int32>& Parents, const TManagedArray<FTransform>& UniformScale, TArray<MatrixType>& Transforms);
+	template<typename MatrixType, typename TransformType>
+	void CHAOS_API GlobalMatrices(const TManagedArray<TransformType>& RelativeTransforms, const TManagedArray<int32>& Parents, const TManagedArray<FTransform>& UniformScale, TArray<MatrixType>& Transforms);
+
+	template<typename MatrixType, typename TransformType>
+	void CHAOS_API GlobalMatrices(const TManagedArray<TransformType>& RelativeTransforms, const TManagedArray<int32>& Parents, TArray<MatrixType>& Transforms);
 
-	template<typename MatrixType>
-	void CHAOS_API GlobalMatrices(const TManagedArray<FTransform>& RelativeTransforms, const TManagedArray<int32>& Parents, TArray<MatrixType>& Transforms);
 
 	/*
 	*  Gets pairs of elements whose bounding boxes overlap.
diff --git a/Engine/Source/Runtime/Experimental/Chaos/Public/GeometryCollection/ManagedArray.h b/Engine/Source/Runtime/Experimental/Chaos/Public/GeometryCollection/ManagedArray.h
--- a/Engine/Source/Runtime/Experimental/Chaos/Public/GeometryCollection/ManagedArray.h
+++ b/Engine/Source/Runtime/Experimental/Chaos/Public/GeometryCollection/ManagedArray.h
@@ -601,7 +601,7 @@ public:
 	FORCEINLINE RangedForIteratorType      end  ()			{ return Array.end(); }
 	FORCEINLINE RangedForConstIteratorType end  () const	{ return Array.end(); }
 
-private:
+protected:
 	/**
 	* Protected Resize to prevent external resizing of the array
 	*
@@ -663,6 +663,18 @@ void InitHelper(TArray<T>& Array, const TManagedArrayBase<T>& NewTypedArray, int
 	}
 }
 
+
+template <typename TSrc, typename TDst>
+void InitHelper(TArray<TDst>& Array, const TManagedArrayBase<TSrc>& NewTypedArray, int32 Size)
+{
+	for (int32 Index = 0; Index < Size; Index++)
+	{
+		Array[Index] = TDst(NewTypedArray[Index]);
+	}
+}
+
+
+
 template <typename T>
 void InitHelper(TArray<TUniquePtr<T>>& Array, const TManagedArrayBase<TUniquePtr<T>>& NewTypedArray, int32 Size)
 {
@@ -1168,6 +1180,30 @@ public:
 };
 
 
+template<>
+class TManagedArray<FTransform3f> : public TManagedArrayBase<FTransform3f>
+{
+	/**
+	* Init from a predefined Array of matching type
+	*/
+	virtual void Init(const FManagedArrayBase& NewArray) override
+	{
+		if (NewArray.GetTypeSize() == GetTypeSize())
+		{
+			TManagedArrayBase::Init(NewArray);
+		}
+		else
+		{
+			check(NewArray.GetTypeSize() == 2 * GetTypeSize());
+			check(NewArray.GetTypeSize() == sizeof(FTransform));
+			const TManagedArrayBase<FTransform>& NewTypedArray = static_cast<const TManagedArrayBase<FTransform>&>(NewArray);
+			const int32 Size = NewTypedArray.Num();
+			Resize(Size);
+			InitHelper(Array, NewTypedArray, Size);
+		}
+	}
+};
+
 template<>
 class TManagedArray<bool> : public FManagedBitArrayBase
 {};
diff --git a/Engine/Source/Runtime/Experimental/Chaos/Public/GeometryCollection/ManagedArrayTypeValues.inl b/Engine/Source/Runtime/Experimental/Chaos/Public/GeometryCollection/ManagedArrayTypeValues.inl
--- a/Engine/Source/Runtime/Experimental/Chaos/Public/GeometryCollection/ManagedArrayTypeValues.inl
+++ b/Engine/Source/Runtime/Experimental/Chaos/Public/GeometryCollection/ManagedArrayTypeValues.inl
@@ -53,6 +53,7 @@ MANAGED_ARRAY_TYPE(TArray<FVector3f>, FVectorArray)
 MANAGED_ARRAY_TYPE(TUniquePtr<Chaos::FPBDRigidParticle>, TPBDRigidParticle3fUniquePtr)
 MANAGED_ARRAY_TYPE(Chaos::FImplicitObjectPtr, FImplicitObjectRefCountedPtr)
 MANAGED_ARRAY_TYPE(Chaos::FConvexPtr, FConvexRefCountedPtr)
+MANAGED_ARRAY_TYPE(FTransform3f, Transform3f)
 
 // NOTE: new types must be added at the bottom to keep serialization from breaking
 
diff --git a/Engine/Source/Runtime/Experimental/Chaos/Public/GeometryCollectionProxyData.h b/Engine/Source/Runtime/Experimental/Chaos/Public/GeometryCollectionProxyData.h
--- a/Engine/Source/Runtime/Experimental/Chaos/Public/GeometryCollectionProxyData.h
+++ b/Engine/Source/Runtime/Experimental/Chaos/Public/GeometryCollectionProxyData.h
@@ -29,7 +29,7 @@ public:
 	FTransformDynamicCollection& operator=(FTransformDynamicCollection&&) = delete;
 
 	// Transform Group
-	TManagedArray<FTransform>   Transform;
+	TManagedArray<FTransform3f> Transform;
 	TManagedArray<int32>        Parent;
 	TManagedArray<TSet<int32>>  Children;
 	TManagedArray<int32>        SimulationType;
diff --git a/Engine/Source/Runtime/Experimental/GeometryCollectionEngine/Private/GeometryCollection/GeometryCollectionComponent.cpp b/Engine/Source/Runtime/Experimental/GeometryCollectionEngine/Private/GeometryCollection/GeometryCollectionComponent.cpp
--- a/Engine/Source/Runtime/Experimental/GeometryCollectionEngine/Private/GeometryCollection/GeometryCollectionComponent.cpp
+++ b/Engine/Source/Runtime/Experimental/GeometryCollectionEngine/Private/GeometryCollection/GeometryCollectionComponent.cpp
@@ -323,6 +323,32 @@ const TManagedArray<Type>& UGeometryCollectionComponent::Get##Name##ArrayRest()
 // Define the methods
 COPY_ON_WRITE_ATTRIBUTES
 
+const TManagedArray<FTransform3f>& UGeometryCollectionComponent::GetTransformArray() const
+{
+	check(IndirectTransformArray != nullptr);
+	return *IndirectTransformArray;
+}
+
+TManagedArray<FTransform3f>& UGeometryCollectionComponent::GetTransformArrayCopyOnWrite()
+{
+	if (!IndirectTransformArray)
+	{
+		DynamicCollection->AddAttribute<FTransform>(FTransformCollection::TransformAttribute, FTransformCollection::TransformGroup);
+		DynamicCollection->CopyAttribute(*RestCollection->GetGeometryCollection(), FTransformCollection::TransformAttribute, FTransformCollection::TransformGroup);
+		IndirectTransformArray = &DynamicCollection->ModifyAttribute<FTransform3f>(FTransformCollection::TransformAttribute, FTransformCollection::TransformGroup);
+		CopyOnWriteAttributeList.Add(reinterpret_cast<FManagedArrayBase**>(&IndirectTransformArray));
+	}
+	return *IndirectTransformArray;
+}
+void UGeometryCollectionComponent::ResetTransformArrayDynamic()
+{
+	IndirectTransformArray = nullptr;
+}
+const TManagedArray<FTransform>& UGeometryCollectionComponent::GetTransformArrayRest() const
+{
+	return RestCollection->GetGeometryCollection()->Transform;
+}
+
 UGeometryCollectionComponent::UGeometryCollectionComponent(const FObjectInitializer& ObjectInitializer)
 	: Super(ObjectInitializer)
 	, ChaosSolverActor(nullptr)
@@ -1087,7 +1113,7 @@ bool UGeometryCollectionComponent::DoCustomNavigableGeometryExport(FNavigableGeo
 			{
 				const int32 RootIndex = GetRootIndex();
 				const FTransform& CompToWorld = GetComponentToWorld(); 
-				const FTransform FinalTransform = (DynamicCollection ? DynamicCollection->Transform[RootIndex] : FTransform::Identity) * CompToWorld;
+				const FTransform FinalTransform = (DynamicCollection ? FTransform(DynamicCollection->Transform[RootIndex]) : FTransform::Identity) * CompToWorld;
 				const FVector Scale3D = FinalTransform.GetScale3D();
 				if (!Scale3D.IsZero())
 				{
@@ -1149,7 +1175,14 @@ bool UGeometryCollectionComponent::DoCustomNavigableGeometryExport(FNavigableGeo
 
 	// Get all the geometry transforms in component space (they are stored natively in parent-bone space)
 	TArray<FTransform> GeomToComponent;
-	GeometryCollectionAlgo::GlobalMatrices(GetTransformArray(), GetParentArray(), TransformIndexBuffer, GeomToComponent);
+	if (DynamicCollection)
+	{
+		GeometryCollectionAlgo::GlobalMatrices(GetTransformArray(), GetParentArray(), TransformIndexBuffer, GeomToComponent);
+	}
+	else
+	{
+		GeometryCollectionAlgo::GlobalMatrices(TManagedArray<FTransform>(GetCurrentRestTransforms()), GetParentArray(), TransformIndexBuffer, GeomToComponent);
+	}
 
 	OutVertexBuffer.AddUninitialized(VertexCount);
 
@@ -2334,11 +2367,11 @@ void UGeometryCollectionComponent::SetInitialTransforms(const TArray<FTransform>
 {
 	if (DynamicCollection)
 	{
-		TManagedArray<FTransform>& Transform = DynamicCollection->Transform;
+		TManagedArray<FTransform3f>& Transform = DynamicCollection->Transform;
 		int32 MaxIdx = FMath::Min(Transform.Num(), InitialTransforms.Num());
 		for (int32 Idx = 0; Idx < MaxIdx; ++Idx)
 		{
-			Transform[Idx] = InitialTransforms[Idx];
+			Transform[Idx] = FTransform3f(InitialTransforms[Idx]);
 		}
 	}
 }
@@ -2426,7 +2459,7 @@ void UGeometryCollectionComponent::GetRestTransforms(TArray<FMatrix44f>& OutRest
 	TArray<FMatrix> RestMatrices;
 	const TArray<FTransform>& RestCollectionTransforms = RestCollection->GetGeometryCollection()->Transform.GetConstArray();
 	const TArray<FTransform>& RestTransformToUse = (RestTransforms.Num() == RestCollectionTransforms.Num()) ? RestTransforms : RestCollectionTransforms;
-	GeometryCollectionAlgo::GlobalMatrices(RestTransformToUse, RestCollection->GetGeometryCollection()->Parent, RestMatrices);
+	GeometryCollectionAlgo::GlobalMatrices(TManagedArray<FTransform>(RestTransformToUse), RestCollection->GetGeometryCollection()->Parent, RestMatrices);
 #if WITH_EDITOR
 	UpdateGlobalMatricesWithExplodedVectors(RestMatrices, *(RestCollection->GetGeometryCollection()));
 #endif
@@ -3386,7 +3419,7 @@ void UGeometryCollectionComponent::UpdateRemovalIfNeeded()
 					const float Scale = 1.0 - Decay;
 					if (Scale < UE_SMALL_NUMBER)
 					{
-						DynamicCollection->Transform[TransformIndex].SetScale3D(FVector::ZeroVector);
+						DynamicCollection->Transform[TransformIndex].SetScale3D(FVector3f::ZeroVector);
 					}
 					// do not try to get this condition out of the loop as this may cause some optimizer related issues
 					else if (RestCollection->bScaleOnRemoval && MassToLocal && CompSpaceTransform)
@@ -3404,7 +3437,7 @@ void UGeometryCollectionComponent::UpdateRemovalIfNeeded()
 						const FVector CenterOfMass = (*MassToLocal)[TransformIndex].GetTranslation();
 						const FVector ScaleCenter = LocalDown + CenterOfMass;
 						const FTransform ScaleTransform(FQuat::Identity, ScaleCenter * FVector::FReal(1.f - Scale), FVector(Scale));
-						DynamicCollection->Transform[TransformIndex] = ScaleTransform * DynamicCollection->Transform[TransformIndex];
+						DynamicCollection->Transform[TransformIndex] = FTransform3f(ScaleTransform) * DynamicCollection->Transform[TransformIndex];
 					}
 				}
 			}
@@ -4734,14 +4767,15 @@ static void CheckForNaNs(const TArray<FTransform>& Transforms)
 }
 #endif
 
-const TArray<FTransform>& UGeometryCollectionComponent::GetCurrentTransforms() const
+const TArray<FTransform>& UGeometryCollectionComponent::GetCurrentRestTransforms() const
 {
 	const bool bRestTransformsOverriden = (RestTransforms.Num() > 0);
-	if (!DynamicCollection && bRestTransformsOverriden)
+	if (bRestTransformsOverriden)
 	{
 		return RestTransforms;
 	}
-	return GetTransformArray().GetConstArray();
+
+	return GetTransformArrayRest().GetConstArray();
 }
 
 const FTransform& UGeometryCollectionComponent::FComponentSpaceTransforms::RequestRootTransform() const
@@ -4756,7 +4790,15 @@ const FTransform& UGeometryCollectionComponent::FComponentSpaceTransforms::Reque
 	{
 		if (bIsRootDirty)
 		{
-			Transforms[RootIndex] = Component->GetCurrentTransforms()[RootIndex];
+			if (Component->DynamicCollection)
+			{
+				Transforms[RootIndex] = FTransform(Component->GetTransformArray()[RootIndex]);
+			}
+			else
+			{
+				Transforms[RootIndex] = Component->GetCurrentRestTransforms()[RootIndex];
+			}
+			
 			bIsRootDirty = false;
 		}
 
@@ -4781,13 +4823,21 @@ const TArray<FTransform>& UGeometryCollectionComponent::FComponentSpaceTransform
 		return Transforms;
 	}
 
-	const TArray<FTransform>& CurrentTransforms = Component->GetCurrentTransforms();
+	int32 CurrentTransformNum;
+	if (Component->DynamicCollection)
+	{
+		CurrentTransformNum = Component->GetTransformArray().Num();
+	}
+	else
+	{
+		CurrentTransformNum = Component->GetCurrentRestTransforms().Num();
+	}
 
 	bool bFastPath = false;
 	if (Component->RestCollection)
 	{
 		const TArray<int32>& BreadthFirstTransformIndices = Component->RestCollection->GetBreadthFirstTransformIndices();
-		bFastPath = (BreadthFirstTransformIndices.Num() == CurrentTransforms.Num());
+		bFastPath = (BreadthFirstTransformIndices.Num() == CurrentTransformNum);
 	}
 
 	if (bFastPath)
@@ -4800,20 +4850,39 @@ const TArray<FTransform>& UGeometryCollectionComponent::FComponentSpaceTransform
 			const int32 TransformIndex = BreadthFirstTransformIndices[Index];
 			const int32 ParentTransformIndex = ParentArray[TransformIndex];
 
+
+			FTransform CurrentTransform;
+			if (Component->DynamicCollection)
+			{
+				CurrentTransform = FTransform(Component->GetTransformArray()[TransformIndex]);
+			}
+			else
+			{
+				CurrentTransform = Component->GetCurrentRestTransforms()[TransformIndex];
+			}
 			if (ParentTransformIndex == INDEX_NONE)
 			{
-				Transforms[TransformIndex] = CurrentTransforms[TransformIndex];
+				
+				Transforms[TransformIndex] = CurrentTransform;
 			}
 			else
 			{
 				const FTransform& ParentTransform = Transforms[ParentTransformIndex];
-				Transforms[TransformIndex] = CurrentTransforms[TransformIndex] * ParentTransform;
+				Transforms[TransformIndex] = CurrentTransform * ParentTransform;
 			}
 		}
 	}
 	else
 	{
-		GeometryCollectionAlgo::GlobalMatrices(Component->GetCurrentTransforms(), Component->GetParentArray(), Transforms);
+		if (Component->DynamicCollection)
+		{
+			GeometryCollectionAlgo::GlobalMatrices(Component->GetTransformArray(), Component->GetParentArray(), Transforms);
+		}
+		else
+		{
+			GeometryCollectionAlgo::GlobalMatrices(TManagedArray<FTransform>(Component->GetCurrentRestTransforms()), Component->GetParentArray(), Transforms);
+		}
+		
 	}
 
 #if WITH_EDITOR
@@ -5231,8 +5300,8 @@ void UGeometryCollectionComponent::IncrementSleepTimer(float DeltaTime)
 						bool ShouldUpdateTimer = (DecayFacade.GetDecay(TransformIdx) > 0);
 						if (!ShouldUpdateTimer && RestCollection->bSlowMovingAsSleeping)
 						{
-							const FVector CurrentPosition = DynamicCollection->Transform[TransformIdx].GetTranslation();
-							ShouldUpdateTimer |= RemoveOnSleepFacade.ComputeSlowMovingState(TransformIdx, CurrentPosition, DeltaTime, RestCollection->SlowMovingVelocityThreshold);
+							const FVector3f CurrentPosition = DynamicCollection->Transform[TransformIdx].GetTranslation();
+							ShouldUpdateTimer |= RemoveOnSleepFacade.ComputeSlowMovingState(TransformIdx, FVector(CurrentPosition), DeltaTime, RestCollection->SlowMovingVelocityThreshold);
 						}
 						if (ShouldUpdateTimer || DynamicStateFacade.IsSleeping(TransformIdx))
 						{
diff --git a/Engine/Source/Runtime/Experimental/GeometryCollectionEngine/Public/GeometryCollection/GeometryCollectionComponent.h b/Engine/Source/Runtime/Experimental/GeometryCollectionEngine/Public/GeometryCollection/GeometryCollectionComponent.h
--- a/Engine/Source/Runtime/Experimental/GeometryCollectionEngine/Public/GeometryCollection/GeometryCollectionComponent.h
+++ b/Engine/Source/Runtime/Experimental/GeometryCollectionEngine/Public/GeometryCollection/GeometryCollectionComponent.h
@@ -760,7 +760,6 @@ public:
 	/* Transform group */																				\
 	COPY_ON_WRITE_ATTRIBUTE(FString, BoneName, FTransformCollection::TransformGroup)					\
 	COPY_ON_WRITE_ATTRIBUTE(FLinearColor, BoneColor, FTransformCollection::TransformGroup)				\
-	COPY_ON_WRITE_ATTRIBUTE(FTransform, Transform, FTransformCollection::TransformGroup)				\
 	COPY_ON_WRITE_ATTRIBUTE(int32, Parent, FTransformCollection::TransformGroup)						\
 	COPY_ON_WRITE_ATTRIBUTE(TSet<int32>, Children, FTransformCollection::TransformGroup)				\
 	COPY_ON_WRITE_ATTRIBUTE(int32, SimulationType, FTransformCollection::TransformGroup)				\
@@ -771,6 +770,15 @@ public:
 	// Declare all the methods
 	COPY_ON_WRITE_ATTRIBUTES
 
+	GEOMETRYCOLLECTIONENGINE_API const TManagedArray<FTransform3f>& GetTransformArray() const;
+	GEOMETRYCOLLECTIONENGINE_API TManagedArray<FTransform3f>& GetTransformArrayCopyOnWrite();
+	GEOMETRYCOLLECTIONENGINE_API void ResetTransformArrayDynamic();
+	GEOMETRYCOLLECTIONENGINE_API const TManagedArray<FTransform>& GetTransformArrayRest() const;
+	private:
+		TManagedArray<FTransform3f>* IndirectTransformArray;
+	public:
+
+
 	UPROPERTY(EditAnywhere, NoClear, BlueprintReadOnly, Category = "ChaosPhysics")
 	TObjectPtr<const UGeometryCollection> RestCollection;
 
@@ -1370,7 +1378,7 @@ private:
 
 	// return the most actual transforms
 	// this can be the rest collection ones, the overriden RestTransforms or the dynamic collection ones
-	const TArray<FTransform>& GetCurrentTransforms() const;
+	const TArray<FTransform>& GetCurrentRestTransforms() const;
 
 	bool bRenderStateDirty;
 	bool bEnableBoneSelection;
