Fix assert
Forgot to checkout a file in CL 27942238

diff --git a/Engine/Source/Runtime/Experimental/Chaos/Private/GeometryCollection/ManagedArrayCollection.cpp b/Engine/Source/Runtime/Experimental/Chaos/Private/GeometryCollection/ManagedArrayCollection.cpp
--- a/Engine/Source/Runtime/Experimental/Chaos/Private/GeometryCollection/ManagedArrayCollection.cpp
+++ b/Engine/Source/Runtime/Experimental/Chaos/Private/GeometryCollection/ManagedArrayCollection.cpp
@@ -547,7 +547,8 @@ void FManagedArrayCollection::CopyAttribute(const FManagedArrayCollection& InCol
 
 	const FValueType& OriginalValue = InCollection.Map[SrcKey];
 	const FValueType& DestValue = Map[DestKey];
-	check(OriginalValue.ArrayType == DestValue.ArrayType);
+	// Now a FTransform3f can be copy in a FTransform, and the way around
+	// check(OriginalValue.ArrayType == DestValue.ArrayType);
 	DestValue.Value->Init(*OriginalValue.Value);
 }
 
