Added trace events to IAS' client for potential later visualisation

#rnx

diff --git a/Engine/Source/Runtime/Experimental/IoStoreOnDemand/Private/CoreHttp/Http.cpp b/Engine/Source/Runtime/Experimental/IoStoreOnDemand/Private/CoreHttp/Http.cpp
--- a/Engine/Source/Runtime/Experimental/IoStoreOnDemand/Private/CoreHttp/Http.cpp
+++ b/Engine/Source/Runtime/Experimental/IoStoreOnDemand/Private/CoreHttp/Http.cpp
@@ -14,6 +14,7 @@
 #include "Misc/StringBuilder.h"
 #include "ProfilingDebugging/CpuProfilerTrace.h"
 #include "Tasks/Task.h"
+#include "Trace/Trace.h"
 #endif
 
 #include <atomic>
@@ -80,6 +81,97 @@ static const SocketType InvalidSocket = ~SocketType(0);
 namespace UE::IO::IAS::HTTP
 {
 
+// {{{1 trace ..................................................................
+
+////////////////////////////////////////////////////////////////////////////////
+#define MAKE_TRACE_ENUM(x) \
+	x(LoopCreate) \
+	x(LoopTick) \
+	x(LoopDestroy) \
+	x(ActivityCreate) \
+	x(ActivityDestroy) \
+	x(RequestBegin) \
+	x(StateChange) \
+	x(StateChangeWait) \
+	x(Wait) \
+	x(Unwait) \
+	x(Connect) \
+	x(Send) \
+	x(Recv) \
+	x(Error) \
+	x($)
+
+enum class ETrace
+{
+#define TRACE_ENUM_VALUES(x) x,
+	MAKE_TRACE_ENUM(TRACE_ENUM_VALUES)
+#undef TRACE_ENUM_VALUES
+};
+
+static const FAnsiStringView TraceEnumNames[] = {
+#define TRACE_ENUM_VALUES(x) #x ,
+	MAKE_TRACE_ENUM(TRACE_ENUM_VALUES)
+#undef TRACE_ENUM_VALUES
+};
+
+#undef MAKE_TRACE_ENUM
+
+UE_TRACE_CHANNEL(IasHttpChannel);
+
+UE_TRACE_EVENT_BEGIN(IasHttp, Enum, NoSync)
+	UE_TRACE_EVENT_FIELD(UE::Trace::AnsiString, Name)
+UE_TRACE_EVENT_END()
+
+UE_TRACE_EVENT_BEGIN(IasHttp, Event, NoSync)
+	UE_TRACE_EVENT_FIELD(uint64, Cycle)
+	UE_TRACE_EVENT_FIELD(uint32, Id)
+	UE_TRACE_EVENT_FIELD(uint32, Param)
+	UE_TRACE_EVENT_FIELD(uint8, Action)
+UE_TRACE_EVENT_END()
+
+static void Activity_TraceStateNames();
+
+////////////////////////////////////////////////////////////////////////////////
+static void TraceEnum(const FAnsiStringView* Names)
+{
+	for (;; ++Names)
+	{
+		UE_TRACE_LOG(IasHttp, Enum, IasHttpChannel)
+			<< Enum.Name(Names[0].GetData(), Names[0].Len());
+
+		if (*Names == "$")
+		{
+			break;
+		}
+	}
+}
+
+////////////////////////////////////////////////////////////////////////////////
+static void TraceInner(const void* Id, ETrace Action, UPTRINT Param)
+{
+#if UE_TRACE_ENABLED
+	static bool Once = [] {
+		TraceEnum(TraceEnumNames);
+		Activity_TraceStateNames();
+		return true;
+	}();
+
+	UE_TRACE_LOG(IasHttp, Event, IasHttpChannel)
+		<< Event.Cycle(FPlatformTime::Cycles64())
+		<< Event.Id(uint32(UPTRINT(Id)))
+		<< Event.Param(uint32(UPTRINT(Param)))
+		<< Event.Action(uint8(Action));
+#endif // UE_TRACE_ENABLED
+}
+
+template <typename T=UPTRINT>
+static void Trace(const void* Id, ETrace Action, T Param=0)
+{
+	TraceInner(Id, Action, UPTRINT(Param));
+}
+
+
+
 // {{{1 misc ...................................................................
 
 ///////////////////////////////////////////////////////////////////////////////
@@ -918,6 +1010,7 @@ struct alignas(16) FActivity
 		Completed,
 		Cancelled,
 		Failed,
+		_Num,
 	};
 
 	int8				Slot = -1;
@@ -938,6 +1031,20 @@ struct alignas(16) FActivity
 	FBuffer				Buffer;
 };
 
+////////////////////////////////////////////////////////////////////////////////
+static void Activity_TraceStateNames()
+{
+	FAnsiStringView StateNames[] = {
+		"Build",	"Resolve",		"Connect",		"Socks",
+		"Send",		"RecvMessage",	"RecvContent",	"RecvStream",
+		"RecvDone",	"Completed",	"Cancelled",	"Failed",
+		"$",
+	};
+	static_assert(UE_ARRAY_COUNT(StateNames) == int32(FActivity::EState::_Num) + 1);
+
+	TraceEnum(StateNames);
+}
+
 ////////////////////////////////////////////////////////////////////////////////
 static FActivity* Activity_Alloc(uint32 BufferSize)
 {
@@ -976,6 +1083,8 @@ static void Activity_Free(FActivity* Activity)
 
 	Activity->~FActivity();
 	FMemory::Free(Activity);
+
+	Trace(Activity, ETrace::ActivityDestroy);
 }
 
 ////////////////////////////////////////////////////////////////////////////////
@@ -986,6 +1095,8 @@ static void Activity_SetError(FActivity* Activity, const char* Reason)
 	Activity->ErrorReason = Reason;
 	Activity->State = FActivity::EState::Failed;
 	Activity->StateParam = LastSocketResult();
+
+	Trace(Activity, ETrace::Error, Activity->StateParam);
 }
 
 // {{{1 request ................................................................
@@ -1336,6 +1447,7 @@ static uint64 ReadyCheck(FActivity** Activities, uint32 Num, uint32 TimeoutMs)
 			{
 				Activities[i]->SocketWait = EWait::None;
 				Ret |= (1ull << Activity->Slot);
+				Trace(Activity, ETrace::Unwait);
 			}
 			break;
 
@@ -1380,6 +1492,7 @@ static uint64 ReadyCheck(FActivity** Activities, uint32 Num, uint32 TimeoutMs)
 		{
 			Activities[i]->SocketWait = EWait::None;
 			Ret |= (1ull << Activities[i]->Slot);
+			Trace(Activities[i], ETrace::Unwait);
 		}
 
 		--SelectNum;
@@ -1409,6 +1522,7 @@ static int32 DoResolve(FActivity* Activity)
 	if (Result > 0)
 	{
 		Activity->State = FActivity::EState::Connect;
+		Trace(Activity, ETrace::StateChange, Activity->State);
 		return 0;
 	}
 
@@ -1416,6 +1530,7 @@ static int32 DoResolve(FActivity* Activity)
 	{
 		Activity->SocketWait = FActivity::EWait::Pool;
 		Activity->State = FActivity::EState::Connect;
+		Trace(Activity, ETrace::StateChangeWait, Activity->State);
 		return 1;
 	}
 
@@ -1429,6 +1544,7 @@ static int32 DoResolve(FActivity* Activity)
 	}
 
 	Activity->State = FActivity::EState::Connect;
+	Trace(Activity, ETrace::StateChange, Activity->State);
 	return 0;
 }
 
@@ -1461,6 +1577,7 @@ static int32 DoConnect(FActivity* Activity)
 		Activity->SocketWait = FActivity::EWait::None;
 		Activity->State = FActivity::EState::Send;
 		Activity->StateParam = 0;
+		Trace(Activity, ETrace::StateChange, Activity->State);
 		return 0;
 	}
 
@@ -1527,6 +1644,8 @@ static int32 DoConnect(FActivity* Activity)
 	AddrInet.sin_port = htons(Port);
 	memcpy(&(AddrInet.sin_addr), &IpAddress, sizeof(IpAddress));
 	{
+		Trace(Activity, ETrace::Connect, IpAddress);
+
 		int Result = connect(Candidate, &(sockaddr&)AddrInet, sizeof(AddrInet));
 		if (Result < 0 && !(IsSocketResult(EWOULDBLOCK) | IsSocketResult(EINPROGRESS)))
 		{
@@ -1544,6 +1663,7 @@ static int32 DoConnect(FActivity* Activity)
 		? FActivity::EState::Socks
 		: FActivity::EState::Send;
 
+	Trace(Activity, ETrace::StateChangeWait, Activity->State);
 	return 1;
 }
 
@@ -1683,6 +1803,7 @@ static int32 DoSocks(FActivity* Activity)
 
 	Activity->State = FActivity::EState::Send;
 	Activity->StateParam = 0;
+	Trace(Activity, ETrace::StateChange, Activity->State);
 	return 0;
 }
 
@@ -1725,10 +1846,14 @@ static int32 DoSend(FActivity* Activity)
 		Activity->StateParam = 0;
 		Activity->State = FActivity::EState::RecvMessage;
 		Activity->SocketWait = FActivity::EWait::Read;
+		Trace(Activity, ETrace::StateChangeWait, Activity->State);
 		return 1;
 	}
 
+	Trace(Activity, ETrace::Send, SendSize);
 	int32 Result = send(Activity->Socket, SendData, SendSize, MsgFlagType(0));
+	Trace(Activity, ETrace::Send, -1);
+
 	if (Result < 0)
 	{
 		if (IsSocketResult(ENOTCONN))
@@ -1762,6 +1887,7 @@ static int32 DoSend(FActivity* Activity)
 	{
 		Activity->StateParam = Index | (Remaining << PackBits);
 		Activity->SocketWait = FActivity::EWait::Write;
+		Trace(Activity, ETrace::Wait);
 		return 1;
 	}
 
@@ -1783,12 +1909,16 @@ static int32 DoRecvMessage(FActivity* Activity)
 	{
 		auto [Dest, DestSize] = Buffer.GetMutableFree(0, PageSize);
 
+		Trace(Activity, ETrace::Recv, -1);
 		int32 Result = recv(Activity->Socket, Dest, DestSize, MsgFlagType(0));
+		Trace(Activity, ETrace::Recv, FMath::Max(Result, 0));
+
 		if (Result < 0)
 		{
 			if (IsSocketResult(EWOULDBLOCK))
 			{
 				Activity->SocketWait = FActivity::EWait::Read;
+				Trace(Activity, ETrace::Wait);
 				return 1;
 			}
 
@@ -1926,6 +2056,7 @@ static int32 DoRecvMessage(FActivity* Activity)
 			return -1;
 		}
 		Activity->State = FActivity::EState::RecvDone;
+		Trace(Activity, ETrace::StateChange, Activity->State);
 		return 0;
 	}
 
@@ -1935,6 +2066,7 @@ static int32 DoRecvMessage(FActivity* Activity)
 
 	Activity->State = bStreamed ? FActivity::EState::RecvStream : FActivity::EState::RecvContent;
 	Activity->StateParam = AlreadyReceived;
+	Trace(Activity, ETrace::StateChange, Activity->State);
 
 	if (AlreadyReceived == 0)
 	{
@@ -1992,13 +2124,17 @@ static FHandlerResult DoRecvContent(FActivity* Activity, uint32 MaxRecvSize)
 			return { 1, RecvSize };
 		}
 
+		Trace(Activity, ETrace::Recv, -1);
 		char* Cursor = (char*)(DestView.GetData()) + Activity->StateParam;
 		int32 Result = recv(Activity->Socket, Cursor, Size, MsgFlagType(0));
+		Trace(Activity, ETrace::Recv, FMath::Max(Result, 0));
+
 		if (Result < 0)
 		{
 			if (IsSocketResult(EWOULDBLOCK))
 			{
 				Activity->SocketWait = FActivity::EWait::Read;
+				Trace(Activity, ETrace::Wait);
 				return { 1, RecvSize };
 			}
 
@@ -2023,6 +2159,7 @@ static FHandlerResult DoRecvContent(FActivity* Activity, uint32 MaxRecvSize)
 	}
 
 	Activity->State = FActivity::EState::RecvDone;
+	Trace(Activity, ETrace::StateChange, Activity->State);
 	return { 0, RecvSize };
 }
 
@@ -2039,6 +2176,7 @@ static int32 DoRecvDone(FActivity* Activity)
 	Activity->Sink(SinkArg);
 
 	Activity->State = FActivity::EState::Completed;
+	Trace(Activity, ETrace::StateChange, Activity->State);
 	return 0;
 }
 
@@ -2060,6 +2198,7 @@ static void DoCancel(FActivity* Activity)
 	}
 
 	Activity->State = FActivity::EState::Cancelled;
+	Trace(Activity, ETrace::StateChange, Activity->State);
 
 	FTicketStatus& SinkArg = *(FTicketStatus*)Activity;
 	Activity->Sink(SinkArg);
@@ -2216,6 +2355,8 @@ FRequest FEventLoop::FImpl::Request(
 	FAnsiStringView Path,
 	FActivity* Activity)
 {
+	Trace(Activity, ETrace::ActivityCreate, this);
+
 	if (Path.Len() == 0)
 	{
 		Path = "/";
@@ -2242,6 +2383,8 @@ FRequest FEventLoop::FImpl::Request(
 ////////////////////////////////////////////////////////////////////////////////
 FTicket FEventLoop::FImpl::Send(FActivity* Activity)
 {
+	Trace(Activity, ETrace::RequestBegin);
+
 	uint64 Slot;
 	{
 		FScopeLock _(&Lock);
@@ -2324,6 +2467,8 @@ uint32 FEventLoop::FImpl::Tick(uint32 PollTimeoutMs)
 	uint64 CancelsLoad = Cancels.load(std::memory_order_relaxed);
 	for (FActivity* Activity : Active)
 	{
+		Trace(this, ETrace::LoopTick, Activity);
+
 		uint64 SlotBit = 1ull << Activity->Slot;
 
 		if (SlotBit & CancelsLoad)
@@ -2456,8 +2601,8 @@ uint32 FEventLoop::FImpl::Tick(uint32 PollTimeoutMs)
 
 
 ////////////////////////////////////////////////////////////////////////////////
-FEventLoop::FEventLoop()						{ Impl = new FEventLoop::FImpl(); }
-FEventLoop::~FEventLoop()						{ delete Impl; }
+FEventLoop::FEventLoop()						{ Impl = new FEventLoop::FImpl(); Trace(Impl, ETrace::LoopCreate); }
+FEventLoop::~FEventLoop()						{ delete Impl; Trace(Impl, ETrace::LoopDestroy); }
 uint32 FEventLoop::Tick(uint32 PollTimeoutMs)	{ return Impl->Tick(PollTimeoutMs); }
 bool FEventLoop::IsIdle() const					{ return Impl->IsIdle(); }
 void FEventLoop::Cancel(FTicket Ticket)			{ return Impl->Cancel(Ticket); }
