Missed a file.

#rnx

diff --git a/Engine/Source/Runtime/Experimental/IoStoreOnDemand/IoStoreOnDemand.Build.cs b/Engine/Source/Runtime/Experimental/IoStoreOnDemand/IoStoreOnDemand.Build.cs
--- a/Engine/Source/Runtime/Experimental/IoStoreOnDemand/IoStoreOnDemand.Build.cs
+++ b/Engine/Source/Runtime/Experimental/IoStoreOnDemand/IoStoreOnDemand.Build.cs
@@ -7,6 +7,7 @@ public class IoStoreOnDemand : ModuleRules
 	public IoStoreOnDemand(ReadOnlyTargetRules Target) : base(Target)
 	{
 		PublicDependencyModuleNames.Add("Core");
+		PublicDependencyModuleNames.Add("TraceLog");
 		PrivateDependencyModuleNames.AddRange(
 			new string[] {
 				"HTTP",
