Tasks are now only used to fetch cached items that are on disk. Strategic step to move cache IO out of tasks.

IIasCache;
- Get() returns hot data straight away (i.e. in memory) or a token that can be used to Materialize() the data

Backend;
- CacheTask is only created if Cache.Get() returned a valid get-token.
- CompleteRequest() can be called immediately if data was returned. Here any read-options should be applied - there never is any however.
- If no data or no token was returned we can go straight off to the cloud.

JC;
- mem-cache only has a Get() and immediately returns any data it has. no need for a Materialize() now.
- because disk-cache has to load off disk, it's API has changed to check for existence, or Materialize() (which I suppose could be called Get() now). As the disk-cache could change, Materialize() does the lookup again. This might not be necessary when IAS Cache thread is doing the work.
- No need for FCache::FEntry() anymore
- FCache::Get(), which previously only did a lookup, split into a Has/Get() pair, plus Materialize() if Get() needs more work to complete.
- The "get token" is opaque, a bit rough and ready, but will do for now.

#rnx

diff --git a/Engine/Source/Runtime/Experimental/IoStoreOnDemand/Private/IasCache.h b/Engine/Source/Runtime/Experimental/IoStoreOnDemand/Private/IasCache.h
--- a/Engine/Source/Runtime/Experimental/IoStoreOnDemand/Private/IasCache.h
+++ b/Engine/Source/Runtime/Experimental/IoStoreOnDemand/Private/IasCache.h
@@ -19,6 +19,9 @@ namespace UE::IO::IAS {
 class IIasCache
 {
 public:
+	using FGetToken = UPTRINT;
+	using FGetWork = UE::Tasks::TTask<TIoStatusOr<FIoBuffer>>;
+
 	virtual ~IIasCache() = default;
 
 	/** Deletes the IAS object, dropping all data persisted to disk and releasing
@@ -30,9 +33,14 @@ public:
 	/** Returns whether the specified cache key is present in the cache. */
 	virtual bool ContainsChunk(const FIoHash& Key) const = 0;
 
-	/** Get the chunk associated with the specified cache key. */
-	virtual UE::Tasks::TTask<TIoStatusOr<FIoBuffer>> Get(
-		const FIoHash& Key,
+	/** Get the chunk associated with the specified cache key. If the data is
+	 already in memory it is return in OutData. Otherwise a FGetToken value is
+	 returned; zero if Key is not found, or non-zero for use with Materialize. */
+	virtual FGetToken Get(const FIoHash& Key, FIoBuffer& OutData) = 0;
+
+	/** Materialize the data for a Get() if it was not immediately available */
+	virtual FGetWork Materialize(
+		FGetToken Token,
 		const FIoReadOptions& Options,
 		const FIoCancellationToken* CancellationToken) = 0;
 
diff --git a/Engine/Source/Runtime/Experimental/IoStoreOnDemand/Private/JournaledCache.cpp b/Engine/Source/Runtime/Experimental/IoStoreOnDemand/Private/JournaledCache.cpp
--- a/Engine/Source/Runtime/Experimental/IoStoreOnDemand/Private/JournaledCache.cpp
+++ b/Engine/Source/Runtime/Experimental/IoStoreOnDemand/Private/JournaledCache.cpp
@@ -50,7 +50,6 @@ namespace UE::IO::IAS::JournaledCache
 
 ////////////////////////////////////////////////////////////////////////////////
 static bool LoadCache(class FDiskCache&);
-using		EntryHandle = UPTRINT;
 
 ////////////////////////////////////////////////////////////////////////////////
 enum class EAilments
@@ -139,8 +138,7 @@ public:
 	uint32			GetCount() const	{ return Items.Num(); }
 	uint32			GetUsed() const		{ return UsedSize; }
 	uint32			GetMax() const		{ return MaxSize; }
-	EntryHandle		Get(uint64 Key) const;
-	bool			Materialize(EntryHandle Handle, FIoBuffer& Out, uint32 Offset=0) const;
+	const FIoBuffer*Get(uint64 Key) const;
 	bool			Put(uint64 Key, FIoBuffer&& Data);
 	int32			Peel(int32 PeelSize, PeelItems& Out);
 	uint32			DebugVisit(void* Param, FDebugCacheEntry::Callback* Callback);
@@ -167,50 +165,23 @@ uint32 FMemCache::GetDemand() const
 }
 
 ////////////////////////////////////////////////////////////////////////////////
-EntryHandle FMemCache::Get(uint64 Key) const
+const FIoBuffer* FMemCache::Get(uint64 Key) const
 {
 	for (auto& Item : Items)
 	{
 		if (Item.Key == Key)
 		{
-			return UPTRINT(&Item);
+			return &(Item.Data);
 		}
 	}
 
-	return 0;
-}
-
-////////////////////////////////////////////////////////////////////////////////
-bool FMemCache::Materialize(EntryHandle Handle, FIoBuffer& Out, uint32 Offset) const
-{
-	TRACE_CPUPROFILER_EVENT_SCOPE(IasCache::Materialize_Memory);
-
-	FItem& Item = *(FItem*)Handle;
-
-	if (Out.GetData() == nullptr)
-	{
-		Out = Item.Data;
-		return true;
-	}
-
-	uint32 Size = uint32(Out.GetSize());
-
-	FMemoryView View = Item.Data.GetView();
-	View = View.RightChop(Offset);
-	if (View.GetSize() > Size)
-	{
-		View = View.Mid(0, Size);
-	}
-
-	Out.GetMutableView().CopyFrom(View);
-
-	return true;
+	return nullptr;
 }
 
 ////////////////////////////////////////////////////////////////////////////////
 bool FMemCache::Put(uint64 Key, FIoBuffer&& Data)
 {
-	if (Get(Key) != 0)
+	if (Get(Key) != nullptr)
 	{
 		uint32 DataSize = uint32(Data.GetSize());
 		FOnDemandIoBackendStats::Get()->OnCachePutExisting(DataSize);
@@ -587,8 +558,8 @@ public:
 	uint32					GetAilments() const;
 	FDiskPhrase				OpenPhrase(uint32 DataSize);
 	void					ClosePhrase(FDiskPhrase&& Phrase);
-	EntryHandle				Get(uint64 Key) const;
-	bool					Materialize(EntryHandle Handle, FIoBuffer& Out, uint32 Offset=0) const;
+	bool					Has(uint64 Key) const;
+	bool					Materialize(uint64 Key, FIoBuffer& Out, uint32 Offset=0) const;
 	int32					Flush();
 	void					Drop();
 	uint32					DebugVisit(void* Param, FDebugCacheEntry::Callback* Callback);
@@ -701,13 +672,13 @@ void FDiskCache::ClosePhrase(FDiskPhrase&& Phrase)
 }
 
 ////////////////////////////////////////////////////////////////////////////////
-EntryHandle FDiskCache::Get(uint64 Key) const
+bool FDiskCache::Has(uint64 Key) const
 {
-	return UPTRINT(DataMap.Find(Key));
+	return (DataMap.Find(Key) != nullptr);
 }
 
 ////////////////////////////////////////////////////////////////////////////////
-bool FDiskCache::Materialize(EntryHandle Handle, FIoBuffer& Out, uint32 Offset) const
+bool FDiskCache::Materialize(uint64 Key, FIoBuffer& Out, uint32 Offset) const
 {
 	TRACE_CPUPROFILER_EVENT_SCOPE(IasCache::Materialize_Disk);
 
@@ -716,9 +687,13 @@ bool FDiskCache::Materialize(EntryHandle Handle, FIoBuffer& Out, uint32 Offset)
 		return false;
 	}
 
-	const FMapEntry& Entry = *(FMapEntry*)Handle;
+	const FMapEntry* Entry = DataMap.Find(Key);
+	if (Entry == nullptr)
+	{
+		return false;
+	}
 
-	uint32 ReadSize = uint32(Entry.Size) - Offset;
+	uint32 ReadSize = uint32(Entry->Size) - Offset;
 
 	if (Out.GetData() == nullptr)
 	{
@@ -727,7 +702,7 @@ bool FDiskCache::Materialize(EntryHandle Handle, FIoBuffer& Out, uint32 Offset)
 
 	ReadSize = FMath::Min<uint32>(uint32(Out.GetSize()), ReadSize);
 
-	DataHandle->Seek(Entry.DataCursor + Offset);
+	DataHandle->Seek(Entry->DataCursor + Offset);
 	return DataHandle->Read(Out.GetData(), ReadSize);
 }
 
@@ -1118,38 +1093,17 @@ public:
 		FString	Path;
 	};
 
-	enum class EHit
-	{
-		None, Memory, Disk
-	};
-
-	class FEntry
-	{
-	public:
-					FEntry() = default;
-					FEntry(FEntry&& Rhs)		= default;
-		FEntry&		operator = (FEntry&& Rhs)	= default;
-		bool		IsHit() const		{ return GetHitType() != EHit::None; }
-		EHit		GetHitType() const	{ return HitType; }
-		bool		Materialize(FIoBuffer& Out, uint32 Offset=0);
-
-	private:
-		friend		FCache;
-		EntryHandle	Handle;
-		EHit		HitType = EHit::None;
-		FReadAccess	Lock;
-		const void*	Owner = nullptr;
-					FEntry(const FEntry& Rhs) = delete;
-		FEntry&		operator = (const FEntry& Rhs) = delete;
-	};
+	using FGetToken = UPTRINT;
 
 					FCache(FConfig&& Config);
 	uint32			GetAilments() const;
 	bool			Load();
 	void			Drop();
 	uint32			GetDemand() const;
-	FEntry			Get(uint64 Key) const;
+	bool			Has(uint64 Key) const;
+	FGetToken		Get(uint64 Key, FIoBuffer& OutData) const;
 	bool			Put(uint64 Key, FIoBuffer& Data);
+	bool			Materialize(FGetToken Token, FIoBuffer& OutData, uint32 Offset=0) const;
 	uint32			Flush();
 	uint32			WriteMemToDisk(int32 Allowance);
 	uint32			DebugVisit(void* Param, FDebugCacheEntry::Callback* Callback);
@@ -1162,31 +1116,6 @@ private:
 	std::atomic_int	Demand;
 };
 
-////////////////////////////////////////////////////////////////////////////////
-bool FCache::FEntry::Materialize(FIoBuffer& Out, uint32 Offset)
-{
-	if (Owner == nullptr)
-	{
-		return false;
-	}
-
-	bool Ret = false;
-	switch (HitType)
-	{
-	case EHit::Memory:	Ret = ((const FMemCache*)Owner)->Materialize(Handle, Out, Offset); break;
-	case EHit::Disk:	Ret = ((const FDiskCache*)Owner)->Materialize(Handle, Out, Offset); break;
-	case EHit::None:	return false;
-	}
-
-	if (!Ret)
-	{
-		return false;
-	}
-
-	FOnDemandIoBackendStats::Get()->OnCacheGet(Out.GetSize());
-	return true;
-}
-
 ////////////////////////////////////////////////////////////////////////////////
 FCache::FCache(FConfig&& Config)
 : MemCache(Config.MemoryQuota)
@@ -1227,28 +1156,34 @@ uint32 FCache::GetDemand() const
 }
 
 ////////////////////////////////////////////////////////////////////////////////
-FCache::FEntry FCache::Get(uint64 Key) const
+bool FCache::Has(uint64 Key) const
 {
-	FEntry Ret;
+	if (FReadAccess _(FsLock); DiskCache.Has(Key))
+	{
+		return true;
+	}
 
-	Ret.Lock = FReadAccess(FsLock);
-	if (Ret.Handle = DiskCache.Get(Key); Ret.Handle)
+	FReadAccess _(MemLock);
+	return (MemCache.Get(Key) != nullptr);
+}
+
+////////////////////////////////////////////////////////////////////////////////
+FCache::FGetToken FCache::Get(uint64 Key, FIoBuffer& OutData) const
+{
+	// Disk first as that will have more data and is more likely to hit
+	if (FReadAccess _(FsLock); DiskCache.Has(Key))
 	{
-		Ret.HitType = EHit::Disk;
-		Ret.Owner = &DiskCache;
-		return Ret;
+		return FGetToken(Key);
 	}
-		
-	Ret.Lock = FReadAccess(MemLock);
-	if (Ret.Handle = MemCache.Get(Key); Ret.Handle)
+
+	// Nothing's on disk, so lets try the memory cache
+	FReadAccess _(MemLock);
+	if (const FIoBuffer* Data = MemCache.Get(Key); Data != nullptr)
 	{
-		Ret.HitType = EHit::Memory;
-		Ret.Owner = &MemCache;
-		return Ret;
+		OutData = *Data;
 	}
 
-	Ret.Lock = FReadAccess();
-	return Ret;
+	return 0;
 }
 
 ////////////////////////////////////////////////////////////////////////////////
@@ -1265,6 +1200,21 @@ bool FCache::Put(uint64 Key, FIoBuffer& Data)
 	return Ok;
 }
 
+////////////////////////////////////////////////////////////////////////////////
+bool FCache::Materialize(FGetToken Token, FIoBuffer& OutData, uint32 Offset) const
+{
+	FReadAccess _(FsLock);
+
+	uint64 Key = Token;
+	if (!DiskCache.Materialize(Key, OutData, Offset))
+	{
+		return false;
+	}
+
+	FOnDemandIoBackendStats::Get()->OnCacheGet(OutData.GetSize());
+	return true;
+}
+
 ////////////////////////////////////////////////////////////////////////////////
 uint32 FCache::Flush()
 {
@@ -1472,13 +1422,12 @@ class FJournaledCache
 	, public FRunnable
 {
 public:
-	using GetRetType = UE::Tasks::TTask<TIoStatusOr<FIoBuffer>>;
-
 								FJournaledCache() = default;
 	bool						Initialize(const TCHAR* RootDir, const FIasCacheConfig& Config);
 	virtual void				Abandon() override;
 	virtual bool				ContainsChunk(const FIoHash& Key) const override;
-	virtual GetRetType			Get(const FIoHash& Key, const FIoReadOptions& Options, const FIoCancellationToken* CancellationToken) override;
+	virtual FGetToken			Get(const FIoHash& Key, FIoBuffer& OutData) override;
+	virtual FGetWork			Materialize(FGetToken Token, const FIoReadOptions& Options, const FIoCancellationToken* CancelToken) override;
 	virtual FIoStatus			Put(const FIoHash& Key, FIoBuffer& Data) override;
 
 private:
@@ -1621,25 +1570,34 @@ void FJournaledCache::Stop()
 bool FJournaledCache::ContainsChunk(const FIoHash& Key) const
 {
 	uint64 InnerKey = ReduceKey(Key);
-	return Cache->Get(InnerKey).IsHit();
+	return Cache->Has(InnerKey);
+}
+
+////////////////////////////////////////////////////////////////////////////////
+FJournaledCache::FGetToken FJournaledCache::Get(const FIoHash& Key, FIoBuffer& OutData)
+{
+	uint64 InnerKey = ReduceKey(Key);
+
+	auto Entry = Cache->Get(InnerKey, OutData);
+	if (Entry == 0)
+	{
+		return 0;
+	}
+
+	static_assert(sizeof(InnerKey) == sizeof(FGetToken));
+	return FGetToken(InnerKey);
 }
 
 ////////////////////////////////////////////////////////////////////////////////
-FJournaledCache::GetRetType	FJournaledCache::Get(
-	const FIoHash& Key,
+FJournaledCache::FGetWork FJournaledCache::Materialize(
+	FGetToken Token,
 	const FIoReadOptions& Options,
 	const FIoCancellationToken* CancelToken)
 {
-	uint64 InnerKey = ReduceKey(Key);
-	return GetPipe.Launch(TEXT("IasCacheGet"), [this, InnerKey, Options, CancelToken] () {
+	return GetPipe.Launch(TEXT("IasCacheGet"), [this, Token, Options, CancelToken] ()
+	{
 		LLM_SCOPE_BYTAG(Ias);
 
-		FCacheInner::FEntry Entry = Cache->Get(InnerKey);
-		if (!Entry.IsHit())
-		{
-			return TIoStatusOr<FIoBuffer>(FIoStatus(EIoErrorCode::Unknown));
-		}
-
 		if (CancelToken != nullptr && CancelToken->IsCancelled())
 		{
 			return TIoStatusOr<FIoBuffer>(FIoStatus(EIoErrorCode::Cancelled));
@@ -1659,7 +1617,7 @@ FJournaledCache::GetRetType	FJournaledCache::Get(
 		}
 
 		uint32 Offset = uint32(Options.GetOffset());
-		if (!Entry.Materialize(Buffer, Offset))
+		if (!Cache->Materialize(Token, Buffer, Offset))
 		{
 			return TIoStatusOr<FIoBuffer>(FIoStatus(EIoErrorCode::ReadError));
 		}
diff --git a/Engine/Source/Runtime/Experimental/IoStoreOnDemand/Private/OnDemandIoDispatcherBackend.cpp b/Engine/Source/Runtime/Experimental/IoStoreOnDemand/Private/OnDemandIoDispatcherBackend.cpp
--- a/Engine/Source/Runtime/Experimental/IoStoreOnDemand/Private/OnDemandIoDispatcherBackend.cpp
+++ b/Engine/Source/Runtime/Experimental/IoStoreOnDemand/Private/OnDemandIoDispatcherBackend.cpp
@@ -1336,9 +1336,28 @@ bool FOnDemandIoBackend::Resolve(FIoRequestImpl* Request)
 
 	if (Cache.IsValid())
 	{
+		const FIoHash& Key = ChunkRequest->Params.ChunkKey;
+		FIoBuffer& Buffer = ChunkRequest->Chunk;
+		IIasCache::FGetToken GetToken = Cache->Get(Key, Buffer);
+
+		if (Buffer.GetData() != nullptr)
+		{
+			ChunkRequest->bCached = true;
+			Launch(UE_SOURCE_LOCATION, [this, ChunkRequest] { CompleteRequest(ChunkRequest); });
+			return true;
+		}
+
+		if (GetToken == 0)
+		{
+			Stats.OnHttpEnqueue();
+			HttpRequests.EnqueueByPriority(ChunkRequest);
+			TickBackendEvent->Trigger();
+			return true;
+		}
+
 		//TODO: Pass priority to cache
 		InflightCacheRequestCount.fetch_add(1, std::memory_order_relaxed);
-		ChunkRequest->CacheTask = Cache->Get(ChunkRequest->Params.ChunkKey, FIoReadOptions(), &ChunkRequest->CancellationToken);
+		ChunkRequest->CacheTask = Cache->Materialize(GetToken, FIoReadOptions(), &ChunkRequest->CancellationToken);
 	}
 
 	const ETaskPriority TaskPriority = ChunkRequest->Priority > IoDispatcherPriority_Medium ? ETaskPriority::High : ETaskPriority::Normal;
