Give the governor some sensible defaults so it's more "raii friendly"
#rnx

diff --git a/Engine/Source/Runtime/Experimental/IoStoreOnDemand/Private/JournaledCache.cpp b/Engine/Source/Runtime/Experimental/IoStoreOnDemand/Private/JournaledCache.cpp
--- a/Engine/Source/Runtime/Experimental/IoStoreOnDemand/Private/JournaledCache.cpp
+++ b/Engine/Source/Runtime/Experimental/IoStoreOnDemand/Private/JournaledCache.cpp
@@ -1282,7 +1282,7 @@ uint32 FCache::DebugVisit(void* Param, FDebugCacheEntry::Callback* Callback)
 class FGovernor
 {
 public:
-			FGovernor() = default;
+			FGovernor();
 	void	Set(uint32 Allowance, uint32 Ops, uint32 Seconds);
 	void	SetDemands(uint32 Threshold, uint32 Boost, uint32 SuperBoost);
 	uint32	BeginAllowance(uint32 DemandPercent);
@@ -1309,6 +1309,12 @@ private:
 	EState	State = EState::Waiting;
 };
 
+////////////////////////////////////////////////////////////////////////////////
+FGovernor::FGovernor()
+{
+	Set(1, 1, 86400);
+}
+
 ////////////////////////////////////////////////////////////////////////////////
 void FGovernor::Set(uint32 Allowance, uint32 Ops, uint32 Seconds)
 {
