Lock helper is no longer required

#rnx

diff --git a/Engine/Source/Runtime/Experimental/IoStoreOnDemand/Private/JournaledCache.cpp b/Engine/Source/Runtime/Experimental/IoStoreOnDemand/Private/JournaledCache.cpp
--- a/Engine/Source/Runtime/Experimental/IoStoreOnDemand/Private/JournaledCache.cpp
+++ b/Engine/Source/Runtime/Experimental/IoStoreOnDemand/Private/JournaledCache.cpp
@@ -87,35 +87,6 @@ static constexpr FStringView GetCacheFsDir()	 { return FStringView(TEXT("ias"));
 static constexpr FStringView GetCacheFsSuffix()  { return FStringView(TEXT(".cache.0")); }
 static constexpr FStringView GetCacheJrnSuffix() { return FStringView(TEXT(".jrn")); }
 
-////////////////////////////////////////////////////////////////////////////////
-template <bool IsExclusive>
-class FAccessScope
-{
-public:
-	FAccessScope(FRWLock& InLock)
-	: Lock(&InLock)
-	{
-		if (IsExclusive) Lock->WriteLock(); else Lock->ReadLock();
-	}
-
-	~FAccessScope()
-	{
-		if (Lock == nullptr) return;
-		if (IsExclusive) Lock->WriteUnlock(); else Lock->ReadUnlock();
-	}
-
-					FAccessScope() = default;
-					FAccessScope(FAccessScope&& Rhs) { Swap(Lock, Rhs.Lock); }
-	FAccessScope&	operator = (FAccessScope&& Rhs) { Swap(Lock, Rhs.Lock); return *this; }
-
-private:
-	FRWLock*		Lock = nullptr;
-					FAccessScope(const FAccessScope&) = delete;
-	FAccessScope	operator = (const FAccessScope&) = delete;
-};
-using FReadAccess	= FAccessScope<false>;
-using FWriteAccess	= FAccessScope<true>;
-
 
 
 // {{{1 mem-cache ..............................................................
@@ -1138,14 +1109,14 @@ uint32 FCache::GetAilments() const
 ////////////////////////////////////////////////////////////////////////////////
 bool FCache::Load()
 {
-	FWriteAccess _(FsLock);
+	FWriteScopeLock _(FsLock);
 	return LoadCache(DiskCache);
 }
 
 ////////////////////////////////////////////////////////////////////////////////
 void FCache::Drop()
 {
-	FWriteAccess _(FsLock);
+	FWriteScopeLock _(FsLock);
 	DiskCache.Drop();
 }
 
@@ -1158,12 +1129,12 @@ uint32 FCache::GetDemand() const
 ////////////////////////////////////////////////////////////////////////////////
 bool FCache::Has(uint64 Key) const
 {
-	if (FReadAccess _(FsLock); DiskCache.Has(Key))
+	if (FReadScopeLock _(FsLock); DiskCache.Has(Key))
 	{
 		return true;
 	}
 
-	FReadAccess _(MemLock);
+	FReadScopeLock _(MemLock);
 	return (MemCache.Get(Key) != nullptr);
 }
 
@@ -1171,13 +1142,13 @@ bool FCache::Has(uint64 Key) const
 FCache::FGetToken FCache::Get(uint64 Key, FIoBuffer& OutData) const
 {
 	// Disk first as that will have more data and is more likely to hit
-	if (FReadAccess _(FsLock); DiskCache.Has(Key))
+	if (FReadScopeLock _(FsLock); DiskCache.Has(Key))
 	{
 		return FGetToken(Key);
 	}
 
 	// Nothing's on disk, so lets try the memory cache
-	FReadAccess _(MemLock);
+	FReadScopeLock _(MemLock);
 	if (const FIoBuffer* Data = MemCache.Get(Key); Data != nullptr)
 	{
 		OutData = *Data;
@@ -1190,7 +1161,7 @@ FCache::FGetToken FCache::Get(uint64 Key, FIoBuffer& OutData) const
 bool FCache::Put(uint64 Key, FIoBuffer& Data)
 {
 	FIoBuffer Cloned = Data;
-	FWriteAccess _(MemLock);
+	FWriteScopeLock _(MemLock);
 	bool Ok = MemCache.Put(Key, MoveTemp(Cloned));
 	if (Ok)
 	{
@@ -1203,7 +1174,7 @@ bool FCache::Put(uint64 Key, FIoBuffer& Data)
 ////////////////////////////////////////////////////////////////////////////////
 bool FCache::Materialize(FGetToken Token, FIoBuffer& OutData, uint32 Offset) const
 {
-	FReadAccess _(FsLock);
+	FReadScopeLock _(FsLock);
 
 	uint64 Key = Token;
 	if (!DiskCache.Materialize(Key, OutData, Offset))
@@ -1218,7 +1189,7 @@ bool FCache::Materialize(FGetToken Token, FIoBuffer& OutData, uint32 Offset) con
 ////////////////////////////////////////////////////////////////////////////////
 uint32 FCache::Flush()
 {
-	FWriteAccess _(FsLock);
+	FWriteScopeLock _(FsLock);
 	return DiskCache.Flush();
 }
 
@@ -1230,7 +1201,7 @@ uint32 FCache::WriteMemToDisk(int32 Allowance)
 	FMemCache::PeelItems PeelItems;
 	int32 MemCacheSize;
 	{
-		FWriteAccess _(MemLock);
+		FWriteScopeLock _(MemLock);
 
 		MemCacheSize = MemCache.Peel(Allowance, PeelItems);
 
@@ -1257,7 +1228,7 @@ uint32 FCache::WriteMemToDisk(int32 Allowance)
 	}
 
 	{
-		FWriteAccess _(FsLock);
+		FWriteScopeLock _(FsLock);
 		DiskCache.ClosePhrase(MoveTemp(Phrase));
 	}
 
@@ -1267,7 +1238,7 @@ uint32 FCache::WriteMemToDisk(int32 Allowance)
 ////////////////////////////////////////////////////////////////////////////////
 uint32 FCache::DebugVisit(void* Param, FDebugCacheEntry::Callback* Callback)
 {
-	FReadAccess _[] = { MemLock, FsLock };
+	FReadScopeLock _[] = { FReadScopeLock(MemLock), FReadScopeLock(FsLock) };
 	uint32 Count = 0;
 	Count += MemCache.DebugVisit(Param, Callback);
 	Count += DiskCache.DebugVisit(Param, Callback);
