Chaos
- fix islands not waking when a particle is disabled
- fix waking of isolated particles when all other particles or constraints have been removed from their island

1) We were not waking islands on removal of a particle. 2) We were destroying single-particle islands before the sleep status was processed, preventing particles from being awoken when the penultimate particle was removed from an island.

The fix is to explicitly wake islands on particle removal, and to defer island/node destruction until FinalizeIslands, after we have processed sleep/wake changes and synched the sleep state to particles.

diff --git a/Engine/Source/Programs/HeadlessChaos/Private/HeadlessChaosTestGraphEvolution.cpp b/Engine/Source/Programs/HeadlessChaos/Private/HeadlessChaosTestGraphEvolution.cpp
--- a/Engine/Source/Programs/HeadlessChaos/Private/HeadlessChaosTestGraphEvolution.cpp
+++ b/Engine/Source/Programs/HeadlessChaos/Private/HeadlessChaosTestGraphEvolution.cpp
@@ -115,7 +115,7 @@ namespace ChaosTest
 	}
 
 	// Start with a kinematic connected to a dynamic. Verify that removing the
-	// constraint removed both particles.
+	// constraint removes both particles.
 	// This version explicitly removes the constraint from the graph.
 	// 
 	//		(d=dynamic, s=sleeping, k=kinematic)
@@ -151,7 +151,7 @@ namespace ChaosTest
 	}
 
 	// Start with a kinematic connected to a dynamic. Verify that removing the
-	// constraint removed both particles.
+	// constraint removes both particles.
 	// This version has the constraint removed by making all particles kinematic
 	//		(d=dynamic, s=sleeping, k=kinematic)
 	//		Ak - Bd
@@ -185,6 +185,62 @@ namespace ChaosTest
 		EXPECT_FALSE(Test.ConstraintHandles[0]->IsInConstraintGraph());
 	}
 
+	// Start with a kinematic connected to a sleeping dynamic. Verify that removing the
+	// kinematic removes both particles from the graph (because we do not keep islands 
+	// unless there are constraints in them), but the dynamic particle is now awake.
+	//		(d=dynamic, s=sleeping, k=kinematic)
+	//		Ak - Bs
+	// =>	{} (B is dynamic/awake but the graph is empty)
+	//
+	// This tests a bug where we were not waking a particle if we removed all other particles 
+	// from its island. We now defer island destruction to after sleep handling.
+	//
+	GTEST_TEST(GraphEvolutionTests, TestConstraintGraph_KinematicDynamic_RemoveKinematic)
+	{
+		FGraphEvolutionTest Test(2);
+		Test.MakeChain();
+
+		Test.Evolution.SetParticleObjectState(Test.ParticleHandles[0], EObjectStateType::Kinematic);
+
+		Test.AdvanceUntilSleeping();
+
+		// Should have 1 island
+		EXPECT_EQ(Test.IslandManager->GetNumIslands(), 1);
+		EXPECT_TRUE(Test.ParticleHandles[0]->IsInConstraintGraph());
+		EXPECT_TRUE(Test.ParticleHandles[1]->IsInConstraintGraph());
+		EXPECT_TRUE(Test.ConstraintHandles[0]->IsInConstraintGraph());
+
+		// B should be asleep
+		EXPECT_TRUE(Test.ParticleHandles[1]->IsSleeping());
+
+		// Remove A
+		Test.Evolution.DisableParticle(Test.ParticleHandles[0]);
+
+		// A and the constraint should have been removed from the graph
+		EXPECT_FALSE(Test.ParticleHandles[0]->IsInConstraintGraph());
+		EXPECT_FALSE(Test.ConstraintHandles[0]->IsInConstraintGraph());
+
+		// B will be removed from the graph because we do not track islands without
+		// constraints, but not until the next tick. For now it will still be in
+		// the graph and still asleep.
+		EXPECT_EQ(Test.IslandManager->GetNumIslands(), 1);
+		EXPECT_TRUE(Test.ParticleHandles[1]->IsInConstraintGraph());
+		EXPECT_TRUE(Test.ParticleHandles[1]->IsSleeping());
+
+		// Tick physics. This will update the graph, waking B's island and therefore B.
+		// B's island will then be destroyed because it has no constraints.
+		Test.Advance();
+
+		// Graph should be empty
+		EXPECT_EQ(Test.IslandManager->GetNumIslands(), 0);
+		EXPECT_FALSE(Test.ParticleHandles[0]->IsInConstraintGraph());
+		EXPECT_FALSE(Test.ParticleHandles[1]->IsInConstraintGraph());
+		EXPECT_FALSE(Test.ConstraintHandles[0]->IsInConstraintGraph());
+		
+		// B should be awake
+		EXPECT_FALSE(Test.ParticleHandles[1]->IsSleeping());
+	}
+
 	// Start with an island containing 4 particle connected in a chain, then make the second one kinematic.
 	// Check that the island splits.
 	//		(d=dynamic, s=sleeping, k=kinematic)
@@ -240,9 +296,10 @@ namespace ChaosTest
 		// Invalidate B
 		Test.Evolution.InvalidateParticle(Test.ParticleHandles[1]);
 
-		// Everything was kicked from the graph
-		EXPECT_FALSE(Test.ParticleHandles[0]->IsInConstraintGraph());
-		EXPECT_FALSE(Test.ParticleHandles[1]->IsInConstraintGraph());
+		// Constraint was kicked from the graph, but particles remain until
+		// explicitly removed or they have no constraint on next update
+		EXPECT_TRUE(Test.ParticleHandles[0]->IsInConstraintGraph());
+		EXPECT_TRUE(Test.ParticleHandles[1]->IsInConstraintGraph());
 		EXPECT_FALSE(Test.ConstraintHandles[0]->IsInConstraintGraph());
 
 		Test.Advance();
diff --git a/Engine/Source/Runtime/Experimental/Chaos/Private/Chaos/Island/IslandManager.cpp b/Engine/Source/Runtime/Experimental/Chaos/Private/Chaos/Island/IslandManager.cpp
--- a/Engine/Source/Runtime/Experimental/Chaos/Private/Chaos/Island/IslandManager.cpp
+++ b/Engine/Source/Runtime/Experimental/Chaos/Private/Chaos/Island/IslandManager.cpp
@@ -662,14 +662,18 @@ namespace Chaos::Private
 	{
 		if (FPBDIslandParticle* Node = GetGraphNode(Particle))
 		{
-			// All nodes in the graph have at least one edge
-			check(Node->Edges.Num() > 0);
+			// Islands are awakened when we disable a particle in it
+			WakeNodeIslands(Node);
 
-			// Remove all the constraints attached to the particle. This will
-			// also destroy the particle's node (and possibly others) 
-			// because empty nodes are removed as they occur
+			// Remove all the constraints attached to the particle
 			RemoveParticleConstraints(Particle);
 
+			// Remove from our island
+			RemoveNodeFromIsland(Node);
+
+			// Destroy the node
+			DestroyGraphNode(Node);
+
 			// We should not be in the graph anymore
 			check(GetGraphNode(Particle) == nullptr);
 		}
@@ -682,10 +686,10 @@ namespace Chaos::Private
 		const int32 VisitEpoch = GetNextVisitEpoch();
 
 		// Process state changed from registered particles.
-		// It would be nice to iterate over transient particles here but the view we need
-		// (Particles.GetActiveDynamicMovingKinematicParticlesView()) actually iterates over
-		// arrays of handles so we lose the benefits. We would also have to visit nodes that
-		// are not represented in that view
+		// To reduce cache-misses it would be nice to iterate over transient particle handles 
+		// here but the view we need (Particles.GetActiveDynamicMovingKinematicParticlesView()) 
+		// actually holds non-transient handles so we lose the benefits.
+		// We would also have to visit nodes that are not represented in that view
 		for (int32 NodeIndex = 0; NodeIndex < Nodes.Num(); ++NodeIndex)
 		{
 			if (NodeIndex < Nodes.Num() - 1)
@@ -698,15 +702,17 @@ namespace Chaos::Private
 			UpdateGraphNode(Node);
 		}
 
-		// The above loop may leave orphaned nodes in the array when dyanmics are made
+		// The above loop may leave orphaned nodes in the Nodes array when dyanmics are made
 		// kinematic and we have removed any kinematic-kinematic constraints. We don't 
 		// handle node removal above because updating one node to kinematic can result 
 		// in removal of many nodes which is awkward while iterating
 		for (int32 NodeIndex = Nodes.Num() - 1; NodeIndex >= 0; --NodeIndex)
 		{
 			FPBDIslandParticle* Node = Nodes[NodeIndex];
-			if (Node->Edges.IsEmpty())
+			if (Node->Edges.IsEmpty() && (Node->Island == nullptr))
 			{
+				// We have an orphaned node - it must be removed from the graph here because
+				// all other loops are over islands and won't visit this node
 				DestroyGraphNode(Node);
 			}
 		}
@@ -734,36 +740,12 @@ namespace Chaos::Private
 		{
 			FPBDIslandParticle* Node0 = Edge->Nodes[0];
 			FPBDIslandParticle* Node1 = Edge->Nodes[1];
-			FPBDIsland* Island = Edge->Island;
 
 			// Remove edge from the island
 			RemoveEdgeFromIsland(Edge);
 
 			// Destroy the edge (also disconnects it from its nodes)
 			DestroyGraphEdge(Edge);
-
-			// Destroy the island if it is now empty of constraints
-			// NOTE: If the island is queued for merging we won't destroy it now. 
-			// It will get destroyed later anyway, but more importantly it may 
-			// contain particles that need to be copied to the merged island.
-			if ((Island != nullptr) && (Island->NumEdges == 0) && (Island->MergeSetIslandIndex == INDEX_NONE))
-			{
-				RemoveNodeFromIsland(Node0);
-				RemoveNodeFromIsland(Node1);
-				DestroyIsland(Island);
-			}
-
-			// Destroy the nodes if they have no constraints
-			if ((Node0 != nullptr) && Node0->Edges.IsEmpty())
-			{
-				RemoveNodeFromIsland(Node0);
-				DestroyGraphNode(Node0);
-			}
-			if ((Node1 != nullptr) && Node1->Edges.IsEmpty())
-			{
-				RemoveNodeFromIsland(Node1);
-				DestroyGraphNode(Node1);
-			}
 		}
 	}
 
@@ -779,9 +761,6 @@ namespace Chaos::Private
 				// NOTE: Destroys the edge and possibly this and other nodes (if this was their last constraint)
 				RemoveConstraint(Edge->Constraint);
 			}
-
-			// A node with no edges will have been removed
-			check(GetGraphNode(Particle) == nullptr);
 		}
 	}
 
@@ -823,9 +802,6 @@ namespace Chaos::Private
 			FPBDIslandConstraint* Edge = Edges[EdgeIndex];
 			RemoveConstraint(Edge->Constraint);
 		}
-
-		// Removing all constraints will remove all islands
-		check(Islands.IsEmpty());
 	}
 
 	void FPBDIslandManager::SetParticleIslandIsSleeping(FGeometryParticleHandle* Particle, const bool bInIsSleeping)
@@ -1305,6 +1281,28 @@ namespace Chaos::Private
 		}
 	}
 
+	void FPBDIslandManager::WakeNodeIslands(const FPBDIslandParticle* Node)
+	{
+		check(Node != nullptr);
+
+		// We want to wake up this tick, regardless of sleep thresholds etc
+		const bool bIsSleepAllowed = false;
+
+		// Dynamic nodes are in one island, kinematic nodes may be in many
+		// and we must vist their edges to discover them
+		if (Node->Flags.bIsDynamic)
+		{
+			EnqueueIslandCheckSleep(Node->Island, bIsSleepAllowed);
+		}
+		else
+		{
+			for (FPBDIslandConstraint* Edge : Node->Edges)
+			{
+				EnqueueIslandCheckSleep(Edge->Island, bIsSleepAllowed);
+			}
+		}
+	}
+
 	void FPBDIslandManager::EnqueueIslandCheckSleep(FPBDIsland* Island, const bool bIsSleepAllowed)
 	{
 		if (Island != nullptr)
@@ -1879,13 +1877,6 @@ namespace Chaos::Private
 		{
 			FPBDIsland* Island = Islands[IslandIndex];
 
-			// Remove empty islands
-			if (Island->NumEdges == 0)
-			{
-				DestroyIsland(Island);
-				continue;
-			}
-
 			// If the island is awake or just changed state, make sure all the particles and constraints agree
 			// NOTE: also handles the case where we add to a sleeping island (by checking bCheckSleep)
 			// @todo(chaos): we should only need to do this when the state changes...?
@@ -1895,6 +1886,23 @@ namespace Chaos::Private
 				Island->UpdateSyncState();
 			}
 
+			// Remove islands without constraints. NOTE: must come after PropagateIslandSleep
+			// so that we wake nodes the have been left on their own after all other nodes were removed.
+			if (Island->NumEdges == 0)
+			{
+				// Destroy any nodes left in the island
+				for (int32 IslandNodeIndex = Island->Nodes.Num() - 1; IslandNodeIndex >= 0; --IslandNodeIndex)
+				{
+					FPBDIslandParticle* Node = Island->Nodes[IslandNodeIndex];
+					check(Node->Edges.IsEmpty());
+					RemoveNodeFromIsland(Node);
+					DestroyGraphNode(Node);
+				}
+
+				DestroyIsland(Island);
+				continue;
+			}
+
 			// Reset of the sleep counter if the island is:
 			// - Not allowed to sleep this tick
 			// - Sleeping, since as soon as it wakes up we could start incrementing the counter
diff --git a/Engine/Source/Runtime/Experimental/Chaos/Private/Chaos/PBDNullConstraints.cpp b/Engine/Source/Runtime/Experimental/Chaos/Private/Chaos/PBDNullConstraints.cpp
--- a/Engine/Source/Runtime/Experimental/Chaos/Private/Chaos/PBDNullConstraints.cpp
+++ b/Engine/Source/Runtime/Experimental/Chaos/Private/Chaos/PBDNullConstraints.cpp
@@ -1,9 +1,33 @@
 // Copyright Epic Games, Inc. All Rights Reserved.
 #include "Chaos/PBDNullConstraints.h"
 #include "Chaos/Island/IslandManager.h"
+#include "Chaos/ParticleHandle.h"
 
 namespace Chaos
 {
+	FPBDNullConstraints::FPBDNullConstraints()
+		: TPBDIndexedConstraintContainer<FPBDNullConstraints>(FPBDNullConstraintHandle::StaticType())
+	{
+	}
+
+	FPBDNullConstraintHandle* FPBDNullConstraints::AddConstraint(const TVec2<FGeometryParticleHandle*>& InConstraintedParticles)
+	{
+		const int32 ConstraintIndex = Constraints.Emplace(FPBDNullConstraint(InConstraintedParticles));
+		const int32 HandleIndex = Handles.Emplace(HandleAllocator.AllocHandle(this, ConstraintIndex));
+		check(ConstraintIndex == HandleIndex);
+
+		if (InConstraintedParticles[0] != nullptr)
+		{
+			InConstraintedParticles[0]->AddConstraintHandle(Handles[HandleIndex]);
+		}
+		if (InConstraintedParticles[1] != nullptr)
+		{
+			InConstraintedParticles[1]->AddConstraintHandle(Handles[HandleIndex]);
+		}
+
+		return Handles[HandleIndex];
+	}
+
 	void FPBDNullConstraints::AddConstraintsToGraph(Private::FPBDIslandManager& IslandManager)
 	{
 		IslandManager.AddContainerConstraints(*this);
diff --git a/Engine/Source/Runtime/Experimental/Chaos/Public/Chaos/Island/IslandManager.h b/Engine/Source/Runtime/Experimental/Chaos/Public/Chaos/Island/IslandManager.h
--- a/Engine/Source/Runtime/Experimental/Chaos/Public/Chaos/Island/IslandManager.h
+++ b/Engine/Source/Runtime/Experimental/Chaos/Public/Chaos/Island/IslandManager.h
@@ -503,6 +503,7 @@ namespace Chaos::Private
 		CHAOS_API void RemoveNodeFromIsland(FPBDIslandParticle* Node);
 		CHAOS_API void AddEdgeToIsland(FPBDIslandConstraint* Edge, FPBDIsland* Island);
 		CHAOS_API void RemoveEdgeFromIsland(FPBDIslandConstraint* Edge);
+		CHAOS_API void WakeNodeIslands(const FPBDIslandParticle* Node);
 		CHAOS_API void EnqueueIslandCheckSleep(FPBDIsland* Island, const bool bIsSleepAllowed);
 
 		// Island merge set management
diff --git a/Engine/Source/Runtime/Experimental/Chaos/Public/Chaos/PBDNullConstraints.h b/Engine/Source/Runtime/Experimental/Chaos/Public/Chaos/PBDNullConstraints.h
--- a/Engine/Source/Runtime/Experimental/Chaos/Public/Chaos/PBDNullConstraints.h
+++ b/Engine/Source/Runtime/Experimental/Chaos/Public/Chaos/PBDNullConstraints.h
@@ -146,17 +146,4 @@ namespace Chaos
 	};
 
 
-	inline FPBDNullConstraints::FPBDNullConstraints()
-		: TPBDIndexedConstraintContainer<FPBDNullConstraints>(FPBDNullConstraintHandle::StaticType())
-	{
-	}
-
-	inline FPBDNullConstraintHandle* FPBDNullConstraints::AddConstraint(const TVec2<FGeometryParticleHandle*>& InConstraintedParticles)
-	{
-		const int32 ConstraintIndex = Constraints.Emplace(FPBDNullConstraint(InConstraintedParticles));
-		const int32 HandleIndex = Handles.Emplace(HandleAllocator.AllocHandle(this, ConstraintIndex));
-		check(ConstraintIndex == HandleIndex);
-		return Handles[HandleIndex];
-	}
-
 }
\ No newline at end of file
diff --git a/Engine/Source/Runtime/Experimental/Chaos/Public/Chaos/PBDRigidsEvolution.h b/Engine/Source/Runtime/Experimental/Chaos/Public/Chaos/PBDRigidsEvolution.h
--- a/Engine/Source/Runtime/Experimental/Chaos/Public/Chaos/PBDRigidsEvolution.h
+++ b/Engine/Source/Runtime/Experimental/Chaos/Public/Chaos/PBDRigidsEvolution.h
@@ -447,11 +447,14 @@ public:
 		TestModeParticleDisabled(Particle);
 #endif
 
+		// NOTE: kinematics must visit their graph edges to determine what islands they are in, so we must remove the 
+		// particle from the graph before we disable its constraints or we don't know what island(s) to wake.
+		IslandManager.RemoveParticle(Particle);
+
 		RemoveParticleFromAccelerationStructure(*Particle);
 		Particles.DisableParticle(Particle);
 		DisableConstraints(Particle);
 		DestroyTransientConstraints(Particle);
-		IslandManager.RemoveParticle(Particle);
 	}
 
 	/**
