Only give access to Current and PrevFrame scope list via FNetBitArrayView

#rnx

diff --git a/Engine/Source/Runtime/Experimental/Iris/Core/Private/Iris/ReplicationSystem/Conditionals/ReplicationConditionals.cpp b/Engine/Source/Runtime/Experimental/Iris/Core/Private/Iris/ReplicationSystem/Conditionals/ReplicationConditionals.cpp
--- a/Engine/Source/Runtime/Experimental/Iris/Core/Private/Iris/ReplicationSystem/Conditionals/ReplicationConditionals.cpp
+++ b/Engine/Source/Runtime/Experimental/Iris/Core/Private/Iris/ReplicationSystem/Conditionals/ReplicationConditionals.cpp
@@ -681,8 +681,8 @@ void FReplicationConditionals::UpdateObjectsInScope()
 {
 	IRIS_PROFILER_SCOPE(FReplicationConditionals_UpdateObjectsInScope);
 
-	const FNetBitArray& ObjectsInScope = NetRefHandleManager->GetScopableInternalIndices();
-	const FNetBitArray& PrevObjectsInScope = NetRefHandleManager->GetPrevFrameScopableInternalIndices();
+	const FNetBitArrayView ObjectsInScope = NetRefHandleManager->GetScopableInternalIndicesView();
+	const FNetBitArrayView PrevObjectsInScope = NetRefHandleManager->GetPrevFrameScopableInternalIndicesView();
 
 	const uint32 WordCountForModifiedWords = Align(FPlatformMath::Max(MaxObjectCount, 1U), 32U)/32U;
 	TArray<uint32> ModifiedWords;
diff --git a/Engine/Source/Runtime/Experimental/Iris/Core/Private/Iris/ReplicationSystem/DeltaCompression/DeltaCompressionBaselineManager.cpp b/Engine/Source/Runtime/Experimental/Iris/Core/Private/Iris/ReplicationSystem/DeltaCompression/DeltaCompressionBaselineManager.cpp
--- a/Engine/Source/Runtime/Experimental/Iris/Core/Private/Iris/ReplicationSystem/DeltaCompression/DeltaCompressionBaselineManager.cpp
+++ b/Engine/Source/Runtime/Experimental/Iris/Core/Private/Iris/ReplicationSystem/DeltaCompression/DeltaCompressionBaselineManager.cpp
@@ -158,8 +158,8 @@ void FDeltaCompressionBaselineManager::UpdateScope()
 	using SignedWordType = TSignedIntType<sizeof(WordType)>::Type;
 	constexpr SIZE_T BitsPerWord = sizeof(WordType)*8U;
 
-	const FNetBitArray& ObjectsInScope = NetRefHandleManager->GetScopableInternalIndices();
-	const FNetBitArray& PrevObjectsInScope = NetRefHandleManager->GetPrevFrameScopableInternalIndices();
+	const FNetBitArrayView ObjectsInScope = NetRefHandleManager->GetScopableInternalIndicesView();
+	const FNetBitArrayView PrevObjectsInScope = NetRefHandleManager->GetPrevFrameScopableInternalIndicesView();
 
 	const WordType* ObjectsInScopeStorage = ObjectsInScope.GetData();
 	const WordType* PrevObjectsInScopeStorage = PrevObjectsInScope.GetData();
diff --git a/Engine/Source/Runtime/Experimental/Iris/Core/Private/Iris/ReplicationSystem/Filtering/ReplicationFiltering.cpp b/Engine/Source/Runtime/Experimental/Iris/Core/Private/Iris/ReplicationSystem/Filtering/ReplicationFiltering.cpp
--- a/Engine/Source/Runtime/Experimental/Iris/Core/Private/Iris/ReplicationSystem/Filtering/ReplicationFiltering.cpp
+++ b/Engine/Source/Runtime/Experimental/Iris/Core/Private/Iris/ReplicationSystem/Filtering/ReplicationFiltering.cpp
@@ -328,11 +328,11 @@ void FReplicationFiltering::FilterNonRelevantObjects()
 	if (!bCVarRepFilterCullNonRelevant)
 	{
 		// Make every object in the global scope part of the relevant list
-		NetRefHandleManager->GetRelevantObjectsInternalIndices().Copy(MakeNetBitArrayView(NetRefHandleManager->GetScopableInternalIndices()));
+		NetRefHandleManager->GetRelevantObjectsInternalIndices().Copy(NetRefHandleManager->GetScopableInternalIndicesView());
 		return;
 	}
 
-	const uint32* const GlobalScopeData = NetRefHandleManager->GetScopableInternalIndices().GetData();
+	const uint32* const GlobalScopeData = NetRefHandleManager->GetScopableInternalIndicesView().GetData();
 	const uint32* const WithOwnerData = ObjectsWithOwnerFilter.GetData();
 	const uint32* const ConnectionFiltersData = AllConnectionFilteredObjects.GetData();
 	const uint32* const DynamicFilteredData = DynamicFilterEnabledObjects.GetData();
@@ -663,7 +663,7 @@ void FReplicationFiltering::InitNewConnections()
 	auto InitNewConnection = [this](uint32 ConnectionId)
 	{
 		// Copy default scope
-		const FNetBitArray& ScopableInternalIndices = NetRefHandleManager->GetScopableInternalIndices();
+		const FNetBitArrayView ScopableInternalIndices = NetRefHandleManager->GetScopableInternalIndicesView();
 		FPerConnectionInfo& ConnectionInfo = this->ConnectionInfos[ConnectionId];
 
 		ConnectionInfo.ConnectionFilteredObjects.Init(MaxObjectCount);
@@ -778,8 +778,8 @@ void FReplicationFiltering::ResetRemovedConnections()
 
 void FReplicationFiltering::UpdateObjectsInScope()
 {
-	const FNetBitArray& ObjectsInScope = NetRefHandleManager->GetScopableInternalIndices();
-	const FNetBitArray& PrevObjectsInScope = NetRefHandleManager->GetPrevFrameScopableInternalIndices();
+	const FNetBitArrayView ObjectsInScope = NetRefHandleManager->GetScopableInternalIndicesView();
+	const FNetBitArrayView PrevObjectsInScope = NetRefHandleManager->GetPrevFrameScopableInternalIndicesView();
 
 	/**
 	 * It's possible for an object to be created, have some filtering applied and then be removed later the same frame.
@@ -996,7 +996,7 @@ void FReplicationFiltering::UpdateOwnerAndConnectionFiltering()
 		ObjectsWithDirtyOwner.ForAllSetBits(UpdateOwners);
 	}
 
-	const FNetBitArray& GlobalObjectsInScope = NetRefHandleManager->GetScopableInternalIndices();
+	const FNetBitArrayView GlobalObjectsInScope = NetRefHandleManager->GetScopableInternalIndicesView();
 
 	// Update filtering
 	if (bHasDirtyConnectionFilter)
@@ -1078,13 +1078,14 @@ void FReplicationFiltering::UpdateGroupFiltering()
 			const FNetObjectGroup* Group = this->Groups->GetGroupByIndex(FNetObjectGroupHandle::FGroupIndexType(GroupIndex));
 			FPerConnectionInfo* LocalConnectionInfos = this->ConnectionInfos.GetData();
 
-			auto UpdateGroupFilterForConnection = [this, ConnectionStateInfo, Group, LocalConnectionInfos](uint32 ConnectionId)
+			const FNetBitArrayView GlobalScopableObjects = NetRefHandleManager->GetScopableInternalIndicesView();
+
+			auto UpdateGroupFilterForConnection = [this, ConnectionStateInfo, Group, LocalConnectionInfos, GlobalScopableObjects](uint32 ConnectionId)
 			{
 				if (this->GetConnectionFilterStatus(*ConnectionStateInfo, ConnectionId) == ENetFilterStatus::Disallow)
 				{
 					FNetBitArray& GroupFilteredOutObjects = LocalConnectionInfos[ConnectionId].GroupFilteredOutObjects;
 					FNetBitArray& ObjectsInScopeBeforeDynamicFiltering = LocalConnectionInfos[ConnectionId].ObjectsInScopeBeforeDynamicFiltering;
-					const FNetBitArray& GlobalScopableObjects = NetRefHandleManager->GetScopableInternalIndices();
 
 					for (uint32 ObjectIndex : Group->Members)
 					{
@@ -1789,7 +1790,7 @@ bool FReplicationFiltering::GetIsFilteredOutByAnyGroup(uint32 ObjectInternalInde
 	return false;
 }
 
-bool FReplicationFiltering::UpdateGroupFilterEffectsForObject(uint32 ObjectIndex, uint32 ConnectionId, const FNetBitArray& ScopableObjects)
+bool FReplicationFiltering::UpdateGroupFilterEffectsForObject(uint32 ObjectIndex, uint32 ConnectionId)
 {
 	FPerConnectionInfo& ConnectionInfo = ConnectionInfos[ConnectionId];
 	FNetBitArray& GroupFilteredOutObjects = ConnectionInfo.GroupFilteredOutObjects;
@@ -1859,7 +1860,7 @@ void FReplicationFiltering::InternalSetGroupFilterStatus(FNetObjectGroupHandle G
 					FPerConnectionInfo& ConnectionInfo = ConnectionInfos[ConnectionId];
 					FNetBitArray& GroupFilteredOutObjects = ConnectionInfo.GroupFilteredOutObjects;
 					FNetBitArray& ObjectsInScopeBeforeDynamicFiltering = ConnectionInfo.ObjectsInScopeBeforeDynamicFiltering;
-					const FNetBitArray& ScopableObjects = NetRefHandleManager->GetScopableInternalIndices();
+					const FNetBitArrayView ScopableObjects = NetRefHandleManager->GetScopableInternalIndicesView();
 					for (uint32 ObjectIndex : MakeArrayView(Group->Members.GetData(), Group->Members.Num()))
 					{
 						GroupFilteredOutObjects.SetBit(ObjectIndex);
@@ -1896,11 +1897,9 @@ void FReplicationFiltering::InternalSetGroupFilterStatus(FNetObjectGroupHandle G
 					const FNetObjectGroup* Group = Groups->GetGroup(GroupHandle);
 
 					// Update group filter mask for the connection
-					const FNetBitArray& ScopedObjects = NetRefHandleManager->GetScopableInternalIndices();
-			
 					for (uint32 ObjectIndex : MakeArrayView(Group->Members.GetData(), Group->Members.Num()))
 					{
-						UpdateGroupFilterEffectsForObject(ObjectIndex, ConnectionId, ScopedObjects);
+						UpdateGroupFilterEffectsForObject(ObjectIndex, ConnectionId);
 					}
 				}
 			}
@@ -1943,14 +1942,13 @@ void FReplicationFiltering::NotifyObjectRemovedFromGroup(FNetObjectGroupHandle G
 	if (FilterGroups.GetBit(GroupHandle.GetGroupIndex()))
 	{
 		FPerObjectInfo* ConnectionState = GetPerObjectInfo(GroupInfos[GroupHandle.GetGroupIndex()].ConnectionStateIndex);
-		const FNetBitArray& ScopedObjects = NetRefHandleManager->GetScopableInternalIndices();
 
-		auto MarkGroupsDirty = [this, ObjectIndex, GroupHandle, ConnectionState, &ScopedObjects](uint32 ConnectionId)
+		auto MarkGroupsDirty = [this, ObjectIndex, GroupHandle, ConnectionState](uint32 ConnectionId)
 		{
 			if (GetConnectionFilterStatus(*ConnectionState, ConnectionId) == ENetFilterStatus::Disallow)
 			{
 				FNetBitArray& GroupFilteredOutObjects = this->ConnectionInfos[ConnectionId].GroupFilteredOutObjects;
-				UpdateGroupFilterEffectsForObject(ObjectIndex, ConnectionId, ScopedObjects);
+				UpdateGroupFilterEffectsForObject(ObjectIndex, ConnectionId);
 			}
 		};
 
diff --git a/Engine/Source/Runtime/Experimental/Iris/Core/Private/Iris/ReplicationSystem/Filtering/ReplicationFiltering.h b/Engine/Source/Runtime/Experimental/Iris/Core/Private/Iris/ReplicationSystem/Filtering/ReplicationFiltering.h
--- a/Engine/Source/Runtime/Experimental/Iris/Core/Private/Iris/ReplicationSystem/Filtering/ReplicationFiltering.h
+++ b/Engine/Source/Runtime/Experimental/Iris/Core/Private/Iris/ReplicationSystem/Filtering/ReplicationFiltering.h
@@ -201,7 +201,7 @@ private:
 	bool HasFragmentFilters() const;
 
 	// Helper to update and reset group filter effects if objects are removed from a filter or after a filter status change, returns true if the group filter was changed
-	bool UpdateGroupFilterEffectsForObject(uint32 ObjectIndex, uint32 ConnectionId, const FNetBitArray& ScopableObjects);
+	bool UpdateGroupFilterEffectsForObject(uint32 ObjectIndex, uint32 ConnectionId);
 	bool HasOwnerFilter(uint32 ObjectIndex) const;
 	bool HasConnectionFilter(uint32 ObjectIndex) const;
 
diff --git a/Engine/Source/Runtime/Experimental/Iris/Core/Private/Iris/ReplicationSystem/NetBlob/NetBlobManager.cpp b/Engine/Source/Runtime/Experimental/Iris/Core/Private/Iris/ReplicationSystem/NetBlob/NetBlobManager.cpp
--- a/Engine/Source/Runtime/Experimental/Iris/Core/Private/Iris/ReplicationSystem/NetBlob/NetBlobManager.cpp
+++ b/Engine/Source/Runtime/Experimental/Iris/Core/Private/Iris/ReplicationSystem/NetBlob/NetBlobManager.cpp
@@ -368,8 +368,8 @@ void FNetBlobManager::FNetObjectAttachmentSendQueue::PrepareProcessQueue(FReplic
 	}
 
 	// Figure out if we have any attachments to objects going out of scope.
-	const FNetBitArray& ScopableObjects = InNetRefHandleManager->GetScopableInternalIndices();
-	const FNetBitArray& PrevScopableObjects = InNetRefHandleManager->GetPrevFrameScopableInternalIndices();
+	const FNetBitArrayView ScopableObjects = InNetRefHandleManager->GetScopableInternalIndicesView();
+	const FNetBitArrayView PrevScopableObjects = InNetRefHandleManager->GetPrevFrameScopableInternalIndicesView();
 	
 	uint32 CurrentEntryIndex = 0U;
 	for (const FNetObjectAttachmentQueueEntry& Entry : MakeArrayView(AttachmentQueue))
diff --git a/Engine/Source/Runtime/Experimental/Iris/Core/Private/Iris/ReplicationSystem/NetRefHandleManager.h b/Engine/Source/Runtime/Experimental/Iris/Core/Private/Iris/ReplicationSystem/NetRefHandleManager.h
--- a/Engine/Source/Runtime/Experimental/Iris/Core/Private/Iris/ReplicationSystem/NetRefHandleManager.h
+++ b/Engine/Source/Runtime/Experimental/Iris/Core/Private/Iris/ReplicationSystem/NetRefHandleManager.h
@@ -186,11 +186,10 @@ public:
 	inline FInternalNetRefIndex GetInternalIndexFromNetHandle(FNetHandle Handle) const;
 
 	/** Get bitarray for all currently scopable internal indices */
-	const FNetBitArray& GetScopableInternalIndices() const { return ScopableInternalIndices; }
 	const FNetBitArrayView GetScopableInternalIndicesView() const { return MakeNetBitArrayView(ScopableInternalIndices); }
 
 	// Get bitarray for all internal indices that was scopable last update
-	const FNetBitArray& GetPrevFrameScopableInternalIndices() const { return PrevFrameScopableInternalIndices; }
+	const FNetBitArrayView GetPrevFrameScopableInternalIndicesView() const { return MakeNetBitArrayView(PrevFrameScopableInternalIndices); }
 	void SetPrevFrameScopableInternalIndicesToCurrent() { PrevFrameScopableInternalIndices = ScopableInternalIndices; }
 
 	/** List of objects that are always relevant or currently relevant to at least one connection. */
@@ -317,7 +316,7 @@ private:
 
 	// Which internal indices were used last net frame. This can be used to find out which ones are new and deleted this frame. 
 	FNetBitArray PrevFrameScopableInternalIndices;
-
+	
 	/** This contains the ScopableInternalIndices list minus filtered objects that are not relevant to any connection this frame. */
 	FNetBitArray RelevantObjectsInternalIndices;
 
diff --git a/Engine/Source/Runtime/Experimental/Iris/Core/Private/Iris/ReplicationSystem/Prioritization/ReplicationPrioritization.cpp b/Engine/Source/Runtime/Experimental/Iris/Core/Private/Iris/ReplicationSystem/Prioritization/ReplicationPrioritization.cpp
--- a/Engine/Source/Runtime/Experimental/Iris/Core/Private/Iris/ReplicationSystem/Prioritization/ReplicationPrioritization.cpp
+++ b/Engine/Source/Runtime/Experimental/Iris/Core/Private/Iris/ReplicationSystem/Prioritization/ReplicationPrioritization.cpp
@@ -596,8 +596,8 @@ void FReplicationPrioritization::UpdatePrioritiesForNewAndDeletedObjects()
 {
 	IRIS_PROFILER_SCOPE(FReplicationPrioritization_UpdatePrioritiesForNewAndDeletedObjects);
 
-	const FNetBitArray& PrevScopedIndices = NetRefHandleManager->GetPrevFrameScopableInternalIndices();
-	const FNetBitArray& ScopedIndices = NetRefHandleManager->GetScopableInternalIndices();
+	const FNetBitArrayView PrevScopedIndices = NetRefHandleManager->GetPrevFrameScopableInternalIndicesView();
+	const FNetBitArrayView ScopedIndices = NetRefHandleManager->GetScopableInternalIndicesView();
 
 	auto ForEachRemovedObject = [this](uint32 ObjectIndex)
 	{
@@ -628,7 +628,7 @@ void FReplicationPrioritization::UpdatePrioritiesForNewAndDeletedObjects()
 		NewIndices.Reserve(FMath::Min(1024U, MaxObjectCount));
 	}
 
-	FNetBitArray::ForAllExclusiveBits(ScopedIndices, PrevScopedIndices, ForEachNewObject, ForEachRemovedObject);
+	FNetBitArrayView::ForAllExclusiveBits(ScopedIndices, PrevScopedIndices, ForEachNewObject, ForEachRemovedObject);
 
 	if (HasNewObjectsWithStaticPriority)
 	{
diff --git a/Engine/Source/Runtime/Experimental/Iris/Core/Private/Iris/ReplicationSystem/ReplicationSystem.cpp b/Engine/Source/Runtime/Experimental/Iris/Core/Private/Iris/ReplicationSystem/ReplicationSystem.cpp
--- a/Engine/Source/Runtime/Experimental/Iris/Core/Private/Iris/ReplicationSystem/ReplicationSystem.cpp
+++ b/Engine/Source/Runtime/Experimental/Iris/Core/Private/Iris/ReplicationSystem/ReplicationSystem.cpp
@@ -241,6 +241,25 @@ public:
 		}
 	}
 
+	void CallPreSendUpdate(float DeltaSeconds)
+	{
+		ReplicationSystemInternal.GetReplicationBridge()->CallPreSendUpdate(DeltaSeconds);
+	}
+
+	void EndPostSendUpdate()
+	{
+		ReplicationSystemInternal.GetChangeMaskCache().ResetCache();
+
+		// Store the state of the previous frames scopable objects
+		ReplicationSystemInternal.GetNetRefHandleManager().SetPrevFrameScopableInternalIndicesToCurrent();
+
+		// Update handles pending tear-off
+		ReplicationSystemInternal.GetReplicationBridge()->UpdateHandlesPendingTearOff();
+
+		// Reset baseline invalidation
+		ReplicationSystemInternal.GetDeltaCompressionBaselineInvalidationTracker().PostSendUpdate();
+	}
+
 	void UpdateDirtyObjectList()
 	{
 		ReplicationSystemInternal.GetDirtyNetObjectTracker().UpdateDirtyNetObjects();
@@ -599,11 +618,11 @@ void UReplicationSystem::PreSendUpdate(float DeltaSeconds)
 		// Update world locations. We need this to happen before both filtering and prioritization.
 		Impl->UpdateWorldLocations();
 
-		// Filters to reduce the top-level scoped object list
+		// Update filters, reduce the top-level scoped object list and set each connection's scope.
 		Impl->UpdateFilterPrePoll();
 
 		// Invoke any operations we need to do before copying state data
-		InternalSys.GetReplicationBridge()->CallPreSendUpdate(DeltaSeconds);
+		Impl->CallPreSendUpdate(DeltaSeconds);
 
 		// Finalize the dirty list with objects set dirty during the poll phase
 		Impl->UpdateDirtyListPostPoll();
@@ -659,30 +678,19 @@ void UReplicationSystem::PostSendUpdate()
 
 	IRIS_PROFILER_SCOPE(FReplicationSystem_PostSendUpdate);
 	
-	FReplicationSystemInternal& InternalSys = Impl->ReplicationSystemInternal;
-
 	Impl->ResetObjectStateDirtiness();
 
-	InternalSys.GetChangeMaskCache().ResetCache();
-
-	// Store the state of the previous frames scopable objects
-	InternalSys.GetNetRefHandleManager().SetPrevFrameScopableInternalIndicesToCurrent();
-
-	// Update handles pending tear-off
-	InternalSys.GetReplicationBridge()->UpdateHandlesPendingTearOff();
-
-	// Reset baseline invalidation
-	InternalSys.GetDeltaCompressionBaselineInvalidationTracker().PostSendUpdate();
+	Impl->EndPostSendUpdate();
 
 	if (bAllowObjectReplication)
 	{
 		FDeltaCompressionBaselineManagerPostSendUpdateParams UpdateParams;
-		InternalSys.GetDeltaCompressionBaselineManager().PostSendUpdate(UpdateParams);
+		Impl->ReplicationSystemInternal.GetDeltaCompressionBaselineManager().PostSendUpdate(UpdateParams);
 	}
 
 #if UE_NET_IRIS_CSV_STATS && CSV_PROFILER
 	{
-		UE::Net::FNetSendStats& SendStats = InternalSys.GetSendStats();
+		UE::Net::FNetSendStats& SendStats = Impl->ReplicationSystemInternal.GetSendStats();
 		SendStats.ReportCsvStats();
 	}
 #endif
