FPBDRigidClustering::FindIslandsInChildren must still only look for children of the input particle even when traversing intercluster edges

- also added in protections in FClusterUnionManager::HandleRemoveOperation against dirtying the cluster union particle when removing 0 particles

diff --git a/Engine/Source/Runtime/Experimental/Chaos/Private/Chaos/ClusterUnionManager.cpp b/Engine/Source/Runtime/Experimental/Chaos/Private/Chaos/ClusterUnionManager.cpp
--- a/Engine/Source/Runtime/Experimental/Chaos/Private/Chaos/ClusterUnionManager.cpp
+++ b/Engine/Source/Runtime/Experimental/Chaos/Private/Chaos/ClusterUnionManager.cpp
@@ -789,6 +789,12 @@ namespace Chaos
 			}
 		}
 
+		// If we don't actually make a change to the cluster union then we don't actually want to trigger any further changes.
+		if (ParticleIndicesToRemove.IsEmpty())
+		{
+			return;
+		}
+
 		ParticleIndicesToRemove.Sort();
 		for (int32 Index = ParticleIndicesToRemove.Num() - 1; Index >= 0; --Index)
 		{
diff --git a/Engine/Source/Runtime/Experimental/Chaos/Private/Chaos/PBDRigidClustering.cpp b/Engine/Source/Runtime/Experimental/Chaos/Private/Chaos/PBDRigidClustering.cpp
--- a/Engine/Source/Runtime/Experimental/Chaos/Private/Chaos/PBDRigidClustering.cpp
+++ b/Engine/Source/Runtime/Experimental/Chaos/Private/Chaos/PBDRigidClustering.cpp
@@ -686,21 +686,43 @@ namespace Chaos
 			ProcessingQueue.Add(Child);
 			while (ProcessingQueue.Num())
 			{
-				FPBDRigidParticleHandle* ChildToProcess = ProcessingQueue.Pop();
-				if (!ProcessedChildren.Contains(ChildToProcess))
+				if (FPBDRigidParticleHandle* ChildToProcess = ProcessingQueue.Pop())
 				{
-					ProcessedChildren.Add(ChildToProcess);
 					Island.Add(ChildToProcess);
 					for (const TConnectivityEdge<FReal>& Edge : ChildToProcess->CastToClustered()->ConnectivityEdges())
 					{
-						if (IsInterclusterEdge(*ChildToProcess, Edge) && !bTraverseInterclusterEdges)
+						FPBDRigidParticleHandle* Sibling = Edge.Sibling;
+						if (IsInterclusterEdge(*ChildToProcess, Edge))
 						{
-							continue;
+							if (!bTraverseInterclusterEdges)
+							{
+								continue;
+							}
+
+							// Intercluster edges need to find the parent particle that's actually a child of the input ClusteredParticle
+							while (Sibling)
+							{
+								if (FPBDRigidClusteredParticleHandle* ClusterSibling = Sibling->CastToClustered())
+								{
+									if (ClusterSibling->Parent() == ClusteredParticle)
+									{
+										break;
+									}
+
+									Sibling = ClusterSibling->Parent();
+								}
+								else
+								{
+									Sibling = nullptr;
+									break;
+								}
+							}
 						}
 
-						if (!ProcessedChildren.Contains(Edge.Sibling))
+						if (Sibling && !ProcessedChildren.Contains(Sibling))
 						{
-							ProcessingQueue.Add(Edge.Sibling);
+							ProcessingQueue.Add(Sibling);
+							ProcessedChildren.Add(Sibling);
 						}
 					}
 				}
