Geometry collection : add cvar to control both marshalled position and rotation tolerances  (when synced from PT to GT )

diff --git a/Engine/Source/Runtime/Experimental/Chaos/Private/PhysicsProxy/GeometryCollectionPhysicsProxy.cpp b/Engine/Source/Runtime/Experimental/Chaos/Private/PhysicsProxy/GeometryCollectionPhysicsProxy.cpp
--- a/Engine/Source/Runtime/Experimental/Chaos/Private/PhysicsProxy/GeometryCollectionPhysicsProxy.cpp
+++ b/Engine/Source/Runtime/Experimental/Chaos/Private/PhysicsProxy/GeometryCollectionPhysicsProxy.cpp
@@ -129,11 +129,18 @@ FAutoConsoleVariableRef CVarGeometryCollectionTransformTolerance(
 	GeometryCollectionTransformTolerance,
 	TEXT("Tolerance to detect if a transform has changed"));
 
-float GeometryCollectionXRUpdateTolerance = UE_KINDA_SMALL_NUMBER;
-FAutoConsoleVariableRef CVarGeometryCollectionXRUpdateTolerance(
-	TEXT("p.GeometryCollection.XRUpdateTolerance"),
-	GeometryCollectionXRUpdateTolerance,
-	TEXT("Tolerance to detect if a XR has changed has changed when syncing PT to GT"));
+float GeometryCollectionPositionUpdateTolerance = UE_KINDA_SMALL_NUMBER;
+FAutoConsoleVariableRef CVarGeometryCollectionPositionUpdateTolerance(
+	TEXT("p.GeometryCollection.PositionUpdateTolerance"),
+	GeometryCollectionPositionUpdateTolerance,
+	TEXT("Tolerance to detect if particle position has changed has changed when syncing PT to GT"));
+
+float GeometryCollectionRotationUpdateTolerance = UE_KINDA_SMALL_NUMBER;
+FAutoConsoleVariableRef CVarGeometryCollectionRotationUpdateTolerance(
+	TEXT("p.GeometryCollection.RotationUpdateTolerance"),
+	GeometryCollectionRotationUpdateTolerance,
+	TEXT("Tolerance to detect if particle rotation has changed has changed when syncing PT to GT"));
+
 
 DEFINE_LOG_CATEGORY_STATIC(UGCC_LOG, Error, All);
 
@@ -3418,14 +3425,14 @@ static inline bool UpdateGTParticleXR(Chaos::FPBDRigidParticle& GTParticle, cons
 	GC_PHYSICSPROXY_CHECK_FOR_NAN(NewX);
 
 	const Chaos::FVec3 OldX = GTParticle.X();
-	const bool bNeedUpdateX = (!NewX.Equals(OldX, GeometryCollectionXRUpdateTolerance));
+	const bool bNeedUpdateX = (!NewX.Equals(OldX, GeometryCollectionPositionUpdateTolerance));
 	if (bNeedUpdateX)
 	{
 		GTParticle.SetX(NewX, false);
 	}
 
 	const Chaos::FRotation3 OldR = GTParticle.R();
-	const bool bNeedUpdateR = (!NewR.Equals(OldR, GeometryCollectionXRUpdateTolerance));
+	const bool bNeedUpdateR = (!NewR.Equals(OldR, GeometryCollectionRotationUpdateTolerance));
 	if (bNeedUpdateR)
 	{
 		GTParticle.SetR(NewR, false);
