Geomety collection : make sure to update the damage eveluation model when changing thedamage model when physics state is already created

diff --git a/Engine/Source/Runtime/Experimental/Chaos/Private/PhysicsProxy/GeometryCollectionPhysicsProxy.cpp b/Engine/Source/Runtime/Experimental/Chaos/Private/PhysicsProxy/GeometryCollectionPhysicsProxy.cpp
--- a/Engine/Source/Runtime/Experimental/Chaos/Private/PhysicsProxy/GeometryCollectionPhysicsProxy.cpp
+++ b/Engine/Source/Runtime/Experimental/Chaos/Private/PhysicsProxy/GeometryCollectionPhysicsProxy.cpp
@@ -2910,6 +2910,7 @@ void FGeometryCollectionPhysicsProxy::PushStateOnGameThread(Chaos::FPBDRigidsSol
 	if (GameThreadPerFrameData.DamageModel.IsSet())
 	{
 		Parameters.DamageModel = GameThreadPerFrameData.DamageModel.GetValue();
+		Parameters.DamageEvaluationModel = GetDamageEvaluationModel(Parameters.DamageModel);
 		GameThreadPerFrameData.DamageModel.Reset();
 		bIsDamageThresholdDataDirty = true;
 	}
