Chaos - fix NaNs from creating GC particles with no geometry

diff --git a/Engine/Source/Runtime/Experimental/Chaos/Private/PhysicsProxy/GeometryCollectionPhysicsProxy.cpp b/Engine/Source/Runtime/Experimental/Chaos/Private/PhysicsProxy/GeometryCollectionPhysicsProxy.cpp
--- a/Engine/Source/Runtime/Experimental/Chaos/Private/PhysicsProxy/GeometryCollectionPhysicsProxy.cpp
+++ b/Engine/Source/Runtime/Experimental/Chaos/Private/PhysicsProxy/GeometryCollectionPhysicsProxy.cpp
@@ -2634,6 +2634,7 @@ int32 FGeometryCollectionPhysicsProxy::CalculateHierarchyLevel(const FGeometryDy
 
 int32 FGeometryCollectionPhysicsProxy::CalculateEffectiveParticles(const FGeometryDynamicCollection& DynamicCollection, int32 NumTransform, TBitArray<>& EffectiveParticles) const
 {
+	int32 NumMissingGeometry = 0;
 	int32 NumEffectiveParticlesFound = 0;
 	EffectiveParticles.Init(false, NumTransform);
 	const int32 MaxSimulatedLevel = FMath::Min(GlobalMaxSimulatedLevel, Parameters.MaxSimulatedLevel);
@@ -2642,10 +2643,23 @@ int32 FGeometryCollectionPhysicsProxy::CalculateEffectiveParticles(const FGeomet
 		const int32 Level = FMath::Clamp(CalculateHierarchyLevel(DynamicCollection, TransformIndex), 0, INT_MAX);
 		if (Level <= MaxSimulatedLevel || !Parameters.EnableClustering)
 		{
-			EffectiveParticles[TransformIndex] = true;
-			NumEffectiveParticlesFound++;
+			const bool bHasGeometry = DynamicCollection.Implicits[TransformIndex].IsValid();
+			if (bHasGeometry)
+			{
+				EffectiveParticles[TransformIndex] = true;
+				NumEffectiveParticlesFound++;
+			}
+			else
+			{
+				// This will happen if some particles are being culled (e.g., via data flow) but the MaxSimulatedLevel is set to a level that includes them.
+				UE_LOG(LogChaos, Verbose, TEXT("GeometryCollection Transform %d has no geometry"), TransformIndex);
+				NumMissingGeometry++;
+			}
 		}
 	}
+
+	UE_CLOG(NumMissingGeometry > 0, LogChaos, Error, TEXT("Geometry collection %s tried to create %d particles with no geometry"), *GetOwner()->GetFullName(), NumMissingGeometry);
+
 	return NumEffectiveParticlesFound;
 }
 
