Moved out IAS' cache servicing thread as a separate concern, loosely supporting more than one cache along with preparation for servicing cache reads.

#rnx

diff --git a/Engine/Source/Runtime/Experimental/IoStoreOnDemand/Private/JournaledCache.cpp b/Engine/Source/Runtime/Experimental/IoStoreOnDemand/Private/JournaledCache.cpp
--- a/Engine/Source/Runtime/Experimental/IoStoreOnDemand/Private/JournaledCache.cpp
+++ b/Engine/Source/Runtime/Experimental/IoStoreOnDemand/Private/JournaledCache.cpp
@@ -1378,15 +1378,9 @@
 	return false;
 }
 
-// }}}
-
-} // namespace UE::IO::IAS::JournaledCache
 
 
+// {{{1 service-thread .........................................................
 
-namespace UE::IO::IAS {
-
-// {{{1 journaled-cache ........................................................
-
 ////////////////////////////////////////////////////////////////////////////////
 TRACE_DECLARE_INT_COUNTER(IasMemDemand, TEXT("Ias/CacheMemDemand"));
@@ -1394,107 +1388,276 @@
 TRACE_DECLARE_INT_COUNTER(IasOpCount,   TEXT("Ias/CacheOpCount"));
 
 ////////////////////////////////////////////////////////////////////////////////
-class FJournaledCache
-	: public IIasCache
-	, public FRunnable
+class FServiceThread
+	: public FRunnable
 {
 public:
-								FJournaledCache() = default;
-	bool						Initialize(const TCHAR* RootDir, const FIasCacheConfig& Config);
-	virtual void				Abandon() override;
-	virtual bool				ContainsChunk(const FIoHash& Key) const override;
-	virtual FGetToken			Get(const FIoHash& Key, FIoBuffer& OutData) override;
-	virtual FGetWork			Materialize(FGetToken Token, const FIoReadOptions& Options, const FIoCancellationToken* CancelToken) override;
-	virtual FIoStatus			Put(const FIoHash& Key, FIoBuffer& Data) override;
+	static FServiceThread&	Get();
+	void					RegisterCache(TUniquePtr<FCache> Cache);
+	void					UnregisterCache(FCache* Cache);
+	void					SetGovernorRate(uint32 Allowance, uint32 Ops, uint32 Seconds);
+	void					SetGovernorDemand(uint32 Threshold, uint32 Boost, uint32 SuperBoost);
 
 private:
-	using FCacheInner = JournaledCache::FCache;
+	struct FWork
+	{
+		enum {
+			Work_Register,
+			Work_Unregister,
+			Work_GovDemand,
+			Work_GovRate,
+		};
 
-	void						Update();
-	static uint64				ReduceKey(const FIoHash& Key);
-	TUniquePtr<FCacheInner>		Cache;
-	UE::Tasks::FPipe			GetPipe = UE::Tasks::FPipe(TEXT("IasCacheGetPipe"));
-	JournaledCache::FGovernor	Governor;
+		void	SetCache(FCache* In) { Cache = UPTRINT(In) >> 3; check((UPTRINT(In) & 0x7) == 0); }
+		FCache* GetCache() const	 { return (FCache*)(Cache << 3); }
 
-	// FRunnable
+		union {
+			struct {
+				UPTRINT			What : 3; // (un)reg
+				UPTRINT			Cache : 44;
+				UPTRINT			_Unused : 17;
+			};
+			struct {
+				uint16			_What0 : 3; // rate
+				uint16			Ops : 13;
+				uint16			Seconds;
+				uint32			Allowance;
+			};
+			struct {
+				uint16			_What1 : 3; // demand
+				uint16			Threshold;
+				uint16			Boost;
+				uint16			SuperBoost;
+			};
+		};
+		//uint64				Key;
+	};
+	static_assert(sizeof(FWork) == sizeof(UPTRINT));
+
+	void						StartThread();
+	int32						Update();
+	void						UpdateCache(FCache* Cache);
 	virtual uint32				Run() override;
 	virtual void				Stop() override;
-	void						StartThread();
-	TUniquePtr<FRunnableThread>	Thread;
+	void						SubmitWork(FWork& Work);
+	void						ReceiveWork();
+	TUniquePtr<FRunnableThread> Thread;
 	FEventRef					WakeEvent;
-	std::atomic<bool>			Running = true;
+	FGovernor					Governor;
+	TArray<TUniquePtr<FCache>>	Caches;
+	std::atomic_int				RunCount = 0;
+	FCriticalSection			Lock;
+	TArray<FWork>				PendingWork;
+	std::atomic_int				PendingCount = 0;
+	int32						PendingPrev;
 };
 
 ////////////////////////////////////////////////////////////////////////////////
-bool FJournaledCache::Initialize(const TCHAR* RootDir, const FIasCacheConfig& Config)
+FServiceThread& FServiceThread::Get()
 {
-	using namespace JournaledCache;
-
-	// Filesystem setup
-	FStringView Name = Config.Name;
-	check(Name.Len() > 0 && !Name.EndsWith('/') && !Name.EndsWith('\\'));
-
-	TStringBuilder<256> CachePath;
-	CachePath << RootDir;
-	FPathViews::Append(CachePath, GetCacheFsDir());
-	FPathViews::Append(CachePath, FPathViews::GetPath(Name));
-
-	if (IFileManager& Ifm = IFileManager::Get(); !Ifm.MakeDirectory(CachePath.ToString(), true))
+	static FServiceThread* Ptr;
+	if (Ptr != nullptr)
 	{
-		UE_LOG(LogIas, Error, TEXT("JournaledCache: Unable to create directory '%s'"), CachePath.ToString());
-		return false;
+		return *Ptr;
 	}
 
-	FPathViews::Append(CachePath, FPathViews::GetBaseFilename(Name));
-	CachePath << GetCacheFsSuffix();
-
-	// Inner cache
-	FCacheInner::FConfig EventualConfig;
-	static_cast<FIasCacheConfig&>(EventualConfig) = Config;
-	EventualConfig.Path = CachePath;
-	Cache = MakeUnique<FCacheInner>(MoveTemp(EventualConfig));
-
-	if (uint32 Ailments = Cache->GetAilments(); Ailments != 0)
-	{
-		UE_LOG(LogIas, Error, TEXT("JournaledCache: Error initialising inner cache '%x'"), Ailments);
-		return false;
-	}
-
-	Cache->Load();
-
-	// Thread setup
-	const FIasCacheConfig::FRate& WriteRate = Config.WriteRate;
-	const FIasCacheConfig::FDemand& Demand = Config.Demand;
-	Governor.Set(WriteRate.Allowance, WriteRate.Ops, WriteRate.Seconds);
-	Governor.SetDemands(Demand.Threshold, Demand.Boost, Demand.SuperBoost);
-
-	StartThread();
-
-	return true;
+	static FServiceThread Instance;
+	Ptr = &Instance;
+	return *Ptr;
 }
 
 ////////////////////////////////////////////////////////////////////////////////
-void FJournaledCache::Abandon()
+void FServiceThread::SubmitWork(FWork& Work)
 {
-	LLM_SCOPE_BYTAG(Ias);
-
-	Thread.Reset();
-	Cache->Drop();
-	delete this;
+	FScopeLock _(&Lock);
+	PendingWork.Add(MoveTemp(Work));
+	PendingCount.fetch_add(1, std::memory_order_relaxed);
+	WakeEvent->Trigger();
 }
 
 ////////////////////////////////////////////////////////////////////////////////
-void FJournaledCache::StartThread()
+void FServiceThread::RegisterCache(TUniquePtr<FCache> Cache)
 {
-	auto* Inst = FRunnableThread::Create(this, TEXT("Ias.FileCache"), 0, TPri_BelowNormal);
+	uint32 PrevRunCount = RunCount.fetch_add(1, std::memory_order_relaxed);
+
+	FCache* RawPtr = Cache.Release();
+
+	FWork Work;
+	Work.What = FWork::Work_Register;
+	Work.SetCache(RawPtr);
+	SubmitWork(Work);
+
+	if (PrevRunCount == 0)
+	{
+		StartThread();
+	}
+}
+
+////////////////////////////////////////////////////////////////////////////////
+void FServiceThread::UnregisterCache(FCache* Cache)
+{
+	FWork Work;
+	Work.What = FWork::Work_Unregister;
+	Work.SetCache(Cache);
+	SubmitWork(Work);
+
+	int32 PrevRunCount = RunCount.fetch_sub(1, std::memory_order_relaxed);
+	if (PrevRunCount == 1)
+	{
+		/* ideally we'd shut down the thread here as there are no active caches
+		 * that need servicing. But this is involved so we'll leave it up for
+		 * now. See "THREAD_ALIVE" comments for add/subs keeping thread up */
+		//Thread.Reset();
+	}
+}
+
+////////////////////////////////////////////////////////////////////////////////
+void FServiceThread::SetGovernorRate(uint32 Allowance, uint32 Ops, uint32 Seconds)
+{
+	FWork Work;
+	Work.What = FWork::Work_GovRate;
+	Work.Allowance = Allowance;
+	Work.Ops = uint16(Ops);
+	Work.Seconds = uint16(Seconds);
+	SubmitWork(Work);
+
+	check(Work.Ops == Ops && Work.Seconds == Seconds);
+}
+
+////////////////////////////////////////////////////////////////////////////////
+void FServiceThread::SetGovernorDemand(uint32 Threshold, uint32 Boost, uint32 SuperBoost)
+{
+	FWork Work;
+	Work.What = FWork::Work_GovDemand;
+	Work.Threshold = uint16(Threshold);
+	Work.Boost = uint16(Boost);
+	Work.SuperBoost = uint16(SuperBoost);
+	SubmitWork(Work);
+
+	check(Work.Threshold == Threshold && Work.Boost == Boost && Work.SuperBoost == SuperBoost);
+}
+
+
+////////////////////////////////////////////////////////////////////////////////
+void FServiceThread::StartThread()
+{
+	RunCount.fetch_add(1, std::memory_order_relaxed); // THREAD_ALIVE
+
+	auto* Inst = FRunnableThread::Create(this, TEXT("Ias.CacheIo"), 0, TPri_BelowNormal);
 	Thread.Reset(Inst);
 }
 
 ////////////////////////////////////////////////////////////////////////////////
-void FJournaledCache::Update()
+uint32 FServiceThread::Run()
 {
 	LLM_SCOPE_BYTAG(Ias);
 
+	while (RunCount.load(std::memory_order_relaxed))
+	{
+		ReceiveWork();
+
+		int32 WaitTime = Update();
+		if (WaitTime < 0)
+		{
+			break;
+		}
+		
+		WaitTime = WaitTime ? WaitTime : 37;
+		WakeEvent->Wait(WaitTime);
+	}
+
+	ReceiveWork();
+	check(Caches.Num() == 0);
+
+	return 0;
+}
+
+////////////////////////////////////////////////////////////////////////////////
+void FServiceThread::Stop()
+{
+	RunCount.fetch_sub(1, std::memory_order_relaxed); // THREAD_ALIVE
+}
+
+////////////////////////////////////////////////////////////////////////////////
+int32 FServiceThread::Update()
+{
+	for (TUniquePtr<FCache>& Item : Caches)
+	{
+		FCache* Cache = Item.Get();
+		UpdateCache(Cache);
+	}
+
+	return 0;
+}
+
+////////////////////////////////////////////////////////////////////////////////
+void FServiceThread::ReceiveWork()
+{
+	int32 PendingLoad = PendingCount.load(std::memory_order_relaxed);
+	if (PendingLoad == PendingPrev)
+	{
+		return;
+	}
+
+	Lock.Lock();
+	TArray<FWork> InboundWork = MoveTemp(PendingWork);
+	Lock.Unlock();
+	PendingPrev = PendingLoad;
+
+	// Unregisters first
+	for (FWork& Work : InboundWork)
+	{
+		if (Work.What != FWork::Work_Unregister)
+		{
+			continue;
+		}
+
+		FCache* CachePtr = Work.GetCache();
+		for (int32 i = 0, n = Caches.Num(); i < n; ++i)
+		{
+			if (Caches[i].Get() != CachePtr)
+			{
+				continue;
+			}
+
+			Caches.RemoveAtSwap(i);
+			break;
+		}
+	}
+
+	// Then the rest
+	for (const FWork& Work : InboundWork)
+	{
+		if (Work.What == FWork::Work_Unregister)
+		{
+			continue;
+		}
+
+		if (Work.What == FWork::Work_Register)
+		{
+			FCache* Cache = Work.GetCache();
+			Cache->Load();
+			Caches.Add(TUniquePtr<FCache>(Cache));
+			continue;
+		}
+
+		if (Work.What == FWork::Work_GovRate)
+		{
+			Governor.Set(Work.Allowance, Work.Ops, Work.Seconds);
+			continue;
+		}
+
+		if (Work.What == FWork::Work_GovDemand)
+		{
+			Governor.SetDemands(Work.Threshold, Work.Boost, Work.SuperBoost);
+			continue;
+		}
+	}
+}
+
+////////////////////////////////////////////////////////////////////////////////
+void FServiceThread::UpdateCache(FCache* Cache)
+{
 	uint32 Demand = Cache->GetDemand();
 	uint32 Allowance = Governor.BeginAllowance(Demand);
 	if (Allowance == 0)
@@ -1525,22 +1688,101 @@
 	TRACE_COUNTER_SET(IasAllowance, 0);
 }
 
+// }}}
+
+} // namespace UE::IO::IAS::JournaledCache
+
+
+
+namespace UE::IO::IAS {
+
+// {{{1 journaled-cache ........................................................
+
 ////////////////////////////////////////////////////////////////////////////////
-uint32 FJournaledCache::Run()
+class FJournaledCache
+	: public IIasCache
 {
-	while (Running.load(std::memory_order_relaxed))
+public:
+								FJournaledCache() = default;
+								~FJournaledCache();
+	bool						Initialize(const TCHAR* RootDir, const FIasCacheConfig& Config);
+	virtual void				Abandon() override;
+	virtual bool				ContainsChunk(const FIoHash& Key) const override;
+	virtual FGetToken			Get(const FIoHash& Key, FIoBuffer& OutData) override;
+	virtual FGetWork			Materialize(FGetToken Token, const FIoReadOptions& Options, const FIoCancellationToken* CancelToken) override;
+	virtual FIoStatus			Put(const FIoHash& Key, FIoBuffer& Data) override;
+
+private:
+	void						Update();
+	static uint64				ReduceKey(const FIoHash& Key);
+	JournaledCache::FCache*		Cache = nullptr;
+	UE::Tasks::FPipe			GetPipe = UE::Tasks::FPipe(TEXT("IasCacheGetPipe"));
+};
+
+////////////////////////////////////////////////////////////////////////////////
+FJournaledCache::~FJournaledCache()
+{
+	if (Cache != nullptr)
 	{
-		Update();
-		WakeEvent->Wait(37);
+		JournaledCache::FServiceThread::Get().UnregisterCache(Cache);
 	}
-	return 0;
 }
 
 ////////////////////////////////////////////////////////////////////////////////
-void FJournaledCache::Stop()
+bool FJournaledCache::Initialize(const TCHAR* RootDir, const FIasCacheConfig& Config)
 {
-	Running.store(false, std::memory_order_relaxed);
-	WakeEvent->Trigger();
+	using namespace JournaledCache;
+
+	// Filesystem setup
+	FStringView Name = Config.Name;
+	check(Name.Len() > 0 && !Name.EndsWith('/') && !Name.EndsWith('\\'));
+
+	TStringBuilder<256> CachePath;
+	CachePath << RootDir;
+	FPathViews::Append(CachePath, GetCacheFsDir());
+	FPathViews::Append(CachePath, FPathViews::GetPath(Name));
+
+	if (IFileManager& Ifm = IFileManager::Get(); !Ifm.MakeDirectory(CachePath.ToString(), true))
+	{
+		UE_LOG(LogIas, Error, TEXT("JournaledCache: Unable to create directory '%s'"), CachePath.ToString());
+		return false;
+	}
+
+	FPathViews::Append(CachePath, FPathViews::GetBaseFilename(Name));
+	CachePath << GetCacheFsSuffix();
+
+	// Inner cache
+	FCache::FConfig EventualConfig;
+	static_cast<FIasCacheConfig&>(EventualConfig) = Config;
+	EventualConfig.Path = CachePath;
+	TUniquePtr<FCache> NewCache = MakeUnique<FCache>(MoveTemp(EventualConfig));
+
+	if (uint32 Ailments = NewCache->GetAilments(); Ailments != 0)
+	{
+		UE_LOG(LogIas, Error, TEXT("JournaledCache: Error initialising inner cache '%x'"), Ailments);
+		return false;
+	}
+
+	Cache = NewCache.Get();
+
+	const FIasCacheConfig::FRate& WriteRate = Config.WriteRate;
+	const FIasCacheConfig::FDemand& Demand = Config.Demand;
+
+	FServiceThread& ServiceThread = FServiceThread::Get();
+	ServiceThread.RegisterCache(MoveTemp(NewCache));
+	ServiceThread.SetGovernorRate(WriteRate.Allowance, WriteRate.Ops, WriteRate.Seconds);
+	ServiceThread.SetGovernorDemand(Demand.Threshold, Demand.Boost, Demand.SuperBoost);
+
+	return true;
 }
 
 ////////////////////////////////////////////////////////////////////////////////
+void FJournaledCache::Abandon()
+{
+	LLM_SCOPE_BYTAG(Ias);
+
+	Cache->Drop();
+	delete this;
+}
+
+////////////////////////////////////////////////////////////////////////////////
