My test server was secretly locking directories on me and causing needless detective work.

#rnx

diff --git a/Engine/Source/Runtime/Experimental/IoStoreOnDemand/Private/CoreHttp/TestServer.py b/Engine/Source/Runtime/Experimental/IoStoreOnDemand/Private/CoreHttp/TestServer.py
--- a/Engine/Source/Runtime/Experimental/IoStoreOnDemand/Private/CoreHttp/TestServer.py
+++ b/Engine/Source/Runtime/Experimental/IoStoreOnDemand/Private/CoreHttp/TestServer.py
@@ -1,4 +1,5 @@
 import io
+import os
 import base64
 import random
 import socket
@@ -6,6 +7,8 @@ import asyncio
 import threading
 from flask import request, Flask, Response
 
+os.chdir("c:/")
+
 #-------------------------------------------------------------------------------
 async def proxy_impl(loop, client):
     msg = b""
