Iris
- Added support for defining policy flags to control how we treat certain types of RPCs - ENetObjectAttachmentSendPolicyFlags

diff --git a/Engine/Source/Runtime/Experimental/Iris/Core/Private/Iris/ReplicationSystem/NetBlob/NetBlobManager.cpp b/Engine/Source/Runtime/Experimental/Iris/Core/Private/Iris/ReplicationSystem/NetBlob/NetBlobManager.cpp
--- a/Engine/Source/Runtime/Experimental/Iris/Core/Private/Iris/ReplicationSystem/NetBlob/NetBlobManager.cpp
+++ b/Engine/Source/Runtime/Experimental/Iris/Core/Private/Iris/ReplicationSystem/NetBlob/NetBlobManager.cpp
@@ -50,7 +50,7 @@ bool FNetBlobManager::RegisterNetBlobHandler(UNetBlobHandler* Handler)
 	return BlobHandlerManager.RegisterHandler(Handler);
 }
 
-bool FNetBlobManager::QueueNetObjectAttachment(uint32 ConnectionId, const FNetObjectReference& TargetRef, const TRefCountPtr<FNetObjectAttachment>& Attachment)
+bool FNetBlobManager::QueueNetObjectAttachment(uint32 ConnectionId, const FNetObjectReference& TargetRef, const TRefCountPtr<FNetObjectAttachment>& Attachment, ENetObjectAttachmentSendPolicyFlags SendFlags)
 {
 	if (!Attachment.IsValid())
 	{
@@ -89,11 +89,11 @@ bool FNetBlobManager::QueueNetObjectAttachment(uint32 ConnectionId, const FNetOb
 	}
 
 	Attachment->SetNetObjectReference(OwnerReference, OwnerOrSubObjectReference);
-	AttachmentSendQueue.Enqueue(ConnectionId, OwnerIndex, SubObjectIndex, Attachment);
+	AttachmentSendQueue.Enqueue(ConnectionId, OwnerIndex, SubObjectIndex, Attachment, SendFlags);
 	return true;
 }
 
-bool FNetBlobManager::SendRPC(const UObject* Object, const UObject* SubObject, const UFunction* Function, const void* Parameters)
+bool FNetBlobManager::SendRPC(const UObject* Object, const UObject* SubObject, const UFunction* Function, const void* Parameters, UE::Net::ENetObjectAttachmentSendPolicyFlags SendFlags)
 {
 	if (CVarEnableIrisRPCs.GetValueOnGameThread() <= 0)
 	{
@@ -150,11 +150,11 @@ bool FNetBlobManager::SendRPC(const UObject* Object, const UObject* SubObject, c
 	}
 
 	RPC->SetNetObjectReference(OwnerReference, OwnerOrSubObjectReference);
-	AttachmentSendQueue.Enqueue(OwnerIndex, SubObjectIndex, reinterpret_cast<const TRefCountPtr<FNetObjectAttachment>&>(RPC));
+	AttachmentSendQueue.Enqueue(OwnerIndex, SubObjectIndex, reinterpret_cast<const TRefCountPtr<FNetObjectAttachment>&>(RPC), SendFlags);
 	return true;
 }
 
-bool FNetBlobManager::SendRPC(uint32 ConnectionId, const UObject* Object, const UObject* SubObject, const UFunction* Function, const void* Parameters)
+bool FNetBlobManager::SendRPC(uint32 ConnectionId, const UObject* Object, const UObject* SubObject, const UFunction* Function, const void* Parameters, UE::Net::ENetObjectAttachmentSendPolicyFlags SendFlags)
 {
 	if (CVarEnableIrisRPCs.GetValueOnGameThread() <= 0)
 	{
@@ -210,7 +210,7 @@ bool FNetBlobManager::SendRPC(uint32 ConnectionId, const UObject* Object, const
 	}
 
 	RPC->SetNetObjectReference(OwnerReference, OwnerOrSubObjectReference);
-	AttachmentSendQueue.Enqueue(ConnectionId, OwnerIndex, SubObjectIndex, reinterpret_cast<const TRefCountPtr<FNetObjectAttachment>&>(RPC));
+	AttachmentSendQueue.Enqueue(ConnectionId, OwnerIndex, SubObjectIndex, reinterpret_cast<const TRefCountPtr<FNetObjectAttachment>&>(RPC), SendFlags);
 	return true;
 }
 
@@ -320,21 +320,23 @@ void FNetBlobManager::FNetObjectAttachmentSendQueue::Init(FNetBlobManager* InMan
 	Manager = InManager;
 }
 
-void FNetBlobManager::FNetObjectAttachmentSendQueue::Enqueue(uint32 ConnectionId, FInternalNetRefIndex OwnerIndex, FInternalNetRefIndex SubObjectIndex, const TRefCountPtr<FNetObjectAttachment>& Attachment)
+void FNetBlobManager::FNetObjectAttachmentSendQueue::Enqueue(uint32 ConnectionId, FInternalNetRefIndex OwnerIndex, FInternalNetRefIndex SubObjectIndex, const TRefCountPtr<FNetObjectAttachment>& Attachment, ENetObjectAttachmentSendPolicyFlags SendFlags)
 {
 	FNetObjectAttachmentQueueEntry& QueueEntry = AttachmentQueue.AddDefaulted_GetRef();
 	QueueEntry.ConnectionId = ConnectionId;
 	QueueEntry.OwnerIndex = OwnerIndex;
 	QueueEntry.SubObjectIndex = SubObjectIndex;
+	QueueEntry.SendFlags = SendFlags;
 	QueueEntry.Attachment = Attachment;
 }
 
-void FNetBlobManager::FNetObjectAttachmentSendQueue::Enqueue(FInternalNetRefIndex OwnerIndex, FInternalNetRefIndex SubObjectIndex, const TRefCountPtr<FNetObjectAttachment>& Attachment)
+void FNetBlobManager::FNetObjectAttachmentSendQueue::Enqueue(FInternalNetRefIndex OwnerIndex, FInternalNetRefIndex SubObjectIndex, const TRefCountPtr<FNetObjectAttachment>& Attachment, ENetObjectAttachmentSendPolicyFlags SendFlags)
 {
 	FNetObjectAttachmentQueueEntry& QueueEntry = AttachmentQueue.AddDefaulted_GetRef();
 	QueueEntry.ConnectionId = 0;
 	QueueEntry.OwnerIndex = OwnerIndex;
 	QueueEntry.SubObjectIndex = SubObjectIndex;
+	QueueEntry.SendFlags = SendFlags;
 	QueueEntry.Attachment = Attachment;
 
 	bHasMulticastAttachments = true;
@@ -458,14 +460,14 @@ void FNetBlobManager::FNetObjectAttachmentSendQueue::ProcessQueue(EProcessMode P
 
 				FReplicationConnection* Connection = ProcessContext.Connections->GetConnection(ConnectionId);
 				// We're only iterating over valid connections so the Connection pointer must be valid.
-				Connection->ReplicationWriter->QueueNetObjectAttachments(Entry.OwnerIndex, Entry.SubObjectIndex, AttachmentsView);
+				Connection->ReplicationWriter->QueueNetObjectAttachments(Entry.OwnerIndex, Entry.SubObjectIndex, AttachmentsView, Entry.SendFlags);
 			}
 		}
 		else
 		{
 			if (FReplicationConnection* Connection = ProcessContext.Connections->GetConnection(Entry.ConnectionId))
 			{
-				Connection->ReplicationWriter->QueueNetObjectAttachments(Entry.OwnerIndex, Entry.SubObjectIndex, AttachmentsView);
+				Connection->ReplicationWriter->QueueNetObjectAttachments(Entry.OwnerIndex, Entry.SubObjectIndex, AttachmentsView, Entry.SendFlags);
 			}
 		}
 	});
diff --git a/Engine/Source/Runtime/Experimental/Iris/Core/Private/Iris/ReplicationSystem/NetBlob/NetBlobManager.h b/Engine/Source/Runtime/Experimental/Iris/Core/Private/Iris/ReplicationSystem/NetBlob/NetBlobManager.h
--- a/Engine/Source/Runtime/Experimental/Iris/Core/Private/Iris/ReplicationSystem/NetBlob/NetBlobManager.h
+++ b/Engine/Source/Runtime/Experimental/Iris/Core/Private/Iris/ReplicationSystem/NetBlob/NetBlobManager.h
@@ -2,6 +2,7 @@
 
 #pragma once
 #include "CoreTypes.h"
+#include "Iris/ReplicationSystem/ReplicationSystemTypes.h"
 #include "Iris/ReplicationSystem/NetBlob/NetBlob.h"
 #include "Iris/ReplicationSystem/NetBlob/NetBlobHandlerManager.h"
 #include "Iris/ReplicationSystem/NetRefHandleManager.h" // For FInternalNetRefIndex
@@ -15,6 +16,7 @@ class UNetRPCHandler;
 namespace UE::Net
 {
 	class FNetObjectReference;
+
 	namespace Private
 	{
 		class FNetRefHandleManager;
@@ -32,7 +34,6 @@ struct FNetBlobManagerInitParams
 	bool bSendAttachmentsWithObject = false;
 };
 
-
 class FNetBlobManager
 {
 public:
@@ -42,9 +43,9 @@ public:
 
 	bool RegisterNetBlobHandler(UNetBlobHandler* Handler);
 
-	bool QueueNetObjectAttachment(uint32 ConnectionId, const FNetObjectReference& TargetRef, const TRefCountPtr<FNetObjectAttachment>& Attachment);
-	bool SendRPC(const UObject* Object, const UObject* SubObject, const UFunction* Function, const void* Parameters);
-	bool SendRPC(uint32 ConnectionId, const UObject* Object, const UObject* SubObject, const UFunction* Function, const void* Parameters);
+	bool QueueNetObjectAttachment(uint32 ConnectionId, const FNetObjectReference& TargetRef, const TRefCountPtr<FNetObjectAttachment>& Attachment, ENetObjectAttachmentSendPolicyFlags SendFlags = ENetObjectAttachmentSendPolicyFlags::None);
+	bool SendRPC(const UObject* Object, const UObject* SubObject, const UFunction* Function, const void* Parameters, ENetObjectAttachmentSendPolicyFlags SendFlags = ENetObjectAttachmentSendPolicyFlags::None);
+	bool SendRPC(uint32 ConnectionId, const UObject* Object, const UObject* SubObject, const UFunction* Function, const void* Parameters, ENetObjectAttachmentSendPolicyFlags SendFlags = ENetObjectAttachmentSendPolicyFlags::None);
 
 	enum class EProcessMode 
 	{
@@ -78,10 +79,10 @@ private:
 		void Init(FNetBlobManager* Manager);
 
 		// Unicast
-		void Enqueue(uint32 ConnectionId, FInternalNetRefIndex OwnerIndex, FInternalNetRefIndex SubObjectIndex, const TRefCountPtr<FNetObjectAttachment>& Attachment);
+		void Enqueue(uint32 ConnectionId, FInternalNetRefIndex OwnerIndex, FInternalNetRefIndex SubObjectIndex, const TRefCountPtr<FNetObjectAttachment>& Attachment, ENetObjectAttachmentSendPolicyFlags SendFlags);
 
 		// Multicast
-		void Enqueue(FInternalNetRefIndex OwnerIndex, FInternalNetRefIndex SubObjectIndex, const TRefCountPtr<FNetObjectAttachment>& Attachment);
+		void Enqueue(FInternalNetRefIndex OwnerIndex, FInternalNetRefIndex SubObjectIndex, const TRefCountPtr<FNetObjectAttachment>& Attachment, ENetObjectAttachmentSendPolicyFlags SendFlags);
 
 		void PrepareProcessQueue(FReplicationConnections* InConnections, const FNetRefHandleManager* InNetRefHandleManager);
 		void ProcessQueue(EProcessMode ProcessMode);
@@ -93,6 +94,7 @@ private:
 			uint32 ConnectionId;
 			FInternalNetRefIndex OwnerIndex;
 			FInternalNetRefIndex SubObjectIndex;
+			ENetObjectAttachmentSendPolicyFlags SendFlags;
 			TRefCountPtr<FNetObjectAttachment> Attachment;
 		};
 		typedef TArray<FNetObjectAttachmentQueueEntry> FQueue;
diff --git a/Engine/Source/Runtime/Experimental/Iris/Core/Private/Iris/ReplicationSystem/ReplicationSystem.cpp b/Engine/Source/Runtime/Experimental/Iris/Core/Private/Iris/ReplicationSystem/ReplicationSystem.cpp
--- a/Engine/Source/Runtime/Experimental/Iris/Core/Private/Iris/ReplicationSystem/ReplicationSystem.cpp
+++ b/Engine/Source/Runtime/Experimental/Iris/Core/Private/Iris/ReplicationSystem/ReplicationSystem.cpp
@@ -16,6 +16,7 @@
 #include "Iris/ReplicationSystem/ReplicationOperationsInternal.h"
 #include "Iris/ReplicationSystem/ReplicationReader.h"
 #include "Iris/ReplicationSystem/ReplicationSystemInternal.h"
+#include "Iris/ReplicationSystem/ReplicationSystemTypes.h"
 #include "Iris/ReplicationSystem/ReplicationTypes.h"
 #include "Iris/ReplicationSystem/ReplicationWriter.h"
 #include "Iris/Serialization/InternalNetSerializationContext.h"
@@ -32,6 +33,10 @@ namespace ReplicationSystemCVars
 static bool bForcePruneBeforeUpdate = false;
 static FAutoConsoleVariableRef CVarForcePruneBeforeUpdate(TEXT("net.Iris.ForcePruneBeforeUpdate"), bForcePruneBeforeUpdate, TEXT("Verify integrity of all tracked instances at the start of every update."));
 #endif
+
+static bool bAllowAttachmentSendPolicyFlags = true;
+static FAutoConsoleVariableRef CVarAllowAttachmentSendPolicyFlags(TEXT("net.Iris.Attachments.AllowSendPolicyFlags"), bAllowAttachmentSendPolicyFlags, TEXT("Allow use of ENetObjectAttachmentSendPolicyFlags to specify behavior of RPCs."));
+
 }
 
 namespace UE::Net::Private
@@ -47,6 +52,8 @@ public:
 	FNetObjectGroupHandle NetGroupOwnerNetObjectGroupHandle;
 	FNetObjectGroupHandle NetGroupReplayNetObjectGroupHandle;
 
+	TMap<FObjectKey, ENetObjectAttachmentSendPolicyFlags> AttachmentSendPolicyFlags;
+
 	explicit FReplicationSystemImpl(UReplicationSystem* InReplicationSystem, const UReplicationSystem::FReplicationSystemParams& Params)
 	: ReplicationSystem(InReplicationSystem)
 	, ReplicationSystemInternal(FReplicationSystemInternalInitParams({ InReplicationSystem->GetId(), Params.MaxReplicatedObjectCount }))
@@ -450,7 +457,8 @@ public:
 			  * the client to the server, if the RPC is allowed to be sent in the first place.
 			  */
 			Params.bAllowSendingAttachmentsToObjectsNotInScope = !ReplicationSystem->IsServer();
-			Params.bAllowReceivingAttachmentsFromRemoteObjectsNotInScope = !Params.bAllowSendingAttachmentsToObjectsNotInScope;
+			Params.bAllowReceivingAttachmentsFromRemoteObjectsNotInScope = true;
+
 			// Delaying attachments with unresolved references on the server could cause massive queues of RPCs, potentially an OOM situation.
 			Params.bAllowDelayingAttachmentsWithUnresolvedReferences = !ReplicationSystem->IsServer();
 
@@ -818,8 +826,39 @@ bool UReplicationSystem::QueueNetObjectAttachment(uint32 ConnectionId, const UE:
 
 bool UReplicationSystem::SendRPC(const UObject* Object, const UObject* SubObject, const UFunction* Function, const void* Parameters)
 {
+	UE::Net::ENetObjectAttachmentSendPolicyFlags SendFlags = UE::Net::ENetObjectAttachmentSendPolicyFlags::None;
+	if (ReplicationSystemCVars::bAllowAttachmentSendPolicyFlags)
+	{
+		if (UE::Net::ENetObjectAttachmentSendPolicyFlags* Flags = Impl->AttachmentSendPolicyFlags.Find(FObjectKey(Function)))
+		{
+			SendFlags = *Flags;
+		}
+	}
+
 	UE::Net::Private::FNetBlobManager& NetBlobManager = Impl->ReplicationSystemInternal.GetNetBlobManager();
-	return NetBlobManager.SendRPC(Object, SubObject, Function, Parameters);
+	return NetBlobManager.SendRPC(Object, SubObject, Function, Parameters, SendFlags);
+}
+
+bool UReplicationSystem::SetRPCSendPolicyFlags(const UFunction* Function, UE::Net::ENetObjectAttachmentSendPolicyFlags SendFlags)
+{
+	if (!Function)
+	{	
+		return false;
+	}
+
+	if (EnumHasAnyFlags(SendFlags, UE::Net::ENetObjectAttachmentSendPolicyFlags::SendImmediate) && (Function->FunctionFlags & FUNC_NetReliable))
+	{
+		ensureAlwaysMsgf(false, TEXT("ENetObjectAttachmentSendPolicyFlags::SendImmediate is not allowed to use on Reliable RPC: %s"), *GetNameSafe(Function));
+		return false;
+	}
+		
+	Impl->AttachmentSendPolicyFlags.Add(FObjectKey(Function), SendFlags);
+	return true;
+}
+
+void UReplicationSystem::ResetRPCSendPolicyFlags()
+{
+	Impl->AttachmentSendPolicyFlags.Reset();
 }
 
 bool UReplicationSystem::SendRPC(uint32 ConnectionId, const UObject* Object, const UObject* SubObject, const UFunction* Function, const void* Parameters)
diff --git a/Engine/Source/Runtime/Experimental/Iris/Core/Private/Iris/ReplicationSystem/ReplicationWriter.cpp b/Engine/Source/Runtime/Experimental/Iris/Core/Private/Iris/ReplicationSystem/ReplicationWriter.cpp
--- a/Engine/Source/Runtime/Experimental/Iris/Core/Private/Iris/ReplicationSystem/ReplicationWriter.cpp
+++ b/Engine/Source/Runtime/Experimental/Iris/Core/Private/Iris/ReplicationSystem/ReplicationWriter.cpp
@@ -277,7 +277,8 @@ static bool s_ValidateReplicationRecord(const FReplicationRecord* ReplicationRec
 		for (uint32 It = 0U; It < RecordInfoCount; ++It)
 		{
  			const FReplicationRecord::FRecordInfo& RecordInfo = ReplicationRecord->PeekInfoAtOffset(It + Offset);
-			if (BitArray.GetBit(RecordInfo.Index))
+			// We allow multiple entries for the OOB attachments but do not expect multiple entires for normal replicated objects
+			if (RecordInfo.Index != 0U && BitArray.GetBit(RecordInfo.Index))
 			{
 				ensure(false);
 				return false;
@@ -309,7 +310,7 @@ bool FReplicationWriter::IsReplicationEnabled() const
 
 // $IRIS TODO : May need to introduce queue and send behaviors. For example one may want to send only with object.
 // One may not want to send unless the object is replicated very soon etc.
-void FReplicationWriter::QueueNetObjectAttachments(FInternalNetRefIndex OwnerInternalIndex, FInternalNetRefIndex SubObjectInternalIndex, TArrayView<const TRefCountPtr<FNetBlob>> InAttachments)
+void FReplicationWriter::QueueNetObjectAttachments(FInternalNetRefIndex OwnerInternalIndex, FInternalNetRefIndex SubObjectInternalIndex, TArrayView<const TRefCountPtr<FNetBlob>> InAttachments, ENetObjectAttachmentSendPolicyFlags SendFlags)
 {
 	if (InAttachments.Num() <= 0)
 	{
@@ -323,9 +324,12 @@ void FReplicationWriter::QueueNetObjectAttachments(FInternalNetRefIndex OwnerInt
 		UE_CLOG_REPLICATIONWRITER_WARNING(bWarnAboutDroppedAttachmentsToObjectsNotInScope, TEXT("Dropping %s attachment due to object ( InternalIndex: %u ) not in scope."), (EnumHasAnyFlags(InAttachments[0]->GetCreationInfo().Flags, ENetBlobFlags::Reliable) ? TEXT("reliable") : TEXT("unreliable")), OwnerInternalIndex);
 		return;
 	}
+	
+	// Route attachments flagged with ScheduleAsOOB through OOB channel if we have started replicating the owner.
+	const bool bScheduleUsingOOBChannel = EnumHasAnyFlags(SendFlags, ENetObjectAttachmentSendPolicyFlags::ScheduleAsOOB) && (GetReplicationInfo(OwnerInternalIndex).GetState() >= EReplicatedObjectState::WaitOnCreateConfirmation);
 
-	const uint32 TargetIndex = bObjectInScope ? (SubObjectInternalIndex != FNetRefHandleManager::InvalidInternalIndex ? SubObjectInternalIndex : OwnerInternalIndex) : ObjectIndexForOOBAttachment;
-	ENetObjectAttachmentType AttachmentType = (bObjectInScope ? ENetObjectAttachmentType::Normal : ENetObjectAttachmentType::OutOfBand);
+	const uint32 TargetIndex = (bObjectInScope && !bScheduleUsingOOBChannel) ? (SubObjectInternalIndex != FNetRefHandleManager::InvalidInternalIndex ? SubObjectInternalIndex : OwnerInternalIndex) : ObjectIndexForOOBAttachment;
+	ENetObjectAttachmentType AttachmentType = ((bObjectInScope && !bScheduleUsingOOBChannel) ? ENetObjectAttachmentType::Normal : ENetObjectAttachmentType::OutOfBand);
 	if (!Attachments.Enqueue(AttachmentType, TargetIndex, InAttachments))
 	{
 		return;
@@ -337,7 +341,6 @@ void FReplicationWriter::QueueNetObjectAttachments(FInternalNetRefIndex OwnerInt
 		return;
 	}
 
-
 	FReplicationInfo& TargetInfo = GetReplicationInfo(TargetIndex);
 	TargetInfo.HasAttachments = 1;
 
diff --git a/Engine/Source/Runtime/Experimental/Iris/Core/Private/Iris/ReplicationSystem/ReplicationWriter.h b/Engine/Source/Runtime/Experimental/Iris/Core/Private/Iris/ReplicationSystem/ReplicationWriter.h
--- a/Engine/Source/Runtime/Experimental/Iris/Core/Private/Iris/ReplicationSystem/ReplicationWriter.h
+++ b/Engine/Source/Runtime/Experimental/Iris/Core/Private/Iris/ReplicationSystem/ReplicationWriter.h
@@ -182,7 +182,7 @@ public:
 	void SetNetExports(FNetExports& InNetExports);
 
 	// Attachments
-	void QueueNetObjectAttachments(FInternalNetRefIndex OwnerInternalIndex, FInternalNetRefIndex SubObjectInternalIndex, TArrayView<const TRefCountPtr<FNetBlob>> Attachments);
+	void QueueNetObjectAttachments(FInternalNetRefIndex OwnerInternalIndex, FInternalNetRefIndex SubObjectInternalIndex, TArrayView<const TRefCountPtr<FNetBlob>> Attachments, ENetObjectAttachmentSendPolicyFlags SendFlags);
 
 private:
 	// Various types
diff --git a/Engine/Source/Runtime/Experimental/Iris/Core/Public/Iris/ReplicationSystem/ReplicationSystem.h b/Engine/Source/Runtime/Experimental/Iris/Core/Public/Iris/ReplicationSystem/ReplicationSystem.h
--- a/Engine/Source/Runtime/Experimental/Iris/Core/Public/Iris/ReplicationSystem/ReplicationSystem.h
+++ b/Engine/Source/Runtime/Experimental/Iris/Core/Public/Iris/ReplicationSystem/ReplicationSystem.h
@@ -213,6 +213,17 @@ public:
 	 */
 	IRISCORE_API bool SendRPC(const UObject* Object, const UObject* SubObject, const UFunction* Function, const void* Parameters);
 
+	/**
+	 * Set the policy flags for an RPC identified by its function
+	 * @param Function a pointer to a valid function identifying the RPC
+	 * @param SendFlags the ENetObjectAttachmentSendPolicyFlags to set for the RPC
+	 * @return Whether the specified SendFlags is valid for the specific RPC
+	 */
+	IRISCORE_API bool SetRPCSendPolicyFlags(const UFunction* Function, UE::Net::ENetObjectAttachmentSendPolicyFlags SendFlags);
+
+	/** Resets all set RPCSendPolicy flags */
+	IRISCORE_API void ResetRPCSendPolicyFlags();
+	
 	/**
 	 * Unicast an RPC targeting a object/subobject.
 	 * @param ConnectionId A valid connection ID. Only this connection will replicate the RPC.
diff --git a/Engine/Source/Runtime/Experimental/Iris/Core/Public/Iris/ReplicationSystem/ReplicationSystemTypes.h b/Engine/Source/Runtime/Experimental/Iris/Core/Public/Iris/ReplicationSystem/ReplicationSystemTypes.h
--- a/Engine/Source/Runtime/Experimental/Iris/Core/Public/Iris/ReplicationSystem/ReplicationSystemTypes.h
+++ b/Engine/Source/Runtime/Experimental/Iris/Core/Public/Iris/ReplicationSystem/ReplicationSystemTypes.h
@@ -29,4 +29,22 @@ enum class EDependentObjectSchedulingHint : uint8
 using FForwardNetRPCCallDelegate = TDelegate<void(UObject* RootObject,UObject* SubObject, UFunction* Function, void* Params)>;
 using FForwardNetRPCCallMulticastDelegate = TMulticastDelegate<typename FForwardNetRPCCallDelegate::TFuncType>;
 
+enum class ENetObjectAttachmentSendPolicyFlags : uint32
+{
+	// Default
+	None = 0,
+
+	// Schedule attachment to use the Out of bounds channel, essentially schedule the attachment to be sent as early as possible. Note: Only valid for unreliable attachments.
+	ScheduleAsOOB = 1U << 0U,
+
+	// Hint that this attachment like to be sent during PostTickDispatch. 
+	// If one attachment has this flag all unreliable attachments scheduled to use the OOB channel they all will be sent during PostTickDispatch.
+	SendInPostTickDispatch = ScheduleAsOOB << 1U,	
+
+	// SendImmediate, Attachment should be sent using OOB channel and from PostTickDispatch.
+	SendImmediate = ScheduleAsOOB | SendInPostTickDispatch,
+};
+
+ENUM_CLASS_FLAGS(ENetObjectAttachmentSendPolicyFlags);
+
 }
