Expose Async Physics Block Mode  and Async Interpolation Multiplier in the Solver API.
Those parameters are still available throw CVar. This change has been made to be able to customize those variable by project.

diff --git a/Engine/Source/Programs/HeadlessChaos/Private/HeadlessChaosTestRewind.cpp b/Engine/Source/Programs/HeadlessChaos/Private/HeadlessChaosTestRewind.cpp
--- a/Engine/Source/Programs/HeadlessChaos/Private/HeadlessChaosTestRewind.cpp
+++ b/Engine/Source/Programs/HeadlessChaos/Private/HeadlessChaosTestRewind.cpp
@@ -775,7 +775,7 @@ namespace ChaosTest {
 					TickSolverHelper(Solver);
 
 					Time += GTDt;
-					const FReal InterpolatedTime = Time - SimDt * Chaos::AsyncInterpolationMultiplier;
+					const FReal InterpolatedTime = Time - SimDt * Solver->GetAsyncInterpolationMultiplier();
 
 
 					if (InterpolatedTime <= InterpStartTime)
@@ -2111,7 +2111,7 @@ namespace ChaosTest {
 				FReal LastCorrectionStep = 0;
 
 				Solver->SetRewindCallback(MoveTemp(UniqueRewindCallback));
-
+				Solver->SetAsyncInterpolationMultiplier(4.0f);
 				{
 					auto& Particle = Proxy->GetGameThreadAPI();
 					Solver->GetEvolution()->GetGravityForces().SetAcceleration(FVec3(0, 0, -1), 0);
@@ -2124,7 +2124,7 @@ namespace ChaosTest {
 						TickSolverHelper(Solver);
 
 						Time += GTDt;
-						const FReal InterpolatedTime = Time - SimDt * Chaos::AsyncInterpolationMultiplier;
+						const FReal InterpolatedTime = Time - SimDt * Solver->GetAsyncInterpolationMultiplier();
 						if (InterpolatedTime < 0)
 						{
 							//not enough time to interpolate so just take initial value
@@ -2138,13 +2138,14 @@ namespace ChaosTest {
 					}
 
 					//resim happened
+					Solver->SetAsyncInterpolationMultiplier(2.0f);
 					for (int Step = LastGameStep; Step < 2*LastGameStep; ++Step)
 					{
 						const FReal PrevZ = Particle.X()[2];
 						TickSolverHelper(Solver);
 
 						Time += GTDt;
-						const FReal InterpolatedTime = Time - SimDt * Chaos::AsyncInterpolationMultiplier;
+						const FReal InterpolatedTime = Time - SimDt * Solver->GetAsyncInterpolationMultiplier();
 
 						if(InterpolatedTime > 20)	//resim happened
 						{
@@ -4327,6 +4328,7 @@ namespace ChaosTest {
 
 	GTEST_TEST(AllTraits, RewindTest_InterpolatedTwoChannels)
 	{
+		Chaos::AsyncInterpolationMultiplier = 3.0f;
 		int32 PrevNumActiveChannels = Chaos::DefaultNumActiveChannels;
 		Chaos::DefaultNumActiveChannels = 2;
 		//Have two moving particles, one in each channel to see that there's a delay in time on second channel
@@ -4354,7 +4356,7 @@ namespace ChaosTest {
 				TickSolverHelper(Solver);
 
 				Time += GtDt;
-				const FReal InterpolatedTime0 = Time - SimDt * Chaos::AsyncInterpolationMultiplier;
+				const FReal InterpolatedTime0 = Time - SimDt * Solver->GetAsyncInterpolationMultiplier();
 				const FReal InterpolatedTime1 = InterpolatedTime0 - Chaos::SecondChannelDelay;
 
 				if (InterpolatedTime0 < 0)
diff --git a/Engine/Source/Runtime/Experimental/Chaos/Private/Chaos/Framework/PhysicsSolverBase.cpp b/Engine/Source/Runtime/Experimental/Chaos/Private/Chaos/Framework/PhysicsSolverBase.cpp
--- a/Engine/Source/Runtime/Experimental/Chaos/Private/Chaos/Framework/PhysicsSolverBase.cpp
+++ b/Engine/Source/Runtime/Experimental/Chaos/Private/Chaos/Framework/PhysicsSolverBase.cpp
@@ -8,6 +8,7 @@
 #include "PBDRigidsSolver.h"
 #include "PhysicsProxy/SingleParticlePhysicsProxy.h"
 #include "Chaos/Framework/ChaosResultsManager.h"
+#include "ChaosSolversModule.h"
 #include "Framework/Threading.h"
 #include "RewindData.h"
 
@@ -182,8 +183,24 @@ namespace Chaos
 	CHAOS_API int32 ForceDisableAsyncPhysics = 0;
 	FAutoConsoleVariableRef CVarForceDisableAsyncPhysics(TEXT("p.ForceDisableAsyncPhysics"), ForceDisableAsyncPhysics, TEXT("Whether to force async physics off regardless of other settings"));
 
+	auto LambdaMul = FConsoleVariableDelegate::CreateLambda([](IConsoleVariable* InVariable)
+		{
+			for (FPhysicsSolverBase* Solver : FChaosSolversModule::GetModule()->GetAllSolvers())
+			{
+				Solver->SetAsyncInterpolationMultiplier(InVariable->GetFloat());
+			}
+		});
+
+	auto LambdaAsyncMode = FConsoleVariableDelegate::CreateLambda([](IConsoleVariable* InVariable)
+		{
+			for (FPhysicsSolverBase* Solver : FChaosSolversModule::GetModule()->GetAllSolvers())
+			{
+				Solver->SetAsyncPhysicsBlockMode(EAsyncBlockMode(InVariable->GetInt()));
+			}
+		});
+
 	CHAOS_API FRealSingle AsyncInterpolationMultiplier = 2.f;
-	FAutoConsoleVariableRef CVarAsyncInterpolationMultiplier(TEXT("p.AsyncInterpolationMultiplier"), AsyncInterpolationMultiplier, TEXT("How many multiples of the fixed dt should we look behind for interpolation"));
+	FAutoConsoleVariableRef CVarAsyncInterpolationMultiplier(TEXT("p.AsyncInterpolationMultiplier"), AsyncInterpolationMultiplier, TEXT("How many multiples of the fixed dt should we look behind for interpolation"), LambdaMul);
 
 	// 0 blocks on any physics steps generated from past GT Frames, and blocks on none of the tasks from current frame.
 	// 1 blocks on everything except the single most recent task (including tasks from current frame)
@@ -192,8 +209,7 @@ namespace Chaos
 	int32 AsyncPhysicsBlockMode = 0;
 	FAutoConsoleVariableRef CVarAsyncPhysicsBlockMode(TEXT("p.AsyncPhysicsBlockMode"), AsyncPhysicsBlockMode, TEXT("Setting to 0 blocks on any physics steps generated from past GT Frames, and blocks on none of the tasks from current frame."
 		" 1 blocks on everything except the single most recent task (including tasks from current frame). 1 should gurantee we will always have a future output for interpolation from 2 frames in the past."
-		" 2 doesn't block the game thread, physics steps could be eventually be dropped if taking too much time."));
-
+		" 2 doesn't block the game thread, physics steps could be eventually be dropped if taking too much time."), LambdaAsyncMode);
 
 	FPhysicsSolverBase::FPhysicsSolverBase(const EMultiBufferMode BufferingModeIn,const EThreadingModeTemp InThreadingMode,UObject* InOwner, Chaos::FReal InAsyncDt)
 		: BufferMode(BufferingModeIn)
@@ -212,6 +228,8 @@ namespace Chaos
 		, MMinDeltaTime(UE_SMALL_NUMBER)
 		, MMaxSubSteps(1)
 		, ExternalSteps(0)
+		, AsyncBlockMode(EAsyncBlockMode(AsyncPhysicsBlockMode))
+		, AsyncMultiplier(AsyncInterpolationMultiplier)
 #if !UE_BUILD_SHIPPING
 		, bStealAdvanceTasksForTesting(false)
 #endif
@@ -425,7 +443,7 @@ namespace Chaos
 		}
 
 		// Eventually drop physics steps in mode 2
-		if (AsyncPhysicsBlockMode == 2)
+		if (AsyncBlockMode == EAsyncBlockMode::DoNoBlock)
 		{
 			// Make sure not to accumulate too many physics solver tasks.
 			constexpr int32 MaxPhysicsStepToKeep = 3;
@@ -475,7 +493,7 @@ namespace Chaos
 			else
 			{
 				// If enabled, block on all but most recent physics task, even tasks generated this frame.
-				if (AsyncPhysicsBlockMode == 1)
+				if (AsyncBlockMode == EAsyncBlockMode::BlockForBestInterpolation)
 				{
 					BlockingTasks = PendingTasks;
 				}
@@ -510,7 +528,7 @@ namespace Chaos
 				break;
 			}
 		}
-		if (AsyncPhysicsBlockMode == 2)
+		if (AsyncBlockMode == EAsyncBlockMode::DoNoBlock)
 		{
 			return {};
 		}
diff --git a/Engine/Source/Runtime/Experimental/Chaos/Public/Chaos/Framework/PhysicsSolverBase.h b/Engine/Source/Runtime/Experimental/Chaos/Public/Chaos/Framework/PhysicsSolverBase.h
--- a/Engine/Source/Runtime/Experimental/Chaos/Public/Chaos/Framework/PhysicsSolverBase.h
+++ b/Engine/Source/Runtime/Experimental/Chaos/Public/Chaos/Framework/PhysicsSolverBase.h
@@ -86,6 +86,14 @@ namespace Chaos
 		bool bSolverSubstepped;
 	};
 
+
+	enum EAsyncBlockMode
+	{
+		BlockOnlyPastFrames = 0,
+		BlockForBestInterpolation = 1,
+		DoNoBlock = 2
+	};
+
 	/**
 	 * Task responsible for processing the command buffer of a single solver and preparing data before solver task and callbacks are run
 	 */
@@ -633,7 +641,7 @@ namespace Chaos
 			if (IsUsingFixedDt())
 			{
 				//fixed dt uses interpolation and looks into the past
-				return ExternalTime - AsyncDt * AsyncInterpolationMultiplier;
+				return ExternalTime - AsyncDt * AsyncMultiplier;
 			}
 			else
 			{
@@ -651,6 +659,22 @@ namespace Chaos
 		/**/
 		FReal GetLastDt() const { return MLastDt; }
 
+		/** 
+		  Set the Async Block Mode, valid mode can be 0, 1, or 2
+		  0 blocks on any physics steps generated from past GT Frames, and blocks on none of the tasks from current frame.
+		  1 blocks on everything except the single most recent task (including tasks from current frame)
+		  1 should guarantee we will always have a future output for interpolation from 2 frames in the past
+		  2 doesn't block the game thread. Physics steps could be eventually be dropped if taking too much time.
+		*/
+		void SetAsyncPhysicsBlockMode(EAsyncBlockMode InAsyncBlockMode) { AsyncBlockMode = InAsyncBlockMode; }
+
+		/** 
+		* Set the async interpolation multiplier which is how many multiples of the fixed dt should we look behind for interpolation.
+		*/
+		void SetAsyncInterpolationMultiplier(FRealSingle InAsyncInterpolationMultiplier) { AsyncMultiplier = InAsyncInterpolationMultiplier; }
+
+		float GetAsyncInterpolationMultiplier() const { return AsyncMultiplier; }
+
 		/** Check if we can enable debugging informations for network physics */
 		static bool CanDebugNetworkPhysicsPrediction()
 		{
@@ -797,6 +821,8 @@ namespace Chaos
 		int32 MMaxSubSteps;
 		int32 ExternalSteps;
 		TArray<FGeometryParticle*> UniqueIdxToGTParticles;
+		EAsyncBlockMode AsyncBlockMode;
+		float AsyncMultiplier;
 
 	public:
 
