Convex shapes simplification

diff --git a/Engine/Source/Runtime/Experimental/Chaos/Private/Chaos/ChaosDebugDraw.cpp b/Engine/Source/Runtime/Experimental/Chaos/Private/Chaos/ChaosDebugDraw.cpp
--- a/Engine/Source/Runtime/Experimental/Chaos/Private/Chaos/ChaosDebugDraw.cpp
+++ b/Engine/Source/Runtime/Experimental/Chaos/Private/Chaos/ChaosDebugDraw.cpp
@@ -25,6 +25,7 @@
 #include "Chaos/PBDRigidParticles.h"
 #include "Chaos/PBDSuspensionConstraints.h"
 #include "Chaos/Sphere.h"
+#include "Chaos/ConvexOptimizer.h"
 #include "Chaos/Utilities.h"
 #include "Chaos/CCDUtilities.h"
 
@@ -57,6 +58,9 @@ namespace Chaos
 		bool bChaosDebugDebugDrawColorShapesByInternalCluster = false;
 		FAutoConsoleVariableRef CVarChaosDebugDebugDrawColorShapesByInternalCluster(TEXT("p.Chaos.DebugDraw.ColorShapesByInternalCluster"), bChaosDebugDebugDrawColorShapesByInternalCluster, TEXT("Whether to check if the particle is an internal cluster to define its color (black : regular particle: red :internal cluster )"));
 
+		bool bChaosDebugDebugDrawColorShapesBySimQueryType = false;
+		FAutoConsoleVariableRef CVarChaosDebugDebugDrawColorShapesBySimQueryType(TEXT("p.Chaos.DebugDraw.ColorShapesBySimQueryType"), bChaosDebugDebugDrawColorShapesBySimQueryType, TEXT("Whether to show with different colors shapes that are sim enabled and query enabled (sim : blue, query : orange)"));
+
 		bool bChaosDebugDebugDrawColorShapesByClusterUnion = false;
 		FAutoConsoleVariableRef CVarChaosDebugDebugDrawColorShapesByClusterUnion(TEXT("p.Chaos.DebugDraw.ColorShapesByClusterUnion"), bChaosDebugDebugDrawColorShapesByClusterUnion, TEXT("An extension of the ColorShapesByInternalCluster option: instead of using a single color for every internal cluster, will use a unique color per cluster union. Non-cluster unions will be black."));
 
@@ -582,40 +586,61 @@ namespace Chaos
 				DrawShapesImpl(Particle, TransformedTransform, Transformed->GetTransformedObject(), Shapes, Margin, Color, Duration, Settings);
 				return;
 			}
-			else if (InnerType == ImplicitObjectType::Union)
+			else if ((InnerType == ImplicitObjectType::Union) || (InnerType == ImplicitObjectType::UnionClustered))
 			{
-				const FImplicitObjectUnion* Union = Implicit->template GetObject<FImplicitObjectUnion>();
-
-				for (int32 UnionIdx = 0; UnionIdx < Union->NumImplicits(); ++UnionIdx)
+				if(InnerType == ImplicitObjectType::Union)
 				{
-					// Retrieve shape from union's shapes array
-					const FPerShapeData* PerShapeData = nullptr;
-					if (!Shapes.IsSingleShape())
+					const FImplicitObjectUnion* Union = Implicit->template GetObject<FImplicitObjectUnion>();
+					for (int32 UnionIdx = 0; UnionIdx < Union->GetObjects().Num(); ++UnionIdx)
 					{
-						const FShapesArray& ShapesArray = *Shapes.GetShapesArray();
-						PerShapeData = ShapesArray[UnionIdx].Get();
+						// Retrieve shape from union's shapes array
+						const FPerShapeData* PerShapeData = nullptr;
+						if (!Shapes.IsSingleShape())
+						{
+							const FShapesArray& ShapesArray = *Shapes.GetShapesArray();
+							PerShapeData = ShapesArray[UnionIdx].Get();
+						}
+						DrawShapesImpl(Particle, ShapeTransform, Union->GetObjects()[UnionIdx].GetReference(), FShapeOrShapesArray(PerShapeData), Margin, Color, Duration, Settings);
 					}
-					
-					DrawShapesImpl(Particle, ShapeTransform, Union->GetImplicit(UnionIdx).GetReference(), FShapeOrShapesArray(PerShapeData), Margin, Color, Duration, Settings);
-					
 				}
-				return;
-			}
-			else if (InnerType == ImplicitObjectType::UnionClustered)
-			{
-				const FImplicitObjectUnionClustered* Union = Implicit->template GetObject<FImplicitObjectUnionClustered>();
-				for (auto& UnionImplicit : Union->GetObjects())
+				else
 				{
-					const TPBDRigidParticleHandle<FReal, 3>* OriginalParticle = Union->FindParticleForImplicitObject(UnionImplicit.GetReference());
-					if (ensure(OriginalParticle))
+					// For union clustered if we have simplified convexes we are generating the particle
+					// shapes since we have more than children. Otherwise we can grab the shape from the original particle lookup
+					const FImplicitObjectUnionClustered* Union = Implicit->template GetObject<FImplicitObjectUnionClustered>();
+					for (auto& UnionImplicit : Union->GetObjects())
 					{
-						DrawShapesImpl(Particle, ShapeTransform, UnionImplicit.GetReference(), FShapeOrShapesArray(OriginalParticle), Margin, Color, Duration, Settings);
+						const TPBDRigidParticleHandle<FReal, 3>* OriginalParticle = Union->FindParticleForImplicitObject(UnionImplicit.GetReference());
+						if (ensure(OriginalParticle))
+						{
+							DrawShapesImpl(Particle, ShapeTransform, UnionImplicit.GetReference(), FShapeOrShapesArray(OriginalParticle), Margin, Color, Duration, Settings);
+						}
+					}
+				}
+				// Draw the manager shapes if the implicit is the root union 
+				if(Particle->GetGeometry() == Implicit)
+				{
+					if(auto ClusteredParticle = Particle->CastToClustered())
+					{
+						const TPimplPtr<Private::FConvexOptimizer>& ConvexOptimizer = ClusteredParticle->ConvexOptimizer();
+						if(ConvexOptimizer && ConvexOptimizer->IsValid())
+						{
+							for (auto& PerShapeData : ConvexOptimizer->GetShapeInstances())
+							{
+								if(FImplicitObjectUnion* Union = PerShapeData->GetGeometry()->AsA<FImplicitObjectUnion>())
+								{
+									for(auto& UnionConvex : Union->GetObjects())
+									{
+										DrawShapesImpl(Particle, ShapeTransform, UnionConvex, FShapeOrShapesArray(PerShapeData.Get()), Margin, Color, Duration, Settings);
+									}
+								}
+							}
+						}
 					}
 				}
 				return;
 			}
 
-
 			// Whether we should show meshes and non-mesh shapes
 			bool bShowMeshes = Settings.bShowComplexCollision;
 			bool bShowNonMeshes = Settings.bShowSimpleCollision;
@@ -679,6 +704,32 @@ namespace Chaos
 					}
 				}
 			}
+			if (ShapeData && bChaosDebugDebugDrawColorShapesBySimQueryType)
+			{
+				if (ShapeData->GetSimEnabled())
+				{
+					const bool bIsUnion = Particle->GetGeometry()->IsUnderlyingUnion();
+					if (bIsUnion)
+					{
+						if(Implicit->GetDoCollide())
+						{
+							ShapeColor = FColor::Green;
+						}
+						else
+						{
+							ShapeColor = FColor::Blue;
+						}
+					}
+					else
+					{
+						ShapeColor = FColor::Orange;
+					}
+				}
+				else if (ShapeData->GetQueryEnabled())
+				{
+					ShapeColor = FColor::Red;
+				}
+			}
 
 			// If we get here, we have an actual shape to render
 			switch (InnerType)
diff --git a/Engine/Source/Runtime/Experimental/Chaos/Private/Chaos/ClusterUnionManager.cpp b/Engine/Source/Runtime/Experimental/Chaos/Private/Chaos/ClusterUnionManager.cpp
--- a/Engine/Source/Runtime/Experimental/Chaos/Private/Chaos/ClusterUnionManager.cpp
+++ b/Engine/Source/Runtime/Experimental/Chaos/Private/Chaos/ClusterUnionManager.cpp
@@ -62,7 +62,7 @@ namespace Chaos
 						// NOTE: This will need to be re-thought if we ever decide to make the mapping between child shapes and child particles not 1-to-1 (e.g. if we ever attempt to simplify the cluster union shape).
 
 						// Ignore intersections against the same particle since we just added the particle (potentially) into the geometry.
-						if (Particle == ClusterUnion.ChildParticles[RootObjectIndex])
+						if ((RootObjectIndex >= ClusterUnion.ChildParticles.Num()) || ((RootObjectIndex < ClusterUnion.ChildParticles.Num()) && (Particle == ClusterUnion.ChildParticles[RootObjectIndex])))
 						{
 							return;
 						}
@@ -915,6 +915,9 @@ namespace Chaos
 		{
 			FlushIncrementalGeometryOperations(ClusterUnion);
 		}
+
+		// Build the convex optimizer if required
+		FRigidClustering::BuildConvexOptimizer(ClusterUnion.InternalCluster);
 		
 		if (ClusterUnion.bGenerateConnectivityEdges)
 		{
diff --git a/Engine/Source/Runtime/Experimental/Chaos/Private/Chaos/Collision/ParticlePairMidPhase.cpp b/Engine/Source/Runtime/Experimental/Chaos/Private/Chaos/Collision/ParticlePairMidPhase.cpp
--- a/Engine/Source/Runtime/Experimental/Chaos/Private/Chaos/Collision/ParticlePairMidPhase.cpp
+++ b/Engine/Source/Runtime/Experimental/Chaos/Private/Chaos/Collision/ParticlePairMidPhase.cpp
@@ -16,6 +16,7 @@
 #include "Chaos/PBDCollisionConstraints.h"
 #include "ChaosStats.h"
 #include "Misc/MemStack.h"
+#include "Chaos/ConvexOptimizer.h"
 #include "ProfilingDebugging/CountersTrace.h"
 
 TRACE_DECLARE_INT_COUNTER_EXTERN(ChaosTraceCounter_MidPhase_NumShapePair);
@@ -109,10 +110,17 @@ namespace Chaos
 	// Get the ShapeInstance data from the particle for the implicit with the specified root object index (its index
 	// in the root union implicit if there is one). Usually every implicit in the root union is represented in the
 	// ShapesArray, but not always. GCS and ClusterUnions sometimes have a Union at the root but only a single ShapeInstance
-	const FShapeInstance* GetShapeInstance(const FShapeInstanceArray& ShapeInstances, const int32 RootObjectIndex)
+	const FShapeInstance* GetShapeInstance(const FShapeInstanceArray& ShapeInstances, const int32 RootObjectIndex, const Private::FConvexOptimizer* ConvexOptimizer = nullptr)
 	{
-		const int32 ShapeIndex = (ShapeInstances.IsValidIndex(RootObjectIndex)) ? RootObjectIndex : 0;
-		return ShapeInstances[ShapeIndex].Get();
+		if(ConvexOptimizer && ConvexOptimizer->IsValid() && (RootObjectIndex == INDEX_NONE))
+		{
+			return ConvexOptimizer->GetShapeInstances()[0].Get();
+		}
+		else
+		{
+			const int32 ShapeIndex = (ShapeInstances.IsValidIndex(RootObjectIndex)) ? RootObjectIndex : 0;
+			return ShapeInstances[ShapeIndex].Get();
+		}
 	}
 
 
@@ -1119,6 +1127,11 @@ namespace Chaos
 		const FImplicitObjectRef Implicit0 = GetParticle0()->GetGeometry();
 		const FImplicitObjectRef Implicit1 = GetParticle1()->GetGeometry();
 
+		FPBDRigidClusteredParticleHandle* ClusteredHandle0 = GetParticle0()->CastToClustered();
+		FPBDRigidClusteredParticleHandle* ClusteredHandle1 = GetParticle1()->CastToClustered();
+
+		Private::FConvexOptimizer* ConvexOptimizer0 = ClusteredHandle0 ? ClusteredHandle0->ConvexOptimizer().Get() : nullptr;
+		Private::FConvexOptimizer* ConvexOptimizer1 = ClusteredHandle1 ? ClusteredHandle1->ConvexOptimizer().Get() : nullptr;
 
 		// See if we have a BVH for either/both of the particles
 		const Private::FImplicitBVH* BVH0 = nullptr;
@@ -1127,11 +1140,19 @@ namespace Chaos
 		{
 			if (const FImplicitObjectUnion* Union0 = Implicit0->template AsA<FImplicitObjectUnion>())
 			{
-				BVH0 = Union0->GetBVH();
+				const bool bHasConvexOptimizer0 = (ConvexOptimizer0 != nullptr) && ConvexOptimizer0->IsValid();
+				if(!bHasConvexOptimizer0)
+				{
+					BVH0 = Union0->GetBVH();
+				}
 			}
 			if (const FImplicitObjectUnion* Union1 = Implicit1->template AsA<FImplicitObjectUnion>())
 			{
-				BVH1 = Union1->GetBVH();
+				const bool bHasConvexOptimizer1 = (ConvexOptimizer1 != nullptr) && ConvexOptimizer1->IsValid();
+				if(!bHasConvexOptimizer1)
+				{
+					BVH1 = Union1->GetBVH();
+				}
 			}
 		}
 
@@ -1144,15 +1165,15 @@ namespace Chaos
 		}
 		else if (BVH0 != nullptr)
 		{
-			GenerateCollisionsBVHImplicitHierarchy(GetParticle0(), BVH0, GetParticle1(), Implicit1, CullDistance, Dt, Context);
+			GenerateCollisionsBVHImplicitHierarchy(GetParticle0(), BVH0, GetParticle1(), Implicit1, ConvexOptimizer1, CullDistance, Dt, Context);
 		}
 		else if (BVH1 != nullptr)
 		{
-			GenerateCollisionsBVHImplicitHierarchy(GetParticle1(), BVH1, GetParticle0(), Implicit0, CullDistance, Dt, Context);
+			GenerateCollisionsBVHImplicitHierarchy(GetParticle1(), BVH1, GetParticle0(), Implicit0, ConvexOptimizer0, CullDistance, Dt, Context);
 		}
 		else
 		{
-			GenerateCollisionsImplicitHierarchyImplicitHierarchy(GetParticle0(), Implicit0, GetParticle1(), Implicit1, CullDistance, Dt, Context);
+			GenerateCollisionsImplicitHierarchyImplicitHierarchy(GetParticle0(), Implicit0, ConvexOptimizer0, GetParticle1(), Implicit1, ConvexOptimizer1, CullDistance, Dt, Context);
 		}
 
 		// Generate manifolds for each constraint we created/recovered and (re)activate if necessary
@@ -1193,18 +1214,20 @@ namespace Chaos
 					CullDistance, Dt, &Context]
 					(const FImplicitObject* ImplicitA, const FRigidTransform3f& RelativeTransformfA, const FAABB3f& RelativeBoundsfA, const int32 RootObjectIndexA, const int32 LeafObjectIndexA) -> void
 					{
-						const FRigidTransform3 RelativeTransformA = FRigidTransform3(RelativeTransformfA);
 						const FShapeInstance* ShapeInstanceA = GetShapeInstance(ShapeInstancesA, RootObjectIndexA);
-
+						if (!FilterHasSimEnabled(ShapeInstanceA)) return;
+						
+						const FRigidTransform3 RelativeTransformA = FRigidTransform3(RelativeTransformfA);
 						BVHB->VisitNodeObjects(NodeIndexB,
 							[this, ParticleA, ImplicitA, ShapeInstanceA, &ParticleWorldTransformA, &RelativeTransformA, LeafObjectIndexA,
 							ParticleB, BVHB, &ParticleWorldTransformB, &ShapeInstancesB,
 							CullDistance, Dt, &Context]
 							(const FImplicitObject* ImplicitB, const FRigidTransform3f& RelativeTransformfB, const FAABB3f& RelativeBoundsfB, const int32 RootObjectIndexB, const int32 LeafObjectIndexB) -> void
 							{
-								const FRigidTransform3 RelativeTransformB = FRigidTransform3(RelativeTransformfB);
 								const FShapeInstance* ShapeInstanceB = GetShapeInstance(ShapeInstancesB, RootObjectIndexB);
-
+								if (!FilterHasSimEnabled(ShapeInstanceB)) return;
+								
+								const FRigidTransform3 RelativeTransformB = FRigidTransform3(RelativeTransformfB);
 								// Detect collisions between the single implicit object pair
 								GenerateCollisionsImplicitLeafImplicitLeaf(
 									ParticleA, ImplicitA, ShapeInstanceA, ParticleWorldTransformA, RelativeTransformA, LeafObjectIndexA,
@@ -1218,17 +1241,18 @@ namespace Chaos
 	// Detect collisions between a BVH and some other implicit object (which may be a hierarchy)
 	void FGenericParticlePairMidPhase::GenerateCollisionsBVHImplicitHierarchy(
 		FGeometryParticleHandle* ParticleA, const Private::FImplicitBVH* BVHA,
-		FGeometryParticleHandle* ParticleB, const FImplicitObject* RootImplicitB,
+		FGeometryParticleHandle* ParticleB, const FImplicitObject* RootImplicitB, const Private::FConvexOptimizer* ConvexOptimizerB,
 		const FReal CullDistance, const FReal Dt, const FCollisionContext& Context)
 	{
 		const FShapeInstanceArray& ShapeInstancesB = ParticleB->ShapeInstances();
 
 		// Visit all the leaf implicits in RootImplicitB and collide against the BVH
-		RootImplicitB->VisitLeafObjects(
-			[this, ParticleA, BVHA, ParticleB, &ShapeInstancesB, CullDistance, Dt, &Context]
+		VisitCollisionObjects(ConvexOptimizerB, RootImplicitB,
+			[this, ParticleA, BVHA, ParticleB, &ShapeInstancesB, CullDistance, Dt, &Context, &ConvexOptimizerB]
 			(const FImplicitObject* ImplicitB, const FRigidTransform3& RelativeTransformB, const int32 RootObjectIndexB, const int32 ObjectIndexB, const int32 LeafObjectIndexB) -> void
 			{
-				const FShapeInstance* ShapeInstanceB = GetShapeInstance(ShapeInstancesB, RootObjectIndexB);
+				const FShapeInstance* ShapeInstanceB = GetShapeInstance(ShapeInstancesB, RootObjectIndexB, ConvexOptimizerB);
+				if (!FilterHasSimEnabled(ShapeInstanceB)) return;
 
 				// ImplicitB is a single object. We perform the bounds tests in A space
 				GenerateCollisionsBVHImplicitLeaf(
@@ -1240,8 +1264,8 @@ namespace Chaos
 
 	// Detect collisions between two implicits, where either or both may be a hierarchy, but neither has a BVH
 	void FGenericParticlePairMidPhase::GenerateCollisionsImplicitHierarchyImplicitHierarchy(
-		FGeometryParticleHandle* ParticleA, const FImplicitObject* RootImplicitA,
-		FGeometryParticleHandle* ParticleB, const FImplicitObject* RootImplicitB,
+		FGeometryParticleHandle* ParticleA, const FImplicitObject* RootImplicitA, const Private::FConvexOptimizer* ConvexOptimizerA,
+		FGeometryParticleHandle* ParticleB, const FImplicitObject* RootImplicitB, const Private::FConvexOptimizer* ConvexOptimizerB,
 		const FReal CullDistance,
 		const FReal Dt,
 		const FCollisionContext& Context)
@@ -1259,25 +1283,26 @@ namespace Chaos
 		// Detect collisions between Implicit Hierarchy of ParticleA and Implicit Hierarchy of ParticleB
 		// Given an ImplicitObject from ParticleA (which we know overlaps the bounds of some parts of ParticleB),
 		// run collision detection on ImplicitA against the implicit object hierarchy of ParticleB.
-		RootImplicitA->VisitLeafObjects(
-			[this, ParticleA, &ShapeInstancesA, &ParticleWorldTransformA,
+		VisitCollisionObjects(ConvexOptimizerA, RootImplicitA, [this, ParticleA, &ShapeInstancesA, &ParticleWorldTransformA,
 			ParticleB, &ShapeInstancesB, RootImplicitB, &ParticleWorldTransformB, &ParticleTransformAToB,
-			CullDistance, Dt, &Context]
-			(const FImplicitObject* ImplicitA, const FRigidTransform3& RelativeTransformA, const int32 RootObjectIndexA, const int32 ObjectIndex, const int32 LeafObjectIndexA)
+			CullDistance, Dt, &Context, &ConvexOptimizerA, &ConvexOptimizerB]
+			(const FImplicitObject* ImplicitA, const FRigidTransform3& RelativeTransformA, const int32 RootObjectIndexA, const int32 ObjectIndexA, const int32 LeafObjectIndexA)
 			{
+				const FShapeInstance* ShapeInstanceA = GetShapeInstance(ShapeInstancesA, RootObjectIndexA, ConvexOptimizerA);
+				if (!FilterHasSimEnabled(ShapeInstanceA)) return;
+
 				const FAABB3 RelativeBoundsA = ImplicitA->CalculateTransformedBounds(RelativeTransformA);
 				const FAABB3 ShapeBoundsAInB = RelativeBoundsA.TransformedAABB(ParticleTransformAToB).ThickenSymmetrically(FVec3(CullDistance));
 
-				const FShapeInstance* ShapeInstanceA = GetShapeInstance(ShapeInstancesA, RootObjectIndexA);
-
 				// Detect collisions between ImplicitA and Implicit Hierarchy of ParticleB
-				RootImplicitB->VisitOverlappingLeafObjects(ShapeBoundsAInB,
+				VisitOverlappingObjects(ConvexOptimizerB, RootImplicitB, ShapeBoundsAInB,
 					[this, ParticleA, ImplicitA, ShapeInstanceA, &ParticleWorldTransformA, &RelativeTransformA, LeafObjectIndexA,
 					ParticleB, &ParticleWorldTransformB, &ShapeInstancesB,
-					CullDistance, Dt, &Context]
+					CullDistance, Dt, &Context, &ConvexOptimizerB]
 					(const FImplicitObject* ImplicitB, const FRigidTransform3& RelativeTransformB, const int32 RootObjectIndexB, const int32 ObjectIndexB, const int32 LeafObjectIndexB)
 					{
-						const FShapeInstance* ShapeInstanceB = GetShapeInstance(ShapeInstancesB, RootObjectIndexB);
+						const FShapeInstance* ShapeInstanceB = GetShapeInstance(ShapeInstancesB, RootObjectIndexB, ConvexOptimizerB);
+						if (!FilterHasSimEnabled(ShapeInstanceB)) return;
 
 						// Detect collisions between ImplicitA and ImplicitB (both leaf implicits)
 						GenerateCollisionsImplicitLeafImplicitLeaf(
@@ -1320,6 +1345,8 @@ namespace Chaos
 			(const FImplicitObject* ImplicitA, const FRigidTransform3f& RelativeTransformfA, const FAABB3f& RelativeBoundsfA, const int32 RootObjectIndexA, const int32 LeafObjectIndexA) -> void
 			{
 				const FShapeInstance* ShapeInstanceA = GetShapeInstance(ShapeInstancesA, RootObjectIndexA);
+				if (!FilterHasSimEnabled(ShapeInstanceA)) return;
+				
 				const FRigidTransform3 RelativeTransformA = FRigidTransform3(RelativeTransformfA);
 
 				GenerateCollisionsImplicitLeafImplicitLeaf(
diff --git a/Engine/Source/Runtime/Experimental/Chaos/Private/Chaos/ConvexOptimizer.cpp b/Engine/Source/Runtime/Experimental/Chaos/Private/Chaos/ConvexOptimizer.cpp
new file mode 100644
--- /dev/null
+++ b/Engine/Source/Runtime/Experimental/Chaos/Private/Chaos/ConvexOptimizer.cpp
@@ -0,0 +1,164 @@
+
+#include "Chaos/ConvexOptimizer.h"
+#include "Chaos/Tribox.h"
+#include "Chaos/ImplicitObjectBVH.h"
+#include "Chaos/PhysicsObjectCollisionInterface.h"
+#include "Chaos/Collision/CollisionContext.h"
+
+namespace Chaos
+{
+	namespace CVars
+	{
+		// Replace all the convexes within an implicit hierarchy with a simplified one (kdop18 tribox for now) for collision
+		bool bChaosSimplifyConvexes = false;
+		FAutoConsoleVariableRef CVarChaosSimplifyConvexes(TEXT("p.Chaos.Collision.SimplifyConvexes"), bChaosSimplifyConvexes, TEXT("If true replace all the convexes within an implcit hierarchy with a simplified one (kdop18 tribox for now) for collision"));
+
+		// Num LODs for the simplification
+		int32 ChaosNumConvexLODs = 0;
+		FAutoConsoleVariableRef CVarChaosNumConvexLODs(TEXT("p.Chaos.Collision.NumConvexLODs"), ChaosNumConvexLODs, TEXT("Max number of collisions LODs required during simplication"));
+	}
+
+namespace Private
+{
+
+	FConvexOptimizer::FConvexOptimizer() :
+		SimplifiedConvexes(), CollisionObjects(MakeUnique<Private::FCollisionObjects>()), ShapesArray()
+	{}
+
+	FConvexOptimizer::~FConvexOptimizer() = default;
+
+	void FConvexOptimizer::VisitCollisionObjects(const FImplicitHierarchyVisitor& VisitorFunc) const
+	{
+		if(CVars::bChaosSimplifyConvexes)
+		{
+			int32 ObjectIndex = 0;
+			for(const Private::FImplicitBVHObject& CollisionObject : CollisionObjects->ImplicitObjects)
+			{
+				int32 LeafObjectIndex = CollisionObject.GetObjectIndex();
+				CollisionObject.GetGeometry()->VisitLeafObjectsImpl(CollisionObject.GetTransform(),
+					CollisionObject.GetRootObjectIndex(), ObjectIndex, LeafObjectIndex, VisitorFunc);
+			}
+		}	
+	}
+
+	void FConvexOptimizer::VisitOverlappingObjects(const FAABB3& LocalBounds, const FImplicitHierarchyVisitor& VisitorFunc) const
+	{
+		if(CVars::bChaosSimplifyConvexes)
+		{
+			int32 ObjectIndex = 0;
+			for(const Private::FImplicitBVHObject& CollisionObject : CollisionObjects->ImplicitObjects)
+			{
+				int32 LeafObjectIndex = CollisionObject.GetObjectIndex();
+				CollisionObject.GetGeometry()->VisitOverlappingLeafObjectsImpl(LocalBounds, CollisionObject.GetTransform(),
+					CollisionObject.GetRootObjectIndex(), ObjectIndex, LeafObjectIndex, VisitorFunc);
+			}
+		}
+	}
+
+	void FConvexOptimizer::BuildConvexShapes(const FShapesArray& UnionShapes)
+	{
+		ShapesArray.Reset();
+		TArray<FImplicitObjectPtr> ImplicitObjects = SimplifiedConvexes;
+		ShapesArray.Add(FShapeInstance::Make(UnionShapes.Num(),
+			MakeImplicitObjectPtr<FImplicitObjectUnion>(MoveTemp(ImplicitObjects))));
+
+		const TUniquePtr<FPerShapeData>& TemplateData = UnionShapes[0];
+		const TUniquePtr<FShapeInstance>& ShapeData = ShapesArray.Last();
+
+		ShapeData->SetQueryData(TemplateData->GetQueryData());
+		ShapeData->SetSimData(TemplateData->GetSimData());
+		ShapeData->SetCollisionTraceType(TemplateData->GetCollisionTraceType());
+
+		// Only sim enabled if all the underlying shapes could be used for physics
+		ShapeData->SetSimEnabled(true);
+
+		// Disable simple shapes for query 
+		ShapeData->SetQueryEnabled(false);
+	}
+
+	DECLARE_CYCLE_STAT(TEXT("FConvexOptimizer::SimplifyRootConvexes"), STAT_SimplifyRootConvexes, STATGROUP_ChaosCollision);
+	void FConvexOptimizer::SimplifyRootConvexes(const Chaos::FImplicitObjectUnionPtr& UnionGeometry, const FShapesArray& UnionShapes)
+	{
+		if(CVars::ChaosNumConvexLODs > 0)
+		{
+			BuildConvexLODs(UnionGeometry, UnionShapes);
+		}
+		else
+		{
+			BuildSingleConvex(UnionGeometry, UnionShapes);
+		}
+	}
+
+	void FConvexOptimizer::BuildSingleConvex(const Chaos::FImplicitObjectUnionPtr& UnionGeometry, const FShapesArray& UnionShapes)
+	{
+		if(CVars::bChaosSimplifyConvexes && UnionGeometry && (UnionShapes.Num() == UnionGeometry->GetObjects().Num()))
+		{
+			FTribox Tribox;
+
+			SimplifiedConvexes.Reset();
+			CollisionObjects->ImplicitObjects.Reset();
+	
+			UnionGeometry->VisitLeafObjects([&Tribox, this, &UnionShapes](
+				const FImplicitObject* ImplicitObject, const FRigidTransform3& RelativeTransform,
+				const int32 RootObjectIndex, const int32 ObjectIndex, const int32 LeafObjectIndex)
+			{
+				if(UnionShapes[RootObjectIndex]->GetSimEnabled())
+				{
+					const FTribox::FRigidTransform3Type ConvexTransform(RelativeTransform);
+					if(const FConvex* Convex = ImplicitObject->AsA<FConvex>())
+					{
+						// Disable collision for all the convexes that are going to be used to build the tribox
+						// For now only used for debug draw
+						const_cast<FImplicitObject*>(ImplicitObject)->SetDoCollide(false);
+						Tribox.AddConvex(Convex, ConvexTransform);
+					}
+					else
+					{
+						// Add non convex implicits to the list of collision objects
+						Private::FImplicitBVH::CollectLeafObject(ImplicitObject, RelativeTransform, RootObjectIndex, CollisionObjects->ImplicitObjects);	  
+					}
+				}
+			});
+			if(Tribox.BuildTribox())
+			{
+				// Build the tribox and add it to the list of optimized convexes
+				SimplifiedConvexes.Add(Tribox.MakeConvex());
+				Private::FImplicitBVH::CollectLeafObject(SimplifiedConvexes.Last(), FRigidTransform3::Identity, INDEX_NONE, CollisionObjects->ImplicitObjects);	  
+
+				// Build the simplified shapes
+				BuildConvexShapes(UnionShapes);
+			}
+		}
+	}
+
+	void FConvexOptimizer::BuildConvexLODs(const Chaos::FImplicitObjectUnionPtr& UnionGeometry, const FShapesArray& UnionShapes)
+	{
+		// WIP Another CL
+	}
+
+	void VisitCollisionObjects(const FConvexOptimizer* ConvexOptimizer, const FImplicitObject* ImplicitObject, const FImplicitHierarchyVisitor& VisitorFunc)
+	{
+		if(ConvexOptimizer && ConvexOptimizer->IsValid())
+		{
+			ConvexOptimizer->VisitCollisionObjects(VisitorFunc);
+		}
+		else
+		{
+			ImplicitObject->VisitLeafObjects(VisitorFunc);
+		}
+	}
+
+	void VisitOverlappingObjects(const FConvexOptimizer* ConvexOptimizer, const FImplicitObject* ImplicitObject, const FAABB3& LocalBounds, const FImplicitHierarchyVisitor& VisitorFunc)
+	{
+		if(ConvexOptimizer && ConvexOptimizer->IsValid())
+		{
+			ConvexOptimizer->VisitOverlappingObjects(LocalBounds, VisitorFunc);
+		}
+		else
+		{
+			ImplicitObject->VisitOverlappingLeafObjects(LocalBounds, VisitorFunc);
+		}
+	}
+
+}
+}
diff --git a/Engine/Source/Runtime/Experimental/Chaos/Private/Chaos/ImplicitObject.cpp b/Engine/Source/Runtime/Experimental/Chaos/Private/Chaos/ImplicitObject.cpp
--- a/Engine/Source/Runtime/Experimental/Chaos/Private/Chaos/ImplicitObject.cpp
+++ b/Engine/Source/Runtime/Experimental/Chaos/Private/Chaos/ImplicitObject.cpp
@@ -377,6 +377,7 @@ namespace Chaos
 		static const FName TriangleMeshName = TEXT("TriangleMesh");
 		static const FName HeightFieldName = TEXT("HeightField");
 		static const FName TaperedCapsuleName = TEXT("TaperedCapsule");
+		static const FName UnionClusteredName = TEXT("UnionClustered");
 
 		switch(GetInnerType(InType))
 		{
@@ -394,6 +395,7 @@ namespace Chaos
 		case ImplicitObjectType::TriangleMesh: return TriangleMeshName;
 		case ImplicitObjectType::HeightField: return HeightFieldName;
 		case ImplicitObjectType::TaperedCapsule: return TaperedCapsuleName;
+		case ImplicitObjectType::UnionClustered: return UnionClusteredName;
 		}
 		return NAME_None;
 	}
diff --git a/Engine/Source/Runtime/Experimental/Chaos/Private/Chaos/ImplicitObjectBVH.cpp b/Engine/Source/Runtime/Experimental/Chaos/Private/Chaos/ImplicitObjectBVH.cpp
--- a/Engine/Source/Runtime/Experimental/Chaos/Private/Chaos/ImplicitObjectBVH.cpp
+++ b/Engine/Source/Runtime/Experimental/Chaos/Private/Chaos/ImplicitObjectBVH.cpp
@@ -54,6 +54,23 @@ namespace Chaos
 			}
 			return NumObjects;
 		}
+		
+		void FImplicitBVH::CollectLeafObject(const FImplicitObject* Object, 
+			const FRigidTransform3& ParentTransform, const int32 RootObjectIndex, TArray<FImplicitBVHObject>& LeafObjects)
+		{
+			// @todo(chaos): clean this up (SetFromRawLowLevel). We know all the objects we visit are children of a UniquePtr because we own it
+			TSerializablePtr<FImplicitObject> SerializableObject;
+			SerializableObject.SetFromRawLowLevel(Object);
+
+			const int32 ObjectIndex = LeafObjects.Num();
+			LeafObjects.Emplace(
+				SerializableObject,
+				ParentTransform.GetTranslation(),
+				ParentTransform.GetRotation(),
+				Object->CalculateTransformedBounds(ParentTransform),
+				RootObjectIndex,
+				ObjectIndex);
+		}
 
 		FImplicitBVH::FObjects FImplicitBVH::CollectLeafObjects(const TArrayView<const Chaos::FImplicitObjectPtr>& InRootObjects)
 		{
@@ -64,24 +81,12 @@ namespace Chaos
 
 			for (int32 RootObjectIndex = 0; RootObjectIndex < InRootObjects.Num(); ++RootObjectIndex)
 			{
-				const Chaos::FImplicitObjectPtr& RootObject = InRootObjects[RootObjectIndex];
-
-				RootObject->VisitLeafObjects(
-					[RootObjectIndex, &Objects](const FImplicitObject* Object, const FRigidTransform3& ParentTransform, const int32 UnusedRootObjectIndex, const int32 UnusedObjectIndex, const int32 UnusedLeafObjectIndex)
-					{
-						// @todo(chaos): clean this up (SetFromRawLowLevel). We know all the objects we visit are children of a UniquePtr because we own it
-						TSerializablePtr<FImplicitObject> SerializableObject;
-						SerializableObject.SetFromRawLowLevel(Object);
-
-						const int32 ObjectIndex = Objects.Num();
-						Objects.Emplace(
-							SerializableObject,
-							ParentTransform.GetTranslation(),
-							ParentTransform.GetRotation(),
-							Object->CalculateTransformedBounds(ParentTransform),
-							RootObjectIndex,
-							ObjectIndex);
-					});
+				InRootObjects[RootObjectIndex]->VisitLeafObjects(
+				[RootObjectIndex, &Objects](const FImplicitObject* Object, const FRigidTransform3& ParentTransform,
+					const int32 UnusedRootObjectIndex, const int32 UnusedObjectIndex, const int32 UnusedLeafObjectIndex)
+				{
+					CollectLeafObject(Object, ParentTransform, RootObjectIndex, Objects);
+				});
 			}
 
 			return Objects;
@@ -94,10 +99,10 @@ namespace Chaos
 
 		TUniquePtr<FImplicitBVH> FImplicitBVH::TryMake(const TArrayView<const Chaos::FImplicitObjectPtr>& InRootObjects, const int32 InMinObjects, const int32 InMaxBVHDepth)
 		{
-			TArray<FImplicitBVHObject> Objects = CollectLeafObjects(InRootObjects);
-			if (Objects.Num() > InMinObjects)
+			TArray<FImplicitBVHObject> LeafObjects = CollectLeafObjects(InRootObjects);
+			if (LeafObjects.Num() > InMinObjects)
 			{
-				return TUniquePtr<FImplicitBVH>(new FImplicitBVH(MoveTemp(Objects), InMaxBVHDepth));
+				return TUniquePtr<FImplicitBVH>(new FImplicitBVH(MoveTemp(LeafObjects), InMaxBVHDepth));
 			}
 			return TUniquePtr<FImplicitBVH>();
 		}
diff --git a/Engine/Source/Runtime/Experimental/Chaos/Private/Chaos/ImplicitObjectUnion.cpp b/Engine/Source/Runtime/Experimental/Chaos/Private/Chaos/ImplicitObjectUnion.cpp
--- a/Engine/Source/Runtime/Experimental/Chaos/Private/Chaos/ImplicitObjectUnion.cpp
+++ b/Engine/Source/Runtime/Experimental/Chaos/Private/Chaos/ImplicitObjectUnion.cpp
@@ -1,7 +1,10 @@
 // Copyright Epic Games, Inc. All Rights Reserved.
 #include "Chaos/ImplicitObjectUnion.h"
+
 #include "Chaos/ImplicitObjectBVH.h"
 #include "Chaos/BoundingVolumeHierarchy.h"
+#include "Chaos/Convex.h"
+#include "Chaos/Tribox.h"
 
 #include "UObject/FortniteMainBranchObjectVersion.h"
 
@@ -28,12 +31,6 @@ namespace Chaos
 		// a random child. See FImplicitBVH::Partition.
 		FRealSingle ChaosUnionBVHSplitBias = 0.1f;
 		FAutoConsoleVariableRef CVarChaosUnionBVHSplitBias(TEXT("p.Chaos.Collision.UnionBVH.SplitBias"), ChaosUnionBVHSplitBias, TEXT(""));
-		
-		// Replace all the union objects with simplified shapes (box for now) for optimisation
-		bool bChaosUnionSimpleShapes = false;
-		FAutoConsoleVariableRef CVarChaosUnionSimpleShapes(TEXT("p.Chaos.Collision.Union.SimpleShapes"), bChaosUnionSimpleShapes, TEXT("If true replace all the union objects with simplified shapes (box for now) for optimisation"));
-
-		
 	}
 
 inline FAABB3 CalculateObjectsBounds(const TArrayView<FImplicitObjectPtr>& Objects)
@@ -63,7 +60,6 @@ FImplicitObjectUnion::FImplicitObjectUnion()
 FImplicitObjectUnion::FImplicitObjectUnion(TArray<TUniquePtr<FImplicitObject>>&& Objects)
 	: FImplicitObject(EImplicitObject::HasBoundingBox, ImplicitObjectType::Union)
 	, MObjects()
-	, Convexes()
 	, NumLeafObjects(0)
 	, Flags()
 {
@@ -73,7 +69,6 @@ FImplicitObjectUnion::FImplicitObjectUnion(TArray<TUniquePtr<FImplicitObject>>&&
 FImplicitObjectUnion::FImplicitObjectUnion(TArray<Chaos::FImplicitObjectPtr>&& Objects)
 	: FImplicitObject(EImplicitObject::HasBoundingBox, ImplicitObjectType::Union)
 	, MObjects(MoveTemp(Objects))
-	, Convexes()
 	, NumLeafObjects(0)
 	, Flags()
 {
@@ -81,16 +76,11 @@ FImplicitObjectUnion::FImplicitObjectUnion(TArray<Chaos::FImplicitObjectPtr>&& O
 
 	MLocalBoundingBox = CalculateObjectsBounds(MakeArrayView(MObjects));
 	SetNumLeafObjects(Private::FImplicitBVH::CountLeafObjects(MakeArrayView(GetObjects())));
-		
-	Flags.bAllowBVH = true;
-	RebuildConvexes();
-	Flags.bAllowBVH = false;
 }
 
 FImplicitObjectUnion::FImplicitObjectUnion(FImplicitObjectUnion&& Other)
 	: FImplicitObject(EImplicitObject::HasBoundingBox, ImplicitObjectType::Union)
 	, MObjects(MoveTemp(Other.MObjects))
-	, Convexes(MoveTemp(Other.Convexes))
 	, MLocalBoundingBox(MoveTemp(Other.MLocalBoundingBox))
 	, BVH(MoveTemp(Other.BVH))
 	, NumLeafObjects(Other.NumLeafObjects)
@@ -103,24 +93,7 @@ FImplicitObjectUnion::~FImplicitObjectUnion() = default;
 bool FImplicitObjectUnion::HasValidBVH() const
 {
 	// No need to create BVH if we are using the simple shape
-	return BVH.IsValid() && CVars::bChaosUnionBVHEnabled && !CVars::bChaosUnionSimpleShapes;
-}
-
-bool FImplicitObjectUnion::HasSimpleShapes() const
-{
-	// Simple shapes only if cvar is on and if we are on the root cluster union (AllowBVH)
-	return Flags.bAllowBVH && Chaos::CVars::bChaosUnionSimpleShapes;
-}
-
-void FImplicitObjectUnion::RebuildConvexes()
-{
-	if(HasSimpleShapes())
-	{
-		Convexes.Reset();
-
-		// Bounding box implicit for physics only
-		Convexes.Add(MakeImplicitObjectPtr<Chaos::FImplicitBox3>(MLocalBoundingBox.Min(), MLocalBoundingBox.Max(), MLocalBoundingBox.GetMargin()));
-	}
+	return BVH.IsValid() && CVars::bChaosUnionBVHEnabled;
 }
 	
 void FImplicitObjectUnion::Combine(TArray<Chaos::FImplicitObjectPtr>& OtherObjects)
@@ -143,7 +116,6 @@ void FImplicitObjectUnion::Combine(TArray<Chaos::FImplicitObjectPtr>& OtherObjec
 		MObjects.Add(MoveTemp(ChildObject));
 	}
 	RebuildBVH();
-	RebuildConvexes();
 }
 
 void FImplicitObjectUnion::RemoveAt(int32 RemoveIndex)
@@ -152,13 +124,11 @@ void FImplicitObjectUnion::RemoveAt(int32 RemoveIndex)
 	{
 		SetNumLeafObjects(GetNumLeafObjects() - Private::FImplicitBVH::CountLeafObjects(MakeArrayView(&MObjects[RemoveIndex], 1)));
 
-		MObjects[RemoveIndex]->Release();
 		MObjects.RemoveAt(RemoveIndex);
 	}
 
 	MLocalBoundingBox = CalculateObjectsBounds(MakeArrayView(MObjects));
 	RebuildBVH();
-	RebuildConvexes();
 }
 
 void FImplicitObjectUnion::SetNumLeafObjects(int32 InNumLeafObjects)
@@ -172,15 +142,16 @@ void FImplicitObjectUnion::SetNumLeafObjects(int32 InNumLeafObjects)
 
 void FImplicitObjectUnion::CreateBVH()
 {
-	if (Flags.bAllowBVH && CVars::bChaosUnionBVHEnabled && !CVars::bChaosUnionSimpleShapes)
+	if (Flags.bAllowBVH && CVars::bChaosUnionBVHEnabled)
 	{
 		const int32 MinBVHShapes = CVars::ChaosUnionBVHMinShapes;
 		const int32 MaxBVHDepth = CVars::ChaosUnionBVHMaxDepth;
+		
 		BVH = Private::FImplicitBVH::TryMake(MakeArrayView(MObjects), MinBVHShapes, MaxBVHDepth);
 		Flags.bHasBVH = BVH.IsValid();
 	}
 }
-
+	
 void FImplicitObjectUnion::DestroyBVH()
 {
 	if (BVH.IsValid())
@@ -206,7 +177,8 @@ void FImplicitObjectUnion::FindAllIntersectingObjects(TArray<Pair<const FImplici
 				if (BVH->NodeIsLeaf(NodeIndex))
 				{
 					BVH->VisitNodeObjects(NodeIndex,
-						[&Out](const FImplicitObject* Implicit, const FRigidTransform3f& RelativeTransformf, const FAABB3f& RelativeBoundsf, const int32 RootObjectIndex, const int32 LeafObjectIndex) -> void
+						[&Out](const FImplicitObject* Implicit, const FRigidTransform3f& RelativeTransformf,
+							const FAABB3f& RelativeBoundsf, const int32 RootObjectIndex, const int32 LeafObjectIndex) -> void
 						{
 							Out.Add(MakePair(Implicit, FRigidTransform3(RelativeTransformf)));
 						});
@@ -215,11 +187,10 @@ void FImplicitObjectUnion::FindAllIntersectingObjects(TArray<Pair<const FImplici
 	}
 	else
 	{
-		const Chaos::FImplicitObjectsArray& Objects = HasSimpleShapes() ? GetConvexes() : GetObjects();
-		for (const Chaos::FImplicitObjectPtr& Object : Objects)
-		{
-			Object->FindAllIntersectingObjects(Out, LocalBounds);
-		}
+		for (const Chaos::FImplicitObjectPtr& Object : MObjects)
+        {
+            Object->FindAllIntersectingObjects(Out, LocalBounds);
+        }
 	}
 }
 
@@ -241,35 +212,34 @@ void FImplicitObjectUnion::VisitOverlappingLeafObjectsImpl(
 	const FAABB3& LocalBounds,
 	const FRigidTransform3& ObjectTransform,
 	const int32 InRootObjectIndex,
-	int32& InOutObjectIndex,
-	int32& InOutLeafObjectIndex,
+	int32& ObjectIndex,
+	int32& LeafObjectIndex,
 	const FImplicitHierarchyVisitor& VisitorFunc) const
 {
 	// Skip self
-	InOutObjectIndex++;
+	ObjectIndex++;
 
 	if (HasValidBVH())
 	{
 		// Visit children
 		// NOTE: ObjectIndex passed to the visitor isn't really correct here. Maybe it should be removed...
 		BVH->VisitAllIntersections(LocalBounds,
-			[this, &ObjectTransform, &InOutObjectIndex, &VisitorFunc](const FImplicitObject* Implicit, const FRigidTransform3f& RelativeTransformf, const FAABB3f& RelativeBoundsf, const int32 RootObjectIndex, const int32 LeafObjectIndex)
+			[this, &ObjectTransform, &ObjectIndex, &VisitorFunc](const FImplicitObject* Implicit, const FRigidTransform3f& RelativeTransformf, const FAABB3f& RelativeBoundsf, const int32 RootObjectIndex, const int32 LeafObjectIndex)
 			{
-				VisitorFunc(Implicit, FRigidTransform3(RelativeTransformf) * ObjectTransform, RootObjectIndex, InOutObjectIndex++, LeafObjectIndex);
+				VisitorFunc(Implicit, FRigidTransform3(RelativeTransformf) * ObjectTransform, RootObjectIndex, ObjectIndex++, LeafObjectIndex);
 			});
 	}
 	else
 	{
-		const Chaos::FImplicitObjectsArray& Objects = HasSimpleShapes() ? GetConvexes() : GetObjects();
-		for (int32 BVHObjectIndex = 0; BVHObjectIndex < Objects.Num(); ++BVHObjectIndex)
+		for (int32 BVHObjectIndex = 0; BVHObjectIndex < MObjects.Num(); ++BVHObjectIndex)
 		{
 			// If we are the root our object index is the root index, otherwise just pass on the value we were given (from the actual root)
 			const int32 RootObjectIndex = (InRootObjectIndex != INDEX_NONE) ? InRootObjectIndex : BVHObjectIndex;
-			Objects[BVHObjectIndex]->VisitOverlappingLeafObjectsImpl(LocalBounds, ObjectTransform, RootObjectIndex, InOutObjectIndex, InOutLeafObjectIndex, VisitorFunc);
+			MObjects[BVHObjectIndex]->VisitOverlappingLeafObjectsImpl(LocalBounds, ObjectTransform, RootObjectIndex, ObjectIndex, LeafObjectIndex, VisitorFunc);
 		}
 	}
 }
-
+	
 void FImplicitObjectUnion::VisitLeafObjectsImpl(
 	const FRigidTransform3& ObjectTransform,
 	const int32 InRootObjectIndex,
@@ -280,12 +250,11 @@ void FImplicitObjectUnion::VisitLeafObjectsImpl(
 	// Skip self
 	++ObjectIndex;
 	
-	const Chaos::FImplicitObjectsArray& Objects = HasSimpleShapes() ? GetConvexes() : GetObjects();
-	for (int32 BVHObjectIndex = 0; BVHObjectIndex < Objects.Num(); ++BVHObjectIndex)
+	for (int32 BVHObjectIndex = 0; BVHObjectIndex < MObjects.Num(); ++BVHObjectIndex)
 	{
 		// If we are the root our object index is the root index, otherwise just pass on the value we were given (from the actual root)
 		const int32 RootObjectIndex = (InRootObjectIndex != INDEX_NONE) ? InRootObjectIndex : BVHObjectIndex;
-		Objects[BVHObjectIndex]->VisitLeafObjectsImpl(ObjectTransform, RootObjectIndex, ObjectIndex, LeafObjectIndex, VisitorFunc);
+		MObjects[BVHObjectIndex]->VisitLeafObjectsImpl(ObjectTransform, RootObjectIndex, ObjectIndex, LeafObjectIndex, VisitorFunc);
 	}
 }
 
@@ -301,12 +270,11 @@ bool FImplicitObjectUnion::VisitObjectsImpl(
 	++ObjectIndex;
 	
 	// Visit Children
-	const Chaos::FImplicitObjectsArray& Objects = HasSimpleShapes() ? GetConvexes() : GetObjects();
-	for (int32 BVHObjectIndex = 0; (BVHObjectIndex < Objects.Num()) && bContinue; ++BVHObjectIndex)
+	for (int32 BVHObjectIndex = 0; (BVHObjectIndex < MObjects.Num()) && bContinue; ++BVHObjectIndex)
 	{
 		// If we are the root our object index is the root index, otherwise just pass on the value we were given (from the actual root)
 		const int32 RootObjectIndex = (InRootObjectIndex != INDEX_NONE) ? InRootObjectIndex : BVHObjectIndex;
-		bContinue = Objects[BVHObjectIndex]->VisitObjectsImpl(ObjectTransform, RootObjectIndex, ObjectIndex, LeafObjectIndex, VisitorFunc);
+		bContinue = MObjects[BVHObjectIndex]->VisitObjectsImpl(ObjectTransform, RootObjectIndex, ObjectIndex, LeafObjectIndex, VisitorFunc);
 	}
 
 	return bContinue;
@@ -322,10 +290,9 @@ bool FImplicitObjectUnion::IsOverlappingBoundsImpl(const FAABB3& LocalBounds) co
 	{
 		if (LocalBounds.Intersects(BoundingBox()))
 		{
-			const Chaos::FImplicitObjectsArray& Objects = HasSimpleShapes() ? GetConvexes() : GetObjects();
-			for (int32 BVHObjectIndex = 0; BVHObjectIndex < Objects.Num(); ++BVHObjectIndex)
+			for (const Chaos::FImplicitObjectPtr& Object : MObjects)
 			{
-				if (Objects[BVHObjectIndex]->IsOverlappingBoundsImpl(LocalBounds))
+				if (Object->IsOverlappingBoundsImpl(LocalBounds))
 				{
 					return true;
 				}
@@ -399,8 +366,7 @@ void FImplicitObjectUnion::ForEachObject(TFunctionRef<bool(const FImplicitObject
 	}
 	else
 	{
-		const Chaos::FImplicitObjectsArray& Objects = HasSimpleShapes() ? GetConvexes() : GetObjects();
-		for (const Chaos::FImplicitObjectPtr& Object : Objects)
+		for (const Chaos::FImplicitObjectPtr& Object : MObjects)
 		{
 			if (Object)
 			{
@@ -465,12 +431,6 @@ void FImplicitObjectUnion::Serialize(FChaosArchive& Ar)
 		{
 			MLocalBoundingBox = FAABB3::ZeroAABB();
 		}
-		
-		if(Chaos::CVars::bChaosUnionSimpleShapes)
-		{
-			Convexes.Reset();
-			Convexes.Add(MakeImplicitObjectPtr<Chaos::FImplicitBox3>(MLocalBoundingBox.Min(), MLocalBoundingBox.Max(), MLocalBoundingBox.GetMargin()));
-		}
 	}
 }
 
@@ -499,7 +459,6 @@ void FImplicitObjectUnion::LegacySerializeBVH(FChaosArchive& Ar)
 	// so just revert to the original behaviour of every Union potentially having a BVH
 	Flags.bAllowBVH = true;
 	RebuildBVH();
-	RebuildConvexes();
 }
 
 FImplicitObjectUnionClustered::FImplicitObjectUnionClustered()
diff --git a/Engine/Source/Runtime/Experimental/Chaos/Private/Chaos/PBDRigidClustering.cpp b/Engine/Source/Runtime/Experimental/Chaos/Private/Chaos/PBDRigidClustering.cpp
--- a/Engine/Source/Runtime/Experimental/Chaos/Private/Chaos/PBDRigidClustering.cpp
+++ b/Engine/Source/Runtime/Experimental/Chaos/Private/Chaos/PBDRigidClustering.cpp
@@ -21,6 +21,7 @@
 #include "Chaos/PerParticleEtherDrag.h"
 #include "Chaos/PerParticlePBDEulerStep.h"
 #include "Chaos/StrainModification.h"
+#include "Chaos/ConvexOptimizer.h"
 #include "PhysicsProxy/GeometryCollectionPhysicsProxy.h"
 #include "PhysicsProxy/ClusterUnionPhysicsProxy.h"
 #include "CoreMinimal.h"
@@ -227,6 +228,10 @@ namespace Chaos
 		}
 		UpdateKinematicProperties(NewParticle, MChildren, MEvolution);
 		UpdateGeometry(NewParticle, ChildrenSet, MChildren, ProxyGeometry, Parameters);
+		
+		// Build the convex optimizer if required
+		BuildConvexOptimizer(NewParticle);
+		
 		GenerateConnectionGraph(NewParticle, Parameters);
 
 		NewParticle->SetSleeping(bClusterIsAsleep);
@@ -310,6 +315,13 @@ namespace Chaos
 		Cluster->SetInternalStrains(Cluster->GetInternalStrains() / static_cast<FRealSingle>(Children.Num()));
 	}
 
+	void FRigidClustering::BuildConvexOptimizer(FPBDRigidClusteredParticleHandle* Particle)
+	{
+		Particle->ConvexOptimizer() = MakePimpl<Private::FConvexOptimizer>();
+		Particle->ConvexOptimizer()->SimplifyRootConvexes( Particle->GetGeometry()->template AsA<FImplicitObjectUnion>(),
+					Particle->ShapesArray());
+	}
+
 	DECLARE_CYCLE_STAT(TEXT("TPBDRigidClustering<>::RemoveParticlesFromCluster"), STAT_RemoveParticlesFromCluster, STATGROUP_Chaos);
 	void
 	FRigidClustering::RemoveParticlesFromCluster(
@@ -497,6 +509,9 @@ namespace Chaos
 		UpdateKinematicProperties(NewParticle, MChildren, MEvolution);
 
 		UpdateGeometry(NewParticle, ChildrenSet, MChildren, FImplicitObjectPtr(nullptr), NoCleanParams);
+		
+		// Build the convex optimizer if required
+		BuildConvexOptimizer(NewParticle);
 
 		return NewParticle;
 	}
diff --git a/Engine/Source/Runtime/Experimental/Chaos/Private/Chaos/ShapeInstance.cpp b/Engine/Source/Runtime/Experimental/Chaos/Private/Chaos/ShapeInstance.cpp
--- a/Engine/Source/Runtime/Experimental/Chaos/Private/Chaos/ShapeInstance.cpp
+++ b/Engine/Source/Runtime/Experimental/Chaos/Private/Chaos/ShapeInstance.cpp
@@ -31,13 +31,12 @@ namespace Chaos
 			const int32 OldShapeNum = ShapesArray.Num();
 			if (const FImplicitObjectUnion* Union = Geometry->template GetObject<FImplicitObjectUnion>())
 			{
-				const int32 NumImplicits = Union->NumImplicits();
-				ShapesArray.Reserve(NumImplicits);
-				ShapesArray.SetNum(NumImplicits);
+				const int32 NumShapes = Union->GetObjects().Num();
+				ShapesArray.SetNum(NumShapes);
 				
-				for (int32 ShapeIndex = 0; ShapeIndex < ShapesArray.Num(); ++ShapeIndex)
+				for (int32 ShapeIndex = 0; ShapeIndex < NumShapes; ++ShapeIndex)
 				{
-					FImplicitObjectPtr ShapeGeometry = Union->GetImplicit(ShapeIndex);
+					FImplicitObjectPtr ShapeGeometry = Union->GetObjects()[ShapeIndex];
 					
 					if (ShapeIndex >= OldShapeNum)
 					{
diff --git a/Engine/Source/Runtime/Experimental/Chaos/Private/Chaos/Tribox.cpp b/Engine/Source/Runtime/Experimental/Chaos/Private/Chaos/Tribox.cpp
new file mode 100644
--- /dev/null
+++ b/Engine/Source/Runtime/Experimental/Chaos/Private/Chaos/Tribox.cpp
@@ -0,0 +1,496 @@
+// Copyright Epic Games, Inc. All Rights Reserved.
+
+#include "Chaos/Tribox.h"
+
+namespace Chaos
+{
+namespace Private
+{
+
+// Math utilities
+static constexpr FTribox::FRealType Sqrt2 = 1.414213562373095f;
+static constexpr FTribox::FRealType InvSqrt2 = 1.0f / Sqrt2;
+
+// k-DOP18 (k-Discrete Oriented Polytopes) Direction Vectors
+static const FTribox::FVec3Type PlanesDirs[18] = 
+{
+	FTribox::FVec3Type( 1.f, 0.f, 0.f),			// Face 0
+	FTribox::FVec3Type(-1.f, 0.f, 0.f),			// Face 1
+	FTribox::FVec3Type( 0.f, 1.f, 0.f),			// Face 2
+	FTribox::FVec3Type( 0.f,-1.f, 0.f),			// Face 3
+	FTribox::FVec3Type( 0.f, 0.f, 1.f),			// Face 4
+	FTribox::FVec3Type( 0.f, 0.f,-1.f),			// Face 5
+	FTribox::FVec3Type( 0.f, InvSqrt2,  InvSqrt2),	// Face 6
+	FTribox::FVec3Type( 0.f,-InvSqrt2, -InvSqrt2),	// Face 7
+	FTribox::FVec3Type( 0.f, InvSqrt2, -InvSqrt2),	// Face 8
+	FTribox::FVec3Type( 0.f,-InvSqrt2,  InvSqrt2),	// Face 9
+	FTribox::FVec3Type( InvSqrt2, 0.f,  InvSqrt2),	// Face 10
+	FTribox::FVec3Type(-InvSqrt2, 0.f, -InvSqrt2),	// Face 11
+	FTribox::FVec3Type( InvSqrt2, 0.f, -InvSqrt2),	// Face 12
+	FTribox::FVec3Type(-InvSqrt2, 0.f,  InvSqrt2),	// Face 13
+	FTribox::FVec3Type( InvSqrt2,  InvSqrt2, 0.f),	// Face 14
+	FTribox::FVec3Type(-InvSqrt2, -InvSqrt2, 0.f),	// Face 15
+	FTribox::FVec3Type( InvSqrt2, -InvSqrt2, 0.f),	// Face 16
+	FTribox::FVec3Type(-InvSqrt2,  InvSqrt2, 0.f)	// Face 17
+};
+
+// List of principal planes indices linked to each of the 8 box corner
+static const int32 PrincipalIndices[8][3] =
+{
+	{0,2,4},
+	{0,2,5},
+	{0,3,4},
+	{1,2,4},
+	{1,3,4},
+	{1,2,5},
+	{0,3,5},
+	{1,3,5}
+};
+
+// List of chamfer planes indices linked to each of the 8 box corner
+static const int32 ChamferIndices[8][3]=
+{
+	{6,10,14},
+	{8,12,14},
+	{9,10,16},
+	{6,13,17},
+	{9,13,15},
+	{8,11,17},
+	{7,12,16},
+	{7,11,15}
+};
+
+// Intersection matrix in between the 3 chamfer planes 
+static const FTribox::FMatrix33Type ChamferMatrices[8] =
+{
+	{-InvSqrt2, InvSqrt2, InvSqrt2, InvSqrt2, -InvSqrt2, InvSqrt2, InvSqrt2, InvSqrt2, -InvSqrt2},
+	{-InvSqrt2, InvSqrt2, InvSqrt2, InvSqrt2, -InvSqrt2, InvSqrt2, -InvSqrt2, -InvSqrt2, InvSqrt2}, 
+	{-InvSqrt2, InvSqrt2, InvSqrt2, -InvSqrt2, InvSqrt2, -InvSqrt2, InvSqrt2, InvSqrt2, -InvSqrt2}, 
+	{InvSqrt2, -InvSqrt2, -InvSqrt2, InvSqrt2, -InvSqrt2, InvSqrt2, InvSqrt2, InvSqrt2, -InvSqrt2}, 
+	{InvSqrt2, -InvSqrt2, -InvSqrt2, -InvSqrt2, InvSqrt2, -InvSqrt2, InvSqrt2, InvSqrt2, -InvSqrt2}, 
+	{InvSqrt2, -InvSqrt2, -InvSqrt2, InvSqrt2, -InvSqrt2, InvSqrt2, -InvSqrt2, -InvSqrt2, InvSqrt2}, 
+	{-InvSqrt2, InvSqrt2, InvSqrt2, -InvSqrt2, InvSqrt2, -InvSqrt2, -InvSqrt2, -InvSqrt2, InvSqrt2}, 
+	{InvSqrt2, -InvSqrt2, -InvSqrt2, -InvSqrt2, InvSqrt2, -InvSqrt2, -InvSqrt2, -InvSqrt2, InvSqrt2}
+};
+
+// Intersection matrix in between the 2 chamfer planes and a principal one
+// (9 combinations for the choice of the 2 chamfer + 1 principal)
+static const FTribox::FMatrix33Type PrincipalMatrices[8][9]=
+{
+	{
+		{1.000000, 0.000000, -0.000000, -1.000000, -0.000000, Sqrt2, -1.000000, Sqrt2, -0.000000},
+		{0.000000, 1.000000, 0.000000, -0.000000, -1.000000, Sqrt2, Sqrt2, 1.000000, -Sqrt2},
+		{0.000000, -0.000000, 1.000000, Sqrt2, -Sqrt2, 1.000000, 0.000000, Sqrt2, -1.000000},
+		{-1.000000, -0.000000, Sqrt2, 1.000000, 0.000000, -0.000000, 1.000000, Sqrt2, -Sqrt2},
+		{-0.000000, -1.000000, Sqrt2, 0.000000, 1.000000, 0.000000, Sqrt2, -1.000000, -0.000000},
+		{-Sqrt2, Sqrt2, 1.000000, -0.000000, 0.000000, 1.000000, Sqrt2, -0.000000, -1.000000},
+		{-1.000000, Sqrt2, 0.000000, 1.000000, -Sqrt2, Sqrt2, 1.000000, -0.000000, 0.000000},
+		{-Sqrt2, 1.000000, Sqrt2, Sqrt2, -1.000000, -0.000000, 0.000000, 1.000000, 0.000000},
+		{-0.000000, Sqrt2, -1.000000, Sqrt2, -0.000000, -1.000000, -0.000000, 0.000000, 1.000000}
+	},
+	{
+	{1.000000, -0.000000, 0.000000, -1.000000, 0.000000, Sqrt2, 1.000000, -Sqrt2, 0.000000},
+	{-0.000000, 1.000000, -0.000000, 0.000000, -1.000000, Sqrt2, -Sqrt2, -1.000000, Sqrt2},
+	{-0.000000, 0.000000, 1.000000, Sqrt2, -Sqrt2, 1.000000, -0.000000, -Sqrt2, 1.000000},
+	{-1.000000, 0.000000, Sqrt2, 1.000000, -0.000000, 0.000000, -1.000000, -Sqrt2, Sqrt2},
+	{0.000000, -1.000000, Sqrt2, -0.000000, 1.000000, -0.000000, -Sqrt2, 1.000000, 0.000000},
+	{-Sqrt2, Sqrt2, 1.000000, 0.000000, -0.000000, 1.000000, -Sqrt2, 0.000000, 1.000000},
+	{-1.000000, Sqrt2, -0.000000, 1.000000, -Sqrt2, Sqrt2, -1.000000, 0.000000, -0.000000},
+	{-Sqrt2, 1.000000, Sqrt2, Sqrt2, -1.000000, 0.000000, -0.000000, -1.000000, -0.000000},
+	{0.000000, Sqrt2, -1.000000, Sqrt2, 0.000000, -1.000000, 0.000000, -0.000000, -1.000000}
+	},
+	{
+	{1.000000, -0.000000, 0.000000, 1.000000, 0.000000, -Sqrt2, -1.000000, Sqrt2, 0.000000},
+	{-0.000000, 1.000000, -0.000000, 0.000000, 1.000000, -Sqrt2, Sqrt2, 1.000000, -Sqrt2},
+	{-0.000000, 0.000000, 1.000000, -Sqrt2, Sqrt2, -1.000000, -0.000000, Sqrt2, -1.000000},
+	{-1.000000, 0.000000, Sqrt2, -1.000000, -0.000000, 0.000000, 1.000000, Sqrt2, -Sqrt2},
+	{0.000000, -1.000000, Sqrt2, -0.000000, -1.000000, -0.000000, Sqrt2, -1.000000, 0.000000},
+	{-Sqrt2, Sqrt2, 1.000000, 0.000000, -0.000000, -1.000000, Sqrt2, 0.000000, -1.000000},
+	{-1.000000, Sqrt2, -0.000000, -1.000000, Sqrt2, -Sqrt2, 1.000000, 0.000000, -0.000000},
+	{-Sqrt2, 1.000000, Sqrt2, -Sqrt2, 1.000000, 0.000000, -0.000000, 1.000000, -0.000000},
+	{0.000000, Sqrt2, -1.000000, -Sqrt2, 0.000000, 1.000000, 0.000000, -0.000000, 1.000000}
+	},
+	{
+	{-1.000000, -0.000000, 0.000000, -1.000000, 0.000000, Sqrt2, -1.000000, Sqrt2, 0.000000},
+	{-0.000000, -1.000000, -0.000000, 0.000000, -1.000000, Sqrt2, Sqrt2, 1.000000, -Sqrt2},
+	{-0.000000, 0.000000, -1.000000, Sqrt2, -Sqrt2, 1.000000, -0.000000, Sqrt2, -1.000000},
+	{1.000000, 0.000000, -Sqrt2, 1.000000, -0.000000, 0.000000, 1.000000, Sqrt2, -Sqrt2},
+	{0.000000, 1.000000, -Sqrt2, -0.000000, 1.000000, -0.000000, Sqrt2, -1.000000, 0.000000},
+	{Sqrt2, -Sqrt2, -1.000000, 0.000000, -0.000000, 1.000000, Sqrt2, 0.000000, -1.000000},
+	{1.000000, -Sqrt2, -0.000000, 1.000000, -Sqrt2, Sqrt2, 1.000000, 0.000000, -0.000000},
+	{Sqrt2, -1.000000, -Sqrt2, Sqrt2, -1.000000, 0.000000, -0.000000, 1.000000, -0.000000},
+	{0.000000, -Sqrt2, 1.000000, Sqrt2, 0.000000, -1.000000, 0.000000, -0.000000, 1.000000}
+	},
+	{
+	{-1.000000, 0.000000, -0.000000, 1.000000, -0.000000, -Sqrt2, -1.000000, Sqrt2, -0.000000},
+	{0.000000, -1.000000, 0.000000, -0.000000, 1.000000, -Sqrt2, Sqrt2, 1.000000, -Sqrt2},
+	{0.000000, -0.000000, -1.000000, -Sqrt2, Sqrt2, -1.000000, 0.000000, Sqrt2, -1.000000},
+	{1.000000, -0.000000, -Sqrt2, -1.000000, 0.000000, -0.000000, 1.000000, Sqrt2, -Sqrt2},
+	{-0.000000, 1.000000, -Sqrt2, 0.000000, -1.000000, 0.000000, Sqrt2, -1.000000, -0.000000},
+	{Sqrt2, -Sqrt2, -1.000000, -0.000000, 0.000000, -1.000000, Sqrt2, -0.000000, -1.000000},
+	{1.000000, -Sqrt2, 0.000000, -1.000000, Sqrt2, -Sqrt2, 1.000000, -0.000000, 0.000000},
+	{Sqrt2, -1.000000, -Sqrt2, -Sqrt2, 1.000000, -0.000000, 0.000000, 1.000000, 0.000000},
+	{-0.000000, -Sqrt2, 1.000000, -Sqrt2, -0.000000, 1.000000, -0.000000, 0.000000, 1.000000}
+	},
+	{
+	{-1.000000, 0.000000, -0.000000, -1.000000, -0.000000, Sqrt2, 1.000000, -Sqrt2, -0.000000},
+	{0.000000, -1.000000, 0.000000, -0.000000, -1.000000, Sqrt2, -Sqrt2, -1.000000, Sqrt2},
+	{0.000000, -0.000000, -1.000000, Sqrt2, -Sqrt2, 1.000000, 0.000000, -Sqrt2, 1.000000},
+	{1.000000, -0.000000, -Sqrt2, 1.000000, 0.000000, -0.000000, -1.000000, -Sqrt2, Sqrt2},
+	{-0.000000, 1.000000, -Sqrt2, 0.000000, 1.000000, 0.000000, -Sqrt2, 1.000000, -0.000000},
+	{Sqrt2, -Sqrt2, -1.000000, -0.000000, 0.000000, 1.000000, -Sqrt2, -0.000000, 1.000000},
+	{1.000000, -Sqrt2, 0.000000, 1.000000, -Sqrt2, Sqrt2, -1.000000, -0.000000, 0.000000},
+	{Sqrt2, -1.000000, -Sqrt2, Sqrt2, -1.000000, -0.000000, 0.000000, -1.000000, 0.000000},
+	{-0.000000, -Sqrt2, 1.000000, Sqrt2, -0.000000, -1.000000, -0.000000, 0.000000, -1.000000}
+	},
+	{
+	{1.000000, 0.000000, -0.000000, 1.000000, -0.000000, -Sqrt2, 1.000000, -Sqrt2, -0.000000},
+	{0.000000, 1.000000, 0.000000, -0.000000, 1.000000, -Sqrt2, -Sqrt2, -1.000000, Sqrt2},
+	{0.000000, -0.000000, 1.000000, -Sqrt2, Sqrt2, -1.000000, 0.000000, -Sqrt2, 1.000000},
+	{-1.000000, -0.000000, Sqrt2, -1.000000, 0.000000, -0.000000, -1.000000, -Sqrt2, Sqrt2},
+	{-0.000000, -1.000000, Sqrt2, 0.000000, -1.000000, 0.000000, -Sqrt2, 1.000000, -0.000000},
+	{-Sqrt2, Sqrt2, 1.000000, -0.000000, 0.000000, -1.000000, -Sqrt2, -0.000000, 1.000000},
+	{-1.000000, Sqrt2, 0.000000, -1.000000, Sqrt2, -Sqrt2, -1.000000, -0.000000, 0.000000},
+	{-Sqrt2, 1.000000, Sqrt2, -Sqrt2, 1.000000, -0.000000, 0.000000, -1.000000, 0.000000},
+	{-0.000000, Sqrt2, -1.000000, -Sqrt2, -0.000000, 1.000000, -0.000000, 0.000000, -1.000000}
+	},
+	{
+	{-1.000000, -0.000000, 0.000000, 1.000000, 0.000000, -Sqrt2, 1.000000, -Sqrt2, 0.000000},
+	{-0.000000, -1.000000, -0.000000, 0.000000, 1.000000, -Sqrt2, -Sqrt2, -1.000000, Sqrt2},
+	{-0.000000, 0.000000, -1.000000, -Sqrt2, Sqrt2, -1.000000, -0.000000, -Sqrt2, 1.000000},
+	{1.000000, 0.000000, -Sqrt2, -1.000000, -0.000000, 0.000000, -1.000000, -Sqrt2, Sqrt2},
+	{0.000000, 1.000000, -Sqrt2, -0.000000, -1.000000, -0.000000, -Sqrt2, 1.000000, 0.000000},
+	{Sqrt2, -Sqrt2, -1.000000, 0.000000, -0.000000, -1.000000, -Sqrt2, 0.000000, 1.000000},
+	{1.000000, -Sqrt2, -0.000000, -1.000000, Sqrt2, -Sqrt2, -1.000000, 0.000000, -0.000000},
+	{Sqrt2, -1.000000, -Sqrt2, -Sqrt2, 1.000000, 0.000000, -0.000000, -1.000000, -0.000000},
+	{0.000000, -Sqrt2, 1.000000, -Sqrt2, 0.000000, 1.000000, 0.000000, -0.000000, -1.000000}
+	}
+};
+
+// For each corner : order of the intersection point within the 3 chamfer face vertex indices
+static const int32 ChamferOrder[8][3] =
+{
+	{1, 1, 1},
+	{1, 1, 5},
+	{1, 5, 1},
+	{5, 1, 1},
+	{5, 5, 1},
+	{5, 1, 5},
+	{1, 5, 5},
+	{5, 5, 5}
+};
+
+// For each corner : order of the intersection point within the 2 chamfer + principal face vertex indices
+// (9 combinations for the choice of the 2 chamfer + 1 principal)
+static const int32 PrincipalOrder[8][9][3] =
+{
+	{
+		{0, 3, 0},
+		{1, 1, 1},
+		{2, 2, 0},
+		{0, 2, 2},
+		{0, 0, 3},
+		{1, 1, 1},
+		{1, 1, 1},
+		{2, 0, 2},
+		{3, 0, 0}
+	},
+	{
+	{2, 0, 7},
+	{2, 2, 6},
+	{1, 1, 3},
+	{7, 1, 5},
+	{3, 6, 4},
+	{2, 2, 6},
+	{0, 2, 6},
+	{1, 1, 5},
+	{0, 3, 0}
+	},
+	{
+	{6, 4, 3},
+	{2, 6, 2},
+	{1, 5, 7},
+	{1, 5, 1},
+	{3, 0, 0},
+	{2, 6, 0},
+	{2, 6, 2},
+	{1, 3, 1},
+	{0, 7, 2}
+	},
+	{
+	{0, 0, 3},
+	{6, 0, 2},
+	{5, 1, 1},
+	{3, 1, 1},
+	{7, 2, 0},
+	{6, 2, 2},
+	{6, 2, 2},
+	{5, 7, 1},
+	{4, 3, 6}
+	},
+	{
+	{2, 7, 0},
+	{5, 3, 1},
+	{6, 6, 2},
+	{6, 6, 2},
+	{4, 6, 3},
+	{5, 5, 7},
+	{5, 5, 1},
+	{6, 4, 2},
+	{7, 4, 4}
+	},
+	{
+	{6, 3, 4},
+	{5, 7, 5},
+	{6, 2, 6},
+	{4, 2, 6},
+	{4, 4, 7},
+	{5, 1, 5},
+	{3, 1, 5},
+	{6, 2, 6},
+	{7, 0, 2}
+	},
+	{
+	{4, 7, 4},
+	{1, 5, 5},
+	{2, 6, 4},
+	{2, 6, 6},
+	{0, 2, 7},
+	{1, 5, 3},
+	{7, 5, 5},
+	{2, 6, 6},
+	{3, 4, 6}
+	},
+	{
+	{4, 4, 7},
+	{6, 4, 6},
+	{5, 5, 5},
+	{5, 5, 5},
+	{7, 4, 4},
+	{6, 6, 4},
+	{4, 6, 6},
+	{5, 5, 5},
+	{4, 7, 4}
+	}
+};
+
+void FTribox::BuildPrincipalDist(const FVec3Type& P, const int32 CoordIndexA, const int32 DistsIndexA, const int32 DistsIndexB)
+{
+	MaxDists[DistsIndexA] = FMath::Max(P[CoordIndexA], MaxDists[DistsIndexA]);
+	MaxDists[DistsIndexB] = FMath::Max(-P[CoordIndexA], MaxDists[DistsIndexB]);
+}
+
+void FTribox::BuildChamferDist(const FVec3Type& P, const int32 CoordIndexA, const int32 CoordIndexB,
+		const int32 DistsIndexA, const int32 DistsIndexB, const int32 DistsIndexC, const int32 DistsIndexD)
+{
+	MaxDists[DistsIndexA] = FMath::Max(P[CoordIndexA]+P[CoordIndexB], MaxDists[DistsIndexA]);
+	MaxDists[DistsIndexB] = FMath::Max(-P[CoordIndexA]-P[CoordIndexB], MaxDists[DistsIndexB]);
+	MaxDists[DistsIndexC] = FMath::Max(P[CoordIndexA]-P[CoordIndexB], MaxDists[DistsIndexC]);
+	MaxDists[DistsIndexD] = FMath::Max(-P[CoordIndexA]+P[CoordIndexB], MaxDists[DistsIndexD]);
+}
+
+void FTribox::AddPoint(const FVec3Type& PointPosition)
+{
+	const FVec3 P = PointPosition;
+
+	// Build the max distance along each principal axis
+	BuildPrincipalDist(P, 0,0,1);
+	BuildPrincipalDist(P, 1,2,3);
+	BuildPrincipalDist(P, 2,4,5);
+
+	// Build the max distance along each chamfer axis
+	BuildChamferDist(P, 1, 2, 6, 7, 8, 9);
+	BuildChamferDist(P, 0, 2, 10, 11, 12, 13);
+	BuildChamferDist(P, 0, 1, 14, 15, 16, 17);
+
+	++NumPoints;
+}
+
+void FTribox::AddConvex(const FConvex* Convex, const FRigidTransform3Type& RelativeTransform)
+{
+	if(Convex)
+	{
+		for(int32 VertexIndex = 0; VertexIndex < Convex->NumVertices(); ++VertexIndex)
+		{
+			AddPoint(RelativeTransform.TransformPosition(FVec3Type(
+				Convex->GetVertex(VertexIndex)[0],
+				Convex->GetVertex(VertexIndex)[1],
+				Convex->GetVertex(VertexIndex)[2])));
+		}
+	}
+}
+
+bool FTribox::BuildTribox()
+{
+	if(NumPoints > 0)
+	{
+		// Inflate the principal planes max distance by a given distance
+		for(int32 PlaneIndex = 0; PlaneIndex < NumPrincipalPlanes; ++PlaneIndex)
+		{
+			MaxDists[PlaneIndex] += InflateDistance;
+		}
+		// Scale the chamfer planes max distance  by 1 / sqrt(2) + Inflate by a given distance
+		for(int32 PlaneIndex = NumPrincipalPlanes; PlaneIndex < NumPlanes; ++PlaneIndex)
+		{
+			MaxDists[PlaneIndex] = MaxDists[PlaneIndex] * InvSqrt2 + InflateDistance;
+		}
+		return true;
+	}
+	return false;
+}
+
+FVec3 FTribox::SolveIntersection(const int32 FaceIndex[3], const FMatrix33Type& A) const
+{
+	const FVec3 B(MaxDists[FaceIndex[0]],MaxDists[FaceIndex[1]],MaxDists[FaceIndex[2]]);
+	
+	return  FVec3(
+	A.GetAt(0,0) * B[0] + A.GetAt(0,1) * B[1] + A.GetAt(0,2) * B[2],
+	A.GetAt(1,0) * B[0] + A.GetAt(1,1) * B[1] + A.GetAt(1,2) * B[2],
+	A.GetAt(2,0) * B[0] + A.GetAt(2,1) * B[1] + A.GetAt(2,2) * B[2]);
+}
+
+void FTribox::AddIntersection(const int32 FaceIndex[3], const int32 FaceOrder[3], const FVec3Type& VertexPosition,
+		TArray<TArray<int32>>& FaceIndices, TArray<FConvex::FVec3Type>& ConvexVertices) const
+{
+	const int32 VertexOffset = ConvexVertices.Num(); 
+
+	FaceIndices[FaceIndex[0]][FaceOrder[0]] = VertexOffset;
+	FaceIndices[FaceIndex[1]][FaceOrder[1]] = VertexOffset;
+	FaceIndices[FaceIndex[2]][FaceOrder[2]] = VertexOffset;
+
+	ConvexVertices.Add(VertexPosition);
+}
+
+void FTribox::ComputeIntersection(const int32 CornerIndex, const int32 PrincipalIndex, const int32 ChamferIndex,
+	TArray<TArray<int32>>& FaceIndices, TArray<FConvex::FVec3Type>& ConvexVertices) const
+{
+	// Compute the intersection in between 2 corner chamfer planes and one principal one
+	int32 TId[3] = {ChamferIndices[CornerIndex][0], ChamferIndices[CornerIndex][1], ChamferIndices[CornerIndex][2]};
+	TId[ChamferIndex] = PrincipalIndices[CornerIndex][PrincipalIndex];
+
+	const int32 LinearIndex = PrincipalIndex*3+ChamferIndex;
+
+	const FVec3Type V = SolveIntersection(TId, PrincipalMatrices[CornerIndex][LinearIndex]);
+	AddIntersection(TId, PrincipalOrder[CornerIndex][LinearIndex], V, FaceIndices, ConvexVertices);
+}
+
+FORCEINLINE void CompressFaces(TArray<TArray<int32>>& FaceIndices)
+{
+	TArray<int32> CompressedIndices;
+	CompressedIndices.Reserve(8);
+	for(int32 PlaneIndex = 0; PlaneIndex < FaceIndices.Num(); ++PlaneIndex)
+	{
+		CompressedIndices.Reset();
+		for(int32 LocalIndex = FaceIndices[PlaneIndex].Num()-1; LocalIndex >= 0; --LocalIndex)
+		{
+			int32 VertexIndex = FaceIndices[PlaneIndex][LocalIndex];
+			if(VertexIndex != INDEX_NONE)
+			{
+				CompressedIndices.Add(VertexIndex);
+			}
+		}
+		FaceIndices[PlaneIndex] = CompressedIndices;
+	}
+}
+
+FImplicitObjectPtr FTribox::MakeConvex() const
+{
+	TArray<TArray<int32>> FaceIndices;
+	TArray<FConvex::FVec3Type> ConvexVertices;
+	TArray<FConvex::FPlaneType> ConvexPlanes;
+
+	FaceIndices.SetNum(NumPlanes);
+	ConvexVertices.Reserve(32);
+	ConvexPlanes.SetNum(NumPlanes);
+
+	for(int32 PlaneIndex = 0; PlaneIndex < NumPlanes; ++PlaneIndex)
+	{
+		ConvexPlanes[PlaneIndex] = FConvex::FPlaneType(
+			PlanesDirs[PlaneIndex] * MaxDists[PlaneIndex], PlanesDirs[PlaneIndex]);
+		FaceIndices[PlaneIndex].Init(INDEX_NONE, 8);
+	}
+	for(int32 CornerIndex = 0; CornerIndex < 8; ++CornerIndex)
+	{
+		// Compute the intersection in between 3 corner chamfer planes
+		const FVec3Type P = SolveIntersection(ChamferIndices[CornerIndex], ChamferMatrices[CornerIndex]);
+
+		if(P.Dot(PlanesDirs[PrincipalIndices[CornerIndex][0]]) > MaxDists[PrincipalIndices[CornerIndex][0]])
+		{
+			ComputeIntersection(CornerIndex, 0,1, FaceIndices, ConvexVertices);
+			ComputeIntersection(CornerIndex, 0,2, FaceIndices, ConvexVertices);
+			ComputeIntersection(CornerIndex, 1,1, FaceIndices, ConvexVertices);
+			ComputeIntersection(CornerIndex, 2,2, FaceIndices, ConvexVertices);
+		}
+		else if(P.Dot(PlanesDirs[PrincipalIndices[CornerIndex][1]]) > MaxDists[PrincipalIndices[CornerIndex][1]])
+		{
+			ComputeIntersection(CornerIndex, 1,0, FaceIndices, ConvexVertices);
+			ComputeIntersection(CornerIndex, 1,2, FaceIndices, ConvexVertices);
+			ComputeIntersection(CornerIndex, 0,0, FaceIndices, ConvexVertices);
+			ComputeIntersection(CornerIndex, 2,2, FaceIndices, ConvexVertices);
+		}
+		else if(P.Dot(PlanesDirs[PrincipalIndices[CornerIndex][2]]) > MaxDists[PrincipalIndices[CornerIndex][2]])
+		{
+			ComputeIntersection(CornerIndex, 2,0, FaceIndices, ConvexVertices);
+			ComputeIntersection(CornerIndex, 2,1, FaceIndices, ConvexVertices);
+			ComputeIntersection(CornerIndex, 0,0, FaceIndices, ConvexVertices);
+			ComputeIntersection(CornerIndex, 1,1, FaceIndices, ConvexVertices);
+		}
+		else
+		{
+			ComputeIntersection(CornerIndex, 0,0, FaceIndices, ConvexVertices);
+			ComputeIntersection(CornerIndex, 1,1, FaceIndices, ConvexVertices);
+			ComputeIntersection(CornerIndex, 2,2, FaceIndices, ConvexVertices);
+
+			AddIntersection(ChamferIndices[CornerIndex], ChamferOrder[CornerIndex], P, FaceIndices, ConvexVertices);
+		}
+	}
+	CompressFaces(FaceIndices);
+	return MakeImplicitObjectPtr<Chaos::FConvex>(MoveTemp(ConvexPlanes), MoveTemp(FaceIndices), MoveTemp(ConvexVertices));
+}
+	
+bool FTribox::OverlapTribox(const FTribox& OtherTribox, FTribox& OverlapTribox) const
+{
+	bool bOverlapTriboxes = true;
+	for(int32 AxisIndex = 0; AxisIndex < NumPlanes; AxisIndex += 2)
+	{
+		if((MaxDists[AxisIndex] < -OtherTribox.MaxDists[AxisIndex+1]) ||
+			(OtherTribox.MaxDists[AxisIndex] < -MaxDists[AxisIndex+1]))
+		{
+			bOverlapTriboxes = false;
+			break;
+		}
+	}
+	if(bOverlapTriboxes)
+	{
+		for(int32 AxisIndex = 0; AxisIndex < NumPlanes; AxisIndex += 2)
+		{
+			OverlapTribox.MaxDists[AxisIndex] = FMath::Min(MaxDists[AxisIndex], OtherTribox.MaxDists[AxisIndex]);
+			OverlapTribox.MaxDists[AxisIndex+1] = FMath::Min(MaxDists[AxisIndex+1], OtherTribox.MaxDists[AxisIndex+1]);
+		}
+	}
+	return bOverlapTriboxes;
+}
+
+FTribox::FRealType FTribox::ComputeVolume() const
+{
+	return (MaxDists[0]+MaxDists[1]) * (MaxDists[2]+MaxDists[3]) * (MaxDists[4]+MaxDists[5]);
+}
+	
+FTribox& FTribox::operator+=(const FTribox& OtherTribox)
+{
+	for(int32 PlaneIndex = 0; PlaneIndex < NumPlanes; ++PlaneIndex)
+	{
+		MaxDists[PlaneIndex]  = FMath::Max(OtherTribox.MaxDists[PlaneIndex], MaxDists[PlaneIndex]);
+	}
+	return *this;
+}
+
+FTribox FTribox::operator+(const FTribox& OtherTribox) const
+{
+	FTribox ReturnTribox(*this);
+	return ReturnTribox += OtherTribox;
+}
+	
+}
+}
+	
\ No newline at end of file
diff --git a/Engine/Source/Runtime/Experimental/Chaos/Public/Chaos/ClusterUnionManager.h b/Engine/Source/Runtime/Experimental/Chaos/Public/Chaos/ClusterUnionManager.h
--- a/Engine/Source/Runtime/Experimental/Chaos/Public/Chaos/ClusterUnionManager.h
+++ b/Engine/Source/Runtime/Experimental/Chaos/Public/Chaos/ClusterUnionManager.h
@@ -318,7 +318,6 @@ namespace Chaos
 		}
 
 		FImplicitObjectUnion* ImplicitUnion = ClusterParticle->GetGeometry() ? ClusterParticle->GetGeometry()->template AsA<FImplicitObjectUnion>() : nullptr;
-		const int32 OldNumSimpleShapes = ImplicitUnion ? ImplicitUnion->GetConvexes().Num() : 0;
 		const int32 OldNumChildShapes = ImplicitUnion ? ImplicitUnion->GetNumRootObjects() : 0;
 
 		Func();
@@ -329,11 +328,10 @@ namespace Chaos
 		ImplicitUnion = ClusterParticle->GetGeometry()->template AsA<FImplicitObjectUnion>();
 		check(ImplicitUnion != nullptr);
 
-		const int32 NewNumSimpleShapes = ImplicitUnion->GetConvexes().Num();
 		const int32 NewNumChildShapes = ImplicitUnion->GetNumRootObjects();
 		const FShapesArray& ShapeArray = ClusterParticle->ShapesArray();
-		check(ShapeArray.Num() == (NewNumChildShapes + NewNumSimpleShapes));
-		check(Particles.Num() == (NewNumChildShapes - OldNumChildShapes));
+		check(ShapeArray.Num() == NewNumChildShapes);
+		check(Particles.Num() == NewNumChildShapes - OldNumChildShapes);
 	
 		for (int32 Index = 0; Index < Particles.Num(); ++Index)
 		{
@@ -349,27 +347,10 @@ namespace Chaos
 				continue;
 			}
 
-			const int32 ShapeIndex = NewNumSimpleShapes + OldNumChildShapes + Index;
+			const int32 ShapeIndex =  OldNumChildShapes + Index;
 			if (const TUniquePtr<Chaos::FPerShapeData>& ShapeData = ShapeArray[ShapeIndex])
 			{
 				TransferClusterUnionShapeData(ShapeData, Particle, TemplateShape, ActorId, ComponentId);
-
-				if (NewNumSimpleShapes > 0)
-				{
-					ShapeData->SetSimEnabled(false);
-				}
-			}
-		}
-
-		for (int32 Index = OldNumSimpleShapes; Index < NewNumSimpleShapes; ++Index)
-		{
-			const TUniquePtr<Chaos::FPerShapeData>& ShapeData = ShapeArray[Index];
-			const TUniquePtr<Chaos::FPerShapeData>& TemplateShape = Particles[0]->ShapesArray()[0];
-
-			if (ShapeData && TemplateShape)
-			{
-				TransferClusterUnionShapeData<TParticle>(ShapeData, nullptr, TemplateShape, ActorId, ComponentId);
-				ShapeData->SetQueryEnabled(false);
 			}
 		}
 	}
@@ -390,8 +371,7 @@ namespace Chaos
 		check(Geometry != nullptr);
 
 		const Chaos::FImplicitObjectUnion& GeometryUnion = Geometry->GetObjectChecked<Chaos::FImplicitObjectUnion>();
-		const int32 NumSimpleShapes = GeometryUnion.GetConvexes().Num();
-		check(AllChildParticles.Num() == (ShapesArray.Num() - NumSimpleShapes));
+		check(AllChildParticles.Num() == ShapesArray.Num());
 		
 		TArray<int32> ShapeIndicesToRemove;
 		ShapeIndicesToRemove.Reserve(ShapeParticles.Num());
@@ -401,7 +381,7 @@ namespace Chaos
 			const int32 Index = AllChildParticles.Find(ShapeParticle);
 			if (Index != INDEX_NONE)
 			{
-				ShapeIndicesToRemove.Add(Index + NumSimpleShapes);
+				ShapeIndicesToRemove.Add(Index);
 			}
 		}
 
@@ -430,10 +410,10 @@ namespace Chaos
 			}
 
 			ClusterParticle->RemoveShapeAtIndex(ToRemoveIndex);
-			AllChildParticles.RemoveAt(ToRemoveIndex - NumSimpleShapes);
+			AllChildParticles.RemoveAt(ToRemoveIndex);
 			LastIndex = ToRemoveIndex;
 		}
 
-		check(AllChildParticles.Num() == (ShapesArray.Num() - NumSimpleShapes));
+		check(AllChildParticles.Num() == ShapesArray.Num());
 	}
 }
diff --git a/Engine/Source/Runtime/Experimental/Chaos/Public/Chaos/Collision/ParticlePairMidPhase.h b/Engine/Source/Runtime/Experimental/Chaos/Public/Chaos/Collision/ParticlePairMidPhase.h
--- a/Engine/Source/Runtime/Experimental/Chaos/Public/Chaos/Collision/ParticlePairMidPhase.h
+++ b/Engine/Source/Runtime/Experimental/Chaos/Public/Chaos/Collision/ParticlePairMidPhase.h
@@ -489,13 +489,13 @@ namespace Chaos
 		// BVH on ParticleA versus implicit hierarchy of ParticleB
 		CHAOS_API void GenerateCollisionsBVHImplicitHierarchy(
 			FGeometryParticleHandle* ParticleA, const Private::FImplicitBVH* BVHA,
-			FGeometryParticleHandle* ParticleB, const FImplicitObject* RootImplicitB,
+			FGeometryParticleHandle* ParticleB, const FImplicitObject* RootImplicitB, const Private::FConvexOptimizer* ConvexOptimizerB,
 			const FReal CullDistance, const FReal Dt, const FCollisionContext& Context);
 
 		// Implicit hierarchy of particle A versus implicit hierarchy of ParticleB (used when no BVH present on either)
 		CHAOS_API void GenerateCollisionsImplicitHierarchyImplicitHierarchy(
-			FGeometryParticleHandle* ParticleA, const FImplicitObject* RootImplicitA,
-			FGeometryParticleHandle* ParticleB, const FImplicitObject* RootImplicitB,
+			FGeometryParticleHandle* ParticleA, const FImplicitObject* RootImplicitA, const Private::FConvexOptimizer* ConvexOptimizerA,
+			FGeometryParticleHandle* ParticleB, const FImplicitObject* RootImplicitB, const Private::FConvexOptimizer* ConvexOptimizerB,
 			const FReal CullDistance, const FReal Dt, const FCollisionContext& Context);
 
 		// BVH on particle A versus a Leaf Implicit of ParticleB
diff --git a/Engine/Source/Runtime/Experimental/Chaos/Public/Chaos/ConvexOptimizer.h b/Engine/Source/Runtime/Experimental/Chaos/Public/Chaos/ConvexOptimizer.h
new file mode 100644
--- /dev/null
+++ b/Engine/Source/Runtime/Experimental/Chaos/Public/Chaos/ConvexOptimizer.h
@@ -0,0 +1,75 @@
+// Copyright Epic Games, Inc. All Rights Reserved.
+
+#pragma once
+
+#include "Chaos/Core.h"
+#include "Chaos/ImplicitObject.h"
+
+namespace Chaos
+{
+	
+namespace Private
+{
+	class FImplicitBVHObject;
+	struct FCollisionObjects
+	{
+		TArray<Private::FImplicitBVHObject> ImplicitObjects;
+	};
+
+	/**
+	 * @brief The convex optimizer goal is to have a central place where
+	 * implicits hierarchy could be modified in order to accelerate collision detection
+	 */
+	class FConvexOptimizer
+	{
+		public :
+		CHAOS_API FConvexOptimizer();
+
+		// Default destructor
+		CHAOS_API ~FConvexOptimizer();
+
+		// Simplify all the convexes in the hierarchy 
+		CHAOS_API void SimplifyRootConvexes(const Chaos::FImplicitObjectUnionPtr& UnionGeometry, const FShapesArray& UnionShapes);
+
+		// Check if the manager is valid or not
+		CHAOS_API bool IsValid() const {return !SimplifiedConvexes.IsEmpty();}
+
+		// Visit all the collision objects if they exist / otherwise forward it to the RootHierarchy
+		CHAOS_API void VisitCollisionObjects(const FImplicitHierarchyVisitor& VisitorFunc) const;
+
+		// Visit all the overlapping objects if they exist / otherwise forward it to the RootHierarchy
+		CHAOS_API void VisitOverlappingObjects(const FAABB3& LocalBounds, const FImplicitHierarchyVisitor& VisitorFunc) const;
+
+		// Get the shapes array 
+		CHAOS_API const FShapeInstanceArray& GetShapeInstances() const { return ShapesArray;}
+
+	private:
+
+		// Build a single convex
+		void BuildSingleConvex(const Chaos::FImplicitObjectUnionPtr& UnionGeometry, const FShapesArray& UnionShapes);
+
+		// Build several convexes for a given LOD
+		void BuildConvexLODs(const Chaos::FImplicitObjectUnionPtr& UnionGeometry, const FShapesArray& UnionShapes);
+
+		// Build the simplified shapes
+		void BuildConvexShapes(const FShapesArray& UnionShapes);
+
+		// List of simplified convexes
+		TArray<FImplicitObjectPtr> SimplifiedConvexes;
+
+		// List of all the collision objects to avoid traversing all the hierarchy during midphase
+		TUniquePtr<Private::FCollisionObjects> CollisionObjects;
+
+		// Additional shapes array that could be used during collision midphase
+		FShapeInstanceArray ShapesArray;
+	};
+
+	// Visit all the collision objects if they exist / otherwise forward it to the RootHierarchy
+	void VisitCollisionObjects(const FConvexOptimizer* ConvexOptimizer, const FImplicitObject* ImplicitObject, const FImplicitHierarchyVisitor& VisitorFunc);
+
+	// Visit all the overlapping objects if they exist / otherwise forward it to the RootHierarchy
+	void VisitOverlappingObjects(const FConvexOptimizer* ConvexOptimizer, const FImplicitObject* ImplicitObject, const FAABB3& LocalBounds, const FImplicitHierarchyVisitor& VisitorFunc);
+
+
+}
+}
\ No newline at end of file
diff --git a/Engine/Source/Runtime/Experimental/Chaos/Public/Chaos/GeometryParticles.h b/Engine/Source/Runtime/Experimental/Chaos/Public/Chaos/GeometryParticles.h
--- a/Engine/Source/Runtime/Experimental/Chaos/Public/Chaos/GeometryParticles.h
+++ b/Engine/Source/Runtime/Experimental/Chaos/Public/Chaos/GeometryParticles.h
@@ -362,7 +362,7 @@ namespace Chaos
 					if (Shape->GetSimEnabled() && (CVars::bCCDAxisThresholdUsesProbeShapes || !Shape->GetIsProbe()))
 					{
 						const FImplicitObjectRef Geometry = Shape->GetGeometry();
-						if (Geometry->HasBoundingBox())
+						if (Geometry && Geometry->HasBoundingBox())
 						{
 							const TVector<T, d> ShapeExtents = Geometry->BoundingBox().Extents();
 							TVector<T, d>& CCDAxisThreshold = MCCDAxisThreshold[Index];
diff --git a/Engine/Source/Runtime/Experimental/Chaos/Public/Chaos/ImplicitObjectBVH.h b/Engine/Source/Runtime/Experimental/Chaos/Public/Chaos/ImplicitObjectBVH.h
--- a/Engine/Source/Runtime/Experimental/Chaos/Public/Chaos/ImplicitObjectBVH.h
+++ b/Engine/Source/Runtime/Experimental/Chaos/Public/Chaos/ImplicitObjectBVH.h
@@ -112,6 +112,8 @@ namespace Chaos
 			// Utility for processing the hierarchy
 			static int32 CountLeafObjects(const TArrayView<const Chaos::FImplicitObjectPtr>& InRootObjects);
 			static FObjects CollectLeafObjects(const TArrayView<const Chaos::FImplicitObjectPtr>& InRootObjects);
+			static void CollectLeafObject(const FImplicitObject* Object, const FRigidTransform3& ParentTransform, const int32 RootObjectIndex,
+				TArray<FImplicitBVHObject>& LeafObjects);
 
 			// Create a BVH around a set of ImplicitObjects. Usually these are the immediate child elements of an FImplcitObjectUnion
 			// TryMake will then recurse into the geometry hierachy and add all descendents to the BVH. Will return null if the 
@@ -122,7 +124,7 @@ namespace Chaos
 
 			int32 GetNumObjects() const { return Objects.Num(); }
 			int32 GetDepth() const { return TreeDepth; }
-
+			
 			const FImplicitBVHObject& GetObject(const int32 ObjectIndex) const { return Objects[ObjectIndex]; }
 
 			const FImplicitObject* GetGeometry(const int32 ObjectIndex) const { return Objects[ObjectIndex].GetGeometry(); }
diff --git a/Engine/Source/Runtime/Experimental/Chaos/Public/Chaos/ImplicitObjectUnion.h b/Engine/Source/Runtime/Experimental/Chaos/Public/Chaos/ImplicitObjectUnion.h
--- a/Engine/Source/Runtime/Experimental/Chaos/Public/Chaos/ImplicitObjectUnion.h
+++ b/Engine/Source/Runtime/Experimental/Chaos/Public/Chaos/ImplicitObjectUnion.h
@@ -50,11 +50,8 @@ class FImplicitObjectUnion : public FImplicitObject
 	// The total number of root objects in the hierarchy (same as GetObjects().Num())
 	int32 GetNumRootObjects() const
 	{
-		return GetObjects().Num();
+		return MObjects.Num();
 	}
-
-	// Return true if the union has a simplified list of shapes
-	CHAOS_API bool HasSimpleShapes() const;
 	
 	// The total number of leaf objects in the hierarchy
 	int32 GetNumLeafObjects() const
@@ -79,7 +76,7 @@ class FImplicitObjectUnion : public FImplicitObject
 	
 	virtual FReal PhiWithNormal(const FVec3& x, FVec3& Normal) const override
 	{
-		const Chaos::FImplicitObjectsArray& Objects = HasSimpleShapes() ? GetConvexes() : GetObjects();
+		const Chaos::FImplicitObjectsArray& Objects = MObjects;
 		FReal Phi = TNumericLimits<FReal>::Max();
 		bool NeedsNormalize = false;
 		for (int32 i = 0; i < Objects.Num(); ++i)
@@ -128,8 +125,7 @@ class FImplicitObjectUnion : public FImplicitObject
 		FReal MinTime = 0;	//initialization not needed, but doing it to avoid warning
 		bool bFound = false;
 
-		const Chaos::FImplicitObjectsArray& Objects = HasSimpleShapes() ? GetConvexes() : GetObjects();
-		for (const Chaos::FImplicitObjectPtr& Obj : Objects)
+		for (const Chaos::FImplicitObjectPtr& Obj : MObjects)
 		{
 			FVec3 Position;
 			FVec3 Normal;
@@ -154,8 +150,7 @@ class FImplicitObjectUnion : public FImplicitObject
 
 	virtual bool Overlap(const FVec3& Point, const FReal Thickness) const override
 	{
-		const Chaos::FImplicitObjectsArray& Objects = HasSimpleShapes() ? GetConvexes() : GetObjects();
-		for (const Chaos::FImplicitObjectPtr& Obj : Objects)
+		for (const Chaos::FImplicitObjectPtr& Obj : MObjects)
 		{
 			if (Obj->Overlap(Point, Thickness))
 			{
@@ -175,18 +170,12 @@ class FImplicitObjectUnion : public FImplicitObject
 		return bValid;
 	}
 
+	// Return the list of objects that will be part of the union
 	CHAOS_API const TArray<Chaos::FImplicitObjectPtr>& GetObjects() const { return MObjects; }
 
+	// Return the const list of objects that will be part of the union
 	CHAOS_API TArray<Chaos::FImplicitObjectPtr>& GetObjects() { return MObjects; }
 	
-	CHAOS_API const TArray<Chaos::FImplicitObjectPtr>& GetConvexes() const { return Convexes; }
-
-	CHAOS_API TArray<Chaos::FImplicitObjectPtr>& GetConvexes() { return Convexes; }
-
-	CHAOS_API int32 NumImplicits() const { return Convexes.Num() + MObjects.Num(); }
-
-	CHAOS_API Chaos::FImplicitObjectPtr GetImplicit(const int32 ImplicitIndex) const { return (ImplicitIndex < Convexes.Num()) ? Convexes[ImplicitIndex] : MObjects[ImplicitIndex-Convexes.Num()]; }
-
 	// The lambda returns TRUE if an object was found and iteration should stop.
 	CHAOS_API void ForEachObject(TFunctionRef<bool(const FImplicitObject&, const FRigidTransform3&)> Lambda) const;
 
@@ -221,17 +210,15 @@ class FImplicitObjectUnion : public FImplicitObject
 #endif // #if INTEL_ISPC
 
 protected:
-	
+
 	virtual Pair<FVec3, bool> FindClosestIntersectionImp(const FVec3& StartPoint, const FVec3& EndPoint, const FReal Thickness) const override
 	{
-		const Chaos::FImplicitObjectsArray& Objects = HasSimpleShapes() ? GetConvexes() : GetObjects();
-		
-		check(Objects.Num());
-		auto ClosestIntersection = Objects[0]->FindClosestIntersection(StartPoint, EndPoint, Thickness);
+		check(MObjects.Num());
+		auto ClosestIntersection = MObjects[0]->FindClosestIntersection(StartPoint, EndPoint, Thickness);
 		FReal Length = ClosestIntersection.Second ? (ClosestIntersection.First - StartPoint).Size() : 0;
-		for (int32 i = 1; i < Objects.Num(); ++i)
+		for (int32 i = 1; i < MObjects.Num(); ++i)
 		{
-			auto NextClosestIntersection = Objects[i]->FindClosestIntersection(StartPoint, EndPoint, Thickness);
+			auto NextClosestIntersection = MObjects[i]->FindClosestIntersection(StartPoint, EndPoint, Thickness);
 			if (!NextClosestIntersection.Second)
 				continue;
 			FReal NewLength = (NextClosestIntersection.First - StartPoint).Size();
@@ -282,7 +269,7 @@ protected:
 	CHAOS_API void RebuildBVH();
 
 	CHAOS_API void LegacySerializeBVH(FChaosArchive& Ar);
-
+	
 	union FFLags
 	{
 	public:
@@ -299,14 +286,8 @@ protected:
 	// Check if the BVH could be used and valid
 	bool HasValidBVH() const;
 
-	// Rebuild all the convexes used for physics
-	void RebuildConvexes();
-
 	// list of implicit objects that are part of the union
 	TArray<Chaos::FImplicitObjectPtr> MObjects;
-
-	// List of implicit convexes that could be used by the union
-	TArray<Chaos::FImplicitObjectPtr> Convexes;
 	
 	FAABB3 MLocalBoundingBox;
 
@@ -376,5 +357,4 @@ struct TImplicitTypeInfo<FImplicitObjectUnion>
 	}
 };
 
-
 }
diff --git a/Engine/Source/Runtime/Experimental/Chaos/Public/Chaos/PBDRigidClusteredParticles.h b/Engine/Source/Runtime/Experimental/Chaos/Public/Chaos/PBDRigidClusteredParticles.h
--- a/Engine/Source/Runtime/Experimental/Chaos/Public/Chaos/PBDRigidClusteredParticles.h
+++ b/Engine/Source/Runtime/Experimental/Chaos/Public/Chaos/PBDRigidClusteredParticles.h
@@ -4,10 +4,16 @@
 #include "Chaos/ArrayCollectionArray.h"
 #include "Chaos/PBDRigidParticles.h"
 #include "Chaos/ImplicitObjectUnion.h"
+#include "Templates/PimplPtr.h"
 
 namespace Chaos
 {
 
+namespace Private
+{
+	class FConvexOptimizer;
+}
+
 /** 
  * Used within the clustering system to describe the clustering hierarchy. The ClusterId
  * stores the children IDs, and a Parent ID. When Id == \c nullptr the cluster is not
@@ -157,6 +163,7 @@ class TPBDRigidClusteredParticles : public TPBDRigidParticles<T, d>
 	, MConnectivityEdges(MoveTemp(Other.MConnectivityEdges))
 	, MExternalStrains(MoveTemp(Other.MExternalStrains))
 	, MRigidClusteredFlags(MoveTemp(Other.MRigidClusteredFlags))
+	, MConvexOptimizers(MoveTemp(Other.MConvexOptimizers))
 	{
 		InitHelper();
 	}
@@ -227,6 +234,8 @@ class TPBDRigidClusteredParticles : public TPBDRigidParticles<T, d>
 	const auto& RigidClusteredFlags() const { return MRigidClusteredFlags; }
 	auto& RigidClusteredFlags() { return MRigidClusteredFlags; }
 
+	const auto& ConvexOptimizers(int32 Idx) const { return MConvexOptimizers[Idx]; }
+	auto& ConvexOptimizers(int32 Idx) { return MConvexOptimizers[Idx]; }
 	
 	typedef TPBDRigidClusteredParticleHandle<T, d> THandleType;
 	const THandleType* Handle(int32 Index) const { return static_cast<const THandleType*>(TGeometryParticles<T,d>::Handle(Index)); }
@@ -250,6 +259,7 @@ class TPBDRigidClusteredParticles : public TPBDRigidParticles<T, d>
 		  TArrayCollection::AddArray(&MConnectivityEdges);
 	  	  TArrayCollection::AddArray(&MExternalStrains);
 	  	  TArrayCollection::AddArray(&MRigidClusteredFlags);
+	  	  TArrayCollection::AddArray(&MConvexOptimizers);
 	  }
 
 	  TArrayCollectionArray<ClusterId> MClusterIds;
@@ -273,6 +283,9 @@ class TPBDRigidClusteredParticles : public TPBDRigidParticles<T, d>
 	  TArrayCollectionArray<TArray<TConnectivityEdge<T>>> MConnectivityEdges;
   
 	  TArrayCollectionArray<FRigidClusteredFlags> MRigidClusteredFlags;
+
+	  // Per clustered particle convex optimizer to reduce the collision cost
+	  TArrayCollectionArray<TPimplPtr<Private::FConvexOptimizer>> MConvexOptimizers;
 };
 
 using FPBDRigidClusteredParticles = TPBDRigidClusteredParticles<FReal, 3>;
diff --git a/Engine/Source/Runtime/Experimental/Chaos/Public/Chaos/PBDRigidClustering.h b/Engine/Source/Runtime/Experimental/Chaos/Public/Chaos/PBDRigidClustering.h
--- a/Engine/Source/Runtime/Experimental/Chaos/Public/Chaos/PBDRigidClustering.h
+++ b/Engine/Source/Runtime/Experimental/Chaos/Public/Chaos/PBDRigidClustering.h
@@ -396,6 +396,14 @@ public:
 	CHAOS_API void SetInternalStrain(FPBDRigidClusteredParticleHandle* Particle, FRealSingle Strain);
 	CHAOS_API void SetExternalStrain(FPBDRigidClusteredParticleHandle* Particle, FRealSingle Strain);
 
+	/*
+	*  BuildConvexOptimizer
+	*    Create the convex optimizer unique ptr and loop over the particle geometry to simplify
+	*    all the convexes within the hierarchy
+	*    @param Particle particle on which the geometry will be simplified   
+	*/
+	static CHAOS_API void BuildConvexOptimizer(FPBDRigidClusteredParticleHandle* Particle);
+	
 	static CHAOS_API bool ShouldUnionsHaveCollisionParticles();
 
 	FClusterUnionManager& GetClusterUnionManager() { return ClusterUnionManager; }
diff --git a/Engine/Source/Runtime/Experimental/Chaos/Public/Chaos/ParticleHandle.h b/Engine/Source/Runtime/Experimental/Chaos/Public/Chaos/ParticleHandle.h
--- a/Engine/Source/Runtime/Experimental/Chaos/Public/Chaos/ParticleHandle.h
+++ b/Engine/Source/Runtime/Experimental/Chaos/Public/Chaos/ParticleHandle.h
@@ -26,6 +26,11 @@ namespace Chaos
 {
 	class FConstraintHandle;
 	class FPBDRigidsEvolutionBase;
+	
+	namespace Private
+	{
+		class FConvexOptimizer;
+	}
 
 struct FGeometryParticleParameters
 {
@@ -1456,6 +1461,10 @@ public:
 
 	int32 TransientParticleIndex() const { return ParticleIdx; }
 
+	// Get the clustered particle convex optmizer
+	const TPimplPtr<Private::FConvexOptimizer>& ConvexOptimizer() const {return PBDRigidClusteredParticles->ConvexOptimizers(ParticleIdx); }
+	TPimplPtr<Private::FConvexOptimizer>& ConvexOptimizer() {return PBDRigidClusteredParticles->ConvexOptimizers(ParticleIdx); }
+
 private:
 	void SetInternalStrains(const FRealSingle Value) { PBDRigidClusteredParticles->Strains(ParticleIdx) = Value; }
 	void SetExternalStrains(const FRealSingle Value) { PBDRigidClusteredParticles->ExternalStrains(ParticleIdx) = Value; }
@@ -2606,7 +2615,6 @@ public:
 	void MergeShapesArray(FShapesArray&& InShapesArray)
 	{
 		MergeShapeInstances(reinterpret_cast<FShapeInstanceProxyArray&&>(InShapesArray));
-		UpdateSimpleShapes();
 	}
 
 	const FShapesArray& ShapesArray() const { return reinterpret_cast<const FShapesArray&>(MShapesArray); }
@@ -2629,7 +2637,7 @@ public:
 	// provided as a utility for use in FPhysInterface_Chaos::AddGeometry. Ideally we would remove this
 	void MergeShapeInstances(FShapeInstanceProxyArray&& InShapes)
 	{
-		check(InShapes.Num() <= MShapesArray.Num());
+		check(InShapes.Num() <= (MShapesArray.Num()));
 		int Idx = MShapesArray.Num() - InShapes.Num();
 		for (FShapeInstanceProxyPtr& Shape : InShapes)
 		{
@@ -2783,37 +2791,6 @@ public:
 
 protected:
 
-	/** Update the simple shapes materials, collisions... */
-	void UpdateSimpleShapes()
-	{
-		if (const FImplicitObjectUnion* Union = GetGeometry()->template GetObject<FImplicitObjectUnion>())
-		{
-			const int32 ObjectsOffset = Union->GetConvexes().Num();
-			if(ObjectsOffset > 0)
-			{
-				for(int32 ShapeIndex = 0; ShapeIndex < MShapesArray.Num(); ++ShapeIndex)
-				{
-					if(ShapeIndex < ObjectsOffset)
-					{
-						MShapesArray[ShapeIndex]->SetQueryData(MShapesArray[ObjectsOffset]->GetQueryData());
-						MShapesArray[ShapeIndex]->SetSimData(MShapesArray[ObjectsOffset]->GetSimData());
-						MShapesArray[ShapeIndex]->SetCollisionTraceType(MShapesArray[ObjectsOffset]->GetCollisionTraceType());
-				
-						// Only sim enabled if the underlying shapes could be used for physics
-						MShapesArray[ShapeIndex]->SetSimEnabled(MShapesArray[ObjectsOffset]->GetSimEnabled());
-
-						// Disable simple shapes for query 
-						MShapesArray[ShapeIndex]->SetQueryEnabled(false);
-					}
-					else
-					{
-						MShapesArray[ShapeIndex]->SetSimEnabled(false);
-					}
-				}
-			}
-		}
-	}
-
 	// Pointer to any data that the solver wants to associate with this particle
 	// TODO: It's important to eventually hide this!
 	// Right now it's exposed to lubricate the creation of the whole proxy system.
@@ -2888,6 +2865,46 @@ protected:
 	}
 };
 
+// Update all the shapes datas including the convexes ones
+template<typename ParticleType>
+void UpdateParticleShapes(const TArray<ParticleType*>& ShapesParticles, const FImplicitObject* ImplicitObject,
+	const FShapesArray& ShapesArray, const int32 ActorId, const int32 ComponentID)
+{
+	if (!ShapesParticles.IsEmpty() && (ShapesParticles.Num() <= (ShapesArray.Num())))
+	{
+		for (int32 ParticleIndex = 0; ParticleIndex < ShapesParticles.Num(); ++ParticleIndex)
+		{
+			const TUniquePtr<Chaos::FPerShapeData>& ShapeData = ShapesArray[ParticleIndex];
+			const TUniquePtr<Chaos::FPerShapeData>& TemplateShape = ShapesParticles[ParticleIndex]->ShapesArray()[0];
+			
+			if (ShapeData && TemplateShape)
+			{
+				{
+					FCollisionData Data = TemplateShape->GetCollisionData();
+					Data.UserData = nullptr;
+					ShapeData->SetCollisionData(Data);
+				}
+    
+				{
+					FCollisionFilterData Data = TemplateShape->GetQueryData();
+					Data.Word0 = ActorId;
+					ShapeData->SetQueryData(Data);
+				}
+    
+				{
+					FCollisionFilterData Data = TemplateShape->GetSimData();
+					Data.Word0 = 0;
+					Data.Word2 = ComponentID;
+					ShapeData->SetSimData(Data);
+				}
+    
+				ShapeData->SetSimEnabled(TemplateShape->GetSimEnabled());
+				ShapeData->SetQueryEnabled(TemplateShape->GetQueryEnabled());
+			}
+		}
+	}
+}
+
 template <typename T, int d>
 FChaosArchive& operator<<(FChaosArchive& Ar, TGeometryParticle<T, d>& Particle)
 {
diff --git a/Engine/Source/Runtime/Experimental/Chaos/Public/Chaos/Tribox.h b/Engine/Source/Runtime/Experimental/Chaos/Public/Chaos/Tribox.h
new file mode 100644
--- /dev/null
+++ b/Engine/Source/Runtime/Experimental/Chaos/Public/Chaos/Tribox.h
@@ -0,0 +1,104 @@
+// Copyright Epic Games, Inc. All Rights Reserved.
+
+#pragma once
+
+#include "Chaos/Core.h"
+#include "Chaos/Convex.h"
+
+namespace Chaos
+{
+namespace Private
+{
+		
+/**
+ * @brief Tribox object that represents a k-DOP18 convex
+ * 
+ * @note This implementation provides a fast way of constructing Tribox in order to be used at runtime to build
+ * approximate convex implicit object
+ * 
+ */
+
+class FTribox
+{
+public :
+
+	using FRealType = FRealSingle;
+	using FVec3Type = TVec3<FRealType>;
+	using FRigidTransform3Type = TRigidTransform<FRealType, 3>;
+	using FMatrix33Type = PMatrix<FRealType, 3, 3>;
+
+	// Number of planes that will be used to define the Tribox 
+	static constexpr int32 NumPlanes = 18;
+
+	// Number of principal planes
+	static constexpr int32 NumPrincipalPlanes = 6;
+
+	// Number of chamfer planes
+	static constexpr int32 NumChamferPlanes = 12;
+
+	// Distance used to inflate the Tribox to avoid degenerate case
+	static constexpr FRealType InflateDistance = 0.5;
+
+	// Base Constructor
+	FORCEINLINE FTribox() : MaxDists()
+	{
+		for(int32 DistsIndex = 0; DistsIndex< NumPlanes; ++DistsIndex)
+		{
+			MaxDists[DistsIndex] = TNumericLimits<FRealType>::Lowest();
+		}
+	};
+
+	// Add a point position to the Tribox
+	void AddPoint(const FVec3Type& PointPosition);
+
+	// Add convex vertices to a tribox
+	void AddConvex(const FConvex* Convex, const FRigidTransform3Type& RelativeTransform);
+
+	// Inflate + Scale the Max distances
+	bool BuildTribox();
+
+	// Find the overlapping tribox 
+	bool OverlapTribox(const FTribox& OtherTribox, FTribox& OverlapTribox) const;
+
+	// Compute the tribox volume
+	FRealType ComputeVolume() const;
+
+	// Create a convex from the Tribox
+	FImplicitObjectPtr MakeConvex() const;
+
+	// Add a tribox to this and return this
+	FTribox& operator+=(const FTribox& OtherTribox);
+
+	// Add a tribox to this and return a new one
+	FTribox operator+( const FTribox& OtherTribox) const;
+
+private : 
+	// Solve the intersection point position
+	FVec3 SolveIntersection(const int32 FaceIndex[3], const FMatrix33Type& A) const;
+
+	// Add the intersection point to the list of faces/vertices
+	void AddIntersection(const int32 FaceIndex[3], const int32 FaceOrder[3], const FVec3Type& VertexPosition, TArray<TArray<int32>>& FaceIndices, TArray<FConvex::FVec3Type>& ConvexVertices) const;
+
+	// Solve and add the intersection point in between the 3 planes
+	void ComputeIntersection(const int32 CornerIndex, const int32 PrincipalIndex, const int32 ChamferIndex, TArray<TArray<int32>>& FaceIndices, TArray<FConvex::FVec3Type>& ConvexVertices) const;
+
+	// Build the max dist along the principal axis
+	void BuildPrincipalDist(const FVec3Type& P, const int32 CoordIndexA, const int32 DistsIndexA, const int32 DistsIndexB);
+
+	// Build the max dist along the chamfer axis
+	void BuildChamferDist(const FVec3Type& P, const int32 CoordIndexA, const int32 CoordIndexB,
+			const int32 DistsIndexA, const int32 DistsIndexB, const int32 DistsIndexC, const int32 DistsIndexD);
+
+	// Max distance along eaxh tribox axis (principal + chamfer)
+	FRealType MaxDists[NumPlanes] = {TNumericLimits<FRealType>::Lowest(), TNumericLimits<FRealType>::Lowest(),
+									 TNumericLimits<FRealType>::Lowest(), TNumericLimits<FRealType>::Lowest(),
+									 TNumericLimits<FRealType>::Lowest(), TNumericLimits<FRealType>::Lowest(),
+									 TNumericLimits<FRealType>::Lowest(), TNumericLimits<FRealType>::Lowest(),
+									 TNumericLimits<FRealType>::Lowest()};
+
+	// Number of points used to build the simplified convexes
+	int32 NumPoints = 0;
+};
+
+}
+}
\ No newline at end of file
diff --git a/Engine/Source/Runtime/PhysicsCore/Private/ChaosEngineInterface.cpp b/Engine/Source/Runtime/PhysicsCore/Private/ChaosEngineInterface.cpp
--- a/Engine/Source/Runtime/PhysicsCore/Private/ChaosEngineInterface.cpp
+++ b/Engine/Source/Runtime/PhysicsCore/Private/ChaosEngineInterface.cpp
@@ -1952,15 +1952,14 @@ void FChaosEngineInterface::SetLocalTransform(const FPhysicsShapeHandle& InShape
 			const FImplicitObjectUnion* AsUnion = static_cast<const FImplicitObjectUnion*>(CurrentGeom);
 			const int32 ShapeIndex = InShape.Shape->GetShapeIndex();
 			const TArray<Chaos::FImplicitObjectPtr>& ObjectArray = AsUnion->GetObjects();
-			const TArray<Chaos::FImplicitObjectPtr>& ConvexesArray = AsUnion->GetConvexes();
 
-			if(ensure(ShapeIndex < (ObjectArray.Num()+ConvexesArray.Num())) && (ShapeIndex > ConvexesArray.Num()))
+			if(ensure(ShapeIndex < ObjectArray.Num()))
 			{
 				TArray<Chaos::FImplicitObjectPtr> NewGeoms;
 				NewGeoms.Reserve(ObjectArray.Num());
 
 				// Duplicate the union and either set transforms, or wrap in transforms
-				int32 CurrentIndex = ConvexesArray.Num();
+				int32 CurrentIndex = 0;
 				for(const Chaos::FImplicitObjectPtr& Obj : ObjectArray)
 				{
 					if(CurrentIndex == ShapeIndex)
@@ -1991,11 +1990,13 @@ int32 GetAllShapesInternalImp_AssumedLocked(const FPhysicsActorHandle& InActorHa
 	if(InActorHandle)
 	{
 		const Chaos::FShapesArray& ShapesArray = InActorHandle->GetGameThreadAPI().ShapesArray();
-		OutShapes.Reset(ShapesArray.Num());
+		const int32 NumRelevantShapes = ShapesArray.Num();
+		OutShapes.Reset(NumRelevantShapes);
+		
 		//todo: can we avoid this construction?
-		for(const TUniquePtr<Chaos::FPerShapeData>& Shape : ShapesArray)
+		for(int32 ShapeIndex = 0; ShapeIndex < NumRelevantShapes; ++ShapeIndex)
 		{
-			OutShapes.Add(FPhysicsShapeReference_Chaos(Shape.Get(),InActorHandle));
+			OutShapes.Add(FPhysicsShapeReference_Chaos(ShapesArray[ShapeIndex].Get(),InActorHandle));
 		}
 
 		return OutShapes.Num();
