Fix missing copyright

diff --git a/Engine/Source/Runtime/Experimental/Chaos/Private/Chaos/ConvexOptimizer.cpp b/Engine/Source/Runtime/Experimental/Chaos/Private/Chaos/ConvexOptimizer.cpp
--- a/Engine/Source/Runtime/Experimental/Chaos/Private/Chaos/ConvexOptimizer.cpp
+++ b/Engine/Source/Runtime/Experimental/Chaos/Private/Chaos/ConvexOptimizer.cpp
@@ -1,3 +1,5 @@
+// Copyright Epic Games, Inc. All Rights Reserved.
+
 
 #include "Chaos/ConvexOptimizer.h"
 #include "Chaos/Tribox.h"
