updating a particle's child to parent in a cluster union must also propagate that change to the implicit object transform

diff --git a/Engine/Source/Runtime/Experimental/Chaos/Private/Chaos/ClusterUnionManager.cpp b/Engine/Source/Runtime/Experimental/Chaos/Private/Chaos/ClusterUnionManager.cpp
--- a/Engine/Source/Runtime/Experimental/Chaos/Private/Chaos/ClusterUnionManager.cpp
+++ b/Engine/Source/Runtime/Experimental/Chaos/Private/Chaos/ClusterUnionManager.cpp
@@ -969,6 +969,7 @@ namespace Chaos
 		check(ClusterUnion.Geometry != nullptr && ClusterUnion.Geometry->GetType() == ImplicitObjectType::Union);
 
 		const TArray<FPBDRigidParticleHandle*>& PendingGeometryAdditions = ClusterUnion.GetPendingGeometryOperationParticles(EClusterUnionGeometryOperation::Add);
+		const FRigidTransform3 ClusterWorldTM(ClusterUnion.InternalCluster->X(), ClusterUnion.InternalCluster->R());
 		if (!PendingGeometryAdditions.IsEmpty())
 		{
 			ModifyAdditionOfChildrenToClusterUnionGeometry(
@@ -976,9 +977,8 @@ namespace Chaos
 				PendingGeometryAdditions,
 				ClusterUnion.ClusterUnionParameters.ActorId,
 				ClusterUnion.ClusterUnionParameters.ComponentId,
-				[this, &ClusterUnion, &PendingGeometryAdditions]()
+				[this, &ClusterUnion, &PendingGeometryAdditions, &ClusterWorldTM]()
 				{
-					const FRigidTransform3 ClusterWorldTM(ClusterUnion.InternalCluster->X(), ClusterUnion.InternalCluster->R());
 
 					TArray<Chaos::FImplicitObjectPtr> Objects;
 					Objects.Reserve(PendingGeometryAdditions.Num());
@@ -1039,6 +1039,18 @@ namespace Chaos
 					continue;
 				}
 
+				if (FImplicitObjectRef ImplicitGeometry = ShapesArray[Index]->GetGeometry())
+				{
+					if (FImplicitObjectTransformed* Transformed = ImplicitGeometry->AsA<FImplicitObjectTransformed>())
+					{
+						const FRigidTransform3 Frame = GetParticleRigidFrameInClusterUnion(Particle, ClusterWorldTM);
+						if (!Transformed->GetTransform().Equals(Frame))
+						{
+							Transformed->SetTransform(Frame);
+						}
+					}
+				}
+
 				TransferClusterUnionShapeData(
 					ShapesArray[Index],
 					Particle,
@@ -1209,6 +1221,10 @@ namespace Chaos
 						// the child to parent update might move the node so far away as to make the old connectivity edges incorrect.
 						ClusterUnion->PendingConnectivityOperations.Add({ Particle, EClusterUnionConnectivityOperation::Remove });
 						ClusterUnion->PendingConnectivityOperations.Add({ Particle, EClusterUnionConnectivityOperation::Add });
+
+						// A child to parent update also requires the geometry to be refreshed since its transform is changed.
+						ClusterUnion->AddPendingGeometryOperation(EClusterUnionGeometryOperation::Refresh, Particle);
+
 						bMadeChanges = true;
 					}
 				}
@@ -1217,7 +1233,9 @@ namespace Chaos
 
 		if (bMadeChanges)
 		{
-			RequestDeferredClusterPropertiesUpdate(ClusterIndex, EUpdateClusterUnionPropertiesFlags::IncrementalGenerateConnectionGraph);
+			constexpr EUpdateClusterUnionPropertiesFlags Flags = EUpdateClusterUnionPropertiesFlags::IncrementalGenerateConnectionGraph
+				| EUpdateClusterUnionPropertiesFlags::IncrementalGenerateGeometry;
+			RequestDeferredClusterPropertiesUpdate(ClusterIndex, Flags);
 		}
 	}
 
