Chaos: Add material transfer to incremental shape path for cluster unions. Otherwise shapes will get default material paramters applied to the root cluster shapes

diff --git a/Engine/Source/Runtime/Experimental/Chaos/Public/Chaos/ClusterUnionManager.h b/Engine/Source/Runtime/Experimental/Chaos/Public/Chaos/ClusterUnionManager.h
--- a/Engine/Source/Runtime/Experimental/Chaos/Public/Chaos/ClusterUnionManager.h
+++ b/Engine/Source/Runtime/Experimental/Chaos/Public/Chaos/ClusterUnionManager.h
@@ -300,6 +300,22 @@ namespace Chaos
 				Data.Word2 = ComponentId;
 				ShapeData->SetSimData(Data);
 			}
+
+			{
+				const int32 NumTemplateMaterials = TemplateShape->NumMaterials();
+
+				if(NumTemplateMaterials == 1)
+				{
+					// We have a special case for 1 material, calling GetMaterials
+					// would allocate a material storage object to apply to the shape
+					// whereas for 1 material we store it inline to the shape data
+					ShapeData->SetMaterial(TemplateShape->GetMaterial(0));
+				}
+				else if(NumTemplateMaterials > 1)
+				{
+					ShapeData->SetMaterials(TemplateShape->GetMaterials());
+				}
+			}
 		}
 	}
 
