[Chaos] Changed the time-slicing implementation in AABBtree generation to be done within a millisecond budget instead of by a number of processed nodes.

The slices are done in the same places as the original implementation, so the behavior is expected to be the same.

This change can be disabled by setting the cvar p.aabbtree.UseTimeSliceMillisecondBudget to false.

The time limit can be adjusted with p.aabbtree.MaxProcessingTimePerSliceSeconds (it is on seconds, so 1ms is 0.001).

As checking for the current time constantly had a cost that added up fast due to being a system call, some sections only check the timespamp after x amounts of elements processed. That can be controlled with the cvars

p.aabbtree.MinNodesChunkToProcessBetweenTimeChecks

p.aabbtree.MinDataChunkToProcessBetweenTimeChecks

The higher their value, less system call we made to get the current timestamp reducing the overall cost, but we might overshoot the budget.

diff --git a/Engine/Source/Runtime/Experimental/Chaos/Private/Chaos/AABBTree.cpp b/Engine/Source/Runtime/Experimental/Chaos/Private/Chaos/AABBTree.cpp
--- a/Engine/Source/Runtime/Experimental/Chaos/Private/Chaos/AABBTree.cpp
+++ b/Engine/Source/Runtime/Experimental/Chaos/Private/Chaos/AABBTree.cpp
@@ -30,3 +30,15 @@ FAutoConsoleVariableRef FAABBTreeCVars::CVarDynamicTreeBoundingBoxPadding(TEXT("
 
 int32 FAABBTreeCVars::DynamicTreeLeafCapacity = 8;
 FAutoConsoleVariableRef FAABBTreeCVars::CVarDynamicTreeLeafCapacity(TEXT("p.aabbtree.DynamicTreeLeafCapacity"), FAABBTreeCVars::DynamicTreeLeafCapacity, TEXT("Dynamic Tree Leaf Capacity"));
+
+bool FAABBTimeSliceCVars::bUseTimeSliceMillisecondBudget = true;
+FAutoConsoleVariableRef FAABBTimeSliceCVars::CVarUseTimeSliceByMillisecondBudget(TEXT("p.aabbtree.UseTimeSliceMillisecondBudget"), FAABBTimeSliceCVars::bUseTimeSliceMillisecondBudget, TEXT("Set to True if we want to timeslice tree generation by a milisecond budget instead of per nodes processed"));
+
+float FAABBTimeSliceCVars::MaxProcessingTimePerSliceSeconds = 0.001f;
+FAutoConsoleVariableRef FAABBTimeSliceCVars::CVarMaxProcessingTimePerSlice(TEXT("p.aabbtree.MaxProcessingTimePerSliceSeconds"), FAABBTimeSliceCVars::MaxProcessingTimePerSliceSeconds, TEXT("Set to True if we want to timeslice tree generation by a milisecond budget instead of per nodes processed or data size copied"));
+
+int32 FAABBTimeSliceCVars::MinNodesChunkToProcessBetweenTimeChecks = 250;
+FAutoConsoleVariableRef FAABBTimeSliceCVars::CVarMinNodesChunkToProcessBetweenTimeChecks(TEXT("p.aabbtree.MinNodesChunkToProcessBetweenTimeChecks"), FAABBTimeSliceCVars::MinNodesChunkToProcessBetweenTimeChecks, TEXT("Minimum amount of nodes we want to process before checking if we are withing budget"));
+
+int32 FAABBTimeSliceCVars::MinDataChunkToProcessBetweenTimeChecks = 500;
+FAutoConsoleVariableRef FAABBTimeSliceCVars::CVarMinDataChunkToProcessBetweenTimeChecks(TEXT("p.aabbtree.MinDataChunkToProcessBetweenTimeChecks"), FAABBTimeSliceCVars::MinDataChunkToProcessBetweenTimeChecks, TEXT("Minimum amount of data elements to process before checking if we are withing budget"));
\ No newline at end of file
diff --git a/Engine/Source/Runtime/Experimental/Chaos/Public/Chaos/AABBTree.h b/Engine/Source/Runtime/Experimental/Chaos/Public/Chaos/AABBTree.h
--- a/Engine/Source/Runtime/Experimental/Chaos/Public/Chaos/AABBTree.h
+++ b/Engine/Source/Runtime/Experimental/Chaos/Public/Chaos/AABBTree.h
@@ -57,6 +57,21 @@ struct FAABBTreeDirtyGridCVars
 	static CHAOS_API FAutoConsoleVariableRef CVarDirtyElementMaxCellCapacity;
 };
 
+struct FAABBTimeSliceCVars
+{
+	static CHAOS_API bool bUseTimeSliceMillisecondBudget;
+	static CHAOS_API FAutoConsoleVariableRef CVarUseTimeSliceByMillisecondBudget;
+
+	static CHAOS_API float MaxProcessingTimePerSliceSeconds;
+	static CHAOS_API FAutoConsoleVariableRef CVarMaxProcessingTimePerSlice;
+
+	static CHAOS_API int32 MinNodesChunkToProcessBetweenTimeChecks;
+	static CHAOS_API FAutoConsoleVariableRef CVarMinNodesChunkToProcessBetweenTimeChecks;
+
+	static CHAOS_API int32 MinDataChunkToProcessBetweenTimeChecks;
+	static CHAOS_API FAutoConsoleVariableRef CVarMinDataChunkToProcessBetweenTimeChecks;
+};
+
 namespace Chaos
 {
 
@@ -794,6 +809,11 @@ public:
 		OverlappingCounts.Reset();
 		
 		NumProcessedThisSlice = 0;
+
+		StartSliceTimeStamp = 0.0;
+		CurrentDataElementsCopiedSinceLastCheck = 0;
+		CurrentProcessedNodesSinceChecked = 0;
+		
 		WorkStack.Reset();
 		WorkPoolFreeList.Reset();
 		WorkPool.Reset();
@@ -833,6 +853,7 @@ public:
 		if (WorkStack.Num())
 		{
 			NumProcessedThisSlice = 0;
+			StartSliceTimeStamp = FPlatformTime::Seconds();
 			SplitNode();
 		}
 	}
@@ -2128,6 +2149,9 @@ public:
 		MaxPayloadBounds = From.MaxPayloadBounds;
 		MaxNumToProcess = From.MaxNumToProcess;
 		NumProcessedThisSlice = From.NumProcessedThisSlice;
+
+		StartSliceTimeStamp = From.StartSliceTimeStamp;
+
 		bShouldRebuild = From.bShouldRebuild;
 
 		RootNode = From.RootNode;
@@ -2164,8 +2188,8 @@ public:
 			DirtyElementTree = nullptr;
 		}
 	}
-	
-	virtual void ProgressCopyTimeSliced(const  ISpatialAcceleration<TPayloadType, T, 3>& InFrom, int MaximumBytesToCopy) override
+
+	virtual void ProgressCopyTimeSliced(const ISpatialAcceleration<TPayloadType, T, 3>& InFrom, int MaximumBytesToCopy) override
 	{
 		check(this != &InFrom);
 		check(InFrom.GetType() == ESpatialAcceleration::AABBTree);
@@ -2174,48 +2198,97 @@ public:
 		int32 SizeToCopyLeft = MaximumBytesToCopy;
 		check(From.CellHashToFlatArray.Num() == 0); // Partial Copy of TMAPs not implemented, and this should be empty for our current use cases
 
-		if (!ContinueTimeSliceCopy(From.Nodes, Nodes, SizeToCopyLeft))
+		TFunction<bool(int32,int32)> CanContinueCopyingDataCallback = [this](int32 MaxSizeToCopy, int32 CurrentCopiedSize)
+		{
+			bool bCanContinueCopy = true;
+			const bool bForceCopyAll = MaxSizeToCopy == -1;
+			if (bForceCopyAll)
+			{
+				return bCanContinueCopy;
+			}
+
+			if (FAABBTimeSliceCVars::bUseTimeSliceMillisecondBudget)
+			{
+				// Checking for platform time every time is expensive as its cost adds up fast, so only do it between a configured amount of elements.
+				// This means we might overshoot the budget, but on the other hand we will keep most of the time doing actual work.
+				CurrentDataElementsCopiedSinceLastCheck++;
+				if (CurrentDataElementsCopiedSinceLastCheck > FAABBTimeSliceCVars::MinDataChunkToProcessBetweenTimeChecks)
+				{
+					CurrentDataElementsCopiedSinceLastCheck = 0;
+					return bCanContinueCopy;
+				}
+
+				const double ElapseTime = FPlatformTime::Seconds() - StartSliceTimeStamp;
+
+				if (!FMath::IsNearlyZero(StartSliceTimeStamp) && ElapseTime > FAABBTimeSliceCVars::MaxProcessingTimePerSliceSeconds)
+				{
+					bCanContinueCopy = false;
+				}
+			}
+			else
+			{
+				bCanContinueCopy = MaxSizeToCopy < 0 || CurrentCopiedSize < MaxSizeToCopy;
+			}
+
+			return bCanContinueCopy;
+		};
+
+		constexpr int32 SizeCopiedSoFar = 0;
+		if (!CanContinueCopyingDataCallback(MaximumBytesToCopy, SizeCopiedSoFar))
+		{
+			// For data copy, the time stamp is set from the setup phase, and it is copied from the tree we are copying from.
+			// Therefore if we reach this point with no time available left, just reset it so the next frame can start counting from scratch 
+			StartSliceTimeStamp = 0.0;
+			return;
+		}
+
+		if (FMath::IsNearlyZero(StartSliceTimeStamp))
+		{
+			StartSliceTimeStamp = FPlatformTime::Seconds();
+		}
+
+		if (!ContinueTimeSliceCopy(From.Nodes, Nodes, SizeToCopyLeft, CanContinueCopyingDataCallback))
 		{
 			return;
 		}
-		if (!ContinueTimeSliceCopy(From.Leaves, Leaves, SizeToCopyLeft))
+		if (!ContinueTimeSliceCopy(From.Leaves, Leaves, SizeToCopyLeft, CanContinueCopyingDataCallback))
 		{
 			return;
 		}
-		if (!ContinueTimeSliceCopy(From.DirtyElements, DirtyElements, SizeToCopyLeft))
+		if (!ContinueTimeSliceCopy(From.DirtyElements, DirtyElements, SizeToCopyLeft, CanContinueCopyingDataCallback))
 		{
 			return;
 		}
 
-		if (!ContinueTimeSliceCopy(From.FlattenedCellArrayOfDirtyIndices, FlattenedCellArrayOfDirtyIndices, SizeToCopyLeft))
+		if (!ContinueTimeSliceCopy(From.FlattenedCellArrayOfDirtyIndices, FlattenedCellArrayOfDirtyIndices, SizeToCopyLeft, CanContinueCopyingDataCallback))
 		{
 			return;
 		}
-		if (!ContinueTimeSliceCopy(From.DirtyElementsGridOverflow, DirtyElementsGridOverflow, SizeToCopyLeft))
+		if (!ContinueTimeSliceCopy(From.DirtyElementsGridOverflow, DirtyElementsGridOverflow, SizeToCopyLeft, CanContinueCopyingDataCallback))
 		{
 			return;
 		}
-		if (!ContinueTimeSliceCopy(From.GlobalPayloads, GlobalPayloads, SizeToCopyLeft))
+		if (!ContinueTimeSliceCopy(From.GlobalPayloads, GlobalPayloads, SizeToCopyLeft, CanContinueCopyingDataCallback))
 		{
 			return;
 		}
-		if (!ContinueTimeSliceCopy(From.PayloadToInfo, PayloadToInfo, SizeToCopyLeft))
+		if (!ContinueTimeSliceCopy(From.PayloadToInfo, PayloadToInfo, SizeToCopyLeft, CanContinueCopyingDataCallback))
 		{
 			return;
 		}
-		if (!ContinueTimeSliceCopy(From.OverlappingLeaves, OverlappingLeaves, SizeToCopyLeft))
+		if (!ContinueTimeSliceCopy(From.OverlappingLeaves, OverlappingLeaves, SizeToCopyLeft, CanContinueCopyingDataCallback))
 		{
 			return;
 		}
-		if (!ContinueTimeSliceCopy(From.OverlappingOffsets, OverlappingOffsets, SizeToCopyLeft))
+		if (!ContinueTimeSliceCopy(From.OverlappingOffsets, OverlappingOffsets, SizeToCopyLeft, CanContinueCopyingDataCallback))
 		{
 			return;
 		}
-		if (!ContinueTimeSliceCopy(From.OverlappingCounts, OverlappingCounts, SizeToCopyLeft))
+		if (!ContinueTimeSliceCopy(From.OverlappingCounts, OverlappingCounts, SizeToCopyLeft, CanContinueCopyingDataCallback))
 		{
 			return;
 		}
-		if (!ContinueTimeSliceCopy(From.OverlappingPairs, OverlappingPairs, SizeToCopyLeft))
+		if (!ContinueTimeSliceCopy(From.OverlappingPairs, OverlappingPairs, SizeToCopyLeft, CanContinueCopyingDataCallback))
 		{
 			return;
 		}
@@ -3129,6 +3202,9 @@ private:
 		TreeExpensiveStats.Reset();
 		PayloadToInfo.Reset();
 		NumProcessedThisSlice = 0;
+
+		StartSliceTimeStamp = FPlatformTime::Seconds();
+
 		GetCVars();  // Safe to copy CVARS here
 
 		if (bDynamicTree)
@@ -3298,12 +3374,79 @@ private:
 		FSplitInfo SplitInfos[2];
 	};
 
-	void FindBestBounds(const int32 StartElemIdx, const int32 LastElem, FWorkSnapshot& CurrentSnapshot, int32 MaxAxis, const TVec3<T>& SplitCenter)
+	int32 GetLastIndexToProcess(int32 CurrentIndex)
+	{
+		int32 LastNodeToProcessIndex = 0;
+		if (FAABBTimeSliceCVars::bUseTimeSliceMillisecondBudget)
+		{
+			// When TimeSlicing by a millisecond budget, try to process all nodes.
+			// We will stop if we ran out of time and continue on the next frame
+			LastNodeToProcessIndex = WorkPool[CurrentIndex].Elems.Num();
+		}
+		else
+		{
+			const bool WeAreTimeslicing = (MaxNumToProcess > 0);
+			const int32 NumWeCanProcess = MaxNumToProcess - NumProcessedThisSlice;
+			LastNodeToProcessIndex = WeAreTimeslicing ? FMath::Min(WorkPool[CurrentIndex].BestBoundsCurIdx + NumWeCanProcess, WorkPool[CurrentIndex].Elems.Num()) : WorkPool[CurrentIndex].Elems.Num();
+		}
+
+		return LastNodeToProcessIndex;
+	}
+
+	bool CanContinueProcessingNodes(bool bOnlyUseTimeStampCheck = true)
+	{
+		bool bCanDoWork = true;
+
+		if (FAABBTimeSliceCVars::bUseTimeSliceMillisecondBudget)
+		{
+			bool bCheckIfHasAvailableTime = false;
+
+			if (bOnlyUseTimeStampCheck)
+			{
+				bCheckIfHasAvailableTime = true;
+			}
+			else
+			{
+				// Checking for platform time every time is expensive, so only do it between a configured amount of elements.
+				// This means we might overshoot the budget, but on the other hand we will keep most of the time doing actual work.
+				CurrentProcessedNodesSinceChecked++;
+				if (CurrentProcessedNodesSinceChecked > FAABBTimeSliceCVars::MinNodesChunkToProcessBetweenTimeChecks)
+				{
+					CurrentProcessedNodesSinceChecked = 0;
+					bCheckIfHasAvailableTime = true;
+				}	
+			}
+	
+			if (bCheckIfHasAvailableTime)
+			{
+				const double ElapsedTime = FPlatformTime::Seconds() - StartSliceTimeStamp;
+				const bool WeAreTimeslicing = FAABBTimeSliceCVars::MaxProcessingTimePerSliceSeconds > 0 && MaxNumToProcess > 0;	
+				if (WeAreTimeslicing && !FMath::IsNearlyZero(StartSliceTimeStamp) && ElapsedTime > FAABBTimeSliceCVars::MaxProcessingTimePerSliceSeconds)
+				{
+					// done enough
+					bCanDoWork = false; 
+				}
+			}
+		}
+		else
+		{
+			const bool WeAreTimeslicing = (MaxNumToProcess > 0);
+			if (WeAreTimeslicing && (NumProcessedThisSlice >= MaxNumToProcess))
+			{
+				// done enough
+				bCanDoWork = false;  
+			}
+		}
+
+		return bCanDoWork;
+	}
+
+	void FindBestBounds(const int32 StartElemIdx, int32& InOutLastElem, FWorkSnapshot& CurrentSnapshot, int32 MaxAxis, const TVec3<T>& SplitCenter)
 	{
 		const T SplitVal = SplitCenter[MaxAxis];
 
 		// add all elements to one of the two split infos at this level - root level [ not taking into account the max number allowed or anything
-		for(int32 ElemIdx = StartElemIdx; ElemIdx < LastElem; ++ElemIdx)
+		for (int32 ElemIdx = StartElemIdx; ElemIdx < InOutLastElem; ++ElemIdx)
 		{
 			const FElement& Elem = CurrentSnapshot.Elems[ElemIdx];
 			int32 BoxIdx = 0;
@@ -3333,15 +3476,22 @@ private:
 			TVec3<T> CenterDelta = ElemCenter - WorkSnapshot.AverageCenter;
 			WorkSnapshot.AverageCenter += CenterDelta / NumElems;
 			WorkSnapshot.ScaledCenterVariance += (ElemCenter - WorkSnapshot.AverageCenter) * CenterDelta;
+
+			constexpr bool bOnlyUSeTimeStampCheck = false;
+			if (!CanContinueProcessingNodes(bOnlyUSeTimeStampCheck))
+			{
+				// If we ended before processing all the requested nodes, update the out last element index variable
+				const bool bIsProcessingLastRequestedIndex = ElemIdx == InOutLastElem - 1;
+				InOutLastElem = bIsProcessingLastRequestedIndex ? InOutLastElem : ElemIdx + 1;
+				break; // done enough
+			}
 		}
 
-		NumProcessedThisSlice += LastElem - StartElemIdx;
+		NumProcessedThisSlice += InOutLastElem - StartElemIdx;
 	}
 	
 	void SplitNode()
-	{
-		const bool WeAreTimeslicing = (MaxNumToProcess > 0);
-
+	{		
 		while (WorkStack.Num())
 		{
 			//NOTE: remember to be careful with this since it's a pointer on a tarray
@@ -3363,7 +3513,7 @@ private:
 				Nodes.AddDefaulted((1 + NewNodeIdx) - Nodes.Num());
 			}
 
-			if (WeAreTimeslicing && (NumProcessedThisSlice >= MaxNumToProcess))
+			if (!CanContinueProcessingNodes())
 			{
 				return; // done enough
 			}
@@ -3426,12 +3576,17 @@ private:
 					WorkPool[SecondChildIdx].AverageCenter = TVec3<T>(0);
 					WorkPool[SecondChildIdx].ScaledCenterVariance = TVec3<T>(0);
 				}
+				
+				if (!CanContinueProcessingNodes())
+				{
+					// done enough
+					return; 
+				}
 			}
 
 			if (WorkPool[CurIdx].TimeslicePhase == eTimeSlicePhase::DuringFindBestBounds)
 			{
-				const int32 NumWeCanProcess = MaxNumToProcess - NumProcessedThisSlice;
-				const int32 LastIdxToProcess = WeAreTimeslicing ? FMath::Min(WorkPool[CurIdx].BestBoundsCurIdx + NumWeCanProcess, WorkPool[CurIdx].Elems.Num()) : WorkPool[CurIdx].Elems.Num();
+				int32 LastIdxToProcess = GetLastIndexToProcess(CurIdx);
 
 				// Determine the axis to split the AABB on based on the SplitOnVarianceAxis console variable. If it is not 1, simply use the largest axis
 				// of the work snapshot bounds; otherwise, select the axis with the greatest center variance. Note that the variance times the number of
@@ -3447,14 +3602,14 @@ private:
 
 				FindBestBounds(WorkPool[CurIdx].BestBoundsCurIdx, LastIdxToProcess, WorkPool[CurIdx], MaxAxis, Center);
 				WorkPool[CurIdx].BestBoundsCurIdx = LastIdxToProcess;
-
-				if (WeAreTimeslicing && (NumProcessedThisSlice >= MaxNumToProcess))
+				
+				if (!CanContinueProcessingNodes())
 				{
-					return; // done enough
+					// done enough
+					return; 
 				}
 			}
 
-
 			const int32 FirstChildIdx = WorkPool[CurIdx].SplitInfos[0].WorkSnapshotIdx;
 			const int32 SecondChildIdx = WorkPool[CurIdx].SplitInfos[1].WorkSnapshotIdx;
 
@@ -3486,7 +3641,7 @@ private:
 				// create the actual node so that no one else can use our children node indices
 				const int32 HighestNodeIdx = Nodes[NewNodeIdx].ChildrenNodes[1];
 				Nodes.AddDefaulted((1 + HighestNodeIdx) - Nodes.Num());
-				
+			
 				WorkPool[CurIdx].TimeslicePhase = eTimeSlicePhase::ProcessingChildren;
 			}
 			else
@@ -3587,11 +3742,11 @@ private:
 	}
 
 	template<typename ContainerType>
-	static bool ContinueTimeSliceCopy(const ContainerType& ContainerFrom, ContainerType& ContainerTo, int32& InOutMaxSize)
+	static bool ContinueTimeSliceCopy(const ContainerType& ContainerFrom, ContainerType& ContainerTo, int32& InOutMaxSize, TFunctionRef<bool(int32,int32)> CanContinueCallback)
 	{
 		int32 SizeCopied = 0;
 
-		for (int32 Index = ContainerTo.Num(); Index < ContainerFrom.Num() && (InOutMaxSize < 0 || SizeCopied < InOutMaxSize); Index++)
+		for (int32 Index = ContainerTo.Num(); Index < ContainerFrom.Num() && CanContinueCallback(InOutMaxSize, SizeCopied); Index++)
 		{
 			AddToContainerHelper(ContainerFrom, ContainerTo, Index);
 			SizeCopied += ContainerElementSizeHelper(ContainerFrom, Index);
@@ -3661,6 +3816,7 @@ private:
 		, MaxPayloadBounds(Other.MaxPayloadBounds)
 		, MaxNumToProcess(Other.MaxNumToProcess)
 		, NumProcessedThisSlice(Other.NumProcessedThisSlice)
+		, StartSliceTimeStamp(Other.StartSliceTimeStamp)
 		, bModifyingTreeMultiThreadingFastCheck(Other.bModifyingTreeMultiThreadingFastCheck)
 		, bShouldRebuild(Other.bShouldRebuild)
 		, bBuildOverlapCache(Other.bBuildOverlapCache)		
@@ -3721,6 +3877,7 @@ private:
 			MaxPayloadBounds = Rhs.MaxPayloadBounds;
 			MaxNumToProcess = Rhs.MaxNumToProcess;
 			NumProcessedThisSlice = Rhs.NumProcessedThisSlice;
+			StartSliceTimeStamp = Rhs.StartSliceTimeStamp;
 			bModifyingTreeMultiThreadingFastCheck = Rhs.bModifyingTreeMultiThreadingFastCheck;
 			bShouldRebuild = Rhs.bShouldRebuild;
 			bBuildOverlapCache = Rhs.bBuildOverlapCache;			
@@ -3777,8 +3934,12 @@ private:
 	int32 MaxTreeDepth;
 	T MaxPayloadBounds;
 	int32 MaxNumToProcess;
-
 	int32 NumProcessedThisSlice;
+
+	double StartSliceTimeStamp = 0.0;
+	int32 CurrentProcessedNodesSinceChecked = 0;
+	int32 CurrentDataElementsCopiedSinceLastCheck = 0;
+	
 	TArray<int32> WorkStack;
 	TArray<int32> WorkPoolFreeList;
 	TArray<FWorkSnapshot> WorkPool;
