Chaos - fix GC creation from skipping cluster particles following the check for no geometry (knock-on from 27980049)

diff --git a/Engine/Source/Runtime/Experimental/Chaos/Private/Chaos/StrainModification.cpp b/Engine/Source/Runtime/Experimental/Chaos/Private/Chaos/StrainModification.cpp
--- a/Engine/Source/Runtime/Experimental/Chaos/Private/Chaos/StrainModification.cpp
+++ b/Engine/Source/Runtime/Experimental/Chaos/Private/Chaos/StrainModification.cpp
@@ -72,7 +72,7 @@ static void ForEachRootChildParticle(const Chaos::FStrainedProxyAndRoot& ProxyAn
 			for (int32 RestChildIdx : *RestChildren)
 			{
 				Chaos::FPBDRigidClusteredParticleHandle* ChildHandle = ProxyGC->GetSolverParticleHandles()[RestChildIdx];
-				if (ChildHandle->Parent() == nullptr) { continue; }
+				if ((ChildHandle == nullptr) || (ChildHandle->Parent() == nullptr)) { continue; }
 
 				Func(ChildHandle);
 			}
diff --git a/Engine/Source/Runtime/Experimental/Chaos/Private/PhysicsProxy/GeometryCollectionPhysicsProxy.cpp b/Engine/Source/Runtime/Experimental/Chaos/Private/PhysicsProxy/GeometryCollectionPhysicsProxy.cpp
--- a/Engine/Source/Runtime/Experimental/Chaos/Private/PhysicsProxy/GeometryCollectionPhysicsProxy.cpp
+++ b/Engine/Source/Runtime/Experimental/Chaos/Private/PhysicsProxy/GeometryCollectionPhysicsProxy.cpp
@@ -515,7 +515,7 @@ void FGeometryCollectionPhysicsProxy::Initialize(Chaos::FPBDRigidsEvolutionBase
 	SolverParticleHandles.Init(nullptr, NumParticles);
 	
 	TBitArray<> EffectiveParticles;
-	NumEffectiveParticles = CalculateEffectiveParticles(DynamicCollection, NumParticles, EffectiveParticles);
+	NumEffectiveParticles = CalculateEffectiveParticles(DynamicCollection, NumParticles, Parameters.MaxSimulatedLevel, Parameters.EnableClustering, GetOwner(), EffectiveParticles);
 
 	PhysicsObjects.Empty();
 	PhysicsObjects.Reserve(NumParticles);
@@ -1080,7 +1080,7 @@ void FGeometryCollectionPhysicsProxy::InitializeBodiesPT(Chaos::FPBDRigidsSolver
 
 		// Here Clean up Additional particles
 		TBitArray<> EffectiveParticles;
-		NumEffectiveParticles = CalculateEffectiveParticles(DynamicCollection, DynamicCollection.NumElements(FGeometryCollection::TransformGroup), EffectiveParticles);
+		NumEffectiveParticles = CalculateEffectiveParticles(DynamicCollection, DynamicCollection.NumElements(FGeometryCollection::TransformGroup), Parameters.MaxSimulatedLevel, Parameters.EnableClustering, GetOwner(), EffectiveParticles);
 
 		CreateNonClusteredParticles(RigidsSolver, *RestCollection, DynamicCollection, EffectiveParticles);
 
@@ -1320,6 +1320,7 @@ void FGeometryCollectionPhysicsProxy::InitializeBodiesPT(Chaos::FPBDRigidsSolver
 
 					// Hook the handle up with the GT particle
 					FParticle* GTParticle = GTParticles[TransformGroupIndex].Get();
+					check(GTParticle != nullptr);
 
 					Chaos::FUniqueIdx ExistingIndex = GTParticle->UniqueIdx();
 					Chaos::FPBDRigidClusteredParticleHandle* Handle = nullptr;
@@ -2631,20 +2632,90 @@ int32 FGeometryCollectionPhysicsProxy::CalculateHierarchyLevel(const FGeometryDy
 	return Level;
 }
 
+TBitArray<> FGeometryCollectionPhysicsProxy::CalculateClustersToCreateFromChildren(const FGeometryDynamicCollection& DynamicCollection, int32 NumTransforms)
+{
+	// Generate a bitmask of all the particles that are clusters without any collision geometry, that have children with collision geometry
+	// This is a workaround following the fix to not create particles that have no geometry. If bGeometryCollectionAlwaysGenerateGTCollisionForClusters
+	// is set, we will be auto-generating cluster geometry from their children (if they have geometry)
+	// @todo(chaos): this dupes functionality in FGeometryCollectionPhysicsProxy::Initialize
+	TBitArray<> ClustersToGenerate;
+
+	if (bGeometryCollectionAlwaysGenerateGTCollisionForClusters)
+	{
+		ClustersToGenerate.Init(false, NumTransforms);
+
+		// All the leaf objects
+		TArray<int32> ChildrenToCheckForParentFix;
+		for (int32 Index = 0; Index < NumTransforms; ++Index)
+		{
+			if (DynamicCollection.Children[Index].Num() == 0)
+			{
+				ChildrenToCheckForParentFix.Add(Index);
+			}
+		}
+
+		TSet<int32> ParentToPotentiallyFix;
+		while (ChildrenToCheckForParentFix.Num())
+		{
+			// step 1 : find parents
+			for (const int32 ChildIndex : ChildrenToCheckForParentFix)
+			{
+				const int32 ParentIndex = DynamicCollection.Parent[ChildIndex];
+				if (ParentIndex != INDEX_NONE)
+				{
+					ParentToPotentiallyFix.Add(ParentIndex);
+				}
+			}
+
+			// step 2: test the parent for having children with geometry
+			for (const int32 ParentToFixIndex : ParentToPotentiallyFix)
+			{
+				const bool bParentHasCollision = (DynamicCollection.Implicits[ParentToFixIndex] != nullptr) || ClustersToGenerate[ParentToFixIndex];
+				if (!bParentHasCollision)
+				{
+					// let's make sure all our children have an implicit defined, otherwise, postpone to next iteration 
+					bool bAllChildrenHaveCollision = true;
+					for (int32 ChildIndex : DynamicCollection.Children[ParentToFixIndex])
+					{
+						// defer if any of the children is a cluster with no collision yet generated 
+						const bool bChildHasCollision = (DynamicCollection.Implicits[ChildIndex] != nullptr) || ClustersToGenerate[ChildIndex];
+						if (!bChildHasCollision && (DynamicCollection.Children[ChildIndex].Num() > 0))
+						{
+							bAllChildrenHaveCollision = false;
+							break;
+						}
+					}
+
+					ClustersToGenerate[ParentToFixIndex] = bAllChildrenHaveCollision;
+				}
+			}
+
+			// step 3 : make the parent the new child to go up the hierarchy and continue the fixing
+			ChildrenToCheckForParentFix = ParentToPotentiallyFix.Array();
+			ParentToPotentiallyFix.Reset();
+		}
+	}
+
+	return ClustersToGenerate;
+}
 
-int32 FGeometryCollectionPhysicsProxy::CalculateEffectiveParticles(const FGeometryDynamicCollection& DynamicCollection, int32 NumTransform, TBitArray<>& EffectiveParticles) const
+int32 FGeometryCollectionPhysicsProxy::CalculateEffectiveParticles(const FGeometryDynamicCollection& DynamicCollection, int32 NumTransform, int32 InMaxSimulatedLevel, bool bEnableClustering, const UObject* Owner, TBitArray<>& EffectiveParticles)
 {
+	TBitArray<> ClustersUsingChildGeometry = CalculateClustersToCreateFromChildren(DynamicCollection, NumTransform);
+
 	int32 NumMissingGeometry = 0;
 	int32 NumEffectiveParticlesFound = 0;
 	EffectiveParticles.Init(false, NumTransform);
-	const int32 MaxSimulatedLevel = FMath::Min(GlobalMaxSimulatedLevel, Parameters.MaxSimulatedLevel);
+	const int32 MaxSimulatedLevel = FMath::Min(GlobalMaxSimulatedLevel, InMaxSimulatedLevel);
 	for (int32 TransformIndex = 0; TransformIndex < NumTransform; ++TransformIndex)
 	{
 		const int32 Level = FMath::Clamp(CalculateHierarchyLevel(DynamicCollection, TransformIndex), 0, INT_MAX);
-		if (Level <= MaxSimulatedLevel || !Parameters.EnableClustering)
+		if (Level <= MaxSimulatedLevel || !bEnableClustering)
 		{
+			const bool bIsClusterUsingChildGeometry = (ClustersUsingChildGeometry.Num() > 0) && ClustersUsingChildGeometry[TransformIndex];
 			const bool bHasGeometry = DynamicCollection.Implicits[TransformIndex].IsValid();
-			if (bHasGeometry)
+
+			if (bHasGeometry || bIsClusterUsingChildGeometry)
 			{
 				EffectiveParticles[TransformIndex] = true;
 				NumEffectiveParticlesFound++;
@@ -2658,7 +2729,17 @@ int32 FGeometryCollectionPhysicsProxy::CalculateEffectiveParticles(const FGeomet
 		}
 	}
 
-	UE_CLOG(NumMissingGeometry > 0, LogChaos, Error, TEXT("Geometry collection %s tried to create %d particles with no geometry"), *GetOwner()->GetFullName(), NumMissingGeometry);
+	if (NumMissingGeometry > 0)
+	{
+		if (Owner != nullptr)
+		{
+			UE_LOG(LogChaos, Error, TEXT("Geometry collection %s tried to create %d particles with no geometry"), *Owner->GetFullName(), NumMissingGeometry);
+		}
+		else
+		{
+			UE_LOG(LogChaos, Error, TEXT("Geometry collection tried to create %d particles with no geometry"), NumMissingGeometry);
+		}
+	}
 
 	return NumEffectiveParticlesFound;
 }
diff --git a/Engine/Source/Runtime/Experimental/Chaos/Public/PhysicsProxy/GeometryCollectionPhysicsProxy.h b/Engine/Source/Runtime/Experimental/Chaos/Public/PhysicsProxy/GeometryCollectionPhysicsProxy.h
--- a/Engine/Source/Runtime/Experimental/Chaos/Public/PhysicsProxy/GeometryCollectionPhysicsProxy.h
+++ b/Engine/Source/Runtime/Experimental/Chaos/Public/PhysicsProxy/GeometryCollectionPhysicsProxy.h
@@ -527,8 +527,6 @@ protected:
 	 */
 	static CHAOS_API int32 CalculateHierarchyLevel(const FGeometryDynamicCollection& DynamicCollection, int32 TransformIndex);
 
-	CHAOS_API int32 CalculateEffectiveParticles(const FGeometryDynamicCollection& DynamicCollection, int32 NumTransform, TBitArray<>& EffectiveParticles) const;
-
 	CHAOS_API void SetClusteredParticleKinematicTarget_Internal(Chaos::FPBDRigidClusteredParticleHandle* Handle, const FTransform& WorldTransform);
 
 	CHAOS_API void PrepareBufferData(Chaos::FDirtyGeometryCollectionData& BufferData, const FGeometryDynamicCollection& ThreadCollection,  Chaos::FReal SolverLastDt = 0.0);
@@ -541,6 +539,9 @@ protected:
 
 private:
 
+	static TBitArray<> CalculateClustersToCreateFromChildren(const FGeometryDynamicCollection& DynamicCollection, int32 NumTransforms);
+	static int32 CalculateEffectiveParticles(const FGeometryDynamicCollection& DynamicCollection, int32 NumTransform, int32 MaxSimulatedLevel, bool bEnableClustering, const UObject* Owner, TBitArray<>& EffectiveParticles);
+
 	/**
 	 * Since geometry collections only buffer data that has changed, when PullFromPhysicsState is given both PrevData and NextData it must
 	 * examine *both* PrevData and NextData for data about a particle (since that particle's data coudl be in PrevData and not NextData).
