Removed two-bytes sends

#rnx

diff --git a/Engine/Source/Runtime/Experimental/IoStoreOnDemand/Private/CoreHttp/Http.cpp b/Engine/Source/Runtime/Experimental/IoStoreOnDemand/Private/CoreHttp/Http.cpp
--- a/Engine/Source/Runtime/Experimental/IoStoreOnDemand/Private/CoreHttp/Http.cpp
+++ b/Engine/Source/Runtime/Experimental/IoStoreOnDemand/Private/CoreHttp/Http.cpp
@@ -1812,43 +1812,14 @@ static int32 DoSend(FActivity* Activity)
 {
 	TRACE_CPUPROFILER_EVENT_SCOPE(IasHttp::DoSend);
 
-	enum { PackBits = 4 };
-	uint32 Index = Activity->StateParam & ((1 << PackBits) - 1);
-	uint32 Remaining = Activity->StateParam >> PackBits;
-
 	FBuffer& Buffer = Activity->Buffer;
+	const char* SendData = Buffer.GetData();
+	int32 SendSize = Buffer.GetSize();
 
-	const char* SendData = nullptr;
-	uint32 SendSize = 0;
-	switch (Index)
-	{
-	case 0:
-		SendData = Buffer.GetData();
-		SendSize = Buffer.GetSize();
-		break;
-
-	case 1: {
-		SendData = "\r\n";
-		SendSize = 2;
-		break;
-		}
-	}
-
+	uint32 Remaining = Activity->StateParam;
 	SendData += Remaining;
 	SendSize -= Remaining;
-
-	if (SendSize == 0 || SendData == nullptr)
-	{
-		// It is expected there will be enough space for a RespInt object
-		Buffer.Reset();
-		Buffer.AdvanceUsed(sizeof(FResponseInternal));
-
-		Activity->StateParam = 0;
-		Activity->State = FActivity::EState::RecvMessage;
-		Activity->SocketWait = FActivity::EWait::Read;
-		Trace(Activity, ETrace::StateChangeWait, Activity->State);
-		return 1;
-	}
+	check(SendSize > 0);
 
 	Trace(Activity, ETrace::Send, SendSize);
 	int32 Result = send(Activity->Socket, SendData, SendSize, MsgFlagType(0));
@@ -1885,14 +1856,21 @@ static int32 DoSend(FActivity* Activity)
 	Remaining = SendSize - Result;
 	if (Remaining != 0)
 	{
-		Activity->StateParam = Index | (Remaining << PackBits);
+		Activity->StateParam = Remaining;
 		Activity->SocketWait = FActivity::EWait::Write;
 		Trace(Activity, ETrace::Wait);
 		return 1;
 	}
 
-	Activity->StateParam = Index + 1;
-	return DoSend(Activity);
+	// It is expected there will be enough space for a RespInt object
+	Buffer.Reset();
+	Buffer.AdvanceUsed(sizeof(FResponseInternal));
+
+	Activity->StateParam = 0;
+	Activity->State = FActivity::EState::RecvMessage;
+	Activity->SocketWait = FActivity::EWait::Read;
+	Trace(Activity, ETrace::StateChangeWait, Activity->State);
+	return 1;
 }
 
 ////////////////////////////////////////////////////////////////////////////////
@@ -2365,7 +2343,7 @@ FRequest FEventLoop::FImpl::Request(
 
 	FMessageBuilder Builder(Activity->Buffer);
 
-	Builder << Method << " " << Path << " HTTP/1.1\r\n"
+	Builder << Method << " " << Path << " HTTP/1.1" "\r\n"
 		"Host: " << Activity->Pool->GetHostName() << "\r\n";
 
 	// HTTP/1.1 is persistent by default thus "Connection" header isn't required
@@ -2384,6 +2362,8 @@ FTicket FEventLoop::FImpl::Send(FActivity* Activity)
 {
 	Trace(Activity, ETrace::RequestBegin);
 
+	FMessageBuilder(Activity->Buffer) << "\r\n";
+
 	uint64 Slot;
 	{
 		FScopeLock _(&Lock);
