Chaos - make Island sleep state propagate to isolated particles when their island is destroyed (we don't keep islands that have no constraints, and isolated particles store their own sleep state)
- fix debug-only ensure (was invalid as a result of a prior change)

diff --git a/Engine/Source/Runtime/Experimental/Chaos/Private/Chaos/Island/IslandManager.cpp b/Engine/Source/Runtime/Experimental/Chaos/Private/Chaos/Island/IslandManager.cpp
--- a/Engine/Source/Runtime/Experimental/Chaos/Private/Chaos/Island/IslandManager.cpp
+++ b/Engine/Source/Runtime/Experimental/Chaos/Private/Chaos/Island/IslandManager.cpp
@@ -683,8 +683,6 @@ namespace Chaos::Private
 	{
 		SCOPE_CYCLE_COUNTER(STAT_IslandManager_UpdateParticles);
 
-		const int32 VisitEpoch = GetNextVisitEpoch();
-
 		// Process state changed from registered particles.
 		// To reduce cache-misses it would be nice to iterate over transient particle handles 
 		// here but the view we need (Particles.GetActiveDynamicMovingKinematicParticlesView()) 
@@ -1239,6 +1237,27 @@ namespace Chaos::Private
 		}
 	}
 
+	void FPBDIslandManager::DestroyIslandNodes(FPBDIsland* Island)
+	{
+		// Destroy any nodes left in the island
+		for (int32 IslandNodeIndex = Island->Nodes.Num() - 1; IslandNodeIndex >= 0; --IslandNodeIndex)
+		{
+			FPBDIslandParticle* Node = Island->Nodes[IslandNodeIndex];
+			check(Node->Edges.IsEmpty());
+
+			// If there was a dynamic particle left in the island, we need to transfer the island 
+			// sleep state to it for use in ProcessParticlesSleep()
+			// @todo(chaos): not great - can we clean this up?
+			if (FPBDRigidParticleHandle* Rigid = Node->GetParticle()->CastToRigidParticle())
+			{
+				Rigid->SetSleepCounter(int8(Island->SleepCounter));
+			}
+
+			RemoveNodeFromIsland(Node);
+			DestroyGraphNode(Node);
+		}
+	}
+
 	void FPBDIslandManager::AddEdgeToIsland(FPBDIslandConstraint* Edge, FPBDIsland* Island)
 	{
 		check(Edge != nullptr);
@@ -1890,14 +1909,9 @@ namespace Chaos::Private
 			// so that we wake nodes the have been left on their own after all other nodes were removed.
 			if (Island->NumEdges == 0)
 			{
-				// Destroy any nodes left in the island
-				for (int32 IslandNodeIndex = Island->Nodes.Num() - 1; IslandNodeIndex >= 0; --IslandNodeIndex)
-				{
-					FPBDIslandParticle* Node = Island->Nodes[IslandNodeIndex];
-					check(Node->Edges.IsEmpty());
-					RemoveNodeFromIsland(Node);
-					DestroyGraphNode(Node);
-				}
+				// Destroy nodes that are left in the island
+				// NOTE: also copies back any island state needed by the particle (sleep counter)
+				DestroyIslandNodes(Island);
 
 				DestroyIsland(Island);
 				continue;
@@ -2237,7 +2251,9 @@ namespace Chaos::Private
 			CHAOS_CONSTRAINTGRAPH_VALIDATE_TEST(Node->Flags.bIsDynamic || (Node->Island == nullptr));
 
 			// Nodes without edges should have been removed
-			CHAOS_CONSTRAINTGRAPH_VALIDATE_TEST(Node->Edges.Num() > 0);
+			// Actually kinematics with no edges may now remain in the graph until the
+			// next UpdateParticles (i.e., the next tick) after their edges were removed
+			CHAOS_CONSTRAINTGRAPH_VALIDATE_TEST(!Node->Flags.bIsDynamic || (Node->Edges.Num() > 0));
 
 			if (Node->Island != nullptr)
 			{
diff --git a/Engine/Source/Runtime/Experimental/Chaos/Public/Chaos/Framework/PoolBackedArray.h b/Engine/Source/Runtime/Experimental/Chaos/Public/Chaos/Framework/PoolBackedArray.h
--- a/Engine/Source/Runtime/Experimental/Chaos/Public/Chaos/Framework/PoolBackedArray.h
+++ b/Engine/Source/Runtime/Experimental/Chaos/Public/Chaos/Framework/PoolBackedArray.h
@@ -77,6 +77,7 @@ namespace Chaos::Private
 			{
 				// Remove item from the array using the index we set in Alloc
 				const int32 Index = FItemAdapter::GetArrayIndex(*Item);
+				check(Items[Index] == Item);
 				Items.RemoveAtSwap(Index, 1, false);
 				FItemAdapter::SetArrayIndex(*Item, INDEX_NONE);
 
diff --git a/Engine/Source/Runtime/Experimental/Chaos/Public/Chaos/Island/IslandManager.h b/Engine/Source/Runtime/Experimental/Chaos/Public/Chaos/Island/IslandManager.h
--- a/Engine/Source/Runtime/Experimental/Chaos/Public/Chaos/Island/IslandManager.h
+++ b/Engine/Source/Runtime/Experimental/Chaos/Public/Chaos/Island/IslandManager.h
@@ -501,6 +501,7 @@ namespace Chaos::Private
 		CHAOS_API void AssignEdgeIsland(FPBDIslandConstraint* Edge);
 		CHAOS_API void AddNodeToIsland(FPBDIslandParticle* Node, FPBDIsland* Island);
 		CHAOS_API void RemoveNodeFromIsland(FPBDIslandParticle* Node);
+		CHAOS_API void DestroyIslandNodes(FPBDIsland* Island);
 		CHAOS_API void AddEdgeToIsland(FPBDIslandConstraint* Edge, FPBDIsland* Island);
 		CHAOS_API void RemoveEdgeFromIsland(FPBDIslandConstraint* Edge);
 		CHAOS_API void WakeNodeIslands(const FPBDIslandParticle* Node);
