Encapsulate Event Lock logic.
Now event container have their own lock for their handles. This CL will allow to dig deeper into unreproducible crash

diff --git a/Engine/Source/Runtime/Experimental/Chaos/Private/EventManager.cpp b/Engine/Source/Runtime/Experimental/Chaos/Private/EventManager.cpp
--- a/Engine/Source/Runtime/Experimental/Chaos/Private/EventManager.cpp
+++ b/Engine/Source/Runtime/Experimental/Chaos/Private/EventManager.cpp
@@ -31,10 +31,10 @@ namespace Chaos
 	void FEventManager::UnregisterHandler(const EEventType& EventType, const void* InHandler)
 	{
 		const FEventID EventID = (FEventID)EventType;
-		ContainerLock.WriteLock();
+		ContainerLock.ReadLock();
 		checkf(EventID < EventContainers.Num(), TEXT("Unregistering event Handler for an event ID that does not exist"));
 		EventContainers[EventID]->UnregisterHandler(InHandler);
-		ContainerLock.WriteUnlock();
+		ContainerLock.ReadUnlock();
 	}
 
 	void FEventManager::FillProducerData(const Chaos::FPBDRigidsSolver* Solver, bool bResetData)
diff --git a/Engine/Source/Runtime/Experimental/Chaos/Public/EventManager.h b/Engine/Source/Runtime/Experimental/Chaos/Public/EventManager.h
--- a/Engine/Source/Runtime/Experimental/Chaos/Public/EventManager.h
+++ b/Engine/Source/Runtime/Experimental/Chaos/Public/EventManager.h
@@ -169,11 +169,13 @@ namespace Chaos
 		 */
 		~TEventContainer()
 		{
+			HandleLock.WriteLock();
 			for (FEventHandlerPtr Handler : HandlerArray)
 			{
 				delete Handler;
 				Handler = nullptr;
 			}
+			HandleLock.WriteUnlock();
 		}
 #
 		/**
@@ -181,6 +183,7 @@ namespace Chaos
 		 */
 		virtual void RegisterHandler(const FEventHandlerPtr& Handler)
 		{
+			HandleLock.WriteLock();
 			HandlerArray.AddUnique(Handler);
 			TArray<UObject*> ProxyOwners;
 			bool bValidProxyFilter = Handler->GetInterestedProxyOwners(ProxyOwners);
@@ -197,8 +200,9 @@ namespace Chaos
 				if (GetProxyToIndexMap(EventBuffer.Get()->GetConsumerBuffer()) != nullptr) // Only if our type supports getting the ProxyToIndexMap do we bother adding this
 				{
 					HandlersNotInMap.AddUnique(Handler);
-				}				
+				}
 			}
+			HandleLock.WriteUnlock();
 		}
 
 		/**
@@ -206,6 +210,7 @@ namespace Chaos
 		 */
 		virtual void UnregisterHandler(const void* InHandler)
 		{
+			HandleLock.WriteLock();
 			TArray<TPair<UObject*, FEventHandlerPtr>> KeysAndValuesToRemove;
 			for (TPair<UObject*, FEventHandlerPtr>& KeyValue : ProxyOwnerToHandlerMap)
 			{
@@ -225,7 +230,6 @@ namespace Chaos
 			{
 				if (HandlersNotInMap[i]->GetHandler() == InHandler)
 				{
-
 					HandlersNotInMap.RemoveAtSwap(i, 1, false);
 					break;
 				}
@@ -240,6 +244,7 @@ namespace Chaos
 					break;
 				}
 			}
+			HandleLock.WriteUnlock();
 		}
 
 		/*
@@ -288,6 +293,7 @@ namespace Chaos
 				return;
 			}
 
+			HandleLock.ReadLock();
 			const TMap<IPhysicsProxyBase*, TArray<int32>>* Map = GetProxyToIndexMap(Buffer); // Use t his map to get all proxies used in the event buffer
 			// Only take this path if we have fewer Events than Handlers
 			if (Map && Map->Num() + HandlersNotInMap.Num() < HandlerArray.Num())
@@ -303,25 +309,25 @@ namespace Chaos
 					}
 				}
 
-				for (FEventHandlerPtr Handler : UniqueHandlers)
+				for (const FEventHandlerPtr Handler : UniqueHandlers)
 				{
 					Handler->HandleEvent(Buffer);
 				}
 
-				for (FEventHandlerPtr Handler : HandlersNotInMap)
+				for (const FEventHandlerPtr Handler : HandlersNotInMap)
 				{
 					Handler->HandleEvent(Buffer);
 				}
-				return;
 			}
-			
+			else
 			// This path is taken if there are fewer Handlers than events or the handler does not support GetProxyToIndexMap
 			{
-				for (FEventHandlerPtr Handler : HandlerArray)
+				for (const FEventHandlerPtr Handler : HandlerArray)
 				{
 					Handler->HandleEvent(Buffer);
 				}
-			}			
+			}
+			HandleLock.ReadUnlock();
 		}
 
 	private:
@@ -349,6 +355,8 @@ namespace Chaos
 		 * Delegate function registered to handle this event when it is dispatched
 		 */
 		TArray<FEventHandlerPtr> HandlerArray;
+
+		FRWLock HandleLock; // protect access ProxyOwnerToHandlerMap, HandlersNotInMap, HandlerArray
 	};
 
 	/**
@@ -433,10 +441,10 @@ namespace Chaos
 			}
 			else
 			{
-				ContainerLock.WriteLock();
+				ContainerLock.ReadLock();
 				checkf(EventID < EventContainers.Num(), TEXT("Registering event Handler for an event ID that does not exist"));
 				EventContainers[EventID]->RegisterHandler(new TRawEventHandler<PayloadType, HandlerType>(Handler, HandlerFunction, InterestedProxyOwnerFunction));
-				ContainerLock.WriteUnlock();
+				ContainerLock.ReadUnlock();
 			}	
 		}
 
