Fix build from CL 28043821

diff --git a/Engine/Source/Runtime/Experimental/Chaos/Public/GeometryCollectionProxyData.h b/Engine/Source/Runtime/Experimental/Chaos/Public/GeometryCollectionProxyData.h
--- a/Engine/Source/Runtime/Experimental/Chaos/Public/GeometryCollectionProxyData.h
+++ b/Engine/Source/Runtime/Experimental/Chaos/Public/GeometryCollectionProxyData.h
@@ -28,9 +28,9 @@ public:
 	FTransformDynamicCollection(FTransformDynamicCollection&&) = delete;
 	FTransformDynamicCollection& operator=(FTransformDynamicCollection&&) = delete;
 
-	const FTransform3f& GetTransform(int32 Index) const;
-	void SetTransform(int32 Index, const FTransform3f& Transform);
-	const TArray<FTransform3f>& GetTransforms() const;
+	CHAOS_API const FTransform3f& GetTransform(int32 Index) const;
+	CHAOS_API void SetTransform(int32 Index, const FTransform3f& Transform);
+	CHAOS_API const TArray<FTransform3f>& GetTransforms() const;
 
 	// Transform Group
 	TManagedArray<int32>        Parent;
