GC Replication improvement
- This is behind a cvar (p.Chaos.GC.UseReplicationV2) ( off for now )
- Split replicated data in two : state and dynamic
     - State data contains a bitArray of what's broken and an array of recently released pieces with matching breaking velocities
     - Dynamic data contains tracked ( XR and VW ) dynamic pieces
- Reduced the precision for quantized breaking velocities
- Store rotation as euler to allow for quantized rotation in dynamic data
- Breaking velocities are removed as particles become disabled to save bandwidth for players joining later or coming into relevancy
- Tried to share as much code as possible between the old and new replication code paths

diff --git a/Engine/Source/Runtime/Experimental/GeometryCollectionEngine/Private/GeometryCollection/GeometryCollectionComponent.cpp b/Engine/Source/Runtime/Experimental/GeometryCollectionEngine/Private/GeometryCollection/GeometryCollectionComponent.cpp
--- a/Engine/Source/Runtime/Experimental/GeometryCollectionEngine/Private/GeometryCollection/GeometryCollectionComponent.cpp
+++ b/Engine/Source/Runtime/Experimental/GeometryCollectionEngine/Private/GeometryCollection/GeometryCollectionComponent.cpp
@@ -131,6 +131,8 @@ FAutoConsoleVariableRef CVarGeometryCollectionEmitRootBreakingEvent(TEXT("p.Chao
 bool GeometryCollectionCreatePhysicsStateInEditor = false;
 FAutoConsoleVariableRef CVarGeometryCollectionCreatePhysicsStateInEditor(TEXT("p.Chaos.GC.CreatePhysicsStateInEditor"), GeometryCollectionCreatePhysicsStateInEditor, TEXT("when on , physics state for a GC will be create in editor ( non PIE )"));
 
+bool GeometryCollectionUseReplicationV2 = false;
+FAutoConsoleVariableRef CVarGeometryCollectionUseReplicationV2(TEXT("p.Chaos.GC.UseReplicationV2"), GeometryCollectionUseReplicationV2, TEXT("When true use new replication data model"));
 
 DEFINE_LOG_CATEGORY_STATIC(UGCC_LOG, Error, All);
 
@@ -242,6 +244,146 @@ bool FGeometryCollectionRepData::NetSerialize(FArchive& Ar, class UPackageMap* M
 	return true;
 }
 
+bool FGeometryCollectionRepStateData::SetBroken(int32 TransformIndex, int32 NumTransforms, bool bDisabled, const FVector& LinV, const FVector& AngVInRadiansPerSecond)
+{
+	if (BrokenState.Num() != NumTransforms)
+	{
+		BrokenState.SetNum(NumTransforms, false);
+	}
+	if (!BrokenState[TransformIndex])
+	{
+		BrokenState[TransformIndex] = true;
+
+		FReleasedData Data;
+		Data.TransformIndex = TransformIndex;
+		Data.LinearVelocity = LinV;
+		Data.AngularVelocityInDegreesPerSecond = FMath::RadiansToDegrees(AngVInRadiansPerSecond);
+		ReleasedData.Emplace(Data);
+
+		return true;
+	}
+
+	// breaking state already recorded, if it's disabled we can remove it from the ReleasedData
+	// because it certainly went throug the removal stage and clients do not need to care about the initial velocity anymore
+	if (bDisabled)
+	{
+		const int32 FoundDataIndex = ReleasedData.IndexOfByPredicate(
+			[&TransformIndex](const FReleasedData& Data) -> bool { return (Data.TransformIndex == TransformIndex); });
+		if (FoundDataIndex != INDEX_NONE)
+		{
+			ReleasedData.RemoveAtSwap(FoundDataIndex, 1, false /* bAllowShrinking */);
+		}
+		// this does not have to be reported as a state change to save bandwidth
+		return false;
+	}
+
+	return false;
+}
+
+bool FGeometryCollectionRepStateData::Identical(const FGeometryCollectionRepStateData* Other, uint32 PortFlags) const
+{
+	return Other && (Version == Other->Version);
+}
+
+bool FGeometryCollectionRepStateData::NetSerialize(FArchive& Ar, class UPackageMap* Map, bool& bOutSuccess)
+{
+	bOutSuccess = true;
+
+	Ar << Version;
+
+	Ar << BrokenState;
+
+	Ar << bIsRootAnchored;
+
+	// we only support up to 2^16 bones for replication ( bandwidth reasons )
+	ensure(ReleasedData.Num() < TNumericLimits<uint16>::Max());
+	uint16 NumReleasedData = static_cast<uint16>(ReleasedData.Num());
+	Ar << NumReleasedData;
+
+	if (Ar.IsLoading())
+	{
+		ReleasedData.SetNumUninitialized(static_cast<int32>(NumReleasedData));
+	}
+
+	for (FReleasedData& Data : ReleasedData)
+	{
+		Ar << Data.TransformIndex;
+		Ar << Data.LinearVelocity;
+		Ar << Data.AngularVelocityInDegreesPerSecond;
+	}
+
+	return true;
+}
+
+bool FGeometryCollectionRepDynamicData::FClusterData::IsEqualPositionsAndVelocities(const FGeometryCollectionRepDynamicData::FClusterData& Data) const
+{
+	return Position.Equals(Data.Position)
+		&& EulerRotation.Equals(Data.EulerRotation)
+		&& LinearVelocity.Equals(Data.LinearVelocity)
+		&& AngularVelocityInDegreesPerSecond.Equals(Data.AngularVelocityInDegreesPerSecond)
+		;
+}
+
+bool FGeometryCollectionRepDynamicData::SetData(const FGeometryCollectionRepDynamicData::FClusterData& Data)
+{
+	FClusterData* ExistingData = ClusterData.FindByPredicate(
+		[&TransformIndex = Data.TransformIndex](const FClusterData& Data) -> bool { return Data.TransformIndex == TransformIndex; });
+	if (ExistingData)
+	{
+		const bool bHasChanged = !ExistingData->IsEqualPositionsAndVelocities(Data);
+		if (bHasChanged)
+		{
+			*ExistingData = Data;
+			ExistingData->LastUpdatedVersion = (Version + 1);
+		}
+		return bHasChanged;
+	}
+	FClusterData& NewData = ClusterData.Emplace_GetRef(Data);
+	NewData.LastUpdatedVersion = (Version + 1);
+	return true;
+}
+
+bool FGeometryCollectionRepDynamicData::RemoveOutOfDateClusterData()
+{
+	const int32 NumRemoved = ClusterData.RemoveAllSwap(
+		[this](const FClusterData& Data) -> bool
+		{
+			return (Data.LastUpdatedVersion <= Version);
+		},
+		false /* bAllowShrinking */);
+
+	return (NumRemoved > 0);
+}
+
+bool FGeometryCollectionRepDynamicData::Identical(const FGeometryCollectionRepDynamicData* Other, uint32 PortFlags) const
+{
+	return Other && (Version == Other->Version);
+}
+
+bool FGeometryCollectionRepDynamicData::NetSerialize(FArchive& Ar, class UPackageMap* Map, bool& bOutSuccess)
+{
+	// we only support up to 2^16 bones for replication ( bandwidth reasons )
+	ensure(ClusterData.Num() < TNumericLimits<uint16>::Max());
+	uint16 NumClusterData = static_cast<uint16>(ClusterData.Num());
+	Ar << NumClusterData;
+
+	if (Ar.IsLoading())
+	{
+		ClusterData.SetNumUninitialized(static_cast<int32>(NumClusterData));
+	}
+
+	for (FClusterData& Data : ClusterData)
+	{
+		Ar << Data.TransformIndex;
+		Ar << Data.Position;
+		Ar << Data.EulerRotation; // as rotator
+		Ar << Data.LinearVelocity;
+		Ar << Data.AngularVelocityInDegreesPerSecond;
+	}
+
+	return true;
+}
+
 int32 GGeometryCollectionNanite = 1;
 FAutoConsoleVariableRef CVarGeometryCollectionNanite(
 	TEXT("r.GeometryCollection.Nanite"),
@@ -512,6 +654,12 @@ void UGeometryCollectionComponent::GetLifetimeReplicatedProps(TArray<FLifetimePr
 	Params.RepNotifyCondition = REPNOTIFY_OnChanged;
 	DOREPLIFETIME_WITH_PARAMS_FAST(UGeometryCollectionComponent, RepData, Params);*/
 	DOREPLIFETIME(UGeometryCollectionComponent, RepData);
+
+	FDoRepLifetimeParams Params;
+	Params.bIsPushBased = true;
+	Params.RepNotifyCondition = REPNOTIFY_OnChanged;
+	DOREPLIFETIME_WITH_PARAMS_FAST(UGeometryCollectionComponent, RepStateData, Params);
+	DOREPLIFETIME_WITH_PARAMS_FAST(UGeometryCollectionComponent, RepDynamicData, Params);
 }
 
 namespace
@@ -1726,6 +1874,8 @@ void UGeometryCollectionComponent::ResetRepData()
 {
 	ClustersToRep.Reset();
 	RepData.Reset();
+	RepStateData.Reset();
+	RepDynamicData.Reset();
 
 	// Those following data are initialized here from the Game Thread, 
 	// but they are read and written from the Physics Thread
@@ -1738,6 +1888,11 @@ void UGeometryCollectionComponent::ResetRepData()
 
 void UGeometryCollectionComponent::UpdateRepData()
 {
+	if (GeometryCollectionUseReplicationV2)
+	{
+		return UpdateRepStateAndDynamicData();
+	}
+	
 	check(GetNetMode() != ENetMode::NM_Client);
 	using namespace Chaos;
 	if(!bEnableReplication)
@@ -1799,7 +1954,6 @@ void UGeometryCollectionComponent::UpdateRepData()
 				bClustersChanged = true;
 			}
 		}
-		
 
 		//see if we have any new clusters that are enabled
 		TSet<FPBDRigidClusteredParticleHandle*> Processed;
@@ -1979,13 +2133,149 @@ void UGeometryCollectionComponent::UpdateRepData()
 	}
 }
 
-int32 GeometryCollectionHardMissingUpdatesSnapThreshold = 20;
-FAutoConsoleVariableRef CVarGeometryCollectionHardMissingUpdatesSnapThreshold(TEXT("p.GeometryCollectionHardMissingUpdatesSnapThreshold"), GeometryCollectionHardMissingUpdatesSnapThreshold,
-	TEXT("Determines how many missing updates before we trigger a hard snap"));
 
-int32 GeometryCollectionHardsnapThresholdMs = 100; // 10 Hz
-FAutoConsoleVariableRef CVarGeometryCollectionHardsnapThresholdMs(TEXT("p.GeometryCollectionHardsnapThresholdMs"), GeometryCollectionHardMissingUpdatesSnapThreshold,
-	TEXT("Determines how many ms since the last hardsnap to trigger a new one"));
+void UGeometryCollectionComponent::UpdateRepStateAndDynamicData()
+{
+	check(GetNetMode() != ENetMode::NM_Client);
+	using namespace Chaos;
+	if (!bEnableReplication || PhysicsProxy == nullptr)
+	{
+		return;
+	}
+
+	AActor* Owner = GetOwner();
+
+	// If we have no owner or our netmode means we never require replication then early out
+	if (!Owner || Owner->GetNetMode() == ENetMode::NM_Standalone)
+	{
+		return;
+	}
+
+	if (GetIsReplicated() && Owner->GetLocalRole() == ROLE_Authority)
+	{
+		if (RestCollection && RestCollection->GetGeometryCollection())
+		{
+			FPBDRigidsSolver* Solver = PhysicsProxy->GetSolver<Chaos::FPBDRigidsSolver>();
+			//const FRigidClustering& RigidClustering = Solver->GetEvolution()->GetRigidClustering();
+
+			// let go through all the transform and see which one has changed its state
+			const TArray<FPBDRigidClusteredParticleHandle*>& ParticleHandles = PhysicsProxy->GetParticles();
+			const int32 NumTransforms = ParticleHandles.Num();
+
+			const TManagedArray<int32>* InitialLevels = PhysicsProxy->GetPhysicsCollection().FindAttribute<int32>("InitialLevel", FGeometryCollection::TransformGroup);
+
+			bool bStateChanged = false;
+			bool bDynamicChanged = false;
+
+			// root level anchor state
+			const int32 InitialRootIndex = PhysicsProxy->GetSimParameters().InitialRootIndex;
+			if (FPBDRigidClusteredParticleHandle* RootHandle = ParticleHandles[InitialRootIndex])
+			{
+				const bool bIsRootAnchored = RootHandle->IsAnchored();
+				if (RepStateData.Version == 0 || bIsRootAnchored != (bool)RepStateData.bIsRootAnchored)
+				{
+					RepStateData.bIsRootAnchored = bIsRootAnchored? 1: 0;
+					bStateChanged = true;
+				}
+			}
+
+			struct FRootHandle
+			{
+				const FPBDRigidClusteredParticleHandle* Handle; // handle of the particle or internal cluster 
+				const int32 TransformIndex; // transform index of the particle or the child of the internal cluster 
+
+				bool operator==(const FRootHandle& Other) const	{ return Handle == Other.Handle; }
+			};
+
+			// contains all the dynamic root of the geometry collection
+			// this includes non-clusterunion internal clusters 
+			// using an array as the number of root should always be small 
+			TArray<FRootHandle> RootsToTrack;
+
+			const int32 RootIndex = GetRootIndex();
+
+			// go through particles and send replicated data 
+			for (int32 TransformIndex = 0; TransformIndex < NumTransforms; TransformIndex++)
+			{
+				if (FPBDRigidClusteredParticleHandle* ParticleHandle = ParticleHandles[TransformIndex])
+				{
+					const int32 Level = (InitialLevels && InitialLevels->IsValidIndex(TransformIndex)) ? (*InitialLevels)[TransformIndex] : INDEX_NONE;
+
+					const FPBDRigidClusteredParticleHandle* ParentHandle = ParticleHandle->Parent();
+
+					// no parent and not root means the particle has been broken  
+					const bool bIsBroken = (ParentHandle == nullptr && TransformIndex != GetRootIndex());
+					if (bIsBroken)
+					{
+						// only record the one at the abandon level and before ( child break of abandon level clusters must be recorded )  
+						if (!bEnableAbandonAfterLevel || Level <= ReplicationAbandonAfterLevel)
+						{
+							bStateChanged |= RepStateData.SetBroken(TransformIndex, NumTransforms, ParticleHandle->Disabled(), ParticleHandle->V(), ParticleHandle->W());
+						}
+
+						// anything beyond ReplicationAbandonAfterLevel must be disabled to save server CPU as they are now client authoritative
+						if (bEnableAbandonAfterLevel && Level >= (ReplicationAbandonAfterLevel + 1))
+						{
+							if (!ParticleHandle->Disabled())
+							{
+								Solver->GetEvolution()->DisableParticle(ParticleHandle);
+								Solver->GetParticles().MarkTransientDirtyParticle(ParticleHandle);
+							}
+						}
+					}
+
+					const bool bIsRoot = bIsBroken && (!ParticleHandle->Disabled());
+					const bool bHasInternalClusterParent = ParentHandle && (ParentHandle->InternalCluster()) && (ParentHandle->PhysicsProxy() == ParticleHandle->PhysicsProxy());
+
+					const FPBDRigidClusteredParticleHandle* RootParticle = bIsRoot ? ParticleHandle : (bHasInternalClusterParent ? ParentHandle : nullptr);
+					const int32 RootParticleLevel = (bHasInternalClusterParent)? (Level - 1): (Level);
+					const bool bTrackPosition = (RootParticleLevel <= this->ReplicationMaxPositionAndVelocityCorrectionLevel);
+					if (bTrackPosition && RootParticle)
+					{
+						RootsToTrack.AddUnique(FRootHandle{ RootParticle, TransformIndex });
+					}
+					
+				}
+			}
+
+			// let's update the data for the tracked clusters 
+			for (const FRootHandle& Root: RootsToTrack)
+			{
+				FGeometryCollectionRepDynamicData::FClusterData Data;
+				Data.TransformIndex = Root.TransformIndex;
+				Data.bIsInternalCluster = Root.Handle->InternalCluster();
+				Data.Position = Root.Handle->X();
+				Data.EulerRotation = Root.Handle->R().Rotator().Euler();
+				Data.LinearVelocity = Root.Handle->V();
+				Data.AngularVelocityInDegreesPerSecond = FMath::RadiansToDegrees(Root.Handle->W());
+				bDynamicChanged |= RepDynamicData.SetData(Data);
+			}
+
+			RepDynamicData.RemoveOutOfDateClusterData();
+
+			if (bStateChanged)
+			{
+				MARK_PROPERTY_DIRTY_FROM_NAME(UGeometryCollectionComponent, RepStateData, this);
+				++RepStateData.Version;
+			}
+
+			if (bDynamicChanged)
+			{
+				MARK_PROPERTY_DIRTY_FROM_NAME(UGeometryCollectionComponent, RepDynamicData, this);
+				++RepDynamicData.Version;
+			}
+
+			if (bDynamicChanged || bStateChanged)
+			{
+				if (Owner->NetDormancy != DORM_Awake)
+				{
+					//If net dormancy is Initial it must be for perf reasons, but since a cluster changed we need to replicate down
+					Owner->SetNetDormancy(DORM_Awake);
+				}
+			}
+		}
+	}
+}
 
 float GeometryCollectionRepLinearMatchStrength = 50;
 FAutoConsoleVariableRef CVarGeometryCollectionRepLinearMatchStrength(TEXT("p.GeometryCollectionRepLinearMatchStrength"), GeometryCollectionRepLinearMatchStrength, TEXT("Units can be interpreted as %/s^2 - acceleration of percent linear correction"));
@@ -2004,14 +2294,156 @@ bool bGeometryCollectionRepUseClusterVelocityMatch = true;
 FAutoConsoleVariableRef CVarGeometryCollectionRepUseClusterVelocityMatch(TEXT("p.bGeometryCollectionRepUseClusterVelocityMatch"), bGeometryCollectionRepUseClusterVelocityMatch,
 	TEXT("Use physical velocity to match cluster states"));
 
+int32 GeometryCollectionHardMissingUpdatesSnapThreshold = 20;
+FAutoConsoleVariableRef CVarGeometryCollectionHardMissingUpdatesSnapThreshold(TEXT("p.GeometryCollectionHardMissingUpdatesSnapThreshold"), GeometryCollectionHardMissingUpdatesSnapThreshold,
+	TEXT("Determines how many missing updates before we trigger a hard snap"));
+
+int32 GeometryCollectionHardsnapThresholdMs = 100; // 10 Hz
+FAutoConsoleVariableRef CVarGeometryCollectionHardsnapThresholdMs(TEXT("p.GeometryCollectionHardsnapThresholdMs"), GeometryCollectionHardMissingUpdatesSnapThreshold,
+	TEXT("Determines how many ms since the last hardsnap to trigger a new one"));
+
 void UGeometryCollectionComponent::ProcessRepData()
 {
 	bGeometryCollectionRepUseClusterVelocityMatch = false;
 	ProcessRepData(0.f, 0.f);
 }
 
-namespace {
-	bool ProcessRepDataCommon(const float DeltaTime, const float SimTime, const FGeometryCollectionRepData& RepData, FGeometryCollectionPhysicsProxy* PhysicsProxy, int32& VersionProcessed, int32& OneOffActivatedProcessed, double& LastHardsnapTimeInMs, float ReceivedTime, bool bReplicateMovement)
+namespace 
+{
+	bool ReplicationShouldHardSnap(int32 RepDataVersion, int32 VersionProcessed, bool bReplicateMovement, double& LastHardsnapTimeInMs)
+	{
+		bool bHardSnap = false;
+		if (bReplicateMovement)
+		{
+			const int64 CurrentTimeInMs = FPlatformTime::ToMilliseconds64(FPlatformTime::Cycles64());
+
+			// Always hard snap on the very first version received
+			if (VersionProcessed == 0)
+			{
+				bHardSnap = true;
+			}
+			else if (VersionProcessed < RepDataVersion)
+			{
+				//TODO: this will not really work if a fracture happens and then immediately goes to sleep without updating client enough times
+				//A time method would work better here, but is limited to async mode. Maybe we can support both
+				bHardSnap = (RepDataVersion - VersionProcessed) > GeometryCollectionHardMissingUpdatesSnapThreshold;
+
+				if (!bGeometryCollectionRepUseClusterVelocityMatch)
+				{
+					// When not doing velocity match for clusters, instead we do periodic hard snapping
+					bHardSnap |= (CurrentTimeInMs - LastHardsnapTimeInMs) > GeometryCollectionHardsnapThresholdMs;
+				}
+			}
+			else if (VersionProcessed > RepDataVersion)
+			{
+				//rollover so just treat as hard snap - this case is extremely rare and a one off
+				bHardSnap = true;
+			}
+
+			if (bHardSnap)
+			{
+				LastHardsnapTimeInMs = CurrentTimeInMs;
+			}
+		}
+		return bHardSnap;
+	}
+
+	void UpdateRootStateFromReplication(const FGeometryCollectionPhysicsProxy& PhysicsProxy, const bool bNewAnchoredState)
+	{
+		using namespace Chaos;
+
+		const int32 InitialRootIndex = PhysicsProxy.GetSimParameters().InitialRootIndex;
+		const TArray<FPBDRigidClusteredParticleHandle*>& ParticleHandles = PhysicsProxy.GetParticles();
+		if (FPBDRigidClusteredParticleHandle* RootHandle = ParticleHandles[InitialRootIndex])
+		{
+			if (RootHandle->Parent() || !RootHandle->Disabled())
+			{
+				if (bNewAnchoredState != RootHandle->IsAnchored())
+				{
+					RootHandle->SetIsAnchored(bNewAnchoredState);
+
+					if (FPBDRigidsSolver* Solver = PhysicsProxy.GetSolver<Chaos::FPBDRigidsSolver>())
+					{
+						if (FPBDRigidsEvolutionGBF* Evolution = Solver->GetEvolution())
+						{
+							FRigidClustering& RigidClustering = Solver->GetEvolution()->GetRigidClustering();
+							// NOTE: We have to start out with the particle as kinematic and let UpdateKinematicProperties correct it if this is wrong
+							Evolution->SetParticleObjectState(RootHandle, EObjectStateType::Kinematic);
+							UpdateKinematicProperties(RootHandle, RigidClustering.GetChildrenMap(), *Evolution);
+						}
+					}
+				}
+			}
+		}
+	}
+
+	// return the awaken state of the cluster 
+	bool UpdateClusterPositionAndVelocitiesFromReplication(
+		Chaos::FPBDRigidsSolver* Solver,
+		Chaos::FPBDRigidParticleHandle& Cluster,
+		const FVector& Position, const FQuat& Rotation, 
+		const FVector& LinearVelocity, const FVector& AngularVelocity,
+		const bool bHardSnap,
+		const double RepExtrapTime,
+		const double DeltaTime)
+	{
+		bool bWake = false;
+
+		if (bHardSnap)
+		{
+			Cluster.SetX(Position);
+			Cluster.SetR(Rotation);
+			Cluster.SetV(LinearVelocity);
+			Cluster.SetW(AngularVelocity);
+			bWake = true;
+		}
+		else if (bGeometryCollectionRepUseClusterVelocityMatch)
+		{
+			//
+			// Match linear velocity
+			//
+			const FVector RepVel = LinearVelocity;
+			const FVector RepExtrapPos = Position + (RepVel * RepExtrapTime);
+			const Chaos::FVec3 DeltaX = RepExtrapPos - Cluster.X();
+			const float DeltaXMagSq = DeltaX.SizeSquared();
+			if (DeltaXMagSq > SMALL_NUMBER && GeometryCollectionRepLinearMatchStrength > SMALL_NUMBER)
+			{
+				bWake = true;
+				//
+				// DeltaX * MatchStrength is an acceleration, m/s^2, which is integrated
+				// by multiplying by DeltaTime.
+				//
+				// It's formulated this way to get a larger correction for a longer time
+				// step, ie. correction velocities are framerate independent.
+				//
+				Cluster.SetV(RepVel + (DeltaX * GeometryCollectionRepLinearMatchStrength * DeltaTime));
+			}
+
+			//
+			// Match angular velocity
+			//
+			const FVector RepAngVel = AngularVelocity;
+			const Chaos::FRotation3 RepExtrapAng = Chaos::FRotation3::IntegrateRotationWithAngularVelocity(Rotation, RepAngVel, RepExtrapTime);
+			const FVector AngVel = Chaos::FRotation3::CalculateAngularVelocity(Cluster.R(), RepExtrapAng, GeometryCollectionRepAngularMatchTime);
+			if (AngVel.SizeSquared() > SMALL_NUMBER)
+			{
+				Cluster.SetW(RepAngVel + AngVel);
+				bWake = true;
+			}
+		}
+
+		//
+		// Wake up particle if it's sleeping and there's a delta to correct
+		//
+		if (bWake && Cluster.IsSleeping())
+		{
+			Solver->GetEvolution()->SetParticleObjectState(&Cluster, Chaos::EObjectStateType::Dynamic);
+		}
+
+		return bWake;
+	}
+
+	bool ProcessRepDataCommon(const float DeltaTime, const float SimTime, const FGeometryCollectionRepData& RepData, FGeometryCollectionPhysicsProxy* PhysicsProxy, int32& VersionProcessed, int32& OneOffActivatedProcessed, double& LastHardsnapTimeInMs, bool bReplicateMovement)
 	{
 		using namespace Chaos;
 
@@ -2020,6 +2452,14 @@ namespace {
 			return false;
 		}
 
+		float ReceivedTime = SimTime;
+
+		// Track the sim time that this rep data was received on.
+		if (RepData.RepDataReceivedTime.IsSet())
+		{
+			ReceivedTime = *RepData.RepDataReceivedTime;
+		}
+		
 		// How far we must extrapolate from when we received the data
 		const float RepExtrapTime = FMath::Max(0.f, SimTime - ReceivedTime);
 
@@ -2063,62 +2503,9 @@ namespace {
 		}
 
 		// Update the anchored state of the root particle
-		const int32 InitialRootIndex = PhysicsProxy->GetSimParameters().InitialRootIndex;
-		const TArray<FPBDRigidClusteredParticleHandle*>& ParticleHandles = PhysicsProxy->GetParticles();
-		if (FPBDRigidClusteredParticleHandle* RootHandle = ParticleHandles[InitialRootIndex])
-		{
-			if (RootHandle->Parent() || !RootHandle->Disabled())
-			{
-				if (RepData.bIsRootAnchored != RootHandle->IsAnchored())
-				{
-					RootHandle->SetIsAnchored(RepData.bIsRootAnchored);
-
-					if (FPBDRigidsSolver* Solver = PhysicsProxy->GetSolver<Chaos::FPBDRigidsSolver>())
-					{
-						if (FPBDRigidsEvolutionGBF* Evolution = Solver->GetEvolution())
-						{
-							FRigidClustering& RigidClustering = Solver->GetEvolution()->GetRigidClustering();
-							// NOTE: We have to start out with the particle as kinematic and let UpdateKinematicProperties correct it if this is wrong
-							Evolution->SetParticleObjectState(RootHandle, Chaos::EObjectStateType::Kinematic);
-							Chaos::UpdateKinematicProperties(RootHandle, RigidClustering.GetChildrenMap(), *Evolution);
-						}
-					}
-				}
-			}
-		}
-		bool bHardSnap = false;
-		if (bReplicateMovement)
-		{
-			const int64 CurrentTimeInMs = FPlatformTime::ToMilliseconds64(FPlatformTime::Cycles64());
+		UpdateRootStateFromReplication(*PhysicsProxy, RepData.bIsRootAnchored);
 
-			// Always hard snap on the very first version received
-			if (VersionProcessed == 0)
-			{
-				bHardSnap = true;
-			}
-			else if (VersionProcessed < RepData.Version)
-			{
-				//TODO: this will not really work if a fracture happens and then immediately goes to sleep without updating client enough times
-				//A time method would work better here, but is limited to async mode. Maybe we can support both
-				bHardSnap = (RepData.Version - VersionProcessed) > GeometryCollectionHardMissingUpdatesSnapThreshold;
-
-				if (!bGeometryCollectionRepUseClusterVelocityMatch)
-				{
-					// When not doing velocity match for clusters, instead we do periodic hard snapping
-					bHardSnap |= (CurrentTimeInMs - LastHardsnapTimeInMs) > GeometryCollectionHardsnapThresholdMs;
-				}
-			}
-			else if (VersionProcessed > RepData.Version)
-			{
-				//rollover so just treat as hard snap - this case is extremely rare and a one off
-				bHardSnap = true;
-			}
-
-			if (bHardSnap)
-			{
-				LastHardsnapTimeInMs = CurrentTimeInMs;
-			}
-		}
+		const bool bHardSnap = ReplicationShouldHardSnap(RepData.Version, VersionProcessed, bReplicateMovement, LastHardsnapTimeInMs);
 
 		FPBDRigidsSolver* Solver = PhysicsProxy->GetSolver<Chaos::FPBDRigidsSolver>();
 		FRigidClustering& RigidClustering = Solver->GetEvolution()->GetRigidClustering();
@@ -2163,70 +2550,148 @@ namespace {
 		bool bProcessed = false;
 		if (bReplicateMovement)
 		{
-			ForEachClusterPair([Solver, DeltaTime, RepExtrapTime, bHardSnap, &bProcessed](const FGeometryCollectionClusterRep& RepCluster, Chaos::FPBDRigidParticleHandle& Cluster)
+			ForEachClusterPair([Solver, DeltaTime, RepExtrapTime, bHardSnap, &bProcessed](const FGeometryCollectionClusterRep& RepData, Chaos::FPBDRigidParticleHandle& Cluster)
 			{
-				bool bWake = false;
+				bProcessed |= UpdateClusterPositionAndVelocitiesFromReplication(Solver, Cluster, RepData.Position, RepData.Rotation, RepData.LinearVelocity, RepData.AngularVelocity, bHardSnap, RepExtrapTime, DeltaTime);
+			});
+		}
 
-				if (bHardSnap)
-				{
-					Cluster.SetX(RepCluster.Position);
-					Cluster.SetR(RepCluster.Rotation);
-					Cluster.SetV(RepCluster.LinearVelocity);
-					Cluster.SetW(RepCluster.AngularVelocity);
-					bWake = true;
-				}
-				else if (bGeometryCollectionRepUseClusterVelocityMatch)
+		VersionProcessed = RepData.Version;
+		return bProcessed;
+	}
+
+	bool ProcessRepStateDataCommon(FGeometryCollectionPhysicsProxy* PhysicsProxy, const FGeometryCollectionRepStateData& RepStateData, int32& VersionProcessed)
+	{
+		using namespace Chaos;
+
+		if (!PhysicsProxy || !PhysicsProxy->IsInitializedOnPhysicsThread() || PhysicsProxy->GetReplicationMode() != FGeometryCollectionPhysicsProxy::EReplicationMode::Client)
+		{
+			return false;
+		}
+
+		if (VersionProcessed == RepStateData.Version)
+		{
+			return false;
+		}
+
+		const TArray<FPBDRigidClusteredParticleHandle*>& ParticleHandles = PhysicsProxy->GetParticles();
+
+		// Update the anchored state of the root particle
+		UpdateRootStateFromReplication(*PhysicsProxy, (bool)RepStateData.bIsRootAnchored);
+		
+		// first apply the breaking velocities 
+		for (const FGeometryCollectionRepStateData::FReleasedData& ReleaseData: RepStateData.ReleasedData)
+		{
+			if (/*ReleaseData.ReleasedVersion > VersionProcessed && */ ParticleHandles.IsValidIndex(ReleaseData.TransformIndex))
+			{
+				if (FPBDRigidClusteredParticleHandle* ParticleHandle = ParticleHandles[ReleaseData.TransformIndex])
 				{
-					//
-					// Match linear velocity
-					//
-					const FVector RepVel = RepCluster.LinearVelocity;
-					const FVector RepExtrapPos = RepCluster.Position + (RepVel * RepExtrapTime);
-					const FVec3 DeltaX = RepExtrapPos - Cluster.X();
-					const float DeltaXMagSq = DeltaX.SizeSquared();
-					if (DeltaXMagSq > SMALL_NUMBER && GeometryCollectionRepLinearMatchStrength > SMALL_NUMBER)
+					const bool bIsNotYetBroken = (ParticleHandle->Parent() != nullptr);
+					if (bIsNotYetBroken)
 					{
-						bWake = true;
-						//
-						// DeltaX * MatchStrength is an acceleration, m/s^2, which is integrated
-						// by multiplying by DeltaTime.
-						//
-						// It's formulated this way to get a larger correction for a longer time
-						// step, ie. correction velocities are framerate independent.
-						//
-						Cluster.SetV(RepVel + (DeltaX * GeometryCollectionRepLinearMatchStrength * DeltaTime));
+						ParticleHandle->SetV(ReleaseData.LinearVelocity);
+						ParticleHandle->SetW(FMath::DegreesToRadians(ReleaseData.AngularVelocityInDegreesPerSecond));
 					}
+				}
+			}
+		}
+		
+		// now sync the broken state of the particles
+		if (RepStateData.BrokenState.Num() == ParticleHandles.Num())
+		{
+			FPBDRigidsSolver* Solver = PhysicsProxy->GetSolver<Chaos::FPBDRigidsSolver>();
+			FRigidClustering& RigidClustering = Solver->GetEvolution()->GetRigidClustering();
 
-					//
-					// Match angular velocity
-					//
-					const FVector RepAngVel = RepCluster.AngularVelocity;
-					const Chaos::FRotation3 RepExtrapAng = Chaos::FRotation3::IntegrateRotationWithAngularVelocity(RepCluster.Rotation, RepAngVel, RepExtrapTime);
-					const FVector AngVel = Chaos::FRotation3::CalculateAngularVelocity(Cluster.R(), RepExtrapAng, GeometryCollectionRepAngularMatchTime);
-					if (AngVel.SizeSquared() > SMALL_NUMBER)
+			for (int32 TransformIndex = 0; TransformIndex < ParticleHandles.Num(); TransformIndex++)
+			{
+				if (FPBDRigidClusteredParticleHandle* ParticleHandle = ParticleHandles[TransformIndex])
+				{
+					// we only check for when thing goes from unbroken to broken 
+					const bool bWasBroken = (ParticleHandle->Parent() == nullptr);
+					const bool bIsBroken = RepStateData.BrokenState[TransformIndex];
+					if (!bWasBroken && bIsBroken)
 					{
-						Cluster.SetW(RepAngVel + AngVel);
-						bWake = true;
+						RigidClustering.ForceReleaseChildParticleAndParents(ParticleHandle, /* bTriggerBreakEvents */true);
 					}
 				}
+			}
+		}
 
-				bProcessed |= bWake;
+		VersionProcessed = RepStateData.Version;
+		return true;
+	}
 
-				//
-				// Wake up particle if it's sleeping and there's a delta to correct
-				//
-				if (bWake && Cluster.IsSleeping())
+	bool ProcessRepDynamicDataCommon(
+		FGeometryCollectionPhysicsProxy* PhysicsProxy,
+		const FGeometryCollectionRepDynamicData& RepDynamicData,
+		int32& VersionProcessed,
+		bool bReplicateMovement,
+		double SimTime,
+		double DeltaTime,
+		double& LastHardsnapTimeInMs
+	)
+	{
+		using namespace Chaos;
+
+		if (!PhysicsProxy || !PhysicsProxy->IsInitializedOnPhysicsThread() || PhysicsProxy->GetReplicationMode() != FGeometryCollectionPhysicsProxy::EReplicationMode::Client)
+		{
+			return false;
+		}
+
+		if (VersionProcessed == RepDynamicData.Version)
+		{
+			return false;
+		}
+
+		// extrapolation time is 0, because we do have the sim time when we receive the data 
+		// we keep this to be mimic the original implementation that was not running on the physics thread 
+		// todo(chaos) shoul dwe fix this ? 
+		const float RepExtrapTime = 0.f;
+
+		// If we've extrapolated past a threshold, then stop tracking
+		// the last received rep data
+		if (RepExtrapTime > GeometryCollectionRepMaxExtrapolationTime)
+		{
+			return false;
+		}
+
+		const bool bHardSnap = ReplicationShouldHardSnap(RepDynamicData.Version, VersionProcessed, bReplicateMovement, LastHardsnapTimeInMs);
+
+		// Keep track of whether we did some "work" on this frame so we can turn off the async tick after
+		// multiple frames of not doing anything.
+		bool bProcessed = false;
+		if (bReplicateMovement)
+		{
+			FPBDRigidsSolver* Solver = PhysicsProxy->GetSolver<Chaos::FPBDRigidsSolver>();
+			const TArray<FPBDRigidClusteredParticleHandle*>& ParticleHandles = PhysicsProxy->GetParticles();
+
+			for (const FGeometryCollectionRepDynamicData::FClusterData& RepCluster : RepDynamicData.ClusterData)
+			{
+				if (Chaos::FPBDRigidParticleHandle* Cluster = ParticleHandles[RepCluster.TransformIndex])
 				{
-					Solver->GetEvolution()->SetParticleObjectState(&Cluster, Chaos::EObjectStateType::Dynamic);
+					if (RepCluster.bIsInternalCluster)
+					{
+						// internal cluster do not have an index so we rep data send one of the children's, let's find the parent
+						Cluster = Cluster->CastToClustered()->Parent();
+					}
+
+					if (Cluster && Cluster->Disabled() == false)
+					{
+						const FQuat RotationAsQuat = FQuat::MakeFromEuler(RepCluster.EulerRotation);
+						const FVector AngularVelocityInRadiansPerSecond = FMath::DegreesToRadians(RepCluster.AngularVelocityInDegreesPerSecond);
+						UpdateClusterPositionAndVelocitiesFromReplication(Solver, *Cluster, RepCluster.Position, RotationAsQuat, RepCluster.LinearVelocity, AngularVelocityInRadiansPerSecond, bHardSnap, RepExtrapTime, DeltaTime);
+					}
 				}
-			});
+			};
 		}
 
-		VersionProcessed = RepData.Version;
-		return bProcessed;
+		VersionProcessed = RepDynamicData.Version;
+		return true;
 	}
 }
 
+
+
 bool UGeometryCollectionComponent::ProcessRepData(const float DeltaTime, const float SimTime)
 {
 
@@ -2236,12 +2701,10 @@ bool UGeometryCollectionComponent::ProcessRepData(const float DeltaTime, const f
 		RepData.RepDataReceivedTime = SimTime;
 	}
 
-	float ReceivedTime = *RepData.RepDataReceivedTime;
-
 	AActor* Owner = GetOwner();
 	const bool bReplicateMovement = Owner ? GetOwner()->IsReplicatingMovement() : true;
 
-	bool ReturnValue = ::ProcessRepDataCommon(DeltaTime, SimTime, RepData, PhysicsProxy, VersionProcessed, OneOffActivatedProcessed, LastHardsnapTimeInMs, ReceivedTime, bReplicateMovement);
+	bool ReturnValue = ::ProcessRepDataCommon(DeltaTime, SimTime, RepData, PhysicsProxy, VersionProcessed, OneOffActivatedProcessed, LastHardsnapTimeInMs, bReplicateMovement);
 
 
 #if ENABLE_DRAW_DEBUG
@@ -2333,22 +2796,44 @@ void UGeometryCollectionComponent::ProcessRepDataOnPT()
 		// which are only used in that function and in the reset function and both are executed on the Physics Thread
 		CurrSolver->EnqueueCommandImmediate([PhysicsProxy = PhysicsProxy, RepDataCopy = RepData, &VersionProcessed = VersionProcessed, &LastHardsnapTimeInMs = LastHardsnapTimeInMs, &OneOffActivatedProcessed = OneOffActivatedProcessed, bReplicateMovement](Chaos::FReal DeltaTime, Chaos::FReal SimTime)
 		{
-				float ReceivedTime;
+			::ProcessRepDataCommon(DeltaTime, SimTime, RepDataCopy, PhysicsProxy, VersionProcessed, OneOffActivatedProcessed, LastHardsnapTimeInMs, bReplicateMovement);
+		});
+	}
+}
 
-				// Track the sim time that this rep data was received on.
-				if (!RepDataCopy.RepDataReceivedTime.IsSet())
-				{
-					ReceivedTime = SimTime;
-				}
-				else
-				{
-					ReceivedTime = *RepDataCopy.RepDataReceivedTime;
-				}
+void UGeometryCollectionComponent::ProcessRepStateDataOnPT()
+{
+	ensure(GeometryCollectionUseReplicationV2);
 
-			::ProcessRepDataCommon(DeltaTime, SimTime, RepDataCopy, PhysicsProxy, VersionProcessed, OneOffActivatedProcessed, LastHardsnapTimeInMs, ReceivedTime, bReplicateMovement);
-		});
+	if (Chaos::FPhysicsSolver* CurrSolver = GetSolver(*this))
+	{
+		// enqueue a callback to process the replicated state data on the physics thread 
+		CurrSolver->EnqueueCommandImmediate(
+			[PhysicsProxy = PhysicsProxy, RepStateDataCopy = RepStateData, &VersionProcessed = VersionProcessed]
+			(Chaos::FReal /*DeltaTime*/, Chaos::FReal /*SimTime*/)
+			{
+				::ProcessRepStateDataCommon(PhysicsProxy, RepStateDataCopy, VersionProcessed);
+			});
 	}
+}
+
+void UGeometryCollectionComponent::ProcessRepDynamicDataOnPT()
+{
+	ensure(GeometryCollectionUseReplicationV2);
 
+	if (Chaos::FPhysicsSolver* CurrSolver = GetSolver(*this))
+	{
+		AActor* Owner = GetOwner();
+		const bool bReplicateMovement = Owner ? GetOwner()->IsReplicatingMovement() : true;
+
+		// enqueue a callback to process the replicated dynamic data on the physics thread 
+		CurrSolver->EnqueueCommandImmediate(
+			[PhysicsProxy = PhysicsProxy, RepDynamicDataCopy = RepDynamicData, &VersionProcessed = VersionProcessed, bReplicateMovement, &LastHardsnapTimeInMs = LastHardsnapTimeInMs]
+			(Chaos::FReal DeltaTime, Chaos::FReal SimTime)
+			{
+				::ProcessRepDynamicDataCommon(PhysicsProxy, RepDynamicDataCopy, VersionProcessed, bReplicateMovement, SimTime, DeltaTime, LastHardsnapTimeInMs);
+			});
+	}
 }
 
 void UGeometryCollectionComponent::SetDynamicState(const Chaos::EObjectStateType& NewDynamicState)
@@ -3373,7 +3858,7 @@ void UGeometryCollectionComponent::UpdateNavigationDataIfNeeded(bool bDynamicCol
 {
 	if (bUpdateNavigationInTick && bDynamicCollectionDirty)
 	{
-		const UWorld* MyWorld = GetWorld();
+		const UWorld* MyWorld = GetWorld();// If not doing velocity match, don't bother processing the same version twice.
 		if (MyWorld && MyWorld->IsGameWorld())
 		{
 			//cycle every 0xff frames
@@ -3471,6 +3956,26 @@ void UGeometryCollectionComponent::OnRep_RepData()
 	ProcessRepDataOnPT();
 }
 
+void UGeometryCollectionComponent::OnRep_RepStateData()
+{
+	// We have new data that was replicated! Turn on the async tick to process instead of just requesting a one-off
+	// since we may want to keep processing for extra time afterwards.
+	check(IsInGameThread());
+	check(GetNetMode() == ENetMode::NM_Client);
+
+	ProcessRepStateDataOnPT();
+}
+
+void UGeometryCollectionComponent::OnRep_RepDynamicData()
+{
+	// We have new dynamic data that was replicated! Turn on the async tick to process instead of just requesting a one-off
+	// since we may want to keep processing for extra time afterwards.
+	check(IsInGameThread());
+	check(GetNetMode() == ENetMode::NM_Client);
+
+	ProcessRepDynamicDataOnPT();
+}
+
 void UGeometryCollectionComponent::SetAbandonedParticleCollisionProfileName(FName CollisionProfile)
 {
 	if (!bEnableAbandonAfterLevel || !GetIsReplicated())
diff --git a/Engine/Source/Runtime/Experimental/GeometryCollectionEngine/Public/GeometryCollection/GeometryCollectionComponent.h b/Engine/Source/Runtime/Experimental/GeometryCollectionEngine/Public/GeometryCollection/GeometryCollectionComponent.h
--- a/Engine/Source/Runtime/Experimental/GeometryCollectionEngine/Public/GeometryCollection/GeometryCollectionComponent.h
+++ b/Engine/Source/Runtime/Experimental/GeometryCollectionEngine/Public/GeometryCollection/GeometryCollectionComponent.h
@@ -437,6 +437,127 @@ struct TStructOpsTypeTraits<FGeometryCollectionRepData> : public TStructOpsTypeT
 	};
 };
 
+/**
+* Replicated state data for a geometry collection when bEnableReplication is true for that component.
+* State data means what is broken and what is not 
+* See UGeomtryCollectionComponent::UpdateRepData
+*/
+USTRUCT()
+struct FGeometryCollectionRepStateData
+{
+	GENERATED_BODY()
+
+	FGeometryCollectionRepStateData()
+	: Version(0)
+	, bIsRootAnchored(false)
+	{
+	}
+
+	// mark a transform as broken nd return true if this was a state change
+	bool SetBroken(int32 TransformIndex, int32 NumTransforms, bool bDisabled, const FVector& LinV, const FVector& AngVInRadiansPerSecond);
+
+	// version for fast comparison
+	int32 Version;
+
+	// broken state of each piece of the GC
+	TBitArray<> BrokenState;
+
+	// Is the root particle of the GC currently anchored
+	// could possibily change in the future to also be a bit array 
+	uint8 bIsRootAnchored;
+
+	// this represents the data for when a particle is released from its parent cluster 
+	// this data is added when the particle is released but will be cleared after a while
+	// so that late client will not replay the break as it is in their past 
+	struct FReleasedData
+	{
+		int16 TransformIndex;
+		FVector_NetQuantize10 LinearVelocity;
+		FVector_NetQuantize10 AngularVelocityInDegreesPerSecond;
+	};
+	TArray<FReleasedData> ReleasedData;
+
+	// Just test version to skip having to traverse the whole pose array for replication
+	bool Identical(const FGeometryCollectionRepStateData* Other, uint32 PortFlags) const;
+	bool NetSerialize(FArchive& Ar, class UPackageMap* Map, bool& bOutSuccess);
+
+	void Reset()
+	{
+		BrokenState.Reset();
+		bIsRootAnchored = 0;
+		ReleasedData.Reset();
+	}
+};
+
+template<>
+struct TStructOpsTypeTraits<FGeometryCollectionRepStateData> : public TStructOpsTypeTraitsBase2<FGeometryCollectionRepStateData>
+{
+	enum
+	{
+		WithNetSerializer = true,
+		WithIdentical = true,
+	};
+};
+
+// this structure holds entries for the tracked pieces to be replicated
+USTRUCT()
+struct FGeometryCollectionRepDynamicData
+{
+	GENERATED_BODY()
+
+	struct FClusterData
+	{
+		FVector_NetQuantize10 Position;
+		FVector_NetQuantize10 EulerRotation;
+		FVector_NetQuantize10 LinearVelocity;
+		FVector_NetQuantize10 AngularVelocityInDegreesPerSecond;
+
+		// Index of the cluster or one of its child if the cluster is internal ( see bIsInternalCluster)
+		uint16 TransformIndex = INDEX_NONE; 
+
+		// Whether this refers to an internal cluster or directly to a cluster in the geometry collection
+		uint8  bIsInternalCluster = false;
+
+		// non serialized data, used to trimn the data back when no longer updated
+		int32 LastUpdatedVersion = 0;
+
+		// comp
+		bool IsEqualPositionsAndVelocities(const FClusterData& Data) const;
+	};
+
+	FGeometryCollectionRepDynamicData()
+		: Version(0)
+	{}
+
+	int32 Version;
+	TArray<FClusterData> ClusterData;
+
+	// return true if the data has changed from stored one
+	bool SetData(const FClusterData& Data);
+
+	// return true if any entries was removed
+	bool RemoveOutOfDateClusterData();
+
+	// Just test version to skip having to traverse the whole pose array for replication
+	bool Identical(const FGeometryCollectionRepDynamicData* Other, uint32 PortFlags) const;
+	bool NetSerialize(FArchive& Ar, class UPackageMap* Map, bool& bOutSuccess);
+
+	void Reset()
+	{
+		ClusterData.Reset();
+	}
+};
+
+template<>
+struct TStructOpsTypeTraits<FGeometryCollectionRepDynamicData> : public TStructOpsTypeTraitsBase2<FGeometryCollectionRepDynamicData>
+{
+	enum
+	{
+		WithNetSerializer = true,
+		WithIdentical = true,
+	};
+};
+
 /**
 *	GeometryCollectionComponent
 */
@@ -1352,12 +1473,25 @@ protected:
 	UPROPERTY(ReplicatedUsing=OnRep_RepData)
 	FGeometryCollectionRepData RepData;
 
+	UPROPERTY(ReplicatedUsing = OnRep_RepStateData)
+	FGeometryCollectionRepStateData RepStateData;
+
+	UPROPERTY(ReplicatedUsing = OnRep_RepDynamicData)
+	FGeometryCollectionRepDynamicData RepDynamicData;
+
 	/** Called post solve to allow authoritative components to update their replication data */
 	UFUNCTION()
 	GEOMETRYCOLLECTIONENGINE_API void OnRep_RepData();
 
+	UFUNCTION()
+	GEOMETRYCOLLECTIONENGINE_API void OnRep_RepStateData();
+
+	UFUNCTION()
+	GEOMETRYCOLLECTIONENGINE_API void OnRep_RepDynamicData();
+	
 	GEOMETRYCOLLECTIONENGINE_API void RequestUpdateRepData();
 	GEOMETRYCOLLECTIONENGINE_API virtual void UpdateRepData();
+	GEOMETRYCOLLECTIONENGINE_API virtual void UpdateRepStateAndDynamicData();
 
 	/** Clear all rep data, this is required if the physics proxy has been recreated */
 	GEOMETRYCOLLECTIONENGINE_API virtual void ResetRepData();
@@ -1378,7 +1512,9 @@ protected:
 	int64 LastAsyncPhysicsTickMs = 0;
 
 private:
-	GEOMETRYCOLLECTIONENGINE_API void ProcessRepDataOnPT();
+	void ProcessRepDataOnPT();
+	void ProcessRepStateDataOnPT();
+	void ProcessRepDynamicDataOnPT();
 
 	// return the most actual transforms
 	// this can be the rest collection ones, the overriden RestTransforms or the dynamic collection ones
