Enable geometry collection replication optimized code path

diff --git a/Engine/Source/Runtime/Experimental/GeometryCollectionEngine/Private/GeometryCollection/GeometryCollectionComponent.cpp b/Engine/Source/Runtime/Experimental/GeometryCollectionEngine/Private/GeometryCollection/GeometryCollectionComponent.cpp
--- a/Engine/Source/Runtime/Experimental/GeometryCollectionEngine/Private/GeometryCollection/GeometryCollectionComponent.cpp
+++ b/Engine/Source/Runtime/Experimental/GeometryCollectionEngine/Private/GeometryCollection/GeometryCollectionComponent.cpp
@@ -131,7 +131,7 @@ FAutoConsoleVariableRef CVarGeometryCollectionEmitRootBreakingEvent(TEXT("p.Chao
 bool GeometryCollectionCreatePhysicsStateInEditor = false;
 FAutoConsoleVariableRef CVarGeometryCollectionCreatePhysicsStateInEditor(TEXT("p.Chaos.GC.CreatePhysicsStateInEditor"), GeometryCollectionCreatePhysicsStateInEditor, TEXT("when on , physics state for a GC will be create in editor ( non PIE )"));
 
-bool GeometryCollectionUseReplicationV2 = false;
+bool GeometryCollectionUseReplicationV2 = true;
 FAutoConsoleVariableRef CVarGeometryCollectionUseReplicationV2(TEXT("p.Chaos.GC.UseReplicationV2"), GeometryCollectionUseReplicationV2, TEXT("When true use new replication data model"));
 
 DEFINE_LOG_CATEGORY_STATIC(UGCC_LOG, Error, All);
