fix edge case where breaking model could cause particles to release from the cluster union and not trigger a connectivity check

- add the connectivity check as a cluster union operation that gets deferred whenever a particle is removed to remove the loophole

diff --git a/Engine/Source/Runtime/Experimental/Chaos/Private/Chaos/ClusterUnionManager.cpp b/Engine/Source/Runtime/Experimental/Chaos/Private/Chaos/ClusterUnionManager.cpp
--- a/Engine/Source/Runtime/Experimental/Chaos/Private/Chaos/ClusterUnionManager.cpp
+++ b/Engine/Source/Runtime/Experimental/Chaos/Private/Chaos/ClusterUnionManager.cpp
@@ -425,7 +425,6 @@ namespace Chaos
 
 		if (!PendingClusterIndexOperations.IsEmpty())
 		{
-			TArray<FClusterUnionIndex> ClustersRequiringConnectivityChecks;
 			for (TPair<FClusterUnionIndex, FClusterOpMap>& OpMap : PendingClusterIndexOperations)
 			{
 				// Is this sort necessary? Better to be safe than sorry. Since we need to guarantee that the UpdateChildToParent happens after add.
@@ -446,7 +445,6 @@ namespace Chaos
 						HandleAddOperation(OpMap.Key, Op.Value, Op.Key == EClusterUnionOperation::AddReleased);
 						break;
 					case EClusterUnionOperation::Remove:
-						ClustersRequiringConnectivityChecks.Add(OpMap.Key);
 						HandleRemoveOperation(OpMap.Key, Op.Value, EClusterUnionOperationTiming::Defer);
 						break;
 					case EClusterUnionOperation::UpdateChildToParent:
@@ -457,18 +455,6 @@ namespace Chaos
 			}
 			PendingClusterIndexOperations.Empty();
 
-			// Needs to be after the emptying of the cluster index operations because otherwise it might remove the removals
-			// added in by the connectivity check.
-			for (FClusterUnionIndex UnionIndex : ClustersRequiringConnectivityChecks)
-			{
-				if (FClusterUnion* ClusterUnion = FindClusterUnion(UnionIndex))
-				{
-					// TODO: Can probably argue that cluster union connectivity should be moved into the
-					// cluster union manager instead?
-					MClustering.HandleConnectivityOnReleaseClusterParticle(ClusterUnion->InternalCluster, false);
-				}
-			}
-
 			if (FRewindData* RewindData = MEvolution.GetRewindData())
 			{
 				// Temp, since rewind doesn't handle altered cluster unions we tell it to never rewind past this frame for now, only relevant when resimulation is used.
@@ -815,7 +801,8 @@ namespace Chaos
 
 		constexpr EUpdateClusterUnionPropertiesFlags RemoveUpdateFlags = EUpdateClusterUnionPropertiesFlags::IncrementalGenerateConnectionGraph
 			| EUpdateClusterUnionPropertiesFlags::UpdateKinematicProperties
-			| EUpdateClusterUnionPropertiesFlags::IncrementalGenerateGeometry;
+			| EUpdateClusterUnionPropertiesFlags::IncrementalGenerateGeometry
+			| EUpdateClusterUnionPropertiesFlags::ConnectivityCheck;
 		switch (UpdateClusterPropertiesTiming)
 		{
 		case EClusterUnionOperationTiming::Immediate:
@@ -930,6 +917,13 @@ namespace Chaos
 			{
 				FlushIncrementalConnectivityGraphOperations(ClusterUnion);
 			}
+
+			if (EnumHasAnyFlags(Flags, EUpdateClusterUnionPropertiesFlags::ConnectivityCheck))
+			{
+				// TODO: Can probably argue that cluster union connectivity should be moved into the
+				// cluster union manager instead?
+				MClustering.HandleConnectivityOnReleaseClusterParticle(ClusterUnion.InternalCluster, false);
+			}
 		}
 
 		for (FPBDRigidClusteredParticleHandle* ChildParticle : PendingParticlesToUndoChildToParentLock)
diff --git a/Engine/Source/Runtime/Experimental/Chaos/Public/Chaos/ClusterUnionManager.h b/Engine/Source/Runtime/Experimental/Chaos/Public/Chaos/ClusterUnionManager.h
--- a/Engine/Source/Runtime/Experimental/Chaos/Public/Chaos/ClusterUnionManager.h
+++ b/Engine/Source/Runtime/Experimental/Chaos/Public/Chaos/ClusterUnionManager.h
@@ -42,7 +42,8 @@ namespace Chaos
 		UpdateKinematicProperties = 1 << 3,
 		ForceGenerateGeometry = 1 << 4,
 		IncrementalGenerateGeometry = 1 << 5,
-		All = RecomputeMassOrientation | ForceGenerateConnectionGraph | UpdateKinematicProperties | ForceGenerateGeometry
+		ConnectivityCheck = 1 << 6,
+		All = RecomputeMassOrientation | ForceGenerateConnectionGraph | UpdateKinematicProperties | ForceGenerateGeometry | ConnectivityCheck
 	};
 	ENUM_CLASS_FLAGS(EUpdateClusterUnionPropertiesFlags);
 
