Enabling Object interface with suspension constraint.
Using Constraint Suspension on modular vehicles

diff --git a/Engine/Source/Runtime/Experimental/Chaos/Private/Chaos/PBDSuspensionConstraintData.cpp b/Engine/Source/Runtime/Experimental/Chaos/Private/Chaos/PBDSuspensionConstraintData.cpp
--- a/Engine/Source/Runtime/Experimental/Chaos/Private/Chaos/PBDSuspensionConstraintData.cpp
+++ b/Engine/Source/Runtime/Experimental/Chaos/Private/Chaos/PBDSuspensionConstraintData.cpp
@@ -1,6 +1,8 @@
 // Copyright Epic Games, Inc. All Rights Reserved.
 
 #include "Chaos/PBDSuspensionConstraintData.h"
+#include "PhysicsProxy/SingleParticlePhysicsProxy.h"
+#include "Chaos/PhysicsObjectInterface.h"
 
 namespace Chaos
 {
@@ -10,5 +12,36 @@ namespace Chaos
 	{
 	}
 
+	void FSuspensionConstraint::SetParticleProxy(IPhysicsProxyBase* InParticleProxy)
+	{
+		check(InParticleProxy);
+
+		SuspensionProxy.Modify(/*bInvalidate=*/true, DirtyFlags, Proxy, [InParticleProxy](FProxyBaseProperty& Data)
+			{
+				ensure(InParticleProxy->GetType() == EPhysicsProxyType::SingleParticleProxy); // legacy, only worked with SingleParticleProxy
+				Data.Proxy = InParticleProxy;
+			});
+
+		// This should work fine since this is a legacy endpoint and was primarily only used with single particle physics proxies.
+		SuspensionBody.Modify(/*bInvalidate=*/true, DirtyFlags, Proxy, [InParticleProxy](FPhysicsObjectProperty& Data)
+			{
+				Data.PhysicsBody = (InParticleProxy->GetType() == EPhysicsProxyType::SingleParticleProxy) ? static_cast<FSingleParticlePhysicsProxy*>(InParticleProxy)->GetPhysicsObject() : nullptr;
+			});
+	}
+
+	void FSuspensionConstraint::SetPhysicsBody(FPhysicsObjectHandle& InBody)
+	{
+		check(InBody);
+
+		SuspensionProxy.Modify(/*bInvalidate=*/true, DirtyFlags, Proxy, [InBody](FProxyBaseProperty& Data)
+			{
+				Data.Proxy = Chaos::FPhysicsObjectInterface::GetProxy({ &InBody, 1 });
+			});
+
+		SuspensionBody.Modify(/*bInvalidate=*/true, DirtyFlags, Proxy, [InBody](FPhysicsObjectProperty& Data)
+			{
+				Data.PhysicsBody = InBody;
+			});
+	}
 
 } // Chaos
diff --git a/Engine/Source/Runtime/Experimental/Chaos/Private/PhysicsProxy/SuspensionConstraintProxy.cpp b/Engine/Source/Runtime/Experimental/Chaos/Private/PhysicsProxy/SuspensionConstraintProxy.cpp
--- a/Engine/Source/Runtime/Experimental/Chaos/Private/PhysicsProxy/SuspensionConstraintProxy.cpp
+++ b/Engine/Source/Runtime/Experimental/Chaos/Private/PhysicsProxy/SuspensionConstraintProxy.cpp
@@ -11,6 +11,7 @@
 #include "Chaos/Serializable.h"
 #include "Chaos/PBDSuspensionConstraints.h"
 #include "Chaos/Framework/MultiBufferResource.h"
+#include "Chaos/PhysicsObjectInternal.h"
 #include "PhysicsProxy/SingleParticlePhysicsProxy.h"
 #include "PhysicsSolver.h"
 
@@ -46,17 +47,19 @@ void FSuspensionConstraintPhysicsProxy::InitializeOnPhysicsThread(FPBDRigidsSolv
 	if (Handles.Size())
 	{
 		auto& SuspensionConstraints = InSolver->GetSuspensionConstraints();
-		if (const FParticleProxyProperty* ProxyProperty = RemoteData.FindSuspensionParticleProxy(Manager, DataIdx))
+		if (const FPhysicsObjectProperty* Body = RemoteData.FindSuspensionPhysicsObject(Manager, DataIdx))
 		{
-			if (FGeometryParticleHandle* Handle0 = GetParticleHandleFromProxy(ProxyProperty->ParticleProxy))
+			FGeometryParticleHandle* Handle = Body->PhysicsBody->GetRootParticle<Chaos::EThreadContext::Internal>();
+
+			if (Handle)
 			{
 				const FPBDSuspensionSettings* SuspensionSettings = RemoteData.FindSuspensionSettings(Manager, DataIdx);
 				const FSuspensionLocation* SuspensionLocation = RemoteData.FindSuspensionLocation(Manager, DataIdx);
 
 				if (SuspensionSettings && SuspensionLocation)
 				{
-					Constraint_PT = SuspensionConstraints.AddConstraint(Handle0, SuspensionLocation->Location, *SuspensionSettings);
-					Handle0->AddConstraintHandle(Constraint_PT);
+					Constraint_PT = SuspensionConstraints.AddConstraint(Handle, SuspensionLocation->Location, *SuspensionSettings);
+					Handle->AddConstraintHandle(Constraint_PT);
 				}
 			}
 		}
diff --git a/Engine/Source/Runtime/Experimental/Chaos/Public/Chaos/PBDSuspensionConstraintData.h b/Engine/Source/Runtime/Experimental/Chaos/Public/Chaos/PBDSuspensionConstraintData.h
--- a/Engine/Source/Runtime/Experimental/Chaos/Public/Chaos/PBDSuspensionConstraintData.h
+++ b/Engine/Source/Runtime/Experimental/Chaos/Public/Chaos/PBDSuspensionConstraintData.h
@@ -37,7 +37,10 @@ namespace Chaos
 		virtual ~FSuspensionConstraint() override {}
 
 
-		const FPBDSuspensionSettings& GetSuspensionSettings()const { return SuspensionSettings.Read(); }
+		CHAOS_API void SetParticleProxy(IPhysicsProxyBase* InParticleProxy);
+		CHAOS_API void SetPhysicsBody(FPhysicsObjectHandle& InBody);
+
+		const FPBDSuspensionSettings& GetSuspensionSettings() const { return SuspensionSettings.Read(); }
 
 #define CHAOS_INNER_SUSP_PROPERTY(OuterProp, Name, InnerType)\
 	void Set##Name(InnerType Val){ OuterProp.Modify(/*bInvalidate=*/true, DirtyFlags, Proxy, [&Val](auto& Data) { Data.Name = Val; }); }\
@@ -47,14 +50,16 @@ namespace Chaos
 
 	protected:
 		TChaosProperty<FPBDSuspensionSettings, EChaosProperty::SuspensionSettings> SuspensionSettings;
-		TChaosProperty<FParticleProxyProperty, EChaosProperty::SuspensionParticleProxy> SuspensionParticleProxy;
 		TChaosProperty<FSuspensionLocation, EChaosProperty::SuspensionLocation> SuspensionLocation; // location = spring local offset
+		TChaosProperty<FProxyBaseProperty, EChaosProperty::SuspensionParticleProxy> SuspensionProxy;
+		TChaosProperty<FPhysicsObjectProperty, EChaosProperty::SuspensionPhysicsObject> SuspensionBody;
 
 		virtual void SyncRemoteDataImp(FDirtyPropertiesManager& Manager, int32 DataIdx, FDirtyChaosProperties& RemoteData) override
 		{
 			SuspensionSettings.SyncRemote(Manager, DataIdx, RemoteData);
-			SuspensionParticleProxy.SyncRemote(Manager, DataIdx, RemoteData);
 			SuspensionLocation.SyncRemote(Manager, DataIdx, RemoteData);
+			SuspensionProxy.SyncRemote(Manager, DataIdx, RemoteData);
+			SuspensionBody.SyncRemote(Manager, DataIdx, RemoteData);
 		}
 	};
 
diff --git a/Engine/Source/Runtime/Experimental/Chaos/Public/Chaos/ParticleDirtyFlags.h b/Engine/Source/Runtime/Experimental/Chaos/Public/Chaos/ParticleDirtyFlags.h
--- a/Engine/Source/Runtime/Experimental/Chaos/Public/Chaos/ParticleDirtyFlags.h
+++ b/Engine/Source/Runtime/Experimental/Chaos/Public/Chaos/ParticleDirtyFlags.h
@@ -216,6 +216,11 @@ struct FProxyBasePairProperty
 	FProxyBasePair ParticleProxies = { nullptr, nullptr };
 };
 
+struct FProxyBaseProperty
+{
+	IPhysicsProxyBase* Proxy = nullptr;
+};
+
 struct FPhysicsObject;
 typedef TVector<FPhysicsObject*, 2> FPhysicsObjectPair;
 
@@ -224,6 +229,11 @@ struct FPhysicsObjectPairProperty
 	FPhysicsObjectPair PhysicsBodies = { nullptr, nullptr };
 };
 
+struct FPhysicsObjectProperty
+{
+	FPhysicsObject* PhysicsBody = nullptr;
+};
+
 inline FChaosArchive& operator<<(FChaosArchive& Ar, FParticleDynamics& Data)
 {
 	Data.Serialize(Ar);
diff --git a/Engine/Source/Runtime/Experimental/Chaos/Public/Chaos/ParticleProperties.inl b/Engine/Source/Runtime/Experimental/Chaos/Public/Chaos/ParticleProperties.inl
--- a/Engine/Source/Runtime/Experimental/Chaos/Public/Chaos/ParticleProperties.inl
+++ b/Engine/Source/Runtime/Experimental/Chaos/Public/Chaos/ParticleProperties.inl
@@ -24,6 +24,7 @@ CHAOS_PROPERTY(JointPhysicsObjects, FPhysicsObjectPairProperty, EPhysicsProxyTyp
 CHAOS_PROPERTY(SuspensionSettings, FPBDSuspensionSettings, EPhysicsProxyType::SuspensionConstraintType)
 CHAOS_PROPERTY(SuspensionParticleProxy, FParticleProxyProperty, EPhysicsProxyType::SuspensionConstraintType)
 CHAOS_PROPERTY(SuspensionLocation, FSuspensionLocation, EPhysicsProxyType::SuspensionConstraintType)
+CHAOS_PROPERTY(SuspensionPhysicsObject, FPhysicsObjectProperty, EPhysicsProxyType::SuspensionConstraintType)
 
 CHAOS_PROPERTY(CharacterGroundConstraintSettings, FCharacterGroundConstraintSettings, EPhysicsProxyType::CharacterGroundConstraintType)
 CHAOS_PROPERTY(CharacterGroundConstraintDynamicData, FCharacterGroundConstraintDynamicData, EPhysicsProxyType::CharacterGroundConstraintType)
diff --git a/Engine/Source/Runtime/Experimental/Chaos/Public/Chaos/SuspensionProperties.inl b/Engine/Source/Runtime/Experimental/Chaos/Public/Chaos/SuspensionProperties.inl
--- a/Engine/Source/Runtime/Experimental/Chaos/Public/Chaos/SuspensionProperties.inl
+++ b/Engine/Source/Runtime/Experimental/Chaos/Public/Chaos/SuspensionProperties.inl
@@ -13,7 +13,6 @@ CHAOS_INNER_SUSP_PROPERTY(SuspensionSettings, MinLength, FReal);
 CHAOS_INNER_SUSP_PROPERTY(SuspensionSettings, MaxLength, FReal);
 CHAOS_INNER_SUSP_PROPERTY(SuspensionSettings, Axis, FVec3);
 CHAOS_INNER_SUSP_PROPERTY(SuspensionSettings, Normal, FVec3)
-CHAOS_INNER_SUSP_PROPERTY(SuspensionParticleProxy, ParticleProxy, FSingleParticlePhysicsProxy*)
 CHAOS_INNER_SUSP_PROPERTY(SuspensionLocation, Location, FVec3)
 
 #undef CHAOS_INNER_SUSP_PROPERTY
\ No newline at end of file
diff --git a/Engine/Source/Runtime/Experimental/ChaosVehicles/ChaosVehiclesCore/Private/SimModule/SuspensionModule.cpp b/Engine/Source/Runtime/Experimental/ChaosVehicles/ChaosVehiclesCore/Private/SimModule/SuspensionModule.cpp
--- a/Engine/Source/Runtime/Experimental/ChaosVehicles/ChaosVehiclesCore/Private/SimModule/SuspensionModule.cpp
+++ b/Engine/Source/Runtime/Experimental/ChaosVehicles/ChaosVehiclesCore/Private/SimModule/SuspensionModule.cpp
@@ -3,6 +3,9 @@
 #include "SimModule/SuspensionModule.h"
 #include "SimModule/SimModuleTree.h"
 #include "SimModule/WheelModule.h"
+#include "Chaos/PBDSuspensionConstraints.h"
+#include "PhysicsProxy/SuspensionConstraintProxy.h"
+#include "PBDRigidsSolver.h"
 #include "VehicleUtility.h"
 
 #if VEHICLE_DEBUGGING_ENABLED
@@ -18,6 +21,10 @@ namespace Chaos
 		, LastDisplacement(0.f)
 		, SpringSpeed(0.f)
 		, WheelSimTreeIndex(INVALID_IDX)
+		, Constraint(nullptr)
+		, TargetPos(FVector::ZeroVector)
+		, ImpactNormal(FVector::ZeroVector)
+		, WheelInContact(false)
 	{
 		AccessSetup().MaxLength = FMath::Abs(Settings.MaxRaise + Settings.MaxDrop);
 	}
@@ -48,36 +55,48 @@ namespace Chaos
 
 	void FSuspensionSimModule::Simulate(float DeltaTime, const FAllInputs& Inputs, FSimModuleTree& VehicleModuleSystem)
 	{
-		float ForceIntoSurface = 0.0f;
-		if (SpringDisplacement > 0)
 		{
-			float Damping = (SpringDisplacement < LastDisplacement) ? Setup().CompressionDamping : Setup().ReboundDamping;
-			SpringSpeed = (LastDisplacement - SpringDisplacement) / DeltaTime;
-
-			float StiffnessForce = SpringDisplacement * Setup().SpringRate;
-			float DampingForce = SpringSpeed * Damping;
-			float SuspensionForce = StiffnessForce - DampingForce;
-			LastDisplacement = SpringDisplacement;
-
-			if (SuspensionForce > 0)
+			float ForceIntoSurface = 0.0f;
+			if (SpringDisplacement > 0)
 			{
-				ForceIntoSurface = SuspensionForce;
-				AddLocalForce(Setup().SuspensionAxis * -SuspensionForce, true, false, true, FColor::Green);
+				float Damping = (SpringDisplacement < LastDisplacement) ? Setup().CompressionDamping : Setup().ReboundDamping;
+				SpringSpeed = (LastDisplacement - SpringDisplacement) / DeltaTime;
+
+				float StiffnessForce = SpringDisplacement * Setup().SpringRate;
+				float DampingForce = SpringSpeed * Damping;
+				float SuspensionForce = StiffnessForce - DampingForce;
+				LastDisplacement = SpringDisplacement;
+
+				if (SuspensionForce > 0)
+				{
+					ForceIntoSurface = SuspensionForce;
+
+					if (Constraint == nullptr)
+					{
+						AddLocalForce(Setup().SuspensionAxis * -SuspensionForce, true, false, true, FColor::Green);
+					}
+				}
 			}
-		}
 
-		// tell wheels how much they are being pressed into the ground
-		if (SimModuleTree && WheelSimTreeIndex != INVALID_IDX)
-		{
-			if (Chaos::ISimulationModuleBase* Module = SimModuleTree->AccessSimModule(WheelSimTreeIndex))
+			// tell wheels how much they are being pressed into the ground
+			if (SimModuleTree && WheelSimTreeIndex != INVALID_IDX)
 			{
-				check(Module->GetSimType() == eSimType::Wheel);
-				Chaos::FWheelSimModule* Wheel = static_cast<Chaos::FWheelSimModule*>(Module);
-
-				Wheel->SetForceIntoSurface(ForceIntoSurface);
+				if (Chaos::ISimulationModuleBase* Module = SimModuleTree->AccessSimModule(WheelSimTreeIndex))
+				{
+					check(Module->GetSimType() == eSimType::Wheel);
+					Chaos::FWheelSimModule* Wheel = static_cast<Chaos::FWheelSimModule*>(Module);
+
+					Wheel->SetForceIntoSurface(ForceIntoSurface);
+				}
+		
 			}
+		}
 
+		if (Constraint)
+		{
+			UpdateConstraint();
 		}
+
 	}
 
 	void FSuspensionSimModule::Animate(Chaos::FClusterUnionPhysicsProxy* Proxy)
@@ -95,6 +114,24 @@ namespace Chaos
 		}
 	}
 
+	void FSuspensionSimModule::SetSuspensionConstraint(FSuspensionConstraint* InConstraint)
+	{
+		Constraint = InConstraint;
+	}
+
+	void FSuspensionSimModule::UpdateConstraint()
+	{
+		if (Constraint && Constraint->IsValid())
+		{
+			if (FSuspensionConstraintPhysicsProxy* Proxy = Constraint->GetProxy<FSuspensionConstraintPhysicsProxy>())
+			{
+				Chaos::FPhysicsSolver* Solver = Proxy->GetSolver<Chaos::FPhysicsSolver>();
+				UE_LOG(LogTemp, Warning, TEXT("TargetPos %s, WheelInContact %d"), *TargetPos.ToString(), WheelInContact);
+				Solver->SetSuspensionTarget(Constraint, TargetPos, ImpactNormal, WheelInContact);
+			}
+		}
+	}
+
 	void FSuspensionSimModuleDatas::FillSimState(ISimulationModuleBase* SimModule)
 	{
 		check(SimModule->GetSimType() == eSimType::Suspension);
diff --git a/Engine/Source/Runtime/Experimental/ChaosVehicles/ChaosVehiclesCore/Public/SimModule/SuspensionModule.h b/Engine/Source/Runtime/Experimental/ChaosVehicles/ChaosVehiclesCore/Public/SimModule/SuspensionModule.h
--- a/Engine/Source/Runtime/Experimental/ChaosVehicles/ChaosVehiclesCore/Public/SimModule/SuspensionModule.h
+++ b/Engine/Source/Runtime/Experimental/ChaosVehicles/ChaosVehiclesCore/Public/SimModule/SuspensionModule.h
@@ -9,6 +9,7 @@ namespace Chaos
 	struct FAllInputs;
 	class FSimModuleTree;
 	class FClusterUnionPhysicsProxy;
+	class FSuspensionConstraint;
 
 	struct CHAOSVEHICLESCORE_API FSuspensionSimModuleDatas : public FModuleNetData
 	{
@@ -147,14 +148,28 @@ namespace Chaos
 		const FVector& GetRestLocation() const { return Setup().RestOffset; }
 
 		void SetWheelSimTreeIndex(int WheelTreeIndexIn) { WheelSimTreeIndex = WheelTreeIndexIn; }
-		int GetWheelSimTreeIndex() const { return WheelSimTreeIndex; }
+		int GetWheelSimTreeIndex() const { return WheelSimTreeIndex; } 
 
+		void UpdateConstraint();
+
+		void SetSuspensionConstraint(FSuspensionConstraint* InConstraint);
+		void SetTargetPoint(const FVector& InTargetPoint, const FVector& InImpactNormal, bool InWheelInContact)
+		{
+			TargetPos = InTargetPoint;
+			ImpactNormal = InImpactNormal;
+			WheelInContact = InWheelInContact;
+		}
 	private:
 
 		float SpringDisplacement;
 		float LastDisplacement;
 		float SpringSpeed;
 		int WheelSimTreeIndex;
+
+		FSuspensionConstraint* Constraint;
+		FVector TargetPos;
+		FVector ImpactNormal;
+		bool WheelInContact;
 	};
 
 
diff --git a/Engine/Source/Runtime/PhysicsCore/Private/ChaosEngineInterface.cpp b/Engine/Source/Runtime/PhysicsCore/Private/ChaosEngineInterface.cpp
--- a/Engine/Source/Runtime/PhysicsCore/Private/ChaosEngineInterface.cpp
+++ b/Engine/Source/Runtime/PhysicsCore/Private/ChaosEngineInterface.cpp
@@ -1362,24 +1362,31 @@ FPhysicsConstraintHandle FChaosEngineInterface::CreateConstraint(const FPhysicsA
 }
 
 FPhysicsConstraintHandle FChaosEngineInterface::CreateSuspension(const FPhysicsActorHandle& InActorRef, const FVector& InLocalFrame)
+{
+	Chaos::FPhysicsObject* Body = InActorRef ? InActorRef->GetPhysicsObject() : nullptr;
+	return FChaosEngineInterface::CreateSuspension(Body, InLocalFrame);
+}
+
+FPhysicsConstraintHandle FChaosEngineInterface::CreateSuspension(Chaos::FPhysicsObject* Body, const FVector& InLocalFrame)
 {
 	FPhysicsConstraintHandle ConstraintRef;
 
 	if (bEnableChaosJointConstraints)
 	{
-		if (InActorRef)
+		if (Body)
 		{
-			if (InActorRef->GetSolverBase())
+			Chaos::FPhysicsSolver* Solver = Chaos::FPhysicsObjectInterface::GetSolver({ &Body, 1 });
+
+			if (Solver)
 			{
 				LLM_SCOPE(ELLMTag::ChaosConstraint);
 
 				auto* SuspensionConstraint = new Chaos::FSuspensionConstraint();
 				ConstraintRef.Constraint = SuspensionConstraint;
 
-				SuspensionConstraint->SetParticleProxy(InActorRef);
+				SuspensionConstraint->SetPhysicsBody(Body);
 				SuspensionConstraint->SetLocation(InLocalFrame);
 
-				Chaos::FPhysicsSolver* Solver = InActorRef->GetSolver<Chaos::FPhysicsSolver>();
 				Solver->RegisterObject(SuspensionConstraint);
 			}
 		}
diff --git a/Engine/Source/Runtime/PhysicsCore/Public/Chaos/ChaosEngineInterface.h b/Engine/Source/Runtime/PhysicsCore/Public/Chaos/ChaosEngineInterface.h
--- a/Engine/Source/Runtime/PhysicsCore/Public/Chaos/ChaosEngineInterface.h
+++ b/Engine/Source/Runtime/PhysicsCore/Public/Chaos/ChaosEngineInterface.h
@@ -488,6 +488,7 @@ public:
 	static PHYSICSCORE_API FPhysicsConstraintHandle CreateConstraint(Chaos::FPhysicsObject* Body1, Chaos::FPhysicsObject* Body2, const FTransform& InLocalFrame1, const FTransform& InLocalFrame2);
 	static PHYSICSCORE_API FPhysicsConstraintHandle CreateConstraint(const FPhysicsActorHandle& InActorRef1,const FPhysicsActorHandle& InActorRef2,const FTransform& InLocalFrame1,const FTransform& InLocalFrame2);
 
+	static PHYSICSCORE_API FPhysicsConstraintHandle CreateSuspension(Chaos::FPhysicsObject* Body, const FVector& InLocalFrame);
 	static PHYSICSCORE_API FPhysicsConstraintHandle CreateSuspension(const FPhysicsActorHandle& InActorRef, const FVector& InLocalFrame);
 	static PHYSICSCORE_API void SetConstraintUserData(const FPhysicsConstraintHandle& InConstraintRef,void* InUserData);
 	static PHYSICSCORE_API void ReleaseConstraint(FPhysicsConstraintHandle& InConstraintRef);
