Iris
- Disabled net.Iris.Attachments.AllowSendPolicyFlags until we have fixed missing exports for OOB rpcs

diff --git a/Engine/Source/Runtime/Experimental/Iris/Core/Private/Iris/ReplicationSystem/ReplicationSystem.cpp b/Engine/Source/Runtime/Experimental/Iris/Core/Private/Iris/ReplicationSystem/ReplicationSystem.cpp
--- a/Engine/Source/Runtime/Experimental/Iris/Core/Private/Iris/ReplicationSystem/ReplicationSystem.cpp
+++ b/Engine/Source/Runtime/Experimental/Iris/Core/Private/Iris/ReplicationSystem/ReplicationSystem.cpp
@@ -34,7 +34,7 @@ static bool bForcePruneBeforeUpdate = false;
 static FAutoConsoleVariableRef CVarForcePruneBeforeUpdate(TEXT("net.Iris.ForcePruneBeforeUpdate"), bForcePruneBeforeUpdate, TEXT("Verify integrity of all tracked instances at the start of every update."));
 #endif
 
-static bool bAllowAttachmentSendPolicyFlags = true;
+static bool bAllowAttachmentSendPolicyFlags = false;
 static FAutoConsoleVariableRef CVarAllowAttachmentSendPolicyFlags(TEXT("net.Iris.Attachments.AllowSendPolicyFlags"), bAllowAttachmentSendPolicyFlags, TEXT("Allow use of ENetObjectAttachmentSendPolicyFlags to specify behavior of RPCs."));
 
 }
