- governor returns approximate number of cycles until it will release some allowance
- renamed a couple members to make it units and association
- updated service thread's use of governor to reflect the above

#rnx

diff --git a/Engine/Source/Runtime/Experimental/IoStoreOnDemand/Private/JournaledCache.cpp b/Engine/Source/Runtime/Experimental/IoStoreOnDemand/Private/JournaledCache.cpp
--- a/Engine/Source/Runtime/Experimental/IoStoreOnDemand/Private/JournaledCache.cpp
+++ b/Engine/Source/Runtime/Experimental/IoStoreOnDemand/Private/JournaledCache.cpp
@@ -1256,8 +1256,8 @@ public:
 			FGovernor();
 	void	Set(uint32 Allowance, uint32 Ops, uint32 Seconds);
 	void	SetDemands(uint32 Threshold, uint32 Boost, uint32 SuperBoost);
-	uint32	BeginAllowance(uint32 DemandPercent);
-	bool	EndAllowance(uint32 UnusedAllowance);
+	int32	BeginAllowance(uint32 DemandPercent);
+	int32	EndAllowance(uint32 UnusedAllowance);
 
 private:
 	enum class EState : uint8
@@ -1266,14 +1266,15 @@ private:
 		Rolling,
 	};
 
+	int32	GetMaxWaitCycles() const;
 	void	Set(uint32 Allowance, uint32 Ops, uint32 Seconds, int64 CycleFreq);
-	uint32	BeginInternal(uint32 Demand, int64 Cycle);
-	int64	FlushInterval;
+	int32	BeginInternal(uint32 Demand, int64 Cycle);
+	int64	OpInterval;
 	int64	PrevCycles;
 	uint32	RunOff = 0;
 	uint32	OpCount = 0;
 	uint32	MaxOpCount;
-	uint32	AllowanceInterval;
+	uint32	OpAllowance;
 	uint8	DemandThreshold = 30;
 	uint8	DemandBoost = 60;
 	uint8	DemandSuperBoost = 87;
@@ -1299,8 +1300,8 @@ void FGovernor::Set(uint32 Allowance, uint32 Ops, uint32 Seconds, int64 CycleFre
 	int32 CommitBufferSize = JournaledCache::GetWriteCommitThreshold();
 	if (CommitBufferSize == 0)
 	{
-		AllowanceInterval = Allowance / Ops;
-		FlushInterval = (CycleFreq * Seconds) / Ops;
+		OpAllowance = Allowance / Ops;
+		OpInterval = (CycleFreq * Seconds) / Ops;
 		MaxOpCount = 4;
 		return;
 	}
@@ -1309,8 +1310,8 @@ void FGovernor::Set(uint32 Allowance, uint32 Ops, uint32 Seconds, int64 CycleFre
 	int32 CommitOpCost = BlockCount * 3;
 	MaxOpCount = (Ops - CommitOpCost) / BlockCount;
 
-	AllowanceInterval = CommitBufferSize / MaxOpCount;
-	FlushInterval = (Seconds * CycleFreq) / (BlockCount * (MaxOpCount - 1));
+	OpAllowance = CommitBufferSize / MaxOpCount;
+	OpInterval = (Seconds * CycleFreq) / (BlockCount * (MaxOpCount - 1));
 }
 
 ////////////////////////////////////////////////////////////////////////////////
@@ -1322,9 +1323,9 @@ void FGovernor::SetDemands(uint32 Threshold, uint32 Boost, uint32 SuperBoost)
 }
 
 ////////////////////////////////////////////////////////////////////////////////
-uint32 FGovernor::BeginInternal(uint32 Demand, int64 Cycles)
+int32 FGovernor::BeginInternal(uint32 Demand, int64 Cycles)
 {
-	int64 Interval = FlushInterval;
+	int64 Interval = OpInterval;
 	Interval >>= int32(Demand >= DemandBoost);
 	Interval >>= int32(Demand >= DemandSuperBoost);
 	Interval <<= int32(Demand <= DemandThreshold);
@@ -1332,7 +1333,7 @@ uint32 FGovernor::BeginInternal(uint32 Demand, int64 Cycles)
 	int64 Delta = Cycles - PrevCycles;
 	if (Delta <= Interval)
 	{
-		return 0;
+		return -int32(Interval - Delta);
 	}
 
 	do { Delta -= Interval; } while (Delta > Interval);
@@ -1340,12 +1341,15 @@ uint32 FGovernor::BeginInternal(uint32 Demand, int64 Cycles)
 
 	OpCount++;
 
-	return AllowanceInterval + RunOff;
+	return OpAllowance + RunOff;
 }
 
 ////////////////////////////////////////////////////////////////////////////////
-uint32 FGovernor::BeginAllowance(uint32 DemandPercent)
+int32 FGovernor::BeginAllowance(uint32 DemandPercent)
 {
+	// A return of >=0 is the allowance of bytes that can be read. Otherwise the
+	// value is number of cycles to wait until allowance may be ready.
+
 	if (State == EState::Rolling)
 	{
 		int64 Cycles = FPlatformTime::Cycles64();
@@ -1354,28 +1358,35 @@ uint32 FGovernor::BeginAllowance(uint32 DemandPercent)
 
 	if (DemandPercent < DemandThreshold)
 	{
-		return 0;
+		return 0 - GetMaxWaitCycles();
 	}
 
 	State = EState::Rolling;
 	PrevCycles = FPlatformTime::Cycles64();
 	OpCount = 1;
 	RunOff = 0;
-	return AllowanceInterval;
+	return OpAllowance;
 }
 
 ////////////////////////////////////////////////////////////////////////////////
-bool FGovernor::EndAllowance(uint32 UnusedAllowance)
+int32 FGovernor::EndAllowance(uint32 UnusedAllowance)
 {
 	RunOff = UnusedAllowance;
 
 	if (OpCount >= MaxOpCount)
 	{
 		State = EState::Waiting;
-		return true;
+		return 0 - GetMaxWaitCycles();
 	}
 
-	return false;
+	return GetMaxWaitCycles();
+}
+
+////////////////////////////////////////////////////////////////////////////////
+int32 FGovernor::GetMaxWaitCycles() const
+{
+	// ">> 2" so we check at four times the speed in case of a S.U.P.E.R BOOST
+	return int32(OpInterval >> 2);
 }
 
 
@@ -1659,8 +1670,8 @@ void FServiceThread::ReceiveWork()
 void FServiceThread::UpdateCache(FCache* Cache)
 {
 	uint32 Demand = Cache->GetDemand();
-	uint32 Allowance = Governor.BeginAllowance(Demand);
-	if (Allowance == 0)
+	int32 Allowance = Governor.BeginAllowance(Demand);
+	if (Allowance <= 0)
 	{
 		return;
 	}
@@ -1668,23 +1679,22 @@ void FServiceThread::UpdateCache(FCache* Cache)
 	TRACE_COUNTER_SET(IasMemDemand, Demand);
 	TRACE_COUNTER_SET(IasAllowance, 0);
 	TRACE_COUNTER_SET(IasAllowance, Allowance);
-	do
-	{
-		TRACE_CPUPROFILER_EVENT_SCOPE(IasCache::Update);
 
-		uint32 AllowanceUsed = Cache->WriteMemToDisk(Allowance);
-		uint32 Unused = Allowance - AllowanceUsed;
-		TRACE_COUNTER_SET(IasAllowance, Unused);
+	TRACE_CPUPROFILER_EVENT_SCOPE(IasCache::Update);
 
-		TRACE_COUNTER_ADD(IasOpCount, 1);
-		if (Governor.EndAllowance(Unused))
-		{
-			TRACE_COUNTER_ADD(IasOpCount, 1); // flush from closing .bin file - we can remove this!
-			TRACE_COUNTER_ADD(IasOpCount, 3); // write-flush-commit from .jrn.
-			Cache->Flush();
-		}
+	uint32 AllowanceUsed = Cache->WriteMemToDisk(Allowance);
+	uint32 Unused = Allowance - AllowanceUsed;
+	TRACE_COUNTER_SET(IasAllowance, Unused);
+	TRACE_COUNTER_ADD(IasOpCount, 1);
+
+	int32 WaitCycles = Governor.EndAllowance(Unused);
+	if (WaitCycles < 0)
+	{
+		TRACE_COUNTER_ADD(IasOpCount, 1); // flush from closing .bin file - we can remove this if we use a single file for .jrn and .bin
+		TRACE_COUNTER_ADD(IasOpCount, 3); // write-flush-commit from .jrn.
+		Cache->Flush();
 	}
-	while (false);
+
 	TRACE_COUNTER_SET(IasAllowance, 0);
 }
 
