Iris
* Split scope bitarray into GlobalScopeList and ScopeFrameData
* The GlobalScopeList is always up to date and should mostly be used outside the NetTickFlush update
* ScopeFrameData is synced at the start of NetTickFlush and valid until the end of the SendUpdate().
* The lists contained in ScopeFrameData are now protected from illegal access outside of NetTickFlush
* Added ensures that prevents creating new replicated root object during the SendUpdate operations.
* Added ensures that prevent changing filter conditions during the SendUpdate operations.
* Added support for replication of subobjects created inside a PreReplication callback.

#rnx

diff --git a/Engine/Plugins/Runtime/ReplicationSystemTestPlugin/Source/Private/Tests/ReplicationSystem/TestPreUpdateCallback.cpp b/Engine/Plugins/Runtime/ReplicationSystemTestPlugin/Source/Private/Tests/ReplicationSystem/TestPreUpdateCallback.cpp
new file mode 100644
--- /dev/null
+++ b/Engine/Plugins/Runtime/ReplicationSystemTestPlugin/Source/Private/Tests/ReplicationSystem/TestPreUpdateCallback.cpp
@@ -0,0 +1,77 @@
+// Copyright Epic Games, Inc. All Rights Reserved.
+
+#include "NetworkAutomationTest.h"
+#include "NetworkAutomationTestMacros.h"
+#include "Iris/ReplicationSystem/DirtyNetObjectTracker.h"
+#include "Iris/ReplicationSystem/ReplicationSystemInternal.h"
+#include "Tests/ReplicationSystem/ReplicationSystemTestFixture.h"
+#include "Tests/ReplicationSystem/ReplicationSystemServerClientTestFixture.h"
+#include "Iris/ReplicationSystem/Filtering/ReplicationFiltering.h"
+#include "Iris/ReplicationSystem/ReplicationSystemInternal.h"
+#include "Iris/ReplicationSystem/Filtering/NetObjectFilter.h"
+#include "Iris/ReplicationSystem/Filtering/NetObjectGroups.h"
+
+namespace UE::Net::Private
+{
+
+/** Test that validates behavior when creating replicated sub objects inside PreUpdate/PreReplication */
+UE_NET_TEST_FIXTURE(FReplicationSystemServerClientTestFixture, CreateObjectsInsidePreUpdateTest)
+{
+	// Add client
+	FReplicationSystemTestClient* Client = CreateClient();
+
+	// Spawn object on server that is polled late in order to test ForceNetUpdate
+	UObjectReplicationBridge::FCreateNetRefHandleParams Params;
+	const uint32 PollPeriod = 100;
+	const float PollFrequency = Server->ConvertPollPeriodIntoFrequency(PollPeriod);
+	Params.PollFrequency = PollFrequency;
+	Params.bCanReceive = true;
+	Params.bAllowDynamicFilter = true;
+	Params.bNeedsPreUpdate = true;
+	UTestReplicatedIrisObject* ServerObject = Server->CreateObject(Params);
+	UTestReplicatedIrisObject* ServerSubObject(nullptr);
+
+	// Send and deliver packet
+	Server->UpdateAndSend({ Client });
+
+	// Object should have been created on the client
+	UTestReplicatedIrisObject* ClientObject = Cast<UTestReplicatedIrisObject>(Client->GetReplicationBridge()->GetReplicatedObject(ServerObject->NetRefHandle));
+	UE_NET_ASSERT_NE(ClientObject, nullptr);
+
+	// Add a PreUpdate where we create a new subobject
+	auto PreUpdateObject = [&](FNetRefHandle NetHandle, UObject* ReplicatedObject, const UReplicationBridge* ReplicationBridge)
+	{
+		if (ServerObject == ReplicatedObject)
+		{
+			if (ServerSubObject == nullptr)
+			{
+				// Dirty a property
+				ServerObject->IntA = 0xBB;
+
+				// Create a subobject
+				ServerSubObject = Server->CreateSubObject<UTestReplicatedIrisObject>(ServerObject->NetRefHandle);
+
+				// Dirty this subobject
+				ServerSubObject->IntA = 0xBB;
+
+			}
+		}
+	};
+	Server->GetReplicationBridge()->SetExternalPreUpdateFunctor(PreUpdateObject);
+
+	Server->ReplicationSystem->ForceNetUpdate(ServerObject->NetRefHandle);
+	Server->UpdateAndSend({ Client });
+	
+	// Was the property changed in PreUpdate replicated ?
+	UE_NET_ASSERT_EQ(ClientObject->IntA, ServerObject->IntA);
+
+	// And in the same Send was the subobject replicated ?
+	UTestReplicatedIrisObject* ClientSubObject = Cast<UTestReplicatedIrisObject>(Client->GetReplicationBridge()->GetReplicatedObject(ServerSubObject->NetRefHandle));
+	UE_NET_ASSERT_NE(ClientSubObject, nullptr);
+
+	// The dirty property of the subobject should have replicated too.
+	UE_NET_ASSERT_EQ(ClientSubObject->IntA, ClientSubObject->IntA);
+}
+
+
+}
\ No newline at end of file
diff --git a/Engine/Source/Runtime/Experimental/Iris/Core/Private/Iris/ReplicationSystem/Conditionals/ReplicationConditionals.cpp b/Engine/Source/Runtime/Experimental/Iris/Core/Private/Iris/ReplicationSystem/Conditionals/ReplicationConditionals.cpp
--- a/Engine/Source/Runtime/Experimental/Iris/Core/Private/Iris/ReplicationSystem/Conditionals/ReplicationConditionals.cpp
+++ b/Engine/Source/Runtime/Experimental/Iris/Core/Private/Iris/ReplicationSystem/Conditionals/ReplicationConditionals.cpp
@@ -681,8 +681,8 @@ void FReplicationConditionals::UpdateObjectsInScope()
 {
 	IRIS_PROFILER_SCOPE(FReplicationConditionals_UpdateObjectsInScope);
 
-	const FNetBitArrayView ObjectsInScope = NetRefHandleManager->GetScopableInternalIndicesView();
-	const FNetBitArrayView PrevObjectsInScope = NetRefHandleManager->GetPrevFrameScopableInternalIndicesView();
+	const FNetBitArrayView ObjectsInScope = NetRefHandleManager->GetCurrentFrameScopableInternalIndices();
+	const FNetBitArrayView PrevObjectsInScope = NetRefHandleManager->GetPrevFrameScopableInternalIndices();
 
 	const uint32 WordCountForModifiedWords = Align(FPlatformMath::Max(MaxObjectCount, 1U), 32U)/32U;
 	TArray<uint32> ModifiedWords;
diff --git a/Engine/Source/Runtime/Experimental/Iris/Core/Private/Iris/ReplicationSystem/DeltaCompression/DeltaCompressionBaselineManager.cpp b/Engine/Source/Runtime/Experimental/Iris/Core/Private/Iris/ReplicationSystem/DeltaCompression/DeltaCompressionBaselineManager.cpp
--- a/Engine/Source/Runtime/Experimental/Iris/Core/Private/Iris/ReplicationSystem/DeltaCompression/DeltaCompressionBaselineManager.cpp
+++ b/Engine/Source/Runtime/Experimental/Iris/Core/Private/Iris/ReplicationSystem/DeltaCompression/DeltaCompressionBaselineManager.cpp
@@ -158,8 +158,8 @@ void FDeltaCompressionBaselineManager::UpdateScope()
 	using SignedWordType = TSignedIntType<sizeof(WordType)>::Type;
 	constexpr SIZE_T BitsPerWord = sizeof(WordType)*8U;
 
-	const FNetBitArrayView ObjectsInScope = NetRefHandleManager->GetScopableInternalIndicesView();
-	const FNetBitArrayView PrevObjectsInScope = NetRefHandleManager->GetPrevFrameScopableInternalIndicesView();
+	const FNetBitArrayView ObjectsInScope = NetRefHandleManager->GetCurrentFrameScopableInternalIndices();
+	const FNetBitArrayView PrevObjectsInScope = NetRefHandleManager->GetPrevFrameScopableInternalIndices();
 
 	const WordType* ObjectsInScopeStorage = ObjectsInScope.GetData();
 	const WordType* PrevObjectsInScopeStorage = PrevObjectsInScope.GetData();
diff --git a/Engine/Source/Runtime/Experimental/Iris/Core/Private/Iris/ReplicationSystem/DirtyNetObjectTracker.cpp b/Engine/Source/Runtime/Experimental/Iris/Core/Private/Iris/ReplicationSystem/DirtyNetObjectTracker.cpp
--- a/Engine/Source/Runtime/Experimental/Iris/Core/Private/Iris/ReplicationSystem/DirtyNetObjectTracker.cpp
+++ b/Engine/Source/Runtime/Experimental/Iris/Core/Private/Iris/ReplicationSystem/DirtyNetObjectTracker.cpp
@@ -101,7 +101,7 @@ void FDirtyNetObjectTracker::UpdateDirtyNetObjects()
 		}
 	}
 
-	const uint32* GlobalScopeListData = NetRefHandleManager->GetScopableInternalIndicesView().GetData();
+	const uint32* GlobalScopeListData = NetRefHandleManager->GetCurrentFrameScopableInternalIndices().GetData();
 	uint32* AccumulatedDirtyNetObjectsData = AccumulatedDirtyNetObjects.GetData();
 
 	const uint32 NumWords = AccumulatedDirtyNetObjects.GetNumWords();
diff --git a/Engine/Source/Runtime/Experimental/Iris/Core/Private/Iris/ReplicationSystem/Filtering/ReplicationFiltering.cpp b/Engine/Source/Runtime/Experimental/Iris/Core/Private/Iris/ReplicationSystem/Filtering/ReplicationFiltering.cpp
--- a/Engine/Source/Runtime/Experimental/Iris/Core/Private/Iris/ReplicationSystem/Filtering/ReplicationFiltering.cpp
+++ b/Engine/Source/Runtime/Experimental/Iris/Core/Private/Iris/ReplicationSystem/Filtering/ReplicationFiltering.cpp
@@ -328,11 +328,11 @@ void FReplicationFiltering::FilterNonRelevantObjects()
 	if (!bCVarRepFilterCullNonRelevant)
 	{
 		// Make every object in the global scope part of the relevant list
-		NetRefHandleManager->GetRelevantObjectsInternalIndices().Copy(NetRefHandleManager->GetScopableInternalIndicesView());
+		NetRefHandleManager->GetRelevantObjectsInternalIndices().Copy(NetRefHandleManager->GetCurrentFrameScopableInternalIndices());
 		return;
 	}
 
-	const uint32* const GlobalScopeData = NetRefHandleManager->GetScopableInternalIndicesView().GetData();
+	const uint32* const CurrentFrameScopeData = NetRefHandleManager->GetCurrentFrameScopableInternalIndices().GetData();
 	const uint32* const WithOwnerData = ObjectsWithOwnerFilter.GetData();
 	const uint32* const ConnectionFiltersData = AllConnectionFilteredObjects.GetData();
 	const uint32* const DynamicFilteredData = DynamicFilterEnabledObjects.GetData();
@@ -347,7 +347,7 @@ void FReplicationFiltering::FilterNonRelevantObjects()
 	for (uint32 WordIndex = 0; WordIndex < MaxWords; ++WordIndex)
 	{
 		// Build the list of always relevant objects (objects that have no filters)
-		GlobalRelevantData[WordIndex] = GlobalScopeData[WordIndex] & ~(WithOwnerData[WordIndex] | ConnectionFiltersData[WordIndex] | DynamicFilteredData[WordIndex] | GroupFilteredData[WordIndex]);
+		GlobalRelevantData[WordIndex] = CurrentFrameScopeData[WordIndex] & ~(WithOwnerData[WordIndex] | ConnectionFiltersData[WordIndex] | DynamicFilteredData[WordIndex] | GroupFilteredData[WordIndex]);
 	}
 
 	// Build the list of currently relevant objects. e.g. always relevant objects + filterable objects relevant to at least one connection.	
@@ -663,7 +663,7 @@ void FReplicationFiltering::InitNewConnections()
 	auto InitNewConnection = [this](uint32 ConnectionId)
 	{
 		// Copy default scope
-		const FNetBitArrayView ScopableInternalIndices = NetRefHandleManager->GetScopableInternalIndicesView();
+		const FNetBitArrayView ScopableInternalIndices = NetRefHandleManager->GetCurrentFrameScopableInternalIndices();
 		FPerConnectionInfo& ConnectionInfo = this->ConnectionInfos[ConnectionId];
 
 		ConnectionInfo.ConnectionFilteredObjects.Init(MaxObjectCount);
@@ -778,8 +778,8 @@ void FReplicationFiltering::ResetRemovedConnections()
 
 void FReplicationFiltering::UpdateObjectsInScope()
 {
-	const FNetBitArrayView ObjectsInScope = NetRefHandleManager->GetScopableInternalIndicesView();
-	const FNetBitArrayView PrevObjectsInScope = NetRefHandleManager->GetPrevFrameScopableInternalIndicesView();
+	const FNetBitArrayView ObjectsInScope = NetRefHandleManager->GetCurrentFrameScopableInternalIndices();
+	const FNetBitArrayView PrevObjectsInScope = NetRefHandleManager->GetPrevFrameScopableInternalIndices();
 
 	/**
 	 * It's possible for an object to be created, have some filtering applied and then be removed later the same frame.
@@ -996,7 +996,7 @@ void FReplicationFiltering::UpdateOwnerAndConnectionFiltering()
 		ObjectsWithDirtyOwner.ForAllSetBits(UpdateOwners);
 	}
 
-	const FNetBitArrayView GlobalObjectsInScope = NetRefHandleManager->GetScopableInternalIndicesView();
+	const FNetBitArrayView CurrentFrameObjectsInScope = NetRefHandleManager->GetCurrentFrameScopableInternalIndices();
 
 	// Update filtering
 	if (bHasDirtyConnectionFilter)
@@ -1006,7 +1006,7 @@ void FReplicationFiltering::UpdateOwnerAndConnectionFiltering()
 			AllConnectionFilteredObjects.SetBitValue(DirtyObjectIndex, HasConnectionFilter(DirtyObjectIndex));
 		});
 
-		auto UpdateConnectionScope = [this, &GlobalObjectsInScope](uint32 ConnectionId)
+		auto UpdateConnectionScope = [this, &CurrentFrameObjectsInScope](uint32 ConnectionId)
 		{
 			FPerConnectionInfo& ConnectionInfo = this->ConnectionInfos[ConnectionId];
 			FNetBitArrayView ConnectionScope = MakeNetBitArrayView(ConnectionInfo.ConnectionFilteredObjects);
@@ -1015,7 +1015,7 @@ void FReplicationFiltering::UpdateOwnerAndConnectionFiltering()
 
 			// Update filter info
 			{
-				auto MaskObject = [this, &ConnectionScope, &GroupFilteredOutObjects, &ObjectsInScopeBeforeDynamicFiltering, ConnectionId, &GlobalObjectsInScope](uint32 ObjectIndex)
+				auto MaskObject = [this, &ConnectionScope, &GroupFilteredOutObjects, &ObjectsInScopeBeforeDynamicFiltering, ConnectionId, &CurrentFrameObjectsInScope](uint32 ObjectIndex)
 				{
 					bool bObjectIsInScope = true;
 					if (HasOwnerFilter(ObjectIndex))
@@ -1042,7 +1042,7 @@ void FReplicationFiltering::UpdateOwnerAndConnectionFiltering()
 					// Subobjects follow suit.
 					for (const FInternalNetRefIndex SubObjectIndex : NetRefHandleManager->GetSubObjects(ObjectIndex))
 					{
-						const bool bEnableObject = bObjectIsInScope && GlobalObjectsInScope.GetBit(SubObjectIndex);
+						const bool bEnableObject = bObjectIsInScope && CurrentFrameObjectsInScope.GetBit(SubObjectIndex);
 						const bool bIsGroupEnabled = !GroupFilteredOutObjects.GetBit(SubObjectIndex);
 
 						ConnectionScope.SetBitValue(SubObjectIndex, bEnableObject);
@@ -1078,9 +1078,9 @@ void FReplicationFiltering::UpdateGroupFiltering()
 			const FNetObjectGroup* Group = this->Groups->GetGroupByIndex(FNetObjectGroupHandle::FGroupIndexType(GroupIndex));
 			FPerConnectionInfo* LocalConnectionInfos = this->ConnectionInfos.GetData();
 
-			const FNetBitArrayView GlobalScopableObjects = NetRefHandleManager->GetScopableInternalIndicesView();
+			const FNetBitArrayView CurrentFrameScopableObjects = NetRefHandleManager->GetCurrentFrameScopableInternalIndices();
 
-			auto UpdateGroupFilterForConnection = [this, ConnectionStateInfo, Group, LocalConnectionInfos, GlobalScopableObjects](uint32 ConnectionId)
+			auto UpdateGroupFilterForConnection = [this, ConnectionStateInfo, Group, LocalConnectionInfos, CurrentFrameScopableObjects](uint32 ConnectionId)
 			{
 				if (this->GetConnectionFilterStatus(*ConnectionStateInfo, ConnectionId) == ENetFilterStatus::Disallow)
 				{
@@ -1095,7 +1095,7 @@ void FReplicationFiltering::UpdateGroupFiltering()
 						// Filter subobjects
 						for (const FInternalNetRefIndex SubObjectIndex : NetRefHandleManager->GetSubObjects(ObjectIndex))
 						{
-							const bool bIsScopable = GlobalScopableObjects.GetBit(SubObjectIndex);
+							const bool bIsScopable = CurrentFrameScopableObjects.GetBit(SubObjectIndex);
 							GroupFilteredOutObjects.SetBitValue(SubObjectIndex, bIsScopable);
 							ObjectsInScopeBeforeDynamicFiltering.ClearBit(SubObjectIndex);
 						}
@@ -1324,7 +1324,7 @@ void FReplicationFiltering::PostUpdateDynamicFiltering(ENetFilterType FilterType
 {
 	IRIS_PROFILER_SCOPE(FReplicationFiltering_PostUpdateDynamicFiltering);
 
-	// Give filters a chance to prepare for filtering. It's only called if any object has the filter set.
+	// Tell filters to clean up after filtering. It's only called if any object has the filter set.
 	{
 		FNetObjectPostFilteringParams PostFilteringParams;
 		for (FFilterInfo& Info : DynamicFilterInfos)
@@ -1860,15 +1860,17 @@ void FReplicationFiltering::InternalSetGroupFilterStatus(FNetObjectGroupHandle G
 					FPerConnectionInfo& ConnectionInfo = ConnectionInfos[ConnectionId];
 					FNetBitArray& GroupFilteredOutObjects = ConnectionInfo.GroupFilteredOutObjects;
 					FNetBitArray& ObjectsInScopeBeforeDynamicFiltering = ConnectionInfo.ObjectsInScopeBeforeDynamicFiltering;
-					const FNetBitArrayView ScopableObjects = NetRefHandleManager->GetScopableInternalIndicesView();
+					const FNetBitArrayView GlobalScopableObjects = NetRefHandleManager->GetGlobalScopableInternalIndices();
 					for (uint32 ObjectIndex : MakeArrayView(Group->Members.GetData(), Group->Members.Num()))
 					{
 						GroupFilteredOutObjects.SetBit(ObjectIndex);
+
+						//$IRIS TODO: ObjectsInScopeBeforeDynamicFiltering should not be accessed outside PreSendUpdate since its reset there. Is this needed ?
 						ObjectsInScopeBeforeDynamicFiltering.ClearBit(ObjectIndex);
 						// Filter subobjects
 						for (const FInternalNetRefIndex SubObjectIndex : NetRefHandleManager->GetSubObjects(ObjectIndex))
 						{
-							GroupFilteredOutObjects.SetBitValue(SubObjectIndex, ScopableObjects.GetBit(SubObjectIndex));
+							GroupFilteredOutObjects.SetBitValue(SubObjectIndex, GlobalScopableObjects.GetBit(SubObjectIndex));
 							ObjectsInScopeBeforeDynamicFiltering.ClearBit(SubObjectIndex);
 						}
 					}
diff --git a/Engine/Source/Runtime/Experimental/Iris/Core/Private/Iris/ReplicationSystem/NetBlob/NetBlobManager.cpp b/Engine/Source/Runtime/Experimental/Iris/Core/Private/Iris/ReplicationSystem/NetBlob/NetBlobManager.cpp
--- a/Engine/Source/Runtime/Experimental/Iris/Core/Private/Iris/ReplicationSystem/NetBlob/NetBlobManager.cpp
+++ b/Engine/Source/Runtime/Experimental/Iris/Core/Private/Iris/ReplicationSystem/NetBlob/NetBlobManager.cpp
@@ -19,13 +19,6 @@ namespace UE::Net::Private
 static TAutoConsoleVariable<int32> CVarEnableIrisRPCs(TEXT("net.Iris.EnableRPCs"), 1, TEXT( "If > 0 let Iris replicate and execute RPCs."));
 
 FNetBlobManager::FNetBlobManager()
-: ReplicationSystem(nullptr)
-, ObjectReferenceCache(nullptr)
-, Connections(nullptr)
-, PartialNetObjectAttachmentHandlerConfig(nullptr)
-, NetRefHandleManager(nullptr)
-, bIsServer(false)
-, bSendAttachmentsWithObject(false)
 {
 }
 	
@@ -39,6 +32,7 @@ void FNetBlobManager::Init(FNetBlobManagerInitParams& InitParams)
 	ObjectReferenceCache = &InitParams.ReplicationSystem->GetReplicationSystemInternal()->GetObjectReferenceCache();
 	bIsServer = ReplicationSystem->IsServer();
 	bSendAttachmentsWithObject = InitParams.bSendAttachmentsWithObject; 
+	bAllowObjectReplication = ReplicationSystem->AllowObjectReplication();
 
 	RegisterDefaultHandlers();
 
@@ -369,25 +363,37 @@ void FNetBlobManager::FNetObjectAttachmentSendQueue::PrepareProcessQueue(FReplic
 		ReplicatingConnections.GetSetBitIndices(0, ~0, ProcessContext.ConnectionIds.GetData(), ProcessContext.ConnectionIds.Num());
 	}
 
-	// Figure out if we have any attachments to objects going out of scope.
-	const FNetBitArrayView ScopableObjects = InNetRefHandleManager->GetScopableInternalIndicesView();
-	const FNetBitArrayView PrevScopableObjects = InNetRefHandleManager->GetPrevFrameScopableInternalIndicesView();
-	
-	uint32 CurrentEntryIndex = 0U;
-	for (const FNetObjectAttachmentQueueEntry& Entry : MakeArrayView(AttachmentQueue))
+	if (Manager->AllowObjectReplication())
 	{
-		const uint32 TargetInternalObjectIndex = Entry.SubObjectIndex != FNetRefHandleManager::InvalidInternalIndex ? Entry.SubObjectIndex : Entry.OwnerIndex;
+		// Figure out if we have any attachments to objects going out of scope.
+		const FNetBitArrayView ScopableObjects = InNetRefHandleManager->GetCurrentFrameScopableInternalIndices();
+		const FNetBitArrayView PrevScopableObjects = InNetRefHandleManager->GetPrevFrameScopableInternalIndices();
 
-		const bool bIsAttachmentToObjectGoingOutOfScope = !ScopableObjects.GetBit(TargetInternalObjectIndex) && PrevScopableObjects.GetBit(TargetInternalObjectIndex);
-		if (bIsAttachmentToObjectGoingOutOfScope)
+		uint32 CurrentEntryIndex = 0U;
+		for (const FNetObjectAttachmentQueueEntry& Entry : MakeArrayView(AttachmentQueue))
 		{
-			ProcessContext.AttachmentsToObjectsGoingOutOfScope.SetBit(CurrentEntryIndex);
+			const uint32 TargetInternalObjectIndex = Entry.SubObjectIndex != FNetRefHandleManager::InvalidInternalIndex ? Entry.SubObjectIndex : Entry.OwnerIndex;
+
+			const bool bIsAttachmentToObjectGoingOutOfScope = !ScopableObjects.GetBit(TargetInternalObjectIndex) && PrevScopableObjects.GetBit(TargetInternalObjectIndex);
+			if (bIsAttachmentToObjectGoingOutOfScope)
+			{
+				ProcessContext.AttachmentsToObjectsGoingOutOfScope.SetBit(CurrentEntryIndex);
+			}
+			else
+			{
+				ProcessContext.AttachmentsToObjectsInScope.SetBit(CurrentEntryIndex);
+			}
+			++CurrentEntryIndex;
 		}
-		else
+	}
+	else
+	{
+		uint32 CurrentEntryIndex = 0U;
+		for (const FNetObjectAttachmentQueueEntry& Entry : MakeArrayView(AttachmentQueue))
 		{
 			ProcessContext.AttachmentsToObjectsInScope.SetBit(CurrentEntryIndex);
+			++CurrentEntryIndex;
 		}
-		++CurrentEntryIndex;
 	}
 }
 
diff --git a/Engine/Source/Runtime/Experimental/Iris/Core/Private/Iris/ReplicationSystem/NetBlob/NetBlobManager.h b/Engine/Source/Runtime/Experimental/Iris/Core/Private/Iris/ReplicationSystem/NetBlob/NetBlobManager.h
--- a/Engine/Source/Runtime/Experimental/Iris/Core/Private/Iris/ReplicationSystem/NetBlob/NetBlobManager.h
+++ b/Engine/Source/Runtime/Experimental/Iris/Core/Private/Iris/ReplicationSystem/NetBlob/NetBlobManager.h
@@ -43,6 +43,8 @@ public:
 
 	bool RegisterNetBlobHandler(UNetBlobHandler* Handler);
 
+	bool AllowObjectReplication() const { return bAllowObjectReplication; }
+
 	bool QueueNetObjectAttachment(uint32 ConnectionId, const FNetObjectReference& TargetRef, const TRefCountPtr<FNetObjectAttachment>& Attachment, ENetObjectAttachmentSendPolicyFlags SendFlags = ENetObjectAttachmentSendPolicyFlags::None);
 	bool SendRPC(const UObject* Object, const UObject* SubObject, const UFunction* Function, const void* Parameters, ENetObjectAttachmentSendPolicyFlags SendFlags = ENetObjectAttachmentSendPolicyFlags::None);
 	bool SendRPC(uint32 ConnectionId, const UObject* Object, const UObject* SubObject, const UFunction* Function, const void* Parameters, ENetObjectAttachmentSendPolicyFlags SendFlags = ENetObjectAttachmentSendPolicyFlags::None);
@@ -132,13 +134,14 @@ private:
 	TStrongObjectPtr<UPartialNetObjectAttachmentHandler> PartialNetObjectAttachmentHandler;
 	TStrongObjectPtr<UNetObjectBlobHandler> NetObjectBlobHandler;
 
-	UReplicationSystem* ReplicationSystem;
-	FObjectReferenceCache* ObjectReferenceCache;
-	FReplicationConnections* Connections;
-	const UPartialNetObjectAttachmentHandlerConfig* PartialNetObjectAttachmentHandlerConfig;
-	const FNetRefHandleManager* NetRefHandleManager;
-	bool bIsServer;
-	bool bSendAttachmentsWithObject;
+	UReplicationSystem* ReplicationSystem = nullptr;
+	FObjectReferenceCache* ObjectReferenceCache = nullptr;
+	FReplicationConnections* Connections = nullptr;
+	const UPartialNetObjectAttachmentHandlerConfig* PartialNetObjectAttachmentHandlerConfig = nullptr;
+	const FNetRefHandleManager* NetRefHandleManager = nullptr;
+	bool bIsServer = false;
+	bool bSendAttachmentsWithObject = false;
+	bool bAllowObjectReplication = false;
 };
 
 }
diff --git a/Engine/Source/Runtime/Experimental/Iris/Core/Private/Iris/ReplicationSystem/NetRefHandleManager.cpp b/Engine/Source/Runtime/Experimental/Iris/Core/Private/Iris/ReplicationSystem/NetRefHandleManager.cpp
--- a/Engine/Source/Runtime/Experimental/Iris/Core/Private/Iris/ReplicationSystem/NetRefHandleManager.cpp
+++ b/Engine/Source/Runtime/Experimental/Iris/Core/Private/Iris/ReplicationSystem/NetRefHandleManager.cpp
@@ -21,8 +21,8 @@ FNetRefHandleManager::FNetRefHandleManager(FReplicationProtocolManager& InReplic
 : ActiveObjectCount(0)
 , MaxActiveObjectCount(InMaxActiveObjectCount)
 , ReplicationSystemId(InReplicationSystemId)
-, ScopableInternalIndices(MaxActiveObjectCount)
-, PrevFrameScopableInternalIndices(MaxActiveObjectCount)
+, GlobalScopableInternalIndices(MaxActiveObjectCount)
+, ScopeFrameData(MaxActiveObjectCount)
 , RelevantObjectsInternalIndices(MaxActiveObjectCount)
 , PolledObjectsInternalIndices(MaxActiveObjectCount)
 , DirtyObjectsToCopy(MaxActiveObjectCount)
@@ -112,7 +112,7 @@ FInternalNetRefIndex FNetRefHandleManager::InternalCreateNetObject(const FNetRef
 
 		// Mark Handle index as assigned and scopable for now
 		AssignedInternalIndices.SetBit(InternalIndex);
-		ScopableInternalIndices.SetBit(InternalIndex);
+		GlobalScopableInternalIndices.SetBit(InternalIndex);
 
 		// When a handle is first created, it is not set to be a subobject
 		SubObjectInternalIndices.ClearBit(InternalIndex);
@@ -370,7 +370,7 @@ void FNetRefHandleManager::RemoveFromScope(FInternalNetRefIndex InternalIndex)
 	// Can only remove an object from scope if it is assignable
 	if (ensure(AssignedInternalIndices.GetBit(InternalIndex)))
 	{
-		ScopableInternalIndices.ClearBit(InternalIndex);
+		GlobalScopableInternalIndices.ClearBit(InternalIndex);
 	}
 }
 
@@ -386,7 +386,7 @@ void FNetRefHandleManager::DestroyNetObject(FNetRefHandle RefHandle)
 		NetHandleToInternalIndex.Remove(Data.NetHandle);
 
 		// Remove from scopable objects if not already done
-		ScopableInternalIndices.ClearBit(InternalIndex);
+		GlobalScopableInternalIndices.ClearBit(InternalIndex);
 
 		// We always defer the actual destroy
 		PendingDestroyInternalIndices.Add(InternalIndex);
@@ -576,6 +576,11 @@ FNetRefHandle FNetRefHandleManager::GetRootObjectOfSubObject(FNetRefHandle SubOb
 	return OwnerInternalIndex != InvalidInternalIndex ? ReplicatedObjectData[OwnerInternalIndex].RefHandle : FNetRefHandle();
 }
 
+FInternalNetRefIndex FNetRefHandleManager::GetRootObjectInternalIndexOfSubObject(FInternalNetRefIndex SubObjectIndex) const
+{
+	return SubObjectIndex != InvalidInternalIndex ? ReplicatedObjectData[SubObjectIndex].SubObjectRootIndex : InvalidInternalIndex;
+}
+
 bool FNetRefHandleManager::AddDependentObject(FNetRefHandle ParentRefHandle, FNetRefHandle DependentObjectRefHandle, EDependentObjectSchedulingHint SchedulingHint, EAddDependentObjectFlags Flags)
 {
 	check(ParentRefHandle != DependentObjectRefHandle);
@@ -790,5 +795,22 @@ FNetRefHandle FNetRefHandleManager::MakeNetRefHandleFromId(uint64 Id)
 	return Handle;
 }
 
+void FNetRefHandleManager::OnPreSendUpdate()
+{
+	// The current frame scope is based on all indexes assigned up to this point.
+	ScopeFrameData.CurrentFrameScopableInternalIndices.Copy(GlobalScopableInternalIndices);
 
+	// Allow the list to be read.
+	ScopeFrameData.bIsValid = true;
 }
+
+void FNetRefHandleManager::OnPostSendUpdate()
+{
+	// Store the scope for the next frame.
+	ScopeFrameData.PrevFrameScopableInternalIndices.Copy(ScopeFrameData.CurrentFrameScopableInternalIndices);
+
+	// From here no-one should access the ScopeFrameData
+	ScopeFrameData.bIsValid = false;
+}
+
+}
\ No newline at end of file
diff --git a/Engine/Source/Runtime/Experimental/Iris/Core/Private/Iris/ReplicationSystem/NetRefHandleManager.h b/Engine/Source/Runtime/Experimental/Iris/Core/Private/Iris/ReplicationSystem/NetRefHandleManager.h
--- a/Engine/Source/Runtime/Experimental/Iris/Core/Private/Iris/ReplicationSystem/NetRefHandleManager.h
+++ b/Engine/Source/Runtime/Experimental/Iris/Core/Private/Iris/ReplicationSystem/NetRefHandleManager.h
@@ -120,8 +120,14 @@ public:
 public:
 	FNetRefHandleManager(FReplicationProtocolManager& InReplicationProtocolManager, uint32 InReplicationSystemId, uint32 MaxActiveObjects);
 
+	/** Callback triggered at the beginning of PreSendUpdate. Used to sync current frame data. */
+	void OnPreSendUpdate();
+
+	/** Callback triggered at the end of SendUpdate. Used to clear current frame data. */
+	void OnPostSendUpdate();
+
 	// Return true if this is a scopable index
-	bool IsScopableIndex(FInternalNetRefIndex InternalIndex) const { return ScopableInternalIndices.GetBit(InternalIndex); }
+	bool IsScopableIndex(FInternalNetRefIndex InternalIndex) const { return GlobalScopableInternalIndices.GetBit(InternalIndex); }
 
 	static FNetRefHandle MakeNetRefHandle(uint64 Id, uint32 ReplicationSystemId);
 	static FNetRefHandle MakeNetRefHandleFromId(uint64 Id);
@@ -185,12 +191,14 @@ public:
 	// Get internal index from NetHandle
 	inline FInternalNetRefIndex GetInternalIndexFromNetHandle(FNetHandle Handle) const;
 
-	/** Get bitarray for all currently scopable internal indices */
-	const FNetBitArrayView GetScopableInternalIndicesView() const { return MakeNetBitArrayView(ScopableInternalIndices); }
+	/** All scopable internal indices of the ReplicationSystem. Always up to date but should mostly be accessed in operations executed outside PreSendUpdate. */
+	const FNetBitArrayView GetGlobalScopableInternalIndices() const { return MakeNetBitArrayView(GlobalScopableInternalIndices); }
 
-	// Get bitarray for all internal indices that was scopable last update
-	const FNetBitArrayView GetPrevFrameScopableInternalIndicesView() const { return MakeNetBitArrayView(PrevFrameScopableInternalIndices); }
-	void SetPrevFrameScopableInternalIndicesToCurrent() { PrevFrameScopableInternalIndices = ScopableInternalIndices; }
+	/** All scopable internal indices of the current frame at the start of PreSendUpdate. Only accessible during that operation. */
+	const FNetBitArrayView GetCurrentFrameScopableInternalIndices() const { check(ScopeFrameData.bIsValid); return MakeNetBitArrayView(ScopeFrameData.CurrentFrameScopableInternalIndices); }
+
+	/** All scopable internal indices of the previous PreSendUpdate. Only accessible during that operation */
+	const FNetBitArrayView GetPrevFrameScopableInternalIndices() const { check(ScopeFrameData.bIsValid); return MakeNetBitArrayView(ScopeFrameData.PrevFrameScopableInternalIndices); }
 
 	/** List of objects that are always relevant or currently relevant to at least one connection. */
 	FNetBitArrayView GetRelevantObjectsInternalIndices() const { return MakeNetBitArrayView(RelevantObjectsInternalIndices); }
@@ -211,7 +219,10 @@ public:
 	bool AddSubObject(FNetRefHandle OwnerHandle, FNetRefHandle SubObjectHandle, FNetRefHandle RelativeOtherSubObjectHandle, EAddSubObjectFlags Flags = EAddSubObjectFlags::Default);
 	bool AddSubObject(FNetRefHandle OwnerHandle, FNetRefHandle SubObjectHandle, EAddSubObjectFlags Flags = EAddSubObjectFlags::Default);
 	void RemoveSubObject(FNetRefHandle SubObjectHandle);
+	
 	FNetRefHandle GetRootObjectOfSubObject(FNetRefHandle SubObjectHandle) const;
+	FInternalNetRefIndex GetRootObjectInternalIndexOfSubObject(FInternalNetRefIndex SubObjectIndex) const;
+
 	bool SetSubObjectNetCondition(FInternalNetRefIndex SubObjectInternalIndex, FLifeTimeConditionStorage SubObjectCondition);
 
 	// DependentObjects
@@ -311,12 +322,31 @@ private:
 	FRefHandleMap RefHandleToInternalIndex;
 	FNetHandleMap NetHandleToInternalIndex;
 
-	// Bitset used in order to track assigned internal which are scopable
-	FNetBitArray ScopableInternalIndices;
+	struct FScopeFrameData
+	{
+		FScopeFrameData(uint32 InMaxActiveObjectCount)
+			: bIsValid(false)
+			, CurrentFrameScopableInternalIndices(InMaxActiveObjectCount)
+			, PrevFrameScopableInternalIndices(InMaxActiveObjectCount)
+			
+		{ }
+
+		// Controls if the frame data can be read or not
+		uint32 bIsValid : 1;
+
+		// Bitset used in order to track assigned internal which are scopable
+		FNetBitArray CurrentFrameScopableInternalIndices;
+
+		// Which internal indices were used last net frame. This can be used to find out which ones are new and deleted this frame. 
+		FNetBitArray PrevFrameScopableInternalIndices;
+	};
+
+	/** Bitset used in order to track assigned internal which are scopable */
+	FNetBitArray GlobalScopableInternalIndices;
+
+	/** Stores scope lists relevant to the current frame. Only valid during SendUpdate(). */
+	FScopeFrameData ScopeFrameData;
 
-	// Which internal indices were used last net frame. This can be used to find out which ones are new and deleted this frame. 
-	FNetBitArray PrevFrameScopableInternalIndices;
-	
 	/** This contains the ScopableInternalIndices list minus filtered objects that are not relevant to any connection this frame. */
 	FNetBitArray RelevantObjectsInternalIndices;
 
diff --git a/Engine/Source/Runtime/Experimental/Iris/Core/Private/Iris/ReplicationSystem/ObjectReplicationBridge.cpp b/Engine/Source/Runtime/Experimental/Iris/Core/Private/Iris/ReplicationSystem/ObjectReplicationBridge.cpp
--- a/Engine/Source/Runtime/Experimental/Iris/Core/Private/Iris/ReplicationSystem/ObjectReplicationBridge.cpp
+++ b/Engine/Source/Runtime/Experimental/Iris/Core/Private/Iris/ReplicationSystem/ObjectReplicationBridge.cpp
@@ -20,6 +20,7 @@
 #include "Iris/ReplicationSystem/ObjectPollFrequencyLimiter.h"
 #include "Iris/ReplicationSystem/ObjectReplicationBridgeConfig.h"
 #include "Iris/ReplicationSystem/Prioritization/NetObjectPrioritizer.h"
+#include "Iris/ReplicationSystem/ReplicationConnections.h"
 #include "Iris/ReplicationSystem/ReplicationFragment.h"
 #include "Iris/ReplicationSystem/ReplicationFragmentInternal.h"
 #include "Iris/ReplicationState/ReplicationStateDescriptorBuilder.h"
@@ -29,6 +30,7 @@
 #include "Iris/ReplicationSystem/ReplicationSystemInternal.h"
 #include "Iris/ReplicationSystem/ReplicationOperations.h"
 #include "Iris/ReplicationSystem/ReplicationOperationsInternal.h"
+#include "Iris/ReplicationSystem/ReplicationWriter.h"
 #include "Iris/ReplicationSystem/RepTag.h"
 #include "Iris/ReplicationSystem/Polling/ObjectPoller.h"
 #include "Iris/Serialization/NetSerializationContext.h"
@@ -213,6 +215,12 @@ UE::Net::FNetRefHandle UObjectReplicationBridge::BeginReplication(UObject* Insta
 	using namespace UE::Net;
 	using namespace UE::Net::Private;
 
+	if (bBlockBeginReplication)
+	{
+		ensureMsgf(false, TEXT("BeginReplication is not allowed during this operation. %s will not be replicated"), *GetNameSafe(Instance));
+		return FNetRefHandle();
+	}
+
 	FNetRefHandle AllocatedRefHandle = ObjectReferenceCache->CreateObjectReferenceHandle(Instance);
 
 	// If we failed to assign a handle, or if the Handle already is replicating, just return the handle
@@ -403,6 +411,9 @@ UE::Net::FNetRefHandle UObjectReplicationBridge::BeginReplication(FNetRefHandle
 	}
 	else
 	{
+		// We support replicating new subobjects even during internal operations.
+		TGuardValue<bool> AllowBeginReplication(bBlockBeginReplication, false);
+
 		FCreateNetRefHandleParams SubObjectCreateParams = Params;
 		// The filtering system ignores subobjects so let's not waste cycles figuring out which filter to use.
 		SubObjectCreateParams.bAllowDynamicFilter = 0U;
@@ -680,11 +691,32 @@ void UObjectReplicationBridge::PreSendUpdateSingleHandle(FNetRefHandle RefHandle
 	ForcePollObject(RefHandle);
 }
 
+void UObjectReplicationBridge::OnStartPreSendUpdate()
+{
+	// During SendUpdate it is not supported to start replication of new root objects.
+	bBlockBeginReplication = true;
+}
+
 void UObjectReplicationBridge::PreSendUpdate()
 {
-	PreUpdateAndPoll();
+	using namespace UE::Net;
+
+	FNetBitArrayView ObjectsConsideredForPolling = GetReplicationSystem()->GetReplicationSystemInternal()->GetNetRefHandleManager().GetPolledObjectsInternalIndices();
+	ObjectsConsideredForPolling.Reset();
+
+	BuildPollList(ObjectsConsideredForPolling);
+
+	PreUpdate(ObjectsConsideredForPolling);
+
+	ReconcileNewSubObjects(ObjectsConsideredForPolling);
+
+	Poll(ObjectsConsideredForPolling);
 }
 
+void UObjectReplicationBridge::OnPostSendUpdate()
+{
+	bBlockBeginReplication = false;
+}
 
 void UObjectReplicationBridge::PruneStaleObjects()
 {
@@ -761,24 +793,21 @@ void UObjectReplicationBridge::ForcePollObject(FNetRefHandle Handle)
 	}
 }
 
-void UObjectReplicationBridge::PreUpdateAndPoll()
+void UObjectReplicationBridge::BuildPollList(UE::Net::FNetBitArrayView ObjectsConsideredForPolling)
 {
 	using namespace UE::Net;
 	using namespace UE::Net::Private;
 
-	IRIS_PROFILER_SCOPE(UObjectReplicationBridge_PreUpdateAndPoll);
+	IRIS_PROFILER_SCOPE(UObjectReplicationBridge_BuildPollList);
 
 	// Update every relevant objects from here
 	FReplicationSystemInternal* ReplicationSystemInternal = GetReplicationSystem()->GetReplicationSystemInternal();
-	const FNetRefHandleManager& LocalNetRefHandleManager = ReplicationSystemInternal->GetNetRefHandleManager();
 	FDirtyNetObjectTracker& DirtyNetObjectTracker = ReplicationSystemInternal->GetDirtyNetObjectTracker();
+
+	const FNetRefHandleManager& LocalNetRefHandleManager = ReplicationSystemInternal->GetNetRefHandleManager();
 	const FNetBitArrayView RelevantObjects = LocalNetRefHandleManager.GetRelevantObjectsInternalIndices();
 	const FNetBitArrayView WantToBeDormantObjects = MakeNetBitArrayView(LocalNetRefHandleManager.GetWantToBeDormantInternalIndices());
 
-	// Filter the set of objects considered for pre-update and polling
-	FNetBitArrayView ObjectsConsideredForPolling = LocalNetRefHandleManager.GetPolledObjectsInternalIndices();
-	ObjectsConsideredForPolling.Reset();
-
 	if (bUseFrequencyBasedPolling)
 	{
 		if (bEnableForceNetUpdate)
@@ -802,7 +831,7 @@ void UObjectReplicationBridge::PreUpdateAndPoll()
 	// Mask off objects pending dormancy as we do not want to poll/pre-update them unless they are marked for flush or are dirty
 	if (bUseDormancyToFilterPolling)
 	{
-		IRIS_PROFILER_SCOPE(PreUpdateAndPoll_Dormancy);
+		IRIS_PROFILER_SCOPE(BuildPollList_Dormancy);
 
 		// Mask off objects pending dormancy that are not dirty
 		const FNetBitArrayView AccumulatedDirtyObjects = DirtyNetObjectTracker.GetAccumulatedDirtyNetObjects();
@@ -833,7 +862,7 @@ void UObjectReplicationBridge::PreUpdateAndPoll()
 	* are replicated atomically this polling propagation is required.
 	*/
 	{
-		IRIS_PROFILER_SCOPE(PreUpdateAndPoll_PropagatePolling);
+		IRIS_PROFILER_SCOPE(BuildPollList_PropagatePolling);
 
 		auto PropagateSubObjectNetForceUpdateToOwner = [&LocalNetRefHandleManager, &ObjectsConsideredForPolling](uint32 InternalObjectIndex)
 		{
@@ -885,7 +914,7 @@ void UObjectReplicationBridge::PreUpdateAndPoll()
 			
 		// If an object with dependents is about to be polled, force it's dependents to poll at the same time.
 		{
-			IRIS_PROFILER_SCOPE(PreUpdateAndPoll_PatchDependentObjects);
+			IRIS_PROFILER_SCOPE(BuildPollList_PatchDependentObjects);
 
 			FNetBitArray TempObjectsConsideredForPolling;
 			TempObjectsConsideredForPolling.InitAndCopy(ObjectsConsideredForPolling);
@@ -900,31 +929,101 @@ void UObjectReplicationBridge::PreUpdateAndPoll()
 			);
 		}
 	}
+}
 	
+void UObjectReplicationBridge::PreUpdate(const UE::Net::FNetBitArrayView ObjectsConsideredForPolling)
+{
+	using namespace UE::Net::Private;
+
+	IRIS_PROFILER_SCOPE(UObjectReplicationBridge_PreUpdate);
+
+	// TODO: Get rid of Poller
 	FObjectPoller::FInitParams PollerInitParams;
 	PollerInitParams.ObjectReplicationBridge = this;
 	PollerInitParams.ReplicationSystemInternal = GetReplicationSystem()->GetReplicationSystemInternal();
-
+	
 	FObjectPoller Poller(PollerInitParams);
+	Poller.PreUpdatePass(ObjectsConsideredForPolling);
 
-	{
-		IRIS_PROFILER_SCOPE(PreUpdateAndPoll_PreUpdate);
-		Poller.PreUpdatePass(ObjectsConsideredForPolling);
-	}
+	FObjectPoller::FPreUpdateAndPollStats Stats = Poller.GetPollStats();
+	UE_NET_TRACE_FRAME_STATSCOUNTER(ReplicationSystem->GetId(), ReplicationSystem.PreUpdatedObjectCount, Stats.PreUpdatedObjectCount, ENetTraceVerbosity::Trace);
+}
 
-	{
-		IRIS_PROFILER_SCOPE(PreUpdateAndPoll_Poll);
-		Poller.PollObjects(ObjectsConsideredForPolling);
-	}
+void UObjectReplicationBridge::Poll(const UE::Net::FNetBitArrayView ObjectsConsideredForPolling)
+{
+	using namespace UE::Net::Private;
+
+	IRIS_PROFILER_SCOPE(UObjectReplicationBridge_Poll);
+
+	FObjectPoller::FInitParams PollerInitParams;
+	PollerInitParams.ObjectReplicationBridge = this;
+	PollerInitParams.ReplicationSystemInternal = GetReplicationSystem()->GetReplicationSystemInternal();
+
+	FObjectPoller Poller(PollerInitParams);
+	Poller.PollObjects(ObjectsConsideredForPolling);
 	
 	FObjectPoller::FPreUpdateAndPollStats Stats = Poller.GetPollStats();
 
 	// Report stats
-	UE_NET_TRACE_FRAME_STATSCOUNTER(ReplicationSystem->GetId(), ReplicationSystem.PreUpdatedObjectCount, Stats.PreUpdatedObjectCount, ENetTraceVerbosity::Trace);
 	UE_NET_TRACE_FRAME_STATSCOUNTER(ReplicationSystem->GetId(), ReplicationSystem.PolledObjectCount, Stats.PolledObjectCount, ENetTraceVerbosity::Trace);
 	UE_NET_TRACE_FRAME_STATSCOUNTER(ReplicationSystem->GetId(), ReplicationSystem.PolledReferencesObjectCount, Stats.PolledReferencesObjectCount, ENetTraceVerbosity::Trace);
 }
 
+void UObjectReplicationBridge::ReconcileNewSubObjects(UE::Net::FNetBitArrayView ObjectsConsideredForPolling)
+{
+	using namespace UE::Net;
+	using namespace UE::Net::Private;
+
+	IRIS_PROFILER_SCOPE(UObjectReplicationBridge_ReconcileNewSubObjects);
+
+	FReplicationSystemInternal* ReplicationSystemInternal = GetReplicationSystem()->GetReplicationSystemInternal();
+	const FNetRefHandleManager& LocalNetRefHandleManager = ReplicationSystemInternal->GetNetRefHandleManager();
+
+	const FNetBitArrayView SubObjectList = LocalNetRefHandleManager.GetSubObjectInternalIndicesView();
+	FReplicationConnections& Connections = ReplicationSystemInternal->GetConnections();
+	FReplicationFiltering& Filtering = ReplicationSystemInternal->GetFiltering();
+
+	auto HandleNewSubObject = [&](FInternalNetRefIndex SubObjectIndex)
+	{
+		const bool bIsSubObject = SubObjectList.IsBitSet(SubObjectIndex);
+		checkf(bIsSubObject, TEXT("Found a root object %s (Index:%u) that was created after the start of PreSendUpdate(). This is not supported"), *GetNameSafe(LocalNetRefHandleManager.GetReplicatedObjectInstance(SubObjectIndex)), SubObjectIndex);
+		if (UNLIKELY(!bIsSubObject))
+		{
+			return;
+		}
+
+		const FInternalNetRefIndex RootObjectIndex = LocalNetRefHandleManager.GetRootObjectInternalIndexOfSubObject(SubObjectIndex);
+		if (UNLIKELY(RootObjectIndex == FNetRefHandleManager::InvalidInternalIndex))
+		{
+			ensureMsgf(RootObjectIndex != FNetRefHandleManager::InvalidInternalIndex, TEXT("SubObject %s (Index:%u) had invalid RootObjectIndex"), *GetNameSafe(LocalNetRefHandleManager.GetReplicatedObjectInstance(SubObjectIndex)), SubObjectIndex);
+			return;
+		}
+
+		// Add the new subobject to the Poll list
+		ObjectsConsideredForPolling.SetBit(SubObjectIndex);
+
+		// Iterate over all connections and add the subobject if the root object is relevant to the connection
+		auto UpdateConnectionScope = [&Filtering, &Connections, RootObjectIndex, SubObjectIndex](uint32 ConnectionId)
+		{
+			FReplicationConnection* Conn = Connections.GetConnection(ConnectionId);
+			FNetBitArrayView ObjectsInScope = Filtering.GetRelevantObjectsInScope(ConnectionId);
+
+			if (ObjectsInScope.IsBitSet(RootObjectIndex))
+			{
+				ObjectsInScope.SetBit(SubObjectIndex);
+			}
+		};
+
+		const FNetBitArray& ValidConnections = Connections.GetValidConnections();
+		ValidConnections.ForAllSetBits(UpdateConnectionScope);
+	};
+
+	// Find any objects that got added since the start of the PreSendUpdate
+	const FNetBitArrayView GlobalScopeList = LocalNetRefHandleManager.GetGlobalScopableInternalIndices();
+	const FNetBitArrayView CurrentFrameScopeList = LocalNetRefHandleManager.GetCurrentFrameScopableInternalIndices();
+	FNetBitArrayView::ForAllSetBits(GlobalScopeList, CurrentFrameScopeList, FNetBitArrayView::AndNotOp, HandleNewSubObject);
+}
+
 void UObjectReplicationBridge::UpdateInstancesWorldLocation()
 {
 	using namespace UE::Net;
@@ -1518,7 +1617,7 @@ void UObjectReplicationBridge::ReinitPollFrequency()
 		}
 	};
 
-	const FNetBitArrayView RootObjects = LocalNetRefHandleManager.GetScopableInternalIndicesView();
+	const FNetBitArrayView RootObjects = LocalNetRefHandleManager.GetGlobalScopableInternalIndices();
 	const FNetBitArrayView SubObjects = MakeNetBitArrayView(LocalNetRefHandleManager.GetSubObjectInternalIndices());
 
 	FNetBitArrayView::ForAllSetBits(RootObjects, SubObjects, FNetBitArrayView::AndNotOp, UpdatePollFrequency);
diff --git a/Engine/Source/Runtime/Experimental/Iris/Core/Private/Iris/ReplicationSystem/Prioritization/ReplicationPrioritization.cpp b/Engine/Source/Runtime/Experimental/Iris/Core/Private/Iris/ReplicationSystem/Prioritization/ReplicationPrioritization.cpp
--- a/Engine/Source/Runtime/Experimental/Iris/Core/Private/Iris/ReplicationSystem/Prioritization/ReplicationPrioritization.cpp
+++ b/Engine/Source/Runtime/Experimental/Iris/Core/Private/Iris/ReplicationSystem/Prioritization/ReplicationPrioritization.cpp
@@ -596,8 +596,8 @@ void FReplicationPrioritization::UpdatePrioritiesForNewAndDeletedObjects()
 {
 	IRIS_PROFILER_SCOPE(FReplicationPrioritization_UpdatePrioritiesForNewAndDeletedObjects);
 
-	const FNetBitArrayView PrevScopedIndices = NetRefHandleManager->GetPrevFrameScopableInternalIndicesView();
-	const FNetBitArrayView ScopedIndices = NetRefHandleManager->GetScopableInternalIndicesView();
+	const FNetBitArrayView PrevScopedIndices = NetRefHandleManager->GetPrevFrameScopableInternalIndices();
+	const FNetBitArrayView ScopedIndices = NetRefHandleManager->GetCurrentFrameScopableInternalIndices();
 
 	auto ForEachRemovedObject = [this](uint32 ObjectIndex)
 	{
diff --git a/Engine/Source/Runtime/Experimental/Iris/Core/Private/Iris/ReplicationSystem/ReplicationBridge.cpp b/Engine/Source/Runtime/Experimental/Iris/Core/Private/Iris/ReplicationSystem/ReplicationBridge.cpp
--- a/Engine/Source/Runtime/Experimental/Iris/Core/Private/Iris/ReplicationSystem/ReplicationBridge.cpp
+++ b/Engine/Source/Runtime/Experimental/Iris/Core/Private/Iris/ReplicationSystem/ReplicationBridge.cpp
@@ -60,6 +60,14 @@ void UReplicationBridge::PreSendUpdate()
 {
 }
 
+void UReplicationBridge::OnStartPreSendUpdate()
+{
+}
+
+void UReplicationBridge::OnPostSendUpdate()
+{
+}
+
 void UReplicationBridge::UpdateInstancesWorldLocation()
 {
 }
diff --git a/Engine/Source/Runtime/Experimental/Iris/Core/Private/Iris/ReplicationSystem/ReplicationSystem.cpp b/Engine/Source/Runtime/Experimental/Iris/Core/Private/Iris/ReplicationSystem/ReplicationSystem.cpp
--- a/Engine/Source/Runtime/Experimental/Iris/Core/Private/Iris/ReplicationSystem/ReplicationSystem.cpp
+++ b/Engine/Source/Runtime/Experimental/Iris/Core/Private/Iris/ReplicationSystem/ReplicationSystem.cpp
@@ -248,6 +248,18 @@ public:
 		}
 	}
 
+	void StartPreSendUpdate()
+	{
+		// Block unsupported operations until SendUpdate is finished
+		ReplicationSystemInternal.SetBlockFilterChanges(true);
+
+		// Tell systems we are starting PreSendUpdate
+		ReplicationSystemInternal.GetReplicationBridge()->OnStartPreSendUpdate();
+
+		// Sync the state of the world at the beginning of PreSendUpdate.
+		ReplicationSystemInternal.GetNetRefHandleManager().OnPreSendUpdate();
+	}
+
 	void CallPreSendUpdate(float DeltaSeconds)
 	{
 		ReplicationSystemInternal.GetReplicationBridge()->CallPreSendUpdate(DeltaSeconds);
@@ -255,16 +267,22 @@ public:
 
 	void EndPostSendUpdate()
 	{
+		// Unblock operations
+		ReplicationSystemInternal.SetBlockFilterChanges(false);
+
 		ReplicationSystemInternal.GetChangeMaskCache().ResetCache();
 
-		// Store the state of the previous frames scopable objects
-		ReplicationSystemInternal.GetNetRefHandleManager().SetPrevFrameScopableInternalIndicesToCurrent();
+		// Store the scope list for the next SendUpdate.
+		ReplicationSystemInternal.GetNetRefHandleManager().OnPostSendUpdate();
 
 		// Update handles pending tear-off
 		ReplicationSystemInternal.GetReplicationBridge()->UpdateHandlesPendingTearOff();
 
 		// Reset baseline invalidation
 		ReplicationSystemInternal.GetDeltaCompressionBaselineInvalidationTracker().PostSendUpdate();
+
+		// Tell systems we finished PostSendUpdate
+		ReplicationSystemInternal.GetReplicationBridge()->OnPostSendUpdate();
 	}
 
 	void UpdateDirtyObjectList()
@@ -620,6 +638,9 @@ void UReplicationSystem::PreSendUpdate(float DeltaSeconds)
 	{
 		UE_NET_TRACE_FRAME_STATSCOUNTER(GetId(), ReplicationSystem.ReplicatedObjectCount, InternalSys.GetNetRefHandleManager().GetActiveObjectCount(), ENetTraceVerbosity::Verbose);
 
+		// Tell systems we are starting PreSendUpdate
+		Impl->StartPreSendUpdate();
+
 		// Refresh the dirty objects we were told about.
 		Impl->UpdateDirtyObjectList();
 
@@ -929,6 +950,13 @@ void UReplicationSystem::SetOwningNetConnection(FNetRefHandle Handle, uint32 Con
 
 	FNetRefHandleManager& NetRefHandleManager = Impl->ReplicationSystemInternal.GetNetRefHandleManager();
 	const FInternalNetRefIndex ObjectInternalIndex = NetRefHandleManager.GetInternalIndex(Handle);
+
+	if (Impl->ReplicationSystemInternal.AreFilterChangesBlocked())
+	{
+		ensureMsgf(false, TEXT("Setting filter conditions is not yet supported during this operation. Filter condition on %s (%s) failed."), *GetNameSafe(NetRefHandleManager.GetReplicatedObjectInstance(ObjectInternalIndex)), *Handle.ToString());
+		return;
+	}
+
 	if (ObjectInternalIndex == FNetRefHandleManager::InvalidInternalIndex)
 	{
 		return;
@@ -964,6 +992,12 @@ bool UReplicationSystem::SetFilter(FNetRefHandle Handle, UE::Net::FNetObjectFilt
 		return false;
 	}
 
+	if (Impl->ReplicationSystemInternal.AreFilterChangesBlocked())
+	{
+		ensureMsgf(false, TEXT("Setting filter conditions is not yet supported during this operation. Filter condition on %s (%s) failed."), *GetNameSafe(NetRefHandleManager.GetReplicatedObjectInstance(ObjectInternalIndex)), *Handle.ToString());
+		return false;
+	}
+
 	FReplicationFiltering& Filtering = Impl->ReplicationSystemInternal.GetFiltering();
 	return Filtering.SetFilter(ObjectInternalIndex, Filter);
 }
@@ -990,6 +1024,13 @@ bool UReplicationSystem::SetConnectionFilter(FNetRefHandle Handle, const TBitArr
 
 	FNetRefHandleManager& NetRefHandleManager = Impl->ReplicationSystemInternal.GetNetRefHandleManager();
 	const FInternalNetRefIndex ObjectInternalIndex = NetRefHandleManager.GetInternalIndex(Handle);
+
+	if (Impl->ReplicationSystemInternal.AreFilterChangesBlocked())
+	{
+		ensureMsgf(false, TEXT("Setting filter conditions is not yet supported during this operation. Filter condition on %s (%s) failed."), *GetNameSafe(NetRefHandleManager.GetReplicatedObjectInstance(ObjectInternalIndex)), *Handle.ToString());
+		return false;
+	}
+
 	if (ObjectInternalIndex == FNetRefHandleManager::InvalidInternalIndex)
 	{
 		return false;
@@ -1052,6 +1093,12 @@ void UReplicationSystem::SetSubObjectFilterStatus(FName GroupName, uint32 Connec
 		return;
 	}
 
+	if (Impl->ReplicationSystemInternal.AreFilterChangesBlocked())
+	{
+		ensureMsgf(false, TEXT("Setting filter conditions is not yet supported during this operation."));
+		return;
+	}
+
 	FNetObjectGroupHandle GroupHandle = GetSubObjectFilterGroupHandle(GroupName);
 	if (GroupHandle.IsValid())
 	{
@@ -1065,6 +1112,12 @@ void UReplicationSystem::RemoveSubObjectFilter(FName GroupName)
 	using namespace UE::Net;
 	using namespace UE::Net::Private;
 
+	if (Impl->ReplicationSystemInternal.AreFilterChangesBlocked())
+	{
+		ensureMsgf(false, TEXT("Setting filter conditions is not yet supported during this operation."));
+		return;
+	}
+
 	FNetObjectGroups& Groups = Impl->ReplicationSystemInternal.GetGroups();
 	FReplicationFiltering& Filtering = Impl->ReplicationSystemInternal.GetFiltering();
 
@@ -1100,6 +1153,12 @@ void UReplicationSystem::AddToGroup(FNetObjectGroupHandle GroupHandle, FNetRefHa
 	FReplicationFiltering& Filtering = Impl->ReplicationSystemInternal.GetFiltering();	
 
 	const FInternalNetRefIndex ObjectInternalIndex = NetRefHandleManager.GetInternalIndex(Handle);
+	
+	if (Impl->ReplicationSystemInternal.AreFilterChangesBlocked())
+	{
+		ensureMsgf(false, TEXT("Setting filter conditions is not yet supported during this operation. Filter condition on %s (%s) failed."), *GetNameSafe(NetRefHandleManager.GetReplicatedObjectInstance(ObjectInternalIndex)), *Handle.ToString());
+		return;
+	}
 
 	if (ObjectInternalIndex)
 	{
@@ -1121,6 +1180,13 @@ void UReplicationSystem::RemoveFromGroup(FNetObjectGroupHandle GroupHandle, FNet
 	FNetRefHandleManager& NetRefHandleManager = Impl->ReplicationSystemInternal.GetNetRefHandleManager();
 
 	const FInternalNetRefIndex ObjectInternalIndex = NetRefHandleManager.GetInternalIndex(Handle);
+
+	if (Impl->ReplicationSystemInternal.AreFilterChangesBlocked())
+	{
+		ensureMsgf(false, TEXT("Setting filter conditions is not yet supported during this operation. Filter condition on %s (%s) failed."), *GetNameSafe(NetRefHandleManager.GetReplicatedObjectInstance(ObjectInternalIndex)), *Handle.ToString());
+		return;
+	}
+
 	if (ObjectInternalIndex)
 	{
 		FNetObjectGroups& Groups = Impl->ReplicationSystemInternal.GetGroups();
@@ -1139,8 +1205,15 @@ void UReplicationSystem::RemoveFromAllGroups(FNetRefHandle Handle)
 	FNetObjectGroups& Groups = Impl->ReplicationSystemInternal.GetGroups();
 	FReplicationFiltering& Filtering = Impl->ReplicationSystemInternal.GetFiltering();
 
-	uint32 NumGroupMemberShips;
+	uint32 NumGroupMemberShips = 0;
 	const FInternalNetRefIndex ObjectInternalIndex = NetRefHandleManager.GetInternalIndex(Handle);
+
+	if (Impl->ReplicationSystemInternal.AreFilterChangesBlocked())
+	{
+		ensureMsgf(false, TEXT("Setting filter conditions is not yet supported during this operation. Filter condition on %s (%s) failed."), *GetNameSafe(NetRefHandleManager.GetReplicatedObjectInstance(ObjectInternalIndex)), *Handle.ToString());
+		return;
+	}
+
 	if (const FNetObjectGroupHandle* GroupHandles = Groups.GetGroupMemberships(ObjectInternalIndex, NumGroupMemberShips))
 	{
 		// We copy the membership array as it is modified during removal
@@ -1221,6 +1294,12 @@ void UReplicationSystem::AddGroupFilter(FNetObjectGroupHandle GroupHandle)
 		return;
 	}
 
+	if (Impl->ReplicationSystemInternal.AreFilterChangesBlocked())
+	{
+		ensureMsgf(false, TEXT("Setting filter conditions is not yet supported during this operation."));
+		return;
+	}
+
 	UE::Net::Private::FReplicationFiltering& Filtering = Impl->ReplicationSystemInternal.GetFiltering();	
 	Filtering.AddGroupFilter(GroupHandle);
 }
@@ -1233,6 +1312,12 @@ void UReplicationSystem::RemoveGroupFilter(FNetObjectGroupHandle GroupHandle)
 		return;
 	}
 
+	if (Impl->ReplicationSystemInternal.AreFilterChangesBlocked())
+	{
+		ensureMsgf(false, TEXT("Setting filter conditions is not yet supported during this operation."));
+		return;
+	}
+
 	UE::Net::Private::FReplicationFiltering& Filtering = Impl->ReplicationSystemInternal.GetFiltering();	
 	Filtering.RemoveGroupFilter(GroupHandle);
 }
@@ -1245,6 +1330,12 @@ void UReplicationSystem::SetGroupFilterStatus(FNetObjectGroupHandle GroupHandle,
 		return;
 	}
 
+	if (Impl->ReplicationSystemInternal.AreFilterChangesBlocked())
+	{
+		ensureMsgf(false, TEXT("Setting filter conditions is not yet supported during this operation."));
+		return;
+	}
+
 	UE::Net::Private::FReplicationFiltering& Filtering = Impl->ReplicationSystemInternal.GetFiltering();	
 	Filtering.SetGroupFilterStatus(GroupHandle, ConnectionId, ReplicationStatus);
 }
@@ -1257,6 +1348,12 @@ void UReplicationSystem::SetGroupFilterStatus(FNetObjectGroupHandle GroupHandle,
 		return;
 	}
 
+	if (Impl->ReplicationSystemInternal.AreFilterChangesBlocked())
+	{
+		ensureMsgf(false, TEXT("Setting filter conditions is not yet supported during this operation."));
+		return;
+	}
+
 	UE::Net::Private::FReplicationFiltering& Filtering = Impl->ReplicationSystemInternal.GetFiltering();	
 	Filtering.SetGroupFilterStatus(GroupHandle, UE::Net::MakeNetBitArrayView(Connections), ReplicationStatus);
 }
@@ -1269,6 +1366,12 @@ void UReplicationSystem::SetGroupFilterStatus(FNetObjectGroupHandle GroupHandle,
 		return;
 	}
 
+	if (Impl->ReplicationSystemInternal.AreFilterChangesBlocked())
+	{
+		ensureMsgf(false, TEXT("Setting filter conditions is not yet supported during this operation."));
+		return;
+	}
+
 	UE::Net::Private::FReplicationFiltering& Filtering = Impl->ReplicationSystemInternal.GetFiltering();	
 	Filtering.SetGroupFilterStatus(GroupHandle, ReplicationStatus);
 }
@@ -1284,6 +1387,12 @@ bool UReplicationSystem::SetReplicationConditionConnectionFilter(FNetRefHandle H
 		return false;
 	}
 
+	if (Impl->ReplicationSystemInternal.AreFilterChangesBlocked())
+	{
+		ensureMsgf(false, TEXT("Setting filter conditions is not yet supported during this operation. Filter condition on %s (%s) failed."), *GetNameSafe(NetRefHandleManager.GetReplicatedObjectInstance(ObjectInternalIndex)), *Handle.ToString());
+		return false;
+	}
+
 	FReplicationConditionals& Conditionals = Impl->ReplicationSystemInternal.GetConditionals();
 	return Conditionals.SetConditionConnectionFilter(ObjectInternalIndex, Condition, ConnectionId, bEnable);
 }
@@ -1315,6 +1424,12 @@ void UReplicationSystem::SetDeltaCompressionStatus(FNetRefHandle Handle, UE::Net
 		return;
 	}
 
+	if (Impl->ReplicationSystemInternal.AreFilterChangesBlocked())
+	{
+		ensureMsgf(false, TEXT("Setting delta compression is not yet supported during this operation. Filter condition on %s (%s) failed."), *GetNameSafe(NetRefHandleManager.GetReplicatedObjectInstance(ObjectInternalIndex)), *Handle.ToString());
+		return;
+	}
+
 	FDeltaCompressionBaselineManager& DC = InternalSys.GetDeltaCompressionBaselineManager();
 	return DC.SetDeltaCompressionStatus(ObjectInternalIndex, Status);
 }
diff --git a/Engine/Source/Runtime/Experimental/Iris/Core/Private/Iris/ReplicationSystem/ReplicationSystemInternal.h b/Engine/Source/Runtime/Experimental/Iris/Core/Private/Iris/ReplicationSystem/ReplicationSystemInternal.h
--- a/Engine/Source/Runtime/Experimental/Iris/Core/Private/Iris/ReplicationSystem/ReplicationSystemInternal.h
+++ b/Engine/Source/Runtime/Experimental/Iris/Core/Private/Iris/ReplicationSystem/ReplicationSystemInternal.h
@@ -108,6 +108,10 @@ public:
 		return ForwardNetRPCCallMulticastDelegate;
 	}
 
+	void SetBlockFilterChanges(bool bBlock) { bBlockFilterChanges = bBlock; }
+
+	bool AreFilterChangesBlocked() const { return bBlockFilterChanges; }
+
 private:
 	FReplicationProtocolManager ReplicationProtocolManager;
 	FNetRefHandleManager NetRefHandleManager;
@@ -133,6 +137,9 @@ private:
 	FNetSendStats SendStats;
 	FForwardNetRPCCallMulticastDelegate ForwardNetRPCCallMulticastDelegate;
 	uint32 Id;
+
+	/** When true this prevents any changes to the filter system. Enabled during times where adding filter options is unsupported. */
+	bool bBlockFilterChanges = false;
 };
 
 }
diff --git a/Engine/Source/Runtime/Experimental/Iris/Core/Private/Iris/ReplicationSystem/ReplicationWriter.cpp b/Engine/Source/Runtime/Experimental/Iris/Core/Private/Iris/ReplicationSystem/ReplicationWriter.cpp
--- a/Engine/Source/Runtime/Experimental/Iris/Core/Private/Iris/ReplicationSystem/ReplicationWriter.cpp
+++ b/Engine/Source/Runtime/Experimental/Iris/Core/Private/Iris/ReplicationSystem/ReplicationWriter.cpp
@@ -478,13 +478,13 @@ void FReplicationWriter::StartReplication(uint32 InternalIndex)
 	// Subobject needs to mark its owner as dirty as the subobject could have been filtered out and now allowed to replicate again.
 	if (Info.IsSubObject)
 	{
-		FNetRefHandleManager::FReplicatedObjectData& ObjectData = NetRefHandleManager->GetReplicatedObjectDataNoCheck(InternalIndex);
-		if (ensureAlways(ObjectData.SubObjectRootIndex != FNetRefHandleManager::InvalidInternalIndex))
+		const uint32 RootObjectInternalIndex = NetRefHandleManager->GetRootObjectInternalIndexOfSubObject(InternalIndex);
+		if (ensure(RootObjectInternalIndex != FNetRefHandleManager::InvalidInternalIndex))
 		{
-			FReplicationInfo& OwnerInfo = ReplicatedObjects[ObjectData.SubObjectRootIndex];
+			FReplicationInfo& OwnerInfo = ReplicatedObjects[RootObjectInternalIndex];
 			if (OwnerInfo.GetState() != EReplicatedObjectState::Invalid)
 			{
-				ObjectsWithDirtyChanges.SetBit(ObjectData.SubObjectRootIndex);
+				ObjectsWithDirtyChanges.SetBit(RootObjectInternalIndex);
 				OwnerInfo.HasDirtySubObjects = 1U;
 			}
 		}
diff --git a/Engine/Source/Runtime/Experimental/Iris/Core/Public/Iris/ReplicationSystem/ObjectReplicationBridge.h b/Engine/Source/Runtime/Experimental/Iris/Core/Public/Iris/ReplicationSystem/ObjectReplicationBridge.h
--- a/Engine/Source/Runtime/Experimental/Iris/Core/Public/Iris/ReplicationSystem/ObjectReplicationBridge.h
+++ b/Engine/Source/Runtime/Experimental/Iris/Core/Public/Iris/ReplicationSystem/ObjectReplicationBridge.h
@@ -163,6 +163,8 @@ protected:
 	IRISCORE_API virtual void Deinitialize() override;
 	IRISCORE_API virtual void PreSendUpdateSingleHandle(FNetRefHandle RefHandle) override;	
 	IRISCORE_API virtual void PreSendUpdate() override;	
+	IRISCORE_API virtual void OnStartPreSendUpdate() override;
+	IRISCORE_API virtual void OnPostSendUpdate() override;
 	IRISCORE_API virtual void UpdateInstancesWorldLocation() override;
 	IRISCORE_API virtual void PruneStaleObjects() override;	
 	IRISCORE_API virtual bool WriteNetRefHandleCreationInfo(FReplicationBridgeSerializationContext& Context, FNetRefHandle Handle) override;
@@ -272,8 +274,17 @@ private:
 	/** Forcibly poll a single replicated object */
 	void ForcePollObject(FNetRefHandle RefHandle);
 
-	/** Pre update and poll all relevant objects who hit their polling period or are force net update. */
-	void PreUpdateAndPoll();
+	/** Build the list of relevant objects who hit their polling period or were flagged ForceNetUpdate. */
+	void BuildPollList(UE::Net::FNetBitArrayView ObjectsConsideredForPolling);
+
+	/** Call the user function PreUpdate (aka PreReplication) on objects about to be polled. */
+	void PreUpdate(const UE::Net::FNetBitArrayView ObjectsConsideredForPolling);
+
+	/** Find any new subobjects created inside PreUpdate and ensure they will be replicated this frame */
+	void ReconcileNewSubObjects(UE::Net::FNetBitArrayView ObjectsConsideredForPolling);
+
+	/** Poll all objects set in the list */
+	void Poll(const UE::Net::FNetBitArrayView ObjectsConsideredForPolling);
 
 	/** Remove mapping between handle and object instance. */
 	void UnregisterInstance(FNetRefHandle RefHandle);
@@ -363,6 +374,9 @@ private:
 	bool bHasPollOverrides = false;
 	bool bHasDirtyClassesInPollPeriodOverrides = false;
 
+	/** Set to true when the system does not allow new objects to begin replication at this moment. Useful when calling into user code and to warn them of illegal operations. */
+	bool bBlockBeginReplication = false;
+
 protected:
 	bool bSuppressCreateInstanceFailedEnsure = false;
 };
diff --git a/Engine/Source/Runtime/Experimental/Iris/Core/Public/Iris/ReplicationSystem/ReplicationBridge.h b/Engine/Source/Runtime/Experimental/Iris/Core/Public/Iris/ReplicationSystem/ReplicationBridge.h
--- a/Engine/Source/Runtime/Experimental/Iris/Core/Public/Iris/ReplicationSystem/ReplicationBridge.h
+++ b/Engine/Source/Runtime/Experimental/Iris/Core/Public/Iris/ReplicationSystem/ReplicationBridge.h
@@ -168,6 +168,12 @@ protected:
 	/** Invoked before ReplicationSystem copies dirty state data. */
 	IRISCORE_API virtual void PreSendUpdate();
 
+	/** Invoked when the ReplicationSystem starts the PreSendUpdate tick. */
+	IRISCORE_API virtual void OnStartPreSendUpdate();
+	
+	/** Invoked after we sent data to all connections. */
+	IRISCORE_API virtual void OnPostSendUpdate();
+	
 	/** Invoked before ReplicationSystem copies dirty state data for a single replicated object. */
 	IRISCORE_API virtual void PreSendUpdateSingleHandle(FNetRefHandle Handle);
 
diff --git a/Engine/Source/Runtime/Experimental/Iris/Core/Public/Iris/ReplicationSystem/ReplicationSystem.h b/Engine/Source/Runtime/Experimental/Iris/Core/Public/Iris/ReplicationSystem/ReplicationSystem.h
--- a/Engine/Source/Runtime/Experimental/Iris/Core/Public/Iris/ReplicationSystem/ReplicationSystem.h
+++ b/Engine/Source/Runtime/Experimental/Iris/Core/Public/Iris/ReplicationSystem/ReplicationSystem.h
@@ -59,6 +59,7 @@ public:
 
 	struct FReplicationSystemParams
 	{
+		//$IRIS TODO: These need documentation
 		UReplicationBridge* ReplicationBridge = nullptr;
 		uint32 MaxReplicatedObjectCount = 65535U;
 		uint32 MaxDeltaCompressedObjectCount = 2048U;
@@ -77,6 +78,9 @@ public:
 	/** @return Whether the system is run on a server. */
 	bool IsServer() const { return bIsServer; }
 
+	/** @return Is this system configured to replicate object properties. */
+	bool AllowObjectReplication() { return bAllowObjectReplication; }
+
 	/**
 	 * PreSendUpdate performs all the necessary work, such as filtering and prioritization of objects,
 	 * so that each connection will be properly updated with all the information needed in order to replicate.
