Reduce memory in Dynamic Collection

diff --git a/Engine/Source/Programs/HeadlessChaos/Private/GeometryCollection/GeometryCollectionTestSimulationSolver.cpp b/Engine/Source/Programs/HeadlessChaos/Private/GeometryCollection/GeometryCollectionTestSimulationSolver.cpp
--- a/Engine/Source/Programs/HeadlessChaos/Private/GeometryCollection/GeometryCollectionTestSimulationSolver.cpp
+++ b/Engine/Source/Programs/HeadlessChaos/Private/GeometryCollection/GeometryCollectionTestSimulationSolver.cpp
@@ -102,7 +102,8 @@ using namespace ChaosTest;
 
 		UnitTest.Solver->RegisterSimOneShotCallback([&]()
 		{
-			TManagedArray<int32>& SimulationType = Collection->DynamicCollection->SimulationType;
+			// Dynamic collection doesn't change the simulation type, RestCollection can be used instead
+			TManagedArray<int32>& SimulationType = RestCollection->SimulationType;
 			EXPECT_EQ(SimulationType[0],FGeometryCollection::ESimulationTypes::FST_Rigid);
 			EXPECT_EQ(SimulationType[1],FGeometryCollection::ESimulationTypes::FST_Rigid);
 			EXPECT_EQ(SimulationType[2],FGeometryCollection::ESimulationTypes::FST_Clustered);
diff --git a/Engine/Source/Runtime/Experimental/Chaos/Private/GeometryCollectionProxyData.cpp b/Engine/Source/Runtime/Experimental/Chaos/Private/GeometryCollectionProxyData.cpp
--- a/Engine/Source/Runtime/Experimental/Chaos/Private/GeometryCollectionProxyData.cpp
+++ b/Engine/Source/Runtime/Experimental/Chaos/Private/GeometryCollectionProxyData.cpp
@@ -27,8 +27,6 @@ void FTransformDynamicCollection::Construct()
 	AddExternalAttribute<FTransform3f>(FTransformCollection::TransformAttribute, FTransformCollection::TransformGroup, Transform);
 	AddExternalAttribute<int32>(FTransformCollection::ParentAttribute, FTransformCollection::TransformGroup, Parent);
 	AddExternalAttribute<TSet<int32>>(FTransformCollection::ChildrenAttribute, FTransformCollection::TransformGroup, Children);
-	AddExternalAttribute<int32>(FGeometryCollection::SimulationTypeAttribute, FTransformCollection::TransformGroup, SimulationType);
-	AddExternalAttribute<int32>(FGeometryCollection::StatusFlagsAttribute, FTransformCollection::TransformGroup, StatusFlags);
 }
 
 const FTransform3f& FTransformDynamicCollection::GetTransform(int32 Index) const
diff --git a/Engine/Source/Runtime/Experimental/Chaos/Private/PhysicsProxy/GeometryCollectionPhysicsProxy.cpp b/Engine/Source/Runtime/Experimental/Chaos/Private/PhysicsProxy/GeometryCollectionPhysicsProxy.cpp
--- a/Engine/Source/Runtime/Experimental/Chaos/Private/PhysicsProxy/GeometryCollectionPhysicsProxy.cpp
+++ b/Engine/Source/Runtime/Experimental/Chaos/Private/PhysicsProxy/GeometryCollectionPhysicsProxy.cpp
@@ -1634,7 +1634,6 @@ Chaos::TPBDGeometryCollectionParticleHandle<Chaos::FReal, 3>* FGeometryCollectio
 	FGeometryDynamicCollection& DynamicCollection = PhysicsThreadCollection;
 	TManagedArray<int32>& DynamicState = DynamicCollection.DynamicState;
 	TManagedArray<int32>& ParentIndex = DynamicCollection.Parent;
-	TManagedArray<TSet<int32>>& Children = DynamicCollection.Children;
 	TManagedArray<FTransform3f>& Transform = DynamicCollection.Transform;
 	const TManagedArray<FTransform>& MassToLocal = Parameters.RestCollection->GetAttribute<FTransform>(MassToLocalAttributeName, FTransformCollection::TransformGroup);
 	const TManagedArray<Chaos::FImplicitObjectPtr>& Implicits = DynamicCollection.Implicits;
@@ -1754,7 +1753,6 @@ FGeometryCollectionPhysicsProxy::BuildClusters_Internal(
 	FGeometryDynamicCollection& DynamicCollection = PhysicsThreadCollection;
 	TManagedArray<int32>& DynamicState = DynamicCollection.DynamicState;
 	TManagedArray<int32>& ParentIndex = DynamicCollection.Parent;
-	TManagedArray<TSet<int32>>& Children = DynamicCollection.Children;
 	TManagedArray<FTransform3f>& Transform = DynamicCollection.Transform;
 	const TManagedArray<FTransform>& MassToLocal = Parameters.RestCollection->GetAttribute<FTransform>(MassToLocalAttributeName, FTransformCollection::TransformGroup);
 	//TManagedArray<TSharedPtr<FCollisionStructureManager::FSimplicial> >& Simplicials = DynamicCollection.Simplicials;
diff --git a/Engine/Source/Runtime/Experimental/Chaos/Public/GeometryCollectionProxyData.h b/Engine/Source/Runtime/Experimental/Chaos/Public/GeometryCollectionProxyData.h
--- a/Engine/Source/Runtime/Experimental/Chaos/Public/GeometryCollectionProxyData.h
+++ b/Engine/Source/Runtime/Experimental/Chaos/Public/GeometryCollectionProxyData.h
@@ -35,8 +35,6 @@ public:
 	// Transform Group
 	TManagedArray<int32>        Parent;
 	TManagedArray<TSet<int32>>  Children;
-	TManagedArray<int32>        SimulationType;
-	TManagedArray<int32>        StatusFlags;
 
 private:
 	TManagedArray<FTransform3f> Transform;
diff --git a/Engine/Source/Runtime/Experimental/GeometryCollectionEngine/Private/GeometryCollection/GeometryCollectionComponent.cpp b/Engine/Source/Runtime/Experimental/GeometryCollectionEngine/Private/GeometryCollection/GeometryCollectionComponent.cpp
--- a/Engine/Source/Runtime/Experimental/GeometryCollectionEngine/Private/GeometryCollection/GeometryCollectionComponent.cpp
+++ b/Engine/Source/Runtime/Experimental/GeometryCollectionEngine/Private/GeometryCollection/GeometryCollectionComponent.cpp
@@ -3433,8 +3433,6 @@ void UGeometryCollectionComponent::ResetDynamicCollection()
 		GetTransformArrayCopyOnWrite();
 		GetParentArrayCopyOnWrite();
 		GetChildrenArrayCopyOnWrite();
-		GetSimulationTypeArrayCopyOnWrite();
-		GetStatusFlagsArrayCopyOnWrite();
 
 		if (bStoreVelocities || bNotifyTrailing)
 		{
diff --git a/Engine/Source/Runtime/Experimental/GeometryCollectionEngine/Public/GeometryCollection/GeometryCollectionComponent.h b/Engine/Source/Runtime/Experimental/GeometryCollectionEngine/Public/GeometryCollection/GeometryCollectionComponent.h
--- a/Engine/Source/Runtime/Experimental/GeometryCollectionEngine/Public/GeometryCollection/GeometryCollectionComponent.h
+++ b/Engine/Source/Runtime/Experimental/GeometryCollectionEngine/Public/GeometryCollection/GeometryCollectionComponent.h
@@ -887,9 +887,7 @@ public:
 	COPY_ON_WRITE_ATTRIBUTE(FLinearColor, BoneColor, FTransformCollection::TransformGroup)				\
 	COPY_ON_WRITE_ATTRIBUTE(int32, Parent, FTransformCollection::TransformGroup)						\
 	COPY_ON_WRITE_ATTRIBUTE(TSet<int32>, Children, FTransformCollection::TransformGroup)				\
-	COPY_ON_WRITE_ATTRIBUTE(int32, SimulationType, FTransformCollection::TransformGroup)				\
 	COPY_ON_WRITE_ATTRIBUTE(int32, TransformToGeometryIndex, FTransformCollection::TransformGroup)		\
-	COPY_ON_WRITE_ATTRIBUTE(int32, StatusFlags, FTransformCollection::TransformGroup)					\
 	COPY_ON_WRITE_ATTRIBUTE(int32, ExemplarIndex, FTransformCollection::TransformGroup)					\
 
 	// Declare all the methods
