Adds CVAR to apply debug coloring to broken primitives
Adds CVAR to force a # amount of primitives to be broken

diff --git a/Engine/Source/Runtime/Experimental/GeometryCollectionEngine/Private/GeometryCollection/GeometryCollectionComponent.cpp b/Engine/Source/Runtime/Experimental/GeometryCollectionEngine/Private/GeometryCollection/GeometryCollectionComponent.cpp
--- a/Engine/Source/Runtime/Experimental/GeometryCollectionEngine/Private/GeometryCollection/GeometryCollectionComponent.cpp
+++ b/Engine/Source/Runtime/Experimental/GeometryCollectionEngine/Private/GeometryCollection/GeometryCollectionComponent.cpp
@@ -462,6 +462,42 @@ const TManagedArray<Type>& UGeometryCollectionComponent::Get##Name##ArrayRest()
 	return RestCollection->GetGeometryCollection()->Name;								\
 }																						\
 
+#if !(UE_BUILD_SHIPPING)
+static TAutoConsoleVariable<int> CVarNumToForceBreak(
+	TEXT("r.GeometryCollection.CustomRenderer.ForceBreak"),
+	-1,
+	TEXT("Force the specified number of pieces to render individually, replacing their root proxy mesh.")
+);
+
+struct _ForcedBroken
+{
+public:
+	void Reset()
+	{
+		Components.Reset();
+		TotalTransforms = 0;
+	}
+
+	bool Break(UGeometryCollectionComponent* t, const TArray<FTransform>& CompSpaceTransforms)
+	{
+		if (Components.Contains(t))
+			return true;
+
+		if (TotalTransforms + CompSpaceTransforms.Num() > (uint32)CVarNumToForceBreak->GetInt())
+			return false;
+
+		TotalTransforms += CompSpaceTransforms.Num();
+		Components.Add(t);
+		return true;
+	}
+
+private:
+	TSet<UGeometryCollectionComponent*> Components;
+	uint32 TotalTransforms = 0;
+
+}static ForcedBroken;
+#endif
+
 // Define the methods
 COPY_ON_WRITE_ATTRIBUTES
 
@@ -593,6 +629,17 @@ UGeometryCollectionComponent::UGeometryCollectionComponent(const FObjectInitiali
 	// make sure older asset are using the default behaviour
 	DamagePropagationData.bEnabled = false;
 
+#if !(UE_BUILD_SHIPPING)
+	static auto ForcedBrokenDelegate = CVarNumToForceBreak->OnChangedDelegate().Add(FConsoleVariableDelegate::CreateLambda([](IConsoleVariable* CVar)
+	{
+		ForcedBroken.Reset();
+
+		for (TObjectIterator<UGeometryCollectionComponent> It; It; ++It)
+		{
+			It->RefreshCustomRenderer();
+		}
+	}));
+#endif
 }
 
 Chaos::FPhysicsSolver* UGeometryCollectionComponent::GetSolver(const UGeometryCollectionComponent& GeometryCollectionComponent)
@@ -5629,7 +5676,15 @@ void UGeometryCollectionComponent::RefreshCustomRenderer()
 					const FTransform ComponentTransform = GetComponentTransform();
 					const int32 RootIndex = GetRootIndex();
 
-					const bool bIsBroken = DynamicCollection ? !DynamicCollection->Active[RootIndex] : false;
+					bool bIsBroken = DynamicCollection ? !DynamicCollection->Active[RootIndex] : false;
+
+				#if !(UE_BUILD_SHIPPING)
+					if (CVarNumToForceBreak->GetInt() >= 0)
+					{
+						bIsBroken = ForcedBroken.Break(this, ComponentSpaceTransforms.RequestAllTransforms());
+					}
+				#endif						
+
 					const bool bRenderRootProxy = bEnableRootProxyForCustomRenderer && !bIsBroken;
 
 					RendererInterface->UpdateState(*RestCollection, ComponentTransform, !bRenderRootProxy, !bHiddenInGame);
