Add fatigue threshold percent and minimum ( replace the old Fatigue threshold value )

diff --git a/Engine/Source/Runtime/Experimental/Chaos/Private/Chaos/StrainModification.cpp b/Engine/Source/Runtime/Experimental/Chaos/Private/Chaos/StrainModification.cpp
--- a/Engine/Source/Runtime/Experimental/Chaos/Private/Chaos/StrainModification.cpp
+++ b/Engine/Source/Runtime/Experimental/Chaos/Private/Chaos/StrainModification.cpp
@@ -113,18 +113,21 @@ int32 Chaos::FStrainedProxyModifier::GetNumBreakingStrains(const bool bDoubleCou
 	return NumBreakingStrains;
 }
 
-float Chaos::FStrainedProxyModifier::GetMaxBreakStrainRatio(const float FatigueThreshold, const uint8 StrainTypes) const
+float Chaos::FStrainedProxyModifier::GetMaxBreakStrainRatio(const float FatigueThresholdPercent, const float FatigueThresholdMinimum, bool bRelative, const uint8 StrainTypes) const
 {
 	float MaxBreakStrainRatio = 0.f;
 	ForEachRootChildParticle(ProxyAndRoot, RestChildren,
-		[this, &MaxBreakStrainRatio, &FatigueThreshold, &StrainTypes] (Chaos::FPBDRigidClusteredParticleHandle* ChildHandle)
+		[this, &MaxBreakStrainRatio, FatigueThresholdPercent, FatigueThresholdMinimum, bRelative, StrainTypes] (Chaos::FPBDRigidClusteredParticleHandle* ChildHandle)
 		{
 			// compute the strain ratio
 			const float InternalStrain = ChildHandle->GetInternalStrains();
 			const float MaxAppliedStrain = GetMaxAppliedStrain(ChildHandle, StrainTypes);
+			const float FatigueThreshold = FMath::Max(FatigueThresholdMinimum, (FatigueThresholdPercent * InternalStrain));
 			if (MaxAppliedStrain >= FatigueThreshold)
 			{
-				const float StrainRatio = (InternalStrain > SMALL_NUMBER) ? (MaxAppliedStrain / InternalStrain) : 1.0f;
+				const float StrainRange = bRelative ? FMath::Max(0.f, InternalStrain - FatigueThreshold): InternalStrain;
+				const float Strain = bRelative ? FMath::Max(0.f, MaxAppliedStrain - FatigueThreshold): MaxAppliedStrain;
+				const float StrainRatio = (StrainRange > SMALL_NUMBER) ? (Strain / StrainRange) : 1.0f;
 				MaxBreakStrainRatio = FMath::Max(MaxBreakStrainRatio, StrainRatio);
 			}
 		});
@@ -132,14 +135,15 @@ float Chaos::FStrainedProxyModifier::GetMaxBreakStrainRatio(const float FatigueT
 	return MaxBreakStrainRatio;
 }
 
-void Chaos::FStrainedProxyModifier::AdjustStrainForBreak(const float FatigueThreshold, const uint8 StrainTypes)
+void Chaos::FStrainedProxyModifier::AdjustStrainForBreak(const float FatigueThresholdPercent, const float FatigueThresholdMinimum, const uint8 StrainTypes)
 {
 	ForEachRootChildParticle(ProxyAndRoot, RestChildren,
-		[this, &FatigueThreshold, &StrainTypes] (Chaos::FPBDRigidClusteredParticleHandle* ChildHandle)
+		[this, FatigueThresholdPercent, FatigueThresholdMinimum, StrainTypes] (Chaos::FPBDRigidClusteredParticleHandle* ChildHandle)
 		{
 			// compute the strain ratio
 			const float InternalStrain = ChildHandle->GetInternalStrains();
 			const float MaxAppliedStrain = GetMaxAppliedStrain(ChildHandle, StrainTypes);
+			const float FatigueThreshold = FMath::Max(FatigueThresholdMinimum, (FatigueThresholdPercent * InternalStrain));
 			if (MaxAppliedStrain >= FatigueThreshold)
 			{
 				RigidClustering.SetExternalStrain(ChildHandle, InternalStrain);
diff --git a/Engine/Source/Runtime/Experimental/Chaos/Public/Chaos/StrainModification.h b/Engine/Source/Runtime/Experimental/Chaos/Public/Chaos/StrainModification.h
--- a/Engine/Source/Runtime/Experimental/Chaos/Public/Chaos/StrainModification.h
+++ b/Engine/Source/Runtime/Experimental/Chaos/Public/Chaos/StrainModification.h
@@ -86,7 +86,7 @@ namespace Chaos
 		// Go through the children and find the one with the largest ratio of applied strain vs internal strain and return this ratio
 		// @param FatigueThreshold if the applied strains are below the fatgue threshold, they will be ignored and the returned ratio will be 0 
 		// @param StrainTypes type strain to account for ( collision and/or external )
-		CHAOS_API float GetMaxBreakStrainRatio(const float FatigueThreshold, const uint8 StrainTypes) const;
+		CHAOS_API float GetMaxBreakStrainRatio(const float FatigueThresholdPercent, const float FatigueThresholdMinimum, bool bRelative, const uint8 StrainTypes) const;
 
 		// Clear strains for all strained cluster children
 		CHAOS_API void ClearStrains();
@@ -94,7 +94,7 @@ namespace Chaos
 		// Adjust strain for all the strained children above the fatigue threshold so that their strain is large enough to break
 		// @param FatigueThreshold if the applied strains are below the fatgue threshold, they will be ignored 
 		// @param StrainTypes type strain to account for ( collision and/or external )
-		CHAOS_API void AdjustStrainForBreak(const float FatigueThreshold, const uint8 StrainTypes);
+		CHAOS_API void AdjustStrainForBreak(const float FatigueThresholdPercent, const float FatigueThresholdMinimum, const uint8 StrainTypes);
 
 	private:
 
