IAS: Allow the generation of the FOnDemandToc from disk to be done on the task system to reduce start up times (default on)

#rnx

- Renamed the cvar 's.IasGenerateOnDemandUtoc' to 's.IasGenerateOnDemandToc'
- Added a new cvar 's.IasEnableThreadedTocGeneration' that controls if we throw the work to generate the FOnDemandToc from disk to the task system or not (defaults to true)
- Note that these cvars are first used in 'FOnDemandIoBackend::Mount' which is called before we load the config file system and apply it to the cvar system. This will be fixed in a later work item.
- When an endpoint is mounted we now kick off the generation of the FOnDemandToc from disk (if enabled) onto the task system, the results of which will be collected by calls to ::ApplyGeneratedOnDemandToc
- When ApplyGeneratedOnDemandToc gets called depends on the project setup. For projects that use a CDN url directly it will be called almost immediately and there is no time saving. If the project uses a distributed endpoint system it will be called once that is resolved, usually several seconds later during initialization and will not end up blocking.
-- In a future work item we need to divorce the generation of the FOnDemandToc entirely from resolving things like the CDN url and have clearly defined pointed during initialization so teams can rely on when IAS can be used. For now we stick to the points that the previews system was using so there is no functional change.
- ApplyGeneratedOnDemandToc takes a CDN url and a TocPath even though in theory it should not need either. The TocPath is needed because we extract part of the look up address from it and both values are used to help validate that the generated data is correct if UE_VALIDATE_GENERATED_TOC is enabled. Both issues should be addressed in future work items so there is no dependency with the online system and generating the FOnDemandToc from disk.
- ::GenerateOnDemandTocFromDisk is no longer a method of FOnDemandIoBackend and has been moved to a utility function (making it clear that it is safe to throw to the task system)
- We are moving around FOnDemandToc in many places but due to it's structure we do not want it copying as this will cause many reallocations and data copies instead the intent is that it be moved when needed. So to make sure this happens the copy constructor and assignment operator have been removed.

- Profiling GameThread start up times on windows, around 95k on demand chunks:
- Downloading from the CDN ~300ms
- Generating from disk ~ 500ms
- Generating from disk on the task system ~ 10ms

diff --git a/Engine/Source/Runtime/Experimental/IoStoreOnDemand/Private/OnDemandIoDispatcherBackend.cpp b/Engine/Source/Runtime/Experimental/IoStoreOnDemand/Private/OnDemandIoDispatcherBackend.cpp
--- a/Engine/Source/Runtime/Experimental/IoStoreOnDemand/Private/OnDemandIoDispatcherBackend.cpp
+++ b/Engine/Source/Runtime/Experimental/IoStoreOnDemand/Private/OnDemandIoDispatcherBackend.cpp
@@ -122,13 +122,20 @@
 	true,
 	TEXT("Enables reporting statics to the analytics system"));
 
-static FAutoConsoleVariable CVar_IasGenerateOnDemandUtoc(
-	TEXT("s.IasGenerateOnDemandUtoc"),
+static FAutoConsoleVariable CVar_IasGenerateOnDemandToc(
+	TEXT("s.IasGenerateOnDemandToc"),
 	false,
 	TEXT("Enables generating the FOnDemandToc from utoc files on disk rather than downloading them"),
 	ECVF_ReadOnly
 );
 
+static FAutoConsoleVariable CVar_IasEnableAsyncTocGeneration(
+	TEXT("s.IasEnableThreadedTocGeneration"),
+	true,
+	TEXT("Enables pushing the work FOnDemandToc generation work to the task system"),
+	ECVF_ReadOnly
+);
+
 #if !UE_BUILD_SHIPPING
 static FAutoConsoleCommand CVar_IasAbandonCache(
 	TEXT("Ias.AbandonCache"),
@@ -834,6 +841,99 @@
 	return FoundFiles;
 }
 
+/** Generate a FOnDemandToc based on utoc files on disk which support the OnDemand feature */
+TIoStatusOr<FOnDemandToc> GenerateOnDemandTocFromDisk()
+{
+	TRACE_CPUPROFILER_EVENT_SCOPE(GenerateOnDemandTocFromDisk);
+
+	FOnDemandToc OutToc;
+
+	TArray<FString> UtocFilePaths = FindOnDemandUtocFilesOnDisk();
+	const TMap<FGuid, FAES::FAESKey> EncryptionKeys = FEncryptionKeyManager::Get().GetAllKeys();
+
+	for (const FString& UtocFilePath : UtocFilePaths)
+	{
+		FIoStoreReader Reader;
+		FIoStatus Status = Reader.Initialize(FPathViews::GetBaseFilenameWithPath(UtocFilePath), EncryptionKeys);
+		if (Status.IsOk() && EnumHasAnyFlags(Reader.GetContainerFlags(), EIoContainerFlags::OnDemand))
+		{
+			FOnDemandTocContainerEntry Container;
+
+			Container.ContainerName = FPathViews::GetBaseFilename(UtocFilePath);
+
+			const uint32 BlockSize = Reader.GetCompressionBlockSize();
+			if (OutToc.Header.BlockSize == 0)
+			{
+				OutToc.Header.BlockSize = Reader.GetCompressionBlockSize();
+			}
+			check(OutToc.Header.BlockSize == Reader.GetCompressionBlockSize());
+
+			TArray<FIoStoreTocChunkInfo> ChunkInfos;
+			Reader.EnumerateChunks([&ChunkInfos](FIoStoreTocChunkInfo&& Info)
+				{
+					ChunkInfos.Emplace(MoveTemp(Info));
+					return true;
+				});
+
+			// We can't actually hit this until we solve the FASEKey Initialize issue above
+			if (EnumHasAnyFlags(Reader.GetContainerFlags(), EIoContainerFlags::Encrypted))
+			{
+				Container.EncryptionKeyGuid = LexToString(Reader.GetEncryptionKeyGuid());
+			}
+
+			for (const FIoStoreTocChunkInfo& ChunkInfo : ChunkInfos)
+			{
+				TIoStatusOr<FIoStoreCompressedChunkInfo> InfoStatus = Reader.GetChunkCompressedInfo(ChunkInfo.Id);
+				if (!InfoStatus.IsOk())
+				{
+					return InfoStatus.Status();
+				}
+				FIoStoreCompressedChunkInfo CompressedChunkInfo = InfoStatus.ConsumeValueOrDie();
+
+				const uint32 BlockOffset = Container.BlockSizes.Num();
+				const uint32 BlockCount = CompressedChunkInfo.Blocks.Num();
+
+				uint64 RawChunkSize = 0;
+				uint64 EncodedChunkSize = 0;
+				for (const FIoStoreCompressedBlockInfo& BlockInfo : CompressedChunkInfo.Blocks)
+				{
+					const uint64 EncodedBlockSize = Align(BlockInfo.CompressedSize, FAES::AESBlockSize);
+					Container.BlockSizes.Add(uint32(EncodedBlockSize));
+
+					FIoBlockHash BlockHash;
+					FMemory::Memcpy(&BlockHash, &BlockInfo.DiskHash, sizeof(FIoBlockHash));
+					Container.BlockHashes.Add(BlockHash);
+
+					EncodedChunkSize += EncodedBlockSize;
+					RawChunkSize += BlockInfo.UncompressedSize;
+
+					if (OutToc.Header.CompressionFormat.IsEmpty() && BlockInfo.CompressionMethod != NAME_None)
+					{
+						OutToc.Header.CompressionFormat = BlockInfo.CompressionMethod.ToString();
+					}
+				}
+
+				FOnDemandTocEntry& TocEntry = Container.Entries.AddDefaulted_GetRef();
+				TocEntry.ChunkId = ChunkInfo.Id;
+				TocEntry.Hash = CompressedChunkInfo.DiskHash;
+				TocEntry.RawSize = RawChunkSize;
+				TocEntry.EncodedSize = EncodedChunkSize;
+				TocEntry.BlockOffset = BlockOffset;
+				TocEntry.BlockCount = BlockCount;
+			}
+
+			OutToc.Containers.Emplace(MoveTemp(Container));
+		}
+	}
+
+	if (!OutToc.Containers.IsEmpty())
+	{
+		OutToc.Header.ChunksDirectory = FString::Printf(TEXT("IoChunksV%u"), EOnDemandChunkVersion::Latest).ToLower();
+	}
+
+	return OutToc;
+}
+
 /** Validation code used during development to ensure that the results are correct */
 #if UE_VALIDATE_GENERATED_TOC
 
@@ -1224,12 +1324,14 @@
 	virtual uint32 Run() override;
 
 private:
-	TIoStatusOr<FOnDemandToc> GenerateOnDemandTocFromDisk() const;
 
 	FString GetEndpointTestPath() const;
 	void ConditionallyStartBackendThread();
 	void CompleteRequest(FChunkRequest* ChunkRequest);
-	FIoStatus LoadOnDemandToc(const FString& CdnUrl, const FString& TocPath);
+
+	FIoStatus ApplyGeneratedOnDemandToc(const FString& CdnUrl, const FString& TocPath);
+	FIoStatus DownloadoadOnDemandToc(const FString& CdnUrl, const FString& TocPath);
+
 	void AddDeferredTocs();
 	void ProcessHttpRequests(FOnDemandHttpClient* HttpClient, FBitWindow& HttpErrors, int32 MaxConcurrentRequests);
 
@@ -1249,6 +1351,9 @@
 	mutable FRWLock Lock;
 	std::atomic_uint32_t InflightCacheRequestCount{0};
 	std::atomic_bool bStopRequested{false};
+
+	bool bGeneratedOnDemandToc = false;
+	UE::Tasks::TTask<TIoStatusOr<FOnDemandToc>> OnDemandTocTask;
 };
 
 ///////////////////////////////////////////////////////////////////////////////
@@ -1324,98 +1429,6 @@
 	BackendContext.Reset();
 }
 
-TIoStatusOr<FOnDemandToc> FOnDemandIoBackend::GenerateOnDemandTocFromDisk() const
-{
-	TRACE_CPUPROFILER_EVENT_SCOPE(FOnDemandIoBackend::GenerateOnDemandTocFromDisk);
-
-	FOnDemandToc OutToc;
-
-	TArray<FString> UtocFilePaths = FindOnDemandUtocFilesOnDisk();
-	const TMap<FGuid, FAES::FAESKey> EncryptionKeys = FEncryptionKeyManager::Get().GetAllKeys();
-
-	for (const FString& UtocFilePath : UtocFilePaths)
-	{
-		FIoStoreReader Reader;
-		FIoStatus Status = Reader.Initialize(FPathViews::GetBaseFilenameWithPath(UtocFilePath), EncryptionKeys);
-		if (Status.IsOk() && EnumHasAnyFlags(Reader.GetContainerFlags(), EIoContainerFlags::OnDemand))
-		{
-			FOnDemandTocContainerEntry Container;
-
-			Container.ContainerName = FPathViews::GetBaseFilename(UtocFilePath);
-
-			const uint32 BlockSize = Reader.GetCompressionBlockSize();
-			if (OutToc.Header.BlockSize == 0)
-			{
-				OutToc.Header.BlockSize = Reader.GetCompressionBlockSize();
-			}
-			check(OutToc.Header.BlockSize == Reader.GetCompressionBlockSize());
-
-			TArray<FIoStoreTocChunkInfo> ChunkInfos;
-			Reader.EnumerateChunks([&ChunkInfos](FIoStoreTocChunkInfo&& Info)
-				{
-					ChunkInfos.Emplace(MoveTemp(Info));
-					return true;
-				});
-
-			// We can't actually hit this until we solve the FASEKey Initialize issue above
-			if (EnumHasAnyFlags(Reader.GetContainerFlags(), EIoContainerFlags::Encrypted))
-			{
-				Container.EncryptionKeyGuid = LexToString(Reader.GetEncryptionKeyGuid());
-			}
-
-			for (const FIoStoreTocChunkInfo& ChunkInfo : ChunkInfos)
-			{
-				TIoStatusOr<FIoStoreCompressedChunkInfo> InfoStatus = Reader.GetChunkCompressedInfo(ChunkInfo.Id);
-				if (!InfoStatus.IsOk())
-				{
-					return InfoStatus.Status();
-				}
-				FIoStoreCompressedChunkInfo CompressedChunkInfo = InfoStatus.ConsumeValueOrDie();
-
-				const uint32 BlockOffset = Container.BlockSizes.Num();
-				const uint32 BlockCount = CompressedChunkInfo.Blocks.Num();
-
-				uint64 RawChunkSize = 0;
-				uint64 EncodedChunkSize = 0;
-				for (const FIoStoreCompressedBlockInfo& BlockInfo : CompressedChunkInfo.Blocks)
-				{
-					const uint64 EncodedBlockSize = Align(BlockInfo.CompressedSize, FAES::AESBlockSize);
-					Container.BlockSizes.Add(uint32(EncodedBlockSize));
-
-					FIoBlockHash BlockHash;
-					FMemory::Memcpy(&BlockHash, &BlockInfo.DiskHash, sizeof(FIoBlockHash));
-					Container.BlockHashes.Add(BlockHash);
-
-					EncodedChunkSize += EncodedBlockSize;
-					RawChunkSize += BlockInfo.UncompressedSize;
-
-					if (OutToc.Header.CompressionFormat.IsEmpty() && BlockInfo.CompressionMethod != NAME_None)
-					{
-						OutToc.Header.CompressionFormat = BlockInfo.CompressionMethod.ToString();
-					}
-				}
-
-				FOnDemandTocEntry& TocEntry = Container.Entries.AddDefaulted_GetRef();
-				TocEntry.ChunkId = ChunkInfo.Id;
-				TocEntry.Hash = CompressedChunkInfo.DiskHash;
-				TocEntry.RawSize = RawChunkSize;
-				TocEntry.EncodedSize = EncodedChunkSize;
-				TocEntry.BlockOffset = BlockOffset;
-				TocEntry.BlockCount = BlockCount;
-			}
-
-			OutToc.Containers.Emplace(MoveTemp(Container));
-		}
-	}
-
-	if (!OutToc.Containers.IsEmpty())
-	{
-		OutToc.Header.ChunksDirectory = FString::Printf(TEXT("IoChunksV%u"), EOnDemandChunkVersion::Latest).ToLower();
-	}
-
-	return OutToc;
-}
-
 FString FOnDemandIoBackend::GetEndpointTestPath() const
 {
 	FString TestPath = IoStore.IsValid() ? IoStore->GetFirstTocPath() : FString();
@@ -1704,52 +1717,63 @@
 	return FIoStatus::Unknown;
 }
 
-FIoStatus FOnDemandIoBackend::LoadOnDemandToc(const FString& CdnUrl, const FString& TocPath)
+FIoStatus FOnDemandIoBackend::ApplyGeneratedOnDemandToc(const FString& CdnUrl, const FString& TocPath)
 {
-	TRACE_CPUPROFILER_EVENT_SCOPE(FOnDemandIoBackend::LoadOnDemandToc);
+	TRACE_CPUPROFILER_EVENT_SCOPE(FOnDemandIoBackend::ApplyGeneratedOnDemandToc);
 
-	if (CVar_IasGenerateOnDemandUtoc->GetBool())
+	if (bGeneratedOnDemandToc || !CVar_IasGenerateOnDemandToc->GetBool())
 	{
-		UE_LOG(LogIas, Log, TEXT("Generating OnDemandToc from disk"));
+		return FIoStatus::Ok;
+	}
 
-		// TODO: We don't really handle the case where the utoc on disk are encrypted and we don't have the key yet/
-		// See the 'DeferredTocs' code below.
-		TIoStatusOr<FOnDemandToc> GeneratedTocResult = GenerateOnDemandTocFromDisk();
+	TIoStatusOr<FOnDemandToc> GeneratedTocResult;
 
-		if (GeneratedTocResult.IsOk())
-		{
+	if (OnDemandTocTask.IsValid())
+	{
+		GeneratedTocResult = MoveTemp(OnDemandTocTask.GetResult());
+	}
+	else
+	{
+		GeneratedTocResult = GenerateOnDemandTocFromDisk();
+	}
+	
+	if (!GeneratedTocResult.IsOk())
+	{
+		return GeneratedTocResult.Status();
+	}
+
 #if UE_VALIDATE_GENERATED_TOC == 0
-			IoStore->AddToc(TocPath, GeneratedTocResult.ConsumeValueOrDie());
+	IoStore->AddToc(TocPath, GeneratedTocResult.ConsumeValueOrDie());
 #else
-			FOnDemandToc UrlToc = LoadTocFromUrl(CdnUrl, TocPath, 1).ConsumeValueOrDie();
-			FOnDemandToc GeneratedToc = GeneratedTocResult.ConsumeValueOrDie();
+	FOnDemandToc UrlToc = LoadTocFromUrl(CdnUrl, TocPath, 1).ConsumeValueOrDie();
+	FOnDemandToc GeneratedToc = GeneratedTocResult.ConsumeValueOrDie();
 
-			UE::IO::IAS::ValidateToc(UrlToc, GeneratedToc);
+	UE::IO::IAS::ValidateToc(UrlToc, GeneratedToc);
 
-			IoStore->AddToc(TocPath, MoveTemp(GeneratedToc));
+	IoStore->AddToc(TocPath, MoveTemp(GeneratedToc));
 #endif // UE_VALIDATE_GENERATED_TOC
-		}
-		else
-		{
-			return GeneratedTocResult.Status();
-		}
+
+	bGeneratedOnDemandToc = true;
+
+	return FIoStatus::Ok;
+}
+
+FIoStatus FOnDemandIoBackend::DownloadoadOnDemandToc(const FString& CdnUrl, const FString& TocPath)
+{
+	TRACE_CPUPROFILER_EVENT_SCOPE(FOnDemandIoBackend::LoadOnDemandToc);
+
+	UE_LOG(LogIas, Log, TEXT("Downloading OnDemandToc from CDN"));
+
+	TIoStatusOr<FOnDemandToc> TocResult = LoadTocFromUrl(CdnUrl, TocPath, 1);
+	if (TocResult.IsOk())
+	{
+		IoStore->AddToc(TocPath, TocResult.ConsumeValueOrDie());
+		return FIoStatus::Ok;
 	}
 	else
 	{
-		UE_LOG(LogIas, Log, TEXT("Downloading OnDemandToc from CDN"));
-
-		TIoStatusOr<FOnDemandToc> TocResult = LoadTocFromUrl(CdnUrl, TocPath, 1);
-		if (TocResult.IsOk())
-		{
-			IoStore->AddToc(TocPath, TocResult.ConsumeValueOrDie());
-		}
-		else
-		{
-			return TocResult.Status();
-		}
+		return TocResult.Status();
 	}
-
-	return EIoErrorCode::Ok;
 }
 
 void FOnDemandIoBackend::AddDeferredTocs()
@@ -1764,12 +1788,31 @@
 		TocPaths = MoveTemp(DeferredTocs);
 	}
 
-	for (const FString& TocPath : TocPaths)
+	// We are still reliant on a valid TocPath to generate the FOnDemandToc from disk as the path contains
+	// info that we still need.
+	if (!TocPaths.IsEmpty())
 	{
-		FIoStatus Result = LoadOnDemandToc(AvailableEps.GetCurrent(), TocPath);
+		FIoStatus Result = ApplyGeneratedOnDemandToc(AvailableEps.GetCurrent(), TocPaths[0]);
 		if (!Result.IsOk())
 		{
-			UE_LOG(LogIas, Error, TEXT("Failed to add TOC '%s/%s', reason '%s'"), *AvailableEps.GetCurrent(), *TocPath, *Result.ToString());
+			UE_LOG(LogIas, Error, TEXT("Failed to add generated toc', reason '%s'"), *Result.ToString());
 		}
 	}
+
+	// Generating the FOnDemandToc is relying on DeferredTocs having entries and cribbing off the callback from resolving the CDN endpoint
+	// This is to a) make it easier to toggle during development b) easier to validate the results.
+	// So keeping that in mind we need to prevent the toc being downloaded if we are supposed to be using the disk generated version 
+	// instead.
+	// This can all be cleaned up when we either remove the toggle or drop the validation requirements AND clean up the initialization flow.
+	if (!CVar_IasGenerateOnDemandToc->GetBool())
+	{
+		for (const FString& TocPath : TocPaths)
+		{
+			FIoStatus Result = DownloadoadOnDemandToc(AvailableEps.GetCurrent(), TocPath);
+			if (!Result.IsOk())
+			{
+				UE_LOG(LogIas, Error, TEXT("Failed to add TOC '%s/%s', reason '%s'"), *AvailableEps.GetCurrent(), *TocPath, *Result.ToString());
+			}
+		}
+	}
 }
@@ -1785,6 +1828,14 @@
 		return;
 	}
 
+	if (CVar_IasGenerateOnDemandToc->GetBool() && CVar_IasEnableAsyncTocGeneration->GetBool())
+	{
+		OnDemandTocTask = UE::Tasks::Launch(UE_SOURCE_LOCATION, [this]() -> TIoStatusOr<FOnDemandToc>
+			{
+				return GenerateOnDemandTocFromDisk();
+			});
+	}
+
 	if (EnumHasAnyFlags(Endpoint.EndpointType, EOnDemandEndpointType::CDN))
 	{
 		{
@@ -1810,15 +1861,24 @@
 
 		check(AvailableEps.HasCurrent());
 
-		FIoStatus Result = LoadOnDemandToc(AvailableEps.GetCurrent(), Endpoint.TocPath);
-		if (!Result.IsOk())
+		FIoStatus GeneratedResult = ApplyGeneratedOnDemandToc(AvailableEps.GetCurrent(), Endpoint.TocPath);
+		if (!GeneratedResult.IsOk())
 		{
-			UE_LOG(LogIas, Error, TEXT("Deferring TOC '%s/%s' due to '%s'"), *AvailableEps.GetCurrent(), *Endpoint.TocPath, *Result.ToString());
-			BackendStatus.SetHttpError(true);
-			FWriteScopeLock _(Lock);
-			DeferredTocs.Add(Endpoint.TocPath);
+			UE_LOG(LogIas, Error, TEXT("Failed to add generated toc', reason '%s'"), *GeneratedResult.ToString());
 		}
 
+		if (!CVar_IasGenerateOnDemandToc->GetBool())
+		{
+			FIoStatus Result = DownloadoadOnDemandToc(AvailableEps.GetCurrent(), Endpoint.TocPath);
+			if (!Result.IsOk())
+			{
+				UE_LOG(LogIas, Error, TEXT("Deferring TOC '%s/%s' due to '%s'"), *AvailableEps.GetCurrent(), *Endpoint.TocPath, *Result.ToString());
+					BackendStatus.SetHttpError(true);
+					FWriteScopeLock _(Lock);
+					DeferredTocs.Add(Endpoint.TocPath);
+			}
+	}
+
 		ConditionallyStartBackendThread();
 	}
 	else
diff --git a/Engine/Source/Runtime/Experimental/IoStoreOnDemand/Public/IO/IoStoreOnDemand.h b/Engine/Source/Runtime/Experimental/IoStoreOnDemand/Public/IO/IoStoreOnDemand.h
--- a/Engine/Source/Runtime/Experimental/IoStoreOnDemand/Public/IO/IoStoreOnDemand.h
+++ b/Engine/Source/Runtime/Experimental/IoStoreOnDemand/Public/IO/IoStoreOnDemand.h
@@ -121,6 +121,17 @@ UE_API bool LoadFromCompactBinary(FCbFieldView Field, FOnDemandTocContainerEntry
 
 struct FOnDemandToc
 {
+	FOnDemandToc() = default;
+	~FOnDemandToc() = default;
+
+	FOnDemandToc(FOnDemandToc&&) = default;
+	FOnDemandToc& operator= (FOnDemandToc&&) = default;
+
+	// Copying this structure would be quite expensive so we want to make sure that it doesn't happen.
+
+	FOnDemandToc(const FOnDemandToc&) = delete;
+	FOnDemandToc&  operator= (const FOnDemandToc&) = delete;
+
 	FOnDemandTocHeader Header;
 	FTocMeta Meta;
 	TArray<FOnDemandTocContainerEntry> Containers;
