Add a 'ias.InvokeHttpFailure' console command that marks the current IAS http connection as failed so that the system needs to reconnect again. Useful for debugging the error flow.

#rnx

diff --git a/Engine/Source/Runtime/Experimental/IoStoreOnDemand/Private/OnDemandIoDispatcherBackend.cpp b/Engine/Source/Runtime/Experimental/IoStoreOnDemand/Private/OnDemandIoDispatcherBackend.cpp
--- a/Engine/Source/Runtime/Experimental/IoStoreOnDemand/Private/OnDemandIoDispatcherBackend.cpp
+++ b/Engine/Source/Runtime/Experimental/IoStoreOnDemand/Private/OnDemandIoDispatcherBackend.cpp
@@ -57,6 +57,9 @@
  */
 #define UE_VALIDATE_GENERATED_TOC (0 && !UE_BUILD_SHIPPING)
 
+/** When enabled the IAS system can add additional debug console commands for development use */
+#define UE_IAS_DEBUG_CONSOLE_CMDS (1 && !NO_CVARS && !UE_BUILD_SHIPPING)
+
 namespace UE::IO::IAS
 {
 
@@ -1354,6 +1357,10 @@ private:
 
 	bool bGeneratedOnDemandToc = false;
 	UE::Tasks::TTask<TIoStatusOr<FOnDemandToc>> OnDemandTocTask;
+
+#if UE_IAS_DEBUG_CONSOLE_CMDS
+	TArray<IConsoleCommand*> DynamicConsoleCommands;
+#endif // UE_IAS_DEBUG_CONSOLE_CMDS
 };
 
 ///////////////////////////////////////////////////////////////////////////////
@@ -1363,10 +1370,33 @@ FOnDemandIoBackend::FOnDemandIoBackend(TUniquePtr<IIasCache>&& InCache)
 	IoStore = MakeUnique<FOnDemandIoStore>();
 	BackendStatus.SetHttpEnabled(true);
 	BackendStatus.SetCacheEnabled(Cache.IsValid());
+
+#if UE_IAS_DEBUG_CONSOLE_CMDS
+	DynamicConsoleCommands.Emplace(
+	IConsoleManager::Get().RegisterConsoleCommand(
+		TEXT("ias.InvokeHttpFailure"),
+		TEXT(""),
+		FConsoleCommandDelegate::CreateLambda([this]()
+			{
+				UE_LOG(LogIas, Display, TEXT("Marks the current ias http connection as failed forcing the system to try to reconnect"));
+				BackendStatus.SetHttpError(true);
+
+				TickBackendEvent->Trigger();
+			}),
+		ECVF_Cheat)
+	);
+#endif // UE_IAS_DEBUG_CONSOLE_CMDS
 }
 
 FOnDemandIoBackend::~FOnDemandIoBackend()
 {
+#if UE_IAS_DEBUG_CONSOLE_CMDS
+	for (IConsoleCommand* Cmd : DynamicConsoleCommands)
+	{
+		IConsoleManager::Get().UnregisterConsoleObject(Cmd);
+	}
+#endif // UE_IAS_DEBUG_CONSOLE_CMDS
+
 	Shutdown();
 }
 
@@ -2106,4 +2136,5 @@ TSharedPtr<IOnDemandIoDispatcherBackend> MakeOnDemandIoDispatcherBackend(TUnique
 
 } // namespace UE::IO::IAS
 
+#undef UE_IAS_DEBUG_CONSOLE_CMDS
 #undef UE_VALIDATE_GENERATED_TOC
