Adding damping to clusterunion particle.  Damping was not working before because the cluster union was not actually having the value set on the PT.

diff --git a/Engine/Source/Runtime/Engine/Private/PrimitiveComponentPhysics.cpp b/Engine/Source/Runtime/Engine/Private/PrimitiveComponentPhysics.cpp
--- a/Engine/Source/Runtime/Engine/Private/PrimitiveComponentPhysics.cpp
+++ b/Engine/Source/Runtime/Engine/Private/PrimitiveComponentPhysics.cpp
@@ -404,7 +404,7 @@ void UPrimitiveComponent::WakeRigidBody(FName BoneName)
 	}
 	else if (Chaos::FPhysicsObject* Object = GetPhysicsObjectByName(BoneName))
 	{
-		FLockedWritePhysicsObjectExternalInterface Interface = FPhysicsObjectExternalInterface::LockWrite({&Object, 1});
+		FLockedWritePhysicsObjectExternalInterface Interface = FPhysicsObjectExternalInterface::LockWrite({ &Object, 1 });
 		Interface->WakeUp({ &Object, 1 });
 	}
 }
@@ -442,6 +442,12 @@ void UPrimitiveComponent::SetLinearDamping(float InDamping)
 		BI->LinearDamping = InDamping;
 		BI->UpdateDampingProperties();
 	}
+
+	else if (Chaos::FPhysicsObject* Object = GetPhysicsObjectByName(NAME_None))
+	{
+		FLockedWritePhysicsObjectExternalInterface Interface = FPhysicsObjectExternalInterface::LockWrite({ &Object, 1 });
+		Interface->SetLinearEtherDrag({ &Object, 1 }, InDamping);
+	}
 }
 
 float UPrimitiveComponent::GetLinearDamping() const
@@ -464,6 +470,12 @@ void UPrimitiveComponent::SetAngularDamping(float InDamping)
 		BI->AngularDamping = InDamping;
 		BI->UpdateDampingProperties();
 	}
+
+	else if (Chaos::FPhysicsObject* Object = GetPhysicsObjectByName(NAME_None))
+	{
+		FLockedWritePhysicsObjectExternalInterface Interface = FPhysicsObjectExternalInterface::LockWrite({ &Object, 1 });
+		Interface->SetAngularEtherDrag({ &Object, 1 }, InDamping);
+	}
 }
 
 float UPrimitiveComponent::GetAngularDamping() const
diff --git a/Engine/Source/Runtime/Experimental/Chaos/Private/Chaos/PhysicsObjectInterface.cpp b/Engine/Source/Runtime/Experimental/Chaos/Private/Chaos/PhysicsObjectInterface.cpp
--- a/Engine/Source/Runtime/Experimental/Chaos/Private/Chaos/PhysicsObjectInterface.cpp
+++ b/Engine/Source/Runtime/Experimental/Chaos/Private/Chaos/PhysicsObjectInterface.cpp
@@ -896,6 +896,7 @@ namespace Chaos
 	template<EThreadContext Id>
 	void FWritePhysicsObjectInterface<Id>::AddRadialImpulse(TArrayView<const FPhysicsObjectHandle> InObjects, FVector Origin, float Radius, float Strength, enum ERadialImpulseFalloff Falloff, bool bApplyStrain, bool bInvalidate)
 	{
+		//TODO: create a PT version of this, plus the damping functions
 		if (Chaos::FPBDRigidsSolver* RigidSolver = Chaos::FPhysicsObjectInterface::GetSolver(InObjects))
 		{
 			//put onto physics thread
@@ -978,6 +979,59 @@ namespace Chaos
 		}
 	}
 
+	template<EThreadContext Id>
+	void Chaos::FWritePhysicsObjectInterface<Id>::SetLinearEtherDrag(TArrayView<const FPhysicsObjectHandle> InObjects, float InLinearDrag)
+	{
+		if (Chaos::FPBDRigidsSolver* RigidSolver = Chaos::FPhysicsObjectInterface::GetSolver(InObjects))
+		{
+			//put onto physics thread
+			RigidSolver->EnqueueCommandImmediate([InObjects = TArray<FPhysicsObjectHandle>{ InObjects }, InLinearDrag, RigidSolver]() {
+				for (const FPhysicsObjectHandle Object : InObjects)
+				{
+					if (!Object)
+					{
+						continue;
+					}
+
+					if (TThreadParticle<EThreadContext::Internal>* Particle = Object->GetParticle<EThreadContext::Internal>())
+					{
+						if (Chaos::TThreadRigidParticle<EThreadContext::Internal>* Rigid = Particle->CastToRigidParticle())
+						{
+							Rigid->SetLinearEtherDrag(InLinearDrag);
+						}
+					}
+				}
+				});
+		}
+	}
+
+
+	template<EThreadContext Id>
+	void Chaos::FWritePhysicsObjectInterface<Id>::SetAngularEtherDrag(TArrayView<const FPhysicsObjectHandle> InObjects, float InAngularDrag)
+	{
+		if (Chaos::FPBDRigidsSolver* RigidSolver = Chaos::FPhysicsObjectInterface::GetSolver(InObjects))
+		{
+			//put onto physics thread
+			RigidSolver->EnqueueCommandImmediate([InObjects = TArray<FPhysicsObjectHandle>{ InObjects }, InAngularDrag, RigidSolver]() {
+				for (const FPhysicsObjectHandle Object : InObjects)
+				{
+					if (!Object)
+					{
+						continue;
+					}
+
+					if (TThreadParticle<EThreadContext::Internal>* Particle = Object->GetParticle<EThreadContext::Internal>())
+					{
+						if (Chaos::TThreadRigidParticle<EThreadContext::Internal>* Rigid = Particle->CastToRigidParticle())
+						{
+							Rigid->SetAngularEtherDrag(InAngularDrag);
+						}
+					}
+				}
+			});
+		}
+	}
+
 	template<EThreadContext Id>
 	void FWritePhysicsObjectInterface<Id>::UpdateShapeCollisionFlags(TArrayView<const FPhysicsObjectHandle> InObjects, bool bSimCollision, bool bQueryCollision)
 	{
diff --git a/Engine/Source/Runtime/Experimental/Chaos/Public/Chaos/PhysicsObjectInterface.h b/Engine/Source/Runtime/Experimental/Chaos/Public/Chaos/PhysicsObjectInterface.h
--- a/Engine/Source/Runtime/Experimental/Chaos/Public/Chaos/PhysicsObjectInterface.h
+++ b/Engine/Source/Runtime/Experimental/Chaos/Public/Chaos/PhysicsObjectInterface.h
@@ -126,7 +126,9 @@ namespace Chaos
 		CHAOS_API void AddForce(TArrayView<const FPhysicsObjectHandle> InObjects, const FVector& Force, bool bInvalidate);
 		CHAOS_API void AddTorque(TArrayView<const FPhysicsObjectHandle> InObjects, const FVector& Torque, bool bInvalidate);
 		CHAOS_API void AddRadialImpulse(TArrayView<const FPhysicsObjectHandle> InObjects, FVector Origin, float Radius, float Strength, enum ERadialImpulseFalloff Falloff, bool bApplyStrain, bool bInvalidate);
-		
+		CHAOS_API void SetLinearEtherDrag(TArrayView<const FPhysicsObjectHandle> InObjects, float InLinearDrag);
+		CHAOS_API void SetAngularEtherDrag(TArrayView<const FPhysicsObjectHandle> InObjects, float InAngularDrag);
+
 		CHAOS_API void UpdateShapeCollisionFlags(TArrayView<const FPhysicsObjectHandle> InObjects, bool bSimCollision, bool bQueryCollision);
 		CHAOS_API void UpdateShapeFilterData(TArrayView<const FPhysicsObjectHandle> InObjects, const FCollisionFilterData& QueryData, const FCollisionFilterData& SimData);
 
@@ -169,6 +171,7 @@ namespace Chaos
 		FWritePhysicsObjectInterface() = default;
 	};
 
+
 	using FWritePhysicsObjectInterface_External = FWritePhysicsObjectInterface<EThreadContext::External>;
 	using FWritePhysicsObjectInterface_Internal = FWritePhysicsObjectInterface<EThreadContext::Internal>;
 
