creating an internal cluster must clean its children particles' parents' book-keeping

- Removing a GC should also disable any active internal clusters immediately
- Switch to using FRigidClustering::DisableCluster in FPBDRigidsSolver::UnregisterObject for GC's
- Fixed an issue where duplicate particles could be added to the same island in FRigidClustering::FindIslandsInChildren
- Use RemoveParticlesFromCluster in FRigidClustering::HandleConnectivityOnReleaseClusterParticle rather than RemoveChildFromParent to ensure MChildren book-keeping is updated

diff --git a/Engine/Source/Runtime/Experimental/Chaos/Private/Chaos/PBDRigidClustering.cpp b/Engine/Source/Runtime/Experimental/Chaos/Private/Chaos/PBDRigidClustering.cpp
--- a/Engine/Source/Runtime/Experimental/Chaos/Private/Chaos/PBDRigidClustering.cpp
+++ b/Engine/Source/Runtime/Experimental/Chaos/Private/Chaos/PBDRigidClustering.cpp
@@ -699,6 +699,8 @@ namespace Chaos
 
 			TArray<FPBDRigidParticleHandle*> ProcessingQueue;
 			ProcessingQueue.Add(Child);
+			ProcessedChildren.Add(Child);
+
 			while (ProcessingQueue.Num())
 			{
 				if (FPBDRigidParticleHandle* ChildToProcess = ProcessingQueue.Pop())
@@ -1169,15 +1171,20 @@ namespace Chaos
 	TSet<FPBDRigidParticleHandle*> FRigidClustering::HandleConnectivityOnReleaseClusterParticle(FPBDRigidClusteredParticleHandle* ClusteredParticle, bool bCreateNewClusters)
 	{
 		SCOPE_CYCLE_COUNTER(HandleConnectivityOnReleaseClusterParticle);
+		if (!ensure(ClusteredParticle))
+		{
+			return {};
+		}
+
 		TSet<FPBDRigidParticleHandle*> ActivatedChildren;
-		TArray<FPBDRigidParticleHandle*>& Children = MChildren[ClusteredParticle];
+		const bool bHasChildren = (ClusteredParticle->ClusterIds().NumChildren > 0);
 
 		// If we're breaking a geometry collection, we'll need to create internal clusters to parent the remaining particles.
 		// However, we do not need to do this if we're currently operating on a cluster union! Its remaining particles should stay
 		// attached to the cluster union because they can handle particles being dynamically added/removed.
 		const FClusterUnion* ParentClusterUnion = ClusterUnionManager.FindClusterUnionFromParticle(ClusteredParticle);
 		const bool bIsClusterUnion = ParentClusterUnion != nullptr;
-		if (Children.Num())
+		if (bHasChildren)
 		{
 			TArray<FParticleIsland> Islands = FindIslandsInChildren(ClusteredParticle, bIsClusterUnion);
 
@@ -1289,8 +1296,8 @@ namespace Chaos
 							{
 								// Need to remove node connections here. Otherwise it may be possible for the cluster union to have erroneous intercluster edges that connect it to another cluster union.
 								RemoveNodeConnections(ChildParticle);
-								RemoveChildFromParent(ChildParticle, ClusteredParticle);
 							}
+							RemoveParticlesFromCluster(ClusteredParticle, ParticlesToRelease);
 							ActivatedChildren.Append(ParticlesToRelease);
 						}
 						IslandIndicesToRemove.Add(IslandIndex);
@@ -1298,7 +1305,7 @@ namespace Chaos
 					else if (Island.Num() == 1 && !AttachedClusterUnion) //need to break single pieces first
 					{
 						FPBDRigidParticleHandle* Child = Island[0];
-						RemoveChildFromParent(Child, ClusteredParticle);
+						RemoveParticlesFromCluster(ClusteredParticle, { Child });
 						ActivatedChildren.Add(Child);
 					}
 
@@ -1312,6 +1319,30 @@ namespace Chaos
 
 				if (bCreateNewClusters)
 				{
+					// Each island is going to be removed from the parent particle. Pre-emptively remove each island from the parent particle's book-keeping.
+					// If we don't do this, the particles will still be stored as children of the particle's previous parents in MChildren.
+					for (const FParticleIsland& Island : Islands)
+					{
+						RemoveParticlesFromCluster(ClusteredParticle, Island);
+
+						if (Island.Num() > 1)
+						{
+							// Need to subsequently disable the particle because they probably
+							// just got re-enabled in RemoveParticlesFromCluster.
+							for (FPBDRigidParticleHandle* Particle : Island)
+							{
+								if (FPBDRigidClusteredParticleHandle* ClusterParticle = Particle->CastToClustered())
+								{
+									DisableCluster(ClusterParticle);
+								}
+								else
+								{
+									MEvolution.DisableParticle(Particle);
+								}
+							}
+						}
+					}
+
 					TArray<FPBDRigidParticleHandle*> NewClusters = CreateClustersFromNewIslands(Islands, ClusteredParticle);
 					ActivatedChildren.Append(MoveTemp(NewClusters));
 				}
@@ -2018,8 +2049,6 @@ namespace Chaos
 	void FRigidClustering::ComputeStrainFromCollision(const FPBDCollisionConstraints& CollisionRule, const FReal Dt)
 	{
 		SCOPE_CYCLE_COUNTER(STAT_ComputeStrainFromCollision);
-		FClusterMap& MParentToChildren = GetChildrenMap();
-
 		const FRealSingle InvDt = (Dt > SMALL_NUMBER) ? (FRealSingle)(1.0 / Dt) : 1.0f;
 
 		ResetCollisionImpulseArray();
@@ -2308,7 +2337,7 @@ namespace Chaos
 
 		// max strain will allow to unconditionally release the children when strain is evaluated
 		constexpr FRealSingle MaxStrain = TNumericLimits<FRealSingle>::Max();
-		if (TArray<FPBDRigidParticleHandle*>* ChildrenHandles = GetChildrenMap().Find(ClusteredParticle))
+		if (const TArray<FPBDRigidParticleHandle*>* ChildrenHandles = GetChildrenMap().Find(ClusteredParticle))
 		{
 			for (FPBDRigidParticleHandle* ChildHandle: *ChildrenHandles)
 			{
diff --git a/Engine/Source/Runtime/Experimental/Chaos/Private/PBDRigidsSolver.cpp b/Engine/Source/Runtime/Experimental/Chaos/Private/PBDRigidsSolver.cpp
--- a/Engine/Source/Runtime/Experimental/Chaos/Private/PBDRigidsSolver.cpp
+++ b/Engine/Source/Runtime/Experimental/Chaos/Private/PBDRigidsSolver.cpp
@@ -791,25 +791,6 @@ namespace Chaos
 
 		EnqueueCommandImmediate([InProxy, this]()
 			{
-				FClusterUnionManager& ClusterUnionManager = MEvolution->GetRigidClustering().GetClusterUnionManager();
-				for (FPBDRigidClusteredParticleHandle* Handle : InProxy->GetSolverParticleHandles())
-				{
-					if (Handle)
-					{
-						// It should be safe to defer here without an immediate call to HandleDeferredClusterUnionUpdateProperties.
-						// This change doesn't really have to go through until FClusterUnionManager::FlushPendingOperations which happens
-						// prior to trying to advance the frame.
-						ClusterUnionManager.HandleRemoveOperationWithClusterLookup({ Handle }, EClusterUnionOperationTiming::Defer);
-
-						// Need to use FRigidClustering::DisableCluster instead of the evolution's to also handle the fact that the GC particle could be in the top level strained sets
-						// which would make it get considered in the breaking model which is undesirable.
-						MEvolution->DisableParticle(Handle);
-
-						// This ensures that this particle won't have any intercluster edges on it. This way no connectivity operations will try to add it into a cluster union.
-						MEvolution->GetRigidClustering().RemoveNodeConnections(Handle);
-					}
-				}
-
 				GeometryCollectionPhysicsProxies_Internal.RemoveSingle(InProxy);
 				InProxy->SyncBeforeDestroy();
 				PendingDestroyGeometryCollectionPhysicsProxy.Add(InProxy);
diff --git a/Engine/Source/Runtime/Experimental/Chaos/Private/PhysicsProxy/GeometryCollectionPhysicsProxy.cpp b/Engine/Source/Runtime/Experimental/Chaos/Private/PhysicsProxy/GeometryCollectionPhysicsProxy.cpp
--- a/Engine/Source/Runtime/Experimental/Chaos/Private/PhysicsProxy/GeometryCollectionPhysicsProxy.cpp
+++ b/Engine/Source/Runtime/Experimental/Chaos/Private/PhysicsProxy/GeometryCollectionPhysicsProxy.cpp
@@ -2885,7 +2885,40 @@ void FGeometryCollectionPhysicsProxy::OnRemoveFromScene()
 
 void FGeometryCollectionPhysicsProxy::SyncBeforeDestroy()
 {
+	// THIS HAPPENS ON THE PHYSICS THREAD.
 
+	if (Chaos::FPhysicsSolver* RigidSolver = GetSolver<Chaos::FPhysicsSolver>())
+	{
+		Chaos::FRigidClustering& RigidClustering = RigidSolver->GetEvolution()->GetRigidClustering();
+		Chaos::FClusterUnionManager& ClusterUnionManager = RigidClustering.GetClusterUnionManager();
+		for (Chaos::FPBDRigidClusteredParticleHandle* Handle : GetSolverParticleHandles())
+		{
+			if (Handle)
+			{
+				// If this particle is in an internal cluster that belongs to the geometry collection, the internal cluster needs to be disabled too.
+				if (Chaos::FPBDRigidClusteredParticleHandle* Parent = Handle->Parent())
+				{
+					if (Parent->PhysicsProxy() == this && !Parent->Disabled())
+					{
+						// We should be able to just disabled it - an internal cluster should never be considered in a cluster union.
+						RigidClustering.DisableCluster(Parent);
+					}
+				}
+		
+				// It should be safe to defer here without an immediate call to HandleDeferredClusterUnionUpdateProperties.
+				// This change doesn't really have to go through until FClusterUnionManager::FlushPendingOperations which happens
+				// prior to trying to advance the frame.
+				ClusterUnionManager.HandleRemoveOperationWithClusterLookup( { Handle }, Chaos::EClusterUnionOperationTiming::Defer);
+		
+				// Need to use FRigidClustering::DisableCluster instead of the evolution's to also handle the fact that the GC particle could be in the top level strained sets
+				// which would make it get considered in the breaking model which is undesirable.
+				RigidClustering.DisableCluster(Handle);
+		
+				// This ensures that this particle won't have any intercluster edges on it. This way no connectivity operations will try to add it into a cluster union.
+				RigidClustering.RemoveNodeConnections(Handle);
+			}
+		}
+	}
 }
 
 void FGeometryCollectionPhysicsProxy::BufferGameState() 
