reverting change to checking when to fire geometry collection OnFullyDecayed events since it instead prevented OnFullyDecayed from ever being fired

diff --git a/Engine/Source/Runtime/Experimental/GeometryCollectionEngine/Private/GeometryCollection/GeometryCollectionComponent.cpp b/Engine/Source/Runtime/Experimental/GeometryCollectionEngine/Private/GeometryCollection/GeometryCollectionComponent.cpp
--- a/Engine/Source/Runtime/Experimental/GeometryCollectionEngine/Private/GeometryCollection/GeometryCollectionComponent.cpp
+++ b/Engine/Source/Runtime/Experimental/GeometryCollectionEngine/Private/GeometryCollection/GeometryCollectionComponent.cpp
@@ -3309,9 +3309,8 @@ void UGeometryCollectionComponent::CheckFullyDecayed()
 	if (DynamicCollection && PhysicsProxy)
 	{
 		FGeometryCollectionDynamicStateFacade DynamicStateFacade(*DynamicCollection);
-		FGeometryCollectionDecayDynamicFacade DecayFacade(*DynamicCollection);
 
-		if (DynamicStateFacade.IsValid() && DecayFacade.IsValid())
+		if (DynamicStateFacade.IsValid())
 		{
 			bool bFullyDecayed = true;
 			const int32 NumTransforms = DynamicCollection->NumElements(FGeometryCollection::TransformGroup);
@@ -3337,12 +3336,6 @@ void UGeometryCollectionComponent::CheckFullyDecayed()
 					bFullyDecayed = false;
 					break;
 				}
-
-				if (DecayFacade.GetDecay(TransformIdx) < 1.0f)
-				{
-					bFullyDecayed = false;
-					break;
-				}
 			}
 
 			if (bFullyDecayed)
