Chaos debug draw : add mass display for kinematic and dynamic particles
(p.Chaos.Solver.DebugDrawMass)

diff --git a/Engine/Source/Runtime/Experimental/Chaos/Private/Chaos/ChaosDebugDraw.cpp b/Engine/Source/Runtime/Experimental/Chaos/Private/Chaos/ChaosDebugDraw.cpp
--- a/Engine/Source/Runtime/Experimental/Chaos/Private/Chaos/ChaosDebugDraw.cpp
+++ b/Engine/Source/Runtime/Experimental/Chaos/Private/Chaos/ChaosDebugDraw.cpp
@@ -1341,6 +1341,36 @@ namespace Chaos
 			}
 		}
 
+		void DrawParticleMassImpl(const FRigidTransform3& SpaceTransform, const FGeometryParticleHandle* InParticle, const FChaosDebugDrawSettings& Settings)
+		{
+			if (const auto RigidParticle = InParticle->CastToRigidParticle())
+			{
+				const FColor Color = RigidParticle->IsDynamic()? FColor::Yellow: FColor::Cyan;
+
+				const FVec3 PCOM = SpaceTransform.TransformPosition(FParticleUtilities::GetCoMWorldPosition(RigidParticle));
+
+				const FReal Mass = RigidParticle->M();
+				if (Mass > 0 && !RigidParticle->Disabled())
+				{
+					auto FormatString = TEXT("{0} kg");
+					FReal FormatMultiplier = 1.0;
+					if (Mass < 1.0)
+					{
+						FormatString = TEXT("{0} g");
+						FormatMultiplier = 1000.0;
+					}
+					else if (Mass > 1000.0)
+					{
+						FormatString = TEXT("{0} t");
+						FormatMultiplier = 0.001;
+					}
+					const FReal RoundedMass = 0.1 * FMath::RoundToDouble(Mass * FormatMultiplier * 10.0);
+					const FString MassStr = FString::Format(FormatString, { FString::SanitizeFloat(RoundedMass, 1) });
+					FDebugDrawQueue::GetInstance().DrawDebugString(PCOM, MassStr, nullptr, Color, UE_KINDA_SMALL_NUMBER, false, Settings.FontScale);
+				}
+			}
+		}
+
 		void DrawJointConstraintImpl(const FRigidTransform3& SpaceTransform, const FVec3& InPa, const FVec3& InCa, const FVec3& InXa, const FMatrix33& Ra, const FVec3& InPb, const FVec3& InCb, const FVec3& InXb, const FMatrix33& Rb, Chaos::FRealSingle ColorScale, const FChaosDebugDrawJointFeatures& FeatureMask, const FChaosDebugDrawSettings& Settings)
 		{
 			using namespace Chaos::DebugDraw;
@@ -2012,6 +2042,28 @@ namespace Chaos
 			}
 		}
 
+		void DrawParticleMass(const FRigidTransform3& SpaceTransform, const TParticleView<FKinematicGeometryParticles>& ParticlesView, const FChaosDebugDrawSettings* Settings)
+		{
+			if (FDebugDrawQueue::IsDebugDrawingEnabled())
+			{
+				for (auto& Particle : ParticlesView)
+				{
+					DrawParticleMassImpl(SpaceTransform, GetHandleHelper(&Particle), GetChaosDebugDrawSettings(Settings));
+				}
+			}
+		}
+
+		void DrawParticleMass(const FRigidTransform3& SpaceTransform, const TParticleView<FPBDRigidParticles>& ParticlesView, const FChaosDebugDrawSettings* Settings)
+		{
+			if (FDebugDrawQueue::IsDebugDrawingEnabled())
+			{
+				for (auto& Particle : ParticlesView)
+				{
+					DrawParticleMassImpl(SpaceTransform, GetHandleHelper(&Particle), GetChaosDebugDrawSettings(Settings));
+				}
+			}
+		}
+
 		void DrawJointConstraints(const FRigidTransform3& SpaceTransform, const TArray<FPBDJointConstraintHandle*>& ConstraintHandles, Chaos::FRealSingle ColorScale, const FChaosDebugDrawJointFeatures& FeatureMask, const FChaosDebugDrawSettings* Settings)
 		{
 			if (FDebugDrawQueue::IsDebugDrawingEnabled())
diff --git a/Engine/Source/Runtime/Experimental/Chaos/Private/PBDRigidsSolver.cpp b/Engine/Source/Runtime/Experimental/Chaos/Private/PBDRigidsSolver.cpp
--- a/Engine/Source/Runtime/Experimental/Chaos/Private/PBDRigidsSolver.cpp
+++ b/Engine/Source/Runtime/Experimental/Chaos/Private/PBDRigidsSolver.cpp
@@ -111,6 +111,7 @@ namespace Chaos
 	namespace CVars
 	{
 		int32 ChaosSolverDebugDrawShapes = CHAOS_SOLVER_ENABLE_DEBUG_DRAW;
+		int32 ChaosSolverDebugDrawMass = 0;
 		int32 ChaosSolverDebugDrawBVHs = 0;
 		int32 ChaosSolverDebugDrawCollisions = CHAOS_SOLVER_ENABLE_DEBUG_DRAW;
 		int32 ChaosSolverDebugDrawCollidingShapes = 0;
@@ -136,6 +137,7 @@ namespace Chaos
 		int32 ChaosSolverDebugDrawShowClient = 1;
 		DebugDraw::FChaosDebugDrawJointFeatures ChaosSolverDrawJointFeatures = DebugDraw::FChaosDebugDrawJointFeatures::MakeDefault();
 		FAutoConsoleVariableRef CVarChaosSolverDrawShapes(TEXT("p.Chaos.Solver.DebugDrawShapes"), ChaosSolverDebugDrawShapes, TEXT("Draw Shapes (0 = never; 1 = end of frame)."));
+		FAutoConsoleVariableRef CVarChaosSolverDebugDrawMass(TEXT("p.Chaos.Solver.DebugDrawMass"), ChaosSolverDebugDrawMass, TEXT("Draw Mass values in Kg (0 = never; 1 = end of frame)."));
 		FAutoConsoleVariableRef CVarChaosSolverDrawBVHs(TEXT("p.Chaos.Solver.DebugDrawBVHs"), ChaosSolverDebugDrawBVHs, TEXT("Draw Particle BVHs where applicable (0 = never; 1 = end of frame)."));
 		FAutoConsoleVariableRef CVarChaosSolverDrawCollisions(TEXT("p.Chaos.Solver.DebugDrawCollisions"), ChaosSolverDebugDrawCollisions, TEXT("Draw Collisions (0 = never; 1 = end of frame)."));
 		FAutoConsoleVariableRef CVarChaosSolverDrawCollidingShapes(TEXT("p.Chaos.Solver.DebugDrawCollidingShapes"), ChaosSolverDebugDrawCollidingShapes, TEXT("Draw Shapes that have collisions on them (0 = never; 1 = end of frame)."));
@@ -2377,7 +2379,6 @@ TRACE_COUNTER_SET(ChaosTraceCounter_##Name, Value)
 			ChaosSolverDebugDebugDrawSettings.ShapesColorsPerState = DebugDraw::GetDefaultShapesColorsByState();
 		}
 #endif
-
 		if (ChaosSolverDebugDrawShapes == 1)
 		{
 			if (ChaosSolverDrawShapesShowStatic)
@@ -2393,6 +2394,11 @@ TRACE_COUNTER_SET(ChaosTraceCounter_##Name, Value)
 				DebugDraw::DrawParticleShapes(FRigidTransform3(), Particles.GetNonDisabledDynamicView(), 1.0f, &ChaosSolverDebugDebugDrawSettings);
 			}
 		}
+		if (ChaosSolverDebugDrawMass == 1)
+		{
+			DebugDraw::DrawParticleMass(FRigidTransform3(), Particles.GetActiveKinematicParticlesView(), &ChaosSolverDebugDebugDrawSettings);
+			DebugDraw::DrawParticleMass(FRigidTransform3(), Particles.GetNonDisabledDynamicView(), &ChaosSolverDebugDebugDrawSettings);
+		}
 		if (ChaosSolverDebugDrawBVHs == 1)
 		{
 			for (const FTransientGeometryParticleHandle& Particle : Particles.GetNonDisabledView())
diff --git a/Engine/Source/Runtime/Experimental/Chaos/Public/Chaos/ChaosDebugDraw.h b/Engine/Source/Runtime/Experimental/Chaos/Public/Chaos/ChaosDebugDraw.h
--- a/Engine/Source/Runtime/Experimental/Chaos/Public/Chaos/ChaosDebugDraw.h
+++ b/Engine/Source/Runtime/Experimental/Chaos/Public/Chaos/ChaosDebugDraw.h
@@ -217,6 +217,8 @@ namespace Chaos
 		CHAOS_API void DrawCollision(const FRigidTransform3& SpaceTransform, const FPBDCollisionConstraint& Collision, FRealSingle ColorScale, const FChaosDebugDrawSettings* Settings = nullptr);
 		CHAOS_API void DrawCollisions(const FRigidTransform3& SpaceTransform, const FPBDCollisionConstraints& Collisions, FRealSingle ColorScale, const FChaosDebugDrawSettings* Settings = nullptr);
 		CHAOS_API void DrawCollisions(const FRigidTransform3& SpaceTransform, const Private::FCollisionConstraintAllocator& CollisionAllocator, FRealSingle ColorScale, const FChaosDebugDrawSettings* Settings = nullptr);
+		CHAOS_API void DrawParticleMass(const FRigidTransform3& SpaceTransform, const TParticleView<FKinematicGeometryParticles>& ParticlesView, const FChaosDebugDrawSettings* Settings = nullptr);
+		CHAOS_API void DrawParticleMass(const FRigidTransform3& SpaceTransform, const TParticleView<FPBDRigidParticles>& ParticlesView, const FChaosDebugDrawSettings* Settings = nullptr);
 		CHAOS_API void DrawJointConstraints(const FRigidTransform3& SpaceTransform, const TArray<FPBDJointConstraintHandle*>& ConstraintHandles, FRealSingle ColorScale, const FChaosDebugDrawJointFeatures& FeatureMask = FChaosDebugDrawJointFeatures::MakeDefault(), const FChaosDebugDrawSettings* Settings = nullptr);
 		CHAOS_API void DrawJointConstraints(const FRigidTransform3& SpaceTransform, const FPBDJointConstraints& Constraints, FRealSingle ColorScale, const FChaosDebugDrawJointFeatures& FeatureMask = FChaosDebugDrawJointFeatures::MakeDefault(), const FChaosDebugDrawSettings* Settings = nullptr);
 		CHAOS_API void DrawCharacterGroundConstraints(const FRigidTransform3& SpaceTransform, const FCharacterGroundConstraintContainer& Constraints, const FChaosDebugDrawSettings* Settings = nullptr);
