Add option to make a GC component to follow it's root bone ( when dynamic )
- this does not work for kinematically moved GC in a cluster union but this should provide a good support for all the other cases

diff --git a/Engine/Source/Runtime/Experimental/Chaos/Private/PhysicsProxy/GeometryCollectionPhysicsProxy.cpp b/Engine/Source/Runtime/Experimental/Chaos/Private/PhysicsProxy/GeometryCollectionPhysicsProxy.cpp
--- a/Engine/Source/Runtime/Experimental/Chaos/Private/PhysicsProxy/GeometryCollectionPhysicsProxy.cpp
+++ b/Engine/Source/Runtime/Experimental/Chaos/Private/PhysicsProxy/GeometryCollectionPhysicsProxy.cpp
@@ -3168,7 +3168,7 @@ void FGeometryCollectionPhysicsProxy::PushToPhysicsState()
 
 						SetClusteredParticleKinematicTarget_Internal(KinematicRootHandle, WorldTransform);
 					}
-					else if (ParentHandle)
+					else if (ParentHandle && !ParentHandle->IsDynamic())
 					{
 						if (ClusterUnionIndex == INDEX_NONE)
 						{
diff --git a/Engine/Source/Runtime/Experimental/GeometryCollectionEngine/Private/GeometryCollection/GeometryCollectionComponent.cpp b/Engine/Source/Runtime/Experimental/GeometryCollectionEngine/Private/GeometryCollection/GeometryCollectionComponent.cpp
--- a/Engine/Source/Runtime/Experimental/GeometryCollectionEngine/Private/GeometryCollection/GeometryCollectionComponent.cpp
+++ b/Engine/Source/Runtime/Experimental/GeometryCollectionEngine/Private/GeometryCollection/GeometryCollectionComponent.cpp
@@ -566,6 +566,7 @@ UGeometryCollectionComponent::UGeometryCollectionComponent(const FObjectInitiali
 	, bGlobalCrumblingEventIncludesChildren(false)
 	, bStoreVelocities(false)
 	, bShowBoneColors(false)
+	, bUpdateComponentTransformToRootBone(false)
 	, bUseRootProxyForNavigation(false)
 	, bUpdateNavigationInTick(true) 
 #if WITH_EDITORONLY_DATA 
@@ -3091,7 +3092,8 @@ void UGeometryCollectionComponent::OnUpdateTransform(EUpdateTransformFlags Updat
 {
 	Super::OnUpdateTransform(UpdateTransformFlags, Teleport);
 
-	if (PhysicsProxy)
+	const bool bSkipPhysicsUpdate = ((UpdateTransformFlags & EUpdateTransformFlags::SkipPhysicsUpdate) == EUpdateTransformFlags::SkipPhysicsUpdate);
+	if (!bSkipPhysicsUpdate && PhysicsProxy)
 	{
 		PhysicsProxy->SetWorldTransform_External(GetComponentTransform());
 	}
@@ -3848,12 +3850,54 @@ void UGeometryCollectionComponent::OnPostPhysicsSync()
 
 		CheckFullyDecayed();
 	}
+	else if (bUpdateComponentTransformToRootBone)
+	{
+		MoveComponentToRootTransform();
+	}
 
 	const bool bDynamicDataIsDirty = (DynamicCollection && DynamicCollection->IsDirty() && HasVisibleGeometry());
 	UpdateRenderSystemsIfNeeded(bDynamicDataIsDirty);
 	UpdateNavigationDataIfNeeded(bDynamicDataIsDirty);
 }
 
+void UGeometryCollectionComponent::MoveComponentToRootTransform()
+{
+	if (RestCollection && PhysicsProxy)
+	{
+		const TSharedPtr<FGeometryCollection, ESPMode::ThreadSafe> AssetCollection = RestCollection->GetGeometryCollection();
+
+		if (AssetCollection && DynamicCollection && DynamicCollection->IsDirty())
+		{
+			static FName MassToLocalAttributeName = "MassToLocal";
+			const TManagedArray<FTransform>& MassToLocalAttribute = RestCollection->GetGeometryCollection()->GetAttribute<FTransform>(MassToLocalAttributeName, FGeometryCollection::TransformGroup);
+
+			FGeometryCollectionDynamicStateFacade DynamicStateFacade(*DynamicCollection);
+
+			const int32 RootIndex = GetRootIndex();
+			const bool bIsRootActive = DynamicStateFacade.IsDynamicOrSleeping(RootIndex);
+			const bool bHasDynamicOPrClusterUnionParent = DynamicStateFacade.HasDynamicInternalClusterParent(RootIndex) || DynamicStateFacade.HasClusterUnionParent(RootIndex);
+			if (bIsRootActive || bHasDynamicOPrClusterUnionParent)
+			{
+				const FTransform& OriginalComponentSpaceRootTransformOffset = AssetCollection->Transform[RootIndex];
+				DynamicCollection->Transform[RootIndex] = FTransform3f(OriginalComponentSpaceRootTransformOffset);
+				const Chaos::FPBDRigidParticle* RootParticle = PhysicsProxy->GetExternalParticles()[RootIndex].Get();
+				const FTransform ParticleWorldPosition(RootParticle->R(), RootParticle->X());
+				const FTransform MassToLocal = MassToLocalAttribute[RootIndex];
+				const FTransform NewRootWorldPosition = MassToLocal.Inverse() * ParticleWorldPosition;
+
+				const FTransform CurrentComponentTransform = GetComponentTransform();
+				if (!NewRootWorldPosition.EqualsNoScale(CurrentComponentTransform))
+				{
+					const FVector NewPosition = NewRootWorldPosition.GetLocation();
+					const FVector MoveBy = NewPosition - CurrentComponentTransform.GetLocation();
+					const FRotator NewRotation = NewRootWorldPosition.Rotator();
+					MoveComponent(MoveBy, NewRotation, /* bSweep */ false,/* Hit */ NULL, MOVECOMP_SkipPhysicsMove);
+				}
+			}
+		}
+	}
+}
+
 void UGeometryCollectionComponent::UpdateRenderSystemsIfNeeded(bool bDynamicCollectionDirty)
 {
 #if WITH_EDITOR
diff --git a/Engine/Source/Runtime/Experimental/GeometryCollectionEngine/Public/GeometryCollection/GeometryCollectionComponent.h b/Engine/Source/Runtime/Experimental/GeometryCollectionEngine/Public/GeometryCollection/GeometryCollectionComponent.h
--- a/Engine/Source/Runtime/Experimental/GeometryCollectionEngine/Public/GeometryCollection/GeometryCollectionComponent.h
+++ b/Engine/Source/Runtime/Experimental/GeometryCollectionEngine/Public/GeometryCollection/GeometryCollectionComponent.h
@@ -1342,6 +1342,14 @@ protected:
 	UPROPERTY(EditAnywhere, BlueprintReadWrite, Category = "ChaosPhysics|General")
 	bool bShowBoneColors;
 
+	/** 
+	* Relocate the component so that the original offset to the root bone is maintained
+	* This only works when the root bone is moving whole being dynamically simulated 
+	* Note: Once the root element is broken, the component will no longer update its position
+	*/
+	UPROPERTY(EditAnywhere, BlueprintReadWrite, Category = "ChaosPhysics|General")
+	bool bUpdateComponentTransformToRootBone;
+
 	UPROPERTY(EditAnywhere, BlueprintReadWrite, Category = Collision, AdvancedDisplay, config)
 	bool bUseRootProxyForNavigation;
 
@@ -1664,6 +1672,8 @@ private:
 
 	float ComputeMassScaleRelativeToAsset() const;
 
+	void MoveComponentToRootTransform();
+
 	/** The clusters we need to replicate */
 	TUniquePtr<TSet<Chaos::FPBDRigidClusteredParticleHandle*>> ClustersToRep;
 
