Added cvar to adjust IAS' receive buffer size

#rnx

(git:d4913eb - Martin Ridgers - [REDACTED_PHONE]:30:06 +0200)

diff --git a/Engine/Source/Runtime/Experimental/IoStoreOnDemand/Private/OnDemandHttpClient.cpp b/Engine/Source/Runtime/Experimental/IoStoreOnDemand/Private/OnDemandHttpClient.cpp
--- a/Engine/Source/Runtime/Experimental/IoStoreOnDemand/Private/OnDemandHttpClient.cpp
+++ b/Engine/Source/Runtime/Experimental/IoStoreOnDemand/Private/OnDemandHttpClient.cpp
@@ -25,6 +25,13 @@ static FAutoConsoleVariableRef CVar_GIasHttpPollTimeoutMs(
 	TEXT("Http tick poll timeout in milliseconds")
 );
 
+static int32 GIasHttpRecvBufKiB = -1;
+static FAutoConsoleVariableRef CVar_GIasHttpRecvBufKiB(
+	TEXT("ias.HttpRecvBufKiB"),
+	GIasHttpRecvBufKiB,
+	TEXT("Recv buffer size")
+);
+
 static void LogHttpResult(const TCHAR* Url, uint32 StatusCode, uint64 DurationMs, uint64 Size, uint64 Offset, const char* Memo = "ok")
 {
 	Size >>= 10;
@@ -43,6 +50,11 @@ FOnDemandHttpClient::FOnDemandHttpClient(const FString& ServiceUrl, int32 MaxCon
 		UE_LOG(LogIas, Error, TEXT("Failed to set host from '%s'"), *ServiceUrl);
 	}
 
+	if (GIasHttpRecvBufKiB >= 0)
+	{
+		UE_LOG(LogIas, VeryVerbose, TEXT("Set HTTP recv buffer size to %dKib"), GIasHttpRecvBufKiB);
+		Params.RecvBufSize = GIasHttpRecvBufKiB << 10;
+	}
 	Params.ConnectionCount = MaxConnectionCount;
 	ConnectionPool = MakeUnique<HTTP::FConnectionPool>(Params);
 }
