Changed IAS backend CPU scopes to use "IasBackend" prefix

#rnx

(git:2b1a510 - Martin Ridgers - [REDACTED_PHONE]:29:34 +0200)

diff --git a/Engine/Source/Runtime/Experimental/IoStoreOnDemand/Private/OnDemandIoDispatcherBackend.cpp b/Engine/Source/Runtime/Experimental/IoStoreOnDemand/Private/OnDemandIoDispatcherBackend.cpp
--- a/Engine/Source/Runtime/Experimental/IoStoreOnDemand/Private/OnDemandIoDispatcherBackend.cpp
+++ b/Engine/Source/Runtime/Experimental/IoStoreOnDemand/Private/OnDemandIoDispatcherBackend.cpp
@@ -307,7 +307,7 @@ void FOnDemandIoStore::AddToc(const FString& TocPath, FOnDemandToc&& Toc)
 {
 	UE_LOG(LogIas, Log, TEXT("Adding TOC '%s'"), *TocPath);
 
-	TRACE_CPUPROFILER_EVENT_SCOPE(FOnDemandIoBackend::AddToc);
+	TRACE_CPUPROFILER_EVENT_SCOPE(IasBackend::AddToc);
 
 	FString Prefix;
 	{
@@ -451,7 +451,7 @@ void FOnDemandIoStore::AddDeferredContainers()
 void FOnDemandIoStore::OnEncryptionKeyAdded(const FGuid& Id, const FAES::FAESKey& Key)
 {
 	LLM_SCOPE_BYTAG(Ias);
-	TRACE_CPUPROFILER_EVENT_SCOPE(FOnDemandIoBackend::OnEncryptionKeyAdded);
+	TRACE_CPUPROFILER_EVENT_SCOPE(IasBackend::OnEncryptionKeyAdded);
 	AddDeferredContainers();
 }
 
@@ -535,7 +535,7 @@ private:
 		else
 		{
 			// NOTE: This can get expensive if the queue gets too long, might be better to have x number of bucket(s)
-			TRACE_CPUPROFILER_EVENT_SCOPE(FOnDemandIoBackend::EnqueueByPriority);
+			TRACE_CPUPROFILER_EVENT_SCOPE(IasBackend::EnqueueByPriority);
 			T* It = Head;
 			while (It->NextRequest != nullptr && Request->Priority <= It->NextRequest->Priority)
 			{
@@ -847,7 +847,7 @@ static TArray<FString> FindOnDemandUtocFilesOnDisk()
 /** Generate a FOnDemandToc based on utoc files on disk which support the OnDemand feature */
 TIoStatusOr<FOnDemandToc> GenerateOnDemandTocFromDisk()
 {
-	TRACE_CPUPROFILER_EVENT_SCOPE(GenerateOnDemandTocFromDisk);
+	TRACE_CPUPROFILER_EVENT_SCOPE(IasBackend::GenerateOnDemandTocFromDisk);
 
 	FOnDemandToc OutToc;
 
@@ -1402,7 +1402,7 @@ FOnDemandIoBackend::~FOnDemandIoBackend()
 
 void FOnDemandIoBackend::Initialize(TSharedRef<const FIoDispatcherBackendContext> Context)
 {
-	TRACE_CPUPROFILER_EVENT_SCOPE(FOnDemandIoBackend::Initialize);
+	TRACE_CPUPROFILER_EVENT_SCOPE(IasBackend::Initialize);
 	LLM_SCOPE_BYTAG(Ias);
 	UE_LOG(LogIas, Log, TEXT("Initializing on demand I/O dispatcher backend"));
 	BackendContext = Context;
@@ -1485,7 +1485,7 @@ void FOnDemandIoBackend::ConditionallyStartBackendThread()
 void FOnDemandIoBackend::CompleteRequest(FChunkRequest* ChunkRequest)
 {
 	LLM_SCOPE_BYTAG(Ias);
-	TRACE_CPUPROFILER_EVENT_SCOPE(FOnDemandIoBackend::CompleteRequest);
+	TRACE_CPUPROFILER_EVENT_SCOPE(IasBackend::CompleteRequest);
 	check(ChunkRequest != nullptr);
 	const bool bCancelled = ChunkRequest->CancellationToken.IsCancelled();
 
@@ -1513,7 +1513,7 @@ void FOnDemandIoBackend::CompleteRequest(FChunkRequest* ChunkRequest)
 		bool bDecoded = false;
 		if (Chunk.GetSize() > 0)
 		{
-			TRACE_CPUPROFILER_EVENT_SCOPE(FOnDemandIoBackend::DecodeBlocks);
+			TRACE_CPUPROFILER_EVENT_SCOPE(IasBackend::DecodeBlocks);
 			const uint64 RawSize = FMath::Min<uint64>(Request->Options.GetSize(), ChunkRequest->Params.ChunkInfo.Entry->RawSize);
 			Request->CreateBuffer(RawSize);
 			DecodingParams.RawOffset = Request->Options.GetOffset(); 
@@ -1589,7 +1589,7 @@ bool FOnDemandIoBackend::Resolve(FIoRequestImpl* Request)
 	{
 		if (bUpdatePriority)
 		{
-			TRACE_CPUPROFILER_EVENT_SCOPE(FOnDemandIoBackend::UpdatePriorityForIoRequest);
+			TRACE_CPUPROFILER_EVENT_SCOPE(IasBackend::UpdatePriorityForIoRequest);
 			HttpRequests.Reprioritize(ChunkRequest);
 		}
 		// The chunk for the request is already inflight 
@@ -1667,7 +1667,7 @@ void FOnDemandIoBackend::CancelIoRequest(FIoRequestImpl* Request)
 
 void FOnDemandIoBackend::UpdatePriorityForIoRequest(FIoRequestImpl* Request)
 {
-	TRACE_CPUPROFILER_EVENT_SCOPE(FOnDemandIoBackend::UpdatePriorityForIoRequest);
+	TRACE_CPUPROFILER_EVENT_SCOPE(IasBackend::UpdatePriorityForIoRequest);
 	if (FChunkRequest* ChunkRequest = ChunkRequests.TryUpdatePriority(Request))
 	{
 		HttpRequests.Reprioritize(ChunkRequest);
@@ -1749,7 +1749,7 @@ TIoStatusOr<FIoMappedRegion> FOnDemandIoBackend::OpenMapped(const FIoChunkId& Ch
 
 FIoStatus FOnDemandIoBackend::ApplyGeneratedOnDemandToc(const FString& CdnUrl, const FString& TocPath)
 {
-	TRACE_CPUPROFILER_EVENT_SCOPE(FOnDemandIoBackend::ApplyGeneratedOnDemandToc);
+	TRACE_CPUPROFILER_EVENT_SCOPE(IasBackend::ApplyGeneratedOnDemandToc);
 
 	if (bGeneratedOnDemandToc || !CVar_IasGenerateOnDemandToc->GetBool())
 	{
@@ -1790,7 +1790,7 @@ FIoStatus FOnDemandIoBackend::ApplyGeneratedOnDemandToc(const FString& CdnUrl, c
 
 FIoStatus FOnDemandIoBackend::DownloadoadOnDemandToc(const FString& CdnUrl, const FString& TocPath)
 {
-	TRACE_CPUPROFILER_EVENT_SCOPE(FOnDemandIoBackend::LoadOnDemandToc);
+	TRACE_CPUPROFILER_EVENT_SCOPE(IasBackend::LoadOnDemandToc);
 
 	UE_LOG(LogIas, Log, TEXT("Downloading OnDemandToc from CDN"));
 
@@ -1850,7 +1850,7 @@ void FOnDemandIoBackend::AddDeferredTocs()
 void FOnDemandIoBackend::Mount(const FOnDemandEndpoint& Endpoint)
 {
 	LLM_SCOPE_BYTAG(Ias);
-	TRACE_CPUPROFILER_EVENT_SCOPE(FOnDemandIoBackend::Mount);
+	TRACE_CPUPROFILER_EVENT_SCOPE(IasBackend::Mount);
 
 	if ((Endpoint.DistributionUrl.IsEmpty() && Endpoint.ServiceUrl.IsEmpty()) || Endpoint.TocPath.IsEmpty())
 	{
@@ -1966,7 +1966,7 @@ void FOnDemandIoBackend::ProcessHttpRequests(FOnDemandHttpClient* HttpClient, FB
 		while (NextChunkRequest)
 		{
 			{
-				TRACE_CPUPROFILER_EVENT_SCOPE(FOnDemandIoBackend::IssueHttpGet);
+				TRACE_CPUPROFILER_EVENT_SCOPE(IasBackend::IssueHttpGet);
 				FChunkRequest* ChunkRequest = NextChunkRequest;
 				NextChunkRequest = ChunkRequest->NextRequest;
 				ChunkRequest->NextRequest = nullptr;
@@ -2036,7 +2036,7 @@ void FOnDemandIoBackend::ProcessHttpRequests(FOnDemandHttpClient* HttpClient, FB
 
 			if (NumConcurrentRequests >= MaxConcurrentRequests)
 			{
-				TRACE_CPUPROFILER_EVENT_SCOPE(FOnDemandIoBackend::TickHttpSaturated);
+				TRACE_CPUPROFILER_EVENT_SCOPE(IasBackend::TickHttpSaturated);
 				while (NumConcurrentRequests >= MaxConcurrentRequests)
 				{
 					HttpClient->Tick(/*Block*/true);
@@ -2051,7 +2051,7 @@ void FOnDemandIoBackend::ProcessHttpRequests(FOnDemandHttpClient* HttpClient, FB
 
 		{
 			// Keep processing pending connections until all requests are completed or a new one is issued
-			TRACE_CPUPROFILER_EVENT_SCOPE(FOnDemandIoBackend::TickHttp);
+			TRACE_CPUPROFILER_EVENT_SCOPE(IasBackend::TickHttp);
 			while (HttpClient->Tick(/*Block*/false))
 			{
 				if (!NextChunkRequest)
