Wrap activity socket in an object. First step in separating activity concerns

#rnx

(git:c9df267 - Martin Ridgers - [REDACTED_PHONE]:49:18 +0200)

diff --git a/Engine/Source/Runtime/Experimental/IoStoreOnDemand/Private/CoreHttp/Http.cpp b/Engine/Source/Runtime/Experimental/IoStoreOnDemand/Private/CoreHttp/Http.cpp
--- a/Engine/Source/Runtime/Experimental/IoStoreOnDemand/Private/CoreHttp/Http.cpp
+++ b/Engine/Source/Runtime/Experimental/IoStoreOnDemand/Private/CoreHttp/Http.cpp
@@ -735,6 +735,100 @@ FMessageBuilder& FMessageBuilder::operator << (FAnsiStringView Lhs)
 
 
 
+// {{{1 socket .................................................................
+
+////////////////////////////////////////////////////////////////////////////////
+class FSocket
+{
+public:
+	SocketType	Get() const { return Socket; } // to be removed
+
+				FSocket() = default;
+				~FSocket()					{ if (IsValid()) Destroy(); }
+				FSocket(FSocket&& Rhs)		{ Move(MoveTemp(Rhs)); }
+	FSocket&	operator = (FSocket&& Rhs)	{ Move(MoveTemp(Rhs)); return *this; }
+	bool		IsValid() const				{ return Socket != InvalidSocket; }
+	bool		Create();
+	void		Destroy();
+	bool		SetBlocking(bool bBlocking);
+	bool		SetSendBufSize(int32 Size);
+	bool		SetRecvBufSize(int32 Size);
+
+private:
+	void		Move(FSocket&& Rhs);
+	SocketType	Socket = InvalidSocket;
+
+				FSocket(const FSocket&) = delete;
+	FSocket&	operator = (const FSocket&) = delete;
+};
+
+////////////////////////////////////////////////////////////////////////////////
+void FSocket::Move(FSocket&& Rhs)
+{
+	check(!IsValid() || !Rhs.IsValid()); // currently we only want to pass one around
+	Swap(Socket, Rhs.Socket);
+}
+
+////////////////////////////////////////////////////////////////////////////////
+bool FSocket::Create()
+{
+	check(!IsValid());
+	Socket = socket(AF_INET, SOCK_STREAM, IPPROTO_TCP);
+	return IsValid();
+}
+
+////////////////////////////////////////////////////////////////////////////////
+void FSocket::Destroy()
+{
+	if (Socket == InvalidSocket)
+	{
+		return;
+	}
+
+	closesocket(Socket);
+	Socket = InvalidSocket;
+}
+
+////////////////////////////////////////////////////////////////////////////////
+bool FSocket::SetBlocking(bool bBlocking)
+{
+	bool bSuccess = false;
+
+#if PLATFORM_MICROSOFT
+	unsigned long NonBlockingMode = 1;
+	if (ioctlsocket(Socket, FIONBIO, &NonBlockingMode) != SOCKET_ERROR)
+	{
+		bSuccess = true;
+	}
+#else
+	int32 Flags = fcntl(Socket, F_GETFL, 0);
+	if (Flags != -1)
+	{
+		Flags |= Flags | int32(O_NONBLOCK);
+		if (fcntl(Socket, F_SETFL, Flags) >= 0)
+		{
+			bSuccess = true;
+		}
+	}
+#endif
+
+	return bSuccess;
+}
+
+////////////////////////////////////////////////////////////////////////////////
+bool FSocket::SetSendBufSize(int32 Size)
+{
+	return 0 == setsockopt(Socket, SOL_SOCKET, SO_SNDBUF, &(char&)Size, sizeof(Size));
+}
+
+////////////////////////////////////////////////////////////////////////////////
+bool FSocket::SetRecvBufSize(int32 Size)
+{
+	return 0 == setsockopt(Socket, SOL_SOCKET, SO_RCVBUF, &(char&)Size, sizeof(Size));
+}
+
+
+
 // {{{1 connection-pool ........................................................
 
 ////////////////////////////////////////////////////////////////////////////////
@@ -746,8 +840,8 @@ public:
 					FSocketPool(const ANSICHAR* InHostName, uint32 InPort, uint32 InMaxLeases);
 					~FSocketPool();
 	static uint32	GetAllocSize(uint32 MaxLeases);
-	bool			LeaseSocket(SocketType& Out);
-	void			ReturnLease(SocketType Socket);
+	bool			LeaseSocket(FSocket& Out);
+	void			ReturnLease(FSocket&& Socket);
 	void			SetBufferSize(EDirection Dir, int32 Size);
 	int32			GetBufferSize(EDirection Dir) const;
 	int32			IsResolved() const;
@@ -764,7 +858,7 @@ private:
 	uint16			Port;
 	uint8			LeaseCount = 0;
 	uint8			MaxLeases;
-	SocketType		Sockets[1/*...N*/]; // this should be the last member
+	FSocket			Sockets[1/*...N*/]; // this should be the last member
 };
 
 ////////////////////////////////////////////////////////////////////////////////
@@ -777,7 +871,7 @@ FSocketPool::FSocketPool(const ANSICHAR* InHostName, uint32 InPort, uint32 InMax
 
 	for (uint32 i = 0; i < InMaxLeases; ++i)
 	{
-		Sockets[i] = InvalidSocket;
+		new (Sockets + i) FSocket();
 	}
 }
 
@@ -786,12 +880,7 @@ FSocketPool::~FSocketPool()
 {
 	for (uint32 i = 0; i < MaxLeases; ++i)
 	{
-		if (Sockets[i] == InvalidSocket)
-		{
-			continue;
-		}
-
-		closesocket(Sockets[i]);
+		Sockets[i].Destroy();
 	}
 }
 
@@ -802,7 +891,7 @@ uint32 FSocketPool::GetAllocSize(uint32 MaxLeases)
 }
 
 ////////////////////////////////////////////////////////////////////////////////
-bool FSocketPool::LeaseSocket(SocketType& Out)
+bool FSocketPool::LeaseSocket(FSocket& Out)
 {
 	check(LeaseCount <= MaxLeases);
 
@@ -811,18 +900,18 @@ bool FSocketPool::LeaseSocket(SocketType& Out)
 		return false;
 	}
 
-	Out = Sockets[LeaseCount];
+	Out = MoveTemp(Sockets[LeaseCount]);
 	++LeaseCount;
 
 	return true;
 }
 
 ////////////////////////////////////////////////////////////////////////////////
-void FSocketPool::ReturnLease(SocketType Socket)
+void FSocketPool::ReturnLease(FSocket&& Socket)
 {
 	check(LeaseCount > 0);
 	--LeaseCount;
-	Sockets[LeaseCount] = Socket;
+	Sockets[LeaseCount] = MoveTemp(Socket);
 }
 
 ////////////////////////////////////////////////////////////////////////////////
@@ -1026,7 +1115,7 @@ struct alignas(16) FActivity
 	UPTRINT				SinkParam;
 	FTicketSink			Sink;
 
-	SocketType			Socket = InvalidSocket;
+	FSocket				Socket;
 
 	FBuffer				Buffer;
 };
@@ -1068,16 +1157,15 @@ static void Activity_Free(FActivity* Activity)
 {
 	if (Activity->State > FActivity::EState::Connect)
 	{
-		SocketType Socket = Activity->Socket;
-		if (!Activity->IsKeepAlive && Socket != InvalidSocket)
+		FSocket& Socket = Activity->Socket;
+		if (!Activity->IsKeepAlive && Socket.IsValid())
 		{
-			closesocket(Socket);
-			Activity->Socket = InvalidSocket;
+			Socket = FSocket();
 		}
 
 		if (Activity->Pool->GetIpAddress() > 0x00ff'ffff)
 		{
-			Activity->Pool->ReturnLease(Activity->Socket);
+			Activity->Pool->ReturnLease(MoveTemp(Socket));
 		}
 	}
 
@@ -1459,7 +1547,7 @@ static uint64 ReadyCheck(FActivity** Activities, uint32 Num, uint32 TimeoutMs)
 			FSelect& Select = Selects[SelectNum];
 			auto SelectEvent = decltype(FSelect::events)((Activity->SocketWait == EWait::Read) ? POLLIN : POLLOUT);
 			Select = {
-				Activity->Socket,
+				Activity->Socket.Get(),
 				SelectEvent,
 				/* 0 */
 			};
@@ -1564,16 +1652,16 @@ static int32 DoConnect(FActivity* Activity)
 	}
 
 	// Claim an existing socket from the pool.
-	SocketType Candidate;
+	FSocket Candidate;
 	if (!Pool->LeaseSocket(Candidate))
 	{
 		// none available at this time
 		return 1;
 	}
 
-	if (Candidate != InvalidSocket)
+	if (Candidate.IsValid())
 	{
-		Activity->Socket = Candidate;
+		Activity->Socket = MoveTemp(Candidate);
 		Activity->SocketWait = FActivity::EWait::None;
 		Activity->State = FActivity::EState::Send;
 		Activity->StateParam = 0;
@@ -1582,51 +1670,28 @@ static int32 DoConnect(FActivity* Activity)
 	}
 
 	// Leased socket isn't valid so we'll create and connect one
-	Candidate = socket(AF_INET, SOCK_STREAM, IPPROTO_TCP);
-	if (Candidate == InvalidSocket)
+	if (!Candidate.Create())
 	{
 		Activity_SetError(Activity, "Failed to create socket");
 		return -1;
 	}
-	ON_SCOPE_EXIT { if (Candidate != InvalidSocket) closesocket(Candidate); };
 
-	// make the socket non-blocking
+	if (!Candidate.SetBlocking(false))
 	{
-		bool Success = false;
-
-#if PLATFORM_MICROSOFT
-		unsigned long NonBlockingMode = 1;
-		if (ioctlsocket(Candidate, FIONBIO, &NonBlockingMode) != SOCKET_ERROR)
-		{
-			Success = true;
-		}
-#else
-		int32 Flags = fcntl(Candidate, F_GETFL, 0);
-		if (Flags != -1)
-		{
-			Flags |= Flags | int32(O_NONBLOCK);
-			if (fcntl(Candidate, F_SETFL, Flags) >= 0)
-			{
-				Success = true;
-			}
-		}
-#endif
-
-		if (!Success)
-		{
-			Activity_SetError(Activity, "Unable to set socket non-blocking");
-			return -1;
-		}
+		Activity_SetError(Activity, "Unable to set socket non-blocking");
+		return -1;
 	}
 
+
 	// Adjust socket send and recv buffer sizes
 	if (int32 OptValue = Pool->GetBufferSize(FSocketPool::EDirection::Send); OptValue >= 0)
 	{
-		setsockopt(Candidate, SOL_SOCKET, SO_SNDBUF, &(char&)OptValue, sizeof(OptValue));
+		Candidate.SetSendBufSize(OptValue);
 	}
+
 	if (int32 OptValue = Pool->GetBufferSize(FSocketPool::EDirection::Recv); OptValue >= 0)
 	{
-		setsockopt(Candidate, SOL_SOCKET, SO_RCVBUF, &(char&)OptValue, sizeof(OptValue));
+		Candidate.SetRecvBufSize(OptValue);
 	}
 
 	uint16 Port = uint16(Pool->GetPort());
@@ -1646,7 +1711,7 @@ static int32 DoConnect(FActivity* Activity)
 	{
 		Trace(Activity, ETrace::Connect, IpAddress);
 
-		int Result = connect(Candidate, &(sockaddr&)AddrInet, sizeof(AddrInet));
+		int Result = connect(Candidate.Get(), &(sockaddr&)AddrInet, sizeof(AddrInet));
 		if (Result < 0 && !(IsSocketResult(EWOULDBLOCK) | IsSocketResult(EINPROGRESS)))
 		{
 			Activity_SetError(Activity, "Socket connect failed");
@@ -1654,9 +1719,8 @@ static int32 DoConnect(FActivity* Activity)
 		}
 	}
 
-	Activity->Socket = Candidate;
+	Activity->Socket = MoveTemp(Candidate);
 	Activity->SocketWait = FActivity::EWait::Write;
-	Candidate = InvalidSocket;
 
 	Activity->StateParam = 0;
 	Activity->State = (GSocksIpAddress != 0)
@@ -1686,7 +1750,7 @@ static int32 DoSocks4(FActivity* Activity)
 		uint32	IpAddress;
 	};
 
-	SocketType Socket = Activity->Socket;
+	SocketType Socket = Activity->Socket.Get();
 	int32 Result = 0;
 
 	FSocks4Request Request = {
@@ -1723,7 +1787,7 @@ static int32 DoSocks5(FActivity* Activity)
 #pragma warning(disable : 6385)
 #endif
 
-	SocketType Socket = Activity->Socket;
+	SocketType Socket = Activity->Socket.Get();
 	int32 Result;
 
 	// Greeting
@@ -1822,7 +1886,7 @@ static int32 DoSend(FActivity* Activity)
 	check(SendSize > 0);
 
 	Trace(Activity, ETrace::Send, SendSize);
-	int32 Result = send(Activity->Socket, SendData, SendSize, MsgFlagType(0));
+	int32 Result = send(Activity->Socket.Get(), SendData, SendSize, MsgFlagType(0));
 	Trace(Activity, ETrace::Send, -1);
 
 	if (Result < 0)
@@ -1831,7 +1895,7 @@ static int32 DoSend(FActivity* Activity)
 		{
 			int32 Error = 0;
 			socklen_t ErrorSize = sizeof(Error);
-			Result = getsockopt(Activity->Socket, SOL_SOCKET, SO_ERROR, (char*)&Error, &ErrorSize);
+			Result = getsockopt(Activity->Socket.Get(), SOL_SOCKET, SO_ERROR, (char*)&Error, &ErrorSize);
 			if (Result < 0 || Error != 0)
 			{
 				Activity_SetError(Activity, "Connection error");
@@ -1888,7 +1952,7 @@ static int32 DoRecvMessage(FActivity* Activity)
 		auto [Dest, DestSize] = Buffer.GetMutableFree(0, PageSize);
 
 		Trace(Activity, ETrace::Recv, -1);
-		int32 Result = recv(Activity->Socket, Dest, DestSize, MsgFlagType(0));
+		int32 Result = recv(Activity->Socket.Get(), Dest, DestSize, MsgFlagType(0));
 		Trace(Activity, ETrace::Recv, FMath::Max(Result, 0));
 
 		if (Result < 0)
@@ -2104,7 +2168,7 @@ static FHandlerResult DoRecvContent(FActivity* Activity, uint32 MaxRecvSize)
 
 		Trace(Activity, ETrace::Recv, -1);
 		char* Cursor = (char*)(DestView.GetData()) + Activity->StateParam;
-		int32 Result = recv(Activity->Socket, Cursor, Size, MsgFlagType(0));
+		int32 Result = recv(Activity->Socket.Get(), Cursor, Size, MsgFlagType(0));
 		Trace(Activity, ETrace::Recv, FMath::Max(Result, 0));
 
 		if (Result < 0)
