Submit IAS cache work in bulk instead of one at a time

#rnx

(git:60fcb90 - ridgers - [REDACTED_PHONE]:31:03 +0200)

diff --git a/Engine/Source/Runtime/Experimental/IoStoreOnDemand/Private/JournaledCache.cpp b/Engine/Source/Runtime/Experimental/IoStoreOnDemand/Private/JournaledCache.cpp
--- a/Engine/Source/Runtime/Experimental/IoStoreOnDemand/Private/JournaledCache.cpp
+++ b/Engine/Source/Runtime/Experimental/IoStoreOnDemand/Private/JournaledCache.cpp
@@ -1450,7 +1450,7 @@ private:
 	void						UpdateCache(FCache* Cache);
 	virtual uint32				Run() override;
 	virtual void				Stop() override;
-	void						SubmitWork(FWork& Work);
+	void						SubmitWork(FWork* Work, uint32 Num);
 	void						ReceiveWork();
 	TUniquePtr<FRunnableThread> Thread;
 	FEventRef					WakeEvent;
@@ -1478,10 +1478,10 @@ FServiceThread& FServiceThread::Get()
 }
 
 ////////////////////////////////////////////////////////////////////////////////
-void FServiceThread::SubmitWork(FWork& Work)
+void FServiceThread::SubmitWork(FWork* Work, uint32 Num)
 {
 	FScopeLock _(&Lock);
-	PendingWork.Add(MoveTemp(Work));
+	PendingWork.Append(Work, Num);
 	PendingCount.fetch_add(1, std::memory_order_relaxed);
 	WakeEvent->Trigger();
 }
@@ -1496,7 +1496,7 @@ void FServiceThread::RegisterCache(TUniquePtr<FCache> Cache)
 	FWork Work;
 	Work.What = FWork::Work_Register;
 	Work.SetCache(RawPtr);
-	SubmitWork(Work);
+	SubmitWork(&Work, 1);
 
 	if (PrevRunCount == 0)
 	{
@@ -1510,7 +1510,7 @@ void FServiceThread::UnregisterCache(FCache* Cache)
 	FWork Work;
 	Work.What = FWork::Work_Unregister;
 	Work.SetCache(Cache);
-	SubmitWork(Work);
+	SubmitWork(&Work, 1);
 
 	int32 PrevRunCount = RunCount.fetch_sub(1, std::memory_order_relaxed);
 	if (PrevRunCount == 1)
@@ -1530,7 +1530,7 @@ void FServiceThread::SetGovernorRate(uint32 Allowance, uint32 Ops, uint32 Second
 	Work.Allowance = Allowance;
 	Work.Ops = uint16(Ops);
 	Work.Seconds = uint16(Seconds);
-	SubmitWork(Work);
+	SubmitWork(&Work, 1);
 
 	check(Work.Ops == Ops && Work.Seconds == Seconds);
 }
@@ -1543,7 +1543,7 @@ void FServiceThread::SetGovernorDemand(uint32 Threshold, uint32 Boost, uint32 Su
 	Work.Threshold = uint16(Threshold);
 	Work.Boost = uint16(Boost);
 	Work.SuperBoost = uint16(SuperBoost);
-	SubmitWork(Work);
+	SubmitWork(&Work, 1);
 
 	check(Work.Threshold == Threshold && Work.Boost == Boost && Work.SuperBoost == SuperBoost);
 }
