Move lower-level connection handling out of activity

#rnx

(git:e4eba75 - Martin Ridgers - [REDACTED_PHONE]:54:18 +0200)

diff --git a/Engine/Source/Runtime/Experimental/IoStoreOnDemand/Private/CoreHttp/Http.cpp b/Engine/Source/Runtime/Experimental/IoStoreOnDemand/Private/CoreHttp/Http.cpp
--- a/Engine/Source/Runtime/Experimental/IoStoreOnDemand/Private/CoreHttp/Http.cpp
+++ b/Engine/Source/Runtime/Experimental/IoStoreOnDemand/Private/CoreHttp/Http.cpp
@@ -35,6 +35,8 @@
 	using SocketType	= SOCKET;
 	using MsgFlagType	= int;
 
+	enum { SHUT_RDWR = SD_BOTH };
+
 #	define IAS_HTTP_USE_POLL
 	template <typename... ArgTypes> auto poll(ArgTypes... Args)
 	{
@@ -750,6 +752,8 @@ public:
 	bool		IsValid() const				{ return Socket != InvalidSocket; }
 	bool		Create();
 	void		Destroy();
+	bool		Connect(uint32 Ip, uint32 Port);
+	void		Disconnect();
 	bool		SetBlocking(bool bBlocking);
 	bool		SetSendBufSize(int32 Size);
 	bool		SetRecvBufSize(int32 Size);
@@ -789,6 +793,33 @@ void FSocket::Destroy()
 	Socket = InvalidSocket;
 }
 
+////////////////////////////////////////////////////////////////////////////////
+bool FSocket::Connect(uint32 IpAddress, uint32 Port)
+{
+	check(IsValid());
+
+	sockaddr_in AddrInet = { sizeof(sockaddr_in) };
+	AddrInet.sin_family = AF_INET;
+	AddrInet.sin_port = htons(uint16(Port));
+	memcpy(&(AddrInet.sin_addr), &IpAddress, sizeof(IpAddress));
+
+	int Result = connect(Socket, &(sockaddr&)AddrInet, sizeof(AddrInet));
+
+	if (IsSocketResult(EWOULDBLOCK) | IsSocketResult(EINPROGRESS))
+	{
+		return true;
+	}
+
+	return (Result >= 0);
+}
+
+////////////////////////////////////////////////////////////////////////////////
+void FSocket::Disconnect()
+{
+	check(IsValid());
+	shutdown(Socket, SHUT_RDWR);
+}
+
 ////////////////////////////////////////////////////////////////////////////////
 bool FSocket::SetBlocking(bool bBlocking)
 {
@@ -1704,19 +1735,11 @@ static int32 DoConnect(FActivity* Activity)
 	}
 
 	// connect
-	sockaddr_in AddrInet = { sizeof(sockaddr_in) };
-	AddrInet.sin_family = AF_INET;
-	AddrInet.sin_port = htons(Port);
-	memcpy(&(AddrInet.sin_addr), &IpAddress, sizeof(IpAddress));
+	Trace(Activity, ETrace::Connect, IpAddress);
+	if (Candidate.Connect(IpAddress, Port))
 	{
-		Trace(Activity, ETrace::Connect, IpAddress);
-
-		int Result = connect(Candidate.Get(), &(sockaddr&)AddrInet, sizeof(AddrInet));
-		if (Result < 0 && !(IsSocketResult(EWOULDBLOCK) | IsSocketResult(EINPROGRESS)))
-		{
-			Activity_SetError(Activity, "Socket connect failed");
-			return -1;
-		}
+		Activity_SetError(Activity, "Socket connect failed");
+		return -1;
 	}
 
 	Activity->Socket = MoveTemp(Candidate);
