Use wait time reported by IAS cache objects

#rnx

(git:4b7206a - ridgers - [REDACTED_PHONE]:28:33 +0200)

diff --git a/Engine/Source/Runtime/Experimental/IoStoreOnDemand/Private/JournaledCache.cpp b/Engine/Source/Runtime/Experimental/IoStoreOnDemand/Private/JournaledCache.cpp
--- a/Engine/Source/Runtime/Experimental/IoStoreOnDemand/Private/JournaledCache.cpp
+++ b/Engine/Source/Runtime/Experimental/IoStoreOnDemand/Private/JournaledCache.cpp
@@ -1447,7 +1447,7 @@ private:
 
 	void						StartThread();
 	int32						Update();
-	void						UpdateCache(FCache* Cache);
+	uint32						UpdateCache(FCache* Cache);
 	virtual uint32				Run() override;
 	virtual void				Stop() override;
 	void						SubmitWork(FWork* Work, uint32 Num);
@@ -1563,20 +1563,28 @@ uint32 FServiceThread::Run()
 {
 	LLM_SCOPE_BYTAG(Ias);
 
+	int64 CycleFreq	= int64(1.0 / FPlatformTime::GetSecondsPerCycle());
+
 	while (RunCount.load(std::memory_order_relaxed))
 	{
 		ReceiveWork();
 
-		int32 WaitTime = Update();
-		if (WaitTime < 0)
+		int32 WaitCycles = Update();
+		if (WaitCycles < 0)
 		{
-			break;
+			continue;
 		}
-		
-		WaitTime = WaitTime ? WaitTime : 37;
-		WakeEvent->Wait(WaitTime);
+
+		uint32 WaitMs = MAX_uint32;
+		if (WaitCycles != MAX_int32)
+		{
+			WaitMs = uint32((int64(WaitCycles) * 1000) / CycleFreq);
+		}
+
+		WakeEvent->Wait(WaitMs);
 	}
 
+	// Loop can exist while there's at least one unregister work to do.
 	ReceiveWork();
 	check(Caches.Num() == 0);
 
@@ -1592,13 +1600,21 @@ void FServiceThread::Stop()
 ////////////////////////////////////////////////////////////////////////////////
 int32 FServiceThread::Update()
 {
+	if (Caches.Num() == 0)
+	{
+		return MAX_int32;
+	}
+
+	// Update caches
+	uint32 CycleSlice = MAX_uint32;
 	for (TUniquePtr<FCache>& Item : Caches)
 	{
 		FCache* Cache = Item.Get();
-		UpdateCache(Cache);
+		uint32 CyclesTillActive = UpdateCache(Cache);
+		CycleSlice = FMath::Min(CyclesTillActive, CycleSlice);
 	}
 
-	return 0;
+	return CycleSlice;
 }
 
 ////////////////////////////////////////////////////////////////////////////////
@@ -1667,13 +1683,13 @@ void FServiceThread::ReceiveWork()
 }
 
 ////////////////////////////////////////////////////////////////////////////////
-void FServiceThread::UpdateCache(FCache* Cache)
+uint32 FServiceThread::UpdateCache(FCache* Cache)
 {
 	uint32 Demand = Cache->GetDemand();
 	int32 Allowance = Governor.BeginAllowance(Demand);
 	if (Allowance <= 0)
 	{
-		return;
+		return -Allowance;
 	}
 
 	TRACE_COUNTER_SET(IasMemDemand, Demand);
@@ -1690,12 +1706,16 @@ void FServiceThread::UpdateCache(FCache* Cache)
 	int32 WaitCycles = Governor.EndAllowance(Unused);
 	if (WaitCycles < 0)
 	{
+		WaitCycles = -WaitCycles;
+
 		TRACE_COUNTER_ADD(IasOpCount, 1); // flush from closing .bin file - we can remove this if we use a single file for .jrn and .bin
 		TRACE_COUNTER_ADD(IasOpCount, 3); // write-flush-commit from .jrn.
 		Cache->Flush();
 	}
 
 	TRACE_COUNTER_SET(IasAllowance, 0);
+
+	return WaitCycles;
 }
 
 // }}}
