Materialize can end three ways; not found, io error, and success

#rnx

(git:cc8d3aa - ridgers - [REDACTED_PHONE]:44:30 +0200)

diff --git a/Engine/Source/Runtime/Experimental/IoStoreOnDemand/Private/JournaledCache.cpp b/Engine/Source/Runtime/Experimental/IoStoreOnDemand/Private/JournaledCache.cpp
--- a/Engine/Source/Runtime/Experimental/IoStoreOnDemand/Private/JournaledCache.cpp
+++ b/Engine/Source/Runtime/Experimental/IoStoreOnDemand/Private/JournaledCache.cpp
@@ -17,6 +17,7 @@
 #include "IO/IoBuffer.h"
 #include "IO/IoDispatcher.h"
 #include "IO/IoHash.h"
+#include "IO/IoStatus.h"
 #include "Math/UnrealMath.h"
 #include "Misc/PathViews.h"
 #include "Misc/Paths.h"
@@ -535,7 +536,7 @@ public:
 	FDiskPhrase				OpenPhrase(uint32 DataSize);
 	void					ClosePhrase(FDiskPhrase&& Phrase);
 	bool					Has(uint64 Key) const;
-	bool					Materialize(uint64 Key, FIoBuffer& Out, uint32 Offset=0) const;
+	EIoErrorCode			Materialize(uint64 Key, FIoBuffer& Out, uint32 Offset=0) const;
 	int32					Flush();
 	void					Drop();
 	uint32					DebugVisit(void* Param, FDebugCacheEntry::Callback* Callback);
@@ -654,19 +655,19 @@ bool FDiskCache::Has(uint64 Key) const
 }
 
 ////////////////////////////////////////////////////////////////////////////////
-bool FDiskCache::Materialize(uint64 Key, FIoBuffer& Out, uint32 Offset) const
+EIoErrorCode FDiskCache::Materialize(uint64 Key, FIoBuffer& Out, uint32 Offset) const
 {
 	TRACE_CPUPROFILER_EVENT_SCOPE(IasCache::Materialize_Disk);
 
 	if (!DataHandle.IsValid())
 	{
-		return false;
+		return EIoErrorCode::FileNotOpen;
 	}
 
 	const FMapEntry* Entry = DataMap.Find(Key);
 	if (Entry == nullptr)
 	{
-		return false;
+		return EIoErrorCode::NotFound;
 	}
 
 	uint32 ReadSize = uint32(Entry->Size) - Offset;
@@ -679,7 +680,8 @@ bool FDiskCache::Materialize(uint64 Key, FIoBuffer& Out, uint32 Offset) const
 	ReadSize = FMath::Min<uint32>(uint32(Out.GetSize()), ReadSize);
 
 	DataHandle->Seek(Entry->DataCursor + Offset);
-	return DataHandle->Read(Out.GetData(), ReadSize);
+	bool bOk = DataHandle->Read(Out.GetData(), ReadSize);
+	return bOk ? EIoErrorCode::Ok : EIoErrorCode::ReadError;
 }
 
 ////////////////////////////////////////////////////////////////////////////////
@@ -1079,7 +1081,7 @@ public:
 	bool			Has(uint64 Key) const;
 	FGetToken		Get(uint64 Key, FIoBuffer& OutData) const;
 	bool			Put(uint64 Key, FIoBuffer& Data);
-	bool			Materialize(FGetToken Token, FIoBuffer& OutData, uint32 Offset=0) const;
+	EIoErrorCode	Materialize(FGetToken Token, FIoBuffer& OutData, uint32 Offset=0) const;
 	uint32			Flush();
 	uint32			WriteMemToDisk(int32 Allowance);
 	uint32			DebugVisit(void* Param, FDebugCacheEntry::Callback* Callback);
@@ -1177,18 +1179,18 @@ bool FCache::Put(uint64 Key, FIoBuffer& Data)
 }
 
 ////////////////////////////////////////////////////////////////////////////////
-bool FCache::Materialize(FGetToken Token, FIoBuffer& OutData, uint32 Offset) const
+EIoErrorCode FCache::Materialize(FGetToken Token, FIoBuffer& OutData, uint32 Offset) const
 {
 	FReadScopeLock _(FsLock);
 
 	uint64 Key = Token;
-	if (!DiskCache.Materialize(Key, OutData, Offset))
+	EIoErrorCode Ret = DiskCache.Materialize(Key, OutData, Offset);
+	if (Ret == EIoErrorCode::Ok)
 	{
-		return false;
+		FOnDemandIoBackendStats::Get()->OnCacheGet(OutData.GetSize());
 	}
 
-	FOnDemandIoBackendStats::Get()->OnCacheGet(OutData.GetSize());
-	return true;
+	return Ret;
 }
 
 ////////////////////////////////////////////////////////////////////////////////
@@ -1866,9 +1868,10 @@ FJournaledCache::FGetWork FJournaledCache::Materialize(
 		}
 
 		uint32 Offset = uint32(Options.GetOffset());
-		if (!Cache->Materialize(Token, Buffer, Offset))
+		EIoErrorCode Ret = Cache->Materialize(Token, Buffer, Offset);
+		if (Ret != EIoErrorCode::Ok)
 		{
-			return TIoStatusOr<FIoBuffer>(FIoStatus(EIoErrorCode::ReadError));
+			return TIoStatusOr<FIoBuffer>(FIoStatus(Ret));
 		}
 
 		return TIoStatusOr<FIoBuffer>(Buffer);
