Fixed a crash in FDistributionEndpoints::CompleteEndpointRequest
When the request times out, the returned response pointer will be null and be treated as an error.

#rnx

diff --git a/Engine/Source/Runtime/Experimental/IoStoreOnDemand/Private/DistributionEndpoints.cpp b/Engine/Source/Runtime/Experimental/IoStoreOnDemand/Private/DistributionEndpoints.cpp
--- a/Engine/Source/Runtime/Experimental/IoStoreOnDemand/Private/DistributionEndpoints.cpp
+++ b/Engine/Source/Runtime/Experimental/IoStoreOnDemand/Private/DistributionEndpoints.cpp
@@ -193,7 +193,7 @@ void FDistributionEndpoints::CompleteEndpointRequest(FResolveRequest& ResolveReq
 	using FJsonReaderPtr = TSharedRef<FJsonReader>;
 
 	TArray<FString> ServiceUrls;
-	if (HttpResponse->GetResponseCode() == 200)
+	if (HttpResponse != nullptr && HttpResponse->GetResponseCode() == 200)
 	{
 		FString Json = HttpResponse->GetContentAsString();
 		FJsonReaderPtr JsonReader = TJsonReaderFactory<TCHAR>::Create(Json);
