Wrap activities' reading and writing

#rnx

(git:78bb6b8 - Martin Ridgers - [REDACTED_PHONE]:58:58 +0200)

diff --git a/Engine/Source/Runtime/Experimental/IoStoreOnDemand/Private/CoreHttp/Http.cpp b/Engine/Source/Runtime/Experimental/IoStoreOnDemand/Private/CoreHttp/Http.cpp
--- a/Engine/Source/Runtime/Experimental/IoStoreOnDemand/Private/CoreHttp/Http.cpp
+++ b/Engine/Source/Runtime/Experimental/IoStoreOnDemand/Private/CoreHttp/Http.cpp
@@ -745,6 +745,14 @@ class FSocket
 public:
 	SocketType	Get() const { return Socket; } // to be removed
 
+	enum class EResult
+	{
+		HangUp			=  0,
+		Wait			= -1,
+		Error			= -2,
+		ConnectError	= -3,
+	};
+
 				FSocket() = default;
 				~FSocket()					{ if (IsValid()) Destroy(); }
 				FSocket(FSocket&& Rhs)		{ Move(MoveTemp(Rhs)); }
@@ -754,6 +762,8 @@ public:
 	void		Destroy();
 	bool		Connect(uint32 Ip, uint32 Port);
 	void		Disconnect();
+	int32		Send(const char* Data, uint32 Size);
+	int32		Recv(char* Dest, uint32 Size);
 	bool		SetBlocking(bool bBlocking);
 	bool		SetSendBufSize(int32 Size);
 	bool		SetRecvBufSize(int32 Size);
@@ -820,6 +830,65 @@ void FSocket::Disconnect()
 	shutdown(Socket, SHUT_RDWR);
 }
 
+////////////////////////////////////////////////////////////////////////////////
+int32 FSocket::Send(const char* Data, uint32 Size)
+{
+	int32 Result = send(Socket, Data, Size, MsgFlagType(0));
+
+	if (Result > 0)
+	{
+		return Result;
+	}
+
+	if (Result == 0)
+	{
+		return int32(EResult::HangUp);
+	}
+
+	if (IsSocketResult(EWOULDBLOCK))
+	{
+		return int32(EResult::Wait);
+	}
+
+	if (IsSocketResult(ENOTCONN))
+	{
+		int32 Error = 0;
+		socklen_t ErrorSize = sizeof(Error);
+		Result = getsockopt(Socket, SOL_SOCKET, SO_ERROR, (char*)&Error, &ErrorSize);
+		if (Result < 0 || Error != 0)
+		{
+			return int32(EResult::ConnectError);
+		}
+
+		return int32(EResult::Wait);
+	}
+
+	return int32(EResult::Error);
+}
+
+////////////////////////////////////////////////////////////////////////////////
+int32 FSocket::Recv(char* Dest, uint32 Size)
+{
+	int32 Result = recv(Socket, Dest, Size, MsgFlagType(0));
+
+	if (Result > 0)
+	{
+		return Result;
+	}
+
+	if (Result == 0)
+	{
+		return int32(EResult::HangUp);
+	}
+
+	if (IsSocketResult(EWOULDBLOCK))
+	{
+		return int32(EResult::Wait);
+	}
+
+	return int32(EResult::Error);
+}
+
 ////////////////////////////////////////////////////////////////////////////////
 bool FSocket::SetBlocking(bool bBlocking)
 {
@@ -1909,44 +1978,26 @@ static int32 DoSend(FActivity* Activity)
 	check(SendSize > 0);
 
 	Trace(Activity, ETrace::Send, SendSize);
-	int32 Result = send(Activity->Socket.Get(), SendData, SendSize, MsgFlagType(0));
+	int32 Result = Activity->Socket.Send(SendData, SendSize);
 	Trace(Activity, ETrace::Send, -1);
 
-	if (Result < 0)
+	switch (FSocket::EResult(Result))
 	{
-		if (IsSocketResult(ENOTCONN))
-		{
-			int32 Error = 0;
-			socklen_t ErrorSize = sizeof(Error);
-			Result = getsockopt(Activity->Socket.Get(), SOL_SOCKET, SO_ERROR, (char*)&Error, &ErrorSize);
-			if (Result < 0 || Error != 0)
-			{
-				Activity_SetError(Activity, "Connection error");
-				return -1;
-			}
-			return 1;
-		}
-
-		if (!IsSocketResult(EWOULDBLOCK))
-		{
-			Activity_SetError(Activity, "Error returned from socket send");
-			return -1;
-		}
-	}
-
-	if (Result == 0)
-	{
-		Activity_SetError(Activity, "ATH0.Send");
-		return -1;
+	case FSocket::EResult::HangUp:		Activity_SetError(Activity, "ATH0.Send"); return -1;
+	case FSocket::EResult::Error:		Activity_SetError(Activity, "Error returned from socket send"); return -1;
+	case FSocket::EResult::ConnectError:Activity_SetError(Activity, "Connection error"); return -1;
+	case FSocket::EResult::Wait:
+		Activity->StateParam = Remaining;
+		Activity->SocketWait = FActivity::EWait::Write;
+		Trace(Activity, ETrace::Wait);
+		return 1;
 	}
 
+	checkf(Result > 0, TEXT("Result wasn't caught by switch statement so it is expected to be a positive amount of bytes sent"));
 	Remaining = SendSize - Result;
 	if (Remaining != 0)
 	{
-		Activity->StateParam = Remaining;
-		Activity->SocketWait = FActivity::EWait::Write;
-		Trace(Activity, ETrace::Wait);
-		return 1;
+		return DoSend(Activity);
 	}
 
 	// It is expected there will be enough space for a RespInt object
@@ -1975,19 +2026,19 @@ static int32 DoRecvMessage(FActivity* Activity)
 		auto [Dest, DestSize] = Buffer.GetMutableFree(0, PageSize);
 
 		Trace(Activity, ETrace::Recv, -1);
-		int32 Result = recv(Activity->Socket.Get(), Dest, DestSize, MsgFlagType(0));
+		int32 Result = Activity->Socket.Recv(Dest, DestSize);
 		Trace(Activity, ETrace::Recv, FMath::Max(Result, 0));
 
-		if (Result < 0)
+		if (Result == int32(FSocket::EResult::Wait))
 		{
-			if (IsSocketResult(EWOULDBLOCK))
-			{
-				Activity->SocketWait = FActivity::EWait::Read;
-				Trace(Activity, ETrace::Wait);
-				return 1;
-			}
+			Activity->SocketWait = FActivity::EWait::Read;
+			Trace(Activity, ETrace::Wait);
+			return 1;
+		}
 
-			Activity_SetError(Activity, "Error occurred on socket recv");
+		if (Result < 0)
+		{
+			Activity_SetError(Activity, "Error returned from socket recv");
 			return -1;
 		}
 
@@ -2189,20 +2240,21 @@ static FHandlerResult DoRecvContent(FActivity* Activity, uint32 MaxRecvSize)
 			return { 1, RecvSize };
 		}
 
-		Trace(Activity, ETrace::Recv, -1);
 		char* Cursor = (char*)(DestView.GetData()) + Activity->StateParam;
-		int32 Result = recv(Activity->Socket.Get(), Cursor, Size, MsgFlagType(0));
+
+		Trace(Activity, ETrace::Recv, -1);
+		int32 Result = Activity->Socket.Recv(Cursor, Size);
 		Trace(Activity, ETrace::Recv, FMath::Max(Result, 0));
 
-		if (Result < 0)
+		if (Result == int32(FSocket::EResult::Wait))
 		{
-			if (IsSocketResult(EWOULDBLOCK))
-			{
-				Activity->SocketWait = FActivity::EWait::Read;
-				Trace(Activity, ETrace::Wait);
-				return { 1, RecvSize };
-			}
+			Activity->SocketWait = FActivity::EWait::Read;
+			Trace(Activity, ETrace::Wait);
+			return { 1, RecvSize };
+		}
 
+		if (Result < 0)
+		{
 			Activity_SetError(Activity, "Socket error while receiving content");
 			return { -1 };
 		}
