Generalised the pointer in work items, plus a small tidy

#rnx

(git:d56b0ef - Martin Ridgers - [REDACTED_PHONE]:35:49 +0200)

diff --git a/Engine/Source/Runtime/Experimental/IoStoreOnDemand/Private/JournaledCache.cpp b/Engine/Source/Runtime/Experimental/IoStoreOnDemand/Private/JournaledCache.cpp
--- a/Engine/Source/Runtime/Experimental/IoStoreOnDemand/Private/JournaledCache.cpp
+++ b/Engine/Source/Runtime/Experimental/IoStoreOnDemand/Private/JournaledCache.cpp
@@ -1421,29 +1421,29 @@ private:
 			Work_GovRate,
 		};
 
-		void	SetCache(FCache* In) { Cache = UPTRINT(In) >> 3; check((UPTRINT(In) & 0x7) == 0); }
-		FCache* GetCache() const	 { return (FCache*)(Cache << 3); }
+		void	SetPtr(const void* In)	{ Ptr = UPTRINT(In) >> 3; check((UPTRINT(In) & 0x7) == 0); }
+		void*	GetPtr() const			{ return (FCache*)(Ptr << 3); }
 
 		union {
-			struct {
-				UPTRINT			What : 3; // (un)reg
-				UPTRINT			Cache : 44;
-				UPTRINT			_Unused : 17;
+			struct {			// reg / unreg
+				UPTRINT			What : 3;
+				UPTRINT			Ptr : 45;
+				UPTRINT			_Unused : 16;
 			};
-			struct {
-				uint16			_What0 : 3; // rate
+			struct {			// rate
+				uint16			_What0 : 3;
 				uint16			Ops : 13;
 				uint16			Seconds;
 				uint32			Allowance;
 			};
-			struct {
-				uint16			_What1 : 3; // demand
+			struct {			// demand
+				uint16			_What1 : 3;
+				uint16			_Unused1 : 13;
 				uint16			Threshold;
 				uint16			Boost;
 				uint16			SuperBoost;
 			};
 		};
-		//uint64				Key;
 	};
 	static_assert(sizeof(FWork) == sizeof(UPTRINT));
 
@@ -1497,7 +1497,7 @@ void FServiceThread::RegisterCache(TUniquePtr<FCache> Cache)
 
 	FWork Work;
 	Work.What = FWork::Work_Register;
-	Work.SetCache(RawPtr);
+	Work.SetPtr(RawPtr);
 	SubmitWork(&Work, 1);
 
 	if (PrevRunCount == 0)
@@ -1511,7 +1511,7 @@ void FServiceThread::UnregisterCache(FCache* Cache)
 {
 	FWork Work;
 	Work.What = FWork::Work_Unregister;
-	Work.SetCache(Cache);
+	Work.SetPtr(Cache);
 	SubmitWork(&Work, 1);
 
 	int32 PrevRunCount = RunCount.fetch_sub(1, std::memory_order_relaxed);
@@ -1634,14 +1634,14 @@ void FServiceThread::ReceiveWork()
 	PendingPrev = PendingLoad;
 
 	// Unregisters first
-	for (FWork& Work : InboundWork)
+	for (const FWork& Work : InboundWork)
 	{
 		if (Work.What != FWork::Work_Unregister)
 		{
 			continue;
 		}
 
-		FCache* CachePtr = Work.GetCache();
+		auto* CachePtr = (FCache*)(Work.GetPtr());
 		for (int32 i = 0, n = Caches.Num(); i < n; ++i)
 		{
 			if (Caches[i].Get() != CachePtr)
@@ -1664,7 +1664,7 @@ void FServiceThread::ReceiveWork()
 
 		if (Work.What == FWork::Work_Register)
 		{
-			FCache* Cache = Work.GetCache();
+			auto* Cache = (FCache*)(Work.GetPtr());
 			Cache->Load();
 			Caches.Add(TUniquePtr<FCache>(Cache));
 			continue;
