Moved proxy support out of the middle of activities

- Moved all related things together in one place
- try and use socks while the socket is in a blocking mode for simplicity

#rnx

(git:0899c97 - Martin Ridgers - [REDACTED_PHONE]:04:21 +0200)

diff --git a/Engine/Source/Runtime/Experimental/IoStoreOnDemand/Private/CoreHttp/Http.cpp b/Engine/Source/Runtime/Experimental/IoStoreOnDemand/Private/CoreHttp/Http.cpp
--- a/Engine/Source/Runtime/Experimental/IoStoreOnDemand/Private/CoreHttp/Http.cpp
+++ b/Engine/Source/Runtime/Experimental/IoStoreOnDemand/Private/CoreHttp/Http.cpp
@@ -176,68 +176,6 @@
 
 // {{{1 misc ...................................................................
 
-///////////////////////////////////////////////////////////////////////////////
-static uint32	GSocksIpAddress	= 0;
-static int32	GSocksPort		= 1080; // default SOCKS5 port
-static int32	GSocksVersion	= 5;
-
-#if !UE_BUILD_SHIPPING
-static void SetSocksIpAddress(const TCHAR* Value)
-{
-	uint32 IpAddress = 0;
-	uint32 Accumulator = 0;
-	while (true)
-	{
-		uint32 c = *Value++;
-
-		if (c - '0' <= '9' - '0')
-		{
-			Accumulator *= 10;
-			Accumulator += (c - '0');
-			continue;
-		}
-
-		if (c == '.' || c == '\0')
-		{
-			IpAddress <<= 8;
-			IpAddress |= Accumulator;
-			Accumulator = 0;
-			if (c == '\0')
-			{
-				break;
-			}
-			continue;
-		}
-
-		GSocksIpAddress = 0;
-		return;
-	}
-	GSocksIpAddress = htonl(IpAddress);
-}
-
-static FString GSocksIpAddressStr;
-static FAutoConsoleVariableRef CVar_IasSocksIpAddress(
-	TEXT("ias.SocksIp"),
-	GSocksIpAddressStr,
-	TEXT("Routes all IAS HTTP traffic through the given SOCKS5 proxy"),
-	FConsoleVariableDelegate::CreateStatic([] (IConsoleVariable*) {
-		SetSocksIpAddress(*GSocksIpAddressStr);
-	})
-);
-
-static FAutoConsoleVariableRef CVar_IasSocksPort(
-	TEXT("ias.SocksPort"),
-	GSocksPort,
-	TEXT("Port of the SOCKS5 proxy to use")
-);
-
-static FAutoConsoleVariableRef CVar_IasSocksVersion(
-	TEXT("ias.SocksVersion"),
-	GSocksVersion,
-	TEXT("SOCKS proxy protocol version to use")
-);
-#endif
-
 ////////////////////////////////////////////////////////////////////////////////
 template <typename LambdaType>
 static void EnumerateHeaders(FAnsiStringView Headers, LambdaType&& Lambda)
@@ -929,6 +867,194 @@
 
 
 
+// {{{1 socks ..................................................................
+
+///////////////////////////////////////////////////////////////////////////////
+static uint32	GSocksIpAddress	= 0;
+static int32	GSocksPort		= 1080; // default SOCKS5 port
+static int32	GSocksVersion	= 5;
+
+#if !UE_BUILD_SHIPPING
+static void SetSocksIpAddress(const TCHAR* Value)
+{
+	uint32 IpAddress = 0;
+	uint32 Accumulator = 0;
+	while (true)
+	{
+		uint32 c = *Value++;
+
+		if (c - '0' <= '9' - '0')
+		{
+			Accumulator *= 10;
+			Accumulator += (c - '0');
+			continue;
+		}
+
+		if (c == '.' || c == '\0')
+		{
+			IpAddress <<= 8;
+			IpAddress |= Accumulator;
+			Accumulator = 0;
+			if (c == '\0')
+			{
+				break;
+			}
+			continue;
+		}
+
+		GSocksIpAddress = 0;
+		return;
+	}
+	GSocksIpAddress = htonl(IpAddress);
+}
+
+static FString GSocksIpAddressStr;
+static FAutoConsoleVariableRef CVar_IasSocksIpAddress(
+	TEXT("ias.SocksIp"),
+	GSocksIpAddressStr,
+	TEXT("Routes all IAS HTTP traffic through the given SOCKS5 proxy"),
+	FConsoleVariableDelegate::CreateStatic([] (IConsoleVariable*) {
+		SetSocksIpAddress(*GSocksIpAddressStr);
+	})
+);
+
+static FAutoConsoleVariableRef CVar_IasSocksPort(
+	TEXT("ias.SocksPort"),
+	GSocksPort,
+	TEXT("Port of the SOCKS5 proxy to use")
+);
+
+static FAutoConsoleVariableRef CVar_IasSocksVersion(
+	TEXT("ias.SocksVersion"),
+	GSocksVersion,
+	TEXT("SOCKS proxy protocol version to use")
+);
+#endif
+
+////////////////////////////////////////////////////////////////////////////////
+static int32 ConnectSocks4(FSocket& Socket, uint32 IpAddress, uint32 Port)
+{
+	struct FSocks4Request
+	{
+		uint8	Version = 4;
+		uint8	Command = 1;
+		uint16	Port;
+		uint32	IpAddress;
+	};
+
+	struct FSocks4Reply
+	{
+		uint8	Version;
+		uint8	Code;
+		uint16	Port;
+		uint32	IpAddress;
+	};
+
+	int32 Result;
+
+	FSocks4Request Request = {
+		.Port		= htons(uint16(Port)),
+		.IpAddress	= IpAddress,
+	};
+	Result = Socket.Send((const char*)&Request, sizeof(Request));
+	if (Result <= 0)
+	{
+		return -1;
+	}
+
+	FSocks4Reply Reply;
+	Result = Socket.Recv((char*)&Reply, sizeof(Reply));
+	if (Result <= 0)
+	{
+		return -1;
+	}
+
+	return 0;
+}
+
+////////////////////////////////////////////////////////////////////////////////
+static int32 ConnectSocks5(FSocket& Socket, uint32 IpAddress, uint32 Port)
+{
+#ifdef _MSC_VER
+	// MSVC's static analysis doesn't see that 'Result' from recv() is checked
+	// to be the exact size of the destination buffer.
+#pragma warning(push)
+#pragma warning(disable : 6385)
+#endif
+
+	int32 Result;
+
+	// Greeting
+	const char Greeting[] = { 5, 1, 0 };
+	Result = Socket.Send(Greeting, sizeof(Greeting));
+	if (Result != sizeof(Greeting))
+	{
+		return -1;
+	}
+
+	// Server auth-choice
+	char ServerChoice[1 + 1];
+	Result = Socket.Recv(ServerChoice, sizeof(ServerChoice));
+	if (Result != sizeof(ServerChoice))
+	{
+		return -1;
+	}
+
+	if (ServerChoice[0] != 0x05 || ServerChoice[1] != 0x00)
+	{
+		return -1;
+	}
+
+	// Connection request
+	uint16 NsPort = htons(uint16(Port));
+	char Request[] = { 5, 1, 0, 1, 0x11,0x11,0x11,0x11, 0x22,0x22 };
+	std::memcpy(Request + 4, &IpAddress, sizeof(IpAddress));
+	std::memcpy(Request + 8, &NsPort, sizeof(NsPort));
+	Result = Socket.Send(Request, sizeof(Request));
+	if (Result != sizeof(Request))
+	{
+		return -1;
+	}
+
+	// Connect reply
+	char Reply[3 + (1 + 4) + 2];
+	Result = Socket.Recv(Reply, sizeof(Reply));
+	if (Result != sizeof(Reply))
+	{
+		return -1;
+	}
+
+	if (Reply[0] != 0x05 && Reply[1] != 0x00)
+	{
+		return -1;
+	}
+
+	return 0;
+
+#ifdef _MSC_VER
+#pragma warning(pop)
+#endif
+}
+
+////////////////////////////////////////////////////////////////////////////////
+static int32 MaybeConnectSocks(FSocket& Socket, uint32 IpAddress, uint32 Port)
+{
+	if (GSocksIpAddress == 0)
+	{
+		return 0;
+	}
+
+	switch (GSocksVersion)
+	{
+	case 4: return ConnectSocks4(Socket, IpAddress, Port);
+	case 5: return ConnectSocks5(Socket, IpAddress, Port);
+	}
+
+	return -1;
+}
+
+
+
 // {{{1 connection-pool ........................................................
 
 ////////////////////////////////////////////////////////////////////////////////
@@ -1190,7 +1316,6 @@
 		Build,
 		Resolve,
 		Connect,
-		Socks,
 		Send,
 		RecvMessage,
 		RecvContent,
@@ -1224,10 +1349,9 @@
 static void Activity_TraceStateNames()
 {
 	FAnsiStringView StateNames[] = {
-		"Build",	"Resolve",		"Connect",		"Socks",
-		"Send",		"RecvMessage",	"RecvContent",	"RecvStream",
-		"RecvDone",	"Completed",	"Cancelled",	"Failed",
-		"$",
+		"Build",		"Resolve",		"Connect",		"Send",
+		"RecvMessage",	"RecvContent",	"RecvStream",	"RecvDone",
+		"Completed",	"Cancelled",	"Failed",		"$",
 	};
 	static_assert(UE_ARRAY_COUNT(StateNames) == int32(FActivity::EState::_Num) + 1);
 
@@ -1745,6 +1869,7 @@
 	check(Pool->IsResolved() > 0);
 
 	uint32 IpAddress = Pool->GetIpAddress();
+	uint32 Port = Pool->GetPort();
 	if (IpAddress <= 0x00ff'ffff)
 	{
 		Activity_SetError(Activity, "Unresolved host");
@@ -1776,14 +1901,27 @@
 		return -1;
 	}
 
+	// Attempt a SOCKS connect
+	Trace(Activity, ETrace::Connect, IpAddress);
+	bool bSocksConnected = false;
+	if (int32 Result = MaybeConnectSocks(Candidate, IpAddress, Port); Result)
+	{
+		if (Result < 0)
+		{
+			Activity_SetError(Activity, "Failed establishing SOCKS connection");
+			return -1;
+		}
+
+		bSocksConnected = true;
+	}
+
+	// Condition the socket
 	if (!Candidate.SetBlocking(false))
 	{
 		Activity_SetError(Activity, "Unable to set socket non-blocking");
 		return -1;
 	}
 
-
-	// Adjust socket send and recv buffer sizes
 	if (int32 OptValue = Pool->GetBufferSize(FSocketPool::EDirection::Send); OptValue >= 0)
 	{
 		Candidate.SetSendBufSize(OptValue);
@@ -1794,18 +1932,8 @@
 		Candidate.SetRecvBufSize(OptValue);
 	}
 
-	uint16 Port = uint16(Pool->GetPort());
-
-	// Redirect connect to a SOCKS proxy if one is set
-	if (GSocksIpAddress != 0)
+	// Adjust socket send and recv buffer sizes
+	if (!bSocksConnected && !Candidate.Connect(IpAddress, Port))
 	{
-		IpAddress = GSocksIpAddress;
-		Port = uint16(GSocksPort);
-	}
-
-	// connect
-	Trace(Activity, ETrace::Connect, IpAddress);
-	if (Candidate.Connect(IpAddress, Port))
-	{
 		Activity_SetError(Activity, "Socket connect failed");
 		return -1;
@@ -1815,154 +1943,11 @@
 	Activity->SocketWait = FActivity::EWait::Write;
 
 	Activity->StateParam = 0;
-	Activity->State = (GSocksIpAddress != 0)
-		? FActivity::EState::Socks
-		: FActivity::EState::Send;
-
+	Activity->State = FActivity::EState::Send;
 	Trace(Activity, ETrace::StateChangeWait, Activity->State);
 	return 1;
 }
 
 ////////////////////////////////////////////////////////////////////////////////
-static int32 DoSocks4(FActivity* Activity)
-{
-	struct FSocks4Request
-	{
-		uint8	Version = 4;
-		uint8	Command = 1;
-		uint16	Port;
-		uint32	IpAddress;
-	};
-
-	struct FSocks4Reply
-	{
-		uint8	Version;
-		uint8	Code;
-		uint16	Port;
-		uint32	IpAddress;
-	};
-
-	SocketType Socket = Activity->Socket.Get();
-	int32 Result = 0;
-
-	FSocks4Request Request = {
-		.Port		= htons(uint16(Activity->Pool->GetPort())),
-		.IpAddress	= Activity->Pool->GetIpAddress(),
-	};
-	Result = send(Socket, (const char*)&Request, sizeof(Request), MsgFlagType(0));
-	if (Result <= 0)
-	{
-		Activity_SetError(Activity, "Failed sending SOCKS4 request");
-		return -1;
-	}
-
-	FSocks4Reply Reply;
-	do {
-		Result = recv(Socket, (char*)&Reply, sizeof(Reply), MsgFlagType(0));
-	} while ((Result < 0) & (IsSocketResult(EWOULDBLOCK) | IsSocketResult(EINPROGRESS)));
-	if (Result <= 0)
-	{
-		Activity_SetError(Activity, "Failed receiving SOCKS4 reply");
-		return -1;
-	}
-
-	return 0;
-}
-
-////////////////////////////////////////////////////////////////////////////////
-static int32 DoSocks5(FActivity* Activity)
-{
-#ifdef _MSC_VER
-	// MSVC's static analysis doesn't see that 'Result' from recv() is checked
-	// to be the exact size of the destination buffer.
-#pragma warning(push)
-#pragma warning(disable : 6385)
-#endif
-
-	SocketType Socket = Activity->Socket.Get();
-	int32 Result;
-
-	// Greeting
-	const char Greeting[] = { 5, 1, 0 };
-	Result = send(Socket, Greeting, sizeof(Greeting), MsgFlagType(0));
-	if (Result != sizeof(Greeting))
-	{
-		Activity_SetError(Activity, "Failed sending SOCKS5 greeting");
-		return -1;
-	}
-
-	// Server auth-choice
-	char ServerChoice[1 + 1];
-	do {
-		Result = recv(Socket, ServerChoice, sizeof(ServerChoice), MsgFlagType(0));
-	} while ((Result < 0) & (IsSocketResult(EWOULDBLOCK) | IsSocketResult(EINPROGRESS)));
-	if (Result != sizeof(ServerChoice))
-	{
-		Activity_SetError(Activity, "Failed getting SOCKS5 server choice");
-		return -1;
-	}
-
-	if (ServerChoice[0] != 0x05 || ServerChoice[1] != 0x00)
-	{
-		Activity_SetError(Activity, "Unable to authenticate with SOCKS");
-		return -1;
-	}
-
-	// Connection request
-	uint32 IpAddress = Activity->Pool->GetIpAddress();
-	uint16 Port = htons(uint16(Activity->Pool->GetPort()));
-	char Request[] = { 5, 1, 0, 1, 0x11,0x11,0x11,0x11, 0x22,0x22 };
-	std::memcpy(Request + 4, &IpAddress, sizeof(IpAddress));
-	std::memcpy(Request + 8, &Port, sizeof(Port));
-	Result = send(Socket, Request, sizeof(Request), MsgFlagType(0));
-	if (Result != sizeof(Request))
-	{
-		Activity_SetError(Activity, "Failed requesting SOCKS5 connection");
-		return -1;
-	}
-
-	// Connect reply
-	char Reply[3 + (1 + 4) + 2];
-	do {
-		Result = recv(Socket, Reply, sizeof(Reply), MsgFlagType(0));
-	} while ((Result < 0) & (IsSocketResult(EWOULDBLOCK) | IsSocketResult(EINPROGRESS)));
-	if (Result != sizeof(Reply))
-	{
-		Activity_SetError(Activity, "Failed SOCKS5 connection");
-		return -1;
-	}
-
-	if (Reply[0] != 0x05 && Reply[1] != 0x00)
-	{
-		Activity_SetError(Activity, "SOCKS5 connection denied");
-		return -1;
-	}
-
-	return 0;
-
-#ifdef _MSC_VER
-#pragma warning(pop)
-#endif
-}
-
-////////////////////////////////////////////////////////////////////////////////
-static int32 DoSocks(FActivity* Activity)
-{
-	switch (GSocksVersion)
-	{
-	case 4: if (DoSocks4(Activity) == -1) return -1; break;
-	case 5: if (DoSocks5(Activity) == -1) return -1; break;
-	default:
-		Activity_SetError(Activity, "Unsupported SOCKS proxy version");
-		return -1;
-	}
-
-	Activity->State = FActivity::EState::Send;
-	Activity->StateParam = 0;
-	Trace(Activity, ETrace::StateChange, Activity->State);
-	return 0;
-}
-
-////////////////////////////////////////////////////////////////////////////////
 static int32 DoSend(FActivity* Activity)
 {
@@ -2660,10 +2645,6 @@
 			--BusyCount;
 			Result = 0;
 			break;
-
-		case FActivity::EState::Socks:
-			Result = DoSocks(Activity);
-			break;
 		}
 
 		if (Result == -1)
