Added counter for the cache seek position

#rnx

(git:7ef0857 - Martin Ridgers - [REDACTED_PHONE]:29:14 +0200)

diff --git a/Engine/Source/Runtime/Experimental/IoStoreOnDemand/Private/JournaledCache.cpp b/Engine/Source/Runtime/Experimental/IoStoreOnDemand/Private/JournaledCache.cpp
--- a/Engine/Source/Runtime/Experimental/IoStoreOnDemand/Private/JournaledCache.cpp
+++ b/Engine/Source/Runtime/Experimental/IoStoreOnDemand/Private/JournaledCache.cpp
@@ -53,6 +53,7 @@ namespace UE::IO::IAS::JournaledCache
 TRACE_DECLARE_INT_COUNTER(IasMemDemand, TEXT("Ias/CacheMemDemand"));
 TRACE_DECLARE_INT_COUNTER(IasAllowance, TEXT("Ias/CacheAllowance"));
 TRACE_DECLARE_INT_COUNTER(IasOpCount,   TEXT("Ias/CacheOpCount"));
+TRACE_DECLARE_INT_COUNTER(IasReadCursor,TEXT("Ias/CacheReadCursor"));
 
 ////////////////////////////////////////////////////////////////////////////////
 static bool LoadCache(class FDiskCache&);
@@ -679,6 +680,8 @@ EIoErrorCode FDiskCache::Materialize(uint64 Key, FIoBuffer& Out, uint32 Offset)
 
 	ReadSize = FMath::Min<uint32>(uint32(Out.GetSize()), ReadSize);
 
+	TRACE_COUNTER_SET(IasReadCursor, Entry->DataCursor + Offset);
+
 	DataHandle->Seek(Entry->DataCursor + Offset);
 	bool bOk = DataHandle->Read(Out.GetData(), ReadSize);
 	return bOk ? EIoErrorCode::Ok : EIoErrorCode::ReadError;
