Moved poll handling out of ReadyCheck and into the appropriate class

#rnx

(git:2953b95 - Martin Ridgers - [REDACTED_PHONE]:11:36 +0200)

diff --git a/Engine/Source/Runtime/Experimental/IoStoreOnDemand/Private/CoreHttp/Http.cpp b/Engine/Source/Runtime/Experimental/IoStoreOnDemand/Private/CoreHttp/Http.cpp
--- a/Engine/Source/Runtime/Experimental/IoStoreOnDemand/Private/CoreHttp/Http.cpp
+++ b/Engine/Source/Runtime/Experimental/IoStoreOnDemand/Private/CoreHttp/Http.cpp
@@ -691,6 +691,15 @@ public:
 		ConnectError	= -3,
 	};
 
+	struct FWaiter
+	{
+				FWaiter(FSocket& Socket, bool bInIsRecvWait);
+		bool	operator == (FSocket& Rhs) const { return UPTRINT(&Rhs) == Candidate; }
+		UPTRINT	Candidate : 62;
+		UPTRINT	bRecvWait : 1;
+		UPTRINT	bReady : 1;
+	};
+
 				FSocket() = default;
 				~FSocket()					{ if (IsValid()) Destroy(); }
 				FSocket(FSocket&& Rhs)		{ Move(MoveTemp(Rhs)); }
@@ -705,6 +714,7 @@ public:
 	bool		SetBlocking(bool bBlocking);
 	bool		SetSendBufSize(int32 Size);
 	bool		SetRecvBufSize(int32 Size);
+	static int	Wait(TArrayView<FWaiter> Waiters, int32 TimeoutMs);
 
 private:
 	void		Move(FSocket&& Rhs);
@@ -714,6 +724,14 @@ private:
 	FSocket&	operator = (const FSocket&) = delete;
 };
 
+////////////////////////////////////////////////////////////////////////////////
+FSocket::FWaiter::FWaiter(FSocket& Socket, bool bInIsRecvWait)
+: Candidate(UPTRINT(&Socket))
+, bRecvWait(bInIsRecvWait == true)
+, bReady(0)
+{
+}
+
 ////////////////////////////////////////////////////////////////////////////////
 void FSocket::Move(FSocket&& Rhs)
 {
@@ -865,6 +883,104 @@ bool FSocket::SetRecvBufSize(int32 Size)
 	return 0 == setsockopt(Socket, SOL_SOCKET, SO_RCVBUF, &(char&)Size, sizeof(Size));
 }
 
+////////////////////////////////////////////////////////////////////////////////
+int32 FSocket::Wait(TArrayView<FWaiter> Waiters, int32 TimeoutMs)
+{
+#if !defined(IAS_HTTP_USE_POLL)
+	struct FSelect
+	{
+		SocketType	fd;
+		int32		events;
+		int32		revents;
+	};
+	static const int32 POLLIN  = 1 << 0;
+	static const int32 POLLOUT = 1 << 1;
+	static const int32 POLLERR = 1 << 2;
+	static const int32 POLLHUP = POLLERR;
+	static const int32 POLLNVAL= POLLERR;
+#else
+	using FSelect = pollfd;
+#endif
+
+	// The following looks odd because POLLFD varies subtly from one platform
+	// to the next. To cleanly set members to zero and to not get narrowing
+	// warnings from the compiler, we list-init and don't assume POD types.
+	TArray<FSelect, TFixedAllocator<64>> Selects;
+	for (FWaiter& Waiter : Waiters)
+	{
+		Selects.Emplace_GetRef() = {
+			((FSocket*)Waiter.Candidate)->Socket,
+			decltype(FSelect::events)(Waiter.bRecvWait ? POLLIN : POLLOUT),
+			/* 0 */
+		};
+	}
+
+	// Poll the sockets
+#if defined(IAS_HTTP_USE_POLL)
+	int32 Result = poll(Selects.GetData(), Selects.Num(), TimeoutMs);
+	if (Result <= 0)
+	{
+		return Result;
+	}
+#else
+	timeval TimeVal = {};
+	timeval* TimeValPtr = (TimeoutMs >= 0 ) ? &TimeVal : nullptr;
+	if (TimeoutMs > 0)
+	{
+		TimeVal = { TimeoutMs >> 10, TimeoutMs & ((1 << 10) - 1) };
+	}
+
+	fd_set FdSetRead;	FD_ZERO(&FdSetRead);
+	fd_set FdSetWrite;	FD_ZERO(&FdSetWrite);
+	fd_set FdSetExcept; FD_ZERO(&FdSetExcept);
+
+	SocketType MaxFd = 0;
+	for (FSelect& Select : Selects)
+	{
+		fd_set* RwSet = (Select.events & POLLIN) ? &FdSetRead : &FdSetWrite;
+		FD_SET(Select.fd, RwSet);
+		FD_SET(Select.fd, &FdSetExcept);
+		MaxFd = FMath::Max(Select.fd, MaxFd);
+	}
+
+	int32 Result = select(int32(MaxFd + 1), &FdSetRead, &FdSetWrite, &FdSetExcept, TimeValPtr);
+	if (Result <= 0)
+	{
+		return Result;
+	}
+
+	for (FSelect& Select : Selects)
+	{
+		if (FD_ISSET(Select.fd, &FdSetExcept))
+		{
+			Select.revents = POLLERR;
+			continue;
+		}
+
+		fd_set* RwSet = (Select.events & POLLIN) ? &FdSetRead : &FdSetWrite;
+		if (FD_ISSET(Select.fd, RwSet))
+		{
+			Select.revents = Select.events;
+		}
+	}
+#endif // IAS_HTTP_USE_POLL
+
+	// Transfer poll results to the input sockets. We don't transfer across error
+	// states. Subsequent sockets ops can take care of that instead.
+	static const auto TestBits = POLLIN|POLLOUT|POLLERR|POLLHUP|POLLNVAL;
+	for (uint32 i = 0, n = Waiters.Num(); i < n; ++i)
+	{
+		if (int32(Selects[i].revents & TestBits) == 0)
+		{
+			continue;
+		}
+
+		Waiters[i].bReady = true;
+	}
+
+	return Result;
+}
+
 
 
 // {{{1 socks ..................................................................
@@ -1662,77 +1778,6 @@ const char* FTicketStatus::GetErrorReason() const
 
 // {{{1 event-loop-int .........................................................
 
-////////////////////////////////////////////////////////////////////////////////
-#if !defined(IAS_HTTP_USE_POLL)
-struct FSelect
-{
-	SocketType	fd;
-	int32		events;
-	int32		revents;
-};
-static const int32 POLLIN  = 1 << 0;
-static const int32 POLLOUT = 1 << 1;
-static const int32 POLLERR = 1 << 2;
-static const int32 POLLHUP = POLLERR;
-static const int32 POLLNVAL= POLLERR;
-#else
-using FSelect = pollfd;
-#endif
-
-////////////////////////////////////////////////////////////////////////////////
-static bool DoSelect(FSelect* Selects, uint32 SelectNum, int32 TimeoutMs)
-{
-#if defined(IAS_HTTP_USE_POLL)
-	return poll(Selects, SelectNum, TimeoutMs) > 0;
-#else
-	timeval TimeVal = {};
-	timeval* TimeValPtr = (TimeoutMs >= 0 ) ? &TimeVal : nullptr;
-	if (TimeoutMs > 0)
-	{
-		TimeVal = { TimeoutMs >> 10, TimeoutMs & ((1 << 10) - 1) };
-	}
-
-	fd_set FdSetRead;	FD_ZERO(&FdSetRead);
-	fd_set FdSetWrite;	FD_ZERO(&FdSetWrite);
-	fd_set FdSetExcept; FD_ZERO(&FdSetExcept);
-
-	SocketType MaxFd = 0;
-	for (uint32 i = 0; i < SelectNum; ++i)
-	{
-		FSelect& Select = Selects[i];
-		fd_set* RwSet = (Select.events & POLLIN) ? &FdSetRead : &FdSetWrite;
-		FD_SET(Select.fd, RwSet);
-		FD_SET(Select.fd, &FdSetExcept);
-		MaxFd = FMath::Max(Select.fd, MaxFd);
-	}
-
-	int32 Result = select(MaxFd + 1, &FdSetRead, &FdSetWrite, &FdSetExcept, TimeValPtr);
-	if (Result == 0)
-	{
-		return false;
-	}
-
-	for (uint32 i = 0; i < SelectNum; ++i)
-	{
-		FSelect& Select = Selects[i];
-
-		if (FD_ISSET(Select.fd, &FdSetExcept))
-		{
-			Select.revents = POLLERR;
-			continue;
-		}
-
-		fd_set* RwSet = (Select.events & POLLIN) ? &FdSetRead : &FdSetWrite;
-		if (FD_ISSET(Select.fd, RwSet))
-		{
-			Select.revents = Select.events;
-		}
-	}
-
-	return true;
-#endif // IAS_HTTP_USE_POLL
-}
-
 ////////////////////////////////////////////////////////////////////////////////
 static uint64 ReadyCheck(FActivity** Activities, uint32 Num, uint32 TimeoutMs)
 {
@@ -1742,8 +1787,7 @@ static uint64 ReadyCheck(FActivity** Activities, uint32 Num, uint32 TimeoutMs)
 
 	uint64 Ret = 0;
 
-	FSelect Selects[64];
-	uint32 SelectNum = 0;
+	TArray<FSocket::FWaiter, TFixedAllocator<64>> Waiters;
 
 	for (uint32 i = 0; i < Num; ++i)
 	{
@@ -1763,52 +1807,37 @@ static uint64 ReadyCheck(FActivity** Activities, uint32 Num, uint32 TimeoutMs)
 			}
 			break;
 
-		case EWait::Read:
-		case EWait::Write: {
-			// The following looks odd because POLLFD varies subtly from one platform
-			// to the next. To cleanly set members to zero and to not get narrowing
-			// warnings from the compiler, we list-init and don't assume POD types.
-			FSelect& Select = Selects[SelectNum];
-			auto SelectEvent = decltype(FSelect::events)((Activity->SocketWait == EWait::Read) ? POLLIN : POLLOUT);
-			Select = {
-				Activity->Socket.Get(),
-				SelectEvent,
-				/* 0 */
-			};
-			++SelectNum;
-			} break;
+		case EWait::Read:  Waiters.Add({ Activity->Socket, true }); break;
+		case EWait::Write: Waiters.Add({ Activity->Socket, false});	break;
 		}
 	}
 
-	// Collect result
-	if (SelectNum == 0)
+	// Early out if there's work to do or nothing to wait on.
+	if (Waiters.IsEmpty())
 	{
 		return Ret;
 	}
 
-	if (!DoSelect(Selects, SelectNum, TimeoutMs))
+	// Wait on sockets and transfer the results
+	if (!FSocket::Wait(Waiters, TimeoutMs))
 	{
 		return Ret;
 	}
 
-	const FSelect* SelectCursor = Selects;
-	for (int32 i = 0; SelectNum > 0; ++i)
+	FActivity* const* Cursor = Activities;
+	for (const FSocket::FWaiter& Waiter : Waiters)
 	{
-		EWait Wait = Activities[i]->SocketWait;
-		if (Wait != EWait::Read && Wait != EWait::Write)
-		{
-			continue;
-		}
+		for (; !(Waiter == Cursor[0]->Socket); ++Cursor);
 
-		if (int32(SelectCursor->revents & (POLLIN|POLLOUT|POLLERR|POLLHUP|POLLNVAL)))
+		if (!Waiter.bReady)
 		{
-			Activities[i]->SocketWait = EWait::None;
-			Ret |= (1ull << Activities[i]->Slot);
-			Trace(Activities[i], ETrace::Unwait);
+			continue;
 		}
 
-		--SelectNum;
-		++SelectCursor;
+		FActivity* Activity = *Cursor;
+		Activity->SocketWait = EWait::None;
+		Ret |= (1ull << Activity->Slot);
+		Trace(Activity, ETrace::Unwait);
 	}
 
 	return Ret;
