Removed now-unused getter()

#rnx

(git:d6d779c - Martin Ridgers - [REDACTED_PHONE]:12:08 +0200)

diff --git a/Engine/Source/Runtime/Experimental/IoStoreOnDemand/Private/CoreHttp/Http.cpp b/Engine/Source/Runtime/Experimental/IoStoreOnDemand/Private/CoreHttp/Http.cpp
--- a/Engine/Source/Runtime/Experimental/IoStoreOnDemand/Private/CoreHttp/Http.cpp
+++ b/Engine/Source/Runtime/Experimental/IoStoreOnDemand/Private/CoreHttp/Http.cpp
@@ -681,8 +681,6 @@ FMessageBuilder& FMessageBuilder::operator << (FAnsiStringView Lhs)
 class FSocket
 {
 public:
-	SocketType	Get() const { return Socket; } // to be removed
-
 	enum class EResult
 	{
 		HangUp			=  0,
