fixed so GC proxy properly applies Parameters.MaterialOverrideMassScaleMultiplier to non-cluster particles on the PT upon initialization

diff --git a/Engine/Source/Runtime/Experimental/Chaos/Private/PhysicsProxy/GeometryCollectionPhysicsProxy.cpp b/Engine/Source/Runtime/Experimental/Chaos/Private/PhysicsProxy/GeometryCollectionPhysicsProxy.cpp
--- a/Engine/Source/Runtime/Experimental/Chaos/Private/PhysicsProxy/GeometryCollectionPhysicsProxy.cpp
+++ b/Engine/Source/Runtime/Experimental/Chaos/Private/PhysicsProxy/GeometryCollectionPhysicsProxy.cpp
@@ -1330,7 +1330,7 @@ void FGeometryCollectionPhysicsProxy::InitializeBodiesPT(Chaos::FPBDRigidsSolver
 					}
 					else
 					{
-						Handle = BuildNonClusters_Internal(TransformGroupIndex, RigidsSolver, Masses[TransformGroupIndex], Inertias[TransformGroupIndex]);
+						Handle = BuildNonClusters_Internal(TransformGroupIndex, RigidsSolver, AdjustMassForScale(Masses[TransformGroupIndex]), Inertias[TransformGroupIndex]);
 					}
 					Handle->GTGeometryParticle() = GTParticle;
 
