Chaos Cloth: Add more controls for self intersections and improve global intersection analysis to detect and floodfill closed regions that include boundaries. This allows us to better resolve intersections that happen along edges.

diff --git a/Engine/Plugins/Experimental/ChaosCloth/Source/Chaos/Private/ChaosCloth/ChaosClothConstraints.cpp b/Engine/Plugins/Experimental/ChaosCloth/Source/Chaos/Private/ChaosCloth/ChaosClothConstraints.cpp
--- a/Engine/Plugins/Experimental/ChaosCloth/Source/Chaos/Private/ChaosCloth/ChaosClothConstraints.cpp
+++ b/Engine/Plugins/Experimental/ChaosCloth/Source/Chaos/Private/ChaosCloth/ChaosClothConstraints.cpp
@@ -34,9 +34,11 @@ FClothConstraints::FClothConstraints()
 	, ConstraintInitOffset(INDEX_NONE)
 	, ConstraintRuleOffset(INDEX_NONE)
 	, PostCollisionConstraintRuleOffset(INDEX_NONE)
+	, PostprocessingConstraintRuleOffset(INDEX_NONE)
 	, NumConstraintInits(0)
 	, NumConstraintRules(0)
 	, NumPostCollisionConstraintRules(0)
+	, NumPostprocessingConstraintRules(0)
 {
 }
 
@@ -96,6 +98,10 @@ void FClothConstraints::Enable(bool bEnable)
 	{
 		Evolution->ActivatePostCollisionConstraintRuleRange(PostCollisionConstraintRuleOffset, bEnable);
 	}
+	if (PostprocessingConstraintRuleOffset != INDEX_NONE)
+	{
+		Evolution->ActivateConstraintPostprocessingsRange(PostprocessingConstraintRuleOffset, bEnable);
+	}
 }
 
 void FClothConstraints::AddRules(
@@ -188,8 +194,6 @@ void FClothConstraints::CreateSelfCollisionConstraints(const Softs::FCollectionP
 			}
 		}
 
-		const bool bUseSelfIntersections = ConfigProperties.GetValue<bool>(TEXT("UseSelfIntersections"));
-
 		SelfCollisionInit = MakeShared<Softs::FPBDTriangleMeshCollisions>(
 			ParticleOffset,
 			NumParticles,
@@ -212,6 +216,7 @@ void FClothConstraints::CreateSelfCollisionConstraints(const Softs::FCollectionP
 			NumParticles,
 			TriangleMesh);
 		++NumConstraintInits;
+		++NumPostprocessingConstraintRules;
 	}
 }
 
@@ -510,14 +515,23 @@ void FClothConstraints::CreateRules()
 	{
 		PostCollisionConstraintRuleOffset = Evolution->AddPostCollisionConstraintRuleRange(NumPostCollisionConstraintRules, false);
 	}
+	check(PostprocessingConstraintRuleOffset == INDEX_NONE);
+	if (NumPostprocessingConstraintRules)
+	{
+		PostprocessingConstraintRuleOffset = Evolution->AddConstraintPostprocessingsRange(NumPostprocessingConstraintRules, false);
+	}
 
 	TFunction<void(Softs::FSolverParticles&, const Softs::FSolverReal)>* const ConstraintInits = Evolution->ConstraintInits().GetData() + ConstraintInitOffset;
 	TFunction<void(Softs::FSolverParticles&, const Softs::FSolverReal)>* const ConstraintRules = Evolution->ConstraintRules().GetData() + ConstraintRuleOffset;
 	TFunction<void(Softs::FSolverParticles&, const Softs::FSolverReal)>* const PostCollisionConstraintRules = Evolution->PostCollisionConstraintRules().GetData() + PostCollisionConstraintRuleOffset;
+	TFunction<void(Softs::FSolverParticles&, const Softs::FSolverReal)>* const
+		PostprocessingConstraintRules = Evolution->ConstraintPostprocessings().GetData() +
+		PostprocessingConstraintRuleOffset;
 
 	int32 ConstraintInitIndex = 0;
 	int32 ConstraintRuleIndex = 0;
 	int32 PostCollisionConstraintRuleIndex = 0;
+	int32 PostprocessingConstraintRuleIndex = 0;
 
 	if (XStretchBiasConstraints)
 	{
@@ -810,17 +824,29 @@ void FClothConstraints::CreateRules()
 			{
 				SelfCollisionConstraints->Apply(Particles, Dt);
 			};
+
 	}
 
-	// The following constraints only run once per subframe, so we do their Apply as part of the Init() which modifies P
-	// To avoid possible dependency order issues, add them last
 	if (SelfCollisionInit && SelfIntersectionConstraints)
 	{
+		// The following constraints only run once per subframe, so we do their Apply as part of the Init() which modifies P
+		// To avoid possible dependency order issues, add them last
 		ConstraintInits[ConstraintInitIndex++] =
 			[this](Softs::FSolverParticles& Particles, const Softs::FSolverReal Dt)
 			{
 				SelfIntersectionConstraints->Apply(Particles, SelfCollisionInit->GetContourMinimizationIntersections(), Dt);
 			};
+
+		PostprocessingConstraintRules[PostprocessingConstraintRuleIndex++] =
+			[this](Softs::FSolverParticles& Particles, const Softs::FSolverReal Dt)
+		{
+			const int32 NumContourIterations = SelfCollisionInit->GetNumContourMinimizationPostSteps();
+			for (int32 Iter = 0; Iter < NumContourIterations; ++Iter)
+			{
+				SelfCollisionInit->PostStepInit(Particles);
+				SelfIntersectionConstraints->Apply(Particles, SelfCollisionInit->GetPostStepContourMinimizationIntersections(), Dt);
+			}
+		};
 	}
 
 	// Long range constraints modify particle P as part of Init. To avoid possible dependency order issues,
@@ -1253,6 +1279,10 @@ void FClothConstraints::Update(
 	{
 		SelfCollisionConstraints->SetProperties(ConfigProperties);
 	}
+	if (SelfCollisionInit)
+	{
+		SelfCollisionInit->SetProperties(ConfigProperties);
+	}
 }
 
 // Deprecated
diff --git a/Engine/Plugins/Experimental/ChaosCloth/Source/Chaos/Private/ChaosCloth/ChaosClothVisualization.cpp b/Engine/Plugins/Experimental/ChaosCloth/Source/Chaos/Private/ChaosCloth/ChaosClothVisualization.cpp
--- a/Engine/Plugins/Experimental/ChaosCloth/Source/Chaos/Private/ChaosCloth/ChaosClothVisualization.cpp
+++ b/Engine/Plugins/Experimental/ChaosCloth/Source/Chaos/Private/ChaosCloth/ChaosClothVisualization.cpp
@@ -2013,7 +2013,10 @@ static FAutoConsoleVariableRef CVarClothVizWeightMapName(TEXT("p.ChaosClothVisua
 			static const FLinearColor White(1.f, 1.f, 1.f);
 			static const FLinearColor Black(0.f, 0.f, 0.f);
 			static const FLinearColor Teal(0.f, 0.5f, 0.5f);
+			static const FLinearColor Orange(1.f, .5f, 0.f);
 			static const FLinearColor Green(0.f, 1.f, 0.f);
+			static const FLinearColor Yellow(1.f, 1.f, 0.f);
+			static const FLinearColor Blue(0.f, 0.f, 1.f);
 
 			if (const Softs::FPBDTriangleMeshCollisions* const SelfCollisionInit = ClothConstraints.GetSelfCollisionInit().Get())
 			{
@@ -2022,6 +2025,7 @@ static FAutoConsoleVariableRef CVarClothVizWeightMapName(TEXT("p.ChaosClothVisua
 
 				// Draw contours
 				const TArray<TArray<Softs::FPBDTriangleMeshCollisions::FBarycentricPoint>>& ContourPoints = SelfCollisionInit->GetIntersectionContourPoints();
+				const TArray<TArray<Softs::FPBDTriangleMeshCollisions::FBarycentricPoint>>& PostStepContourPoints = SelfCollisionInit->GetPostStepIntersectionContourPoints();
 				const TArray<Softs::FPBDTriangleMeshCollisions::FContourType>& ContourTypes = SelfCollisionInit->GetIntersectionContourTypes();
 				check(ContourPoints.Num() == ContourTypes.Num());
 
@@ -2029,23 +2033,38 @@ static FAutoConsoleVariableRef CVarClothVizWeightMapName(TEXT("p.ChaosClothVisua
 				{
 					Teal,
 					Red,
+					Blue,
+					Yellow,
 					White,
 					Black
 				};
-				for( int32 ContourIndex = 0; ContourIndex < ContourPoints.Num(); ++ContourIndex)
+
+				auto DrawContour = [&LocalSpaceLocation, PDI, &Positions](const TArray<Softs::FPBDTriangleMeshCollisions::FBarycentricPoint>& Contour,
+					const FLinearColor& ContourColor)
 				{
-					const TArray<Softs::FPBDTriangleMeshCollisions::FBarycentricPoint>& Contour = ContourPoints[ContourIndex];
-					const FLinearColor& ContourColor = ColorsForType[(int8)ContourTypes[ContourIndex]];
 					for (int32 PointIdx = 0; PointIdx < Contour.Num() - 1; ++PointIdx)
 					{
 						const Softs::FPBDTriangleMeshCollisions::FBarycentricPoint& Point0 = Contour[PointIdx];
 						const FVector EndPoint0 = LocalSpaceLocation + (1.f - Point0.Bary[0] - Point0.Bary[1]) * Positions[Point0.Vertices[0]] + Point0.Bary[0] * Positions[Point0.Vertices[1]] + Point0.Bary[1] * Positions[Point0.Vertices[2]];
-						const Softs::FPBDTriangleMeshCollisions::FBarycentricPoint& Point1 = Contour[PointIdx+1];
+						const Softs::FPBDTriangleMeshCollisions::FBarycentricPoint& Point1 = Contour[PointIdx + 1];
 						const FVector EndPoint1 = LocalSpaceLocation + (1.f - Point1.Bary[0] - Point1.Bary[1]) * Positions[Point1.Vertices[0]] + Point1.Bary[0] * Positions[Point1.Vertices[1]] + Point1.Bary[1] * Positions[Point1.Vertices[2]];
 						DrawLine(PDI, EndPoint0, EndPoint1, ContourColor);
 						DrawPoint(PDI, EndPoint0, ContourColor, nullptr, 1.f);
 						DrawPoint(PDI, EndPoint1, ContourColor, nullptr, 1.f);
 					}
+
+				};
+
+				for( int32 ContourIndex = 0; ContourIndex < ContourPoints.Num(); ++ContourIndex)
+				{
+					const TArray<Softs::FPBDTriangleMeshCollisions::FBarycentricPoint>& Contour = ContourPoints[ContourIndex];
+					const FLinearColor& ContourColor = ColorsForType[(int8)ContourTypes[ContourIndex]];
+					DrawContour(Contour, ContourColor);
+				}
+				for (int32 ContourIndex = 0; ContourIndex < PostStepContourPoints.Num(); ++ContourIndex)
+				{
+					const TArray<Softs::FPBDTriangleMeshCollisions::FBarycentricPoint>& Contour = PostStepContourPoints[ContourIndex];
+					DrawContour(Contour, Orange);
 				}
 
 				// Draw GIA colors
@@ -2059,9 +2078,10 @@ static FAutoConsoleVariableRef CVarClothVizWeightMapName(TEXT("p.ChaosClothVisua
 						if (VertexGIAColors[ParticleIdx].ContourIndexBits)
 						{
 							const bool bIsLoop = VertexGIAColors[ParticleIdx].IsLoop();
+							const bool bIsBoundary = VertexGIAColors[ParticleIdx].IsBoundary();
 							const bool bAnyWhite = (VertexGIAColors[ParticleIdx].ContourIndexBits & ~VertexGIAColors[ParticleIdx].ColorBits);
 							const bool bAnyBlack = (VertexGIAColors[ParticleIdx].ContourIndexBits & VertexGIAColors[ParticleIdx].ColorBits);
-							const FLinearColor& VertColor = bIsLoop ? Red : (bAnyWhite && bAnyBlack) ? Gray : bAnyWhite ? White : Black;
+							const FLinearColor& VertColor = bIsLoop ? Red : bIsBoundary ? Blue : (bAnyWhite && bAnyBlack) ? Gray : bAnyWhite ? White : Black;
 						
 							DrawPoint(PDI, LocalSpaceLocation + Positions[ParticleIdx], VertColor, nullptr, 5.f);
 						}
diff --git a/Engine/Plugins/Experimental/ChaosCloth/Source/Chaos/Public/ChaosCloth/ChaosClothConstraints.h b/Engine/Plugins/Experimental/ChaosCloth/Source/Chaos/Public/ChaosCloth/ChaosClothConstraints.h
--- a/Engine/Plugins/Experimental/ChaosCloth/Source/Chaos/Public/ChaosCloth/ChaosClothConstraints.h
+++ b/Engine/Plugins/Experimental/ChaosCloth/Source/Chaos/Public/ChaosCloth/ChaosClothConstraints.h
@@ -272,8 +272,10 @@ PRAGMA_ENABLE_DEPRECATION_WARNINGS
 		int32 ConstraintInitOffset;
 		int32 ConstraintRuleOffset;
 		int32 PostCollisionConstraintRuleOffset;
+		int32 PostprocessingConstraintRuleOffset;
 		int32 NumConstraintInits;
 		int32 NumConstraintRules;
 		int32 NumPostCollisionConstraintRules;
+		int32 NumPostprocessingConstraintRules;
 	};
 } // namespace Chaos
diff --git a/Engine/Plugins/Experimental/ChaosClothAssetEditor/Source/ChaosClothAssetDataflowNodes/Private/ChaosClothAsset/SimulationSelfCollisionConfigNode.cpp b/Engine/Plugins/Experimental/ChaosClothAssetEditor/Source/ChaosClothAssetDataflowNodes/Private/ChaosClothAsset/SimulationSelfCollisionConfigNode.cpp
--- a/Engine/Plugins/Experimental/ChaosClothAssetEditor/Source/ChaosClothAssetDataflowNodes/Private/ChaosClothAsset/SimulationSelfCollisionConfigNode.cpp
+++ b/Engine/Plugins/Experimental/ChaosClothAssetEditor/Source/ChaosClothAssetDataflowNodes/Private/ChaosClothAsset/SimulationSelfCollisionConfigNode.cpp
@@ -19,5 +19,10 @@ void FChaosClothAssetSimulationSelfCollisionConfigNode::AddProperties(Dataflow::
 	UE_CHAOS_CLOTHASSET_SIMULATIONCONFIG_SETPROPERTY(SelfCollisionThickness);
 	UE_CHAOS_CLOTHASSET_SIMULATIONCONFIG_SETPROPERTY(SelfCollisionStiffness);
 	UE_CHAOS_CLOTHASSET_SIMULATIONCONFIG_SETPROPERTY(SelfCollisionFriction);
+
 	UE_CHAOS_CLOTHASSET_SIMULATIONCONFIG_SETPROPERTYBOOL(UseSelfIntersections);
+	UE_CHAOS_CLOTHASSET_SIMULATIONCONFIG_SETPROPERTYBOOL(UseGlobalIntersectionAnalysis);
+	UE_CHAOS_CLOTHASSET_SIMULATIONCONFIG_SETPROPERTYBOOL(UseContourMinimization);
+	UE_CHAOS_CLOTHASSET_SIMULATIONCONFIG_SETPROPERTY(NumContourMinimizationPostSteps);
+	UE_CHAOS_CLOTHASSET_SIMULATIONCONFIG_SETPROPERTYBOOL(UseGlobalPostStepContours);
 }
diff --git a/Engine/Plugins/Experimental/ChaosClothAssetEditor/Source/ChaosClothAssetDataflowNodes/Public/ChaosClothAsset/SimulationSelfCollisionConfigNode.h b/Engine/Plugins/Experimental/ChaosClothAssetEditor/Source/ChaosClothAssetDataflowNodes/Public/ChaosClothAsset/SimulationSelfCollisionConfigNode.h
--- a/Engine/Plugins/Experimental/ChaosClothAssetEditor/Source/ChaosClothAssetDataflowNodes/Public/ChaosClothAsset/SimulationSelfCollisionConfigNode.h
+++ b/Engine/Plugins/Experimental/ChaosClothAssetEditor/Source/ChaosClothAssetDataflowNodes/Public/ChaosClothAsset/SimulationSelfCollisionConfigNode.h
@@ -18,17 +18,33 @@ public:
 	float SelfCollisionThickness = 0.5f;
 
 	/** The stiffness of the springs used to control self collision. */
-	UPROPERTY(EditAnywhere, Category = "Self-Collision Properties", meta = (UIMin = "0", UIMax = "100", ClampMin = "0", ClampMax = "1000", EditCondition = "bUseSelfCollisions"))
+	UPROPERTY(EditAnywhere, Category = "Self-Collision Properties", meta = (UIMin = "0", UIMax = "100", ClampMin = "0", ClampMax = "1000"))
 	float SelfCollisionStiffness = 0.5f;
 
 	/** Friction coefficient for cloth - cloth interaction. */
-	UPROPERTY(EditAnywhere, Category = "Self-Collision Properties", meta = (UIMin = "0", UIMax = "1", ClampMin = "0", ClampMax = "10", EditCondition = "bUseSelfCollisions"))
+	UPROPERTY(EditAnywhere, Category = "Self-Collision Properties", meta = (UIMin = "0", UIMax = "1", ClampMin = "0", ClampMax = "10"))
 	float SelfCollisionFriction = 0.0f;
 
 	/** Enable self intersection resolution. This will try to fix any cloth intersections that are not handled by collision repulsions. */
-	UPROPERTY(EditAnywhere, Category = "Self-Collision Properties", meta = (EditCondition = "bUseSelfCollisions"))
+	UPROPERTY(EditAnywhere, Category = "Experimental")
 	bool bUseSelfIntersections = false;
 
+	/** Do global intersection analysis to determine the correct normals for the collision springs */
+	UPROPERTY(EditAnywhere, Category = "Experimental", meta = (EditCondition = "bUseSelfIntersections"))
+	bool bUseGlobalIntersectionAnalysis = true;
+
+	/** Do a step of contour minimization at the beginning of the timestep. */
+	UPROPERTY(EditAnywhere, Category = "Experimental", meta = (EditCondition = "bUseSelfIntersections"))
+	bool bUseContourMinimization = true;
+
+	/** Number of post timestep contour minimization steps to do. (Expensive!)*/
+	UPROPERTY(EditAnywhere, Category = "Experimental", meta = (ClampMin = "0", EditCondition = "bUseSelfIntersections"))
+	int32 NumContourMinimizationPostSteps = 0;
+
+	/** Use global contour gradients when doing post timestep contour minimization */
+	UPROPERTY(EditAnywhere, Category = "Experimental", meta = (EditCondition = "bUseSelfIntersections && NumContourMinimizationPostSteps > 0"))
+	bool bUseGlobalPostStepContours = true;
+
 	FChaosClothAssetSimulationSelfCollisionConfigNode(const Dataflow::FNodeParameters& InParam, FGuid InGuid = FGuid::NewGuid());
 
 private:
diff --git a/Engine/Source/Runtime/Experimental/Chaos/Private/Chaos/PBDCollisionSpringConstraintsBase.cpp b/Engine/Source/Runtime/Experimental/Chaos/Private/Chaos/PBDCollisionSpringConstraintsBase.cpp
--- a/Engine/Source/Runtime/Experimental/Chaos/Private/Chaos/PBDCollisionSpringConstraintsBase.cpp
+++ b/Engine/Source/Runtime/Experimental/Chaos/Private/Chaos/PBDCollisionSpringConstraintsBase.cpp
@@ -89,9 +89,27 @@ void FPBDCollisionSpringConstraintsBase::Init(const FSolverParticles& Particles,
 					[this, &VertexGIAColors, &TriangleGIAColors](const int32 PointIndex, const int32 TriangleIndex)->bool
 					{
 						const TVector<int32, 3>& Elem = Elements[TriangleIndex];
-						if (bGlobalIntersectionAnalysis && VertexGIAColors[PointIndex].IsLoop() && (VertexGIAColors[Elem[0]].IsLoop() || VertexGIAColors[Elem[1]].IsLoop() || VertexGIAColors[Elem[2]].IsLoop() || TriangleGIAColors[TriangleIndex].IsLoop()))
+						if (bGlobalIntersectionAnalysis)
 						{
-							return false;
+							const bool bIsAnyBoundary = VertexGIAColors[PointIndex].IsBoundary()
+								|| VertexGIAColors[Elem[0]].IsBoundary() 
+								|| VertexGIAColors[Elem[1]].IsBoundary() 
+								|| VertexGIAColors[Elem[2]].IsBoundary();
+							if (bIsAnyBoundary)
+							{
+								return false;
+							}
+
+							const bool bAreBothLoop = VertexGIAColors[PointIndex].IsLoop() &&
+								(VertexGIAColors[Elem[0]].IsLoop() 
+									|| VertexGIAColors[Elem[1]].IsLoop() 
+									|| VertexGIAColors[Elem[2]].IsLoop() 
+									|| TriangleGIAColors[TriangleIndex].IsLoop());
+
+							if (bAreBothLoop)
+							{
+								return false;
+							}
 						}
 
 						if (DisabledCollisionElements.Contains({ PointIndex, Elem[0] }) ||
diff --git a/Engine/Source/Runtime/Experimental/Chaos/Private/Chaos/PBDTriangleMeshCollisions.cpp b/Engine/Source/Runtime/Experimental/Chaos/Private/Chaos/PBDTriangleMeshCollisions.cpp
--- a/Engine/Source/Runtime/Experimental/Chaos/Private/Chaos/PBDTriangleMeshCollisions.cpp
+++ b/Engine/Source/Runtime/Experimental/Chaos/Private/Chaos/PBDTriangleMeshCollisions.cpp
@@ -120,17 +120,32 @@ namespace GIA
 		{}
 	};
 
-	using FIntersectionContour = TArray<FIntersectionContourTriangleSection>;
+	struct FIntersectionContour
+	{
+		int8 BoundaryEdgeCount = 0;
+		TArray<FIntersectionContourTriangleSection> Contour;		
+	};
 
 	struct FIntersectionContourPair
 	{
-		bool bIsClosed = false;
-		int8 LoopVertexCount = 0;
-		int8 BoundaryEdgeCount = 0;
+		enum struct EClosedStatus : uint8
+		{
+			Open, // Can't color.
+			SimpleClosed, // Expect to color both contours separately.
+			LoopClosed, // Expect to have a single colorable contour.
+			BoundaryClosed // Expect to have one colorable contour (two boundary ends), and one open contour.
+		};
+		EClosedStatus ClosedStatus = EClosedStatus::Open;
+		int8 LoopVertexCount = 0; // Shared between ColorContours
 
 		FIntersectionContour ColorContours[2];
 		TArray<FEdgeFaceIntersection> Intersections;
 		TArray<int32> ContourPointCurves[2]; // Track which curves correspond with which Contours. (Sometimes we generate multiple curves per contour)
+
+		int8 NumContourEnds() const
+		{
+			return LoopVertexCount + ColorContours[0].BoundaryEdgeCount + ColorContours[1].BoundaryEdgeCount;
+		}
 	};
 
 	// methods for building contours
@@ -271,16 +286,16 @@ namespace GIA
 				check(FirstEdgeFace != -1); // Each Edge should be connected to at least one face
 
 				// We will rely on finding the First FaceSection and EdgeFaceSection as the [0] element in these arrays.
-				check(FaceContour->Num() == 0);
-				check(EdgeContour->Num() == 0);
+				check(FaceContour->Contour.Num() == 0);
+				check(EdgeContour->Contour.Num() == 0);
 
 
 				// Start main loop to consume intersections to build this contour
 				bool bReverseDirection = false; // when not reversing, write to CrossingEdgeLocalIndex[0] first. when reversing, writing to CrossingEdgeLocalIndex[1] first.
 				FEdgeFaceIntersection CurrIntersection = FirstIntersection;
 				int32 EdgeFace = FirstEdgeFace;
-				FIntersectionContourTriangleSection* FaceSection = &FaceContour->Add_GetRef(FIntersectionContourTriangleSection(FirstIntersection.FaceIndex));
-				FIntersectionContourTriangleSection* EdgeFaceSection = &EdgeContour->Add_GetRef(FIntersectionContourTriangleSection(FirstEdgeFace));
+				FIntersectionContourTriangleSection* FaceSection = &FaceContour->Contour.Add_GetRef(FIntersectionContourTriangleSection(FirstIntersection.FaceIndex));
+				FIntersectionContourTriangleSection* EdgeFaceSection = &EdgeContour->Contour.Add_GetRef(FIntersectionContourTriangleSection(FirstEdgeFace));
 
 				// Setup first EdgeFaceSection crossing
 				EdgeFaceSection->CrossingEdgeLocalIndex[0] = __internal::GetLocalEdgeIndex(FaceToEdges, FirstEdgeFace, FirstIntersection.EdgeIndex);
@@ -314,20 +329,23 @@ namespace GIA
 						EdgeFaceSection->LoopVertexLocalIndex = (EdgeFaceCrossingEdgeLocalIndex + 2) % 3;
 						++ContourPair.LoopVertexCount;
 
-						if (ContourPair.LoopVertexCount + ContourPair.BoundaryEdgeCount == 2)
+						if (ContourPair.NumContourEnds() == 2)
 						{
 							// We've found both ends of this contour.
-							ContourPair.bIsClosed = ContourPair.LoopVertexCount == 2;
+							if (ContourPair.LoopVertexCount == 2)
+							{
+								ContourPair.ClosedStatus = FIntersectionContourPair::EClosedStatus::LoopClosed;
+							}
 							break;
 						}
-						check(ContourPair.LoopVertexCount + ContourPair.BoundaryEdgeCount == 1);
+						check(ContourPair.NumContourEnds() == 1);
 						// We hit one loop vertex. Pick up at the FirstIntersection and move the opposite direction if possible.
 
 						EdgeFace = EdgeToFaces[FirstIntersection.EdgeIndex][1];
 						if (EdgeFace == -1)
 						{
 							// First intersection edge was a boundary, so this contour is done. We're NOT closed.
-							++ContourPair.BoundaryEdgeCount;
+							++ContourPair.ColorContours[InitialEdgeContourIndex].BoundaryEdgeCount;
 							break;
 						}
 
@@ -343,8 +361,8 @@ namespace GIA
 						ContourPair.ContourPointCurves[InitialEdgeContourIndex].Add(ContourPoints.Num());
 						EdgeContourPoints = &ContourPoints.AddDefaulted_GetRef();
 
-						FaceSection = &ContourPair.ColorContours[InitialFaceContourIndex][0];
-						EdgeFaceSection = &EdgeContour->Add_GetRef(FIntersectionContourTriangleSection(EdgeFace));
+						FaceSection = &FaceContour->Contour[0];
+						EdgeFaceSection = &EdgeContour->Contour.Add_GetRef(FIntersectionContourTriangleSection(EdgeFace));
 						EdgeFaceSection->CrossingEdgeLocalIndex[1] = __internal::GetLocalEdgeIndex(FaceToEdges, EdgeFace, CurrIntersection.EdgeIndex);
 						continue;
 					}
@@ -356,7 +374,7 @@ namespace GIA
 						// We can close the contour
 						EdgeFaceSection->CrossingEdgeLocalIndex[1] = LocalEdgeIndex;
 
-						FIntersectionContourTriangleSection& FirstFaceSection = ContourPair.ColorContours[InitialFaceContourIndex][0];
+						FIntersectionContourTriangleSection& FirstFaceSection = ContourPair.ColorContours[InitialFaceContourIndex].Contour[0];
 
 						// Merge current FaceSection and FirstFaceSection (if they're not already the same)
 						check(FaceContour == &ContourPair.ColorContours[InitialFaceContourIndex]);
@@ -365,11 +383,11 @@ namespace GIA
 						{
 							check(FaceSection->CrossingEdgeLocalIndex[0] != INDEX_NONE);
 							FirstFaceSection.CrossingEdgeLocalIndex[0] = FaceSection->CrossingEdgeLocalIndex[0];
-							check(FaceSection == &FaceContour->Last());
-							FaceContour->RemoveAt(FaceContour->Num() - 1, 1, false);
+							check(FaceSection == &FaceContour->Contour.Last());
+							FaceContour->Contour.RemoveAt(FaceContour->Contour.Num() - 1, 1, false);
 						}
 
-						ContourPair.bIsClosed = true;
+						ContourPair.ClosedStatus = FIntersectionContourPair::EClosedStatus::SimpleClosed;
 
 						// Repeat first point in contour points for ease of drawing closed loop
 						FaceContourPoints->Add(FPBDTriangleMeshCollisions::FBarycentricPoint((*FaceContourPoints)[0]));
@@ -384,16 +402,16 @@ namespace GIA
 						FaceSection->CrossingEdgeLocalIndex[1] = LocalEdgeIndex;
 
 						// Merge current EdgeFaceSection with FirstFaceSection
-						FIntersectionContourTriangleSection& FirstFaceSection = ContourPair.ColorContours[InitialFaceContourIndex][0];
+						FIntersectionContourTriangleSection& FirstFaceSection = ContourPair.ColorContours[InitialFaceContourIndex].Contour[0];
 						check(EdgeContour == &ContourPair.ColorContours[InitialFaceContourIndex]);
 						check(EdgeFaceSection->TriangleIndex == FirstFaceSection.TriangleIndex);
 						check(EdgeFaceSection->CrossingEdgeLocalIndex[0] != INDEX_NONE);
 						FirstFaceSection.CrossingEdgeLocalIndex[0] = EdgeFaceSection->CrossingEdgeLocalIndex[0];
 
-						check(EdgeFaceSection == &EdgeContour->Last());
-						EdgeContour->RemoveAt(EdgeContour->Num() - 1, 1, false);
+						check(EdgeFaceSection == &EdgeContour->Contour.Last());
+						EdgeContour->Contour.RemoveAt(EdgeContour->Contour.Num() - 1, 1, false);
 
-						ContourPair.bIsClosed = true;
+						ContourPair.ClosedStatus = FIntersectionContourPair::EClosedStatus::SimpleClosed;
 
 						// Repeat first point in contour points for ease of drawing closed loop
 						FaceContourPoints->Add(FPBDTriangleMeshCollisions::FBarycentricPoint((*FaceContourPoints)[0]));
@@ -421,30 +439,40 @@ namespace GIA
 						if (NextEdgeFace == -1)
 						{
 							// We've hit a boundary.
-							++ContourPair.BoundaryEdgeCount;
+							++EdgeContour->BoundaryEdgeCount;
 
-							if (ContourPair.LoopVertexCount + ContourPair.BoundaryEdgeCount == 2)
+							if (ContourPair.NumContourEnds() == 2)
 							{
-								// We've hit both ends of this contour. Since (at least) one end is a boundary, we won't be able to form a closed contour
+								// We've hit both ends of this contour.
+								if (EdgeContour->BoundaryEdgeCount == 2)
+								{
+									// This contour is "closed" by hitting boundaries on both sides
+									ContourPair.ClosedStatus = FIntersectionContourPair::EClosedStatus::BoundaryClosed;
+								}
 								break;
 							}
 
-							check(ContourPair.LoopVertexCount + ContourPair.BoundaryEdgeCount == 1);
+							check(ContourPair.NumContourEnds() == 1);
 
 							// Pick up at the FirstIntersection and move the opposite direction if possible.
 							EdgeFace = EdgeToFaces[FirstIntersection.EdgeIndex][1];
 							if (EdgeFace == -1)
 							{
-								// First intersection edge was a boundary, so this contour is done. We're NOT closed.
-								++ContourPair.BoundaryEdgeCount;
+								// First intersection edge was a boundary, so this contour is done.
+								++ContourPair.ColorContours[InitialEdgeContourIndex].BoundaryEdgeCount;
+								if (EdgeContour->BoundaryEdgeCount == 2)
+								{
+									// This contour is "closed" by hitting boundaries on both sides
+									ContourPair.ClosedStatus = FIntersectionContourPair::EClosedStatus::BoundaryClosed;
+								}
 								break;
 							}
 
 							check(bReverseDirection == false);
 							bReverseDirection = true;
 							CurrIntersection = FirstIntersection;
-							FaceContour = &ContourPair.ColorContours[0];
-							EdgeContour = &ContourPair.ColorContours[1];
+							FaceContour = &ContourPair.ColorContours[InitialFaceContourIndex];
+							EdgeContour = &ContourPair.ColorContours[InitialEdgeContourIndex];
 
 							// ContourPoints are currently just used for debug drawing. Start new contours for reverse section
 							ContourPair.ContourPointCurves[InitialFaceContourIndex].Add(ContourPoints.Num());
@@ -452,8 +480,8 @@ namespace GIA
 							ContourPair.ContourPointCurves[InitialEdgeContourIndex].Add(ContourPoints.Num());
 							EdgeContourPoints = &ContourPoints.AddDefaulted_GetRef();
 
-							FaceSection = &ContourPair.ColorContours[InitialFaceContourIndex][0];
-							EdgeFaceSection = &EdgeContour->Add_GetRef(FIntersectionContourTriangleSection(EdgeFace));
+							FaceSection = &ContourPair.ColorContours[InitialFaceContourIndex].Contour[0];
+							EdgeFaceSection = &EdgeContour->Contour.Add_GetRef(FIntersectionContourTriangleSection(EdgeFace));
 							EdgeFaceSection->CrossingEdgeLocalIndex[1] = __internal::GetLocalEdgeIndex(FaceToEdges, EdgeFace, CurrIntersection.EdgeIndex);
 
 							continue;
@@ -464,7 +492,7 @@ namespace GIA
 						CurrIntersection = NextIntersection;
 						EdgeFace = NextEdgeFace;
 
-						EdgeFaceSection = &EdgeContour->Add_GetRef(FIntersectionContourTriangleSection(EdgeFace));
+						EdgeFaceSection = &EdgeContour->Contour.Add_GetRef(FIntersectionContourTriangleSection(EdgeFace));
 						int32& FirstCrossingLocalIndex = bReverseDirection ? EdgeFaceSection->CrossingEdgeLocalIndex[1] : EdgeFaceSection->CrossingEdgeLocalIndex[0];
 						FirstCrossingLocalIndex = __internal::GetLocalEdgeIndex(FaceToEdges, EdgeFace, CurrIntersection.EdgeIndex);
 						continue;
@@ -487,22 +515,32 @@ namespace GIA
 						if (NextEdgeFace == -1)
 						{
 							// We've hit a boundary.
-							++ContourPair.BoundaryEdgeCount;
+							++FaceContour->BoundaryEdgeCount;
 
-							if (ContourPair.LoopVertexCount + ContourPair.BoundaryEdgeCount == 2)
+							if (ContourPair.NumContourEnds() == 2)
 							{
-								// We've hit both ends of this contour. Since (at least) one end is a boundary, we won't be able to form a closed contour
+								// We've hit both ends of this contour. 
+								if (FaceContour->BoundaryEdgeCount == 2)
+								{
+									// This contour is "closed" by hitting boundaries on both sides
+									ContourPair.ClosedStatus = FIntersectionContourPair::EClosedStatus::BoundaryClosed;
+								}
 								break;
 							}
 
-							check(ContourPair.LoopVertexCount + ContourPair.BoundaryEdgeCount == 1);
+							check(ContourPair.NumContourEnds() == 1);
 
 							// Pick up at the FirstIntersection and move the opposite direction if possible.
 							EdgeFace = EdgeToFaces[FirstIntersection.EdgeIndex][1];
 							if (EdgeFace == -1)
 							{
-								// First intersection edge was a boundary, so this contour is done. We're NOT closed.
-								++ContourPair.BoundaryEdgeCount;
+								// First intersection edge was a boundary.
+								++ContourPair.ColorContours[InitialEdgeContourIndex].BoundaryEdgeCount;
+								if (FaceContour->BoundaryEdgeCount == 2)
+								{
+									// This contour is "closed" by hitting boundaries on both sides
+									ContourPair.ClosedStatus = FIntersectionContourPair::EClosedStatus::BoundaryClosed;
+								}
 								break;
 							}
 
@@ -518,8 +556,8 @@ namespace GIA
 							ContourPair.ContourPointCurves[InitialEdgeContourIndex].Add(ContourPoints.Num());
 							EdgeContourPoints = &ContourPoints.AddDefaulted_GetRef();
 
-							FaceSection = &ContourPair.ColorContours[InitialFaceContourIndex][0];
-							EdgeFaceSection = &EdgeContour->Add_GetRef(FIntersectionContourTriangleSection(EdgeFace));
+							FaceSection = &ContourPair.ColorContours[InitialFaceContourIndex].Contour[0];
+							EdgeFaceSection = &EdgeContour->Contour.Add_GetRef(FIntersectionContourTriangleSection(EdgeFace));
 							EdgeFaceSection->CrossingEdgeLocalIndex[1] = __internal::GetLocalEdgeIndex(FaceToEdges, EdgeFace, CurrIntersection.EdgeIndex);
 
 							continue;
@@ -539,7 +577,7 @@ namespace GIA
 						FaceSection = EdgeFaceSection;
 
 						EdgeFace = NextEdgeFace;
-						EdgeFaceSection = &EdgeContour->Add_GetRef(FIntersectionContourTriangleSection(EdgeFace));
+						EdgeFaceSection = &EdgeContour->Contour.Add_GetRef(FIntersectionContourTriangleSection(EdgeFace));
 						int32& FirstCrossingLocalIndex = bReverseDirection ? EdgeFaceSection->CrossingEdgeLocalIndex[1] : EdgeFaceSection->CrossingEdgeLocalIndex[0];
 						FirstCrossingLocalIndex = __internal::GetLocalEdgeIndex(FaceToEdges, EdgeFace, CurrIntersection.EdgeIndex);
 						continue;
@@ -926,49 +964,92 @@ namespace GIA
 			int32 CurrentContourIndex = 1; // ContourIndex 0 is reserved for loop 
 			for (const FIntersectionContourPair& Contour : Contours)
 			{
-				if (!Contour.bIsClosed)
+				if (Contour.ClosedStatus == FIntersectionContourPair::EClosedStatus::Open)
 				{
 					continue;
 				}
 
 				TMultiMap<int32 /*TriangleIndex*/, const FIntersectionContourTriangleSection*> ContourSegments;
-				const bool bCombineContours = Contour.LoopVertexCount > 0; // Loop contours were built still as a contour pair, but really represent a single contour doubling back on itself.
-				const int32 ContourSegmentReserveNum = bCombineContours ? Contour.ColorContours[0].Num() + Contour.ColorContours[1].Num() : FMath::Max(Contour.ColorContours[0].Num(), Contour.ColorContours[1].Num());
-				ContourSegments.Reserve(ContourSegmentReserveNum);
+				
+				const bool bTwoSimpleContours = Contour.ClosedStatus == FIntersectionContourPair::EClosedStatus::SimpleClosed;
 
-				for (const FIntersectionContourTriangleSection& Section : Contour.ColorContours[0])
+				int32 ContourIndex = INDEX_NONE;
+				bool bIsColorB = false;
+				switch (Contour.ClosedStatus)
 				{
-					ContourSegments.Emplace(Section.TriangleIndex, &Section);
-				}
+				case FIntersectionContourPair::EClosedStatus::SimpleClosed:
+				{
+					ContourIndex = (((CurrentContourIndex++) - 1) % 31) + 1;
+					bIsColorB = false;
+					ContourSegments.Reserve(FMath::Max(Contour.ColorContours[0].Contour.Num(), Contour.ColorContours[1].Contour.Num()));
 
-				if (bCombineContours)
+					for (const FIntersectionContourTriangleSection& Section : Contour.ColorContours[0].Contour)
+					{
+						ContourSegments.Emplace(Section.TriangleIndex, &Section);
+					}
+				}break;
+				case FIntersectionContourPair::EClosedStatus::LoopClosed:
 				{
-					for (const FIntersectionContourTriangleSection& Section : Contour.ColorContours[1])
+					ContourIndex = FPBDTriangleMeshCollisions::FGIAColor::LoopContourIndex;
+					bIsColorB = true;
+					ContourSegments.Reserve(Contour.ColorContours[0].Contour.Num() + Contour.ColorContours[1].Contour.Num());
+					for (const FIntersectionContourTriangleSection& Section : Contour.ColorContours[0].Contour)
 					{
 						ContourSegments.Emplace(Section.TriangleIndex, &Section);
 					}
+					for (const FIntersectionContourTriangleSection& Section : Contour.ColorContours[1].Contour)
+					{
+						ContourSegments.Emplace(Section.TriangleIndex, &Section);
+					}
+				}break;
+				case FIntersectionContourPair::EClosedStatus::BoundaryClosed:
+				{
+					ContourIndex = FPBDTriangleMeshCollisions::FGIAColor::BoundaryContourIndex;
+					bIsColorB = false;
+					if (Contour.ColorContours[0].BoundaryEdgeCount == 2)
+					{
+						ContourSegments.Reserve(Contour.ColorContours[0].Contour.Num());
+
+						for (const FIntersectionContourTriangleSection& Section : Contour.ColorContours[0].Contour)
+						{
+							ContourSegments.Emplace(Section.TriangleIndex, &Section);
+						}
+					}
+					else
+					{
+						check(Contour.ColorContours[1].BoundaryEdgeCount == 2);
+						ContourSegments.Reserve(Contour.ColorContours[1].Contour.Num());
+
+						for (const FIntersectionContourTriangleSection& Section : Contour.ColorContours[1].Contour)
+						{
+							ContourSegments.Emplace(Section.TriangleIndex, &Section);
+						}
+					}
+				}break;
+				default:
+					checkNoEntry();
 				}
 
-				const int32 ContourIndex = bCombineContours ? FPBDTriangleMeshCollisions::FGIAColor::LoopContourIndex : (((CurrentContourIndex++) - 1) % 31) + 1;
+				check(ContourIndex != INDEX_NONE);
 
-				if (!__internal::FloodFillContourColor(TriangleMesh, NumParticles, Offset, ContourSegments, ContourIndex, bCombineContours, VertexGIAColors, TriangleGIAColors))
+				if (!__internal::FloodFillContourColor(TriangleMesh, NumParticles, Offset, ContourSegments, ContourIndex, bIsColorB, VertexGIAColors, TriangleGIAColors))
 				{
-					const_cast<FIntersectionContourPair&>(Contour).bIsClosed = false;
+					const_cast<FIntersectionContourPair&>(Contour).ClosedStatus = FIntersectionContourPair::EClosedStatus::Open;
 					continue;
 				}
 
-				if (!bCombineContours)
+				if (bTwoSimpleContours)
 				{
 					ContourSegments.Reset();
 
-					for (const FIntersectionContourTriangleSection& Section : Contour.ColorContours[1])
+					for (const FIntersectionContourTriangleSection& Section : Contour.ColorContours[1].Contour)
 					{
 						ContourSegments.Emplace(Section.TriangleIndex, &Section);
 					}
 
 					if (!__internal::FloodFillContourColor(TriangleMesh, NumParticles, Offset, ContourSegments, ContourIndex, true, VertexGIAColors, TriangleGIAColors))
 					{
-						const_cast<FIntersectionContourPair&>(Contour).bIsClosed = false;
+						const_cast<FIntersectionContourPair&>(Contour).ClosedStatus = FIntersectionContourPair::EClosedStatus::Open;
 						continue;
 					}
 				}
@@ -982,31 +1063,49 @@ namespace GIA
 			IntersectionContourTypes.SetNumZeroed(ContourPoints.Num());
 			for (const FIntersectionContourPair& ContourPair : Contours)
 			{
-				if (ContourPair.bIsClosed)
+				switch (ContourPair.ClosedStatus)
+				{
+				case FIntersectionContourPair::EClosedStatus::SimpleClosed:
 				{
-					if (ContourPair.LoopVertexCount > 0)
+					for (const int32 PointIndex : ContourPair.ContourPointCurves[0])
 					{
-						for (const int32 PointIndex : ContourPair.ContourPointCurves[0])
-						{
-							IntersectionContourTypes[PointIndex] = FPBDTriangleMeshCollisions::FContourType::Loop;
-						}
-						for (const int32 PointIndex : ContourPair.ContourPointCurves[1])
-						{
-							IntersectionContourTypes[PointIndex] = FPBDTriangleMeshCollisions::FContourType::Loop;
-						}
+						IntersectionContourTypes[PointIndex] = FPBDTriangleMeshCollisions::FContourType::Contour0;
 					}
-					else
+					for (const int32 PointIndex : ContourPair.ContourPointCurves[1])
 					{
-						for (const int32 PointIndex : ContourPair.ContourPointCurves[0])
-						{
-							IntersectionContourTypes[PointIndex] = FPBDTriangleMeshCollisions::FContourType::Contour0;
-						}
-						for (const int32 PointIndex : ContourPair.ContourPointCurves[1])
+						IntersectionContourTypes[PointIndex] = FPBDTriangleMeshCollisions::FContourType::Contour1;
+					}
+				}break;
+				case FIntersectionContourPair::EClosedStatus::LoopClosed:
+				{
+					for (const int32 PointIndex : ContourPair.ContourPointCurves[0])
+					{
+						IntersectionContourTypes[PointIndex] = FPBDTriangleMeshCollisions::FContourType::Loop;
+					}
+					for (const int32 PointIndex : ContourPair.ContourPointCurves[1])
+					{
+						IntersectionContourTypes[PointIndex] = FPBDTriangleMeshCollisions::FContourType::Loop;
+					}
+				}break;
+				case FIntersectionContourPair::EClosedStatus::BoundaryClosed:
+				{
+					for (int32 ContourIndex = 0; ContourIndex < 2; ++ContourIndex)
+					{
+						const FPBDTriangleMeshCollisions::FContourType ContourType =
+							ContourPair.ColorContours[ContourIndex].BoundaryEdgeCount == 2 ?
+							FPBDTriangleMeshCollisions::FContourType::BoundaryClosed
+							: FPBDTriangleMeshCollisions::FContourType::BoundaryOpen;
+
+						for (const int32 PointIndex : ContourPair.ContourPointCurves[ContourIndex])
 						{
-							IntersectionContourTypes[PointIndex] = FPBDTriangleMeshCollisions::FContourType::Contour1;
+							IntersectionContourTypes[PointIndex] = ContourType;;
 						}
 					}
 				}
+				break;
+				default:
+					break;
+				}
 			}
 		}
 	} // namespace FloodFill
@@ -1112,7 +1211,7 @@ namespace ContourMinimization
 		int32 IndexStart = 0;
 		for (int32 ContourPairIndex = 0; ContourPairIndex < IntersectionContours.Num(); ++ContourPairIndex)
 		{
-			if (IntersectionContours[ContourPairIndex].bIsClosed && IntersectionContours[ContourPairIndex].LoopVertexCount == 0)
+			if (IntersectionContours[ContourPairIndex].ClosedStatus == GIA::FIntersectionContourPair::EClosedStatus::SimpleClosed)
 			{
 				continue;
 			}
@@ -1209,4 +1308,40 @@ void FPBDTriangleMeshCollisions::Init(const FSolverParticles& Particles, const F
 		}
 	}
 }
+
+void FPBDTriangleMeshCollisions::PostStepInit(const FSolverParticles& Particles)
+{
+	if (TriangleMesh.GetNumElements() == 0)
+	{
+		return;
+	}
+
+	PostStepIntersectionContourPoints.Reset();
+	PostStepContourMinimizationIntersections.Reset();
+
+	if (NumContourMinimizationPostSteps > 0)
+	{
+		// For now just going to rebuild the spatial grid every time. In reality, should really do some sort of refitting.
+		{
+			TRACE_CPUPROFILER_EVENT_SCOPE(ChaosFPBDTriangleMeshCollisions_BuildSpatialHash);
+			constexpr FSolverReal MinSpatialLodSize = 1.f;
+			TriangleMesh.BuildSpatialHash(static_cast<const TArrayView<const FSolverVec3>&>(Particles.XArray()), SpatialHash, MinSpatialLodSize);
+
+
+			// Detect all EdgeFace Intersections
+			TArray<FEdgeFaceIntersection> Intersections = FindEdgeFaceIntersections(TriangleMesh, SpatialHash, Particles);
+
+			if (bUseGlobalPostStepContours)
+			{
+				TArray<GIA::FIntersectionContourPair> IntersectionContours = GIA::ContourBuilding::BuildIntersectionContours(TriangleMesh, Intersections, PostStepIntersectionContourPoints);
+				ContourMinimization::BuildGlobalContourMinimizationIntersections(TriangleMesh, Particles, IntersectionContours, PostStepContourMinimizationIntersections);
+			}
+			else
+			{
+				// Just build local gradient for all Intersections
+				ContourMinimization::BuildLocalContourMinimizationIntersections(TriangleMesh, Particles, Intersections, PostStepContourMinimizationIntersections);
+			}
+		}
+	}
+}
 }  // End namespace Chaos::Softs
diff --git a/Engine/Source/Runtime/Experimental/Chaos/Public/Chaos/PBDTriangleMeshCollisions.h b/Engine/Source/Runtime/Experimental/Chaos/Public/Chaos/PBDTriangleMeshCollisions.h
--- a/Engine/Source/Runtime/Experimental/Chaos/Public/Chaos/PBDTriangleMeshCollisions.h
+++ b/Engine/Source/Runtime/Experimental/Chaos/Public/Chaos/PBDTriangleMeshCollisions.h
@@ -37,6 +37,7 @@ public:
 		static constexpr int32 LoopContourIndex = 0;
 		static constexpr int32 LoopBits = 1 << LoopContourIndex;
 		static constexpr int32 NonLoopMask = ~LoopBits;
+		static constexpr int32 BoundaryContourIndex = 0; // Same as Loop. Loop is ColorA, Boundary is ColorB
 		
 		void SetContourColor(int32 ContourIndex, bool bIsColorB)
 		{
@@ -68,6 +69,16 @@ public:
 			SetContourColor(LoopContourIndex, true);
 		}
 
+		void SetBoundary()
+		{
+			SetContourColor(BoundaryContourIndex, false);
+		}
+
+		bool IsBoundary() const
+		{
+			return (ContourIndexBits & LoopBits) && !(ColorBits & LoopBits);
+		}
+
 		// Because they are opposite colors with a shared contour index. This will cause repulsion forces to attract and fix the intersection.
 		// NOTE: We flip normals if ANY of the TriVertColors are opposite the PointColor or if the TriColor (used for thin regions is flipped). This does a better job for thin features
 		// than only flipping normals if ALL TriVertColors are opposite.
@@ -90,6 +101,8 @@ public:
 	{
 		Open = 0,
 		Loop,
+		BoundaryClosed,
+		BoundaryOpen,
 		Contour0,
 		Contour1,
 		Count
@@ -109,11 +122,16 @@ public:
 		:TriangleMesh(InTriangleMesh)
 		, Offset(InOffset)
 		, NumParticles(InNumParticles)
-		, bGlobalIntersectionAnalysis(GetUseSelfIntersections(PropertyCollection, false) || GetUseGlobalIntersectionAnalysis(PropertyCollection, false))
-		, bContourMinimization(GetUseSelfIntersections(PropertyCollection, false) || GetUseContourMinimization(PropertyCollection, false))
+		, bUseSelfIntersections(GetUseSelfIntersections(PropertyCollection, false))
+		, bGlobalIntersectionAnalysis(GetUseSelfIntersections(PropertyCollection, false) && GetUseGlobalIntersectionAnalysis(PropertyCollection, true))
+		, bContourMinimization(GetUseSelfIntersections(PropertyCollection, false) && GetUseContourMinimization(PropertyCollection, true))
+		, NumContourMinimizationPostSteps(GetUseSelfIntersections(PropertyCollection, false) ? GetNumContourMinimizationPostSteps(PropertyCollection, 0) : 0)
+		, bUseGlobalPostStepContours(GetUseGlobalPostStepContours(PropertyCollection, true))
 		, UseSelfIntersectionsIndex(PropertyCollection)
 		, UseGlobalIntersectionAnalysisIndex(PropertyCollection)
 		, UseContourMinimizationIndex(PropertyCollection)
+		, NumContourMinimizationPostStepsIndex(PropertyCollection)
+		, UseGlobalPostStepContoursIndex(PropertyCollection)
 	{}
 
 	FPBDTriangleMeshCollisions(
@@ -126,58 +144,86 @@ public:
 		:TriangleMesh(InTriangleMesh)
 		, Offset(InOffset)
 		, NumParticles(InNumParticles)
+		, bUseSelfIntersections(bInGlobalIntersectionAnalysis || bInContourMinimization)
 		, bGlobalIntersectionAnalysis(bInGlobalIntersectionAnalysis)
 		, bContourMinimization(bInContourMinimization)
 		, UseSelfIntersectionsIndex(ForceInit)
 		, UseGlobalIntersectionAnalysisIndex(ForceInit)
 		, UseContourMinimizationIndex(ForceInit)
+		, NumContourMinimizationPostStepsIndex(ForceInit)
+		, UseGlobalPostStepContoursIndex(ForceInit)
 	{}
 
 	virtual ~FPBDTriangleMeshCollisions() = default;
 
 	void SetProperties(const FCollectionPropertyConstFacade& PropertyCollection)
 	{
-		bool bUseSelfIntersections = false;
-		if (IsUseSelfIntersectionsMutable(PropertyCollection))
+		const bool bSelfIntersectionsMutable = IsUseSelfIntersectionsMutable(PropertyCollection);
+		if (bSelfIntersectionsMutable)
 		{
 			bUseSelfIntersections = GetUseSelfIntersections(PropertyCollection);
 		}
 		if (bUseSelfIntersections)
 		{
-			bGlobalIntersectionAnalysis = bContourMinimization = true;
-		}
-		else
-		{
-			if (IsUseGlobalIntersectionAnalysisMutable(PropertyCollection))
+			if (bSelfIntersectionsMutable || IsUseGlobalIntersectionAnalysisMutable(PropertyCollection))
 			{
 				bGlobalIntersectionAnalysis = GetUseGlobalIntersectionAnalysis(PropertyCollection);
 			}
-			if (IsUseContourMinimizationMutable(PropertyCollection))
+			if (bSelfIntersectionsMutable || IsUseContourMinimizationMutable(PropertyCollection))
 			{
 				bContourMinimization = GetUseContourMinimization(PropertyCollection);
 			}
+			if (bSelfIntersectionsMutable || IsNumContourMinimizationPostStepsMutable(PropertyCollection))
+			{
+				NumContourMinimizationPostSteps = GetNumContourMinimizationPostSteps(PropertyCollection);
+			}
+			if (bSelfIntersectionsMutable || IsUseGlobalPostStepContoursMutable(PropertyCollection))
+			{
+				bUseGlobalPostStepContours = GetUseGlobalPostStepContours(PropertyCollection);
+			}
+		}
+		else
+		{
+			bGlobalIntersectionAnalysis = bContourMinimization = false;
+			NumContourMinimizationPostSteps = 0;
 		}
 	}
 
 	CHAOS_API void Init(const FSolverParticles& Particles, const FSolverReal MinProximityQueryRadius = (FSolverReal)0.);
 
+	CHAOS_API void PostStepInit(const FSolverParticles& Particles);
+
 	void SetGlobalIntersectionAnalysis(bool bInGlobalIntersectionAnalysis) { bGlobalIntersectionAnalysis = bInGlobalIntersectionAnalysis; }
 	void SetContourMinimization(bool bInContourMinimization) { bContourMinimization = bInContourMinimization; }
 
+	int32 GetNumContourMinimizationPostSteps() const
+	{
+		return NumContourMinimizationPostSteps;
+	}
+
 	const FTriangleMesh::TSpatialHashType<FSolverReal>& GetSpatialHash() const { return SpatialHash; }
 	const TArray<FContourMinimizationIntersection>& GetContourMinimizationIntersections() const { return ContourMinimizationIntersections; }
 	const TConstArrayView<FGIAColor> GetVertexGIAColors() const { return bGlobalIntersectionAnalysis && VertexGIAColors.Num() == NumParticles ? TConstArrayView<FGIAColor>(VertexGIAColors.GetData() - Offset, NumParticles + Offset) : TConstArrayView<FGIAColor>(); }
 	const TArray<FGIAColor>& GetTriangleGIAColors() const { return TriangleGIAColors; }
 	const TArray<TArray<FBarycentricPoint>>& GetIntersectionContourPoints() const { return IntersectionContourPoints; }
 	const TArray<FContourType>& GetIntersectionContourTypes() const { return IntersectionContourTypes; }
+
+	// Same data but for the post step contour minimization. Just making them separate arrays
+	// for debug draw purposes.
+	const TArray<FContourMinimizationIntersection>& GetPostStepContourMinimizationIntersections() const { return PostStepContourMinimizationIntersections; }
+	const TArray<TArray<FBarycentricPoint>>& GetPostStepIntersectionContourPoints() const { return PostStepIntersectionContourPoints; }
 private:
 
 	const FTriangleMesh& TriangleMesh;
 	int32 Offset;
 	int32 NumParticles;
+	bool bUseSelfIntersections;
 	bool bGlobalIntersectionAnalysis;
 	bool bContourMinimization;
-
+	
+	int32 NumContourMinimizationPostSteps = 0;
+	bool bUseGlobalPostStepContours = true;
+	
 	FTriangleMesh::TSpatialHashType<FSolverReal> SpatialHash;
 	TArray<FContourMinimizationIntersection> ContourMinimizationIntersections;
 	TArray<FGIAColor> VertexGIAColors;
@@ -187,9 +233,15 @@ private:
 	TArray<TArray<FBarycentricPoint>> IntersectionContourPoints;
 	TArray<FContourType> IntersectionContourTypes;
 
+	// PostStep contour data. Keeping it separate for debug drawing for now.
+	TArray<FContourMinimizationIntersection> PostStepContourMinimizationIntersections;
+	TArray<TArray<FBarycentricPoint>> PostStepIntersectionContourPoints;
+
 	UE_CHAOS_DECLARE_PROPERTYCOLLECTION_NAME(UseSelfIntersections, bool);
 	UE_CHAOS_DECLARE_PROPERTYCOLLECTION_NAME(UseGlobalIntersectionAnalysis, bool);
 	UE_CHAOS_DECLARE_PROPERTYCOLLECTION_NAME(UseContourMinimization, bool);
+	UE_CHAOS_DECLARE_PROPERTYCOLLECTION_NAME(NumContourMinimizationPostSteps, int32);
+	UE_CHAOS_DECLARE_PROPERTYCOLLECTION_NAME(UseGlobalPostStepContours, bool);
 };
 
 }  // End namespace Chaos::Softs
