Chaos - extra stats (disabled by default) for helping to analyse heavy collision loads

diff --git a/Engine/Source/Runtime/Experimental/Chaos/Private/Chaos/Collision/ParticlePairMidPhase.cpp b/Engine/Source/Runtime/Experimental/Chaos/Private/Chaos/Collision/ParticlePairMidPhase.cpp
--- a/Engine/Source/Runtime/Experimental/Chaos/Private/Chaos/Collision/ParticlePairMidPhase.cpp
+++ b/Engine/Source/Runtime/Experimental/Chaos/Private/Chaos/Collision/ParticlePairMidPhase.cpp
@@ -24,6 +24,13 @@ TRACE_DECLARE_INT_COUNTER_EXTERN(ChaosTraceCounter_MidPhase_NumGeneric);
 
 extern bool Chaos_Collision_NarrowPhase_AABBBoundsCheck;
 
+// Enable for extended stats. Slow but useful for determining counts and relative costs of midphases
+#if 0
+#define CHAOS_MIDPHASE_SCOPE_CYCLE_TIMER(X) QUICK_SCOPE_CYCLE_COUNTER(X)
+#else
+#define CHAOS_MIDPHASE_SCOPE_CYCLE_TIMER(X)
+#endif
+
 namespace Chaos
 {
 #if CHAOS_DEBUG_DRAW
@@ -377,6 +384,8 @@ namespace Chaos
 		const FReal Dt,
 		const FCollisionContext& Context)
 	{
+		CHAOS_MIDPHASE_SCOPE_CYCLE_TIMER(FSingleShapePairCollisionDetector_GenerateCollision);
+
 		const int32 CurrentEpoch = Context.GetAllocator()->GetCurrentEpoch();
 		if (DoBoundsOverlap(CullDistance, CurrentEpoch))
 		{
@@ -391,6 +400,8 @@ namespace Chaos
 		const FReal Dt,
 		const FCollisionContext& Context)
 	{
+		CHAOS_MIDPHASE_SCOPE_CYCLE_TIMER(FSingleShapePairCollisionDetector_GenerateCollisionCCD);
+
 		return GenerateCollisionCCDImpl(bEnableCCDSweep, CullDistance, Dt, Context);
 	}
 
@@ -416,6 +427,8 @@ namespace Chaos
 		const FReal Dt,
 		const FCollisionContext& Context)
 	{
+		CHAOS_MIDPHASE_SCOPE_CYCLE_TIMER(FSingleShapePairCollisionDetector_GenerateCollisionImpl);
+
 		if (BoundsTestFlags.bIsProbe)
 		{
 			return GenerateCollisionProbeImpl(CullDistance, Dt, Context);
@@ -470,6 +483,8 @@ namespace Chaos
 
 				if (!Context.GetSettings().bDeferNarrowPhase)
 				{
+					CHAOS_MIDPHASE_SCOPE_CYCLE_TIMER(FSingleShapePairCollisionDetector_GenerateCollision_NarrowPhase);
+
 					// Run the narrow phase
 					Collisions::UpdateConstraint(*Constraint.Get(), ShapeWorldTransform0, ShapeWorldTransform1, Dt);
 				}
@@ -870,6 +885,8 @@ namespace Chaos
 			return;
 		}
 
+		CHAOS_MIDPHASE_SCOPE_CYCLE_TIMER(FParticlePairMidPhase_GenerateCollision);
+
 		if (Flags.bIsActive)
 		{
 			// CullDistance is scaled by the size of the dynamic objects.
@@ -1046,6 +1063,8 @@ namespace Chaos
 		const FReal Dt,
 		const FCollisionContext& Context)
 	{
+		CHAOS_MIDPHASE_SCOPE_CYCLE_TIMER(FShapePairParticlePairMidPhase_GenerateCollision);
+
 		//TRACE_COUNTER_INCREMENT(ChaosTraceCounter_MidPhase_NumShapePair);
 
 		int32 NumActive = 0;
@@ -1123,6 +1142,8 @@ namespace Chaos
 		const FReal Dt,
 		const FCollisionContext& Context)
 	{
+		CHAOS_MIDPHASE_SCOPE_CYCLE_TIMER(FGenericParticlePairMidPhase_GenerateCollisionImpl);
+
 		//TRACE_COUNTER_INCREMENT(ChaosTraceCounter_MidPhase_NumGeneric);
 		const FImplicitObjectRef Implicit0 = GetParticle0()->GetGeometry();
 		const FImplicitObjectRef Implicit1 = GetParticle1()->GetGeometry();
@@ -1177,7 +1198,11 @@ namespace Chaos
 		}
 
 		// Generate manifolds for each constraint we created/recovered and (re)activate if necessary
-		int32 NumActive = ProcessNewConstraints(CullDistance, Dt, Context);
+		int32 NumActive = 0;
+		if (NewConstraints.Num() > 0)
+		{
+			NumActive = ProcessNewConstraints(CullDistance, Dt, Context);
+		}
 
 		// @todo(chaos): we could clean up unused collisions between this pair, but probably not worth it
 		//PruneConstraints();
@@ -1193,6 +1218,8 @@ namespace Chaos
 		const FReal Dt,
 		const FCollisionContext& Context)
 	{
+		CHAOS_MIDPHASE_SCOPE_CYCLE_TIMER(FGenericParticlePairMidPhase_GenerateCollision_BVHBVH);
+
 		const FRigidTransform3 ParticleWorldTransformA = FConstGenericParticleHandle(ParticleA)->GetTransformPQ();
 		const FRigidTransform3 ParticleWorldTransformB = FConstGenericParticleHandle(ParticleB)->GetTransformPQ();
 		const FRigidTransform3 ParticleTransformBToA = ParticleWorldTransformB.GetRelativeTransform(ParticleWorldTransformA);
@@ -1244,6 +1271,8 @@ namespace Chaos
 		FGeometryParticleHandle* ParticleB, const FImplicitObject* RootImplicitB, const Private::FConvexOptimizer* ConvexOptimizerB,
 		const FReal CullDistance, const FReal Dt, const FCollisionContext& Context)
 	{
+		CHAOS_MIDPHASE_SCOPE_CYCLE_TIMER(FGenericParticlePairMidPhase_GenerateCollision_BVHImplicitHierarchy);
+
 		const FShapeInstanceArray& ShapeInstancesB = ParticleB->ShapeInstances();
 
 		// Visit all the leaf implicits in RootImplicitB and collide against the BVH
@@ -1270,6 +1299,8 @@ namespace Chaos
 		const FReal Dt,
 		const FCollisionContext& Context)
 	{
+		CHAOS_MIDPHASE_SCOPE_CYCLE_TIMER(FGenericParticlePairMidPhase_GenerateCollision_ImplicitHierarchyImplicitHierarch);
+
 		const FConstGenericParticleHandle PA = ParticleA;
 		const FConstGenericParticleHandle PB = ParticleB;
 		const FShapeInstanceArray& ShapeInstancesA = ParticleA->ShapeInstances();
@@ -1319,6 +1350,8 @@ namespace Chaos
 		FGeometryParticleHandle* ParticleB, const FImplicitObject* ImplicitB, const FShapeInstance* ShapeInstanceB, const FRigidTransform3& RelativeTransformB, const int32 LeafObjectIndexB,
 		const FReal CullDistance, const FReal Dt, const FCollisionContext& Context)
 	{
+		CHAOS_MIDPHASE_SCOPE_CYCLE_TIMER(FGenericParticlePairMidPhase_GenerateCollision_BVHImplicitLeaf);
+
 		const FConstGenericParticleHandle PA = ParticleA;
 		const FConstGenericParticleHandle PB = ParticleB;
 
@@ -1407,6 +1440,8 @@ namespace Chaos
 		const FReal Dt,
 		const FCollisionContext& Context)
 	{
+		CHAOS_MIDPHASE_SCOPE_CYCLE_TIMER(FGenericParticlePairMidPhase_GenerateCollision_ImplicitLeafImplicitLeaf);
+
 		// Check the sim filter to see if these shapes collide
 		const EImplicitObjectType ImplicitTypeA = Private::GetImplicitCollisionType(ParticleA, ImplicitA);
 		const EImplicitObjectType ImplicitTypeB = Private::GetImplicitCollisionType(ParticleB, ImplicitB);
@@ -1561,7 +1596,8 @@ namespace Chaos
 		// this will depend on the shape types involved. E.g., with two particles each with a sphere and a box in a union
 		// would require up to two Sphere-Box contacts, with the particles in opposite orders.
 #if !UE_BUILD_TEST && !UE_BUILD_SHIPPING
-		if (!ensure(((InParticle0 == Particle0) && (InParticle1 == Particle1)) || ((InParticle0 == Particle1) && (InParticle1 == Particle0))))
+		const bool bIsCorrectParticles = ((InParticle0 == Particle0) && (InParticle1 == Particle1)) || ((InParticle0 == Particle1) && (InParticle1 == Particle0));
+		if (!ensureMsgf(bIsCorrectParticles, TEXT("Attempt to us MidPhase for particles %d - %d with particles %d - %d"), Particle0->ParticleID().LocalID, Particle1->ParticleID().LocalID, InParticle0->ParticleID().LocalID, InParticle1->ParticleID().LocalID))
 		{
 			// We somehow received a callback for the wrong particle pair...this should not happen
 			return nullptr;
@@ -1664,6 +1700,8 @@ namespace Chaos
 		const FReal Dt,
 		const FCollisionContext& Context)
 	{
+		CHAOS_MIDPHASE_SCOPE_CYCLE_TIMER(FGenericParticlePairMidPhase_ProcessNewConstraints);
+
 		int32 NumActive = 0;
 		const bool bUseCCDSweep = IsCCD() && Flags.bUseSweep;
 
