Chaos - optimize FClusterUnionManager::FlushIncrementalGeometryOperations to avoid repeated rebuild of BVH and ShapesArray

diff --git a/Engine/Source/Runtime/Engine/Private/PhysicsEngine/ClusterUnionComponent.cpp b/Engine/Source/Runtime/Engine/Private/PhysicsEngine/ClusterUnionComponent.cpp
--- a/Engine/Source/Runtime/Engine/Private/PhysicsEngine/ClusterUnionComponent.cpp
+++ b/Engine/Source/Runtime/Engine/Private/PhysicsEngine/ClusterUnionComponent.cpp
@@ -852,7 +852,12 @@ void UClusterUnionComponent::SyncClusterUnionFromProxy()
 		}
 	}
 
-	PhysicsProxy->ForceSetGeometryChildParticles_External(MoveTemp(ChildParticles));
+	// We don't set the child particles if we did not sync geometry in PullFromPhysicsState because ChildParticles 
+	// would not have been populated. However ChildParticles may legitimately be empty if all children were removed
+	if (FullData.bDidSyncGeometry)
+	{
+		PhysicsProxy->ForceSetGeometryChildParticles_External(MoveTemp(ChildParticles));
+	}
 
 	// We need to handle any additions, deletions, and modifications to any child in the cluster union here.
 	// If a component lives in MappedData but not in PerComponentData, new component!
diff --git a/Engine/Source/Runtime/Experimental/Chaos/Private/Chaos/ClusterUnionManager.cpp b/Engine/Source/Runtime/Experimental/Chaos/Private/Chaos/ClusterUnionManager.cpp
--- a/Engine/Source/Runtime/Experimental/Chaos/Private/Chaos/ClusterUnionManager.cpp
+++ b/Engine/Source/Runtime/Experimental/Chaos/Private/Chaos/ClusterUnionManager.cpp
@@ -960,9 +960,24 @@ namespace Chaos
 	void FClusterUnionManager::FlushIncrementalGeometryOperations(FClusterUnion& ClusterUnion)
 	{
 		SCOPE_CYCLE_COUNTER(STAT_FlushIncrementalGeometryOperations);
-		check(ClusterUnion.Geometry != nullptr && ClusterUnion.Geometry->GetType() == ImplicitObjectType::Union);
+		check(ClusterUnion.Geometry != nullptr);
+
+		FImplicitObjectUnion* ImplicitUnion = ClusterUnion.Geometry->template AsA<FImplicitObjectUnion>();
+		check(ImplicitUnion != nullptr);
 
 		const TArray<FPBDRigidParticleHandle*>& PendingGeometryAdditions = ClusterUnion.GetPendingGeometryOperationParticles(EClusterUnionGeometryOperation::Add);
+		const TArray<FPBDRigidParticleHandle*>& PendingGeometryRemovals = ClusterUnion.GetPendingGeometryOperationParticles(EClusterUnionGeometryOperation::Remove);
+		const TArray<FPBDRigidParticleHandle*>& PendingGeometryRefresh = ClusterUnion.GetPendingGeometryOperationParticles(EClusterUnionGeometryOperation::Refresh);
+		if (PendingGeometryAdditions.IsEmpty() && PendingGeometryRemovals.IsEmpty() && PendingGeometryRefresh.IsEmpty())
+		{
+			// NOTE: Early out is important to prevent collision reset and BVH rebuild when there are no changes
+			return;
+		}
+
+		// We are about to change the geometry so clear all cached collisions and prevent the BVH from being rebuilt until we are done
+		MEvolution.InvalidateParticle(ClusterUnion.InternalCluster);
+		ImplicitUnion->SetAllowBVH(false);
+
 		const FRigidTransform3 ClusterWorldTM(ClusterUnion.InternalCluster->X(), ClusterUnion.InternalCluster->R());
 		if (!PendingGeometryAdditions.IsEmpty())
 		{
@@ -998,14 +1013,12 @@ namespace Chaos
 			ClusterUnion.ClearPendingGeometryOperations(EClusterUnionGeometryOperation::Add);
 		}
 
-		const TArray<FPBDRigidParticleHandle*>& PendingGeometryRemovals = ClusterUnion.GetPendingGeometryOperationParticles(EClusterUnionGeometryOperation::Remove);
 		if (!PendingGeometryRemovals.IsEmpty())
 		{
 			RemoveParticlesFromClusterUnionGeometry(ClusterUnion.InternalCluster, PendingGeometryRemovals, ClusterUnion.GeometryChildParticles);
 			ClusterUnion.ClearPendingGeometryOperations(EClusterUnionGeometryOperation::Remove);
 		}
 
-		const TArray<FPBDRigidParticleHandle*>& PendingGeometryRefresh = ClusterUnion.GetPendingGeometryOperationParticles(EClusterUnionGeometryOperation::Refresh);
 		if (!PendingGeometryRefresh.IsEmpty())
 		{
 			// For each particle we need to find the corresponding shape.
@@ -1056,6 +1069,9 @@ namespace Chaos
 
 			ClusterUnion.ClearPendingGeometryOperations(EClusterUnionGeometryOperation::Refresh);
 		}
+
+		// Re-enable the BVH to rebuild it
+		ImplicitUnion->SetAllowBVH(true);
 	}
 
 	DECLARE_CYCLE_STAT(TEXT("FClusterUnionManager::GetOrCreateClusterUnionIndexFromExplicitIndex"), STAT_GetOrCreateClusterUnionIndexFromExplicitIndex, STATGROUP_Chaos);
diff --git a/Engine/Source/Runtime/Experimental/Chaos/Private/Chaos/GeometryParticles.cpp b/Engine/Source/Runtime/Experimental/Chaos/Private/Chaos/GeometryParticles.cpp
--- a/Engine/Source/Runtime/Experimental/Chaos/Private/Chaos/GeometryParticles.cpp
+++ b/Engine/Source/Runtime/Experimental/Chaos/Private/Chaos/GeometryParticles.cpp
@@ -7,6 +7,7 @@
 #include "Chaos/ImplicitObjectUnion.h"
 #include "Chaos/Particle/ParticleUtilities.h"
 #include "Chaos/ParticleHandle.h"
+#include "Chaos/Framework/ArrayAlgorithm.h"
 #include "Chaos/Framework/PhysicsSolverBase.h"
 
 namespace Chaos
@@ -50,10 +51,9 @@ namespace Chaos
 	}
 
 	template <typename T, int d, EGeometryParticlesSimType SimType>
-	void TGeometryParticlesImp<T, d, SimType>::RemoveShapeAtIndex(const int32 ParticleIndex, const int32 Index)
+	void TGeometryParticlesImp<T, d, SimType>::RemoveShapesAtSortedIndices(const int32 ParticleIndex, const TArrayView<const int32>& InIndices)
 	{
-		check(MShapesArray.IsValidIndex(ParticleIndex));
-		MShapesArray[ParticleIndex].RemoveAt(Index);
+		RemoveArrayItemsAtSortedIndices(MShapesArray[ParticleIndex], InIndices);
 	}
 
 	template <typename T, int d, EGeometryParticlesSimType SimType>
diff --git a/Engine/Source/Runtime/Experimental/Chaos/Private/Chaos/ImplicitObjectUnion.cpp b/Engine/Source/Runtime/Experimental/Chaos/Private/Chaos/ImplicitObjectUnion.cpp
--- a/Engine/Source/Runtime/Experimental/Chaos/Private/Chaos/ImplicitObjectUnion.cpp
+++ b/Engine/Source/Runtime/Experimental/Chaos/Private/Chaos/ImplicitObjectUnion.cpp
@@ -4,6 +4,7 @@
 #include "Chaos/ImplicitObjectBVH.h"
 #include "Chaos/BoundingVolumeHierarchy.h"
 #include "Chaos/Convex.h"
+#include "Chaos/Framework/ArrayAlgorithm.h"
 #include "Chaos/Tribox.h"
 
 #include "UObject/FortniteMainBranchObjectVersion.h"
@@ -131,6 +132,22 @@ void FImplicitObjectUnion::RemoveAt(int32 RemoveIndex)
 	RebuildBVH();
 }
 
+void FImplicitObjectUnion::RemoveAtSortedIndices(const TArrayView<const int32>& InIndices)
+{
+	if (InIndices.IsEmpty())
+	{
+		return;
+	}
+
+	RemoveArrayItemsAtSortedIndices(MObjects, InIndices);
+
+	SetNumLeafObjects(Private::FImplicitBVH::CountLeafObjects(MakeArrayView(GetObjects())));
+
+	MLocalBoundingBox = CalculateObjectsBounds(MakeArrayView(MObjects));
+	RebuildBVH();
+}
+
+
 void FImplicitObjectUnion::SetNumLeafObjects(int32 InNumLeafObjects)
 {
 	constexpr int32 MaxNumLeafObjects = int32(TNumericLimits<decltype(NumLeafObjects)>::Max());
diff --git a/Engine/Source/Runtime/Experimental/Chaos/Private/Chaos/ParticleHandle.cpp b/Engine/Source/Runtime/Experimental/Chaos/Private/Chaos/ParticleHandle.cpp
--- a/Engine/Source/Runtime/Experimental/Chaos/Private/Chaos/ParticleHandle.cpp
+++ b/Engine/Source/Runtime/Experimental/Chaos/Private/Chaos/ParticleHandle.cpp
@@ -111,44 +111,32 @@ namespace Chaos
 		// NOTE: only intended use is to remove objects from inside a FImplicitObjectUnion
 		CHAOS_ENSURE(MNonFrequentData.Read().GetGeometry()->GetType() == FImplicitObjectUnion::StaticType());
 
-		int32 FoundIndex = INDEX_NONE;
 		for (int32 Index = 0; Index < MShapesArray.Num(); Index++)
 		{
 			if (InShape == MShapesArray[Index].Get())
 			{
-				FoundIndex = Index;
-				break;
+				RemoveShapesAtSortedIndices({ Index });
+				return;
 			}
 		}
-
-		RemoveShapeAtIndex(FoundIndex);
 	}
 
 	template <typename T, int d>
-	void Chaos::TGeometryParticle<T, d>::RemoveShapeAtIndex(int32 InIndex)
+	void Chaos::TGeometryParticle<T, d>::RemoveShapesAtSortedIndices(const TArrayView<const int32>& InIndices)
 	{
 		// NOTE: only intended use is to remove objects from inside a FImplicitObjectUnion
 		CHAOS_ENSURE(MNonFrequentData.Read().GetGeometry()->GetType() == FImplicitObjectUnion::StaticType());
 
-		if (InIndex == INDEX_NONE)
-		{
-			return;
-		}
-
-		check(MShapesArray.IsValidIndex(InIndex));
-		MShapesArray.RemoveAt(InIndex);
-
-		if (MNonFrequentData.Read().GetGeometry()->GetType() == FImplicitObjectUnion::StaticType())
-		{
-			// if we are currently a union then remove geometry from this union
-			ModifyGeometry([InIndex](FImplicitObject& GeomToModify)
+		ModifyGeometry(
+			[this, &InIndices](FImplicitObject& GeomToModify)
 			{
-				if (FImplicitObjectUnion* Union = GeomToModify.template GetObject<FImplicitObjectUnion>())
+				if (FImplicitObjectUnion* Union = GeomToModify.template AsA<FImplicitObjectUnion>())
 				{
-					Union->RemoveAt(InIndex);
+					RemoveArrayItemsAtSortedIndices(MShapesArray, InIndices);
+
+					Union->RemoveAtSortedIndices(InIndices);
 				}
 			});
-		}
 	}
 
 	template <typename T, int d, bool bPersistent>
@@ -158,52 +146,40 @@ namespace Chaos
 		const FImplicitObjectRef CurrentGeometry = GetGeometry();
 		if (ensure(CurrentGeometry != nullptr))
 		{
-			if (ensure(CurrentGeometry->GetType() == FImplicitObjectUnion::StaticType()))
+			const FShapesArray& CurrentShapesArray = ShapesArray();
+			for (int32 Index = 0; Index < CurrentShapesArray.Num(); Index++)
 			{
-				const FShapesArray& CurrentShapesArray = ShapesArray();
-
-				int32 FoundIndex = INDEX_NONE;
-				for (int32 Index = 0; Index < CurrentShapesArray.Num(); Index++)
+				if (InShape == CurrentShapesArray[Index].Get())
 				{
-					if (InShape == CurrentShapesArray[Index].Get())
-					{
-						FoundIndex = Index;
-						break;
-					}
+					RemoveShapesAtSortedIndices(MakeArrayView({ Index }));
+					return;
 				}
-
-				RemoveShapeAtIndex(FoundIndex);
 			}
 		}
 	}
 
 	template <typename T, int d, bool bPersistent>
-	void Chaos::TGeometryParticleHandleImp<T, d, bPersistent>::RemoveShapeAtIndex(int32 InIndex)
+	void Chaos::TGeometryParticleHandleImp<T, d, bPersistent>::RemoveShapesAtSortedIndices(const TArrayView<const int32>& InIndices)
 	{
-		if (InIndex == INDEX_NONE)
-		{
-			return;
-		}
-
 		// NOTE: only intended use is to remove objects from inside a FImplicitObjectUnion
 		const FImplicitObjectRef CurrentGeometry = GetGeometry();
-		if (!ensure(CurrentGeometry != nullptr))
+		if (CurrentGeometry == nullptr)
 		{
 			return;
 		}
 
-		if (!ensure(CurrentGeometry->GetType() == FImplicitObjectUnion::StaticType()))
+		FImplicitObjectUnion* Union = CurrentGeometry->template AsA<FImplicitObjectUnion>();
+		if (Union == nullptr)
 		{
 			return;
 		}
 
-		GeometryParticles->RemoveShapeAtIndex(ParticleIdx, InIndex);
+		GeometryParticles->RemoveShapesAtSortedIndices(ParticleIdx, InIndices);
 
-		// if we are currently a union then remove geometry from this union
-		FImplicitObjectUnion& Union = CurrentGeometry->GetObjectChecked<FImplicitObjectUnion>();
-		Union.RemoveAt(InIndex);
+		Union->RemoveAtSortedIndices(InIndices);
 
 		// Needed to update the shapes array.
+		// @todo(chaos): is it though? Maybe for the bounds etc?
 		SetGeometry(GeometryParticles->GetGeometry(ParticleIdx));
 	}
 
diff --git a/Engine/Source/Runtime/Experimental/Chaos/Private/PhysicsProxy/ClusterUnionPhysicsProxy.cpp b/Engine/Source/Runtime/Experimental/Chaos/Private/PhysicsProxy/ClusterUnionPhysicsProxy.cpp
--- a/Engine/Source/Runtime/Experimental/Chaos/Private/PhysicsProxy/ClusterUnionPhysicsProxy.cpp
+++ b/Engine/Source/Runtime/Experimental/Chaos/Private/PhysicsProxy/ClusterUnionPhysicsProxy.cpp
@@ -397,6 +397,7 @@ namespace Chaos
 		}
 
 		SyncedData_External.bIsAnchored = CurrentPullData.bIsAnchored;
+		SyncedData_External.bDidSyncGeometry = false;
 		SyncedData_External.ChildParticles.Empty(CurrentPullData.ChildParticles.Num());
 
 		// I question the need to do this individual copy one by one...maybe we should've used the same type for the synced data.
@@ -413,6 +414,7 @@ namespace Chaos
 		if (CurrentPullData.Geometry)
 		{
 			Particle_External->SetGeometry(CurrentPullData.Geometry);
+			SyncedData_External.bDidSyncGeometry = true;
 		}
 		Particle_External->SetObjectState(CurrentPullData.ObjectState, true, /*bInvalidate=*/false);
 
diff --git a/Engine/Source/Runtime/Experimental/Chaos/Public/Chaos/ClusterUnionManager.h b/Engine/Source/Runtime/Experimental/Chaos/Public/Chaos/ClusterUnionManager.h
--- a/Engine/Source/Runtime/Experimental/Chaos/Public/Chaos/ClusterUnionManager.h
+++ b/Engine/Source/Runtime/Experimental/Chaos/Public/Chaos/ClusterUnionManager.h
@@ -2,6 +2,7 @@
 #pragma once
 
 #include "Chaos/ClusterCreationParameters.h"
+#include "Chaos/Framework/ArrayAlgorithm.h"
 #include "Chaos/ParticleHandleFwd.h"
 #include "Containers/Array.h"
 #include "Containers/Map.h"
@@ -383,12 +384,7 @@ namespace Chaos
 			return;
 		}
 
-		const Chaos::FShapesArray& ShapesArray = ClusterParticle->ShapesArray();
-		const Chaos::FImplicitObjectRef Geometry = ClusterParticle->GetGeometry();
-		check(Geometry != nullptr);
-
-		const Chaos::FImplicitObjectUnion& GeometryUnion = Geometry->GetObjectChecked<Chaos::FImplicitObjectUnion>();
-		check(AllChildParticles.Num() == ShapesArray.Num());
+		check(AllChildParticles.Num() == ClusterParticle->ShapesArray().Num());
 		
 		TArray<int32> ShapeIndicesToRemove;
 		ShapeIndicesToRemove.Reserve(ShapeParticles.Num());
@@ -402,35 +398,12 @@ namespace Chaos
 			}
 		}
 
-		// To make sure we preserve the index as we remove shapes one-by-one, we need to order
-		// shape removal based on their index.
-		ShapeIndicesToRemove.Sort(
-			[](const int32 A, const int32 B)
-			{
-				// This way we guarantee descending order.
-				return A > B;
-			}
-		);
+		ShapeIndicesToRemove.Sort();
 
-		int32 LastIndex = INDEX_NONE;
-		// TODO: Can we batch this to remove all shapes at once?
-		// We're going to need to do that so we don't do multiple iterations through the shapes array.
-		for (int32 ToRemoveIndex : ShapeIndicesToRemove)
-		{
-			// This is a good cheap way of preventing us from trying to remove the same shape (index)
-			// multiple times.
-			if (ToRemoveIndex == LastIndex)
-			{
-				// Probably a good indication of a failure if we end up trying to remove the same shape more than once.
-				check(false);
-				continue;
-			}
+		ClusterParticle->RemoveShapesAtSortedIndices(ShapeIndicesToRemove);
 
-			ClusterParticle->RemoveShapeAtIndex(ToRemoveIndex);
-			AllChildParticles.RemoveAt(ToRemoveIndex);
-			LastIndex = ToRemoveIndex;
-		}
+		RemoveArrayItemsAtSortedIndices(AllChildParticles, ShapeIndicesToRemove);
 
-		check(AllChildParticles.Num() == ShapesArray.Num());
+		check(AllChildParticles.Num() == ClusterParticle->ShapesArray().Num());
 	}
 }
diff --git a/Engine/Source/Runtime/Experimental/Chaos/Public/Chaos/Framework/ArrayAlgorithm.h b/Engine/Source/Runtime/Experimental/Chaos/Public/Chaos/Framework/ArrayAlgorithm.h
new file mode 100644
--- /dev/null
+++ b/Engine/Source/Runtime/Experimental/Chaos/Public/Chaos/Framework/ArrayAlgorithm.h
@@ -0,0 +1,62 @@
+// Copyright Epic Games, Inc. All Rights Reserved.
+#pragma once
+#include "Chaos/Core.h"
+#include "Containers/Array.h"
+
+namespace Chaos
+{
+	/**
+	* Move the elements in the array Items in the range [BeginIndex, EndIndex) down by DownShift. Does not resize the array.
+	*/
+	template<typename TItemArray>
+	void MoveArrayItemsDown(TItemArray& Items, const int32 BeginIndex, const int32 EndIndex, const int32 DownShift)
+	{
+		check(DownShift >= 0);
+		check(BeginIndex - DownShift >= 0);
+		check(EndIndex >= BeginIndex);
+		check(EndIndex <= Items.Num());
+
+		if (DownShift > 0)
+		{
+			for (int32 SrcIndex = BeginIndex; SrcIndex < EndIndex; ++SrcIndex)
+			{
+				Items[SrcIndex - DownShift] = MoveTemp(Items[SrcIndex]);
+			}
+		}
+	}
+
+	/**
+	* Remove a set of items from an array
+	* @tparam TItemArray An array of items of any movable type
+	* @param Items an array of items to remove elements from
+	* @param SortedIndicesToRemove the set of array indices to be removed. Must be sorted low to high. May contain duplicates.
+	*/
+	template<typename TItemArray>
+	void RemoveArrayItemsAtSortedIndices(TItemArray& Items, const TArrayView<const int32>& SortedIndicesToRemove)
+	{
+		if (SortedIndicesToRemove.IsEmpty())
+		{
+			return;
+		}
+
+		int32 DestShift = 1;
+		int32 IndexToRemove0 = SortedIndicesToRemove[0];
+		for (int32 IndicesIndex = 1; IndicesIndex < SortedIndicesToRemove.Num(); ++IndicesIndex)
+		{
+			const int32 IndexToRemove1 = SortedIndicesToRemove[IndicesIndex];
+
+			if (IndexToRemove1 != IndexToRemove0)
+			{
+				MoveArrayItemsDown(Items, IndexToRemove0 + 1, IndexToRemove1, DestShift);
+
+				IndexToRemove0 = IndexToRemove1;
+				DestShift++;
+			}
+		}
+
+		MoveArrayItemsDown(Items, IndexToRemove0 + 1, Items.Num(), DestShift);
+
+		Items.SetNum(Items.Num() - DestShift, false);
+	}
+
+}
\ No newline at end of file
diff --git a/Engine/Source/Runtime/Experimental/Chaos/Public/Chaos/GeometryParticles.h b/Engine/Source/Runtime/Experimental/Chaos/Public/Chaos/GeometryParticles.h
--- a/Engine/Source/Runtime/Experimental/Chaos/Public/Chaos/GeometryParticles.h
+++ b/Engine/Source/Runtime/Experimental/Chaos/Public/Chaos/GeometryParticles.h
@@ -252,7 +252,7 @@ namespace Chaos
 		bool& HasCollision(const int32 Index) { return MHasCollision[Index]; }
 
 		const FShapesArray& ShapesArray(const int32 Index) const { return reinterpret_cast<const FShapesArray&>(MShapesArray[Index]); }
-		void RemoveShapeAtIndex(const int32 ParticleIndex, const int32 Index);
+		void RemoveShapesAtSortedIndices(const int32 ParticleIndex, const TArrayView<const int32>& InIndices);
 
 		const FShapeInstanceArray& ShapeInstances(const int32 Index) const { return MShapesArray[Index]; }
 
@@ -312,6 +312,11 @@ namespace Chaos
 
 		void SetGeometryImpl(const int32 Index, const FImplicitObjectPtr& InGeometry)
 		{
+			// We hit these checks if there's a call to modify geometry without first clearing constraints 
+			// on the particle (e.g., PBDRigidsEvolutionGBF::InvalidateParticle)
+			check(MParticleCollisions[Index].Num() == 0);
+			check(MParticleConstraints[Index].Num() == 0);
+
 			MGeometry[Index] = InGeometry;
 
 			UpdateShapesArray(Index);
diff --git a/Engine/Source/Runtime/Experimental/Chaos/Public/Chaos/ImplicitObjectUnion.h b/Engine/Source/Runtime/Experimental/Chaos/Public/Chaos/ImplicitObjectUnion.h
--- a/Engine/Source/Runtime/Experimental/Chaos/Public/Chaos/ImplicitObjectUnion.h
+++ b/Engine/Source/Runtime/Experimental/Chaos/Public/Chaos/ImplicitObjectUnion.h
@@ -46,6 +46,7 @@ class FImplicitObjectUnion : public FImplicitObject
 
 	CHAOS_API void Combine(TArray<Chaos::FImplicitObjectPtr>& Objects);
 	CHAOS_API void RemoveAt(int32 RemoveIndex);
+	CHAOS_API void RemoveAtSortedIndices(const TArrayView<const int32>& InIndices);
 
 	// The total number of root objects in the hierarchy (same as GetObjects().Num())
 	int32 GetNumRootObjects() const
diff --git a/Engine/Source/Runtime/Experimental/Chaos/Public/Chaos/ParticleHandle.h b/Engine/Source/Runtime/Experimental/Chaos/Public/Chaos/ParticleHandle.h
--- a/Engine/Source/Runtime/Experimental/Chaos/Public/Chaos/ParticleHandle.h
+++ b/Engine/Source/Runtime/Experimental/Chaos/Public/Chaos/ParticleHandle.h
@@ -589,7 +589,7 @@ public:
 	const FShapeInstanceArray& ShapeInstances() const { return GeometryParticles->ShapeInstances(ParticleIdx); }
 
 	void RemoveShape(FPerShapeData* InShape);
-	void RemoveShapeAtIndex(int32 InIndex);
+	void RemoveShapesAtSortedIndices(const TArrayView<const int32>& InIndices);
 
 	const TAABB<T, d>& LocalBounds() const { return GeometryParticles->LocalBounds(ParticleIdx); }
 	void SetLocalBounds(const TAABB<T, d>& NewBounds) { GeometryParticles->LocalBounds(ParticleIdx) = NewBounds; }
@@ -2543,7 +2543,7 @@ public:
 	TSerializablePtr<FImplicitObject> Geometry() const { check(false); return TSerializablePtr<FImplicitObject>(); }
 
 	CHAOS_API void RemoveShape(FPerShapeData* InShape, bool bWakeTouching);
-	CHAOS_API void RemoveShapeAtIndex(int32 InIndex);
+	CHAOS_API void RemoveShapesAtSortedIndices(const TArrayView<const int32>& InIndices);
 
 	void* UserData() const { return MUserData; }
 	void SetUserData(void* InUserData)
diff --git a/Engine/Source/Runtime/Experimental/Chaos/Public/PhysicsProxy/ClusterUnionPhysicsProxy.h b/Engine/Source/Runtime/Experimental/Chaos/Public/PhysicsProxy/ClusterUnionPhysicsProxy.h
--- a/Engine/Source/Runtime/Experimental/Chaos/Public/PhysicsProxy/ClusterUnionPhysicsProxy.h
+++ b/Engine/Source/Runtime/Experimental/Chaos/Public/PhysicsProxy/ClusterUnionPhysicsProxy.h
@@ -45,6 +45,9 @@ namespace Chaos
 		// Whether the cluster is anchored or not.
 		bool bIsAnchored;
 
+		// Whether we assigned new geometry from the PT in PullFromPhysicsState
+		bool bDidSyncGeometry;
+
 		// Data on every child particle in the cluster union.
 		TArray<FClusterUnionChildData> ChildParticles;
 	};
