Chaos Cloth: Add simple Self Collision Spheres constraint which is much cheaper than the existing vertex-face self collisions.

This constraint works similar to the old apex self collisions. Spheres are seeded on cloth vertices, and are used for simple repulsion-based constraints.

* The spheres are generated in the Dataflow graph (based on radius and an additional cull multiplier similar to old apex properties) and stored on the asset as a Vertex Selection Set.
* Vertex sets are now threaded through the asset to the solver in the same manner as weight maps.
* The constraint uses a new version of the spatial grid hashing that has been optimized for the simple case of all the same size spheres. The grid spacing = diameter of a sphere. This way you only need to look for collisions against other spheres within your 3x3x3 neighborhood. Spatial grid building and self collision detection has been parallelized using atomics. Constraints between spheres are generated once per timestep.

diff --git a/Engine/Plugins/Experimental/ChaosCloth/Source/Chaos/Private/ChaosCloth/ChaosClothConstraints.cpp b/Engine/Plugins/Experimental/ChaosCloth/Source/Chaos/Private/ChaosCloth/ChaosClothConstraints.cpp
--- a/Engine/Plugins/Experimental/ChaosCloth/Source/Chaos/Private/ChaosCloth/ChaosClothConstraints.cpp
+++ b/Engine/Plugins/Experimental/ChaosCloth/Source/Chaos/Private/ChaosCloth/ChaosClothConstraints.cpp
@@ -16,6 +16,7 @@
 #include "Chaos/PBDAnimDriveConstraint.h"
 #include "Chaos/PBDShapeConstraints.h"
 #include "Chaos/PBDCollisionSpringConstraints.h"
+#include "Chaos/PBDSelfCollisionSphereConstraints.h"
 #include "Chaos/PBDTriangleMeshCollisions.h"
 #include "Chaos/PBDTriangleMeshIntersections.h"
 #include "Chaos/PBDEvolution.h"
@@ -129,7 +130,8 @@ void FClothConstraints::AddRules(
 	}
 
 	// Call new AddRules function
-	AddRules(ConfigProperties, TriangleMesh, nullptr, WeightMaps, Tethers, MeshScale, bEnabled);
+	AddRules(ConfigProperties, TriangleMesh, nullptr, WeightMaps, TMap<FString, const TSet<int32>*>(),
+		Tethers, MeshScale, bEnabled);
 }
 
 void FClothConstraints::AddRules(
@@ -139,9 +141,22 @@ void FClothConstraints::AddRules(
 	const TMap<FString, TConstArrayView<FRealSingle>>& WeightMaps,
 	const TArray<TConstArrayView<TTuple<int32, int32, FRealSingle>>>& Tethers,
 	Softs::FSolverReal MeshScale, bool bEnabled)
+{
+	AddRules(ConfigProperties, TriangleMesh, PatternData, WeightMaps, TMap<FString, const TSet<int32>*>(),
+		Tethers, MeshScale, bEnabled);
+}
+
+void FClothConstraints::AddRules(
+	const Softs::FCollectionPropertyConstFacade& ConfigProperties,
+	const FTriangleMesh& TriangleMesh,
+	const FClothingPatternData* PatternData,
+	const TMap<FString, TConstArrayView<FRealSingle>>& WeightMaps,
+	const TMap<FString, const TSet<int32>*>& VertexSets,
+	const TArray<TConstArrayView<TTuple<int32, int32, FRealSingle>>>& Tethers,
+	Softs::FSolverReal MeshScale, bool bEnabled)
 {
 	// Self collisions
-	CreateSelfCollisionConstraints(ConfigProperties, TriangleMesh);
+	CreateSelfCollisionConstraints(ConfigProperties, VertexSets, TriangleMesh);
 
 	// Edge constraints
 	CreateStretchConstraints(ConfigProperties, WeightMaps, TriangleMesh, PatternData);
@@ -173,7 +188,8 @@ PRAGMA_ENABLE_DEPRECATION_WARNINGS
 	Enable(bEnabled);
 }
 
-void FClothConstraints::CreateSelfCollisionConstraints(const Softs::FCollectionPropertyConstFacade& ConfigProperties, const FTriangleMesh& TriangleMesh)
+void FClothConstraints::CreateSelfCollisionConstraints(const Softs::FCollectionPropertyConstFacade& ConfigProperties, 
+	const TMap<FString, const TSet<int32>*>& VertexSets, const FTriangleMesh& TriangleMesh)
 {
 	const bool bUseSelfCollisions = ConfigProperties.GetValue<bool>(TEXT("UseSelfCollisions"));
 
@@ -218,6 +234,16 @@ void FClothConstraints::CreateSelfCollisionConstraints(const Softs::FCollectionP
 		++NumConstraintInits;
 		++NumPostprocessingConstraintRules;
 	}
+	else if (Softs::FPBDSelfCollisionSphereConstraints::IsEnabled(ConfigProperties))
+	{
+		SelfCollisionSphereConstraints = MakeShared< Softs::FPBDSelfCollisionSphereConstraints>(
+			ParticleOffset,
+			NumParticles,
+			VertexSets,
+			ConfigProperties);
+		++NumConstraintInits;
+		++NumPostCollisionConstraintRules;
+	}
 }
 
 void FClothConstraints::CreateStretchConstraints(
@@ -827,6 +853,23 @@ void FClothConstraints::CreateRules()
 
 	}
 
+	if (SelfCollisionSphereConstraints)
+	{
+		ConstraintInits[ConstraintInitIndex++] =
+			[this](Softs::FSolverParticles& Particles, const Softs::FSolverReal /*Dt*/)
+			{
+				SelfCollisionSphereConstraints->Init(Particles);
+			};
+
+		PostCollisionConstraintRules[PostCollisionConstraintRuleIndex++] =
+			[this](Softs::FSolverParticles& Particles, const Softs::FSolverReal Dt)
+			{
+				SelfCollisionSphereConstraints->Apply(Particles, Dt);
+			};
+	}
+
+	// The following constraints only run once per subframe, so we do their Apply as part of the Init() which modifies P
+	// To avoid possible dependency order issues, add them last
 	if (SelfCollisionInit && SelfIntersectionConstraints)
 	{
 		// The following constraints only run once per subframe, so we do their Apply as part of the Init() which modifies P
@@ -1216,6 +1259,7 @@ void FClothConstraints::SetSelfCollisionConstraints(const FTriangleMesh& Triangl
 void FClothConstraints::Update(
 	const Softs::FCollectionPropertyConstFacade& ConfigProperties,
 	const TMap<FString, TConstArrayView<FRealSingle>>& WeightMaps,
+	const TMap<FString, const TSet<int32>*>& VertexSets,
 	Softs::FSolverReal MeshScale,
 	Softs::FSolverReal MaxDistancesScale)
 {
@@ -1283,6 +1327,20 @@ void FClothConstraints::Update(
 	{
 		SelfCollisionInit->SetProperties(ConfigProperties);
 	}
+	if (SelfCollisionSphereConstraints)
+	{
+		SelfCollisionSphereConstraints->SetProperties(ConfigProperties, VertexSets);
+	}
+}
+
+// Deprecated
+void FClothConstraints::Update(
+	const Softs::FCollectionPropertyConstFacade& ConfigProperties,
+	const TMap<FString, TConstArrayView<FRealSingle>>& WeightMaps,
+	Softs::FSolverReal MeshScale,
+	Softs::FSolverReal MaxDistancesScale)
+{
+	Update(ConfigProperties, WeightMaps, TMap<FString, const TSet<int32>*>(), MeshScale, MaxDistancesScale);
 }
 
 // Deprecated
@@ -1291,7 +1349,7 @@ void FClothConstraints::Update(
 	Softs::FSolverReal MeshScale,
 	Softs::FSolverReal MaxDistancesScale)
 {
-	Update(ConfigProperties, TMap<FString, TConstArrayView<FRealSingle>>(), MeshScale, MaxDistancesScale);
+	Update(ConfigProperties, TMap<FString, TConstArrayView<FRealSingle>>(), TMap<FString, const TSet<int32>*>(), MeshScale, MaxDistancesScale);
 }
 
 void FClothConstraints::SetEdgeProperties(const Softs::FSolverVec2& EdgeStiffness, const Softs::FSolverVec2& DampingRatio)
diff --git a/Engine/Plugins/Experimental/ChaosCloth/Source/Chaos/Private/ChaosCloth/ChaosClothVisualization.cpp b/Engine/Plugins/Experimental/ChaosCloth/Source/Chaos/Private/ChaosCloth/ChaosClothVisualization.cpp
--- a/Engine/Plugins/Experimental/ChaosCloth/Source/Chaos/Private/ChaosCloth/ChaosClothVisualization.cpp
+++ b/Engine/Plugins/Experimental/ChaosCloth/Source/Chaos/Private/ChaosCloth/ChaosClothVisualization.cpp
@@ -19,6 +19,7 @@
 #include "Chaos/PBDBendingConstraints.h"
 #include "Chaos/PBDCollisionSpringConstraints.h"
 #include "Chaos/PBDLongRangeConstraints.h"
+#include "Chaos/PBDSelfCollisionSphereConstraints.h"
 #include "Chaos/PBDSphericalConstraint.h"
 #include "Chaos/PBDSpringConstraints.h"
 #include "Chaos/PBDTriangleMeshCollisions.h"
@@ -1987,6 +1988,32 @@ static FAutoConsoleVariableRef CVarClothVizWeightMapName(TEXT("p.ChaosClothVisua
 					DrawLine(PDI, Pos1, P, Orange);
 				}
 			}
+
+			if (const Softs::FPBDSelfCollisionSphereConstraints* const SelfCollisionSphereConstraints =
+				ClothConstraints.GetSelfCollisionSphereConstraints().Get())
+			{
+				const TArray<Softs::FSolverVec3>& Positions = Solver->GetParticleXs();
+				const TArray<TVec2<int32>>& Constraints = SelfCollisionSphereConstraints->GetConstraints();
+				for (int32 Index = 0; Index < Constraints.Num(); ++Index)
+				{
+					const TVec2<int32>& Constraint = Constraints[Index];
+					const FVector P0 = LocalSpaceLocation + FVector(Positions[Constraint[0]]);
+					const FVector P1 = LocalSpaceLocation + FVector(Positions[Constraint[1]]);
+					static const FLinearColor Brown(0.1f, 0.05f, 0.f);
+					DrawLine(PDI, P0, P1, Brown);
+				}
+
+				if (const TSet<int32>* const VertexSet = SelfCollisionSphereConstraints->GetVertexSet())
+				{
+					const FReal Radius = (FReal)SelfCollisionSphereConstraints->GetRadius();
+					for (const int32 Vertex : *VertexSet)
+					{
+						const FVector P0 = LocalSpaceLocation + FVector(Positions[Vertex + Offset]);
+						DrawSphere(PDI, TSphere<FReal, 3>(FVector::ZeroVector, Radius), FQuat::Identity,
+							P0, FLinearColor(FColor::Orange));
+					}
+				}
+			}
 		}
 	}
 
diff --git a/Engine/Plugins/Experimental/ChaosCloth/Source/Chaos/Private/ChaosCloth/ChaosClothingSimulationCloth.cpp b/Engine/Plugins/Experimental/ChaosCloth/Source/Chaos/Private/ChaosCloth/ChaosClothingSimulationCloth.cpp
--- a/Engine/Plugins/Experimental/ChaosCloth/Source/Chaos/Private/ChaosCloth/ChaosClothingSimulationCloth.cpp
+++ b/Engine/Plugins/Experimental/ChaosCloth/Source/Chaos/Private/ChaosCloth/ChaosClothingSimulationCloth.cpp
@@ -47,6 +47,7 @@ struct FClothingSimulationCloth::FLODData
 	const int32 NumParticles;
 	const TConstArrayView<uint32> Indices;
 	const TMap<FString, TConstArrayView<FRealSingle>> WeightMaps;
+	const TMap<FString, const TSet<int32>*> VertexSets;
 	const TArray<TConstArrayView<TTuple<int32, int32, FRealSingle>>> Tethers;
 
 	const FClothingPatternData PatternData;
@@ -71,6 +72,7 @@ struct FClothingSimulationCloth::FLODData
 		const TConstArrayView<uint32>& InPatternIndices,
 		const TConstArrayView<uint32>& InPatternToWeldedIndices,
 		TMap<FString, TConstArrayView<FRealSingle>>&& InWeightMaps,
+		TMap<FString, const TSet<int32>*>&& InVertexSets,
 		const TArray<TConstArrayView<TTuple<int32, int32, FRealSingle>>>& InTethers);
 
 	void Add(FClothingSimulationSolver* Solver, FClothingSimulationCloth* Cloth, int32 LODIndex);
@@ -92,10 +94,12 @@ FClothingSimulationCloth::FLODData::FLODData(
 	const TConstArrayView<uint32>& InPatternIndices,
 	const TConstArrayView<uint32>& InPatternToWeldedIndices,
 	TMap<FString, TConstArrayView<FRealSingle>>&& InWeightMaps,
+	TMap<FString, const TSet<int32>*>&& InVertexSets,
 	const TArray<TConstArrayView<TTuple<int32, int32, FRealSingle>>>& InTethers)
 	: NumParticles(InNumParticles)
 	, Indices(InIndices)
 	, WeightMaps(MoveTemp(InWeightMaps))
+	, VertexSets(MoveTemp(InVertexSets))
 	, Tethers(InTethers)
 	, PatternData(NumParticles, Indices, InPatternPositions, InPatternIndices, InPatternToWeldedIndices)
 {
@@ -211,7 +215,7 @@ PRAGMA_ENABLE_DEPRECATION_WARNINGS
 
 	// Create constraints
 	const bool bEnabled = false;  // Set constraint disabled by default
-	ClothConstraints.AddRules(ConfigProperties, TriangleMesh, &PatternData, WeightMaps, Tethers, MeshScale, bEnabled);
+	ClothConstraints.AddRules(ConfigProperties, TriangleMesh, &PatternData, WeightMaps, VertexSets, Tethers, MeshScale, bEnabled);
 
 	// Update LOD stats
 	const TConstArrayView<Softs::FSolverReal> InvMasses(Solver->GetParticleInvMasses(Offset), NumParticles);
@@ -252,7 +256,7 @@ PRAGMA_DISABLE_DEPRECATION_WARNINGS  // TODO: CHAOS_IS_CLOTHINGSIMULATIONMESH_AB
 	const Softs::FSolverReal MeshScale = Cloth->Mesh->GetScale();
 PRAGMA_ENABLE_DEPRECATION_WARNINGS
 	const Softs::FSolverReal MaxDistancesScale = (Softs::FSolverReal)Cloth->MaxDistancesMultiplier;
-	ClothConstraints.Update(Cloth->Config->GetProperties(SolverDatum.LODIndex), WeightMaps, MeshScale, MaxDistancesScale);
+	ClothConstraints.Update(Cloth->Config->GetProperties(SolverDatum.LODIndex), WeightMaps, VertexSets, MeshScale, MaxDistancesScale);
 }
 
 void FClothingSimulationCloth::FLODData::Enable(FClothingSimulationSolver* Solver, bool bEnable) const
@@ -567,6 +571,8 @@ void FClothingSimulationCloth::SetMesh(FClothingSimulationMesh* InMesh)
 					TConstArrayView<FRealSingle>());
 		}
 
+		TMap<FString, const TSet<int32>*> VertexSets = Mesh->GetVertexSets(LODIndex);
+
 		const bool bUseGeodesicTethers = Config->GetProperties(LODIndex).GetValue<bool>(TEXT("UseGeodesicTethers"), ClothingSimulationClothDefault::bUseGeodesicTethers);
 
 		// Add LOD data
@@ -577,6 +583,7 @@ void FClothingSimulationCloth::SetMesh(FClothingSimulationMesh* InMesh)
 			Mesh->GetPatternIndices(LODIndex),
 			Mesh->GetPatternToWeldedIndices(LODIndex),
 			MoveTemp(WeightMaps),
+			MoveTemp(VertexSets),
 			Mesh->GetTethers(LODIndex, bUseGeodesicTethers)));
 	}
 
diff --git a/Engine/Plugins/Experimental/ChaosCloth/Source/Chaos/Private/ChaosCloth/ChaosClothingSimulationMesh.cpp b/Engine/Plugins/Experimental/ChaosCloth/Source/Chaos/Private/ChaosCloth/ChaosClothingSimulationMesh.cpp
--- a/Engine/Plugins/Experimental/ChaosCloth/Source/Chaos/Private/ChaosCloth/ChaosClothingSimulationMesh.cpp
+++ b/Engine/Plugins/Experimental/ChaosCloth/Source/Chaos/Private/ChaosCloth/ChaosClothingSimulationMesh.cpp
@@ -200,6 +200,12 @@ TArray<TConstArrayView<FRealSingle>> FClothingSimulationMesh::GetWeightMaps(int3
 	return WeightMaps;
 }
 
+TMap<FString, const TSet<int32>*> FClothingSimulationMesh::GetVertexSets(int32 LODIndex) const
+{
+	// Not supported
+	return TMap<FString, const TSet<int32>*>();
+}
+
 TArray<TConstArrayView<TTuple<int32, int32, float>>> FClothingSimulationMesh::GetTethers(int32 LODIndex, bool bUseGeodesicTethers) const
 {
 	TArray<TConstArrayView<TTuple<int32, int32, float>>> Tethers;
diff --git a/Engine/Plugins/Experimental/ChaosCloth/Source/Chaos/Private/ChaosCloth/ChaosClothingSimulationSkeletalMesh.cpp b/Engine/Plugins/Experimental/ChaosCloth/Source/Chaos/Private/ChaosCloth/ChaosClothingSimulationSkeletalMesh.cpp
--- a/Engine/Plugins/Experimental/ChaosCloth/Source/Chaos/Private/ChaosCloth/ChaosClothingSimulationSkeletalMesh.cpp
+++ b/Engine/Plugins/Experimental/ChaosCloth/Source/Chaos/Private/ChaosCloth/ChaosClothingSimulationSkeletalMesh.cpp
@@ -159,6 +159,12 @@ PRAGMA_ENABLE_DEPRECATION_WARNINGS
 		return WeightMaps;
 	}
 
+	TMap<FString, const TSet<int32>*> FClothingSimulationSkeletalMesh::GetVertexSets(int32 LODIndex) const
+	{
+		// Not supported
+		return TMap<FString, const TSet<int32>*>();
+	}
+
 	TArray<TConstArrayView<TTuple<int32, int32, float>>> FClothingSimulationSkeletalMesh::GetTethers(int32 LODIndex, bool bUseGeodesicTethers) const
 	{
 		TArray<TConstArrayView<TTuple<int32, int32, float>>> Tethers;
diff --git a/Engine/Plugins/Experimental/ChaosCloth/Source/Chaos/Public/ChaosCloth/ChaosClothConstraints.h b/Engine/Plugins/Experimental/ChaosCloth/Source/Chaos/Public/ChaosCloth/ChaosClothConstraints.h
--- a/Engine/Plugins/Experimental/ChaosCloth/Source/Chaos/Public/ChaosCloth/ChaosClothConstraints.h
+++ b/Engine/Plugins/Experimental/ChaosCloth/Source/Chaos/Public/ChaosCloth/ChaosClothConstraints.h
@@ -49,11 +49,22 @@ PRAGMA_ENABLE_DEPRECATION_WARNINGS
 			const FTriangleMesh& TriangleMesh,
 			const FClothingPatternData* PatternData,
 			const TMap<FString, TConstArrayView<FRealSingle>>& WeightMaps,
+			const TMap<FString, const TSet<int32>*>& VertexSets,
 			const TArray<TConstArrayView<TTuple<int32, int32, FRealSingle>>>& Tethers,
 			Softs::FSolverReal MeshScale,
 			bool bEnabled);
 
-		UE_DEPRECATED(5.3, "Use AddRules() with WeightMaps and optional PatternData instead.")
+		UE_DEPRECATED(5.4, "Use AddRules() with WeightMaps, VertexSets, and optional PatternData instead.")
+		void AddRules(
+			const Softs::FCollectionPropertyConstFacade& ConfigProperties,
+			const FTriangleMesh& TriangleMesh,
+			const FClothingPatternData* PatternData,
+			const TMap<FString, TConstArrayView<FRealSingle>>& WeightMaps,
+			const TArray<TConstArrayView<TTuple<int32, int32, FRealSingle>>>& Tethers,
+			Softs::FSolverReal MeshScale,
+			bool bEnabled);
+
+		UE_DEPRECATED(5.3, "Use AddRules() with WeightMaps, VertexSets, and optional PatternData instead.")
 		void AddRules(
 			const Softs::FCollectionPropertyConstFacade& ConfigProperties,
 			const FTriangleMesh& TriangleMesh,
@@ -62,6 +73,14 @@ PRAGMA_ENABLE_DEPRECATION_WARNINGS
 			Softs::FSolverReal MeshScale,
 			bool bEnabled);
 
+		void Update(
+			const Softs::FCollectionPropertyConstFacade& ConfigProperties,
+			const TMap<FString, TConstArrayView<FRealSingle>>& WeightMaps,
+			const TMap<FString, const TSet<int32>*>& VertexSets,
+			Softs::FSolverReal MeshScale,
+			Softs::FSolverReal MaxDistancesScale = (Softs::FSolverReal)1.);
+
+		UE_DEPRECATED(5.4, "Use Update() with WeightMaps and VertexSets instead.")
 		void Update(
 			const Softs::FCollectionPropertyConstFacade& ConfigProperties,
 			const TMap<FString, TConstArrayView<FRealSingle>>& WeightMaps,
@@ -193,11 +212,14 @@ PRAGMA_ENABLE_DEPRECATION_WARNINGS
 		const TSharedPtr<Softs::FPBDCollisionSpringConstraints>& GetSelfCollisionConstraints() const { return SelfCollisionConstraints; }
 		const TSharedPtr<Softs::FPBDTriangleMeshIntersections>& GetSelfIntersectionConstraints() const { return SelfIntersectionConstraints; }
 		const TSharedPtr<Softs::FPBDTriangleMeshCollisions>& GetSelfCollisionInit() const { return SelfCollisionInit; }
+		const TSharedPtr<Softs::FPBDSelfCollisionSphereConstraints>& GetSelfCollisionSphereConstraints() const
+		{ return SelfCollisionSphereConstraints; }
 		// ---- End of debug functions ----
 
 	private:
 		void CreateSelfCollisionConstraints(
 			const Softs::FCollectionPropertyConstFacade& ConfigProperties,
+			const TMap<FString, const TSet<int32>*>& VertexSets,
 			const FTriangleMesh& TriangleMesh);
 		void CreateStretchConstraints(
 			const Softs::FCollectionPropertyConstFacade& ConfigProperties,
@@ -260,6 +282,7 @@ PRAGMA_ENABLE_DEPRECATION_WARNINGS
 		TSharedPtr<Softs::FPBDTriangleMeshCollisions> SelfCollisionInit;
 		TSharedPtr<Softs::FPBDCollisionSpringConstraints> SelfCollisionConstraints;
 		TSharedPtr<Softs::FPBDTriangleMeshIntersections> SelfIntersectionConstraints;
+		TSharedPtr<Softs::FPBDSelfCollisionSphereConstraints> SelfCollisionSphereConstraints;
 		
 		Softs::FPBDEvolution* Evolution;
 		const TArray<Softs::FSolverVec3>* AnimationPositions;
diff --git a/Engine/Plugins/Experimental/ChaosCloth/Source/Chaos/Public/ChaosCloth/ChaosClothingSimulationMesh.h b/Engine/Plugins/Experimental/ChaosCloth/Source/Chaos/Public/ChaosCloth/ChaosClothingSimulationMesh.h
--- a/Engine/Plugins/Experimental/ChaosCloth/Source/Chaos/Public/ChaosCloth/ChaosClothingSimulationMesh.h
+++ b/Engine/Plugins/Experimental/ChaosCloth/Source/Chaos/Public/ChaosCloth/ChaosClothingSimulationMesh.h
@@ -72,6 +72,7 @@ namespace Chaos
 		UE_DEPRECATED(5.3, "Use LODIndex version.")
 		virtual TMap<FString, int32> GetWeightMapIndices() const { return GetWeightMapIndices(0); }
 		virtual TArray<TConstArrayView<FRealSingle>> GetWeightMaps(int32 LODIndex) const;
+		virtual TMap<FString, const TSet<int32>*> GetVertexSets(int32 LODIndex) const;
 		virtual TArray<TConstArrayView<TTuple<int32, int32, float>>> GetTethers(int32 LODIndex, bool bUseGeodesicTethers) const;
 		virtual int32 GetReferenceBoneIndex() const;
 		virtual FTransform GetReferenceBoneTransform() const;
@@ -132,6 +133,9 @@ namespace Chaos
 		/* Return the specified LOD's weight map. */
 		virtual TArray<TConstArrayView<FRealSingle>> GetWeightMaps(int32 LODIndex) const = 0;
 
+		/* Return the specified LOD's vertex sets. */
+		virtual const TMap<FString, TSet<int32>*> GetVertexSets(int32 LODIndex) const = 0;
+
 		/* Return the tethers connections for the long range attachment into convenient parallel friendly batches. */
 		virtual TArray<TConstArrayView<TTuple<int32, int32, float>>> GetTethers(int32 LODIndex, bool bUseGeodesicTethers) const = 0;
 
diff --git a/Engine/Plugins/Experimental/ChaosCloth/Source/Chaos/Public/ChaosCloth/ChaosClothingSimulationSkeletalMesh.h b/Engine/Plugins/Experimental/ChaosCloth/Source/Chaos/Public/ChaosCloth/ChaosClothingSimulationSkeletalMesh.h
--- a/Engine/Plugins/Experimental/ChaosCloth/Source/Chaos/Public/ChaosCloth/ChaosClothingSimulationSkeletalMesh.h
+++ b/Engine/Plugins/Experimental/ChaosCloth/Source/Chaos/Public/ChaosCloth/ChaosClothingSimulationSkeletalMesh.h
@@ -52,6 +52,7 @@ PRAGMA_ENABLE_DEPRECATION_WARNINGS
 		UE_DEPRECATED(5.3, "Use LODIndex version.")
 		virtual TMap<FString, int32> GetWeightMapIndices() const override { return GetWeightMapIndices(0); }
 		virtual TArray<TConstArrayView<FRealSingle>> GetWeightMaps(int32 LODIndex) const override;
+		virtual TMap<FString, const TSet<int32>*> GetVertexSets(int32 LODIndex) const override;
 		virtual TArray<TConstArrayView<TTuple<int32, int32, float>>> GetTethers(int32 LODIndex, bool bUseGeodesicTethers) const override;
 		virtual int32 GetReferenceBoneIndex() const override;
 		virtual FTransform GetReferenceBoneTransform() const override;
diff --git a/Engine/Plugins/Experimental/ChaosClothAsset/Source/ChaosClothAsset/Private/ChaosClothAsset/CollectionClothSelectionFacade.cpp b/Engine/Plugins/Experimental/ChaosClothAsset/Source/ChaosClothAsset/Private/ChaosClothAsset/CollectionClothSelectionFacade.cpp
--- a/Engine/Plugins/Experimental/ChaosClothAsset/Source/ChaosClothAsset/Private/ChaosClothAsset/CollectionClothSelectionFacade.cpp
+++ b/Engine/Plugins/Experimental/ChaosClothAsset/Source/ChaosClothAsset/Private/ChaosClothAsset/CollectionClothSelectionFacade.cpp
@@ -86,6 +86,34 @@ namespace UE::Chaos::ClothAsset
 		}
 	}
 
+	void FCollectionClothSelectionFacade::Append(const FCollectionClothSelectionConstFacade& Other, bool bOverwriteExistingIfMismatched)
+	{
+		if (Other.IsValid())
+		{
+			const int32 NumInSelections = Other.GetNumSelections();
+			const TArray<FName> InSelectionNames = Other.GetNames();
+			for (int32 InSelectionIndex = 0; InSelectionIndex < NumInSelections; ++InSelectionIndex)
+			{
+				const FName& SelectionName = InSelectionNames[InSelectionIndex];
+				if (HasSelection(SelectionName))
+				{
+					if (GetSelectionGroup(SelectionName) == Other.GetSelectionGroup(SelectionName))
+					{
+						TSet<int32>& UnionedSet = GetSelectionSet(SelectionName);
+						UnionedSet.Append(Other.GetSelectionSet(SelectionName));
+						continue;
+					}
+					if (!bOverwriteExistingIfMismatched)
+					{
+						continue;
+					}
+				}
+
+				FindOrAddSelectionSet(SelectionName, Other.GetSelectionGroup(SelectionName)) = Other.GetSelectionSet(SelectionName);
+			}
+		}
+	}
+
 	TSet<int32>& FCollectionClothSelectionFacade::GetSelectionSet(const FName& Name)
 	{
 		check(IsValid());
diff --git a/Engine/Plugins/Experimental/ChaosClothAsset/Source/ChaosClothAsset/Public/ChaosClothAsset/CollectionClothSelectionFacade.h b/Engine/Plugins/Experimental/ChaosClothAsset/Source/ChaosClothAsset/Public/ChaosClothAsset/CollectionClothSelectionFacade.h
--- a/Engine/Plugins/Experimental/ChaosClothAsset/Source/ChaosClothAsset/Public/ChaosClothAsset/CollectionClothSelectionFacade.h
+++ b/Engine/Plugins/Experimental/ChaosClothAsset/Source/ChaosClothAsset/Public/ChaosClothAsset/CollectionClothSelectionFacade.h
@@ -75,6 +75,16 @@ namespace UE::Chaos::ClothAsset
 		/** Find a selection with the given name, or nullptr if no such selection exists. */
 		CHAOSCLOTHASSET_API TSet<int32>* FindSelectionSet(const FName& Name);
 
+		/** 
+		 * Append all sets from an existing collection to this collection.
+		 * Matching sets (i.e., same name and type) will be unioned.
+		 * Mismatching sets (same name, different type) will be handled according to bOverwriteExistingIfMismatched
+		 * 
+		 * @param bOverwriteExistingIfMismatched: If true, overwrite existing mismatched sets.
+		 *   If false, keep the existing set.
+		 */
+		CHAOSCLOTHASSET_API void Append(const FCollectionClothSelectionConstFacade& Other, bool bOverwriteExistingIfMismatched);
+
 		/**
 		 * Find, or add if it doesn't already exist, a selection for the specified group with the given name.
 		 * If the group doesn't already exists, this function will create it.
diff --git a/Engine/Plugins/Experimental/ChaosClothAsset/Source/ChaosClothAssetEngine/Private/ChaosClothAsset/ClothSimulationMesh.cpp b/Engine/Plugins/Experimental/ChaosClothAsset/Source/ChaosClothAssetEngine/Private/ChaosClothAsset/ClothSimulationMesh.cpp
--- a/Engine/Plugins/Experimental/ChaosClothAsset/Source/ChaosClothAssetEngine/Private/ChaosClothAsset/ClothSimulationMesh.cpp
+++ b/Engine/Plugins/Experimental/ChaosClothAsset/Source/ChaosClothAssetEngine/Private/ChaosClothAsset/ClothSimulationMesh.cpp
@@ -123,6 +123,20 @@ PRAGMA_ENABLE_DEPRECATION_WARNINGS
 		return WeightMaps;
 	}
 
+	TMap<FString, const TSet<int32>*> FClothSimulationMesh::GetVertexSets(int32 LODIndex) const
+	{
+		TMap<FString, const TSet<int32>*> VertexSets;
+		if (ClothSimulationModel.ClothSimulationLodModels.IsValidIndex(LODIndex))
+		{
+			VertexSets.Reserve(ClothSimulationModel.ClothSimulationLodModels[LODIndex].VertexSets.Num());
+			for (TMap<FName, TSet<int32>>::TConstIterator Iter = ClothSimulationModel.ClothSimulationLodModels[LODIndex].VertexSets.CreateConstIterator(); Iter; ++Iter)
+			{
+				VertexSets.Emplace(Iter.Key().ToString(), &Iter.Value());
+			}
+		}
+		return VertexSets;
+	}
+
 	TArray<TConstArrayView<TTuple<int32, int32, float>>> FClothSimulationMesh::GetTethers(int32 LODIndex, bool /*bUseGeodesicTethers*/) const
 	{
 		return ClothSimulationModel.GetTethers(LODIndex);
diff --git a/Engine/Plugins/Experimental/ChaosClothAsset/Source/ChaosClothAssetEngine/Private/ChaosClothAsset/ClothSimulationMesh.h b/Engine/Plugins/Experimental/ChaosClothAsset/Source/ChaosClothAssetEngine/Private/ChaosClothAsset/ClothSimulationMesh.h
--- a/Engine/Plugins/Experimental/ChaosClothAsset/Source/ChaosClothAssetEngine/Private/ChaosClothAsset/ClothSimulationMesh.h
+++ b/Engine/Plugins/Experimental/ChaosClothAsset/Source/ChaosClothAssetEngine/Private/ChaosClothAsset/ClothSimulationMesh.h
@@ -45,6 +45,7 @@ PRAGMA_ENABLE_DEPRECATION_WARNINGS
 		UE_DEPRECATED(5.3, "Use LODIndex version.")
 		virtual TMap<FString, int32> GetWeightMapIndices() const override { return GetWeightMapIndices(0); }
 		virtual TArray<TConstArrayView<::Chaos::FRealSingle>> GetWeightMaps(int32 LODIndex) const override;
+		virtual TMap<FString, const TSet<int32>*> GetVertexSets(int32 LODIndex) const override;
 		// Note: there is only one set of tethers stored on ClothSimulationMesh assets
 		virtual TArray<TConstArrayView<TTuple<int32, int32, float>>> GetTethers(int32 LODIndex, bool /*bUseGeodesicTethers*/) const override;
 		virtual int32 GetReferenceBoneIndex() const override;
diff --git a/Engine/Plugins/Experimental/ChaosClothAsset/Source/ChaosClothAssetEngine/Private/ChaosClothAsset/ClothSimulationModel.cpp b/Engine/Plugins/Experimental/ChaosClothAsset/Source/ChaosClothAssetEngine/Private/ChaosClothAsset/ClothSimulationModel.cpp
--- a/Engine/Plugins/Experimental/ChaosClothAsset/Source/ChaosClothAssetEngine/Private/ChaosClothAsset/ClothSimulationModel.cpp
+++ b/Engine/Plugins/Experimental/ChaosClothAsset/Source/ChaosClothAssetEngine/Private/ChaosClothAsset/ClothSimulationModel.cpp
@@ -2,6 +2,8 @@
 
 #include "ChaosClothAsset/ClothSimulationModel.h"
 #include "ChaosClothAsset/CollectionClothFacade.h"
+#include "ChaosClothAsset/CollectionClothSelectionFacade.h"
+#include "ChaosClothAsset/ClothCollectionGroup.h"
 #include "ChaosClothAsset/ClothGeometryTools.h"
 #include "ChaosClothAsset/ClothLodTransitionDataCache.h"
 #include "Utils/ClothingMeshUtils.h"
@@ -143,6 +145,8 @@ bool FChaosClothSimulationLodModel::Serialize(FArchive& Ar)
 	Ar << LODTransitionUpData;
 	Ar << LODTransitionDownData;
 
+	Ar << VertexSets;
+
 	// Return true to confirm that serialization has already been taken care of
 	return true;
 }
@@ -177,6 +181,18 @@ FChaosClothSimulationModel::FChaosClothSimulationModel(const TArray<TSharedRef<c
 			LodModel.WeightMaps.Add(WeightMapName) = Cloth.GetWeightMap(WeightMapName);
 		}
 
+		// Copy vertex sets
+		FCollectionClothSelectionConstFacade Selection(ClothCollections[LodIndex]);
+		const TArray<FName> SelectionNames = Selection.GetNames();
+		LodModel.VertexSets.Reserve(SelectionNames.Num()); // At this point, only the SimVertex3d selections should be in the collection.
+		for (const FName& SelectionName : SelectionNames)
+		{
+			if (Selection.GetSelectionGroup(SelectionName) == ClothCollectionGroup::SimVertices3D)
+			{
+				LodModel.VertexSets.Add(SelectionName) = Selection.GetSelectionSet(SelectionName);
+			}
+		}
+
 		// Copy bone influences and gather tether data
 		const int32 NumSimVertices3D = Cloth.GetNumSimVertices3D();
 		TConstArrayView<TArray<int32>> SimBoneIndices = Cloth.GetSimBoneIndices();
diff --git a/Engine/Plugins/Experimental/ChaosClothAsset/Source/ChaosClothAssetEngine/Private/ChaosClothAsset/ClothSimulationModel.h b/Engine/Plugins/Experimental/ChaosClothAsset/Source/ChaosClothAssetEngine/Private/ChaosClothAsset/ClothSimulationModel.h
--- a/Engine/Plugins/Experimental/ChaosClothAsset/Source/ChaosClothAssetEngine/Private/ChaosClothAsset/ClothSimulationModel.h
+++ b/Engine/Plugins/Experimental/ChaosClothAsset/Source/ChaosClothAssetEngine/Private/ChaosClothAsset/ClothSimulationModel.h
@@ -62,6 +62,9 @@ struct FChaosClothSimulationLodModel
 	UPROPERTY()
 	FClothTetherData TetherData;
 
+	/** Vertex sets */
+	TMap<FName, TSet<int32>> VertexSets;
+
 	// Custom serialize for weight maps
 	bool Serialize(FArchive& Ar);
 };
diff --git a/Engine/Plugins/Experimental/ChaosClothAssetEditor/Source/ChaosClothAssetDataflowNodes/Private/ChaosClothAsset/ChaosClothAssetDataflowNodesModule.cpp b/Engine/Plugins/Experimental/ChaosClothAssetEditor/Source/ChaosClothAssetDataflowNodes/Private/ChaosClothAsset/ChaosClothAssetDataflowNodesModule.cpp
--- a/Engine/Plugins/Experimental/ChaosClothAssetEditor/Source/ChaosClothAssetDataflowNodes/Private/ChaosClothAsset/ChaosClothAssetDataflowNodesModule.cpp
+++ b/Engine/Plugins/Experimental/ChaosClothAssetEditor/Source/ChaosClothAssetDataflowNodes/Private/ChaosClothAsset/ChaosClothAssetDataflowNodesModule.cpp
@@ -30,6 +30,7 @@
 #include "ChaosClothAsset/SimulationPBDEdgeSpringConfigNode.h"
 #include "ChaosClothAsset/SimulationPressureConfigNode.h"
 #include "ChaosClothAsset/SimulationSelfCollisionConfigNode.h"
+#include "ChaosClothAsset/SimulationSelfCollisionSpheresConfigNode.h"
 #include "ChaosClothAsset/SimulationSolverConfigNode.h"
 #include "ChaosClothAsset/SimulationVelocityScaleConfigNode.h"
 #include "ChaosClothAsset/SimulationXPBDAreaSpringConfigNode.h"
@@ -81,6 +82,7 @@ namespace UE::Chaos::ClothAsset::Private
 		DATAFLOW_NODE_REGISTER_CREATION_FACTORY(FChaosClothAssetSimulationPBDEdgeSpringConfigNode);
 		DATAFLOW_NODE_REGISTER_CREATION_FACTORY(FChaosClothAssetSimulationPressureConfigNode);
 		DATAFLOW_NODE_REGISTER_CREATION_FACTORY(FChaosClothAssetSimulationSelfCollisionConfigNode);
+		DATAFLOW_NODE_REGISTER_CREATION_FACTORY(FChaosClothAssetSimulationSelfCollisionSpheresConfigNode);
 		DATAFLOW_NODE_REGISTER_CREATION_FACTORY(FChaosClothAssetSimulationSolverConfigNode);
 		DATAFLOW_NODE_REGISTER_CREATION_FACTORY(FChaosClothAssetSimulationVelocityScaleConfigNode);
 		DATAFLOW_NODE_REGISTER_CREATION_FACTORY(FChaosClothAssetSimulationXPBDAnisoBendingConfigNode);
diff --git a/Engine/Plugins/Experimental/ChaosClothAssetEditor/Source/ChaosClothAssetDataflowNodes/Private/ChaosClothAsset/MergeClothCollectionsNode.cpp b/Engine/Plugins/Experimental/ChaosClothAssetEditor/Source/ChaosClothAssetDataflowNodes/Private/ChaosClothAsset/MergeClothCollectionsNode.cpp
--- a/Engine/Plugins/Experimental/ChaosClothAssetEditor/Source/ChaosClothAssetDataflowNodes/Private/ChaosClothAsset/MergeClothCollectionsNode.cpp
+++ b/Engine/Plugins/Experimental/ChaosClothAssetEditor/Source/ChaosClothAssetDataflowNodes/Private/ChaosClothAsset/MergeClothCollectionsNode.cpp
@@ -2,6 +2,7 @@
 
 #include "ChaosClothAsset/MergeClothCollectionsNode.h"
 #include "ChaosClothAsset/CollectionClothFacade.h"
+#include "ChaosClothAsset/CollectionClothSelectionFacade.h"
 #include "Chaos/CollectionPropertyFacade.h"
 #include "Dataflow/DataflowInputOutput.h"
 
@@ -40,6 +41,9 @@ void FChaosClothAssetMergeClothCollectionsNode::Evaluate(Dataflow::FContext& Con
 		FCollectionPropertyMutableFacade PropertyFacade(ClothCollection);
 		bAreAnyValid |= PropertyFacade.IsValid();
 
+		FCollectionClothSelectionFacade SelectionFacade(ClothCollection);
+		bAreAnyValid |= SelectionFacade.IsValid();
+
 		// Iterate through the inputs and append them to LOD 0
 		const TArray<const FManagedArrayCollection*> Collections = GetCollections();
 		for (int32 InputIndex = 1; InputIndex < Collections.Num(); ++InputIndex)
@@ -60,6 +64,14 @@ void FChaosClothAssetMergeClothCollectionsNode::Evaluate(Dataflow::FContext& Con
 				PropertyFacade.Append(OtherClothCollection.ToSharedPtr(), bUpdateExistingProperties);
 				bAreAnyValid = true;
 			}
+			// Copy selections
+			const FCollectionClothSelectionConstFacade OtherSelectionFacade(OtherClothCollection);
+			if (OtherSelectionFacade.IsValid())
+			{
+				constexpr bool bUpdateExistingSelections = true; // Want last one wins.
+				SelectionFacade.Append(OtherSelectionFacade, bUpdateExistingSelections);
+				bAreAnyValid = true;
+			}
 		}
 
 		// Set the output
diff --git a/Engine/Plugins/Experimental/ChaosClothAssetEditor/Source/ChaosClothAssetDataflowNodes/Private/ChaosClothAsset/SimulationBaseConfigNodePrivate.h b/Engine/Plugins/Experimental/ChaosClothAssetEditor/Source/ChaosClothAssetDataflowNodes/Private/ChaosClothAsset/SimulationBaseConfigNodePrivate.h
--- a/Engine/Plugins/Experimental/ChaosClothAssetEditor/Source/ChaosClothAssetDataflowNodes/Private/ChaosClothAsset/SimulationBaseConfigNodePrivate.h
+++ b/Engine/Plugins/Experimental/ChaosClothAssetEditor/Source/ChaosClothAssetDataflowNodes/Private/ChaosClothAsset/SimulationBaseConfigNodePrivate.h
@@ -20,6 +20,11 @@
 	const int32 Property##KeyIndex = AddPropertyHelper(Properties, FName(TEXT(#Property))); \
 	Properties.SetValue(Property##KeyIndex, (int32)Property);
 
+/** Macro for adding and setting string value from inside AddProperties() */
+#define UE_CHAOS_CLOTHASSET_SIMULATIONCONFIG_SETPROPERTYSTRING(Property) \
+	const int32 Property##KeyIndex = AddPropertyHelper(Properties, FName(TEXT(#Property))); \
+	Properties.SetStringValue(Property##KeyIndex, Property);
+
 /** Macro for adding and setting weighted value properties from inside AddProperties() and check for a similar property being overriden. */
 #define UE_CHAOS_CLOTHASSET_SIMULATIONCONFIG_SETPROPERTYCHECKED1(Property, SimilarProperty1) \
 	const int32 Property##KeyIndex = AddPropertyHelper(Properties, FName(TEXT(#Property)), true, \
diff --git a/Engine/Plugins/Experimental/ChaosClothAssetEditor/Source/ChaosClothAssetDataflowNodes/Private/ChaosClothAsset/SimulationSelfCollisionSpheresConfigNode.cpp b/Engine/Plugins/Experimental/ChaosClothAssetEditor/Source/ChaosClothAssetDataflowNodes/Private/ChaosClothAsset/SimulationSelfCollisionSpheresConfigNode.cpp
new file mode 100644
--- /dev/null
+++ b/Engine/Plugins/Experimental/ChaosClothAssetEditor/Source/ChaosClothAssetDataflowNodes/Private/ChaosClothAsset/SimulationSelfCollisionSpheresConfigNode.cpp
@@ -0,0 +1,66 @@
+// Copyright Epic Games, Inc. All Rights Reserved.
+
+#include "ChaosClothAsset/SimulationSelfCollisionSpheresConfigNode.h"
+#include "ChaosClothAsset/SimulationBaseConfigNodePrivate.h"
+#include "Chaos/CollectionPropertyFacade.h"
+
+#include "ChaosClothAsset/CollectionClothFacade.h"
+#include "ChaosClothAsset/CollectionClothSelectionFacade.h"
+#include "ChaosClothAsset/ClothCollectionGroup.h"
+
+#include UE_INLINE_GENERATED_CPP_BY_NAME(SimulationSelfCollisionSpheresConfigNode)
+
+FChaosClothAssetSimulationSelfCollisionSpheresConfigNode::FChaosClothAssetSimulationSelfCollisionSpheresConfigNode(const Dataflow::FNodeParameters& InParam, FGuid InGuid)
+	: FChaosClothAssetSimulationBaseConfigNode(InParam, InGuid)
+{
+	RegisterCollectionConnections();
+}
+
+void FChaosClothAssetSimulationSelfCollisionSpheresConfigNode::AddProperties(Dataflow::FContext& Context, ::Chaos::Softs::FCollectionPropertyMutableFacade& Properties) const
+{
+	UE_CHAOS_CLOTHASSET_SIMULATIONCONFIG_SETPROPERTY(SelfCollisionSphereRadius);
+	UE_CHAOS_CLOTHASSET_SIMULATIONCONFIG_SETPROPERTY(SelfCollisionSphereStiffness); 
+}
+
+void FChaosClothAssetSimulationSelfCollisionSpheresConfigNode::EvaluateClothCollection(Dataflow::FContext& Context, const TSharedRef<FManagedArrayCollection>& ClothCollection) const
+{
+	using namespace UE::Chaos::ClothAsset;
+	FCollectionClothConstFacade Cloth(ClothCollection);
+	const float CullDiameterSq = FMath::Square(SelfCollisionSphereRadius * SelfCollisionSphereRadiusCullMultiplier * 2.f);
+	if (Cloth.IsValid() && CullDiameterSq > 0.f)
+	{
+		TConstArrayView<FVector3f> SimPositions = Cloth.GetSimPosition3D();
+		TArray<bool> VertexIsValid;
+		VertexIsValid.Init(true, SimPositions.Num());
+
+		TSet<int32> VertexSet;
+		for (int32 Index = 0; Index < SimPositions.Num(); ++Index)
+		{
+			if (!VertexIsValid[Index])
+			{
+				continue;
+			}
+			VertexSet.Add(Index);
+
+			const FVector3f& Pos0 = SimPositions[Index];
+			for (int32 CompareIndex = Index + 1; CompareIndex < SimPositions.Num(); ++CompareIndex)
+			{
+				if (!VertexIsValid[CompareIndex])
+				{
+					continue;
+				}
+				if (FVector3f::DistSquared(Pos0, SimPositions[CompareIndex]) < CullDiameterSq)
+				{
+					VertexIsValid[CompareIndex] = false;
+				}
+			}
+		}
+
+		FCollectionClothSelectionFacade Selection(ClothCollection);
+		Selection.DefineSchema();
+
+		static const FName SelectionSetName("_SelfCollisionSpheres");
+		Selection.FindOrAddSelectionSet(SelectionSetName, ClothCollectionGroup::SimVertices3D) = VertexSet;
+	}
+
+}
diff --git a/Engine/Plugins/Experimental/ChaosClothAssetEditor/Source/ChaosClothAssetDataflowNodes/Private/ChaosClothAsset/TerminalNode.cpp b/Engine/Plugins/Experimental/ChaosClothAssetEditor/Source/ChaosClothAssetDataflowNodes/Private/ChaosClothAsset/TerminalNode.cpp
--- a/Engine/Plugins/Experimental/ChaosClothAssetEditor/Source/ChaosClothAssetDataflowNodes/Private/ChaosClothAsset/TerminalNode.cpp
+++ b/Engine/Plugins/Experimental/ChaosClothAssetEditor/Source/ChaosClothAssetDataflowNodes/Private/ChaosClothAsset/TerminalNode.cpp
@@ -6,6 +6,8 @@
 #include "ChaosClothAsset/ClothGeometryTools.h"
 #include "ChaosClothAsset/ClothLODTransitionDataCache.h"
 #include "ChaosClothAsset/CollectionClothFacade.h"
+#include "ChaosClothAsset/CollectionClothSelectionFacade.h"
+#include "ChaosClothAsset/ClothCollectionGroup.h"
 #include "Animation/Skeleton.h"
 #include "Chaos/CollectionPropertyFacade.h"
 #include "Dataflow/DataflowInputOutput.h"
@@ -30,6 +32,21 @@ namespace UE::Chaos::ClothAsset::Private
 			{
 				Checksum = Cloth.CalculateTypeHash(bIncludeWeightMapsTrue, Checksum);
 			}
+			FCollectionClothSelectionConstFacade Selection(ClothCollection);
+			if (Selection.IsValid())
+			{
+				// Just checksum the sets that are SimVertex3D sets since those are the only ones we care about right now
+				const TArray<FName> SelectionNames = Selection.GetNames();
+				for (const FName& SelectionName : SelectionNames)
+				{
+					if (Selection.GetSelectionGroup(SelectionName) == ClothCollectionGroup::SimVertices3D)
+					{
+						const TArray<int32> SelectionAsArray = Selection.GetSelectionSet(SelectionName).Array();
+						Checksum = HashCombineFast(Checksum, GetTypeHash(SelectionName));
+						Checksum = GetArrayHash(SelectionAsArray.GetData(), SelectionAsArray.Num(), Checksum);
+					}
+				}
+			}
 		}
 		return Checksum;
 	}
@@ -50,10 +67,10 @@ namespace UE::Chaos::ClothAsset::Private
 			}
 			for (int32 PropertyIndex = 0; PropertyIndex < Property0.Num(); ++PropertyIndex)
 			{
-				if (Property0.GetKey(PropertyIndex) != Property1.GetKey(PropertyIndex))
-				{
-					return false;
-				}
+if (Property0.GetKey(PropertyIndex) != Property1.GetKey(PropertyIndex))
+{
+	return false;
+}
 			}
 		}
 		return true;
@@ -140,6 +157,23 @@ void FChaosClothAssetTerminalNode::SetAssetValue(TObjectPtr<UObject> Asset, Data
 			constexpr bool bUpdateExistingProperties = false;
 			Chaos::Softs::FCollectionPropertyMutableFacade(ClothCollection).Append(InClothCollections[LodIndex].ToSharedPtr(), bUpdateExistingProperties);
 
+			// Set selections
+			FCollectionClothSelectionFacade Selection(ClothCollection);
+			FCollectionClothSelectionConstFacade InSelection(InClothCollections[LodIndex]);
+			if (InSelection.IsValid())
+			{
+				Selection.DefineSchema();
+				const TArray<FName> InSelectionNames = InSelection.GetNames();
+				for (const FName& InSelectionName : InSelectionNames)
+				{
+					if (InSelection.GetSelectionGroup(InSelectionName) == ClothCollectionGroup::SimVertices3D)
+					{
+						Selection.FindOrAddSelectionSet(InSelectionName, ClothCollectionGroup::SimVertices3D) = InSelection.GetSelectionSet(InSelectionName);
+					}
+				}
+			}
+
+
 			// Set physics asset only with LOD 0 at the moment
 			if (LodIndex == 0)
 			{
diff --git a/Engine/Plugins/Experimental/ChaosClothAssetEditor/Source/ChaosClothAssetDataflowNodes/Public/ChaosClothAsset/SimulationSelfCollisionSpheresConfigNode.h b/Engine/Plugins/Experimental/ChaosClothAssetEditor/Source/ChaosClothAssetDataflowNodes/Public/ChaosClothAsset/SimulationSelfCollisionSpheresConfigNode.h
new file mode 100644
--- /dev/null
+++ b/Engine/Plugins/Experimental/ChaosClothAssetEditor/Source/ChaosClothAssetDataflowNodes/Public/ChaosClothAsset/SimulationSelfCollisionSpheresConfigNode.h
@@ -0,0 +1,37 @@
+// Copyright Epic Games, Inc. All Rights Reserved.
+
+#pragma once 
+
+#include "ChaosClothAsset/SimulationBaseConfigNode.h"
+#include "SimulationSelfCollisionSpheresConfigNode.generated.h"
+
+/** Self-collision spheres properties configuration node. */
+USTRUCT(Meta = (DataflowCloth))
+struct FChaosClothAssetSimulationSelfCollisionSpheresConfigNode : public FChaosClothAssetSimulationBaseConfigNode
+{
+	GENERATED_USTRUCT_BODY()
+	DATAFLOW_NODE_DEFINE_INTERNAL(FChaosClothAssetSimulationSelfCollisionSpheresConfigNode, "SimulationSelfCollisionSpheresConfig", "Cloth", "Cloth Simulation Self Collision Spheres Config")
+
+public:
+	/** The radius of the spheres used in self collision centered at each vertex. */
+	UPROPERTY(EditAnywhere, Category = "Self-Collision Spheres Properties", meta = (UIMin = "0", UIMax = "100", ClampMin = "0", ClampMax = "1000"))
+	float SelfCollisionSphereRadius = 0.5f;
+
+	/** The stiffness of the springs used to control self collision. */
+	UPROPERTY(EditAnywhere, Category = "Self-Collision Spheres Properties", meta = (UIMin = "0", UIMax = "100", ClampMin = "0", ClampMax = "1000"))
+	float SelfCollisionSphereStiffness = 1.f;
+
+	/** 
+	 * Multiplier for culling the spheres generated by this node. Spheres are seeded on every vertex, 
+	 * and culled based on SelfCollisionSphereRadius * SelfCollisionSphereRadiusCullMultiplier.
+	 */
+	UPROPERTY(EditAnywhere, Category = "Self-Collision Spheres Properties", meta = (ClampMin = "1"))
+	float SelfCollisionSphereRadiusCullMultiplier = 1.f;
+
+	FChaosClothAssetSimulationSelfCollisionSpheresConfigNode(const Dataflow::FNodeParameters& InParam, FGuid InGuid = FGuid::NewGuid());
+
+private:
+	virtual void AddProperties(Dataflow::FContext& Context, ::Chaos::Softs::FCollectionPropertyMutableFacade& Properties) const override;
+
+	virtual void EvaluateClothCollection(Dataflow::FContext& Context, const TSharedRef<FManagedArrayCollection>& ClothCollection) const override;
+};
diff --git a/Engine/Source/Runtime/Experimental/Chaos/Private/Chaos/PBDSelfCollisionSphereConstraints.cpp b/Engine/Source/Runtime/Experimental/Chaos/Private/Chaos/PBDSelfCollisionSphereConstraints.cpp
new file mode 100644
--- /dev/null
+++ b/Engine/Source/Runtime/Experimental/Chaos/Private/Chaos/PBDSelfCollisionSphereConstraints.cpp
@@ -0,0 +1,145 @@
+// Copyright Epic Games, Inc. All Rights Reserved.
+#include "Chaos/PBDSelfCollisionSphereConstraints.h"
+
+#if !COMPILE_WITHOUT_UNREAL_SUPPORT
+#include "Chaos/Framework/Parallel.h"
+#include "Chaos/HierarchicalSpatialHash.h"
+
+namespace Chaos::Softs {
+
+	const FString FPBDSelfCollisionSphereConstraints::SelectionSetName("_SelfCollisionSpheres");
+
+	namespace Private
+	{
+		struct FSphereSpatialEntry
+		{
+			const TConstArrayView<FSolverVec3>* Points;
+			int32 Index;
+
+			FSolverVec3 X() const
+			{
+				return (*Points)[Index];
+			}
+
+			template<typename TPayloadType>
+			int32 GetPayload(int32) const
+			{
+				return Index;
+			}
+		};
+	}
+
+	FPBDSelfCollisionSphereConstraintsBase::FPBDSelfCollisionSphereConstraintsBase(
+		const int32 InOffset,
+		const int32 InNumParticles,
+		const TSet<int32>* InVertexSetNoOffset,
+		const TArray<FSolverVec3>* InReferencePositions,
+		const FSolverReal InRadius,
+		const FSolverReal InStiffness)
+		: Radius(InRadius)
+		, Stiffness(InStiffness)
+		, VertexSetNoOffset(InVertexSetNoOffset)
+		, Offset(InOffset)
+		, NumParticles(InNumParticles)
+		, ReferencePositions(InReferencePositions)
+	{
+	}
+
+	void FPBDSelfCollisionSphereConstraintsBase::Init(const FSolverParticles& Particles)
+	{
+		Constraints.Reset();
+
+		if (NumParticles == 0 || (VertexSetNoOffset && VertexSetNoOffset->IsEmpty()))
+		{
+			return;
+		}
+
+		// Build Spatial
+		TArray<Private::FSphereSpatialEntry> Entries;
+		const FSolverReal Diameter = 2.f * Radius;
+		TConstArrayView<FSolverVec3> Points = Particles.XArray();
+		if (VertexSetNoOffset)
+		{
+			Entries.Reset(VertexSetNoOffset->Num());
+			for (const int32 VertexNoOffset : *VertexSetNoOffset)
+			{
+				Entries.Add({ &Points, VertexNoOffset + Offset });
+			}
+		}
+		else
+		{
+			Entries.Reset(NumParticles);
+			for (int32 Index = 0; Index < NumParticles; ++Index)
+			{
+				Entries.Add({ &Points, Index + Offset });
+			}
+		}
+		TSpatialHashGridPoints<int32, FSolverReal> SpatialHash(Diameter);
+		SpatialHash.InitializePoints(Entries);
+		
+		const FSolverReal DiamSq = FMath::Square(Diameter);
+		constexpr int32 CellRadius = 1; // We set the cell size of the spatial hash such that we only need to look 1 cell away to find proximities.
+		constexpr int32 MaxNumExpectedConnectionsPerParticle = 3;
+		const int32 MaxNumExpectedConnections = MaxNumExpectedConnectionsPerParticle * Entries.Num();
+		Constraints = SpatialHash.FindAllSelfProximities(CellRadius, MaxNumExpectedConnections,
+			[this, &Particles, DiamSq](const int32 i1, const int32 i2)
+		{
+			const FSolverReal CombinedMass = Particles.InvM(i1) + Particles.InvM(i2);
+			if (CombinedMass < (FSolverReal)1e-7)
+			{
+				return false;
+			}
+			if (ReferencePositions)
+			{
+				if (FSolverVec3::DistSquared((*ReferencePositions)[i1], (*ReferencePositions)[i2]) < DiamSq)
+				{
+					return false;
+				}
+			}
+			return true;
+		}
+		);
+	}
+
+	void FPBDSelfCollisionSphereConstraintsBase::Apply(FSolverParticles& Particles, const FSolverReal Dt) const
+	{
+		const FSolverReal Diameter = 2.f * Radius;
+		const FSolverReal DiamSq = FMath::Square(Diameter);
+		for (const TVec2<int32>& Constraint : Constraints)
+		{
+			const int32 i1 = Constraint[0];
+			const int32 i2 = Constraint[1];
+			const FSolverReal CombinedMass = Particles.InvM(i1) + Particles.InvM(i2);
+			checkSlow(CombinedMass >= (FSolverReal)1e-7);
+
+			FSolverVec3& P1 = Particles.P(i1);
+			FSolverVec3& P2 = Particles.P(i2);
+			const FSolverVec3 Difference = P1 - P2;
+			const FSolverReal DistSq = Difference.SizeSquared();
+			if (DistSq >= DiamSq)
+			{
+				continue;
+			}
+			const FSolverReal Dist = FMath::Sqrt(DistSq);
+			const FSolverVec3 Delta = Stiffness * (Dist - Diameter) * Difference / (CombinedMass * Dist);
+			if (Particles.InvM(i1) > (FSolverReal)0.)
+			{
+				P1 -= Particles.InvM(i1) * Delta;
+			}
+			if (Particles.InvM(i2) > (FSolverReal)0.)
+			{
+				P2 += Particles.InvM(i2) * Delta;
+			}
+		}
+	}
+
+	void FPBDSelfCollisionSphereConstraints::SetProperties(
+		const FCollectionPropertyConstFacade& PropertyCollection,
+		const TMap<FString, const TSet<int32>*>& VertexSets)
+	{
+		Radius = (FSolverReal)FMath::Max(GetSelfCollisionSphereRadius(PropertyCollection), 0.f);
+		Stiffness = (FSolverReal)FMath::Clamp(GetSelfCollisionSphereStiffness(PropertyCollection), 0.f, 1.f);
+	}
+}  // End namespace Chaos::Softs
+
+#endif
diff --git a/Engine/Source/Runtime/Experimental/Chaos/Public/Chaos/HierarchicalSpatialHash.h b/Engine/Source/Runtime/Experimental/Chaos/Public/Chaos/HierarchicalSpatialHash.h
--- a/Engine/Source/Runtime/Experimental/Chaos/Public/Chaos/HierarchicalSpatialHash.h
+++ b/Engine/Source/Runtime/Experimental/Chaos/Public/Chaos/HierarchicalSpatialHash.h
@@ -69,10 +69,10 @@ inline uint32 GetTypeHash(const Chaos::FHierarchicalSpatialHashCellIdx& Idx)
 	return (uint32)Hash;
 }
 
-// Currently this assumes a 3 dimensional grid. 
-template<typename PayloadType, typename T>
-class THierarchicalSpatialHash
+template<typename PayloadType>
+class TSpatialHashGridBase
 {
+protected:	
 	struct FPayloadAndBounds
 	{
 		PayloadType Payload;
@@ -121,7 +121,17 @@ class THierarchicalSpatialHash
 						++IndexToWrite;
 					}
 				}
-			}			
+			}
+		}
+
+		int32 ConcurrentAddElement(const TVec3<int32>& CellIdx, int32 Lod, const FPayloadAndBounds& Value)
+		{
+			const int32 IndexToWrite = ConcurrentElementAddIdx.fetch_add(1, std::memory_order_relaxed);
+			Elements[IndexToWrite].Key = FHierarchicalSpatialHashCellIdx{ CellIdx, Lod };
+			Elements[IndexToWrite].Value = Value;
+			const uint32 HashKey = GetTypeHash(Elements[IndexToWrite].Key) & (Hash.Num() - 1);
+			Elements[IndexToWrite].NextIndex = Hash[HashKey].exchange(IndexToWrite);
+			return IndexToWrite;
 		}
 
 		void ShrinkElementsAfterConcurrentAdd()
@@ -157,6 +167,11 @@ class THierarchicalSpatialHash
 			return Elements[Index].Value;
 		}
 
+		const FHierarchicalSpatialHashCellIdx& GetKey(int32 Index) const
+		{
+			return Elements[Index].Key;
+		}
+
 	private:
 
 		TArray<std::atomic<int32>> Hash;
@@ -171,6 +186,23 @@ class THierarchicalSpatialHash
 		std::atomic<int32> ConcurrentElementAddIdx;
 	};
 
+	FHashMap HashMap;
+
+	void Reset()
+	{
+		HashMap.Reset();
+	}
+};
+
+// Currently this assumes a 3 dimensional grid. 
+template<typename PayloadType, typename T>
+class THierarchicalSpatialHash : public TSpatialHashGridBase<PayloadType>
+{
+	typedef TSpatialHashGridBase<PayloadType> Base;
+	using typename Base::FPayloadAndBounds;
+	using typename Base::FHashMap;
+	using Base::HashMap;
+
 	class FVectorAABB
 	{
 	public:
@@ -210,7 +242,7 @@ public:
 
 	void Reset()
 	{
-		HashMap.Reset();
+		Base::Reset();
 		UsedLods.Reset();
 	}
 
@@ -346,12 +378,161 @@ public:
 		return Result;
 	}
 
-
-
 private:
 	TMap<int32, T> UsedLods;
-	FHashMap HashMap;
 	TArray<FVectorAABB> Bounds;
 };
+
+template<typename PayloadType, typename T>
+class TSpatialHashGridPoints : public TSpatialHashGridBase<PayloadType>
+{
+	typedef TSpatialHashGridBase<PayloadType> Base;
+	using typename Base::FPayloadAndBounds;
+	using typename Base::FHashMap;
+	using Base::HashMap;
+
+public:
+	TSpatialHashGridPoints(const T InCellSize)
+		: Base()
+		, CellSize(InCellSize)
+	{}
+
+	void Reset()
+	{
+		Base::Reset();
+		ParticleIndexToElement.Reset();
+	}
+
+	template<typename ParticlesType>
+	void InitializePoints(const ParticlesType& Particles)
+	{
+		TRACE_CPUPROFILER_EVENT_SCOPE(SpatialHashGridConstantBounds_Init);
+		Reset();
+
+		ParticleIndexToElement.SetNum(Particles.Num());
+		HashMap.PreallocateElementsForConcurrentAdd(Particles.Num(), Particles.Num());
+		PhysicsParallelFor(Particles.Num(),
+			[this, &Particles](int32 ParticleIdx)
+		{
+			const auto& Particle = Particles[ParticleIdx];
+			const PayloadType Payload = Particle.template GetPayload<PayloadType>(ParticleIdx);
+			const TVec3<int32> CellIdx = FHierarchicalSpatialHashCellIdx::CellIdxForPoint(Particle.X(), CellSize);
+			const FPayloadAndBounds Value = { Payload, ParticleIdx };
+			const int32 ElementIdx = HashMap.ConcurrentAddElement(CellIdx, INDEX_NONE, Value);
+			ParticleIndexToElement[ParticleIdx] = ElementIdx;
+		}
+		);
+
+		HashMap.ShrinkElementsAfterConcurrentAdd();
+	}
+
+	TArray<TVec2<PayloadType>> FindAllSelfProximities(int32 CellRadius, int32 MaxNumExpectedConnections, TFunctionRef<bool(const int32 Payload0, const int32 Payload1)> NarrowTest)
+	{
+		check(CellRadius >= 0);
+
+		// Resize using a RWLock. Start with MaxNumExpectedConnections
+		TArray<TVec2<PayloadType>> Result;
+		Result.SetNum(MaxNumExpectedConnections);
+		Result.SetNum(Result.Max()); // Remove all slack.
+		std::atomic<int32> ResultIndex(0);
+		std::atomic<int32> CurrentResultNum(Result.Num()); // Arrays aren't thread-safe, so use this to track size instead.
+		FRWLock ResultResizeRWLock;
+
+		auto AddResult = [this, &NarrowTest, &Result, &ResultIndex, &CurrentResultNum, &ResultResizeRWLock]
+		(const TVec2<PayloadType>& Proximity)
+		{
+			if (!NarrowTest(Proximity[0], Proximity[1]))
+			{
+				return;
+			}
+			const int32 NewResultIndex = ResultIndex.fetch_add(1);
+			if (NewResultIndex < CurrentResultNum.load())
+			{
+				ResultResizeRWLock.ReadLock();
+				Result[NewResultIndex] = Proximity;
+				ResultResizeRWLock.ReadUnlock();
+			}
+			else
+			{
+				ResultResizeRWLock.WriteLock();
+				if (NewResultIndex >= CurrentResultNum.load())
+				{
+					// Use array's allocator to decide how much slack to create
+					Result.SetNum(NewResultIndex+1);
+					Result.SetNum(Result.Max());
+					CurrentResultNum.store(Result.Num());
+				}
+				Result[NewResultIndex] = Proximity;
+				ResultResizeRWLock.WriteUnlock();
+			}
+		};
+
+		auto AddResultsInCell = [this,&AddResult](PayloadType ThisPayload, const FHierarchicalSpatialHashCellIdx& Cell)
+		{
+			for (int32 HashIndex = HashMap.First(Cell); HashMap.IsValid(HashIndex); HashIndex = HashMap.Next(HashIndex))
+			{
+				if (HashMap.KeyMatches(HashIndex, Cell))
+				{
+					AddResult(TVec2<PayloadType>(ThisPayload, HashMap.GetValue(HashIndex).Payload));
+				}
+			}
+		};
+
+		PhysicsParallelFor(ParticleIndexToElement.Num(),
+			[this, &AddResult, &AddResultsInCell, CellRadius](int32 ParticleIdx)
+		{
+			const int32 ThisElementIdx = ParticleIndexToElement[ParticleIdx];
+			const int32 ThisPayload = HashMap.GetValue(ThisElementIdx).Payload;
+			const FHierarchicalSpatialHashCellIdx& ThisKey = HashMap.GetKey(ThisElementIdx);
+
+			// All neighbors in our cell after us in the list
+			int32 OtherElementIdx = HashMap.Next(ThisElementIdx);
+			for (; HashMap.IsValid(OtherElementIdx); OtherElementIdx = HashMap.Next(OtherElementIdx))
+			{
+				if (HashMap.KeyMatches(OtherElementIdx, ThisKey))
+				{
+					AddResult(TVec2<PayloadType>(ThisPayload, HashMap.GetValue(OtherElementIdx).Payload));
+				}
+			}
+
+			// Iterate over cells within CellRadius away. Only look "forward". Particles behind us will find us by looking forward.
+			// I dunno... there's probably a better way to do this.....
+			for (int32 ZIndex = 1; ZIndex <= CellRadius; ++ZIndex)
+			{
+				AddResultsInCell(ThisPayload, FHierarchicalSpatialHashCellIdx{ ThisKey.Cell + TVec3<int32>(0,0,ZIndex), INDEX_NONE });
+			}
+
+			for (int32 YIndex = 1; YIndex <= CellRadius; ++YIndex)
+			{
+				for (int32 ZIndex = -CellRadius; ZIndex <= CellRadius; ++ZIndex)
+				{
+					AddResultsInCell(ThisPayload, FHierarchicalSpatialHashCellIdx{ ThisKey.Cell + TVec3<int32>(0,YIndex,ZIndex), INDEX_NONE });
+				}
+			}
+
+			for (int32 XIndex = 1; XIndex <= CellRadius; ++XIndex)
+			{
+				for (int32 YIndex = -CellRadius; YIndex <= CellRadius; ++YIndex)
+				{
+					for (int32 ZIndex = -CellRadius; ZIndex <= CellRadius; ++ZIndex)
+					{
+						AddResultsInCell(ThisPayload, FHierarchicalSpatialHashCellIdx{ ThisKey.Cell + TVec3<int32>(XIndex,YIndex,ZIndex), INDEX_NONE });
+					}
+				}
+			}
+		}
+		);
+
+		// Shrink Result array to actual number of found proximities
+		const int32 ResultNum = ResultIndex.load();
+		Result.SetNum(ResultNum, /*bAllowShrinking*/ true);
+		return Result;
+	}
+private:
+	const T CellSize;
+
+	// Particle Index is dense array index passed in at build time. 
+	TArray<int32> ParticleIndexToElement;
+};
 }
 
diff --git a/Engine/Source/Runtime/Experimental/Chaos/Public/Chaos/PBDSelfCollisionSphereConstraints.h b/Engine/Source/Runtime/Experimental/Chaos/Public/Chaos/PBDSelfCollisionSphereConstraints.h
new file mode 100644
--- /dev/null
+++ b/Engine/Source/Runtime/Experimental/Chaos/Public/Chaos/PBDSelfCollisionSphereConstraints.h
@@ -0,0 +1,87 @@
+// Copyright Epic Games, Inc. All Rights Reserved.
+#pragma once
+
+#if !COMPILE_WITHOUT_UNREAL_SUPPORT
+#include "Chaos/Core.h"
+#include "Chaos/CollectionPropertyFacade.h"
+#include "Chaos/PBDSoftsEvolutionFwd.h"
+#include "Chaos/PBDSoftsSolverParticles.h"
+#include "Containers/Set.h"
+
+namespace Chaos::Softs
+{
+	class FPBDSelfCollisionSphereConstraintsBase
+	{
+	public:
+		CHAOS_API FPBDSelfCollisionSphereConstraintsBase(
+			const int32 InOffset,
+			const int32 InNumParticles,
+			const TSet<int32>* InVertexSetNoOffset,
+			const TArray<FSolverVec3>* InReferencePositions,
+			const FSolverReal InRadius,
+			const FSolverReal InStiffness);
+
+		virtual ~FPBDSelfCollisionSphereConstraintsBase() {}
+
+		CHAOS_API void Init(const FSolverParticles& Particles);
+
+		CHAOS_API void Apply(FSolverParticles& InParticles, const FSolverReal Dt) const;
+
+		const TArray<TVec2<int32>>& GetConstraints() const { return Constraints; }
+		const TSet<int32>* GetVertexSet() const { return VertexSetNoOffset; }
+		FSolverReal GetRadius() const { return Radius; }
+
+	protected:
+		TArray<TVec2<int32>> Constraints;
+		FSolverReal Radius;
+		FSolverReal Stiffness;
+		const TSet<int32>* VertexSetNoOffset;
+
+	private:
+		int32 Offset;
+		int32 NumParticles;
+		const TArray<FSolverVec3>* ReferencePositions;
+	};
+
+	class FPBDSelfCollisionSphereConstraints : public FPBDSelfCollisionSphereConstraintsBase
+	{
+		typedef FPBDSelfCollisionSphereConstraintsBase Base;
+
+	public:
+		CHAOS_API static const FString SelectionSetName;
+
+		static bool IsEnabled(const FCollectionPropertyConstFacade& PropertyCollection)
+		{
+			return IsSelfCollisionSphereStiffnessEnabled(PropertyCollection, false);
+		}
+
+		FPBDSelfCollisionSphereConstraints(
+			const int32 InOffset,
+			const int32 InNumParticles,
+			const TMap<FString, const TSet<int32>*>& VertexSets,
+			const FCollectionPropertyConstFacade& PropertyCollection)
+			: Base(InOffset,
+				InNumParticles,
+				VertexSets.FindRef(SelectionSetName, nullptr),
+				nullptr, // No reference positions--using vertex sets instead
+				(FSolverReal)FMath::Max(GetSelfCollisionSphereRadius(PropertyCollection, 0.5f), 0.f),
+				(FSolverReal)FMath::Clamp(GetSelfCollisionSphereStiffness(PropertyCollection, 0.5f), 0.f, 1.f)),
+			SelfCollisionSphereRadiusIndex(PropertyCollection),
+			SelfCollisionSphereStiffnessIndex(PropertyCollection)
+		{}
+
+		CHAOS_API void SetProperties(
+			const FCollectionPropertyConstFacade& PropertyCollection,
+			const TMap<FString, const TSet<int32>*>& VertexSets);
+
+	private:
+		using Base::Radius;
+		using Base::Stiffness;
+		using Base::VertexSetNoOffset;
+
+		UE_CHAOS_DECLARE_PROPERTYCOLLECTION_NAME(SelfCollisionSphereRadius, float);
+		UE_CHAOS_DECLARE_PROPERTYCOLLECTION_NAME(SelfCollisionSphereStiffness, float);
+
+	};
+}
+#endif
diff --git a/Engine/Source/Runtime/Experimental/Chaos/Public/Chaos/PBDSoftsEvolutionFwd.h b/Engine/Source/Runtime/Experimental/Chaos/Public/Chaos/PBDSoftsEvolutionFwd.h
--- a/Engine/Source/Runtime/Experimental/Chaos/Public/Chaos/PBDSoftsEvolutionFwd.h
+++ b/Engine/Source/Runtime/Experimental/Chaos/Public/Chaos/PBDSoftsEvolutionFwd.h
@@ -69,6 +69,7 @@ namespace Chaos
 		class FPBDCollisionSpringConstraints;
 		class FPBDTriangleMeshIntersections;
 		class FPBDTriangleMeshCollisions;
+		class FPBDSelfCollisionSphereConstraints;
 
 		// Softs solver forces
 		class FVelocityAndPressureField;
