Fixed geometry collection to run on the scene pipe and defer dynamic data until the proxy is fully created.

diff --git a/Engine/Source/Runtime/Experimental/GeometryCollectionEngine/Private/GeometryCollection/GeometryCollectionComponent.cpp b/Engine/Source/Runtime/Experimental/GeometryCollectionEngine/Private/GeometryCollection/GeometryCollectionComponent.cpp
--- a/Engine/Source/Runtime/Experimental/GeometryCollectionEngine/Private/GeometryCollection/GeometryCollectionComponent.cpp
+++ b/Engine/Source/Runtime/Experimental/GeometryCollectionEngine/Private/GeometryCollection/GeometryCollectionComponent.cpp
@@ -44,6 +44,7 @@
 #include "PhysicsSolver.h"
 #include "UObject/FortniteValkyrieBranchObjectVersion.h"
 #include "Chaos/PBDRigidClusteringAlgo.h"
+#include "Rendering/RenderCommandPipes.h"
 
 #include "Algo/RemoveIf.h"
 
@@ -1011,8 +1012,8 @@ FPrimitiveSceneProxy* UGeometryCollectionComponent::CreateSceneProxy()
 			if (bForceMotionBlur)
 			{
 				bIsMoving = true;
-				ENQUEUE_RENDER_COMMAND(NaniteProxyOnMotionEnd)(
-					[NaniteProxy](FRHICommandListImmediate& RHICmdList)
+				ENQUEUE_RENDER_COMMAND(NaniteProxyOnMotionEnd)(UE::RenderCommandPipe::Scene,
+					[NaniteProxy]
 					{
 						NaniteProxy->OnMotionBegin();
 					}
@@ -1623,8 +1624,8 @@ void UGeometryCollectionComponent::SetRestState(TArray<FTransform>&& InRestTrans
 		if (SceneProxy->IsNaniteMesh())
 		{
 			FNaniteGeometryCollectionSceneProxy* GeometryCollectionSceneProxy = static_cast<FNaniteGeometryCollectionSceneProxy*>(SceneProxy);
-			ENQUEUE_RENDER_COMMAND(SendRenderDynamicData)(
-				[GeometryCollectionSceneProxy, DynamicData](FRHICommandListImmediate& RHICmdList)
+			ENQUEUE_RENDER_COMMAND(SendRenderDynamicData)(UE::RenderCommandPipe::Scene,
+				[GeometryCollectionSceneProxy, DynamicData]
 				{
 					GeometryCollectionSceneProxy->SetDynamicData_RenderThread(DynamicData);
 				}
@@ -1633,10 +1634,10 @@ void UGeometryCollectionComponent::SetRestState(TArray<FTransform>&& InRestTrans
 		else
 		{
 			FGeometryCollectionSceneProxy* GeometryCollectionSceneProxy = static_cast<FGeometryCollectionSceneProxy*>(SceneProxy);
-			ENQUEUE_RENDER_COMMAND(SendRenderDynamicData)(
-				[GeometryCollectionSceneProxy, DynamicData](FRHICommandListImmediate& RHICmdList)
+			ENQUEUE_RENDER_COMMAND(SendRenderDynamicData)(UE::RenderCommandPipe::Scene,
+				[GeometryCollectionSceneProxy, DynamicData](FRHICommandListBase& RHICmdList)
 				{
-					GeometryCollectionSceneProxy->SetDynamicData_RenderThread(DynamicData);
+					GeometryCollectionSceneProxy->SetDynamicData_RenderThread(RHICmdList, DynamicData);
 				}
 			);
 		}
@@ -3056,8 +3057,8 @@ FGeometryCollectionDynamicData* UGeometryCollectionComponent::InitDynamicData(bo
 				if (SceneProxy && SceneProxy->IsNaniteMesh())
 				{
 					FNaniteGeometryCollectionSceneProxy* NaniteProxy = static_cast<FNaniteGeometryCollectionSceneProxy*>(SceneProxy);
-					ENQUEUE_RENDER_COMMAND(NaniteProxyOnMotionEnd)(
-						[NaniteProxy](FRHICommandListImmediate& RHICmdList)
+					ENQUEUE_RENDER_COMMAND(NaniteProxyOnMotionEnd)(UE::RenderCommandPipe::Scene,
+						[NaniteProxy]
 						{
 							NaniteProxy->OnMotionEnd();
 						}
@@ -3074,8 +3075,8 @@ FGeometryCollectionDynamicData* UGeometryCollectionComponent::InitDynamicData(bo
 				if (SceneProxy && SceneProxy->IsNaniteMesh())
 				{
 					FNaniteGeometryCollectionSceneProxy* NaniteProxy = static_cast<FNaniteGeometryCollectionSceneProxy*>(SceneProxy);
-					ENQUEUE_RENDER_COMMAND(NaniteProxyOnMotionBegin)(
-						[NaniteProxy](FRHICommandListImmediate& RHICmdList)
+					ENQUEUE_RENDER_COMMAND(NaniteProxyOnMotionBegin)(UE::RenderCommandPipe::Scene,
+						[NaniteProxy]
 						{
 							NaniteProxy->OnMotionBegin();
 						}
@@ -4298,8 +4299,8 @@ void UGeometryCollectionComponent::SendRenderDynamicData_Concurrent()
 			if (SceneProxy->IsNaniteMesh())
 			{
 				FNaniteGeometryCollectionSceneProxy* GeometryCollectionSceneProxy = static_cast<FNaniteGeometryCollectionSceneProxy*>(SceneProxy);
-				ENQUEUE_RENDER_COMMAND(SendRenderDynamicData)(
-					[GeometryCollectionSceneProxy, DynamicData](FRHICommandListImmediate& RHICmdList)
+				ENQUEUE_RENDER_COMMAND(SendRenderDynamicData)(UE::RenderCommandPipe::Scene,
+					[GeometryCollectionSceneProxy, DynamicData]
 					{
 						if (DynamicData)
 						{
@@ -4316,12 +4317,12 @@ void UGeometryCollectionComponent::SendRenderDynamicData_Concurrent()
 			else
 			{
 				FGeometryCollectionSceneProxy* GeometryCollectionSceneProxy = static_cast<FGeometryCollectionSceneProxy*>(SceneProxy);
-				ENQUEUE_RENDER_COMMAND(SendRenderDynamicData)(
-					[GeometryCollectionSceneProxy, DynamicData](FRHICommandListImmediate& RHICmdList)
+				ENQUEUE_RENDER_COMMAND(SendRenderDynamicData)(UE::RenderCommandPipe::Scene,
+					[GeometryCollectionSceneProxy, DynamicData](FRHICommandListBase& RHICmdList)
 					{
 						if (GeometryCollectionSceneProxy)
 						{
-							GeometryCollectionSceneProxy->SetDynamicData_RenderThread(DynamicData);
+							GeometryCollectionSceneProxy->SetDynamicData_RenderThread(RHICmdList, DynamicData);
 						}
 					}
 				);
diff --git a/Engine/Source/Runtime/Experimental/GeometryCollectionEngine/Private/GeometryCollection/GeometryCollectionSceneProxy.cpp b/Engine/Source/Runtime/Experimental/GeometryCollectionEngine/Private/GeometryCollection/GeometryCollectionSceneProxy.cpp
--- a/Engine/Source/Runtime/Experimental/GeometryCollectionEngine/Private/GeometryCollection/GeometryCollectionSceneProxy.cpp
+++ b/Engine/Source/Runtime/Experimental/GeometryCollectionEngine/Private/GeometryCollection/GeometryCollectionSceneProxy.cpp
@@ -377,6 +377,9 @@ void FGeometryCollectionSceneProxy::CreateRenderThreadResources(FRHICommandListB
 		bGeometryResourceUpdated = true;
 	}
 #endif
+
+	bRenderResourcesCreated = true;
+	SetDynamicData_RenderThread(RHICmdList, DynamicData);
 }
 
 void FGeometryCollectionSceneProxy::DestroyRenderThreadResources()
@@ -415,18 +418,19 @@ void FGeometryCollectionSceneProxy::DestroyRenderThreadResources()
 #endif
 }
 
-void FGeometryCollectionSceneProxy::SetDynamicData_RenderThread(FGeometryCollectionDynamicData* NewDynamicData)
+void FGeometryCollectionSceneProxy::SetDynamicData_RenderThread(FRHICommandListBase& RHICmdList, FGeometryCollectionDynamicData* NewDynamicData)
 {
-	check(IsInRenderingThread());
-	FRHICommandListImmediate& RHICmdList = FRHICommandListExecutor::GetImmediateCommandList();
-	if (DynamicData)
+	if (NewDynamicData != DynamicData)
 	{
-		GDynamicDataPool.Release(DynamicData);
-		DynamicData = nullptr;
+		if (DynamicData)
+		{
+			GDynamicDataPool.Release(DynamicData);
+			DynamicData = nullptr;
+		}
+		DynamicData = NewDynamicData;
 	}
-	DynamicData = NewDynamicData;
 
-	if (MeshDescription.NumVertices == 0)
+	if (MeshDescription.NumVertices == 0 || !DynamicData || !bRenderResourcesCreated)
 	{
 		return;
 	}
diff --git a/Engine/Source/Runtime/Experimental/GeometryCollectionEngine/Private/GeometryCollection/GeometryCollectionSceneProxy.h b/Engine/Source/Runtime/Experimental/GeometryCollectionEngine/Private/GeometryCollection/GeometryCollectionSceneProxy.h
--- a/Engine/Source/Runtime/Experimental/GeometryCollectionEngine/Private/GeometryCollection/GeometryCollectionSceneProxy.h
+++ b/Engine/Source/Runtime/Experimental/GeometryCollectionEngine/Private/GeometryCollection/GeometryCollectionSceneProxy.h
@@ -209,6 +209,7 @@ class FGeometryCollectionSceneProxy final : public FPrimitiveSceneProxy
 	int32 CurrentTransformBufferIndex = 0;
 	bool TransformVertexBuffersContainsRestTransforms = true;
 	bool bSupportsTripleBufferVertexUpload = false;
+	bool bRenderResourcesCreated = false;
 	TArray<FGeometryCollectionTransformBuffer, TInlineAllocator<3>> TransformBuffers;
 	TArray<FGeometryCollectionTransformBuffer, TInlineAllocator<3>> PrevTransformBuffers;
 
@@ -242,7 +243,7 @@ public:
 	virtual ~FGeometryCollectionSceneProxy();
 
 	/** Called on render thread to setup dynamic geometry for rendering */
-	void SetDynamicData_RenderThread(FGeometryCollectionDynamicData* NewDynamicData);
+	void SetDynamicData_RenderThread(FRHICommandListBase& RHICmdList, FGeometryCollectionDynamicData* NewDynamicData);
 
 	uint32 GetMemoryFootprint() const override { return sizeof(*this) + GetAllocatedSize(); }
 	uint32 GetAllocatedSize() const;
