In Event Manager fix the differed handler registration.

diff --git a/Engine/Source/Runtime/Experimental/Chaos/Private/EventManager.cpp b/Engine/Source/Runtime/Experimental/Chaos/Private/EventManager.cpp
--- a/Engine/Source/Runtime/Experimental/Chaos/Private/EventManager.cpp
+++ b/Engine/Source/Runtime/Experimental/Chaos/Private/EventManager.cpp
@@ -87,11 +87,6 @@ namespace Chaos
 
 	void FEventManager::DispatchEvents()
 	{
-		FScopeLock ScopeLock(&AccessDeferredHandlersLock);
-
-		check(!bCurrentlyDispatchingEvents);
-		bCurrentlyDispatchingEvents = true;
-		
 		if (BufferMode == EMultiBufferMode::Double)
 		{
 			ResourceLock.ReadLock();
@@ -111,22 +106,6 @@ namespace Chaos
 		{
 			ResourceLock.ReadUnlock();
 		}
-		
-		ensure(bCurrentlyDispatchingEvents);
-		bCurrentlyDispatchingEvents = false;
-
-		// If we deferred any handler registration in RegisterHandler, complete it now. 
-		if (DeferredHandlers.Num() > 0)
-		{
-			ContainerLock.WriteLock();
-			for (TPair<FEventID, IEventHandler*>& DeferredHandler : DeferredHandlers)
-			{
-				checkf(DeferredHandler.Key < EventContainers.Num(), TEXT("Registering event Handler for an event ID that does not exist"));
-				EventContainers[DeferredHandler.Key]->RegisterHandler(DeferredHandler.Value);
-			}
-			DeferredHandlers.Reset();
-			ContainerLock.WriteUnlock();
-		}
 	}
 
 	void FEventManager::InternalRegisterInjector(const FEventID& EventID, const FEventContainerBasePtr& Container)
diff --git a/Engine/Source/Runtime/Experimental/Chaos/Public/EventManager.h b/Engine/Source/Runtime/Experimental/Chaos/Public/EventManager.h
--- a/Engine/Source/Runtime/Experimental/Chaos/Public/EventManager.h
+++ b/Engine/Source/Runtime/Experimental/Chaos/Public/EventManager.h
@@ -169,13 +169,13 @@ namespace Chaos
 		 */
 		~TEventContainer()
 		{
-			HandleLock.WriteLock();
+			HandlerLock.WriteLock();
 			for (FEventHandlerPtr Handler : HandlerArray)
 			{
 				delete Handler;
 				Handler = nullptr;
 			}
-			HandleLock.WriteUnlock();
+			HandlerLock.WriteUnlock();
 		}
 #
 		/**
@@ -183,26 +183,36 @@ namespace Chaos
 		 */
 		virtual void RegisterHandler(const FEventHandlerPtr& Handler)
 		{
-			HandleLock.WriteLock();
-			HandlerArray.AddUnique(Handler);
-			TArray<UObject*> ProxyOwners;
-			bool bValidProxyFilter = Handler->GetInterestedProxyOwners(ProxyOwners);
-
-			if (bValidProxyFilter)
+			// If Register Handler is called by the user inside a dispatch event we ll end up with a dead lock. 
+			// so if we cannot lock it we store the event handler to try to register later
+			if (HandlerLock.TryWriteLock())
 			{
-				for (UObject* ProxyOwner : ProxyOwners)
+				HandlerArray.AddUnique(Handler);
+				TArray<UObject*> ProxyOwners;
+				bool bValidProxyFilter = Handler->GetInterestedProxyOwners(ProxyOwners);
+
+				if (bValidProxyFilter)
 				{
-					ProxyOwnerToHandlerMap.Add(ProxyOwner, Handler);
+					for (UObject* ProxyOwner : ProxyOwners)
+					{
+						ProxyOwnerToHandlerMap.Add(ProxyOwner, Handler);
+					}
 				}
+				else
+				{
+					if (GetProxyToIndexMap(EventBuffer.Get()->GetConsumerBuffer()) != nullptr) // Only if our type supports getting the ProxyToIndexMap do we bother adding this
+					{
+						HandlersNotInMap.AddUnique(Handler);
+					}
+				}
+				HandlerLock.WriteUnlock();
 			}
 			else
 			{
-				if (GetProxyToIndexMap(EventBuffer.Get()->GetConsumerBuffer()) != nullptr) // Only if our type supports getting the ProxyToIndexMap do we bother adding this
-				{
-					HandlersNotInMap.AddUnique(Handler);
-				}
+				DeferredHandlerLock.WriteLock();
+				DeferredHandlers.AddUnique(Handler);
+				DeferredHandlerLock.WriteUnlock();
 			}
-			HandleLock.WriteUnlock();
 		}
 
 		/**
@@ -210,7 +220,7 @@ namespace Chaos
 		 */
 		virtual void UnregisterHandler(const void* InHandler)
 		{
-			HandleLock.WriteLock();
+			HandlerLock.WriteLock();
 			TArray<TPair<UObject*, FEventHandlerPtr>> KeysAndValuesToRemove;
 			for (TPair<UObject*, FEventHandlerPtr>& KeyValue : ProxyOwnerToHandlerMap)
 			{
@@ -244,7 +254,7 @@ namespace Chaos
 					break;
 				}
 			}
-			HandleLock.WriteUnlock();
+			HandlerLock.WriteUnlock();
 		}
 
 		/*
@@ -293,7 +303,7 @@ namespace Chaos
 				return;
 			}
 
-			HandleLock.ReadLock();
+			HandlerLock.ReadLock();
 			const TMap<IPhysicsProxyBase*, TArray<int32>>* Map = GetProxyToIndexMap(Buffer); // Use t his map to get all proxies used in the event buffer
 			// Only take this path if we have fewer Events than Handlers
 			if (Map && Map->Num() + HandlersNotInMap.Num() < HandlerArray.Num())
@@ -327,10 +337,24 @@ namespace Chaos
 					Handler->HandleEvent(Buffer);
 				}
 			}
-			HandleLock.ReadUnlock();
+			HandlerLock.ReadUnlock();
+			RegisterDeferedHandler();
 		}
 
-	private:
+private:
+
+		void RegisterDeferedHandler()
+		{
+			DeferredHandlerLock.WriteLock();
+			// Move array
+			TArray<FEventHandlerPtr> DeferredHandlersCopy(MoveTemp(DeferredHandlers));
+			check(DeferredHandlers.Num() == 0);
+			DeferredHandlerLock.WriteUnlock();
+			for (const FEventHandlerPtr& HandlerPtr : DeferredHandlersCopy)
+			{
+				RegisterHandler(HandlerPtr);
+			}
+		}
 
 		void DeleteHandler(FEventHandlerPtr& HandlerPtr)
 		{
@@ -356,7 +380,10 @@ namespace Chaos
 		 */
 		TArray<FEventHandlerPtr> HandlerArray;
 
-		FRWLock HandleLock; // protect access ProxyOwnerToHandlerMap, HandlersNotInMap, HandlerArray
+		FRWLock HandlerLock; // protect access ProxyOwnerToHandlerMap, HandlersNotInMap, HandlerArray
+		TArray<FEventHandlerPtr> DeferredHandlers; // Store handler to register, they couldn't registered to avoid reentrant lock
+		FRWLock DeferredHandlerLock; // protect access to DeferredHandlers
+
 	};
 
 	/**
@@ -372,7 +399,6 @@ namespace Chaos
 
 		FEventManager(const Chaos::EMultiBufferMode& BufferModeIn) 
 			: BufferMode(BufferModeIn)
-			, bCurrentlyDispatchingEvents(false)
 			{}
 
 		~FEventManager()
@@ -430,22 +456,11 @@ namespace Chaos
 		template<typename PayloadType, typename HandlerType>
 		void RegisterHandler(const EEventType& EventType, HandlerType* Handler, typename TRawEventHandler<PayloadType, HandlerType>::FHandlerFunction HandlerFunction, typename TRawEventHandler<PayloadType, HandlerType>::FInterestedProxyOwnerFunction InterestedProxyOwnerFunction = nullptr)
 		{
-			FScopeLock ScopeLock(&AccessDeferredHandlersLock);
-
 			const FEventID EventID = FEventID(EventType);
-			
-			// If we are currently dispatching events, defer handler registration until completion of dispatch to avoid deadlock
-			if (bCurrentlyDispatchingEvents)
-			{
-				DeferredHandlers.Add(TPair<FEventID, IEventHandler*>(EventID, new TRawEventHandler<PayloadType, HandlerType>(Handler, HandlerFunction, InterestedProxyOwnerFunction)));
-			}
-			else
-			{
-				ContainerLock.ReadLock();
-				checkf(EventID < EventContainers.Num(), TEXT("Registering event Handler for an event ID that does not exist"));
-				EventContainers[EventID]->RegisterHandler(new TRawEventHandler<PayloadType, HandlerType>(Handler, HandlerFunction, InterestedProxyOwnerFunction));
-				ContainerLock.ReadUnlock();
-			}	
+			ContainerLock.ReadLock();
+			checkf(EventID < EventContainers.Num(), TEXT("Registering event Handler for an event ID that does not exist"));
+			EventContainers[EventID]->RegisterHandler(new TRawEventHandler<PayloadType, HandlerType>(Handler, HandlerFunction, InterestedProxyOwnerFunction));
+			ContainerLock.ReadUnlock();
 		}
 
 		/**
@@ -501,11 +516,6 @@ namespace Chaos
 		TArray<FEventContainerBasePtr> EventContainers;	// Array of event types
 		FRWLock ResourceLock;
 		FRWLock ContainerLock;
-		FCriticalSection AccessDeferredHandlersLock;
-
-		/** Defer handler registration if we are currently dispatching events. */
-		bool bCurrentlyDispatchingEvents;
-		TArray<TPair<FEventID, IEventHandler*>> DeferredHandlers;
 	};
 
 }
