Geometry collection : fix new replication to set the velocity after the break is forced

diff --git a/Engine/Source/Runtime/Experimental/GeometryCollectionEngine/Private/GeometryCollection/GeometryCollectionComponent.cpp b/Engine/Source/Runtime/Experimental/GeometryCollectionEngine/Private/GeometryCollection/GeometryCollectionComponent.cpp
--- a/Engine/Source/Runtime/Experimental/GeometryCollectionEngine/Private/GeometryCollection/GeometryCollectionComponent.cpp
+++ b/Engine/Source/Runtime/Experimental/GeometryCollectionEngine/Private/GeometryCollection/GeometryCollectionComponent.cpp
@@ -2628,22 +2628,6 @@ namespace
 		// Update the anchored state of the root particle
 		UpdateRootStateFromReplication(*PhysicsProxy, (bool)RepStateData.bIsRootAnchored);
 		
-		// first apply the breaking velocities 
-		for (const FGeometryCollectionRepStateData::FReleasedData& ReleaseData: RepStateData.ReleasedData)
-		{
-			if (/*ReleaseData.ReleasedVersion > VersionProcessed && */ ParticleHandles.IsValidIndex(ReleaseData.TransformIndex))
-			{
-				if (FPBDRigidClusteredParticleHandle* ParticleHandle = ParticleHandles[ReleaseData.TransformIndex])
-				{
-					const bool bIsNotYetBroken = (ParticleHandle->Parent() != nullptr);
-					if (bIsNotYetBroken)
-					{
-						ParticleHandle->SetV(ReleaseData.LinearVelocity);
-						ParticleHandle->SetW(FMath::DegreesToRadians(ReleaseData.AngularVelocityInDegreesPerSecond));
-					}
-				}
-			}
-		}
 		
 		// now sync the broken state of the particles
 		if (RepStateData.BrokenState.Num() == ParticleHandles.Num())
@@ -2661,6 +2645,25 @@ namespace
 					if (!bWasBroken && bIsBroken)
 					{
 						RigidClustering.ForceReleaseChildParticleAndParents(ParticleHandle, /* bTriggerBreakEvents */true);
+
+						if (!ParticleHandle->Disabled())
+						{
+							// check if we have a release velocity to be applied 
+							const FGeometryCollectionRepStateData::FReleasedData* ReleaseData = RepStateData.ReleasedData.FindByPredicate(
+								[TransformIndex](const FGeometryCollectionRepStateData::FReleasedData& Data)
+								{
+									return Data.TransformIndex == TransformIndex;
+								});
+							if (ReleaseData)
+							{
+								ParticleHandle->SetV(ReleaseData->LinearVelocity);
+								ParticleHandle->SetW(FMath::DegreesToRadians(ReleaseData->AngularVelocityInDegreesPerSecond));
+							}
+							else
+							{
+								// todo(chaos) : no release data : we can disable the particle right away ( it was released in the past )
+							}
+						}
 					}
 				}
 			}
