Sequencer: Constraints: Constraints in LS update
Main fix is to not use the Constraint Manager/Actor when a constraint is added to a level sequence. Only static constraints are now owned by the manager, animated constraints are owned by the section that the constraint channel is in, and we use a subsystem to gather both sets of these constraints as weak pointers so we can iterate over both groups.  Kept the controller logic around but that now feeds into the sub system.

Also as part of this got rid of the soft object ptr to the constraint in the section, instead we just always use the Spawn copy which we reference via a Get/Set functions.

Since FNames may not be unique when combining level sequences, we now use a FGuid to identify Cosntraints.

We now have options to show Selected,Level Sequence, Current, and All constaints in drop down. As part of this we need to refresh the content more often(on time change) since new constraints may show up and we also no longer cache constraint names or state but get it dynamically.
When we add actors to Sequencer we look to see if it has any static contraints if so we change them to dynamic/owned by the LS and key them.

In the template/system we had issues with constraints not getting turned off when exiting so we make sure to do so if no bound object or when exiting. Also we now always 'init' constraints which involves setting up dependencies/tick functions as needed on every tick in order to handle orders changing and spawning.
As part of starting to fix cycle issues that happen with tick dependencies we no longer set up a child trs dependency but instead always tick the child in the GetFunction(better but need to revisit one last time).

Bake to Control Rig with Sequencer Interp Type
Fixed issue when smart keying constriant active or not not working when over a current, we now just toggle based upon previous value.

diff --git a/Engine/Plugins/Animation/ControlRig/Source/ControlRig/Private/Constraints/ControlRigTransformableHandle.cpp b/Engine/Plugins/Animation/ControlRig/Source/ControlRig/Private/Constraints/ControlRigTransformableHandle.cpp
--- a/Engine/Plugins/Animation/ControlRig/Source/ControlRig/Private/Constraints/ControlRigTransformableHandle.cpp
+++ b/Engine/Plugins/Animation/ControlRig/Source/ControlRig/Private/Constraints/ControlRigTransformableHandle.cpp
@@ -12,6 +12,7 @@
 #include "Sequencer/MovieSceneControlRigParameterSection.h"
 #include "Sections/MovieScene3DTransformSection.h"
 #include "Sequencer/ControlRigSequencerHelpers.h"
+#include "UObject/Package.h"
 
 #include UE_INLINE_GENERATED_CPP_BY_NAME(ControlRigTransformableHandle)
 
@@ -52,7 +53,7 @@ bool UTransformableControlHandle::IsValid() const
 	return true;
 }
 
-void UTransformableControlHandle::TickForBaking()
+void UTransformableControlHandle::TickForBaking() const
 {
 	USkeletalMeshComponent* SkeletalMeshComponent = GetSkeletalMesh();
 	if (SkeletalMeshComponent)
diff --git a/Engine/Plugins/Animation/ControlRig/Source/ControlRig/Private/ControlRig.cpp b/Engine/Plugins/Animation/ControlRig/Source/ControlRig/Private/ControlRig.cpp
--- a/Engine/Plugins/Animation/ControlRig/Source/ControlRig/Private/ControlRig.cpp
+++ b/Engine/Plugins/Animation/ControlRig/Source/ControlRig/Private/ControlRig.cpp
@@ -3150,7 +3150,6 @@ void UControlRig::PostInitInstance(URigVMHost* InCDO)
 }
 
 UTransformableControlHandle* UControlRig::CreateTransformableControlHandle(
-	UObject* InOuter,
 	const FName& InControlName) const
 {
 	auto IsConstrainable = [this](const FName& InControlName)
@@ -3178,7 +3177,7 @@ UTransformableControlHandle* UControlRig::CreateTransformableControlHandle(
 		return nullptr;
 	}
 	
-	UTransformableControlHandle* CtrlHandle = NewObject<UTransformableControlHandle>(InOuter, NAME_None, RF_Transactional);
+	UTransformableControlHandle* CtrlHandle = NewObject<UTransformableControlHandle>(GetTransientPackage(), NAME_None, RF_Transactional);
 	check(CtrlHandle);
 	CtrlHandle->ControlRig = this;
 	CtrlHandle->ControlName = InControlName;
diff --git a/Engine/Plugins/Animation/ControlRig/Source/ControlRig/Private/ControlRigModule.cpp b/Engine/Plugins/Animation/ControlRig/Source/ControlRig/Private/ControlRigModule.cpp
--- a/Engine/Plugins/Animation/ControlRig/Source/ControlRig/Private/ControlRigModule.cpp
+++ b/Engine/Plugins/Animation/ControlRig/Source/ControlRig/Private/ControlRigModule.cpp
@@ -34,11 +34,11 @@ void FControlRigModule::RegisterTransformableCustomization() const
 	FModuleManager::Get().LoadModule(ConstraintsModuleName);
 	
 	// register UControlRig and AControlRigShapeActor
-	auto CreateControlHandle = [](UObject* Outer, UObject* InObject, const FName& InControlName)->UTransformableHandle*
+	auto CreateControlHandle = [](UObject* InObject, const FName& InControlName)->UTransformableHandle*
 	{
 		if (const UControlRig* ControlRig = Cast<UControlRig>(InObject))
 		{
-			return ControlRig->CreateTransformableControlHandle(Outer, InControlName);
+			return ControlRig->CreateTransformableControlHandle(InControlName);
 		}
 		return nullptr;
 	};
@@ -52,11 +52,11 @@ void FControlRigModule::RegisterTransformableCustomization() const
 		return 0;
 	};
 	
-	auto CreateControlHandleFromActor = [CreateControlHandle](UObject* Outer, UObject* InObject, const FName&)->UTransformableHandle*
+	auto CreateControlHandleFromActor = [CreateControlHandle](UObject* InObject, const FName&)->UTransformableHandle*
 	{
 		if (const AControlRigShapeActor* ControlActor = Cast<AControlRigShapeActor>(InObject))
 		{
-			return CreateControlHandle(Outer, ControlActor->ControlRig.Get(), ControlActor->ControlName);
+			return CreateControlHandle(ControlActor->ControlRig.Get(), ControlActor->ControlName);
 		}
 		return nullptr;
 	};
diff --git a/Engine/Plugins/Animation/ControlRig/Source/ControlRig/Private/Sequencer/MovieSceneControlRigParameterSection.cpp b/Engine/Plugins/Animation/ControlRig/Source/ControlRig/Private/Sequencer/MovieSceneControlRigParameterSection.cpp
--- a/Engine/Plugins/Animation/ControlRig/Source/ControlRig/Private/Sequencer/MovieSceneControlRigParameterSection.cpp
+++ b/Engine/Plugins/Animation/ControlRig/Source/ControlRig/Private/Sequencer/MovieSceneControlRigParameterSection.cpp
@@ -782,7 +782,7 @@ void UMovieSceneControlRigParameterSection::OnBindingIDsUpdated(const TMap<UE::M
 {
 	for (FConstraintAndActiveChannel& ConstraintChannel : ConstraintsChannels)
 	{
-		if (UTickableTransformConstraint* TransformConstraint = Cast< UTickableTransformConstraint>(ConstraintChannel.Constraint.Get()))
+		if (UTickableTransformConstraint* TransformConstraint = Cast< UTickableTransformConstraint>(ConstraintChannel.GetConstraint()))
 		{
 			if (TransformConstraint->ChildTRSHandle)
 			{
@@ -793,17 +793,6 @@ void UMovieSceneControlRigParameterSection::OnBindingIDsUpdated(const TMap<UE::M
 				TransformConstraint->ParentTRSHandle->OnBindingIDsUpdated(OldFixedToNewFixedMap, LocalSequenceID, Hierarchy, Player);
 			}
 		}
-		if (UTickableTransformConstraint* SpawnCopy = Cast< UTickableTransformConstraint>(ConstraintChannel.ConstraintCopyToSpawn))
-		{
-			if (SpawnCopy->ChildTRSHandle)
-			{
-				SpawnCopy->ChildTRSHandle->OnBindingIDsUpdated(OldFixedToNewFixedMap, LocalSequenceID, Hierarchy, Player);
-			}
-			if (SpawnCopy->ParentTRSHandle)
-			{
-				SpawnCopy->ParentTRSHandle->OnBindingIDsUpdated(OldFixedToNewFixedMap, LocalSequenceID, Hierarchy, Player);
-			}
-		}
 	}
 }
 
@@ -811,7 +800,7 @@ void UMovieSceneControlRigParameterSection::GetReferencedBindings(TArray<FGuid>&
 {
 	for (FConstraintAndActiveChannel& ConstraintChannel : ConstraintsChannels)
 	{
-		if (UTickableTransformConstraint* TransformConstraint = Cast< UTickableTransformConstraint>(ConstraintChannel.Constraint.Get()))
+		if (UTickableTransformConstraint* TransformConstraint = Cast< UTickableTransformConstraint>(ConstraintChannel.GetConstraint().Get()))
 		{
 			if (TransformConstraint->ChildTRSHandle && TransformConstraint->ChildTRSHandle->ConstraintBindingID.IsValid())
 			{
@@ -828,14 +817,6 @@ void UMovieSceneControlRigParameterSection::GetReferencedBindings(TArray<FGuid>&
 void UMovieSceneControlRigParameterSection::PreSave(FObjectPreSaveContext SaveContext)
 {
 	Super::PreSave(SaveContext);
-	
-	for (FConstraintAndActiveChannel& ActiveChannel : ConstraintsChannels)
-	{
-		if (ActiveChannel.Constraint.IsValid())
-		{
-			ActiveChannel.ConstraintCopyToSpawn = ActiveChannel.Constraint->Duplicate(this);
-		}
-	}
 }
 
 bool UMovieSceneControlRigParameterSection::RenameParameterName(const FName& OldParameterName, const FName& NewParameterName)
@@ -1028,21 +1009,7 @@ void UMovieSceneControlRigParameterSection::Serialize(FArchive& Ar)
 #if WITH_EDITOR
 void UMovieSceneControlRigParameterSection::PostDuplicate(bool bDuplicateForPIE)
 {
-	Super::PostDuplicate(bDuplicateForPIE);
-	//if not duplicating for PIE clear out the soft object ptrs to the constraints we dont' want to hold pointers to the old constraints.
-	//also make the copy to spawn so we can duplicate it if it doesn't exist in the level
-	//to the old ones but make new duplicates
-	if (bDuplicateForPIE == false)
-	{
-		for (FConstraintAndActiveChannel& ConstraintChannel : ConstraintsChannels)
-		{
-			if (ConstraintChannel.Constraint.IsValid() && !ConstraintChannel.ConstraintCopyToSpawn)
-			{
-				ConstraintChannel.ConstraintCopyToSpawn = ConstraintChannel.Constraint->Duplicate(this);
-			}
-			ConstraintChannel.Constraint.Reset();
-		}
-	}
+	Super::PostDuplicate(bDuplicateForPIE);	
 }
 #endif
 
@@ -1064,8 +1031,7 @@ void UMovieSceneControlRigParameterSection::PostLoad()
 	{
 		for (FConstraintAndActiveChannel& ConstraintChannel : ConstraintsChannels)
 		{
-			ConstraintChannel.Constraint.Reset(); //clear it out may be referencing wrong level
-			if (UTickableTransformConstraint* TransformConstraint = Cast<UTickableTransformConstraint>(ConstraintChannel.ConstraintCopyToSpawn))
+			if (UTickableTransformConstraint* TransformConstraint = Cast<UTickableTransformConstraint>(ConstraintChannel.GetConstraint()))
 			{
 				if (UTransformableControlHandle* Handle = Cast<UTransformableControlHandle>(TransformConstraint->ChildTRSHandle))
 				{
@@ -1466,20 +1432,19 @@ void UMovieSceneControlRigParameterSection::AddSpaceChannel(FName InControlName,
 	}
 }
 
-bool UMovieSceneControlRigParameterSection::HasConstraintChannel(const FName& InConstraintName) const
+bool UMovieSceneControlRigParameterSection::HasConstraintChannel(const FGuid& InGuid) const
 {
-	return ConstraintsChannels.ContainsByPredicate( [InConstraintName](const FConstraintAndActiveChannel& InChannel)
+	return ConstraintsChannels.ContainsByPredicate( [InGuid](const FConstraintAndActiveChannel& InChannel)
 	{
-		return InChannel.Constraint.IsValid() ? InChannel.Constraint->GetFName() == InConstraintName : false;
+		return InChannel.GetConstraint().Get() ? InChannel.GetConstraint()->ConstraintID == InGuid : false;
 	});
 }
 
-FConstraintAndActiveChannel* UMovieSceneControlRigParameterSection::GetConstraintChannel(const FName& InConstraintName)
+FConstraintAndActiveChannel* UMovieSceneControlRigParameterSection::GetConstraintChannel(const FGuid& InConstraintID)
 {
-	const int32 Index = ConstraintsChannels.IndexOfByPredicate([InConstraintName](const FConstraintAndActiveChannel& InChannel)
+	const int32 Index = ConstraintsChannels.IndexOfByPredicate([InConstraintID](const FConstraintAndActiveChannel& InChannel)
 		{
-			return InChannel.Constraint.IsValid() ? InChannel.Constraint->GetFName() == InConstraintName : (InChannel.ConstraintCopyToSpawn ?
-			InChannel.ConstraintCopyToSpawn->GetFName() == InConstraintName : false);
+			return InChannel.GetConstraint().Get() ? InChannel.GetConstraint()->ConstraintID == InConstraintID : false;
 		});
 	return (Index != INDEX_NONE) ? &ConstraintsChannels[Index] : nullptr;	
 }
@@ -1488,12 +1453,12 @@ void UMovieSceneControlRigParameterSection::ReplaceConstraint(const FName InCons
 {
 	const int32 Index = ConstraintsChannels.IndexOfByPredicate([InConstraintName](const FConstraintAndActiveChannel& InChannel)
 	{
-		return InChannel.ConstraintCopyToSpawn ? InChannel.ConstraintCopyToSpawn->GetFName() == InConstraintName : false;
+		return InChannel.GetConstraint().Get() ? InChannel.GetConstraint()->GetFName() == InConstraintName : false;
 	});
 	if (Index != INDEX_NONE)
 	{
 		Modify();
-		ConstraintsChannels[Index].Constraint = InConstraint;
+		ConstraintsChannels[Index].SetConstraint(InConstraint);
 		ReconstructChannelProxy();
 	}
 }
@@ -1505,8 +1470,7 @@ void UMovieSceneControlRigParameterSection::OnConstraintsChanged()
 
 void UMovieSceneControlRigParameterSection::AddConstraintChannel(UTickableConstraint* InConstraint)
 {
-	
-	if (InConstraint && !HasConstraintChannel(InConstraint->GetFName()))
+	if (InConstraint && !HasConstraintChannel(InConstraint->ConstraintID))
 	{
 		Modify();
 		
@@ -1516,8 +1480,9 @@ void UMovieSceneControlRigParameterSection::AddConstraintChannel(UTickableConstr
 		ExistingChannel->SetDefault(false);
 		
 		//make copy that we can spawn if it doesn't exist
-		ConstraintsChannels[NewIndex].ConstraintCopyToSpawn = InConstraint->Duplicate(this);
-
+		//the rename changes the outer to this section (from any actor manager)
+		InConstraint->Rename(nullptr, this, REN_ForceNoResetLoaders | REN_DontCreateRedirectors);
+		
 		if (OnConstraintChannelAdded.IsBound())
 		{
 			OnConstraintChannelAdded.Broadcast(this, ExistingChannel);
@@ -1536,7 +1501,7 @@ void UMovieSceneControlRigParameterSection::RemoveConstraintChannel(const UTicka
 	}
 	const int32 Index = ConstraintsChannels.IndexOfByPredicate([InConstraint](const FConstraintAndActiveChannel& InChannel)
 	{
-		return InChannel.Constraint.IsValid() ? InChannel.Constraint == InConstraint : false;
+		return InChannel.GetConstraint().Get() ? InChannel.GetConstraint().Get() == InConstraint : false;
 	});
 
 	if (ConstraintsChannels.IsValidIndex(Index))
@@ -1670,42 +1635,45 @@ void UMovieSceneControlRigParameterSection::ReconstructChannelProxy()
 		auto AddConstrainChannels = [this, GetConstraints, &ConstraintsChannelIndex, &TotalIndex, &Channels](
 			const FName& InControlName, const FText& InGroup, const bool bEnabled)
 		{
-			TArray<TObjectPtr<UTickableConstraint>> Constraints = GetConstraints(InControlName);
-			for (const TObjectPtr<UTickableConstraint>& Constraint: Constraints)
+			TArray<TWeakObjectPtr<UTickableConstraint>> Constraints = GetConstraints(InControlName);
+			for (const TWeakObjectPtr<UTickableConstraint>& Constraint: Constraints)
 			{
-				const FName& ConstraintName = Constraint->GetFName();
-				if(FConstraintAndActiveChannel* ConstraintChannel = GetConstraintChannel(ConstraintName))
+				if (Constraint.IsValid())
 				{
-					if (FChannelMapInfo* ChannelInfo = ControlChannelMap.Find(InControlName))
+					const FGuid& ConstraintID = Constraint->ConstraintID;
+					if (FConstraintAndActiveChannel* ConstraintChannel = GetConstraintChannel(ConstraintID))
 					{
-						ChannelInfo->ConstraintsIndex.Add(ConstraintsChannelIndex);
-					}
+						if (FChannelMapInfo* ChannelInfo = ControlChannelMap.Find(InControlName))
+						{
+							ChannelInfo->ConstraintsIndex.Add(ConstraintsChannelIndex);
+						}
 
 #if WITH_EDITOR
-					ConstraintChannel->ActiveChannel.ExtraLabel = [WeakConstraint = MakeWeakObjectPtr(Constraint)]
-					{
-						if (WeakConstraint.IsValid())
+						ConstraintChannel->ActiveChannel.ExtraLabel = [Constraint]
 						{
-							FString ParentStr; WeakConstraint->GetLabel().Split(TEXT("."), &ParentStr, nullptr);
-							if (!ParentStr.IsEmpty())
+							if (Constraint.IsValid())
 							{
-								return ParentStr;
-							}		
-						}
-						static const FString DummyStr;
-						return DummyStr;
-					};
-					
-					const FText DisplayText = FText::FromString(Constraint->GetTypeLabel());
-					FMovieSceneChannelMetaData MetaData(ConstraintName, DisplayText,InGroup, bEnabled);
-					ConstraintsChannelIndex += 1;
-					MetaData.SortOrder = TotalIndex++;
-					MetaData.bCanCollapseToTrack = false;
-		
-					Channels.Add(ConstraintChannel->ActiveChannel, MetaData, TMovieSceneExternalValue<bool>());
+								FString ParentStr; Constraint->GetLabel().Split(TEXT("."), &ParentStr, nullptr);
+								if (!ParentStr.IsEmpty())
+								{
+									return ParentStr;
+								}
+							}
+							static const FString DummyStr;
+							return DummyStr;
+						};
+
+						const FText DisplayText = FText::FromString(Constraint->GetTypeLabel());
+						FMovieSceneChannelMetaData MetaData(Constraint->GetFName(), DisplayText, InGroup, bEnabled);
+						ConstraintsChannelIndex += 1;
+						MetaData.SortOrder = TotalIndex++;
+						MetaData.bCanCollapseToTrack = false;
+
+						Channels.Add(ConstraintChannel->ActiveChannel, MetaData, TMovieSceneExternalValue<bool>());
 #else
-					Channels.Add(ConstraintChannel->ActiveChannel);
+						Channels.Add(ConstraintChannel->ActiveChannel);
 #endif
+					}
 				}
 			}
 		};
@@ -2886,7 +2854,7 @@ void UMovieSceneControlRigParameterSection::AutoSetTangents(const FName& Control
 }
 #if WITH_EDITOR
 
-void UMovieSceneControlRigParameterSection::RecordControlRigKey(FFrameNumber FrameNumber, bool bSetDefault, ERichCurveInterpMode InInterpMode)
+void UMovieSceneControlRigParameterSection::RecordControlRigKey(FFrameNumber FrameNumber, bool bSetDefault, EMovieSceneKeyInterpolation InInterpMode)
 {
 	if (ControlRig)
 	{
@@ -2900,24 +2868,29 @@ void UMovieSceneControlRigParameterSection::RecordControlRigKey(FFrameNumber Fra
 		{
 			switch (InInterpMode)
 			{
-			case RCIM_Linear:
+			case EMovieSceneKeyInterpolation::Linear:
 				FloatChannels[ChannelIndex++]->AddLinearKey(FrameNumber, Value.X);
 				FloatChannels[ChannelIndex++]->AddLinearKey(FrameNumber, Value.Y);
 				FloatChannels[ChannelIndex++]->AddLinearKey(FrameNumber, Value.Z);
 				break;
 
-			case RCIM_Constant:
+			case EMovieSceneKeyInterpolation::Constant:
 				FloatChannels[ChannelIndex++]->AddConstantKey(FrameNumber, Value.X);
 				FloatChannels[ChannelIndex++]->AddConstantKey(FrameNumber, Value.Y);
 				FloatChannels[ChannelIndex++]->AddConstantKey(FrameNumber, Value.Z);
 				break;
 
-			case RCIM_Cubic:
-			default:
+			case  EMovieSceneKeyInterpolation::Auto:
 				FloatChannels[ChannelIndex++]->AddCubicKey(FrameNumber, Value.X, ERichCurveTangentMode::RCTM_Auto);
 				FloatChannels[ChannelIndex++]->AddCubicKey(FrameNumber, Value.Y, ERichCurveTangentMode::RCTM_Auto);
 				FloatChannels[ChannelIndex++]->AddCubicKey(FrameNumber, Value.Z, ERichCurveTangentMode::RCTM_Auto);
 				break;
+			case  EMovieSceneKeyInterpolation::SmartAuto:
+			default:
+				FloatChannels[ChannelIndex++]->AddCubicKey(FrameNumber, Value.X, ERichCurveTangentMode::RCTM_SmartAuto);
+				FloatChannels[ChannelIndex++]->AddCubicKey(FrameNumber, Value.Y, ERichCurveTangentMode::RCTM_SmartAuto);
+				FloatChannels[ChannelIndex++]->AddCubicKey(FrameNumber, Value.Z, ERichCurveTangentMode::RCTM_SmartAuto);
+				break;
 			}
 		};
 
@@ -2983,18 +2956,22 @@ void UMovieSceneControlRigParameterSection::RecordControlRigKey(FFrameNumber Fra
 
 					switch (InInterpMode)
 					{
-					case RCIM_Linear:
+					case EMovieSceneKeyInterpolation::Linear:
 						FloatChannels[ChannelIndex++]->AddLinearKey(FrameNumber, Val);
 						break;
 
-					case RCIM_Constant:
+					case EMovieSceneKeyInterpolation::Constant:
 						FloatChannels[ChannelIndex++]->AddConstantKey(FrameNumber, Val);
 						break;
 
-					case RCIM_Cubic:
-					default:
+					case EMovieSceneKeyInterpolation::Auto:
 						FloatChannels[ChannelIndex++]->AddCubicKey(FrameNumber, Val, ERichCurveTangentMode::RCTM_Auto);
 						break;
+
+					case EMovieSceneKeyInterpolation::SmartAuto:
+					default:
+						FloatChannels[ChannelIndex++]->AddCubicKey(FrameNumber, Val, ERichCurveTangentMode::RCTM_SmartAuto);
+						break;
 					}
 
 					break;
@@ -3010,21 +2987,24 @@ void UMovieSceneControlRigParameterSection::RecordControlRigKey(FFrameNumber Fra
 
 					switch (InInterpMode)
 					{
-					case RCIM_Linear:
+					case EMovieSceneKeyInterpolation::Linear:
 						FloatChannels[ChannelIndex++]->AddLinearKey(FrameNumber, Val.X);
 						FloatChannels[ChannelIndex++]->AddLinearKey(FrameNumber, Val.Y);
 						break;
 
-					case RCIM_Constant:
+					case  EMovieSceneKeyInterpolation::Constant:
 						FloatChannels[ChannelIndex++]->AddConstantKey(FrameNumber, Val.X);
 						FloatChannels[ChannelIndex++]->AddConstantKey(FrameNumber, Val.Y);
 						break;
-
-					case RCIM_Cubic:
-					default:
+					case  EMovieSceneKeyInterpolation::Auto:
 						FloatChannels[ChannelIndex++]->AddCubicKey(FrameNumber, Val.X, ERichCurveTangentMode::RCTM_Auto);
 						FloatChannels[ChannelIndex++]->AddCubicKey(FrameNumber, Val.Y, ERichCurveTangentMode::RCTM_Auto);
 						break;
+					case EMovieSceneKeyInterpolation::SmartAuto:
+					default:
+						FloatChannels[ChannelIndex++]->AddCubicKey(FrameNumber, Val.X, ERichCurveTangentMode::RCTM_SmartAuto);
+						FloatChannels[ChannelIndex++]->AddCubicKey(FrameNumber, Val.Y, ERichCurveTangentMode::RCTM_SmartAuto);
+						break;
 					}
 
 					break;
@@ -3075,9 +3055,7 @@ void UMovieSceneControlRigParameterSection::RecordControlRigKey(FFrameNumber Fra
 					}
 					else
 					{
-						Val = 
-							ControlRig
-							->GetControlValue(ControlElement, ERigControlValueType::Current).Get<FRigControlValue::FTransform_Float>().ToTransform();
+						Val = ControlRig->GetControlValue(ControlElement, ERigControlValueType::Current).Get<FRigControlValue::FTransform_Float>().ToTransform();
 					}
 					FVector CurrentVector = Val.GetTranslation();
 					if (bSetDefault)
@@ -3128,7 +3106,7 @@ void UMovieSceneControlRigParameterSection::RecordControlRigKey(FFrameNumber Fra
 	}
 }
 
-bool UMovieSceneControlRigParameterSection::LoadAnimSequenceIntoThisSection(UAnimSequence* AnimSequence, UMovieScene* MovieScene, UObject* BoundObject, bool bKeyReduce, float Tolerance, FFrameNumber InStartFrame)
+bool UMovieSceneControlRigParameterSection::LoadAnimSequenceIntoThisSection(UAnimSequence* AnimSequence, UMovieScene* MovieScene, UObject* BoundObject, bool bKeyReduce, float Tolerance, FFrameNumber InStartFrame, EMovieSceneKeyInterpolation InInterpolation)
 {
 	USkeletalMeshComponent* SkelMeshComp = Cast<USkeletalMeshComponent>(BoundObject);
 	
@@ -3265,8 +3243,7 @@ bool UMovieSceneControlRigParameterSection::LoadAnimSequenceIntoThisSection(UAni
 		}
 		ControlRig->Execute(FRigUnit_InverseExecution::EventName);
 
-		const ERichCurveInterpMode InterpMode = bKeyReduce ? RCIM_Cubic : RCIM_Linear;
-		RecordControlRigKey(FrameNumber, Index == 0, InterpMode);
+		RecordControlRigKey(FrameNumber, Index == 0, InInterpolation);
 		Progress.EnterProgressFrame(1);
 		if (Progress.ShouldCancel())
 		{
diff --git a/Engine/Plugins/Animation/ControlRig/Source/ControlRig/Private/Sequencer/MovieSceneControlRigParameterTemplate.cpp b/Engine/Plugins/Animation/ControlRig/Source/ControlRig/Private/Sequencer/MovieSceneControlRigParameterTemplate.cpp
--- a/Engine/Plugins/Animation/ControlRig/Source/ControlRig/Private/Sequencer/MovieSceneControlRigParameterTemplate.cpp
+++ b/Engine/Plugins/Animation/ControlRig/Source/ControlRig/Private/Sequencer/MovieSceneControlRigParameterTemplate.cpp
@@ -177,13 +177,11 @@ struct FEulerTransformParameterStringAndValue
 
 struct FConstraintAndActiveValue
 {
-	FConstraintAndActiveValue(TWeakObjectPtr<UTickableConstraint> InConstraint, TWeakObjectPtr<UTickableConstraint> InSpawnConstraint, bool InValue)
+	FConstraintAndActiveValue(TWeakObjectPtr<UTickableConstraint> InConstraint,  bool InValue)
 		: Constraint(InConstraint)
-		, SpawnConstraint(InSpawnConstraint)
 		, Value(InValue)
 	{}
 	TWeakObjectPtr<UTickableConstraint> Constraint;
-	TWeakObjectPtr<UTickableConstraint> SpawnConstraint;
 	bool Value;
 };
 
@@ -1121,6 +1119,26 @@ static UControlRig* GetControlRig(const UMovieSceneControlRigParameterSection* S
 	return ControlRig;
 }
 
+static UTickableConstraint* CreateConstraintIfNeeded(const FConstraintsManagerController& Controller, FConstraintAndActiveValue& ConstraintValue, UMovieSceneControlRigParameterSection* Section)
+{
+	UTickableConstraint* Constraint = ConstraintValue.Constraint.Get();
+
+	if (!Constraint) 
+	{
+		return nullptr;
+	}
+	else // it's possible that we have it but it's not in the manager, due to manager not being saved with it (due to spawning or undo/redo).
+	{
+		const TArray< TWeakObjectPtr<UTickableConstraint>>& ConstraintsArray = Controller.GetConstraintsArray();
+		if (Controller.GetConstraint(Constraint->ConstraintID) == nullptr)
+		{
+			Controller.AddConstraint(ConstraintValue.Constraint.Get());
+		}
+	}
+
+	return Constraint;
+}
+
 /* Simple token used for non-blendables*/
 struct FControlRigParameterExecutionToken : IMovieSceneExecutionToken
 {
@@ -1148,11 +1166,10 @@ struct FControlRigParameterExecutionToken : IMovieSceneExecutionToken
 		UControlRig* ControlRig = Section->GetControlRig();
 
 		// Update the animation's state
-		
+		TArrayView<TWeakObjectPtr<>> BoundObjects = Player.FindBoundObjects(Operand);
 		if (ControlRig)
 		{
 			const UMovieSceneSequence* Sequence = Player.State.FindSequence(Operand.SequenceID);
-			TArrayView<TWeakObjectPtr<>> BoundObjects = Player.FindBoundObjects(Operand);
 
 			UObject* BoundObject = BoundObjects.Num() > 0 ? BoundObjects[0].Get() : nullptr;
 			if (Sequence && BoundObject)
@@ -1301,35 +1318,48 @@ struct FControlRigParameterExecutionToken : IMovieSceneExecutionToken
 						}
 					}
 				}
-
-				const FConstraintsManagerController& Controller = FConstraintsManagerController::Get(ControlRig->GetWorld());
-				for (FConstraintAndActiveValue& ConstraintValue : ConstraintsValues)
+				UObject* BoundObject = BoundObjects.Num() > 0 ? BoundObjects[0].Get() : nullptr;
+				if (BoundObject)
 				{
-					//if the constraint isn't valid try to make a copy of our spawn copy if it exists and add that
-					if (ConstraintValue.Constraint.IsValid() == false && ConstraintValue.SpawnConstraint.IsValid())
+					const FConstraintsManagerController& Controller = FConstraintsManagerController::Get(BoundObject->GetWorld());
+					for (FConstraintAndActiveValue& ConstraintValue : ConstraintsValues)
 					{
-						UTickableConstraint* NewOne = Controller.AddConstraintFromCopy(ConstraintValue.SpawnConstraint.Get());
-						ConstraintValue.Constraint = NewOne;
 						UMovieSceneControlRigParameterSection* NonConstSection = const_cast<UMovieSceneControlRigParameterSection*>(Section);
-						NonConstSection->ReplaceConstraint(ConstraintValue.SpawnConstraint.Get()->GetFName(), NewOne);
+						CreateConstraintIfNeeded(Controller, ConstraintValue, NonConstSection);
+
+						if (ConstraintValue.Constraint.IsValid())
+						{
+							//For Control Rig we may need to explicitly set the control rig
+							if (UTickableTransformConstraint* TransformConstraint = Cast<UTickableTransformConstraint>(ConstraintValue.Constraint))
+							{
+								TransformConstraint->InitConstraint(BoundObject->GetWorld());
+							}
+							ConstraintValue.Constraint->ResolveBoundObjects(Operand.SequenceID, Player, ControlRig);
+							ConstraintValue.Constraint->SetActive(ConstraintValue.Value);
+						}
 					}
-					if (ConstraintValue.Constraint.IsValid())
+					//unfortunately for Constraints with ControlRig we need to resolve all Parents also. Don't need to do children since they wil be handled by
+					//the channel resolve above
+					TArray< TWeakObjectPtr<UTickableConstraint>> Constraints = Controller.GetAllConstraints();
+					for (TWeakObjectPtr<UTickableConstraint>& TickConstraint : Constraints)
 					{
-						//For Control Rig we may need to explicitly set the control rig
-						ConstraintValue.Constraint->ResolveBoundObjects(Operand.SequenceID, Player, ControlRig);
-						ConstraintValue.Constraint->SetActive(ConstraintValue.Value);
+						if (UTickableTransformConstraint* TransformConstraint = Cast< UTickableTransformConstraint>(TickConstraint.Get()))
+						{
+							if (TransformConstraint->ParentTRSHandle)
+							{
+								TransformConstraint->ParentTRSHandle->ResolveBoundObjects(Operand.SequenceID, Player, ControlRig);
+								TransformConstraint->EnsurePrimaryDependency(BoundObject->GetWorld());
+							}
+						}
 					}
 				}
-				//unfortunately for Constraints with ControlRig we need to resolve all Parents also. Don't need to do children since they wil be handled by
-				//the channel resolve above
-				TArray< TObjectPtr<UTickableConstraint>> Constraints =  Controller.GetAllConstraints();
-				for (UTickableConstraint* TickConstraint : Constraints)
+				else  //no bound object so turn off constraint
 				{
-					if (UTickableTransformConstraint* TransformConstraint = Cast< UTickableTransformConstraint>(TickConstraint))
+					for (FConstraintAndActiveValue& ConstraintValue : ConstraintsValues)
 					{
-						if (TransformConstraint->ParentTRSHandle)
+						if (ConstraintValue.Constraint.IsValid())
 						{
-							TransformConstraint->ParentTRSHandle->ResolveBoundObjects(Operand.SequenceID, Player, ControlRig);
+							ConstraintValue.Constraint->SetActive(ConstraintValue.Value);
 						}
 					}
 				}
@@ -1891,11 +1921,7 @@ void FMovieSceneControlRigParameterTemplate::EvaluateCurvesWithMasks(const FMovi
 			bool Value = false;
 			const FConstraintAndActiveChannel& ConstraintAndActiveChannel = Constraints[Index];
 			ConstraintAndActiveChannel.ActiveChannel.Evaluate(RoundTime, Value);
-			if (ConstraintAndActiveChannel.Constraint.IsPending())
-			{
-				ConstraintAndActiveChannel.Constraint.LoadSynchronous();
-			}
-			Values.ConstraintsValues.Emplace(ConstraintAndActiveChannel.Constraint.Get(), ConstraintAndActiveChannel.ConstraintCopyToSpawn, Value);
+			Values.ConstraintsValues.Emplace(ConstraintAndActiveChannel.GetConstraint().Get(), Value);
 		}
 		
 		for (int32 Index = 0; Index < Bools.Num(); ++Index)
diff --git a/Engine/Plugins/Animation/ControlRig/Source/ControlRig/Public/Constraints/ControlRigTransformableHandle.h b/Engine/Plugins/Animation/ControlRig/Source/ControlRig/Public/Constraints/ControlRigTransformableHandle.h
--- a/Engine/Plugins/Animation/ControlRig/Source/ControlRig/Public/Constraints/ControlRigTransformableHandle.h
+++ b/Engine/Plugins/Animation/ControlRig/Source/ControlRig/Public/Constraints/ControlRigTransformableHandle.h
@@ -72,7 +72,7 @@ public:
 	virtual UTransformableHandle* Duplicate(UObject* NewOuter) const override;
 
 	/** tick things when baking*/ 
-	virtual void TickForBaking() override;
+	virtual void TickForBaking() const override;
 
 	/** Returns the skeletal mesh bound to ControlRig. */
 	USkeletalMeshComponent* GetSkeletalMesh() const;
diff --git a/Engine/Plugins/Animation/ControlRig/Source/ControlRig/Public/ControlRig.h b/Engine/Plugins/Animation/ControlRig/Source/ControlRig/Public/ControlRig.h
--- a/Engine/Plugins/Animation/ControlRig/Source/ControlRig/Public/ControlRig.h
+++ b/Engine/Plugins/Animation/ControlRig/Source/ControlRig/Public/ControlRig.h
@@ -110,7 +110,7 @@ public:
 	/** Creates a transformable control handle for the specified control to be used by the constraints system. Should use the UObject from 
 	ConstraintsScriptingLibrary::GetManager(UWorld* InWorld)*/
 	UFUNCTION(BlueprintCallable, Category = "Control Rig | Constraints")
-	UTransformableControlHandle* CreateTransformableControlHandle(UObject* Outer, const FName& ControlName) const;
+	UTransformableControlHandle* CreateTransformableControlHandle(const FName& ControlName) const;
 
 
 #if WITH_EDITOR
diff --git a/Engine/Plugins/Animation/ControlRig/Source/ControlRig/Public/Sequencer/MovieSceneControlRigParameterSection.h b/Engine/Plugins/Animation/ControlRig/Source/ControlRig/Public/Sequencer/MovieSceneControlRigParameterSection.h
--- a/Engine/Plugins/Animation/ControlRig/Source/ControlRig/Public/Sequencer/MovieSceneControlRigParameterSection.h
+++ b/Engine/Plugins/Animation/ControlRig/Source/ControlRig/Public/Sequencer/MovieSceneControlRigParameterSection.h
@@ -19,7 +19,7 @@
 #include "Channels/MovieSceneBoolChannel.h"
 #include "Sequencer/MovieSceneControlRigSpaceChannel.h"
 #include "ConstraintChannel.h"
-
+#include "KeyParams.h"
 #include "MovieSceneControlRigParameterSection.generated.h"
 
 class UAnimSequence;
@@ -272,12 +272,12 @@ public:
 	/*
 	* Whether it has that channel
 	*/
-	virtual bool HasConstraintChannel(const FName& InConstraintName) const override;
+	virtual bool HasConstraintChannel(const FGuid& InConstraintName) const override;
 
 	/*
 	* Get constraint with that name
 	*/
-	virtual FConstraintAndActiveChannel* GetConstraintChannel(const FName& InConstraintName) override;
+	virtual FConstraintAndActiveChannel* GetConstraintChannel(const FGuid& InConstraintID) override;
 
 	/*
 	*  Add Constraint channel
@@ -309,11 +309,11 @@ public:
 
 #if WITH_EDITOR
 	//Function to save control rig key when recording.
-	void RecordControlRigKey(FFrameNumber FrameNumber, bool bSetDefault, ERichCurveInterpMode InInterpMode);
+	void RecordControlRigKey(FFrameNumber FrameNumber, bool bSetDefault, EMovieSceneKeyInterpolation InInterpMode);
 
 	//Function to load an Anim Sequence into this section. It will automatically resize to the section size.
 	//Will return false if fails or is canceled
-	virtual bool LoadAnimSequenceIntoThisSection(UAnimSequence* Sequence, UMovieScene* MovieScene, UObject* BoundObject, bool bKeyReduce, float Tolerance, FFrameNumber InStartFrame = 0);
+	virtual bool LoadAnimSequenceIntoThisSection(UAnimSequence* Sequence, UMovieScene* MovieScene, UObject* BoundObject, bool bKeyReduce, float Tolerance, FFrameNumber InStartFrame = 0, EMovieSceneKeyInterpolation InInterolation = EMovieSceneKeyInterpolation::SmartAuto);
 #endif
 	const TArray<bool>& GetControlsMask() const
 	{
diff --git a/Engine/Plugins/Animation/ControlRig/Source/ControlRigEditor/Private/ControlRigSequencerEditorLibrary.cpp b/Engine/Plugins/Animation/ControlRig/Source/ControlRigEditor/Private/ControlRigSequencerEditorLibrary.cpp
--- a/Engine/Plugins/Animation/ControlRig/Source/ControlRigEditor/Private/ControlRigSequencerEditorLibrary.cpp
+++ b/Engine/Plugins/Animation/ControlRig/Source/ControlRigEditor/Private/ControlRigSequencerEditorLibrary.cpp
@@ -471,6 +471,11 @@ UTickableConstraint* UControlRigSequencerEditorLibrary::AddConstraint(UWorld* Wo
 	{
 		FMovieSceneConstraintChannelHelper::SmartConstraintKey(WeakSequencer.Pin(), Constraint, TOptional<bool>(), TOptional<FFrameNumber>());
 	}
+	else
+	{
+		FConstraintsManagerController& Controller = FConstraintsManagerController::Get(World);
+		Controller.StaticConstraintCreated(World, Constraint);
+	}
 	return Constraint;
 }
 
@@ -478,10 +483,10 @@ TArray <UTickableConstraint*> UControlRigSequencerEditorLibrary::GetConstraintsF
 {
 	TArray <UTickableConstraint*> Constraints;
 	const FConstraintsManagerController& Controller = FConstraintsManagerController::Get(InWorld);
-	const TArray<TObjectPtr<UTickableConstraint>>& AllConstraints = Controller.GetAllConstraints(false);
-	for (const TObjectPtr<UTickableConstraint>& TickConstraint : AllConstraints)
+	const TArray<TWeakObjectPtr<UTickableConstraint>>& AllConstraints = Controller.GetAllConstraints(false);
+	for (const TWeakObjectPtr<UTickableConstraint>& TickConstraint : AllConstraints)
 	{
-		if (TObjectPtr<UTickableTransformConstraint> Constraint = Cast<UTickableTransformConstraint>(TickConstraint))
+		if (TObjectPtr<UTickableTransformConstraint> Constraint = Cast<UTickableTransformConstraint>(TickConstraint.Get()))
 		{
 			if (Constraint->ChildTRSHandle == InChild)
 			{
@@ -588,7 +593,7 @@ bool UControlRigSequencerEditorLibrary::GetConstraintKeys(UTickableConstraint* I
 		UE_LOG(LogControlRig, Error, TEXT("GetConstraintKeys: Section doesn't support constraints"));
 		return false;
 	}
-	FConstraintAndActiveChannel* ConstraintAndChannel = ConstrainedSection->GetConstraintChannel(InConstraint->GetFName());
+	FConstraintAndActiveChannel* ConstraintAndChannel = ConstrainedSection->GetConstraintChannel(InConstraint->ConstraintID);
 	if (ConstraintAndChannel == nullptr)
 	{
 		UE_LOG(LogControlRig, Error, TEXT("GetConstraintKeys: Constraint not found in section"));
@@ -661,7 +666,7 @@ bool UControlRigSequencerEditorLibrary::MoveConstraintKey(UTickableConstraint* C
 		UE_LOG(LogControlRig, Error, TEXT("MoveConstraintKey: Constraint not valid"));
 		return false;
 	}
-	FConstraintAndActiveChannel* ConstraintAndChannel = ConstrainedSection->GetConstraintChannel(Constraint->GetFName());
+	FConstraintAndActiveChannel* ConstraintAndChannel = ConstrainedSection->GetConstraintChannel(Constraint->ConstraintID);
 	if (ConstraintAndChannel == nullptr)
 	{
 		UE_LOG(LogControlRig, Error, TEXT("MoveConstraintKey: Constraint not found in section"));
@@ -719,7 +724,7 @@ bool UControlRigSequencerEditorLibrary::DeleteConstraintKey(UTickableConstraint*
 		UE_LOG(LogControlRig, Error, TEXT("DeleteConstraintKey: Constraint not valid"));
 		return false;
 	}
-	FConstraintAndActiveChannel* ConstraintAndChannel = ConstrainedSection->GetConstraintChannel(Constraint->GetFName());
+	FConstraintAndActiveChannel* ConstraintAndChannel = ConstrainedSection->GetConstraintChannel(Constraint->ConstraintID);
 	if (ConstraintAndChannel == nullptr)
 	{
 		UE_LOG(LogControlRig, Error, TEXT("DeleteConstraintKey: Constraint not found in section"));
diff --git a/Engine/Plugins/Animation/ControlRig/Source/ControlRigEditor/Private/EditMode/ControlRigEditMode.cpp b/Engine/Plugins/Animation/ControlRig/Source/ControlRigEditor/Private/EditMode/ControlRigEditMode.cpp
--- a/Engine/Plugins/Animation/ControlRig/Source/ControlRigEditor/Private/EditMode/ControlRigEditMode.cpp
+++ b/Engine/Plugins/Animation/ControlRig/Source/ControlRigEditor/Private/EditMode/ControlRigEditMode.cpp
@@ -3972,7 +3972,7 @@ void FControlRigEditMode::SetControlShapeTransform(
 	}
 	
 	// find the last constraint in the stack (this could be cached on mouse press)
-	TArray< TObjectPtr<UTickableConstraint> > Constraints;
+	TArray< TWeakObjectPtr<UTickableConstraint> > Constraints;
 	FTransformConstraintUtils::GetParentConstraints(ControlRig->GetWorld(), InShapeActor, Constraints);
 
 	const int32 LastActiveIndex = FTransformConstraintUtils::GetLastActiveConstraintIndex(Constraints);
diff --git a/Engine/Plugins/Animation/ControlRig/Source/ControlRigEditor/Private/EditMode/SControlRigEditModeTools.cpp b/Engine/Plugins/Animation/ControlRig/Source/ControlRigEditor/Private/EditMode/SControlRigEditModeTools.cpp
--- a/Engine/Plugins/Animation/ControlRig/Source/ControlRigEditor/Private/EditMode/SControlRigEditModeTools.cpp
+++ b/Engine/Plugins/Animation/ControlRig/Source/ControlRigEditor/Private/EditMode/SControlRigEditModeTools.cpp
@@ -34,6 +34,7 @@
 #include "ControlRigSpaceChannelEditors.h"
 #include "IKeyArea.h"
 #include "Widgets/Notifications/SNotificationList.h"
+#include "Widgets/Input/SComboButton.h"
 #include "Framework/Notifications/NotificationManager.h"
 #include "ScopedTransaction.h"
 #include "EditMode/ControlRigEditModeToolkit.h"
@@ -189,6 +190,7 @@ void SControlRigEditModeTools::Construct(const FArguments& InArgs, TSharedPtr<FC
 	RigTreeDelegates.OnSelectionChanged = FOnRigTreeSelectionChanged::CreateSP(this, &SControlRigEditModeTools::HandleSelectionChanged);
 #endif
 
+
 	ChildSlot
 	[
 		SNew(SScrollBox)
@@ -362,7 +364,51 @@ void SControlRigEditModeTools::Construct(const FArguments& InArgs, TSharedPtr<FC
 					[
 						SNew(SSpacer)
 					]
+					// "Selected" button
+					+SHorizontalBox::Slot()
+					.AutoWidth()
+					.HAlign(HAlign_Right)
+					.VAlign(VAlign_Center)
+					.Padding(0.f, 2.f, 8.f, 2.f)
+					[
+						// Combo Button to swap what constraints we show  
+						SNew(SComboButton)
+						.OnGetMenuContent_Lambda([this]()
+						{
+							FMenuBuilder MenuBuilder(true, NULL);
+							MenuBuilder.BeginSection("Constraints");
+							if (ConstraintsEditionWidget.IsValid())
+							{
+								for (int32 Index = 0; Index < 4; ++Index)
+								{
+									FUIAction ItemAction(FExecuteAction::CreateSP(this, &SControlRigEditModeTools::OnSelectShowConstraints, Index));
+									const TAttribute<FText> Text = ConstraintsEditionWidget->GetShowConstraintsText((FBaseConstraintListWidget::EShowConstraints)(Index));
+									const TAttribute<FText> Tooltip = ConstraintsEditionWidget->GetShowConstraintsTooltip((FBaseConstraintListWidget::EShowConstraints)(Index));
+									MenuBuilder.AddMenuEntry(Text, Tooltip, FSlateIcon(), ItemAction);
+								}
+							}
+							MenuBuilder.EndSection();
 
+							return MenuBuilder.MakeWidget();
+						})
+						.ButtonContent()
+						[
+							SNew(SHorizontalBox)
+			
+							+SHorizontalBox::Slot()
+							[
+								SNew(STextBlock)
+								.Text_Lambda([this]()
+									{
+										return GetShowConstraintsName();
+									})
+								.ToolTipText_Lambda([this]()
+								{
+									return GetShowConstraintsTooltip();
+								})
+							]
+						]
+					]
 					// "Plus" icon
 					+SHorizontalBox::Slot()
 					.AutoWidth()
@@ -432,9 +478,38 @@ void SControlRigEditModeTools::SetSettingsDetailsObject(const TWeakObjectPtr<>&
 		TArray<TWeakObjectPtr<>> Objects;
 		Objects.Add(InObject);
 		SettingsDetailsView->SetObjects(Objects);
+	}
+}
+
+void SControlRigEditModeTools::OnSelectShowConstraints(int32 Index)
+{
+	if (ConstraintsEditionWidget.IsValid())
+	{
+		SConstraintsEditionWidget::EShowConstraints ShowConstraint = (SConstraintsEditionWidget::EShowConstraints)(Index);
+		ConstraintsEditionWidget->SetShowConstraints(ShowConstraint);
+	}
+}
 
+FText SControlRigEditModeTools::GetShowConstraintsName() const
+{
+	FText Text = FText::GetEmpty();
+	if (ConstraintsEditionWidget.IsValid())
+	{
+		Text = ConstraintsEditionWidget->GetShowConstraintsText(FBaseConstraintListWidget::ShowConstraints);
+	}
+	return Text;
+}
+
+FText SControlRigEditModeTools::GetShowConstraintsTooltip() const
+{
+	FText Text = FText::GetEmpty();
+	if (ConstraintsEditionWidget.IsValid())
+	{
+		Text = ConstraintsEditionWidget->GetShowConstraintsTooltip(FBaseConstraintListWidget::ShowConstraints);
 	}
+	return Text;
 }
+
 #if USE_LOCAL_DETAILS
 
 void SControlRigEditModeTools::SetEulerTransformDetailsObjects(const TArray<TWeakObjectPtr<>>& InObjects)
diff --git a/Engine/Plugins/Animation/ControlRig/Source/ControlRigEditor/Private/EditMode/SControlRigEditModeTools.h b/Engine/Plugins/Animation/ControlRig/Source/ControlRigEditor/Private/EditMode/SControlRigEditModeTools.h
--- a/Engine/Plugins/Animation/ControlRig/Source/ControlRigEditor/Private/EditMode/SControlRigEditModeTools.h
+++ b/Engine/Plugins/Animation/ControlRig/Source/ControlRigEditor/Private/EditMode/SControlRigEditModeTools.h
@@ -125,6 +125,11 @@ private:
 	EVisibility GetRigOptionExpanderVisibility() const;
 
 	void OnRigOptionFinishedChange(const FPropertyChangedEvent& PropertyChangedEvent);
+	
+	/** constraint type to show selection*/
+	void OnSelectShowConstraints(int32 Index);
+	FText GetShowConstraintsName() const;
+	FText GetShowConstraintsTooltip() const;
 
 private:
 	/** Toolbar functions and windows*/
diff --git a/Engine/Plugins/Animation/ControlRig/Source/ControlRigEditor/Private/Editor/Constraints/SConstraintsWidget.cpp b/Engine/Plugins/Animation/ControlRig/Source/ControlRigEditor/Private/Editor/Constraints/SConstraintsWidget.cpp
--- a/Engine/Plugins/Animation/ControlRig/Source/ControlRigEditor/Private/Editor/Constraints/SConstraintsWidget.cpp
+++ b/Engine/Plugins/Animation/ControlRig/Source/ControlRigEditor/Private/Editor/Constraints/SConstraintsWidget.cpp
@@ -25,6 +25,7 @@
 #include "ISequencer.h"
 #include "Tools/BakingHelper.h"
 #include "MovieSceneToolHelpers.h"
+#include "ConstraintsManager.h"
 #include "Styling/SlateIconFinder.h"
 #include "Widgets/Input/NumericTypeInterface.h"
 #include "FrameNumberDetailsCustomization.h"
@@ -346,6 +347,11 @@ void SDroppableConstraintItem::CreateConstraint(
 					{
 						FMovieSceneConstraintChannelHelper::SmartConstraintKey(WeakSequencer.Pin(), Constraint, TOptional<bool>(), TOptional<FFrameNumber>());
 					}
+					else
+					{
+						FConstraintsManagerController& Controller = FConstraintsManagerController::Get(World);
+						Controller.StaticConstraintCreated(World, Constraint);
+					}		
 				}
 			}
 		}
@@ -504,17 +510,16 @@ void SEditableConstraintItem::Construct(
 	};
 	TWeakObjectPtr<UTickableConstraint> Constraint = GetConstraint();
 
+	UWorld* World = GetCurrentWorld();
+
 	// sequencer
 	TWeakPtr<ISequencer> WeakSequencer = GetSequencerChecked();
-	
+
+	//manager for static check
+	UConstraintsManager* Manager = UConstraintsManager::Find(World);
+
 	// labels
 	FString ParentLabel(TEXT("undefined")), ChildLabel(TEXT("undefined"));
-	if (!InItem->GetLabel().IsEmpty())
-	{
-		ChildLabel = InItem->GetLabel();
-		InItem->GetLabel().Split(TEXT("."), &ParentLabel, &ChildLabel);
-	}
-	
 	FString ParentFullLabel = ParentLabel, ChildFullLabel = ChildLabel;
 	if (IsValid(Constraint.Get()))
 	{
@@ -536,12 +541,16 @@ void SEditableConstraintItem::Construct(
 			SNew(SBorder)
 			.Padding(FMargin(5.0, 2.0, 5.0, 2.0))
 			.BorderImage(RoundedBoxBrush)
-			.BorderBackgroundColor_Lambda([Constraint]()
+			.BorderBackgroundColor_Lambda([Manager, World,Constraint]()
 			{
 				if (!Constraint.IsValid() || !IsValid(Constraint.Get()))
 				{
 					return FStyleColors::Transparent;
 				}
+				if (Manager && Manager->IsStaticConstraint(Constraint.Get()))
+				{
+					return FSlateColor(FLinearColor::Green * 0.5);
+				}
 				return Constraint->IsFullyActive() ? FStyleColors::Select : FStyleColors::Transparent;
 			})
 			.Content()
@@ -568,9 +577,17 @@ void SEditableConstraintItem::Construct(
 				.Padding(0)
 				[
 					SNew( STextBlock )
-					.Text_Lambda( [ParentLabel]()
+					.Text_Lambda( [Constraint, ParentLabel, ChildLabel,World]()
 					{
-						return FText::FromString(ParentLabel);
+							FString CurrentParentLabel = ParentLabel, CurrentChildLabel = ChildLabel;
+							if (Constraint.IsValid() && IsValid(Constraint.Get()) && Constraint->IsValid())
+							{
+								FString Label = Constraint->GetLabel();
+								Label.Split(TEXT("."), &CurrentParentLabel, &CurrentChildLabel);
+							}
+							static constexpr TCHAR LabelFormat[] = TEXT("%s <-- %s");
+							const FString Label = FString::Printf(LabelFormat, *CurrentParentLabel, *CurrentChildLabel);
+							return FText::FromString(Label);
 					})
 					.Font_Lambda([Constraint]()
 					{
@@ -779,6 +796,66 @@ void FBaseConstraintListWidget::InvalidateConstraintList()
 	}
 }
 
+FBaseConstraintListWidget::EShowConstraints FBaseConstraintListWidget::ShowConstraints = FBaseConstraintListWidget::EShowConstraints::ShowSelected;
+
+FText FBaseConstraintListWidget::GetShowConstraintsText(EShowConstraints Index) const
+{
+	FText Text = FText::GetEmpty();
+	switch (Index)
+	{
+		case EShowConstraints::ShowSelected:
+		{
+			Text = FText(LOCTEXT("Selected", "Selected"));
+		}
+		break;
+		case EShowConstraints::ShowLevelSequence:
+		{
+			Text = FText(LOCTEXT("LevelSequence", "Level Sequence"));
+		}
+		break;
+		case EShowConstraints::ShowValid:
+		{
+			Text = FText(LOCTEXT("Current", "Current"));
+		}
+		break;
+		case EShowConstraints::ShowAll:
+		{
+			Text = FText(LOCTEXT("All", "All"));
+		}
+		break;
+	}
+	return Text;
+}
+
+FText FBaseConstraintListWidget::GetShowConstraintsTooltip(EShowConstraints Index) const
+{
+	FText Text = FText::GetEmpty();
+	switch (Index)
+	{
+		case EShowConstraints::ShowSelected:
+		{
+			Text = FText(LOCTEXT("SelectedTooltip", "Show constraints on selected"));
+		}
+		break;
+		case EShowConstraints::ShowLevelSequence:
+		{
+			Text = FText(LOCTEXT("LevelSequenceTooltip", "Show constraints on active level sequence"));
+		}
+		break;
+		case EShowConstraints::ShowValid:
+		{
+			Text = FText(LOCTEXT("CurrentTooltip", "Show constraints that are current in level sequences or static"));
+		}
+		break;
+		case EShowConstraints::ShowAll:
+		{
+			Text = FText(LOCTEXT("AllTooltip", "Show all constraints"));
+		}
+		break;
+	}
+	return Text;
+}
+
 int32 FBaseConstraintListWidget::RefreshConstraintList()
 {
 	// get constraints
@@ -792,18 +869,67 @@ int32 FBaseConstraintListWidget::RefreshConstraintList()
 
 	const bool bIsConstraintsActor = Selection.Num() == 1 && Selection[0]->IsA<AConstraintsActor>();
 
-	TArray< TObjectPtr<UTickableConstraint> > Constraints;
-	if (bIsConstraintsActor)
+	TArray< TWeakObjectPtr<UTickableConstraint> > Constraints;
+	if (ShowConstraints == EShowConstraints::ShowSelected)
 	{
-		const FConstraintsManagerController& Controller = FConstraintsManagerController::Get(World);
-		static constexpr bool bSorted = true;
-		Constraints = Controller.GetAllConstraints(bSorted);
+		if (bIsConstraintsActor)
+		{
+			const FConstraintsManagerController& Controller = FConstraintsManagerController::Get(World);
+			static constexpr bool bSorted = true;
+			TArray< TObjectPtr<UTickableConstraint> > StaticConstraints;
+
+			StaticConstraints = Controller.GetStaticConstraints(bSorted);
+			for (TObjectPtr<UTickableConstraint>& Constraint : StaticConstraints)
+			{
+				Constraints.Add(Constraint);
+			}
+		}
+		else
+		{
+			for (const AActor* Actor : Selection)
+			{
+				FTransformConstraintUtils::GetParentConstraints(World, Actor, Constraints);
+			}
+			//remove if not active...
+			for (int32 Index = Constraints.Num() - 1; Index >= 0; --Index)
+			{
+				if (Constraints[Index].Get() == nullptr || Constraints[Index].Get()->IsValid() == false)
+				{
+					Constraints.RemoveAt(Index);
+				}
+			}
+		}
 	}
 	else
 	{
-		for (const AActor* Actor : Selection)
+		const FConstraintsManagerController& Controller = FConstraintsManagerController::Get(World);
+		static constexpr bool bSorted = true;
+		Constraints = Controller.GetAllConstraints(bSorted);
+		if (ShowConstraints == EShowConstraints::ShowLevelSequence)
+		{
+			TWeakPtr<ISequencer> WeakSequencer = GetSequencerChecked();
+			if (WeakSequencer.IsValid())
+			{
+				UMovieScene* MovieScene = WeakSequencer.Pin()->GetFocusedMovieSceneSequence()->GetMovieScene();
+				for (int32 Index = Constraints.Num() - 1; Index >= 0; --Index)
+				{
+					UMovieScene* OuterMovieScene = Constraints[Index].Get() ? Constraints[Index]->GetTypedOuter<UMovieScene>() : nullptr;
+					if (OuterMovieScene != MovieScene)
+					{
+						Constraints.RemoveAt(Index);
+					}
+				}
+			}
+		}
+		else if (ShowConstraints == EShowConstraints::ShowValid)
 		{
-			FTransformConstraintUtils::GetParentConstraints(World, Actor, Constraints);
+			for (int32 Index = Constraints.Num() - 1; Index >= 0; --Index)
+			{
+				if (Constraints[Index].Get() == nullptr || Constraints[Index].Get()->IsValid() == false)
+				{
+					Constraints.RemoveAt(Index);
+				}
+			}
 		}
 	}
 
@@ -811,13 +937,13 @@ int32 FBaseConstraintListWidget::RefreshConstraintList()
 	ListItems.Empty();
 
 	const UEnum* ETransformConstraintTypeEnum = StaticEnum<ETransformConstraintType>();
-	for (const TObjectPtr<UTickableConstraint>& Constraint : Constraints)
+	for (const TWeakObjectPtr<UTickableConstraint>& Constraint : Constraints)
 	{
 		const int8 Type = FConstraintInfo::GetType(Constraint->GetClass());
 		if (ETransformConstraintTypeEnum->IsValidEnumValue(Type))
 		{
 			const ETransformConstraintType ConstraintType = static_cast<ETransformConstraintType>(Type);
-			ListItems.Emplace(FEditableConstraintItem::Make(Constraint, ConstraintType));
+			ListItems.Emplace(FEditableConstraintItem::Make(Constraint.Get(), ConstraintType));
 		}
 	}
 
@@ -842,6 +968,7 @@ void FBaseConstraintListWidget::OnActorSelectionChanged(const TArray<UObject*>&
 
 void SConstraintsEditionWidget::Construct(const FArguments& InArgs)
 {
+	WeakSequencer = GetSequencerChecked();
 
 	ChildSlot
 	[
@@ -907,12 +1034,24 @@ void SConstraintsEditionWidget::Construct(const FArguments& InArgs)
 	RegisterSelectionChanged();
 }
 
+bool SConstraintsEditionWidget::SequencerTimeChanged() 
+{
+	if (WeakSequencer.IsValid())
+	{
+		TSharedPtr<ISequencer> Sequencer = WeakSequencer.Pin();
+		const bool bTimeChanged = (SequencerLastTime.Rate != Sequencer->GetLocalTime().Rate ||
+			SequencerLastTime.Time != Sequencer->GetLocalTime().Time);
+		SequencerLastTime = Sequencer->GetLocalTime();
+		return bTimeChanged;
+	}
+	return false;
+}
 
 void SConstraintsEditionWidget::Tick(const FGeometry& AllottedGeometry, const double InCurrentTime, const float InDeltaTime)
 {
 	SCompoundWidget::Tick(AllottedGeometry, InCurrentTime, InDeltaTime);
 
-	if(bNeedsRefresh)
+	if(bNeedsRefresh || SequencerTimeChanged())
 	{
 		RefreshConstraintList();
 		bNeedsRefresh = false;
@@ -1017,7 +1156,7 @@ void SConstraintsEditionWidget::RemoveItem(const TSharedPtr<FEditableConstraintI
 	
 	FScopedTransaction Transaction(LOCTEXT("RemoveConstraint", "Remove Constraint"));
 	
-	const FConstraintsManagerController& Controller = FConstraintsManagerController::Get(World);
+	FConstraintsManagerController& Controller = FConstraintsManagerController::Get(World);
 	
 	Controller.RemoveConstraint(Item->Constraint.Get());
 
@@ -1027,7 +1166,6 @@ void SConstraintsEditionWidget::RemoveItem(const TSharedPtr<FEditableConstraintI
 FReply SConstraintsEditionWidget::OnBakeClicked()
 {
 	UTickableTransformConstraint* InConstraint = nullptr;
-	TWeakPtr<ISequencer> WeakSequencer = GetSequencerChecked();
 	if (!WeakSequencer.IsValid() || ListItems.Num() < 1)
 	{
 		return FReply::Unhandled();
@@ -1140,11 +1278,10 @@ TSharedPtr<SWidget> SConstraintsEditionWidget::CreateContextMenu()
 			LOCTEXT("BakeConstraintLabel", "Bake"),
 			FText::Format(LOCTEXT("BakeConstraintDoItTooltip", "Bake {0} transforms."), ConstraintLabel),
 			FSlateIcon(),
-			FUIAction(FExecuteAction::CreateLambda([Constraint, World]()
+			FUIAction(FExecuteAction::CreateLambda([this, Constraint, World]()
 				{
 					if (UTickableTransformConstraint* TransformConstraint = Cast<UTickableTransformConstraint>(Constraint))
 					{
-						const TWeakPtr<ISequencer> WeakSequencer = GetSequencerChecked();
 						if (!WeakSequencer.IsValid())
 						{
 							return;
@@ -1165,7 +1302,6 @@ TSharedPtr<SWidget> SConstraintsEditionWidget::CreateContextMenu()
 			return TransformConstraint->bDynamicOffset;
 		});
 
-		TWeakPtr<ISequencer> WeakSequencer = GetSequencerChecked();
 		if (!bIsLookAtConstraint && WeakSequencer.IsValid())
 		{
 			MenuBuilder.BeginSection("KeyConstraint", LOCTEXT("KeyConstraintHeader", "Keys"));
@@ -1174,7 +1310,7 @@ TSharedPtr<SWidget> SConstraintsEditionWidget::CreateContextMenu()
 				LOCTEXT("CompensateKeyLabel", "Compensate Key"),
 				FText::Format(LOCTEXT("CompensateKeyTooltip", "Compensate transform key for {0}."), ConstraintLabel),
 				FSlateIcon(),
-				FUIAction(FExecuteAction::CreateLambda([TransformConstraint, WeakSequencer]()
+				FUIAction(FExecuteAction::CreateLambda([TransformConstraint, this]()
 				{
 					const TSharedPtr<ISequencer> Sequencer = WeakSequencer.Pin();
 					const FFrameRate TickResolution = Sequencer->GetFocusedTickResolution();
@@ -1189,7 +1325,7 @@ TSharedPtr<SWidget> SConstraintsEditionWidget::CreateContextMenu()
 				LOCTEXT("CompensateAllKeysLabel", "Compensate All Keys"),
 				FText::Format(LOCTEXT("CompensateAllKeysTooltip", "Compensate all transform keys for {0}."), ConstraintLabel),
 				FSlateIcon(),
-				FUIAction(FExecuteAction::CreateLambda([TransformConstraint, WeakSequencer]()
+				FUIAction(FExecuteAction::CreateLambda([TransformConstraint, this]()
 				{
 					FMovieSceneConstraintChannelHelper::Compensate(WeakSequencer.Pin(), TransformConstraint, TOptional<FFrameNumber>());
 				}), IsCompensationEnabled),
@@ -1503,11 +1639,7 @@ void SBakeConstraintItem::Construct(
 
 	// labels
 	FString ParentLabel(TEXT("undefined")), ChildLabel(TEXT("undefined"));
-	if (!InItem->GetLabel().IsEmpty())
-	{
-		ChildLabel = InItem->GetLabel();
-		InItem->GetLabel().Split(TEXT("."), &ParentLabel, &ChildLabel);
-	}
+	
 
 	FString ParentFullLabel = ParentLabel, ChildFullLabel = ChildLabel;
 	if (IsValid(Constraint.Get()))
@@ -1558,9 +1690,17 @@ void SBakeConstraintItem::Construct(
 					.Padding(0)
 					[
 						SNew(STextBlock)
-						.Text_Lambda([ParentLabel]()
+						.Text_Lambda([Constraint, ParentLabel, ChildLabel]()
 							{
-								return FText::FromString(ParentLabel);
+								FString CurrentParentLabel = ParentLabel, CurrentChildLabel = ChildLabel;
+								if (Constraint.IsValid() && IsValid(Constraint.Get()))
+								{
+									FString Label = Constraint->GetLabel();
+									Label.Split(TEXT("."), &CurrentParentLabel, &CurrentChildLabel);
+								}
+								static constexpr TCHAR LabelFormat[] = TEXT("%s <-- %s");
+								const FString Label = FString::Printf(LabelFormat, *CurrentParentLabel, *CurrentChildLabel);
+								return FText::FromString(Label);
 							})
 					.Font_Lambda([Constraint]()
 						{
diff --git a/Engine/Plugins/Animation/ControlRig/Source/ControlRigEditor/Private/Editor/Constraints/SConstraintsWidget.h b/Engine/Plugins/Animation/ControlRig/Source/ControlRigEditor/Private/Editor/Constraints/SConstraintsWidget.h
--- a/Engine/Plugins/Animation/ControlRig/Source/ControlRigEditor/Private/Editor/Constraints/SConstraintsWidget.h
+++ b/Engine/Plugins/Animation/ControlRig/Source/ControlRigEditor/Private/Editor/Constraints/SConstraintsWidget.h
@@ -8,6 +8,7 @@
 #include "EditorUndoClient.h"
 #include "IStructureDetailsView.h"
 #include "BakingAnimationKeySettings.h"
+#include "Misc/QualifiedFrameTime.h"
 #include "ConstraintsManager.h"
 
 class AActor;
@@ -213,6 +214,18 @@ class CONTROLRIGEDITOR_API FBaseConstraintListWidget : public FEditorUndoClient
 {
 public:
 
+	/**
+	*  Constraints to show in Widget
+	*/
+	enum class EShowConstraints
+	{
+		ShowSelected = 0x0,
+		ShowLevelSequence = 0x1,
+		ShowValid = 0x2,
+		ShowAll = 0x3,
+	};
+public:
+
 	virtual ~FBaseConstraintListWidget() override;
 	/* FEditorUndoClient interface */
 	virtual void PostUndo(bool bSuccess);
@@ -228,6 +241,18 @@ public:
 	/** Triggers a constraint list invalidation when selection in the level viewport. */
 	void OnActorSelectionChanged(const TArray<UObject*>& NewSelection, bool bForceRefresh);
 
+	/** Which constraints to show  */
+	EShowConstraints GetShowConstraints() const { return ShowConstraints; }
+	void SetShowConstraints(EShowConstraints InShowConstraints)
+	{
+		if (InShowConstraints != ShowConstraints) {
+			ShowConstraints = InShowConstraints;
+			bNeedsRefresh = true;
+		}
+	}
+	FText GetShowConstraintsText(EShowConstraints Index) const;
+	FText GetShowConstraintsTooltip(EShowConstraints Index) const;
+
 protected:
 	/** Types */
 	using ItemSharedPtr = TSharedPtr<FEditableConstraintItem>;
@@ -248,6 +273,10 @@ protected:
 
 	FDelegateHandle OnSelectionChangedHandle;
 
+public:
+
+	static EShowConstraints ShowConstraints;
+
 };
 
 /**
@@ -259,7 +288,6 @@ class CONTROLRIGEDITOR_API SConstraintsEditionWidget : public SCompoundWidget, p
 public:
 	SLATE_BEGIN_ARGS(SConstraintsEditionWidget)	{}
 	SLATE_END_ARGS()
-
 	/** Constructs this widget with InArgs */
 	void Construct(const FArguments& InArgs);
 
@@ -291,6 +319,12 @@ private:
 	void OnItemDoubleClicked(ItemSharedPtr InItem);
 
 	FReply OnBakeClicked();
+
+	//sequencer and it's time
+	TWeakPtr<ISequencer> WeakSequencer;
+	bool SequencerTimeChanged();
+	FQualifiedFrameTime SequencerLastTime;
+
 };
 
 /** Widget allowing baking of constraints */
diff --git a/Engine/Plugins/Animation/ControlRig/Source/ControlRigEditor/Private/Sequencer/ControlConstraintChannelInterface.cpp b/Engine/Plugins/Animation/ControlRig/Source/ControlRigEditor/Private/Sequencer/ControlConstraintChannelInterface.cpp
--- a/Engine/Plugins/Animation/ControlRig/Source/ControlRigEditor/Private/Sequencer/ControlConstraintChannelInterface.cpp
+++ b/Engine/Plugins/Animation/ControlRig/Source/ControlRigEditor/Private/Sequencer/ControlConstraintChannelInterface.cpp
@@ -12,7 +12,7 @@
 
 #include "TransformConstraint.h"
 #include "Algo/Copy.h"
-
+#include "ConstraintsManager.h"
 #include "ScopedTransaction.h"
 #include "Channels/MovieSceneFloatChannel.h"
 #include "Constraints/MovieSceneConstraintChannelHelper.inl"
@@ -91,11 +91,17 @@ bool FControlConstraintChannelInterface::SmartConstraintKey(
 	// set constraint as dynamic
 	InConstraint->bDynamicOffset = true;
 	
+	//check if static if so we need to delete it from world, will get added later again
+	if (UConstraintsManager* Manager = InConstraint->GetTypedOuter<UConstraintsManager>())
+	{
+		Manager->RemoveStaticConstraint(InConstraint);
+	}
+
 	// add the channel
 	ConstraintSection->AddConstraintChannel(InConstraint);
 
 	// add key if needed
-	if (FConstraintAndActiveChannel* Channel = ConstraintSection->GetConstraintChannel(InConstraint->GetFName()))
+	if (FConstraintAndActiveChannel* Channel = ConstraintSection->GetConstraintChannel(InConstraint->ConstraintID))
 	{
 		bool ActiveValueToBeSet = false;
 		//add key if we can and make sure the key we are setting is what we want
diff --git a/Engine/Plugins/Animation/ControlRig/Source/ControlRigEditor/Private/Sequencer/ControlRigParameterTrackEditor.cpp b/Engine/Plugins/Animation/ControlRig/Source/ControlRigEditor/Private/Sequencer/ControlRigParameterTrackEditor.cpp
--- a/Engine/Plugins/Animation/ControlRig/Source/ControlRigEditor/Private/Sequencer/ControlRigParameterTrackEditor.cpp
+++ b/Engine/Plugins/Animation/ControlRig/Source/ControlRigEditor/Private/Sequencer/ControlRigParameterTrackEditor.cpp
@@ -1974,33 +1974,68 @@ void FControlRigParameterTrackEditor::AddTrackForComponent(USceneComponent* InCo
 	}
 }
 
+//test to see if actor has a constraint, in which case we need to add a constraint channel/key
+//or a control rig in which case we create a track if cvar is off
 void FControlRigParameterTrackEditor::HandleActorAdded(AActor* Actor, FGuid TargetObjectGuid)
 {
-	if (!CVarAutoGenerateControlRigTrack.GetValueOnGameThread())
+	if (Actor == nullptr)
 	{
 		return;
 	}
-
-	if (Actor)
+	//test for constraint
+	FConstraintsManagerController& Controller = FConstraintsManagerController::Get(Actor->GetWorld());
+	TArray< TWeakObjectPtr<UTickableConstraint>> Constraints = Controller.GetAllConstraints();
+	for (TWeakObjectPtr<UTickableConstraint>& WeakConstraint : Constraints)
 	{
-		if (UControlRigComponent* ControlRigComponent = Actor->FindComponentByClass<UControlRigComponent>())
+		if (WeakConstraint.IsValid())
 		{
-			AddControlRigFromComponent(TargetObjectGuid);
-			return;
+			if (UTickableTransformConstraint* Constraint = Cast<UTickableTransformConstraint>(WeakConstraint.Get()))
+			{
+				if (Constraint->ChildTRSHandle)
+				{
+					if (UObject* Child = Constraint->ChildTRSHandle->GetTarget().Get())
+					{
+						AActor* TargetActor = Cast<AActor>(Child);
+						if (TargetActor == nullptr)
+						{
+							TargetActor = Child->GetTypedOuter<AActor>();
+						}
+						if (TargetActor == Actor)
+						{
+							const TOptional<bool> bActive = true;
+							const TSharedPtr<ISequencer> SequencerPtr = GetSequencer();
+							FMovieSceneConstraintChannelHelper::SmartConstraintKey(SequencerPtr, Constraint, bActive, TOptional<FFrameNumber>());
+						}		
+					}
+				}
+			}
 		}
+	}
+	
+	//test for control rig
 
-		if (USkeletalMeshComponent* SkeletalMeshComponent = Cast<USkeletalMeshComponent>(Actor->GetRootComponent()))
-		{
-			AddTrackForComponent(SkeletalMeshComponent, TargetObjectGuid);
-			return;
-		}
+	if (!CVarAutoGenerateControlRigTrack.GetValueOnGameThread())
+	{
+		return;
+	}
+
+	if (UControlRigComponent* ControlRigComponent = Actor->FindComponentByClass<UControlRigComponent>())
+	{
+		AddControlRigFromComponent(TargetObjectGuid);
+		return;
+	}
+
+	if (USkeletalMeshComponent* SkeletalMeshComponent = Cast<USkeletalMeshComponent>(Actor->GetRootComponent()))
+	{
+		AddTrackForComponent(SkeletalMeshComponent, TargetObjectGuid);
+		return;
+	}
 
-		for (UActorComponent* Component : Actor->GetComponents())
+	for (UActorComponent* Component : Actor->GetComponents())
+	{
+		if (USceneComponent* SceneComp = Cast<USceneComponent>(Component))
 		{
-			if (USceneComponent* SceneComp = Cast<USceneComponent>(Component))
-			{
-				AddTrackForComponent(SceneComp, FGuid());
-			}
+			AddTrackForComponent(SceneComp, FGuid());
 		}
 	}
 }
@@ -2700,7 +2735,7 @@ namespace
 				return Data;
 			}
 
-			Data.Constraint = Cast<UTickableTransformConstraint>(ConstraintChannels[Index].Constraint.Get());
+			Data.Constraint = Cast<UTickableTransformConstraint>(ConstraintChannels[Index].GetConstraint().Get());
 
 			// get constraint name
 			auto GetControlName = [InSection, Index]()
@@ -2854,17 +2889,17 @@ void FControlRigParameterTrackEditor::HandleConstraintRemoved(IMovieSceneConstra
 								return;
 							}
 
-							const FConstraintAndActiveChannel* ConstraintChannel = InSection->GetConstraintChannel(Constraint->GetFName());
-							if (!ConstraintChannel || ConstraintChannel->Constraint != Constraint)
+							const FConstraintAndActiveChannel* ConstraintChannel = InSection->GetConstraintChannel(Constraint->ConstraintID);
+							if (!ConstraintChannel || ConstraintChannel->GetConstraint().Get() != Constraint)
 							{
 								return;
 							}
 
 							const bool bCompensate = (InNotifyType == EConstraintsManagerNotifyType::ConstraintRemovedWithCompensation);
-							if (bCompensate && ConstraintChannel->Constraint.IsValid())
+							if (bCompensate && ConstraintChannel->GetConstraint().Get())
 							{
 								FMovieSceneConstraintChannelHelper::HandleConstraintRemoved(
-									ConstraintChannel->Constraint.Get(),
+									ConstraintChannel->GetConstraint().Get(),
 									&ConstraintChannel->ActiveChannel,
 									GetSequencer(),
 									Section);
@@ -2915,7 +2950,7 @@ void FControlRigParameterTrackEditor::HandleConstraintPropertyChanged(UTickableT
 	const TArray<FConstraintAndActiveChannel>& ConstraintChannels = ConstraintSection->GetConstraintsChannels();
 	const FConstraintAndActiveChannel* Channel = ConstraintChannels.FindByPredicate([InConstraint](const FConstraintAndActiveChannel& Channel)
 	{
-		return Channel.Constraint == InConstraint || Channel.ConstraintCopyToSpawn == InConstraint;
+		return Channel.GetConstraint() == InConstraint;
 	});
 
 	if (!Channel)
@@ -4671,8 +4706,7 @@ bool FControlRigParameterTrackEditor::CollapseAllLayers(TSharedPtr<ISequencer>&S
 					FScopedSlowTask Feedback(Frames.Num(), LOCTEXT("CollapsingSections", "Collapsing Sections"));
 					Feedback.MakeDialog(true);
 
-					const ERichCurveInterpMode InterpMode = InSettings.bReduceKeys ? RCIM_Cubic : RCIM_Linear;
-
+					const EMovieSceneKeyInterpolation InterpMode = SequencerPtr->GetSequencerSettings()->GetKeyInterpolation();
 					Index = 0;
 					for (Index = 0; Index < Frames.Num(); ++Index)
 					{
@@ -5094,15 +5128,18 @@ bool FControlRigParameterSection::RequestDeleteCategory(const TArray<FName>& Cat
 
 	if (ParameterSection && SequencerPtr)
 	{
-		const FName& Channel = CategoryNamePaths.Last();
-
+		const FName& ChannelName = CategoryNamePaths.Last();
+		const int32 Index = ParameterSection->GetConstraintsChannels().IndexOfByPredicate([ChannelName](const FConstraintAndActiveChannel& InChannel)
+			{
+				return InChannel.GetConstraint().Get() ? InChannel.GetConstraint()->GetFName() == ChannelName : false;
+			});
 		// remove constraint channel if there are no keys
-		const FConstraintAndActiveChannel* ConstraintChannel = ParameterSection->GetConstraintChannel(Channel);
+		const FConstraintAndActiveChannel* ConstraintChannel = Index != INDEX_NONE ? &(ParameterSection->GetConstraintsChannels()[Index]): nullptr;
 		if (ConstraintChannel && ConstraintChannel->ActiveChannel.GetNumKeys() == 0)
 		{
 			if (ParameterSection->TryModify())
 			{
-				ParameterSection->RemoveConstraintChannel(ConstraintChannel->Constraint.Get());
+				ParameterSection->RemoveConstraintChannel(ConstraintChannel->GetConstraint().Get());
 				SequencerPtr->NotifyMovieSceneDataChanged(EMovieSceneDataChangeType::MovieSceneStructureItemsChanged);
 				return true;
 			}
diff --git a/Engine/Plugins/Animation/ControlRig/Source/ControlRigEditor/Private/Tools/ConstraintBaker.cpp b/Engine/Plugins/Animation/ControlRig/Source/ControlRigEditor/Private/Tools/ConstraintBaker.cpp
--- a/Engine/Plugins/Animation/ControlRig/Source/ControlRigEditor/Private/Tools/ConstraintBaker.cpp
+++ b/Engine/Plugins/Animation/ControlRig/Source/ControlRigEditor/Private/Tools/ConstraintBaker.cpp
@@ -41,7 +41,7 @@ void FConstraintBaker::GetMinimalFramesToBake(
 	TArrayView<const bool> ConstraintValues;
 
 	// note that we might want to bake a constraint which is not animated
-	FConstraintAndActiveChannel* ThisActiveChannel = ConstrainedSection->GetConstraintChannel(InConstraint->GetFName());
+	FConstraintAndActiveChannel* ThisActiveChannel = ConstrainedSection->GetConstraintChannel(InConstraint->ConstraintID);
 	if (ThisActiveChannel)
 	{
 		const TMovieSceneChannelData<const bool> ConstraintChannelData = ThisActiveChannel->ActiveChannel.GetData();
@@ -142,7 +142,7 @@ void FConstraintBaker::GetMinimalFramesToBake(
 		static constexpr bool bSorted = true;
 		const uint32 ChildHash = InConstraint->ChildTRSHandle->GetHash();
 		const FConstraintsManagerController& Controller = FConstraintsManagerController::Get(InWorld);
-		using ConstraintPtr = TObjectPtr<UTickableConstraint>;
+		using ConstraintPtr = TWeakObjectPtr<UTickableConstraint>;
 		const TArray<ConstraintPtr> Constraints = Controller.GetParentConstraints(ChildHash, bSorted);
 
 		// store the other channels that may need compensation
@@ -152,8 +152,8 @@ void FConstraintBaker::GetMinimalFramesToBake(
 			const UTickableTransformConstraint* TransformConstraint = Cast<UTickableTransformConstraint>(Constraint.Get());
 			if (TransformConstraint && TransformConstraint->NeedsCompensation() && TransformConstraint != InConstraint) 
 			{
-				const FName ConstraintName = TransformConstraint->GetFName();
-				if (FConstraintAndActiveChannel* ConstraintChannel = ConstrainedSection->GetConstraintChannel(ConstraintName))
+				const FGuid ConstraintID= TransformConstraint->ConstraintID;
+				if (FConstraintAndActiveChannel* ConstraintChannel = ConstrainedSection->GetConstraintChannel(ConstraintID))
 				{
 					OtherChannels.Add(ConstraintChannel);
 				}
diff --git a/Engine/Plugins/Importers/USDImporter/Source/USDExporter/Private/LevelSequenceExporterUSD.cpp b/Engine/Plugins/Importers/USDImporter/Source/USDExporter/Private/LevelSequenceExporterUSD.cpp
--- a/Engine/Plugins/Importers/USDImporter/Source/USDExporter/Private/LevelSequenceExporterUSD.cpp
+++ b/Engine/Plugins/Importers/USDImporter/Source/USDExporter/Private/LevelSequenceExporterUSD.cpp
@@ -1097,7 +1097,7 @@ namespace UE
 
 				const FConstraintsManagerController& Controller = FConstraintsManagerController::Get(Context.World);
 				static constexpr bool bSorted = true;
-				const TArray<TObjectPtr<UTickableConstraint>> AllConstraints = Controller.GetAllConstraints(bSorted);
+				const TArray<TWeakObjectPtr<UTickableConstraint>> AllConstraints = Controller.GetAllConstraints(bSorted);
 
 				for ( FFrameTime EvalTime = StartFrame; EvalTime <= EndFrame; EvalTime += Interval )
 				{
@@ -1105,9 +1105,12 @@ namespace UE
 					Context.Sequencer->ForceEvaluate();
 
 					// Evaluate constraints (these run on tick in the editor, so here we must trigger them manually)
-					for (const UTickableConstraint* Constraint : AllConstraints)
+					for (const TWeakObjectPtr<UTickableConstraint>& Constraint : AllConstraints)
 					{
-						Constraint->Evaluate();
+						if (Constraint.IsValid())
+						{
+							Constraint->Evaluate();
+						}
 					}
 
 					FFrameTime KeyTime = FFrameRate::Snap( EvalTime, Resolution, DisplayRate ).FloorToFrame();
diff --git a/Engine/Source/Editor/MovieSceneTools/Private/Constraints/ComponentConstraintChannelInterface.cpp b/Engine/Source/Editor/MovieSceneTools/Private/Constraints/ComponentConstraintChannelInterface.cpp
--- a/Engine/Source/Editor/MovieSceneTools/Private/Constraints/ComponentConstraintChannelInterface.cpp
+++ b/Engine/Source/Editor/MovieSceneTools/Private/Constraints/ComponentConstraintChannelInterface.cpp
@@ -12,6 +12,7 @@
 #include "Sections/MovieScene3DTransformSection.h"
 #include "Systems/MovieScenePropertyInstantiator.h"
 #include "Tracks/MovieScene3DTransformTrack.h"
+#include "ConstraintsManager.h"
 #include "ScopedTransaction.h"
 
 UMovieSceneSection* FComponentConstraintChannelInterface::GetHandleSection(const UTransformableHandle* InHandle, const TSharedPtr<ISequencer>& InSequencer)
@@ -78,8 +79,13 @@ bool FComponentConstraintChannelInterface::SmartConstraintKey(
 	InConstraint->bDynamicOffset = true;
 	
 	// add the channel
-	if (ConstraintSection->HasConstraintChannel(InConstraint->GetFName()) == false)
+	if (ConstraintSection->HasConstraintChannel(InConstraint->ConstraintID) == false)
 	{
+		//check if static if so we need to delete it from world, will get added later again
+		if (UConstraintsManager* Manager = InConstraint->GetTypedOuter<UConstraintsManager>())
+		{
+			Manager->RemoveStaticConstraint(InConstraint);
+		}
 		ConstraintSection->AddConstraintChannel(InConstraint);
 		if (InSequencer.IsValid())
 		{
@@ -88,7 +94,7 @@ bool FComponentConstraintChannelInterface::SmartConstraintKey(
 	}
 
 	// add key if needed
-	if (FConstraintAndActiveChannel* Channel = ConstraintSection->GetConstraintChannel(InConstraint->GetFName()))
+	if (FConstraintAndActiveChannel* Channel = ConstraintSection->GetConstraintChannel(InConstraint->ConstraintID))
 	{
 		bool ActiveValueToBeSet = false;
 		//add key if we can and make sure the key we are setting is what we want
diff --git a/Engine/Source/Editor/MovieSceneTools/Private/Constraints/MovieSceneConstraintChannelHelper.cpp b/Engine/Source/Editor/MovieSceneTools/Private/Constraints/MovieSceneConstraintChannelHelper.cpp
--- a/Engine/Source/Editor/MovieSceneTools/Private/Constraints/MovieSceneConstraintChannelHelper.cpp
+++ b/Engine/Source/Editor/MovieSceneTools/Private/Constraints/MovieSceneConstraintChannelHelper.cpp
@@ -42,7 +42,7 @@ void FCompensationEvaluator::ComputeLocalTransforms(
 		return;
 	}
 
-	using ConstraintPtr = TObjectPtr<UTickableConstraint>;
+	using ConstraintPtr = TWeakObjectPtr<UTickableConstraint>;
 	const TArray< ConstraintPtr > Constraints = GetHandleTransformConstraints(InWorld);
 	if (Constraints.IsEmpty())
 	{
@@ -108,9 +108,12 @@ void FCompensationEvaluator::ComputeLocalTransforms(
 		InSequencer->GetEvaluationTemplate().EvaluateSynchronousBlocking(Context, *InSequencer);
 
 		// evaluate constraints
-		for (const UTickableConstraint* InConstraint : AllConstraints)
+		for (const TWeakObjectPtr<UTickableConstraint>& InConstraint : AllConstraints)
 		{
-			InConstraint->Evaluate();
+			if (InConstraint.IsValid())
+			{
+				InConstraint->Evaluate();
+			}
 		}
 
 		for (IMovieSceneToolsAnimationBakeHelper* BakeHelper : BakeHelpers)
@@ -199,7 +202,7 @@ void FCompensationEvaluator::ComputeLocalTransformsForBaking(UWorld* InWorld, co
 		return;
 	}
 
-	using ConstraintPtr = TObjectPtr<UTickableConstraint>;
+	using ConstraintPtr = TWeakObjectPtr<UTickableConstraint>;
 	const TArray< ConstraintPtr > Constraints = GetHandleTransformConstraints(InWorld);
 	
 	const TArray< ConstraintPtr > ConstraintsMinusThis =
@@ -263,9 +266,12 @@ void FCompensationEvaluator::ComputeLocalTransformsForBaking(UWorld* InWorld, co
 		InSequencer->GetEvaluationTemplate().EvaluateSynchronousBlocking(Context, *InSequencer);
 
 		// evaluate constraints
-		for (UTickableConstraint* InConstraint : AllConstraints)
+		for (const TWeakObjectPtr<UTickableConstraint>& InConstraint : AllConstraints)
 		{
-			InConstraint->Evaluate(true);
+			if (InConstraint.IsValid())
+			{
+				InConstraint->Evaluate(true);
+			}
 		}
 
 		for (IMovieSceneToolsAnimationBakeHelper* BakeHelper : BakeHelpers)
@@ -318,12 +324,12 @@ void FCompensationEvaluator::ComputeLocalTransformsBeforeDeletion(
 		return;
 	}
 
-	using ConstraintPtr = TObjectPtr<UTickableConstraint>;
+	using ConstraintPtr = TWeakObjectPtr<UTickableConstraint>;
 	const TArray<ConstraintPtr> Constraints = GetHandleTransformConstraints(InWorld);
 	const TArray<ConstraintPtr> ConstraintsMinusThis = Constraints.FilterByPredicate(
 		[this](const ConstraintPtr& InConstraint)
 		{
-			return InConstraint != Constraint;
+			return InConstraint.Get() != Constraint;
 		});
 
 	// find last active constraint in the list that is different than the on we want to compensate for
@@ -378,9 +384,12 @@ void FCompensationEvaluator::ComputeLocalTransformsBeforeDeletion(
 		InSequencer->GetEvaluationTemplate().EvaluateSynchronousBlocking(Context, *InSequencer);
 
 		// evaluate constraints
-		for (const UTickableConstraint* InConstraint : AllConstraints)
+		for (const TWeakObjectPtr<UTickableConstraint>& InConstraint : AllConstraints)
 		{
-			InConstraint->Evaluate(true);
+			if (InConstraint.IsValid())
+			{
+				InConstraint->Evaluate(true);
+			}
 		}
 		
 		for (IMovieSceneToolsAnimationBakeHelper* BakeHelper : BakeHelpers)
@@ -424,7 +433,7 @@ void FCompensationEvaluator::ComputeLocalTransformsBeforeDeletion(
 
 void FCompensationEvaluator::ComputeCompensation(UWorld* InWorld, const TSharedPtr<ISequencer>& InSequencer, const FFrameNumber& InTime)
 {
-	using ConstraintPtr = TObjectPtr<UTickableConstraint>;
+	using ConstraintPtr = TWeakObjectPtr<UTickableConstraint>;
 	const TArray<ConstraintPtr> Constraints = GetHandleTransformConstraints(InWorld);
 	if (Constraints.IsEmpty())
 	{
@@ -477,9 +486,12 @@ void FCompensationEvaluator::ComputeCompensation(UWorld* InWorld, const TSharedP
 		}
 		InSequencer->GetEvaluationTemplate().EvaluateSynchronousBlocking(Context0, *InSequencer);
 
-		for (const UTickableConstraint* InConstraint : AllConstraints)
+		for (const TWeakObjectPtr<UTickableConstraint>& InConstraint : AllConstraints)
 		{
-			InConstraint->Evaluate(true);
+			if (InConstraint.IsValid())
+			{
+				InConstraint->Evaluate(true);
+			}
 		}
 
 		for (IMovieSceneToolsAnimationBakeHelper* BakeHelper : BakeHelpers)
@@ -543,7 +555,7 @@ void FCompensationEvaluator::CacheTransforms(UWorld* InWorld, const TSharedPtr<I
 		return;
 	}
 
-	using ConstraintPtr = TObjectPtr<UTickableConstraint>;
+	using ConstraintPtr = TWeakObjectPtr<UTickableConstraint>;
 
 	// get all constraints for evaluation
 	const FConstraintsManagerController& Controller = FConstraintsManagerController::Get(InWorld);
@@ -579,9 +591,12 @@ void FCompensationEvaluator::CacheTransforms(UWorld* InWorld, const TSharedPtr<I
 		InSequencer->GetEvaluationTemplate().EvaluateSynchronousBlocking(Context, *InSequencer);
 
 		// evaluate constraints
-		for (const UTickableConstraint* InConstraint : AllConstraints)
+		for (const TWeakObjectPtr<UTickableConstraint>& InConstraint : AllConstraints)
 		{
-			InConstraint->Evaluate(true);
+			if (InConstraint.IsValid())
+			{
+				InConstraint->Evaluate(true);
+			}
 		}
 		
 		for (IMovieSceneToolsAnimationBakeHelper* BakeHelper : BakeHelpers)
@@ -625,16 +640,19 @@ void FCompensationEvaluator::ComputeCurrentTransforms(UWorld* InWorld)
 {
 	ChildLocals = ChildGlobals = SpaceGlobals = {FTransform::Identity};
 
-	using ConstraintPtr = TObjectPtr<UTickableConstraint>;
+	using ConstraintPtr = TWeakObjectPtr<UTickableConstraint>;
 	const TArray< ConstraintPtr > Constraints = GetHandleTransformConstraints(InWorld);
 	if (Constraints.IsEmpty())
 	{
 		return;
 	}
 
-	for (const UTickableConstraint* InConstraint : Constraints)
+	for (const TWeakObjectPtr<UTickableConstraint>& InConstraint : Constraints)
 	{
-		InConstraint->Evaluate();
+		if (InConstraint.IsValid())
+		{
+			InConstraint->Evaluate();
+		}
 	}
 
 	ChildLocals[0] = Handle->GetLocalTransform();
@@ -661,9 +679,9 @@ void FCompensationEvaluator::ComputeCurrentTransforms(UWorld* InWorld)
 	}
 }
 
-TArray< TObjectPtr<UTickableConstraint> > FCompensationEvaluator::GetHandleTransformConstraints(UWorld* InWorld) const
+const TArray< TWeakObjectPtr<UTickableConstraint> > FCompensationEvaluator::GetHandleTransformConstraints(UWorld* InWorld) const
 {
-	using ConstraintPtr = TObjectPtr<UTickableConstraint>;
+	using ConstraintPtr = TWeakObjectPtr<UTickableConstraint>;
 	
 	if (Handle)
 	{
@@ -673,7 +691,7 @@ TArray< TObjectPtr<UTickableConstraint> > FCompensationEvaluator::GetHandleTrans
 		const TArray< ConstraintPtr > Constraints = Controller.GetParentConstraints(Handle->GetHash(), bSorted);
 		return Constraints.FilterByPredicate([](const ConstraintPtr& InConstraint)
 		{
-			return IsValid(InConstraint) && InConstraint.IsA<UTickableTransformConstraint>();  
+			return IsValid(InConstraint.Get()) && InConstraint.Get()->IsA<UTickableTransformConstraint>();  
 		});
 	}
 
@@ -883,12 +901,12 @@ bool FMovieSceneConstraintChannelHelper::SmartConstraintKey(
 		Time = FrameTime.GetFrame();
 	}
 
-	const bool bSucceeded = Interface->SmartConstraintKey(InConstraint, InOptActive, Time, InSequencer);
-
-	//todo need to revisit this to see if we need to create this even if we don't set a key, it's harmless I think either way
+	//create bindings before smart keying so added to spawn copies
 	CreateBindingIDForHandle(InSequencer, InConstraint->ChildTRSHandle);
 	CreateBindingIDForHandle(InSequencer, InConstraint->ParentTRSHandle);
 
+	const bool bSucceeded = Interface->SmartConstraintKey(InConstraint, InOptActive, Time, InSequencer);
+
 	return bSucceeded;
 }
 
@@ -956,17 +974,17 @@ void FMovieSceneConstraintChannelHelper::CompensateIfNeeded(
 	Algo::CopyIf(ConstraintChannels, TransformConstraintsChannels,
 		[InChildHash](const FConstraintAndActiveChannel& InChannel)
 		{
-			if (!InChannel.Constraint.IsValid())
+			if (!InChannel.GetConstraint().Get())
 			{
 				return false;
 			}
 
-			if ((InChildHash != INDEX_NONE) && (InChannel.Constraint->GetTargetHash() != InChildHash))
+			if ((InChildHash != INDEX_NONE) && (InChannel.GetConstraint()->GetTargetHash() != InChildHash))
 			{
 				return false;
 			}
 
-			const UTickableTransformConstraint* Constraint = Cast<UTickableTransformConstraint>(InChannel.Constraint.Get());
+			const UTickableTransformConstraint* Constraint = Cast<UTickableTransformConstraint>(InChannel.GetConstraint().Get());
 			return Constraint && (Constraint->GetTargetHash() == InChildHash) && Constraint->NeedsCompensation();
 		}
 	);
@@ -990,7 +1008,7 @@ void FMovieSceneConstraintChannelHelper::CompensateIfNeeded(
 
 			if (CurrentValue != PreviousValue) //if they are the same no need to do anything
 			{
-				UTickableTransformConstraint* Constraint = Cast<UTickableTransformConstraint>(Channel.Constraint.Get());
+				UTickableTransformConstraint* Constraint = Cast<UTickableTransformConstraint>(Channel.GetConstraint().Get());
 
 				// is the child already in that array?
 				int32 DataIndex = ToCompensate.IndexOfByPredicate([Constraint](const CompensationData& InData)
@@ -1077,7 +1095,7 @@ FConstraintSections FMovieSceneConstraintChannelHelper::GetConstraintSectionAndC
 		return ReturnValue;
 	}
 
-	ReturnValue.ActiveChannel = ConstrainedSection->GetConstraintChannel(InConstraint->GetFName());
+	ReturnValue.ActiveChannel = ConstrainedSection->GetConstraintChannel(InConstraint->ConstraintID);
 	return ReturnValue;
 }
 
diff --git a/Engine/Source/Editor/MovieSceneTools/Private/Constraints/TransformConstraintChannelInterface.cpp b/Engine/Source/Editor/MovieSceneTools/Private/Constraints/TransformConstraintChannelInterface.cpp
--- a/Engine/Source/Editor/MovieSceneTools/Private/Constraints/TransformConstraintChannelInterface.cpp
+++ b/Engine/Source/Editor/MovieSceneTools/Private/Constraints/TransformConstraintChannelInterface.cpp
@@ -8,35 +8,54 @@
 // {
 // 	return Cast<IMovieSceneConstrainedSection>(GetHandleSection(InHandle, InSequencer));
 // }
-
-bool ITransformConstraintChannelInterface::CanAddKey(const FMovieSceneConstraintChannel& ActiveChannel, const FFrameNumber& InTime, bool& ActiveValue)
+//we changed this so that it will always return true
+bool ITransformConstraintChannelInterface::CanAddKey(FMovieSceneConstraintChannel& ActiveChannel, const FFrameNumber& InTime, bool& ActiveValue)
 {
-	const TMovieSceneChannelData<const bool> ChannelData = ActiveChannel.GetData();
+	TMovieSceneChannelData<bool> ChannelData = ActiveChannel.GetData();
 	const TArrayView<const FFrameNumber> Times = ChannelData.GetTimes();
 	if (Times.IsEmpty())
 	{
 		ActiveValue = true;
 		return true;
 	}
-
-	const TArrayView<const bool> Values = ChannelData.GetValues();
-	if (InTime < Times[0])
+	ActiveChannel.Evaluate(InTime, ActiveValue);
+	ActiveValue = !ActiveValue;
+	const int32 Index = Times.Find(InTime);
+	const int32 NextTimeIndex = Algo::UpperBound(Times, InTime);
+	//have key at that time so need to delete it
+	auto DeleteNextKey = [&ActiveChannel, ActiveValue, &Times, Index,NextTimeIndex]()
+	{
+		TMovieSceneChannelData<bool> ChannelData = ActiveChannel.GetData();
+		//if has same value then delete the next
+		if (NextTimeIndex != Index && Times.IsValidIndex(NextTimeIndex))
+		{
+			bool NextValue = false;
+			ActiveChannel.Evaluate(Times[NextTimeIndex], NextValue);
+			if (NextValue == ActiveValue)
+			{
+				//same value so delete
+				ChannelData.RemoveKey(NextTimeIndex);
+			}
+		}
+	};
+	if (Index != INDEX_NONE)
 	{
-		if (!Values[0])
+		DeleteNextKey();
+		//delete key will add it later
+		ChannelData.RemoveKey(Index);
+		//if last key then we won't add any more keys
+		if (Index == (Times.Num() - 1))
 		{
-			ActiveValue = true;
-			return true;
+			return false;
 		}
-		return false;
 	}
-	
-	if (InTime > Times.Last())
+	//if has same value then delete the next
+	else if (Times.IsValidIndex(NextTimeIndex))
 	{
-		ActiveValue = !Values.Last();
-		return true;
+		DeleteNextKey();
 	}
-
-	return false;
+	
+	return true;
 }
 
 FConstraintChannelInterfaceRegistry& FConstraintChannelInterfaceRegistry::Get()
diff --git a/Engine/Source/Editor/MovieSceneTools/Private/TrackEditors/TransformTrackEditor.cpp b/Engine/Source/Editor/MovieSceneTools/Private/TrackEditors/TransformTrackEditor.cpp
--- a/Engine/Source/Editor/MovieSceneTools/Private/TrackEditors/TransformTrackEditor.cpp
+++ b/Engine/Source/Editor/MovieSceneTools/Private/TrackEditors/TransformTrackEditor.cpp
@@ -942,7 +942,7 @@ void F3DTransformTrackEditor::AddTransformKeys( UObject* ObjectToKey, const TOpt
 //todo move to external function so it can also be used by sdk/python
 static void UpdateTransformBasedOnConstraint(FTransform& CurrentTransform, USceneComponent* SceneComponent)
 {
-	TArray< TObjectPtr<UTickableConstraint> > Constraints;
+	TArray< TWeakObjectPtr<UTickableConstraint> > Constraints;
 	AActor* ShapeActor = SceneComponent->GetTypedOuter<AActor>();
 	
 	if (ShapeActor)
@@ -1558,7 +1558,7 @@ static UTickableTransformConstraint* GetTickableTransformConstraint(IMovieSceneC
 
 	if (Index != INDEX_NONE)
 	{
-		Constraint = Cast<UTickableTransformConstraint>(ConstraintChannels[Index].Constraint.Get());
+		Constraint = Cast<UTickableTransformConstraint>(ConstraintChannels[Index].GetConstraint().Get());
 	}
 	return Constraint;
 }
@@ -1636,17 +1636,17 @@ void F3DTransformTrackEditor::HandleConstraintRemoved(IMovieSceneConstrainedSect
 									return;
 								}
 
-								const FConstraintAndActiveChannel* ConstraintChannel = InSection->GetConstraintChannel(Constraint->GetFName());
-								if (!ConstraintChannel || ConstraintChannel->Constraint != Constraint)
+								const FConstraintAndActiveChannel* ConstraintChannel = InSection->GetConstraintChannel(Constraint->ConstraintID);
+								if (!ConstraintChannel || ConstraintChannel->GetConstraint().Get() != Constraint)
 								{
 									return;
 								}
 
 								const bool bCompensate = (InNotifyType == EConstraintsManagerNotifyType::ConstraintRemovedWithCompensation);
-								if (bCompensate && ConstraintChannel->Constraint.IsValid())
+								if (bCompensate && ConstraintChannel->GetConstraint().Get())
 								{
 									FMovieSceneConstraintChannelHelper::HandleConstraintRemoved(
-										ConstraintChannel->Constraint.Get(),
+										ConstraintChannel->GetConstraint().Get(),
 										&ConstraintChannel->ActiveChannel,
 										GetSequencer(),
 										Section);
@@ -1697,7 +1697,7 @@ void F3DTransformTrackEditor::HandleConstraintPropertyChanged(UTickableTransform
 	const TArray<FConstraintAndActiveChannel>& ConstraintChannels = ConstraintSection->GetConstraintsChannels();
 	const FConstraintAndActiveChannel* Channel = ConstraintChannels.FindByPredicate([InConstraint](const FConstraintAndActiveChannel& Channel)
 	{
-		return Channel.Constraint == InConstraint || Channel.ConstraintCopyToSpawn == InConstraint;
+		return Channel.GetConstraint() == InConstraint;
 	});
 
 	if (!Channel)
diff --git a/Engine/Source/Editor/MovieSceneTools/Public/Constraints/MovieSceneConstraintChannelHelper.h b/Engine/Source/Editor/MovieSceneTools/Public/Constraints/MovieSceneConstraintChannelHelper.h
--- a/Engine/Source/Editor/MovieSceneTools/Public/Constraints/MovieSceneConstraintChannelHelper.h
+++ b/Engine/Source/Editor/MovieSceneTools/Public/Constraints/MovieSceneConstraintChannelHelper.h
@@ -39,7 +39,7 @@ public:
 	
 private:
 
-	TArray< TObjectPtr<UTickableConstraint> > GetHandleTransformConstraints(UWorld* InWorld) const;
+	const TArray< TWeakObjectPtr<UTickableConstraint> > GetHandleTransformConstraints(UWorld* InWorld) const;
 
 	UTickableTransformConstraint* Constraint = nullptr;
 	UTransformableHandle* Handle = nullptr;
diff --git a/Engine/Source/Editor/MovieSceneTools/Public/Constraints/TransformConstraintChannelInterface.h b/Engine/Source/Editor/MovieSceneTools/Public/Constraints/TransformConstraintChannelInterface.h
--- a/Engine/Source/Editor/MovieSceneTools/Public/Constraints/TransformConstraintChannelInterface.h
+++ b/Engine/Source/Editor/MovieSceneTools/Public/Constraints/TransformConstraintChannelInterface.h
@@ -46,7 +46,7 @@ struct ITransformConstraintChannelInterface
 		const EMovieSceneTransformChannel& InChannels) = 0;
 
 protected:
-	static MOVIESCENETOOLS_API bool CanAddKey(const FMovieSceneConstraintChannel& ActiveChannel, const FFrameNumber& InTime, bool& ActiveValue);
+	static MOVIESCENETOOLS_API bool CanAddKey(FMovieSceneConstraintChannel& ActiveChannel, const FFrameNumber& InTime, bool& ActiveValue);
 };
 
 /**
diff --git a/Engine/Source/Runtime/Experimental/Animation/Constraints/Private/ConstraintSubsystem.cpp b/Engine/Source/Runtime/Experimental/Animation/Constraints/Private/ConstraintSubsystem.cpp
new file mode 100644
--- /dev/null
+++ b/Engine/Source/Runtime/Experimental/Animation/Constraints/Private/ConstraintSubsystem.cpp
@@ -0,0 +1,232 @@
+// Copyright Epic Games, Inc. All Rights Reserved.
+
+#include "ConstraintSubsystem.h"
+#include "Containers/Ticker.h"
+#include "Engine/Engine.h"
+#include "ConstraintsManager.h"
+#include "TransformConstraint.h"
+
+//needs to be static to avoid system getting deleted with dangling handles.
+FDelegateHandle UConstraintSubsystem::OnWorldInitHandle;
+FDelegateHandle UConstraintSubsystem::OnWorldCleanupHandle;
+
+UConstraintSubsystem::UConstraintSubsystem()
+{
+}
+
+void UConstraintSubsystem::Initialize(FSubsystemCollectionBase& Collection)
+{
+	Super::Initialize(Collection);
+
+	OnWorldInitHandle = FWorldDelegates::OnPreWorldInitialization.AddStatic(&UConstraintSubsystem::OnWorldInit);
+	OnWorldCleanupHandle = FWorldDelegates::OnWorldCleanup.AddStatic(&UConstraintSubsystem::OnWorldCleanup);
+	
+	SetFlags(RF_Transactional);
+}
+
+void UConstraintSubsystem::Deinitialize()
+{
+	for (int32 Index = ConstraintsInWorld.Num() - 1; Index >= 0; --Index)
+	{
+		ConstraintsInWorld[Index].RemoveConstraints(ConstraintsInWorld[Index].World);
+	}
+	ConstraintsInWorld.Reset();
+	FWorldDelegates::OnPreWorldInitialization.Remove(OnWorldInitHandle);
+	FWorldDelegates::OnWorldCleanup.Remove(OnWorldCleanupHandle);
+	Super::Deinitialize();
+}
+
+UConstraintSubsystem* UConstraintSubsystem::Get()
+{
+	return GEngine->GetEngineSubsystem<UConstraintSubsystem>();
+}
+
+const FConstraintsInWorld* UConstraintSubsystem::ConstraintsInWorldFind(UWorld* InWorld) const
+{
+	for (const FConstraintsInWorld& CInW : ConstraintsInWorld)
+	{
+		if (CInW.World.Get() == InWorld)
+		{
+			return &CInW;
+		}
+	}
+	return nullptr;
+}
+
+FConstraintsInWorld* UConstraintSubsystem::ConstraintsInWorldFind(UWorld* InWorld) 
+{
+	for (FConstraintsInWorld& CInW : ConstraintsInWorld)
+	{
+		if (CInW.World.Get() == InWorld)
+		{
+			return &CInW;
+		}
+	}
+	return nullptr;
+}
+
+FConstraintsInWorld& UConstraintSubsystem::ConstraintsInWorldFindOrAdd(UWorld* InWorld)
+{
+	for (FConstraintsInWorld& CInW : ConstraintsInWorld)
+	{
+		if (CInW.World.Get() == InWorld)
+		{
+			return CInW;
+		}
+	}
+	FConstraintsInWorld NewCInW;
+	NewCInW.World = InWorld;
+	int32 Index = ConstraintsInWorld.Add(NewCInW);
+	return ConstraintsInWorld[Index];
+}
+TArray<TWeakObjectPtr<UTickableConstraint>> UConstraintSubsystem::GetConstraints(UWorld* InWorld) const
+{
+	static const TArray< TWeakObjectPtr<UTickableConstraint> > DummyArray;
+	if (const FConstraintsInWorld* Constraints = ConstraintsInWorldFind(InWorld))
+	{
+		return (Constraints->Constraints);
+	}
+	return DummyArray;
+}
+
+const TArray<TWeakObjectPtr<UTickableConstraint>>& UConstraintSubsystem::GetConstraintsArray(UWorld* InWorld) const
+{
+	static const TArray< TWeakObjectPtr<UTickableConstraint> > DummyArray;
+	if (const FConstraintsInWorld* Constraints = ConstraintsInWorldFind(InWorld))
+	{
+		return (Constraints->Constraints);
+	}
+	return DummyArray;
+}
+
+void UConstraintSubsystem::AddConstraint(UWorld* InWorld, UTickableConstraint* InConstraint)
+{	
+	Modify();
+	FConstraintsInWorld& Constraints = ConstraintsInWorldFindOrAdd(InWorld);
+	if (Constraints.Constraints.Contains(InConstraint) == false)
+	{
+		Constraints.Constraints.Emplace(InConstraint);
+	}
+	OnConstraintAddedToSystem_BP.Broadcast(this, InConstraint);
+}
+
+void UConstraintSubsystem::RemoveConstraint(UWorld* InWorld, UTickableConstraint* InConstraint, bool bDoNoCompensate)
+{
+	Modify();
+	OnConstraintRemovedFromSystem_BP.Broadcast(this, InConstraint, bDoNoCompensate);
+
+	// disable constraint
+	InConstraint->Modify();
+	InConstraint->TeardownConstraint(InWorld);
+	InConstraint->SetActive(false);
+
+	if (FConstraintsInWorld* Constraints = ConstraintsInWorldFind(InWorld))
+	{
+		Constraints->Constraints.Remove(InConstraint);
+	}
+}
+
+// we want InFunctionToTickBefore to tick first = InFunctionToTickBefore is a prerex of InFunctionToTickAfter
+void UConstraintSubsystem::SetConstraintDependencies(
+	FConstraintTickFunction* InFunctionToTickBefore,
+	FConstraintTickFunction* InFunctionToTickAfter)
+{
+	// look for child tick function in in parent's prerequisites. 
+	const TArray<FTickPrerequisite>& ParentPrerequisites = InFunctionToTickAfter->GetPrerequisites();
+	const bool bIsChildAPrerexOfParent = ParentPrerequisites.ContainsByPredicate([InFunctionToTickBefore](const FTickPrerequisite& Prerex)
+		{
+			return Prerex.PrerequisiteTickFunction == InFunctionToTickBefore;
+		});
+
+	// child tick function is already a prerex -> parent already ticks after child
+	if (bIsChildAPrerexOfParent)
+	{
+		return;
+	}
+
+	// look for parent tick function in in child's prerequisites
+	const TArray<FTickPrerequisite>& ChildPrerequisites = InFunctionToTickBefore->GetPrerequisites();
+	const bool bIsParentAPrerexOfChild = ChildPrerequisites.ContainsByPredicate([InFunctionToTickAfter](const FTickPrerequisite& Prerex)
+		{
+			return Prerex.PrerequisiteTickFunction == InFunctionToTickAfter;
+		});
+
+	// parent tick function is a prerex of the child tick function (child ticks after parent)
+	// so remove it before setting new dependencies.
+	if (bIsParentAPrerexOfChild)
+	{
+		InFunctionToTickBefore->RemovePrerequisite(this, *InFunctionToTickAfter);
+	}
+
+	// set dependency
+	InFunctionToTickAfter->AddPrerequisite(this, *InFunctionToTickBefore);
+}
+
+bool UConstraintSubsystem::HasConstraint(UWorld* InWorld, UTickableConstraint* InConstraint)
+{
+	const TArray<TWeakObjectPtr<UTickableConstraint>>& Constraints = GetConstraintsArray(InWorld);
+	return  Constraints.Contains(InConstraint);
+}
+
+#if WITH_EDITOR
+
+void UConstraintSubsystem::PostEditUndo()
+{
+	Super::PostEditUndo();
+
+	const FConstraintsManagerController& Controller = FConstraintsManagerController::Get(GetWorld());
+	Controller.Notify(EConstraintsManagerNotifyType::ManagerUpdated, this);
+
+}
+#endif
+
+void UConstraintSubsystem::OnWorldInit(UWorld* InWorld, const UWorld::InitializationValues IVS)
+{
+	UConstraintSubsystem* System = UConstraintSubsystem::Get();
+	FConstraintsInWorld &ConstraintInWorld = System->ConstraintsInWorldFindOrAdd(InWorld);
+	ConstraintInWorld.Init(InWorld);
+}
+
+void UConstraintSubsystem::OnWorldCleanup(UWorld* InWorld, bool bSessionEnded, bool bCleanupResources)
+{
+	UConstraintSubsystem* System = UConstraintSubsystem::Get();
+
+	for (int32 Index = System->ConstraintsInWorld.Num() - 1; Index >= 0; --Index)
+	{
+		if (System->ConstraintsInWorld[Index].World.Get() == InWorld)
+		{
+			System->ConstraintsInWorld[Index].RemoveConstraints(InWorld);
+			System->ConstraintsInWorld.RemoveAt(Index);
+			break;
+		}
+	}
+}
+
+/*************************************
+*FConstraintsInWorld
+*************************************/
+
+void FConstraintsInWorld::RemoveConstraints(UWorld* InWorld)
+{
+	for (TWeakObjectPtr<UTickableConstraint>& Constraint : Constraints)
+	{
+		if (Constraint.IsValid())
+		{
+			Constraint->TeardownConstraint(InWorld);
+			Constraint->SetActive(false);
+		}
+	}
+	Constraints.SetNum(0);
+}
+
+
+void FConstraintsInWorld::Init(UWorld* InWorld)
+{
+	if (InWorld)
+	{
+		World = InWorld;
+	}
+}
+
+
+
diff --git a/Engine/Source/Runtime/Experimental/Animation/Constraints/Private/ConstraintSubsystem.h b/Engine/Source/Runtime/Experimental/Animation/Constraints/Private/ConstraintSubsystem.h
new file mode 100644
--- /dev/null
+++ b/Engine/Source/Runtime/Experimental/Animation/Constraints/Private/ConstraintSubsystem.h
@@ -0,0 +1,90 @@
+// Copyright Epic Games, Inc. All Rights Reserved.
+
+#pragma once
+
+#include "CoreMinimal.h"
+#include "UObject/Object.h"
+#include "Engine/EngineBaseTypes.h"
+#include "Subsystems/EngineSubsystem.h"
+#include "ConstraintsManager.h"
+#include "Engine/World.h"
+#include "ConstraintSubsystem.generated.h"
+
+class UTickableConstraint;
+class ULevel;
+
+USTRUCT()
+struct  FConstraintsInWorld 
+{
+public:
+	GENERATED_BODY()
+
+	UPROPERTY(transient)
+	TObjectPtr<UWorld> World;
+
+	UPROPERTY(transient)
+	mutable TArray<TWeakObjectPtr<UTickableConstraint>> Constraints;
+
+	CONSTRAINTS_API void Init(UWorld* World);
+	CONSTRAINTS_API void RemoveConstraints(UWorld* World);
+};
+
+
+UCLASS()
+class  UConstraintSubsystem: public UEngineSubsystem
+{
+	GENERATED_BODY()
+
+public:	
+	/** Dynamic blueprintable delegates for knowing when a constraints are added or deleted*/
+	DECLARE_DYNAMIC_MULTICAST_SPARSE_DELEGATE_TwoParams(FOnConstraintAddedToSystem, UConstraintSubsystem, OnConstraintAddedToSystem_BP, UConstraintSubsystem*, Mananger, UTickableConstraint*, Constraint);
+	DECLARE_DYNAMIC_MULTICAST_SPARSE_DELEGATE_ThreeParams(FOnConstraintRemovedFromSystem, UConstraintSubsystem, OnConstraintRemovedFromSystem_BP, UConstraintSubsystem*, Mananger, UTickableConstraint*, Constraint, bool, bDoNotCompensate);
+
+	UConstraintSubsystem();
+
+	static UConstraintSubsystem* Get();
+	virtual void Initialize(FSubsystemCollectionBase& Collection) override;
+	virtual void Deinitialize() override;
+
+	/* Set tick dependencies between two constraints. */
+	CONSTRAINTS_API void SetConstraintDependencies(
+		FConstraintTickFunction* InFunctionToTickBefore,
+		FConstraintTickFunction* InFunctionToTickAfter);
+
+	static void OnWorldInit(UWorld* World, const UWorld::InitializationValues IVS);
+	static void OnWorldCleanup(UWorld* World, bool bSessionEnded, bool bCleanupResources);
+
+public:
+
+	CONSTRAINTS_API TArray<TWeakObjectPtr<UTickableConstraint>> GetConstraints(UWorld* InWorld) const;
+	CONSTRAINTS_API const TArray<TWeakObjectPtr<UTickableConstraint>>& GetConstraintsArray(UWorld* InWorld) const;
+
+	CONSTRAINTS_API void AddConstraint(UWorld* InWorld, UTickableConstraint* InConstraint);
+	CONSTRAINTS_API void RemoveConstraint(UWorld* InWorld, UTickableConstraint* InConstraint, bool bDoNoCompensate);
+	CONSTRAINTS_API bool HasConstraint(UWorld* InWorld, UTickableConstraint* InConstraint);
+
+#if WITH_EDITOR
+	virtual void PostEditUndo() override;
+#endif
+public:
+	/** BP Delegate fired when constraints are added*/
+	UPROPERTY(BlueprintAssignable, Category = Constraints, meta = (DisplayName = "OnConstraintAdded"))
+	FOnConstraintAddedToSystem OnConstraintAddedToSystem_BP;
+	
+	/** BP Delegate fired when constraints are removed*/
+	UPROPERTY(BlueprintAssignable, Category = Constraints, meta = (DisplayName = "OnConstraintRemoved"))
+	FOnConstraintRemovedFromSystem OnConstraintRemovedFromSystem_BP;
+
+private:
+	UPROPERTY(transient)
+	TArray<FConstraintsInWorld> ConstraintsInWorld;
+	
+	//handles for handling world creation init and teardown
+	static FDelegateHandle OnWorldInitHandle;
+	static FDelegateHandle OnWorldCleanupHandle;
+
+	const FConstraintsInWorld* ConstraintsInWorldFind(UWorld* InWorld) const;
+	FConstraintsInWorld* ConstraintsInWorldFind(UWorld* InWorld);
+	FConstraintsInWorld& ConstraintsInWorldFindOrAdd(UWorld* InWorld);
+
+};
diff --git a/Engine/Source/Runtime/Experimental/Animation/Constraints/Private/ConstraintsActor.cpp b/Engine/Source/Runtime/Experimental/Animation/Constraints/Private/ConstraintsActor.cpp
--- a/Engine/Source/Runtime/Experimental/Animation/Constraints/Private/ConstraintsActor.cpp
+++ b/Engine/Source/Runtime/Experimental/Animation/Constraints/Private/ConstraintsActor.cpp
@@ -2,7 +2,6 @@
 
 
 #include "ConstraintsActor.h"
-
 #include "ConstraintsManager.h"
 
 #include UE_INLINE_GENERATED_CPP_BY_NAME(ConstraintsActor)
@@ -74,14 +73,15 @@ void AConstraintsActor::RegisterConstraintsTickFunctions() const
 			// remove invalid pointers
 			ConstraintsManager->Constraints.RemoveAll( [](const TObjectPtr<UTickableConstraint>& InConstraint)
 			{
-				return !IsValid(InConstraint);
+				return !IsValid(InConstraint) || !InConstraint->IsValid();
 			});
 			
-			// ensure registration
+			// ensure registration, we now do this by adding it to the controller (which registers and sets it with the
+			//subsystem
+			FConstraintsManagerController& Controller = FConstraintsManagerController::Get(GetWorld());
 			for (const TObjectPtr<UTickableConstraint>& ConstPtr: ConstraintsManager->Constraints)
 			{
-				ConstPtr->ConstraintTick.RegisterTickFunction(Level);
-				ConstPtr->ConstraintTick.Constraint = ConstPtr;
+				Controller.AddConstraint(ConstPtr);
 			}
 		}
 	}
@@ -91,12 +91,5 @@ void AConstraintsActor::RegisterConstraintsTickFunctions() const
 void AConstraintsActor::PostEditUndo()
 {
 	Super::PostEditUndo();
-
-	if (IsValid(ConstraintsManager))
-	{
-		const FConstraintsManagerController& Controller = FConstraintsManagerController::Get(GetWorld());
-		Controller.Notify(EConstraintsManagerNotifyType::ManagerUpdated, ConstraintsManager);
-		ConstraintsManager->Init(GetWorld());
-	}
 }
 #endif
diff --git a/Engine/Source/Runtime/Experimental/Animation/Constraints/Private/ConstraintsManager.cpp b/Engine/Source/Runtime/Experimental/Animation/Constraints/Private/ConstraintsManager.cpp
--- a/Engine/Source/Runtime/Experimental/Animation/Constraints/Private/ConstraintsManager.cpp
+++ b/Engine/Source/Runtime/Experimental/Animation/Constraints/Private/ConstraintsManager.cpp
@@ -2,7 +2,7 @@
 
 #include "ConstraintsManager.inl"
 #include "ConstraintsActor.h"
-
+#include "ConstraintSubsystem.h"
 #include "Algo/Copy.h"
 #include "Algo/StableSort.h"
 
@@ -27,6 +27,15 @@ FConstraintTickFunction::FConstraintTickFunction()
 FConstraintTickFunction::~FConstraintTickFunction()
 {}
 
+FConstraintTickFunction::FConstraintTickFunction(const FConstraintTickFunction& In) : FConstraintTickFunction()
+{
+	if (In.ConstraintFunctions.Num() > 0) //should only be 1
+	{
+		ConstraintFunctions.Add(In.ConstraintFunctions[0]);
+	}
+	Constraint = In.Constraint;
+}
+
 void FConstraintTickFunction::ExecuteTick(
 	float DeltaTime,
 	ELevelTick TickType,
@@ -38,7 +47,10 @@ void FConstraintTickFunction::ExecuteTick(
 
 void FConstraintTickFunction::RegisterFunction(ConstraintFunction InConstraint)
 {
-	ConstraintFunctions.Add(InConstraint);
+	if (ConstraintFunctions.IsEmpty())
+	{
+		ConstraintFunctions.Add(InConstraint);
+	}
 }
 
 void FConstraintTickFunction::EvaluateFunctions() const
@@ -67,10 +79,28 @@ FString FConstraintTickFunction::DiagnosticMessage()
  * UTickableConstraint
  **/
 
+FConstraintTickFunction& UTickableConstraint::GetTickFunction(UWorld* InWorld) 
+{
+	FConstraintTickFunction& ConstraintTick = ConstraintTicks.FindOrAdd(InWorld->GetCurrentLevel());
+	return ConstraintTick;
+}
+
+const FConstraintTickFunction& UTickableConstraint::GetTickFunction(UWorld* InWorld) const
+{
+	const FConstraintTickFunction& ConstraintTick = ConstraintTicks.FindOrAdd(InWorld->GetCurrentLevel());
+	return ConstraintTick;
+}
+
 void UTickableConstraint::SetActive(const bool bIsActive)
 {
 	Active = bIsActive;
-	ConstraintTick.SetTickFunctionEnable(bIsActive);
+	for (TPair<TWeakObjectPtr<ULevel>, FConstraintTickFunction>& Pair : ConstraintTicks)
+	{
+		if (Pair.Key.IsValid())
+		{
+			Pair.Value.SetTickFunctionEnable(bIsActive);
+		}
+	}
 }
 
 bool UTickableConstraint::IsFullyActive() const
@@ -80,7 +110,13 @@ bool UTickableConstraint::IsFullyActive() const
 
 void UTickableConstraint::Evaluate(bool bTickHandlesAlso) const
 {
-	ConstraintTick.EvaluateFunctions();
+	for (const TPair<TWeakObjectPtr<ULevel>, FConstraintTickFunction>& Pair : ConstraintTicks)
+	{
+		if (Pair.Key.IsValid())
+		{
+			Pair.Value.EvaluateFunctions();
+		}
+	}
 }
 
 UTickableConstraint* UTickableConstraint::Duplicate(UObject* NewOuter) const
@@ -105,6 +141,32 @@ FString UTickableConstraint::GetTypeLabel() const
 	return GetLabel();
 }
 
+void UTickableConstraint::PostDuplicate(bool bDuplicateForPIE)
+{
+	Super::PostDuplicate(bDuplicateForPIE);
+	if (bDuplicateForPIE == false) //doing actually copy then give unique id
+	{
+		ConstraintID = FGuid::NewGuid();
+	}
+}
+void UTickableConstraint::PostLoad()
+{
+	Super::PostLoad();
+	if (ConstraintID.IsValid() == false)
+	{
+		ConstraintID = FGuid::NewGuid();
+	}
+}
+
+void UTickableConstraint::PostInitProperties()
+{
+	Super::PostInitProperties();
+	if (HasAnyFlags(RF_ClassDefaultObject | RF_NeedLoad) == false)
+	{
+		ConstraintID = FGuid::NewGuid();
+	}
+}
+
 void UTickableConstraint::PostEditChangeProperty(FPropertyChangedEvent& PropertyChangedEvent)
 {
 	Super::PostEditChangeProperty(PropertyChangedEvent);
@@ -112,10 +174,16 @@ void UTickableConstraint::PostEditChangeProperty(FPropertyChangedEvent& Property
 	const FName PropertyName = PropertyChangedEvent.GetPropertyName();
 	if (PropertyName == GET_MEMBER_NAME_CHECKED(UTickableConstraint, Active))
 	{
-		ConstraintTick.SetTickFunctionEnable(Active);
-		if (Active)
+		for (TPair<TWeakObjectPtr<ULevel>, FConstraintTickFunction>& Pair : ConstraintTicks)
 		{
-			Evaluate();
+			if (Pair.Key.IsValid())
+			{
+				Pair.Value.SetTickFunctionEnable(Active);
+				if (Active)
+				{
+					Pair.Value.EvaluateFunctions();
+				}
+			}	
 		}
 	}
 }
@@ -125,9 +193,15 @@ void UTickableConstraint::PostEditUndo()
 	Super::PostEditUndo();
 
 	// make sure we update ticking when undone
-	const bool bActiveTick = IsValid(this);
-	ConstraintTick.SetTickFunctionEnable(bActiveTick);
-	ConstraintTick.bCanEverTick = bActiveTick;
+	const bool bActiveTick = ::IsValid(this);
+	for (TPair<TWeakObjectPtr<ULevel>, FConstraintTickFunction>& Pair : ConstraintTicks)
+	{
+		if (Pair.Key.IsValid())
+		{
+			Pair.Value.SetTickFunctionEnable(bActiveTick);
+			Pair.Value.bCanEverTick = bActiveTick;
+		}
+	}
 }
 
 #endif
@@ -151,22 +225,21 @@ void UConstraintsManager::PostLoad()
 	{
 		if (ConstPtr)
 		{
-			ConstPtr->ConstraintTick.Constraint = ConstPtr;
+			ConstPtr->InitConstraint(GetWorld());
 		}
 	}
 }
 
 void UConstraintsManager::OnActorDestroyed(AActor* InActor)
 {
+
 	if (USceneComponent* SceneComponent = InActor->GetRootComponent())
 	{
-
 		TArray< int32 > IndicesToRemove;
 		for (int32 Index = 0; Index < Constraints.Num(); ++Index)
 		{
 			//if the constraint has a bound object(in sequencer) we don't remove the constraint, it could be a spawnable
-			if (IsValid(Constraints[Index]) && Constraints[Index]->HasBoundObjects() == false && Constraints[Index]->ReferencesObject(SceneComponent))
-
+			if (Constraints[Index] && IsValid(Constraints[Index].Get()) && Constraints[Index]->HasBoundObjects() == false && Constraints[Index]->ReferencesObject(SceneComponent))
 			{
 				IndicesToRemove.Add(Index);
 			}
@@ -259,74 +332,10 @@ UConstraintsManager* UConstraintsManager::Find(const UWorld* InWorld)
 	const AConstraintsActor* ConstraintsActor = FindFirstConstraintsActor(Level);
 	return ConstraintsActor ? ConstraintsActor->ConstraintsManager : nullptr;
 }
-
-// we want InFunctionToTickBefore to tick first = InFunctionToTickBefore is a prerex of InFunctionToTickAfter
-void UConstraintsManager::SetConstraintDependencies(
-	FConstraintTickFunction* InFunctionToTickBefore,
-	FConstraintTickFunction* InFunctionToTickAfter)
-{
-	// look for child tick function in in parent's prerequisites. 
-	const TArray<FTickPrerequisite>& ParentPrerequisites = InFunctionToTickAfter->GetPrerequisites();
-	const bool bIsChildAPrerexOfParent = ParentPrerequisites.ContainsByPredicate([InFunctionToTickBefore](const FTickPrerequisite& Prerex)
-	{
-		return Prerex.PrerequisiteTickFunction == InFunctionToTickBefore;
-	});
-	
-	// child tick function is already a prerex -> parent already ticks after child
-	if (bIsChildAPrerexOfParent)
-	{
-		return;
-	}
-
-	// look for parent tick function in in child's prerequisites
-	const TArray<FTickPrerequisite>& ChildPrerequisites = InFunctionToTickBefore->GetPrerequisites();
-	const bool bIsParentAPrerexOfChild = ChildPrerequisites.ContainsByPredicate([InFunctionToTickAfter](const FTickPrerequisite& Prerex)
-	{
-		return Prerex.PrerequisiteTickFunction == InFunctionToTickAfter;
-	});
-	
-	// parent tick function is a prerex of the child tick function (child ticks after parent)
-	// so remove it before setting new dependencies.
-	if (bIsParentAPrerexOfChild)
-	{
-		InFunctionToTickBefore->RemovePrerequisite(this, *InFunctionToTickAfter);
-	}
-	
-	// set dependency
-	InFunctionToTickAfter->AddPrerequisite(this, *InFunctionToTickBefore);
-}
-
 void UConstraintsManager::Clear(UWorld* InWorld)
 {
 	UnregisterDelegates(InWorld);
-
-	if (!Constraints.IsEmpty())
-	{
-		static constexpr bool bDoNoCompensate = false;
-		static constexpr EConstraintsManagerNotifyType RemovalNotification =
-						EConstraintsManagerNotifyType::ConstraintRemovedWithCompensation;
-	
-		const FConstraintsManagerController& Controller = FConstraintsManagerController::Get(InWorld);
-		for (UTickableConstraint* Constraint : Constraints)
-		{
-			if ( IsValid(Constraint) )
-			{
-				// notify deletion
-				Controller.Notify(RemovalNotification, Constraint);
-				OnConstraintRemoved_BP.Broadcast(this, Constraint, bDoNoCompensate);
-
-				// disable constraint
-				Constraint->Modify();
-				Constraint->ConstraintTick.UnRegisterTickFunction();
-				Constraint->SetActive(false);
-			}
-		}
-
-		Modify();
-		Constraints.Empty();
-	}
 }
-
 void UConstraintsManager::Dump() const
 {
 	UE_LOG(LogTemp, Error, TEXT("nb consts = %d"), Constraints.Num());
@@ -340,6 +349,7 @@ void UConstraintsManager::Dump() const
 	}
 }
 
+
 /**
  * FConstraintsManagerController
  **/
@@ -354,38 +364,22 @@ FConstraintsManagerController& FConstraintsManagerController::Get(UWorld* InWorl
 	return Singleton;
 }
 
-bool FConstraintsManagerController::DoesExistInAnyWorld(UTickableConstraint* InConstraint)
+
+UConstraintsManager* FConstraintsManagerController::GetManager() 
 {
-	bool bFound = false;
-	if (InConstraint)
+	UConstraintsManager* Manager = FindManager();
+	if (Manager)
 	{
-		UConstraintsManager* Outer = InConstraint->GetTypedOuter<UConstraintsManager>();
-		if (Outer)
-		{
-			UWorld* CurrentWorld = World; //save current World
-			for (const FWorldContext& Context : GEngine->GetWorldContexts())
-			{
-				World = Context.World();
-				UConstraintsManager* Manager = FindManager();
-				if (Manager)
-				{
-					if (Manager->Constraints.Find(InConstraint) != INDEX_NONE)
-					{
-						bFound = true;
-						break;
-					}
-				}
-			}
-			World = CurrentWorld; //restore current World
-		}
+		return Manager;
 	}
-	return bFound;
-}
 
+	// create a new ConstraintsActor
+	AConstraintsActor* ConstraintsActor = World->SpawnActor<AConstraintsActor>();
+#if WITH_EDITOR
+	ConstraintsActor->SetActorLabel(TEXT("Constraints Manager"));
+#endif // WITH_EDITOR
+	return ConstraintsActor->ConstraintsManager;
 
-UConstraintsManager* FConstraintsManagerController::GetManager() const
-{
-	return UConstraintsManager::Get(World);
 }
 
 UConstraintsManager* FConstraintsManagerController::FindManager() const
@@ -436,30 +430,41 @@ void FConstraintsManagerController::DestroyManager() const
 	}
 }
 
+void FConstraintsManagerController::StaticConstraintCreated(UWorld* InWorld, UTickableConstraint* InConstraint)
+{
+	if (UConstraintsManager* Manager = GetManager()) //created If needed
+	{
+		if (Manager->Constraints.Contains(InConstraint) == false)
+		{
+			Manager->Modify();
+			Manager->Constraints.Add(InConstraint);
+			InConstraint->Rename(nullptr, Manager, REN_ForceNoResetLoaders | REN_DontCreateRedirectors);
+			Manager->OnConstraintAdded_BP.Broadcast(Manager, InConstraint);
+		}
+	}
+}
+
 bool FConstraintsManagerController::AddConstraint(UTickableConstraint* InConstraint) const
 {
 	if (!InConstraint)
 	{
 		return false;
 	}
-	
-	UConstraintsManager* Manager = GetManager(); //this will allocate if doesn't exist
-	if (!Manager)
+	UConstraintSubsystem* Subsystem = UConstraintSubsystem::Get();
+	if (!Subsystem)
 	{
 		return false;
 	}
+	if (Subsystem->HasConstraint(World, InConstraint) == true)
+	{
+		return false;
+	}
+	InConstraint->InitConstraint(World);
 
-	Manager->Modify();
-
-	Manager->Constraints.Emplace(InConstraint);
-
-	InConstraint->ConstraintTick.RegisterFunction(InConstraint->GetFunction());
-	InConstraint->ConstraintTick.RegisterTickFunction(World->GetCurrentLevel());
-	InConstraint->ConstraintTick.Constraint = InConstraint;
+	Subsystem->AddConstraint(World,InConstraint);
 
 	// notify
 	Notify(EConstraintsManagerNotifyType::ConstraintAdded, InConstraint);
-	Manager->OnConstraintAdded_BP.Broadcast(Manager, InConstraint);
 
 	return true;
 }
@@ -470,13 +475,21 @@ UTickableConstraint* FConstraintsManagerController::AddConstraintFromCopy(UTicka
 	{
 		return nullptr;
 	}
-	UConstraintsManager* Manager = GetManager(); //this will allocate if doesn't exist
-	if (!Manager)
+	UConstraintSubsystem* Subsystem = UConstraintSubsystem::Get();
+	if (!Subsystem)
 	{
 		return nullptr;
 	}
+	int32 Index = GetConstraintIndex(CopyOfConstraint->ConstraintID);
+	if (Index != INDEX_NONE)
+	{
+		const TArray<TWeakObjectPtr<UTickableConstraint>>& Constraints = Subsystem->GetConstraintsArray(World);
+		return Constraints[Index].Get();
+	}
 
-	UTickableConstraint* OurCopy = CopyOfConstraint->Duplicate(Manager);
+	UObject* Outer = CopyOfConstraint->GetOuter();
+	UTickableConstraint* OurCopy = CopyOfConstraint->Duplicate(Outer);
+	OurCopy->ConstraintID = CopyOfConstraint->ConstraintID;
 	OurCopy->SetActive(false); //always set false
 	if (AddConstraint(OurCopy))
 	{
@@ -485,36 +498,51 @@ UTickableConstraint* FConstraintsManagerController::AddConstraintFromCopy(UTicka
 	return nullptr;
 }
 
+int32 FConstraintsManagerController::GetConstraintIndex(const FGuid& InGuid) const
+{
+	UConstraintSubsystem* Subsystem = UConstraintSubsystem::Get();
+	if (!Subsystem)
+	{
+		return INDEX_NONE;
+	}
+	TArray<TWeakObjectPtr<UTickableConstraint>> Constraints = Subsystem->GetConstraints(World);
+	return Constraints.IndexOfByPredicate([InGuid](const TWeakObjectPtr<UTickableConstraint>& Constraint)
+		{
+			return 	(Constraint.IsValid() && Constraint->ConstraintID == InGuid);
+		});
+}
 
 int32 FConstraintsManagerController::GetConstraintIndex(const FName& InConstraintName) const
 {
-	const UConstraintsManager* Manager = FindManager();
-	if (!Manager)
+	UConstraintSubsystem* Subsystem = UConstraintSubsystem::Get();
+	if (!Subsystem)
 	{
 		return INDEX_NONE;
 	}
-	
-	return Manager->Constraints.IndexOfByPredicate([InConstraintName](const TObjectPtr<UTickableConstraint>& Constraint)
+	TArray<TWeakObjectPtr<UTickableConstraint>> Constraints = Subsystem->GetConstraints(World);
+	return Constraints.IndexOfByPredicate([InConstraintName](const TWeakObjectPtr<UTickableConstraint>& Constraint)
 	{
-		return 	(Constraint && Constraint->GetFName() == InConstraintName);
+		return 	(Constraint.IsValid() && Constraint->GetFName() == InConstraintName);
 	} );
 }
+
 	
-bool FConstraintsManagerController::RemoveConstraint(UTickableConstraint* InConstraint, bool bDoNotCompensate) const
+bool FConstraintsManagerController::RemoveConstraint(UTickableConstraint* InConstraint, bool bDoNotCompensate) 
 {
 	if (FConstraintsManagerController::bDoNotRemoveConstraint == true)
 	{
 		return false;
 	}
-	const UConstraintsManager* Manager = FindManager();
-	if (!Manager)
+	UConstraintSubsystem* Subsystem = UConstraintSubsystem::Get();
+	if (!Subsystem)
 	{
 		return false;
 	}
+	TArray<TWeakObjectPtr<UTickableConstraint>> Constraints = Subsystem->GetConstraints(World);
 
-	int32 Index = Manager->Constraints.IndexOfByPredicate([InConstraint](const TObjectPtr<UTickableConstraint>& Constraint)
+	int32 Index = Constraints.IndexOfByPredicate([InConstraint](const TWeakObjectPtr<UTickableConstraint>& Constraint)
 		{
-			return 	(Constraint && Constraint == InConstraint);
+			return 	(Constraint.IsValid() && Constraint.Get() == InConstraint);
 		});
 	if (Index == INDEX_NONE)
 	{
@@ -523,32 +551,17 @@ bool FConstraintsManagerController::RemoveConstraint(UTickableConstraint* InCons
 	return RemoveConstraint(Index, bDoNotCompensate);
 }
 
-bool  FConstraintsManagerController::RemoveConstraint(const FName& InConstraintName, bool bDoNotCompensate) const
-{
-	const int32 Index = GetConstraintIndex(InConstraintName);
-	if (Index == INDEX_NONE)
-	{
-		return false;
-	}
-	
-	return RemoveConstraint(Index, bDoNotCompensate);
-}
-
-bool FConstraintsManagerController::RemoveConstraint(const int32 InConstraintIndex,bool bDoNotCompensate) const
+bool FConstraintsManagerController::RemoveConstraint(const int32 InConstraintIndex,bool bDoNotCompensate) 
 {
-	UConstraintsManager* Manager = FindManager();
-	if (!Manager)
-	{
-		return false;
-	}
-	
-	if (!Manager->Constraints.IsValidIndex(InConstraintIndex))
+	UConstraintSubsystem* Subsystem = UConstraintSubsystem::Get();
+	if (!Subsystem)
 	{
 		return false;
 	}
+	TArray<TWeakObjectPtr<UTickableConstraint>> Constraints = Subsystem->GetConstraints(World);
 
-	const FName ConstraintName = Manager->Constraints[InConstraintIndex]->GetFName();
-	UTickableConstraint* Constraint = Manager->Constraints[InConstraintIndex];
+	const FName ConstraintName = Constraints[InConstraintIndex]->GetFName();
+	UTickableConstraint* Constraint = Constraints[InConstraintIndex].Get();
 	
 	// notify deletion
 	auto GetRemoveNotifyType = [bDoNotCompensate]()
@@ -560,80 +573,89 @@ bool FConstraintsManagerController::RemoveConstraint(const int32 InConstraintInd
 		return EConstraintsManagerNotifyType::ConstraintRemovedWithCompensation;
 	};
 	Notify(GetRemoveNotifyType(), Constraint);	
-	Manager->OnConstraintRemoved_BP.Broadcast(Manager, Constraint, bDoNotCompensate);
-	
-	Manager->Constraints[InConstraintIndex]->Modify();
-	Manager->Modify();
-
-	Manager->Constraints[InConstraintIndex]->SetActive(false);
-	Manager->Constraints.RemoveAt(InConstraintIndex);
+	Subsystem->Modify();
+	Subsystem->RemoveConstraint(World, Constraint, bDoNotCompensate);
 
-	// destroy constraints actor if no constraints left
-	if (Manager->Constraints.IsEmpty())
+	if (UConstraintsManager* Manager = FindManager())
 	{
-		// todo will re-evaluate this based upon spawnable work. DestroyManager();
+		if (Manager->Constraints.Contains(Constraint) == true)
+		{
+			Manager->Modify();
+			Manager->Constraints.Remove(Constraint);
+			Manager->OnConstraintRemoved_BP.Broadcast(Manager, Constraint, bDoNotCompensate);
+			// destroy constraints actor if no constraints left
+			if (Manager->Constraints.IsEmpty())
+			{
+				DestroyManager();
+			}
+		}
 	}
+
 	return true;
+	
 }
 
-UTickableConstraint* FConstraintsManagerController::GetConstraint(const FName& InConstraintName) const
+UTickableConstraint* FConstraintsManagerController::GetConstraint(const FGuid& InGuid) const
 {
-	const int32 Index = GetConstraintIndex(InConstraintName);
+	const int32 Index = GetConstraintIndex(InGuid);
 	if (Index == INDEX_NONE)
 	{
-		return nullptr;	
+		return nullptr;
 	}
-	
+
 	return GetConstraint(Index);
 }
 
 UTickableConstraint* FConstraintsManagerController::GetConstraint(const int32 InConstraintIndex) const
 {
-	UConstraintsManager* Manager = FindManager();
-	if (!Manager)
+	UConstraintSubsystem* Subsystem = UConstraintSubsystem::Get();
+	if (!Subsystem)
 	{
-		return nullptr;	
+		return nullptr;
 	}
-	
-	if (!Manager->Constraints.IsValidIndex(InConstraintIndex))
+	TArray<TWeakObjectPtr<UTickableConstraint>> Constraints = Subsystem->GetConstraints(World);
+
+	if (!Constraints.IsValidIndex(InConstraintIndex))
 	{
 		return nullptr;	
 	}
 
-	return Manager->Constraints[InConstraintIndex];
+	return Constraints[InConstraintIndex].Get();
 }
 
-TArray< TObjectPtr<UTickableConstraint> > FConstraintsManagerController::GetParentConstraints(
+TArray< TWeakObjectPtr<UTickableConstraint> > FConstraintsManagerController::GetParentConstraints(
 	const uint32 InTargetHash,
 	const bool bSorted) const
 {
-		static const TArray< TObjectPtr<UTickableConstraint> > DummyArray;
-
-	const UConstraintsManager* Manager = FindManager();
-	if (!Manager)
+	static const TArray< TWeakObjectPtr<UTickableConstraint> > DummyArray;
+	
+	if (InTargetHash == 0)
 	{
 		return DummyArray;
 	}
-	
-	if (InTargetHash == 0)
+
+	UConstraintSubsystem* Subsystem = UConstraintSubsystem::Get();
+	if (!Subsystem)
 	{
 		return DummyArray;
 	}
+	TArray<TWeakObjectPtr<UTickableConstraint>> Constraints = Subsystem->GetConstraints(World);
 
-	using ConstraintPtr = TObjectPtr<UTickableConstraint>;
-	TArray< TObjectPtr<UTickableConstraint> > FilteredConstraints =
-	Manager->Constraints.FilterByPredicate( [InTargetHash](const ConstraintPtr& Constraint)
+	using ConstraintPtr = TWeakObjectPtr<UTickableConstraint>;
+
+	TArray< TWeakObjectPtr<UTickableConstraint> > FilteredConstraints =
+	Constraints.FilterByPredicate( [InTargetHash](const ConstraintPtr& Constraint)
 	{
-		return IsValid(Constraint) && Constraint->GetTargetHash() == InTargetHash;
+		return Constraint.IsValid() && IsValid(Constraint.Get()) && Constraint->GetTargetHash() == InTargetHash;
 	} );
 	
 	if (bSorted)
 	{
 		// LHS ticks before RHS = LHS is a prerex of RHS 
-		auto TicksBefore = [](const UTickableConstraint& LHS, const UTickableConstraint& RHS)
+		auto TicksBefore = [this](const UTickableConstraint& LHS, const UTickableConstraint& RHS)
 		{
-			const TArray<FTickPrerequisite>& RHSPrerex = RHS.ConstraintTick.GetPrerequisites();
-			const FConstraintTickFunction& LHSTickFunction = LHS.ConstraintTick;
+			const TArray<FTickPrerequisite>& RHSPrerex = RHS.GetTickFunction(World).GetPrerequisites();
+			const FConstraintTickFunction& LHSTickFunction = LHS.GetTickFunction(World);
 			const bool bIsLHSAPrerexOfRHS = RHSPrerex.ContainsByPredicate([&LHSTickFunction](const FTickPrerequisite& Prerex)
 			{
 				return Prerex.PrerequisiteTickFunction == &LHSTickFunction;
@@ -654,11 +676,12 @@ void FConstraintsManagerController::SetConstraintsDependencies(
 	const FName& InNameToTickBefore,
 	const FName& InNameToTickAfter) const
 {
-	UConstraintsManager* Manager = FindManager();
-	if (!Manager)
+	UConstraintSubsystem* Subsystem = UConstraintSubsystem::Get();
+	if (!Subsystem)
 	{
 		return;
 	}
+	TArray<TWeakObjectPtr<UTickableConstraint>> Constraints = Subsystem->GetConstraints(World);
 
 	const int32 IndexBefore = GetConstraintIndex(InNameToTickBefore);
 	const int32 IndexAfter = GetConstraintIndex(InNameToTickAfter);
@@ -667,25 +690,85 @@ void FConstraintsManagerController::SetConstraintsDependencies(
 		return;
 	}
 
-	FConstraintTickFunction& FunctionToTickBefore = Manager->Constraints[IndexBefore]->ConstraintTick;
-	FConstraintTickFunction& FunctionToTickAfter = Manager->Constraints[IndexAfter]->ConstraintTick;
+	FConstraintTickFunction& FunctionToTickBefore = Constraints[IndexBefore]->GetTickFunction(World);
+	FConstraintTickFunction& FunctionToTickAfter = Constraints[IndexAfter]->GetTickFunction(World);
 
-	Manager->SetConstraintDependencies( &FunctionToTickBefore, &FunctionToTickAfter);
+	Subsystem->SetConstraintDependencies( &FunctionToTickBefore, &FunctionToTickAfter);
 }
 
-const TArray< TObjectPtr<UTickableConstraint> >& FConstraintsManagerController::GetConstraintsArray() const
+const TArray< TWeakObjectPtr<UTickableConstraint> >& FConstraintsManagerController::GetConstraintsArray() const
 {
-	UConstraintsManager* Manager = FindManager();
-	if (!Manager)
+	UConstraintSubsystem* Subsystem = UConstraintSubsystem::Get();
+	if (!Subsystem)
 	{
-		static const TArray< TObjectPtr<UTickableConstraint> > Empty;
+		static const TArray< TWeakObjectPtr<UTickableConstraint> > Empty;
 		return Empty;
 	}
+	const TArray<TWeakObjectPtr<UTickableConstraint>>& Constraints = Subsystem->GetConstraintsArray(World);
+	return Constraints;
+}
 
-	return Manager->Constraints;
+bool FConstraintsManagerController::RemoveAllConstraints(bool bDoNotCompensate) 
+{
+	if (UConstraintSubsystem* Subsystem = UConstraintSubsystem::Get())
+	{
+		Subsystem->Modify();
+		TArray<TWeakObjectPtr<UTickableConstraint>> Constraints = Subsystem->GetConstraints(World);
+		for (TWeakObjectPtr<UTickableConstraint>& Constraint : Constraints)
+		{
+			RemoveConstraint(Constraint.Get(), bDoNotCompensate);
+		}
+	}
+	return false;
 }
 
-TArray< TObjectPtr<UTickableConstraint> > FConstraintsManagerController::GetAllConstraints(const bool bSorted) const
+static void SortConstraints(UWorld* InWorld, TArray< TWeakObjectPtr<UTickableConstraint> >& SortedConstraints)
+{
+	using ConstraintPtr = TWeakObjectPtr<UTickableConstraint>;
+	// LHS ticks before RHS = LHS is a prerex of RHS 
+	auto TicksBefore = [InWorld](const UTickableConstraint& LHS, const UTickableConstraint& RHS)
+	{
+		if (RHS.IsValid() == false || LHS.IsValid() == false)
+		{
+			return false;
+		}
+		const TArray<FTickPrerequisite>& RHSPrerex = RHS.GetTickFunction(InWorld).GetPrerequisites();
+		const FConstraintTickFunction& LHSTickFunction = LHS.GetTickFunction(InWorld);
+		const bool bIsLHSAPrerexOfRHS = RHSPrerex.ContainsByPredicate([&LHSTickFunction](const FTickPrerequisite& Prerex)
+			{
+				return Prerex.PrerequisiteTickFunction == &LHSTickFunction;
+			});
+		return bIsLHSAPrerexOfRHS;
+	};
+
+	Algo::StableSort(SortedConstraints, [TicksBefore](const ConstraintPtr& LHS, const ConstraintPtr& RHS)
+		{
+			return TicksBefore(*LHS, *RHS);
+		});
+}
+
+static void SortConstraints(UWorld* InWorld, TArray< TObjectPtr<UTickableConstraint> > &SortedConstraints)
+{
+	using ConstraintPtr = TObjectPtr<UTickableConstraint>;
+	// LHS ticks before RHS = LHS is a prerex of RHS 
+	auto TicksBefore = [InWorld](const UTickableConstraint& LHS, const UTickableConstraint& RHS)
+	{
+		const TArray<FTickPrerequisite>& RHSPrerex = RHS.GetTickFunction(InWorld).GetPrerequisites();
+		const FConstraintTickFunction& LHSTickFunction = LHS.GetTickFunction(InWorld);
+		const bool bIsLHSAPrerexOfRHS = RHSPrerex.ContainsByPredicate([&LHSTickFunction](const FTickPrerequisite& Prerex)
+			{
+				return Prerex.PrerequisiteTickFunction == &LHSTickFunction;
+			});
+		return bIsLHSAPrerexOfRHS;
+	};
+
+	Algo::StableSort(SortedConstraints, [TicksBefore](const ConstraintPtr& LHS, const ConstraintPtr& RHS)
+		{
+			return TicksBefore(*LHS, *RHS);
+		});
+}
+
+TArray< TObjectPtr<UTickableConstraint> > FConstraintsManagerController::GetStaticConstraints(const bool bSorted) const
 {
 	UConstraintsManager* Manager = FindManager();
 	if (!Manager)
@@ -693,6 +776,7 @@ TArray< TObjectPtr<UTickableConstraint> > FConstraintsManagerController::GetAllC
 		static const TArray< TObjectPtr<UTickableConstraint> > Empty;
 		return Empty;
 	}
+	TArray<TObjectPtr<UTickableConstraint>> Constraints = Manager->Constraints;
 
 	// Remove stale constraints. Stale constraints may be caused to due unexpected unloading
 	// of constributing objects, such as a level sequence.
@@ -703,37 +787,45 @@ TArray< TObjectPtr<UTickableConstraint> > FConstraintsManagerController::GetAllC
 
 	if (!bSorted)
 	{
-		return Manager->Constraints;
+		return Constraints;
 	}
 
-	using ConstraintPtr = TObjectPtr<UTickableConstraint>;
-	TArray< TObjectPtr<UTickableConstraint> > SortedConstraints(Manager->Constraints);
-	// LHS ticks before RHS = LHS is a prerex of RHS 
-	auto TicksBefore = [](const UTickableConstraint& LHS, const UTickableConstraint& RHS)
+	TArray< TObjectPtr<UTickableConstraint> > SortedConstraints(Constraints);
+	SortConstraints(World,SortedConstraints);
+
+	return SortedConstraints;
+}
+
+TArray< TWeakObjectPtr<UTickableConstraint> > FConstraintsManagerController::GetAllConstraints(const bool bSorted) const
+{
+	UConstraintSubsystem* Subsystem = UConstraintSubsystem::Get();
+	if (!Subsystem)
 	{
-		const TArray<FTickPrerequisite>& RHSPrerex = RHS.ConstraintTick.GetPrerequisites();
-		const FConstraintTickFunction& LHSTickFunction = LHS.ConstraintTick;
-		const bool bIsLHSAPrerexOfRHS = RHSPrerex.ContainsByPredicate([&LHSTickFunction](const FTickPrerequisite& Prerex)
-		{
-			return Prerex.PrerequisiteTickFunction == &LHSTickFunction;
-		});
-		return bIsLHSAPrerexOfRHS;
-	};
-	
-	Algo::StableSort(SortedConstraints, [TicksBefore](const ConstraintPtr& LHS, const ConstraintPtr& RHS)
+		static const TArray< TWeakObjectPtr<UTickableConstraint> > Empty;
+		return Empty;
+	}
+	TArray<TWeakObjectPtr<UTickableConstraint>> Constraints = Subsystem->GetConstraints(World);
+
+	if (!bSorted)
 	{
-		return TicksBefore(*LHS, *RHS);
-	});
+		return Constraints;
+	}
+
+	TArray< TWeakObjectPtr<UTickableConstraint> > SortedConstraints(Constraints);
+	SortConstraints(World,SortedConstraints);
 	
 	return SortedConstraints;
 }
 
 void FConstraintsManagerController::EvaluateAllConstraints() const
 {
-	TArray< TObjectPtr<UTickableConstraint>>Constraints = GetAllConstraints(true);
-	for (const UTickableConstraint* InConstraint : Constraints)
+	TArray< TWeakObjectPtr<UTickableConstraint>>Constraints = GetAllConstraints(true);
+	for (const TWeakObjectPtr<UTickableConstraint>& InConstraint : Constraints)
 	{
-		InConstraint->Evaluate(true);
+		if (InConstraint.IsValid())
+		{
+			InConstraint->Evaluate(true);
+		}
 	}
 }
 
@@ -748,7 +840,7 @@ void FConstraintsManagerController::Notify(EConstraintsManagerNotifyType InNotif
 			break;
 
 		case EConstraintsManagerNotifyType::ManagerUpdated:
-			checkSlow(Cast<UConstraintsManager>(InObject) != nullptr);
+			checkSlow(Cast<UConstraintSubsystem>(InObject) != nullptr);
 			break;
 		default:
 			checkfSlow(false, TEXT("Unchecked EConstraintsManagerNotifyType!"));
@@ -756,4 +848,28 @@ void FConstraintsManagerController::Notify(EConstraintsManagerNotifyType InNotif
 	}
 
 	NotifyDelegate.Broadcast(InNotifyType, InObject);
+}
+
+bool FConstraintsManagerController::DoesExistInAnyWorld(UTickableConstraint* InConstraint)
+{
+	bool bFound = false;
+
+	if (InConstraint)
+	{
+		if (UConstraintSubsystem * Subsystem = UConstraintSubsystem::Get())
+		{
+			UWorld* CurrentWorld = World; //save current World
+			for (const FWorldContext& Context : GEngine->GetWorldContexts())
+			{
+				World = Context.World();
+				if (Subsystem->GetConstraintsArray(World).Find(InConstraint) != INDEX_NONE)
+				{
+					bFound = true;
+					break;
+				}
+			}
+			World = CurrentWorld; //restore current World
+		}
+	}
+	return bFound;
 }
\ No newline at end of file
diff --git a/Engine/Source/Runtime/Experimental/Animation/Constraints/Private/ConstraintsScripting.cpp b/Engine/Source/Runtime/Experimental/Animation/Constraints/Private/ConstraintsScripting.cpp
--- a/Engine/Source/Runtime/Experimental/Animation/Constraints/Private/ConstraintsScripting.cpp
+++ b/Engine/Source/Runtime/Experimental/Animation/Constraints/Private/ConstraintsScripting.cpp
@@ -11,37 +11,24 @@
 
 UConstraintsManager* UConstraintsScriptingLibrary::GetManager(UWorld* InWorld)
 {
-	UConstraintsManager* ConstraintsManager = UConstraintsManager::Get(InWorld);
-	//add ue_log
-	return ConstraintsManager;
+	return nullptr;
 }
 
 UTransformableComponentHandle* UConstraintsScriptingLibrary::CreateTransformableComponentHandle(
 	UWorld* InWorld, USceneComponent* InSceneComponent, const FName& InSocketName)
 {
-	UConstraintsManager* ConstraintsManager = UConstraintsManager::Get(InWorld);
-	if (ConstraintsManager)
-	{
-		UTransformableComponentHandle* Handle = FTransformConstraintUtils::CreateHandleForSceneComponent(InSceneComponent, InSocketName, ConstraintsManager);
-		return Handle;
-
-	}
-	return nullptr;
+	UTransformableComponentHandle* Handle = FTransformConstraintUtils::CreateHandleForSceneComponent(InSceneComponent, InSocketName);
+	return Handle;
 }
 
 UTransformableHandle* UConstraintsScriptingLibrary::CreateTransformableHandle(UWorld* InWorld, UObject* InObject, const FName& InAttachmentName)
 {
-	UConstraintsManager* ConstraintsManager = UConstraintsManager::Get(InWorld);
-	if (!ConstraintsManager)
-	{
-		return nullptr;
-	}
 	
 	// look for customized transform handle
 	const FTransformableRegistry& Registry = FTransformableRegistry::Get();
 	if (const FTransformableRegistry::CreateHandleFuncT CreateFunction = Registry.GetCreateFunction(InObject->GetClass()))
 	{
-		return CreateFunction(ConstraintsManager, InObject, InAttachmentName);
+		return CreateFunction(InObject, InAttachmentName);
 	}
 	
 	return nullptr;
@@ -67,8 +54,8 @@ TArray<UTickableConstraint*> UConstraintsScriptingLibrary::GetConstraintsArray(U
 {
 	TArray<UTickableConstraint*> Constraints;
 	const FConstraintsManagerController& Controller = FConstraintsManagerController::Get(InWorld);
-	const TArray< TObjectPtr<UTickableConstraint> >& ConstraintsArray = Controller.GetConstraintsArray();
-	for (const TObjectPtr<UTickableConstraint>& Constraint : ConstraintsArray)
+	const TArray< TWeakObjectPtr<UTickableConstraint> >& ConstraintsArray = Controller.GetConstraintsArray();
+	for (const TWeakObjectPtr<UTickableConstraint>& Constraint : ConstraintsArray)
 	{
 		Constraints.Add(Constraint.Get());
 	}
@@ -77,17 +64,17 @@ TArray<UTickableConstraint*> UConstraintsScriptingLibrary::GetConstraintsArray(U
 
 bool UConstraintsScriptingLibrary::RemoveConstraint(UWorld* InWorld, int32 InIndex)
 {
-	const FConstraintsManagerController& Controller = FConstraintsManagerController::Get(InWorld);
+	FConstraintsManagerController& Controller = FConstraintsManagerController::Get(InWorld);
 	return Controller.RemoveConstraint(InIndex);
 }
 
 bool UConstraintsScriptingLibrary::RemoveThisConstraint(UWorld* InWorld, UTickableConstraint* InTickableConstraint)
 {
-	const FConstraintsManagerController& Controller = FConstraintsManagerController::Get(InWorld);
-	const TArray< TObjectPtr<UTickableConstraint> >& ConstraintsArray = Controller.GetConstraintsArray();
+	FConstraintsManagerController& Controller = FConstraintsManagerController::Get(InWorld);
+	const TArray< TWeakObjectPtr<UTickableConstraint> >& ConstraintsArray = Controller.GetConstraintsArray();
 	for (int32 Index = 0; Index < ConstraintsArray.Num(); ++Index)
 	{
-		if (InTickableConstraint == ConstraintsArray[Index])
+		if (InTickableConstraint == ConstraintsArray[Index].Get())
 		{
 			return Controller.RemoveConstraint(Index);
 
diff --git a/Engine/Source/Runtime/Experimental/Animation/Constraints/Private/TransformConstraint.cpp b/Engine/Source/Runtime/Experimental/Animation/Constraints/Private/TransformConstraint.cpp
--- a/Engine/Source/Runtime/Experimental/Animation/Constraints/Private/TransformConstraint.cpp
+++ b/Engine/Source/Runtime/Experimental/Animation/Constraints/Private/TransformConstraint.cpp
@@ -7,6 +7,8 @@
 #include "GameFramework/Actor.h"
 #include "MovieSceneSection.h"
 #include "Engine/World.h"
+#include "UObject/Package.h"
+#include "Components/SkeletalMeshComponent.h"
 
 #include UE_INLINE_GENERATED_CPP_BY_NAME(TransformConstraint)
 
@@ -16,13 +18,13 @@ namespace
 struct FConstraintCycleChecker
 {
 public:
-	using ConstraintPtr = TObjectPtr<UTickableConstraint>;
+	using ConstraintPtr = TWeakObjectPtr<UTickableConstraint>;
 	using ConstraintArray = TArray<ConstraintPtr>;
 
 	/** Checks if this handle is cycle from a tick dependencies perspective. */
-	static bool IsCycling(const TObjectPtr<UTransformableHandle>& InHandle)
+	static bool IsCycling(const TWeakObjectPtr<UTransformableHandle>& InHandle)
 	{
-		if (!IsValid(InHandle))
+		if (!IsValid(InHandle.Get()))
 		{
 			return false;
 		}
@@ -40,6 +42,7 @@ public:
 			return;
 		}
 
+		//todo constraints on level sequences aren't in a world
 		UWorld* World = InConstraint->GetWorld();
 		if (!IsValid(World))
 		{
@@ -83,7 +86,7 @@ public:
 		TSet<const FTickFunction*> VisitedFunctions;
 		for (const ConstraintPtr& Constraint: CyclingConstraints)
 		{
-			if (HasPrerequisiteDependencyWith(&Constraint->ConstraintTick, &InConstraint->ConstraintTick, VisitedFunctions))
+			if (HasPrerequisiteDependencyWith(&Constraint->GetTickFunction(World), &InConstraint->GetTickFunction(World), VisitedFunctions))
 			{
 				UpdateCyclingDependency(World, Cast<UTickableTransformConstraint>(Constraint));
 			}
@@ -139,7 +142,7 @@ private:
 	static void UpdateCyclingDependency(UWorld* InWorld, UTickableTransformConstraint* InConstraintToUpdate)
 	{
 		// nothing to do if this constraint doesn't tick
-		if (!InConstraintToUpdate->ConstraintTick.IsTickFunctionEnabled())
+		if (!InConstraintToUpdate->GetTickFunction(InWorld).IsTickFunctionEnabled())
 		{
 			return;
 		}
@@ -172,7 +175,7 @@ private:
 		// check if there's any active constraint in ParentConstraints
 		const bool bHasActiveParentConstraint = ParentConstraints.ContainsByPredicate([](const ConstraintPtr& Constraint)
 		{
-			return IsValid(Constraint) && Constraint->IsFullyActive(); 
+			return Constraint.IsValid() && IsValid(Constraint.Get()) && Constraint->IsFullyActive(); 
 		});
 
 		// update the constraint prerequisites based on the result
@@ -181,12 +184,12 @@ private:
 			if (bHasActiveParentConstraint)
 			{
 				// UE_LOG(LogTemp, Warning, TEXT("REMOVE %s prerex on %s."), *TargetObject->GetName(), *InConstraintToUpdate->GetName());
-				InConstraintToUpdate->ConstraintTick.RemovePrerequisite(TargetObject, *TargetTickFunction);
+				InConstraintToUpdate->GetTickFunction(InWorld).RemovePrerequisite(TargetObject, *TargetTickFunction);
 			}
 			else
 			{
 				// UE_LOG(LogTemp, Warning, TEXT("ADD %s prerex on %s."), *TargetObject->GetName(), *InConstraintToUpdate->GetName());
-				InConstraintToUpdate->ConstraintTick.AddPrerequisite(TargetObject, *TargetTickFunction);
+				InConstraintToUpdate->GetTickFunction(InWorld).AddPrerequisite(TargetObject, *TargetTickFunction);
 			}
 		}
 	}
@@ -238,8 +241,13 @@ UTickableConstraint* UTickableTransformConstraint::Duplicate(UObject* NewOuter)
 		UTransformableHandle* HandleCopy = ParentTRSHandle->Duplicate(Dup);
 		Dup->ParentTRSHandle = HandleCopy;
 	}
-	Dup->SetupDependencies();
-	Dup->RegisterDelegates();
+	for (TPair<TWeakObjectPtr<ULevel>, FConstraintTickFunction>& Pair : ConstraintTicks)
+	{
+		if (Pair.Key.IsValid())
+		{
+			UWorld* World = Pair.Key->GetTypedOuter<UWorld>();
+		}
+	}
 	return Dup;
 }
 
@@ -316,7 +324,6 @@ void UTickableTransformConstraint::PostEditChangeProperty(FPropertyChangedEvent&
 void UTickableTransformConstraint::PostEditUndo()
 {
 	Super::PostEditUndo();
-	InitConstraint();
 }
 
 UTickableTransformConstraint::FOnConstraintChanged& UTickableTransformConstraint::GetOnConstraintChanged()
@@ -361,12 +368,29 @@ void UTickableTransformConstraint::Setup()
 	}
 	
 	ComputeOffset();
-	SetupDependencies();
-	RegisterDelegates();
+
 }
 
-void UTickableTransformConstraint::SetupDependencies()
+void UTickableTransformConstraint::SetupDependencies(UWorld* InWorld)
 {
+	//we may not be the outer for old files so move it over
+	if (ParentTRSHandle)
+	{
+		UTickableTransformConstraint* PTRS = ParentTRSHandle->GetTypedOuter<UTickableTransformConstraint>();
+		if(PTRS != this)
+		{
+			ParentTRSHandle->Rename(nullptr, this, REN_ForceNoResetLoaders | REN_DontCreateRedirectors);
+		}
+	}
+	if (ChildTRSHandle)
+	{
+		UTickableTransformConstraint* CTRS = ChildTRSHandle->GetTypedOuter<UTickableTransformConstraint>();
+		if (CTRS != this)
+		{
+			ChildTRSHandle->Rename(nullptr, this, REN_ForceNoResetLoaders | REN_DontCreateRedirectors);
+		}
+	}
+	
 	FTickFunction* ParentTickFunction = GetParentHandleTickFunction();
 	FTickFunction* ChildTickFunction = GetChildHandleTickFunction();
 	
@@ -375,24 +399,32 @@ void UTickableTransformConstraint::SetupDependencies()
 		// manage dependencies
 		// force ConstraintTickFunction to tick after InParent does.
 		// Note that this might not register anything if the parent can't tick (static meshes for instance)
+		FConstraintTickFunction& ConstraintTick = ConstraintTicks.FindOrAdd(InWorld->GetCurrentLevel());
 		ConstraintTick.AddPrerequisite(ParentTRSHandle->GetPrerequisiteObject(), *ParentTickFunction);
 	}
 	
 	// TODO also check for cycle dependencies here
 	if (ChildTickFunction)
 	{
-		// force InChild to tick after ConstraintTickFunction does.
-		// Note that this might not register anything if the child can't tick (static meshes for instance)
-		ChildTickFunction->AddPrerequisite(this, ConstraintTick);
+		USkeletalMeshComponent* ChildOwner = Cast<USkeletalMeshComponent>(ChildTRSHandle->GetTarget().Get());
+		if (ChildOwner == nullptr)
+		{
+			// force InChild to tick after ConstraintTickFunction does.
+			// Note that this might not register anything if the child can't tick (static meshes for instance)
+			// Also don't do if it's a skel mesh will cause cycle.
+			FConstraintTickFunction& ConstraintTick = ConstraintTicks.FindOrAdd(InWorld->GetCurrentLevel());
+			ChildTickFunction->AddPrerequisite(this, ConstraintTick);
+		}
 	}
 }
 
-void UTickableTransformConstraint::EnsurePrimaryDependency()
+void UTickableTransformConstraint::EnsurePrimaryDependency(UWorld* InWorld)
 {
 	const FTickFunction* ParentTickFunction = GetParentHandleTickFunction();
 	const FTickFunction* ChildTickFunction = GetChildHandleTickFunction();
 	if (ParentTickFunction && (ChildTickFunction != ParentTickFunction))
 	{
+		FConstraintTickFunction& ConstraintTick = ConstraintTicks.FindOrAdd(InWorld->GetCurrentLevel());
 		const TArray<FTickPrerequisite>& ParentPrerequisites = ConstraintTick.GetPrerequisites();
 		if (ParentPrerequisites.IsEmpty())
 		{
@@ -418,24 +450,12 @@ void UTickableTransformConstraint::OnActiveStateChanged() const
 void UTickableTransformConstraint::PostLoad()
 {
 	Super::PostLoad();
-	InitConstraint();
 }
 
 void UTickableTransformConstraint::PostDuplicate(bool bDuplicateForPIE)
 {
 	Super::PostDuplicate(bDuplicateForPIE);
-	if (ConstraintTick.ConstraintFunctions.IsEmpty())
-	{
-		ConstraintTick.RegisterFunction(GetFunction());
-	}
 
-	SetupDependencies();
-	RegisterDelegates();
-
-	if (bDuplicateForPIE)
-	{
-		EnsurePrimaryDependency();
-	}
 }
 
 uint32 UTickableTransformConstraint::GetTargetHash() const
@@ -495,7 +515,7 @@ void UTickableTransformConstraint::Evaluate(bool bTickHandlesAlso) const
 		}
 
 		Super::Evaluate();
-
+		//todo test this more may be able to remove it
 		if (ChildTRSHandle)
 		{
 			ChildTRSHandle->TickForBaking();
@@ -509,7 +529,7 @@ void UTickableTransformConstraint::Evaluate(bool bTickHandlesAlso) const
 
 void UTickableTransformConstraint::SetActive(const bool bIsActive)
 {
-	const bool bNeedsUpdate = (Active != bIsActive) || (bIsActive != ConstraintTick.IsTickFunctionEnabled());
+	const bool bNeedsUpdate = (Active != bIsActive);
 	Super::SetActive(bIsActive);
 	
 	if (bNeedsUpdate)
@@ -565,7 +585,14 @@ void UTickableTransformConstraint::OnHandleModified(UTransformableHandle* InHand
 	{
 		if (InNotification == EHandleEvent::ComponentUpdated)
 		{
-			SetupDependencies();
+			UWorld* World = nullptr;
+			if (InHandle->GetTarget().Get())
+			{
+				if (InHandle->GetTarget().Get()->GetWorld())
+				{
+					SetupDependencies(InHandle->GetTarget().Get()->GetWorld());
+				}
+			}
 			return;
 		}
 	}
@@ -573,18 +600,31 @@ void UTickableTransformConstraint::OnHandleModified(UTransformableHandle* InHand
 	// update the constraint if the parent's transform has been modified 
 	if (InHandle == ParentTRSHandle && InNotification == EHandleEvent::GlobalTransformUpdated)
 	{
-		if (IsFullyActive() && ConstraintTick.IsTickFunctionRegistered() && ConstraintTick.IsTickFunctionEnabled())
+		if (InHandle->GetTarget().Get())
 		{
-			Evaluate();
+			if (UWorld* World = InHandle->GetTarget().Get()->GetWorld())
+			{
+				FConstraintTickFunction* ConstraintTick = ConstraintTicks.Find(World->GetCurrentLevel());
+				if (IsFullyActive() && ConstraintTick->IsTickFunctionRegistered() && ConstraintTick->IsTickFunctionEnabled())
+				{
+					Evaluate();
+				}
+
+			}
 		}
 		return;
 	}
 }
 
+bool UTickableTransformConstraint::IsValid() const
+{
+	return (bValid && ::IsValid(ChildTRSHandle) && ChildTRSHandle->IsValid()
+		&& ::IsValid(ParentTRSHandle) && ParentTRSHandle->IsValid());
+}
+
 bool UTickableTransformConstraint::IsFullyActive() const
 {
-	return (Active && IsValid(ChildTRSHandle) && ChildTRSHandle->IsValid() 
-		&& IsValid(ParentTRSHandle) && ParentTRSHandle->IsValid());
+	return (Active && IsValid());
 }
 
 bool UTickableTransformConstraint::NeedsCompensation() const
@@ -605,14 +645,7 @@ FTickFunction* UTickableTransformConstraint::GetParentHandleTickFunction() const
 
 FTickFunction* UTickableTransformConstraint::GetHandleTickFunction(const TObjectPtr<UTransformableHandle>& InHandle) const 
 {
-	if (!IsValid(InHandle) || !InHandle->IsValid())
-	{
-		return nullptr;
-	}
-
-	// avoid creating dependencies between functions that are registered in levels that don't leave in the same world 
-	const UObject* PrerequisiteObject = InHandle->GetPrerequisiteObject();
-	if (!PrerequisiteObject || PrerequisiteObject->GetWorld() != GetWorld())
+	if (!::IsValid(InHandle) || !InHandle->IsValid())
 	{
 		return nullptr;
 	}
@@ -620,17 +653,27 @@ FTickFunction* UTickableTransformConstraint::GetHandleTickFunction(const TObject
 	return InHandle->GetTickFunction();
 }
 
-void UTickableTransformConstraint::InitConstraint()
+void UTickableTransformConstraint::InitConstraint(UWorld *InWorld)
 {
+	FConstraintTickFunction& ConstraintTick = ConstraintTicks.FindOrAdd(InWorld->GetCurrentLevel());
+
 	if (ConstraintTick.ConstraintFunctions.IsEmpty())
 	{
 		ConstraintTick.RegisterFunction(GetFunction());
 	}
-
-	SetupDependencies();
+	ConstraintTick.RegisterTickFunction(InWorld->GetCurrentLevel());
+	ConstraintTick.Constraint = this;
+	SetupDependencies(InWorld);
 	RegisterDelegates();
+	bValid = true;
 
-	EnsurePrimaryDependency();
+}
+
+void UTickableTransformConstraint::TeardownConstraint(UWorld* InWorld)
+{
+	FConstraintTickFunction& ConstraintTick = ConstraintTicks.FindOrAdd(InWorld->GetCurrentLevel());
+	ConstraintTick.UnRegisterTickFunction(); 
+	ConstraintTicks.Remove(InWorld->GetCurrentLevel());
 }
 
 /** 
@@ -709,6 +752,10 @@ FConstraintTickFunction::ConstraintFunction UTickableTranslationConstraint::GetF
 		Transform.SetLocation(NewTranslation);
 			
 		SetChildGlobalTransform(Transform);
+		if (ChildTRSHandle)
+		{
+			ChildTRSHandle->TickForBaking();
+		}
 	};
 }
 
@@ -830,7 +877,7 @@ FConstraintTickFunction::ConstraintFunction UTickableRotationConstraint::GetFunc
 		{
 			return;
 		}
-		
+
 		const FQuat ParentRotation = GetParentGlobalTransform().GetRotation();
 		FTransform Transform = GetChildGlobalTransform();
 		const FQuat ChildRotation = Transform.GetRotation();
@@ -846,6 +893,10 @@ FConstraintTickFunction::ConstraintFunction UTickableRotationConstraint::GetFunc
 		Transform.SetRotation(NewRotation);
 		
 		SetChildGlobalTransform(Transform);
+		if (ChildTRSHandle)
+		{
+			ChildTRSHandle->TickForBaking();
+		}
 	};
 }
 
@@ -987,6 +1038,10 @@ FConstraintTickFunction::ConstraintFunction UTickableScaleConstraint::GetFunctio
 		Transform.SetScale3D(NewScale);
 			
 		SetChildGlobalTransform(Transform);
+		if (ChildTRSHandle)
+		{
+			ChildTRSHandle->TickForBaking();
+		}
 	};
 }
 
@@ -1132,7 +1187,7 @@ FConstraintTickFunction::ConstraintFunction UTickableParentConstraint::GetFuncti
 			OutTransform.SetRotation(NewRotation);
 			OutTransform.SetScale3D(NewScale);
 		};
-
+	
 		const FTransform ParentTransform = GetParentGlobalTransform();
 		
 		FTransform TargetTransform = (!bMaintainOffset) ? ParentTransform : OffsetTransform * ParentTransform;
@@ -1147,6 +1202,10 @@ FConstraintTickFunction::ConstraintFunction UTickableParentConstraint::GetFuncti
 		}
 		
 		SetChildGlobalTransform(TargetTransform);
+		if (ChildTRSHandle)
+		{
+			ChildTRSHandle->TickForBaking();
+		}
 	};
 }
 
@@ -1295,6 +1354,10 @@ FConstraintTickFunction::ConstraintFunction UTickableLookAtConstraint::GetFuncti
 				FTransform Transform = ChildTransform;
 				Transform.SetRotation(Rotation.GetNormalized());
 				SetChildGlobalTransform(Transform);
+				if (ChildTRSHandle)
+				{
+					ChildTRSHandle->TickForBaking();
+				}
 			}
 		}
 	};
@@ -1355,13 +1418,13 @@ void UTickableLookAtConstraint::PostEditChangeProperty(FPropertyChangedEvent& Pr
 namespace
 {
 
-UTransformableHandle* GetHandle(UObject* InObject, const FName& InSocketName, UObject* Outer)
+UTransformableHandle* GetHandle(UObject* InObject, const FName& InSocketName)
 {
 	// look for customized transform handle
 	const FTransformableRegistry& Registry = FTransformableRegistry::Get();
 	if (const FTransformableRegistry::CreateHandleFuncT CreateFunction = Registry.GetCreateFunction(InObject->GetClass()))
 	{
-		return CreateFunction(Outer, InObject, InSocketName);
+		return CreateFunction(InObject, InSocketName);
 	}
 
 	return nullptr;
@@ -1383,13 +1446,12 @@ uint32 GetConstrainableHash(const UObject* InObject)
 
 UTransformableComponentHandle* FTransformConstraintUtils::CreateHandleForSceneComponent(
 	USceneComponent* InSceneComponent,
-	const FName& InSocketName,
-	UObject* Outer)
+	const FName& InSocketName)
 {
 	UTransformableComponentHandle* ComponentHandle = nullptr;
 	if (InSceneComponent)
 	{
-		ComponentHandle = NewObject<UTransformableComponentHandle>(Outer, NAME_None, RF_Transactional);
+		ComponentHandle = NewObject<UTransformableComponentHandle>(GetTransientPackage(), NAME_None, RF_Transactional);
 		ComponentHandle->Component = InSceneComponent;
 		ComponentHandle->SocketName = InSocketName;
 		InSceneComponent->SetMobility(EComponentMobility::Movable);
@@ -1401,7 +1463,7 @@ UTransformableComponentHandle* FTransformConstraintUtils::CreateHandleForSceneCo
 void FTransformConstraintUtils::GetParentConstraints(
 	UWorld* InWorld,
 	const AActor* InChild,
-	TArray< TObjectPtr<UTickableConstraint> >& OutConstraints)
+	TArray< TWeakObjectPtr<UTickableConstraint> >& OutConstraints)
 {
 	if (!InWorld || !InChild)
 	{
@@ -1473,7 +1535,7 @@ namespace
 // we suppose that both InParentHandle and InChildHandle are safe to use
 bool HasConstraintDependencyWith(UWorld* InWorld, const UTransformableHandle* InParentHandle, const UTransformableHandle* InChildHandle)
 {
-	using ConstraintPtr = TObjectPtr<UTickableConstraint>;
+	using ConstraintPtr = TWeakObjectPtr<UTickableConstraint>;
 	using HandlePtr = TObjectPtr<UTransformableHandle>;
 
 	static constexpr bool bSorted = false;
@@ -1568,20 +1630,13 @@ UTickableTransformConstraint* FTransformConstraintUtils::CreateAndAddFromObjects
 		return nullptr;
 	}
 
-	UConstraintsManager* ConstraintsManager = UConstraintsManager::Get(InWorld);
-	if (!ConstraintsManager)
-	{
-		UE_LOG(LogTemp, Error, TEXT("%s constraint manager is null."), ErrorPrefix);
-		return nullptr;
-	}
-
-	UTransformableHandle* ParentHandle = GetHandle(InParent, InParentSocketName, ConstraintsManager);
+	UTransformableHandle* ParentHandle = GetHandle(InParent, InParentSocketName);
 	if (!ParentHandle)
 	{
 		return nullptr;
 	}
 	
-	UTransformableHandle* ChildHandle = GetHandle(InChild, InChildSocketName, ConstraintsManager);
+	UTransformableHandle* ChildHandle = GetHandle(InChild, InChildSocketName);
 	if (!ChildHandle)
 	{
 		return nullptr;
@@ -1630,7 +1685,7 @@ bool FTransformConstraintUtils::AddConstraint(
 
 	// store previous child constraints
 	const FConstraintsManagerController& Controller = FConstraintsManagerController::Get(InWorld);
-	const TArray<TObjectPtr<UTickableConstraint>> ChildParentConstraints = Controller.GetParentConstraints(InChildHandle->GetHash(), true);
+	const TArray<TWeakObjectPtr<UTickableConstraint>> ChildParentConstraints = Controller.GetParentConstraints(InChildHandle->GetHash(), true);
 
 	// add the new one
 	const bool bConstraintAdded = Controller.AddConstraint(Constraint);
@@ -1639,17 +1694,11 @@ bool FTransformConstraintUtils::AddConstraint(
 		return false;
 	}
 
-	// setup the constraint
-	auto SetupConstraint = [InParentHandle, InChildHandle, bMaintainOffset](UTickableTransformConstraint* InConstraint)
-	{
-		InConstraint->ParentTRSHandle = InParentHandle;
-		InConstraint->ChildTRSHandle = InChildHandle;
-		InConstraint->bMaintainOffset = bMaintainOffset;
-		InConstraint->Setup();
-	};
-	
-	SetupConstraint(Constraint);
-
+	Constraint->ParentTRSHandle = InParentHandle;
+	Constraint->ChildTRSHandle = InChildHandle;
+	Constraint->bMaintainOffset = bMaintainOffset;
+	Constraint->Setup();
+	Constraint->InitConstraint(InWorld);
 	// add dependencies with the last child constraint
 	const FName NewConstraintName = Constraint->GetFName();
 	if (!ChildParentConstraints.IsEmpty())
@@ -1659,12 +1708,12 @@ bool FTransformConstraintUtils::AddConstraint(
 	}
 
 	// make sure we tick after the parent.
-	Constraint->EnsurePrimaryDependency();
+	Constraint->EnsurePrimaryDependency(InWorld);
 
 	// if child handle is the parent of some other constraints, ensure they will tick after that new one
-	TArray<TObjectPtr<UTickableConstraint>> ChildChildConstraints;
+	TArray<TWeakObjectPtr<UTickableConstraint>> ChildChildConstraints;
 	GetChildrenConstraints(InWorld, InChildHandle, ChildChildConstraints);
-	for (const TObjectPtr<UTickableConstraint>& ChildConstraint: ChildChildConstraints)
+	for (const TWeakObjectPtr<UTickableConstraint>& ChildConstraint: ChildChildConstraints)
 	{
 		Controller.SetConstraintsDependencies( NewConstraintName, ChildConstraint->GetFName());
 	}
@@ -1755,7 +1804,7 @@ TOptional<FTransform> FTransformConstraintUtils::GetRelativeTransform(UWorld* In
 	static constexpr bool bSorted = true;
 	const FConstraintsManagerController& Controller = FConstraintsManagerController::Get(InWorld);
 
-	const TArray< TObjectPtr<UTickableConstraint> > Constraints = Controller.GetParentConstraints(InHandleHash, bSorted);
+	const TArray< TWeakObjectPtr<UTickableConstraint> > Constraints = Controller.GetParentConstraints(InHandleHash, bSorted);
 	if (Constraints.IsEmpty())
 	{
 		return TOptional<FTransform>();
@@ -1777,7 +1826,7 @@ TOptional<FTransform> FTransformConstraintUtils::GetRelativeTransform(UWorld* In
 }
 
 TOptional<FTransform> FTransformConstraintUtils::GetConstraintsRelativeTransform(
-	const TArray< TObjectPtr<UTickableConstraint> >& InConstraints,
+	const TArray< TWeakObjectPtr<UTickableConstraint> >& InConstraints,
 	const FTransform& InChildLocal, const FTransform& InChildWorld)
 {
 	if (InConstraints.IsEmpty())
@@ -1802,7 +1851,7 @@ TOptional<FTransform> FTransformConstraintUtils::GetConstraintsRelativeTransform
 	}
 
 	// otherwise, we need to look for constraints on a sub-transform basis so we compute the relative transform for each of them
-	using ConstraintPtr = TObjectPtr<UTickableConstraint>;
+	using ConstraintPtr = TWeakObjectPtr<UTickableConstraint>;
 	auto GetLastSubTransformIndex = [InConstraints](const EMovieSceneTransformChannel& InChannel)
 	{
 		return InConstraints.FindLastByPredicate([InChannel](const ConstraintPtr& InConstraint)
@@ -1848,9 +1897,9 @@ TOptional<FTransform> FTransformConstraintUtils::GetConstraintsRelativeTransform
 	return ChildLocal;
 }
 
-int32 FTransformConstraintUtils::GetLastActiveConstraintIndex(const TArray< TObjectPtr<UTickableConstraint> >& InConstraints)
+int32 FTransformConstraintUtils::GetLastActiveConstraintIndex(const TArray< TWeakObjectPtr<UTickableConstraint> >& InConstraints)
 {
-	return InConstraints.FindLastByPredicate([](const TObjectPtr<UTickableConstraint>& InConstraint)
+	return InConstraints.FindLastByPredicate([](const TWeakObjectPtr<UTickableConstraint>& InConstraint)
 	{
 	   if (const UTickableTransformConstraint* TransformConstraint = Cast<UTickableTransformConstraint>(InConstraint.Get()))
 	   {
@@ -1863,9 +1912,9 @@ int32 FTransformConstraintUtils::GetLastActiveConstraintIndex(const TArray< TObj
 void FTransformConstraintUtils::GetChildrenConstraints(
 	UWorld* World,
 	const UTransformableHandle* InParentHandle,
-	TArray< TObjectPtr<UTickableConstraint> >& OutConstraints)
+	TArray< TWeakObjectPtr<UTickableConstraint> >& OutConstraints)
 {
-	using ConstraintPtr = TObjectPtr<UTickableConstraint>;
+	using ConstraintPtr = TWeakObjectPtr<UTickableConstraint>;
 	
 	// filter for transform constraints where the InParentHandle is the parent (based on its hash value)
 	const uint32 ParentHash = InParentHandle->GetHash();
diff --git a/Engine/Source/Runtime/Experimental/Animation/Constraints/Private/TransformableHandle.cpp b/Engine/Source/Runtime/Experimental/Animation/Constraints/Private/TransformableHandle.cpp
--- a/Engine/Source/Runtime/Experimental/Animation/Constraints/Private/TransformableHandle.cpp
+++ b/Engine/Source/Runtime/Experimental/Animation/Constraints/Private/TransformableHandle.cpp
@@ -67,7 +67,7 @@ bool UTransformableComponentHandle::IsValid() const
 }
 
 //need to tick any skelmesh component, sibling or parent
-void UTransformableComponentHandle::TickForBaking()
+void UTransformableComponentHandle::TickForBaking() const
 {
 	if (!Component.IsValid())
 	{
@@ -223,16 +223,16 @@ FTickPrerequisite LookForPrimaryPrerequisite(USceneComponent* Component)
 		static constexpr bool bSorted = true;
 		
 		const FConstraintsManagerController& Controller = FConstraintsManagerController::Get(World);
-		const TArray<TObjectPtr<UTickableConstraint>> ParentConstraints =
+		const TArray<TWeakObjectPtr<UTickableConstraint>> ParentConstraints =
 			Controller.GetParentConstraints(GetTypeHash(Component), bSorted);
 		
 		for (int Index = ParentConstraints.Num()-1; Index >= 0; Index--)
 		{
-			if (UTickableConstraint* Constraint = ParentConstraints[Index])
+			if (UTickableConstraint* Constraint = ParentConstraints[Index].Get())
 			{
-				if(IsValidTickFunction(&Constraint->ConstraintTick))
+				if(IsValidTickFunction(&Constraint->GetTickFunction(World)))
 				{
-					return FTickPrerequisite(Constraint->GetOuter(), Constraint->ConstraintTick);
+					return FTickPrerequisite(Constraint->GetOuter(), Constraint->GetTickFunction(World));
 				}
 			}
 		}	
diff --git a/Engine/Source/Runtime/Experimental/Animation/Constraints/Private/TransformableRegistry.cpp b/Engine/Source/Runtime/Experimental/Animation/Constraints/Private/TransformableRegistry.cpp
--- a/Engine/Source/Runtime/Experimental/Animation/Constraints/Private/TransformableRegistry.cpp
+++ b/Engine/Source/Runtime/Experimental/Animation/Constraints/Private/TransformableRegistry.cpp
@@ -69,11 +69,11 @@ const FTransformableRegistry::FTransformableInfo* FTransformableRegistry::FindIn
 void FTransformableRegistry::RegisterBaseObjects()
 {
     // register USceneComponent and AActor
-    auto CreateComponentHandle = [](UObject* Outer, UObject* InObject, const FName& InSocketName)->UTransformableHandle*
+    auto CreateComponentHandle = [](UObject* InObject, const FName& InSocketName)->UTransformableHandle*
     {
     	if (USceneComponent* Component = Cast<USceneComponent>(InObject))
     	{
-    		return FTransformConstraintUtils::CreateHandleForSceneComponent(Component, InSocketName, Outer);
+    		return FTransformConstraintUtils::CreateHandleForSceneComponent(Component, InSocketName);
     	}
     	return nullptr;
     };
@@ -87,11 +87,11 @@ void FTransformableRegistry::RegisterBaseObjects()
     	return 0;
     };
     
-    auto CreateComponentHandleFromActor = [CreateComponentHandle](UObject* Outer, UObject* InObject, const FName& InSocketName)->UTransformableHandle*
+    auto CreateComponentHandleFromActor = [CreateComponentHandle](UObject* InObject, const FName& InSocketName)->UTransformableHandle*
     {
     	if (const AActor* Actor = Cast<AActor>(InObject))
     	{
-    		return CreateComponentHandle(Outer, Actor->GetRootComponent(), InSocketName);
+    		return CreateComponentHandle(Actor->GetRootComponent(), InSocketName);
     	}
     	return nullptr;
     };
diff --git a/Engine/Source/Runtime/Experimental/Animation/Constraints/Public/ConstraintChannel.h b/Engine/Source/Runtime/Experimental/Animation/Constraints/Public/ConstraintChannel.h
--- a/Engine/Source/Runtime/Experimental/Animation/Constraints/Public/ConstraintChannel.h
+++ b/Engine/Source/Runtime/Experimental/Animation/Constraints/Public/ConstraintChannel.h
@@ -99,15 +99,15 @@ struct FConstraintAndActiveChannel
 
 	FConstraintAndActiveChannel() {}
 	FConstraintAndActiveChannel(const TObjectPtr<UTickableConstraint>& InConstraint)
-		: Constraint(InConstraint)
+		: ConstraintCopyToSpawn(InConstraint)
 	{};
-
-	UPROPERTY()
-	TSoftObjectPtr<UTickableConstraint> Constraint;
+	TObjectPtr<UTickableConstraint> GetConstraint() const { return ConstraintCopyToSpawn; }
+	void SetConstraint(UTickableConstraint* InConstraint) { ConstraintCopyToSpawn = InConstraint; }
 
 	UPROPERTY()
 	FMovieSceneConstraintChannel ActiveChannel;
 
+private:
 	UPROPERTY()
 	TObjectPtr<UTickableConstraint> ConstraintCopyToSpawn;
 };
diff --git a/Engine/Source/Runtime/Experimental/Animation/Constraints/Public/ConstraintsManager.h b/Engine/Source/Runtime/Experimental/Animation/Constraints/Public/ConstraintsManager.h
--- a/Engine/Source/Runtime/Experimental/Animation/Constraints/Public/ConstraintsManager.h
+++ b/Engine/Source/Runtime/Experimental/Animation/Constraints/Public/ConstraintsManager.h
@@ -5,11 +5,12 @@
 #include "UObject/Object.h"
 #include "Engine/EngineBaseTypes.h"
 #include "MovieSceneSequenceID.h"
+#include "Misc/Guid.h"
+#include "Engine/Level.h"
 #include "ConstraintsManager.generated.h"
 
 class IMovieScenePlayer;
 class UTickableConstraint;
-class ULevel;
 
 /** 
  * FConstraintTickFunction
@@ -26,6 +27,7 @@ public:
 	CONSTRAINTS_API FConstraintTickFunction();
 	CONSTRAINTS_API ~FConstraintTickFunction();
 
+	CONSTRAINTS_API FConstraintTickFunction(const FConstraintTickFunction& In);
 	/* Begin FTickFunction Interface */
 	CONSTRAINTS_API virtual void ExecuteTick(
 		float DeltaTime,
@@ -81,13 +83,11 @@ public:
 	/** Evaluates the constraint in a context where it's mot done thru the ConstraintTick's tick function. */
 	CONSTRAINTS_API virtual void Evaluate(bool bTickHandlesAlso = false) const;
 
-	/** Tick function that will be registered and evaluated. */
-	UPROPERTY()
-	FConstraintTickFunction ConstraintTick;
-
+	/** Whether or not it's valid for example it may not be fully loaded, missing handles or set to not be valid*/
+	virtual bool IsValid() const PURE_VIRTUAL(IsValid, return false;);
 	/** Sets the Active value and enable/disable the tick function. */
 	CONSTRAINTS_API virtual void SetActive(const bool bIsActive);
-	
+
 	/** Get whether or not it's fully active, it's set to active and all pieces are set up,e.g. objects really exist that are being constrainted*/
 	CONSTRAINTS_API virtual bool IsFullyActive() const;
 	/** If true it contains objects bound to an external system, like sequencer so we don't do certain things, like remove constraints when they don't resolve*/
@@ -101,6 +101,13 @@ public:
 	CONSTRAINTS_API virtual bool ReferencesObject(TWeakObjectPtr<UObject> InObject) const PURE_VIRTUAL(ReferencesObject, return false;);
 	/** Create duplicate with new Outer*/
 	CONSTRAINTS_API virtual UTickableConstraint* Duplicate(UObject* NewOuter) const;
+	/** Initialize and register the constraint with that world*/
+	CONSTRAINTS_API virtual void InitConstraint(UWorld * InWorld)  PURE_VIRTUAL(UTickableConstraint::InitConstraint);
+	/** Teardown and unregister the constraint with that world*/
+	CONSTRAINTS_API virtual void TeardownConstraint(UWorld* InWorld)  PURE_VIRTUAL(UTickableConstraint::DeInitConstraint);
+	/** Get Tick funtion in that world*/
+	CONSTRAINTS_API FConstraintTickFunction& GetTickFunction(UWorld* InWorld);
+	CONSTRAINTS_API const FConstraintTickFunction& GetTickFunction(UWorld* InWorld) const;
 
 #if WITH_EDITOR
 	/** Returns the constraint's label used for UI. */
@@ -113,17 +120,34 @@ public:
 	// UObject interface
 	CONSTRAINTS_API virtual void PostEditChangeProperty(FPropertyChangedEvent& PropertyChangedEvent) override;
 	CONSTRAINTS_API virtual void PostEditUndo() override;
+	CONSTRAINTS_API virtual void PostLoad() override;
+	CONSTRAINTS_API virtual void PostDuplicate(bool bDuplicateForPIE) override;
+	CONSTRAINTS_API virtual void PostInitProperties() override;
+
 	// End of UObject interface
 #endif
 	/** @todo documentation. */
 	UPROPERTY(EditAnywhere, BlueprintReadWrite, DisplayName="Active State", Category="Constraint")
 	bool Active = true;
+
+	UPROPERTY()
+	bool bValid = true;
+
+
+	/** Unique id to be used to fine constraint in a particular world, may be shared*/
+	FGuid ConstraintID;
+
+protected:
+
+	/** Tick function that will be registered and evaluated. */
+	mutable TMap<TWeakObjectPtr<ULevel>,FConstraintTickFunction> ConstraintTicks;
 };
 
 
 /** 
  * UConstraintsManager
- * This object gathers the different constraints of the level and is held by the ConstraintsActor (unique in the level)
+ * This object gathers the different static/nonanimated constraints of the level and is held by the ConstraintsActor (unique in the level)
+ * Note in 5.4 all of the constraints are owned by the subsystem, so need to get that to get at animated constraints
  **/
 
 UCLASS(BLUEPRINTABLE, MinimalAPI)
@@ -144,24 +168,22 @@ public:
 	
 	//UObjects
 	CONSTRAINTS_API virtual void PostLoad() override;
-	/** Get the existing Constraints Manager if existing or create a new one. */
-	static CONSTRAINTS_API UConstraintsManager* Get(UWorld* InWorld);
+
+	CONSTRAINTS_API void Init(UWorld* InWorld);
+	CONSTRAINTS_API void Clear(UWorld* World);
+
+	CONSTRAINTS_API void RemoveStaticConstraint(UTickableConstraint* InConstraint) { Constraints.Remove(InConstraint); }
+	CONSTRAINTS_API bool IsStaticConstraint(UTickableConstraint* InConstraint) const { return Constraints.Contains(InConstraint); }
 
 	/** Find the existing Constraints Manager. */
 	static CONSTRAINTS_API UConstraintsManager* Find(const UWorld* InWorld);
 
-	/** @todo document */
-	CONSTRAINTS_API void Init(UWorld* InWorld);
-	
-	/* Set tick dependencies between two constraints. */
-	CONSTRAINTS_API void SetConstraintDependencies(
-		FConstraintTickFunction* InFunctionToTickBefore,
-		FConstraintTickFunction* InFunctionToTickAfter);
-
-	/** @todo document */
-	CONSTRAINTS_API void Clear(UWorld* World);
+private:
+	/** Get the existing Constraints Manager if existing or create a new one. */
+	static CONSTRAINTS_API UConstraintsManager* Get(UWorld* InWorld);
 
 
+public:
 	/** BP Delegate fired when constraints are added*/
 	UPROPERTY(BlueprintAssignable, Category = Constraints, meta = (DisplayName = "OnConstraintAdded"))
 	FOnConstraintAdded OnConstraintAdded_BP;
@@ -172,23 +194,25 @@ public:
 
 private:
 
+
 	/** @todo document */
 	FDelegateHandle OnActorDestroyedHandle;
-	
 	CONSTRAINTS_API void OnActorDestroyed(AActor* InActor);
 
 	CONSTRAINTS_API void RegisterDelegates(UWorld* World);
 	CONSTRAINTS_API void UnregisterDelegates(UWorld* World);
-	
+
+
 	/** @todo document */
 	CONSTRAINTS_API void Dump() const;
 
-	/** All of the constraints*/
+	/**  For 5.4, this now just contains all of the static constraints.
+	Other constraints will exist in the subsystem and owned by the section*/
 	UPROPERTY()
 	TArray< TObjectPtr<UTickableConstraint> > Constraints;
 
-	friend class FConstraintsManagerController;
 	friend class AConstraintsActor;
+	friend class FConstraintsManagerController;
 };
 
 /** 
@@ -225,36 +249,42 @@ public:
 	/** Make a copy of this constraint and add it to the manager, returns the copy mode if it was added*/
 	CONSTRAINTS_API UTickableConstraint* AddConstraintFromCopy(UTickableConstraint* CopyOfConstraint) const;
 
-	/** Get the index of the given constraint's name. */
-	CONSTRAINTS_API int32 GetConstraintIndex(const FName& InConstraintName) const;
-	
+	/** Get the index of the given constraint's ID. */
+	CONSTRAINTS_API int32 GetConstraintIndex(const FGuid& InGuid) const;
+
 	/** Remove the constraint by name. */
-	CONSTRAINTS_API bool RemoveConstraint(const FName& InConstraintName, bool bDoNotCompensate = false) const;
+	CONSTRAINTS_API bool RemoveConstraint(const FName& InConstraintName, bool bDoNotCompensate = false);
 
 	/** Remove the constraint at the given index. */
-	CONSTRAINTS_API bool RemoveConstraint(const int32 InConstraintIndex, bool bDoNotCompensate = false) const;
+	CONSTRAINTS_API bool RemoveConstraint(const int32 InConstraintIndex, bool bDoNotCompensate = false);
 
 	/** Remove constraint by ptr*/
-	CONSTRAINTS_API bool RemoveConstraint(UTickableConstraint* InConstraint, bool bDoNotCompensate = false) const;
+	CONSTRAINTS_API bool RemoveConstraint(UTickableConstraint* InConstraint, bool bDoNotCompensate = false);
+
+	/** Remove All Constraints*/
+	CONSTRAINTS_API bool RemoveAllConstraints(bool bDoNotCompensate = false);
 
-	/** Returns the constraint based on it's name within the manager's constraints array. */
-	CONSTRAINTS_API UTickableConstraint* GetConstraint(const FName& InConstraintName) const;
+	/** Returns the constraint based on it's ID within the manager's constraints array. */
+	CONSTRAINTS_API UTickableConstraint* GetConstraint(const FGuid& InGuid) const;
 
 	/** Returns the constraint based on it's index within the manager's constraints array. */
 	CONSTRAINTS_API UTickableConstraint* GetConstraint(const int32 InConstraintIndex) const;
 	
 	/** Get read-only access to the array of constraints. */
-	CONSTRAINTS_API const TArray< TObjectPtr<UTickableConstraint> >& GetConstraintsArray() const;
+	CONSTRAINTS_API const TArray< TWeakObjectPtr<UTickableConstraint> >& GetConstraintsArray() const;
 
-	/** Returns manager's constraints array (sorted if needed). */
-	CONSTRAINTS_API TArray< TObjectPtr<UTickableConstraint> > GetAllConstraints(const bool bSorted = false) const;
+	/** Returns systems' constraints array (sorted if needed). */
+	CONSTRAINTS_API TArray< TWeakObjectPtr<UTickableConstraint> > GetAllConstraints(const bool bSorted = false) const;
+
+	/** Returns the static/non-animated constraints*/
+	CONSTRAINTS_API TArray< TObjectPtr<UTickableConstraint> >GetStaticConstraints(const bool bSorted = false) const;
 
 	/** Returns a filtered constraints array checking if the predicate for each element is true. */
 	template <typename Predicate>
-	TArray< TObjectPtr<UTickableConstraint> > GetConstraintsByPredicate(Predicate Pred, const bool bSorted = false) const;
+	TArray< TWeakObjectPtr<UTickableConstraint> > GetConstraintsByPredicate(Predicate Pred, const bool bSorted = false) const;
 
 	/** Get parent constraints of the specified child. If bSorted is true, then the constraints will be sorted by dependency. */
-	CONSTRAINTS_API TArray< TObjectPtr<UTickableConstraint> > GetParentConstraints(
+	CONSTRAINTS_API TArray< TWeakObjectPtr<UTickableConstraint> > GetParentConstraints(
 		const uint32 InTargetHash,
 		const bool bSorted = false) const;
 
@@ -269,7 +299,16 @@ public:
 	/** when PIEing/Simulating it's possible that the constraint isn't in the active manager but still lives*/
 	CONSTRAINTS_API bool DoesExistInAnyWorld(UTickableConstraint* InConstraint);
 
+	/** If static contraint created, we call this, which stores it to the manager/actor*/
+	CONSTRAINTS_API void StaticConstraintCreated(UWorld* InWorld, UTickableConstraint* InConstraint);
+
+	/** If an animated constraint is ready we initialize it (ticks) and add it to the world*/
+	CONSTRAINTS_API void DynamicConstraintCreated(UWorld* InWorld, UTickableConstraint* InConstraint);
+
 private:
+	/** Get the index of the given constraint's name. */
+	int32 GetConstraintIndex(const FName& InConstraintName) const;
+
 	/** Delegeate that's fired when a scene component is constrained, this is needed to make sure things like gizmo's get updated after the constraint tick happens*/
 	DECLARE_MULTICAST_DELEGATE_OneParam(FOnSceneComponentConstrained, USceneComponent* /*InSceneComponent*/);
 	FOnSceneComponentConstrained SceneComponentConstrained;
@@ -278,7 +317,7 @@ private:
 	FConstraintsManagerNotifyDelegate NotifyDelegate;
 	
 	/** Find the existing Constraints Manager in World or create a new one. */
-	CONSTRAINTS_API UConstraintsManager* GetManager() const;
+	CONSTRAINTS_API UConstraintsManager* GetManager();
 	
 	/** Find the existing Constraints Manager in World. */
 	CONSTRAINTS_API UConstraintsManager* FindManager() const;
diff --git a/Engine/Source/Runtime/Experimental/Animation/Constraints/Public/ConstraintsManager.inl b/Engine/Source/Runtime/Experimental/Animation/Constraints/Public/ConstraintsManager.inl
--- a/Engine/Source/Runtime/Experimental/Animation/Constraints/Public/ConstraintsManager.inl
+++ b/Engine/Source/Runtime/Experimental/Animation/Constraints/Public/ConstraintsManager.inl
@@ -3,49 +3,49 @@
 #pragma once
 
 #include "ConstraintsManager.h"
-
+#include "ConstraintSubsystem.h"
 #include "Algo/StableSort.h"
 
 template< typename TConstraint >
 TConstraint* FConstraintsManagerController::AllocateConstraintT(const FName& InBaseName) const
 {
-	UConstraintsManager* Manager = GetManager();
-	if (!Manager)
+	UConstraintSubsystem* Subsystem =  UConstraintSubsystem::Get();
+	if (!Subsystem)
 	{
 		return nullptr;
 	}
 	// unique name (we may want to use another approach here to manage uniqueness)
-	const FName Name = MakeUniqueObjectName(Manager, TConstraint::StaticClass(), InBaseName);
+	const FName Name = MakeUniqueObjectName(Subsystem, TConstraint::StaticClass(), InBaseName);
 
-	TConstraint* NewConstraint = NewObject<TConstraint>(Manager, Name, RF_Transactional);
+	TConstraint* NewConstraint = NewObject<TConstraint>(Subsystem, Name, RF_Transactional);
 	NewConstraint->Modify();
 	return NewConstraint;
 }
 
 template <typename Predicate>
-TArray< TObjectPtr<UTickableConstraint> > FConstraintsManagerController::GetConstraintsByPredicate(
+TArray< TWeakObjectPtr<UTickableConstraint> > FConstraintsManagerController::GetConstraintsByPredicate(
 	Predicate Pred, const bool bSorted) const
 {
-	using ConstraintPtr = TObjectPtr<UTickableConstraint>;
-	
-	UConstraintsManager* Manager = FindManager();
-	if (!Manager)
+	using ConstraintPtr = TWeakObjectPtr<UTickableConstraint>;
+	UConstraintSubsystem* Subsystem = UConstraintSubsystem::Get();
+
+	if (!Subsystem)
 	{
-		static const TArray< TObjectPtr<UTickableConstraint> > Empty;
+		static const TArray< TWeakObjectPtr<UTickableConstraint> > Empty;
 		return Empty;
 	}
 
-	TArray< ConstraintPtr > Constraints = Manager->Constraints.FilterByPredicate(Pred);
+	TArray< ConstraintPtr > Constraints = Subsystem->GetConstraints(World).FilterByPredicate(Pred);
 	if (!bSorted)
 	{
 		return Constraints;
 	}
 
 	// LHS ticks before RHS = LHS is a prerex of RHS 
-	auto TicksBefore = [](const UTickableConstraint& LHS, const UTickableConstraint& RHS)
+	auto TicksBefore = [this](const UTickableConstraint& LHS, const UTickableConstraint& RHS)
 	{
-		const TArray<FTickPrerequisite>& RHSPrerex = RHS.ConstraintTick.GetPrerequisites();
-		const FConstraintTickFunction& LHSTickFunction = LHS.ConstraintTick;
+		const TArray<FTickPrerequisite>& RHSPrerex = RHS.GetTickFunction(World).GetPrerequisites();
+		const FConstraintTickFunction& LHSTickFunction = LHS.GetTickFunction(World);
 		const bool bIsLHSAPrerexOfRHS = RHSPrerex.ContainsByPredicate([&LHSTickFunction](const FTickPrerequisite& Prerex)
 		{
 			return Prerex.PrerequisiteTickFunction == &LHSTickFunction;
diff --git a/Engine/Source/Runtime/Experimental/Animation/Constraints/Public/ConstraintsScripting.h b/Engine/Source/Runtime/Experimental/Animation/Constraints/Public/ConstraintsScripting.h
--- a/Engine/Source/Runtime/Experimental/Animation/Constraints/Public/ConstraintsScripting.h
+++ b/Engine/Source/Runtime/Experimental/Animation/Constraints/Public/ConstraintsScripting.h
@@ -27,13 +27,10 @@ public:
 public:
 
 	/**
-	* Get the manager of the constraints. This object contains delegates to listen to for when constraints are added,deleted,
-	* and is also the outer used when creating custom transformable handles, for example this is used to create control rig transformable handles
-	* Note this function will create the mananager and it's actor if one doesn't exist.
-	* @param InWorld, the world you are in
-	* @return Returns the mananger
+	* No longer used, the system will internally set up the outers for the handles and see UConstraintSubsystem
+	* for the delegates.
 	*/
-	UFUNCTION(BlueprintCallable, Category = "Editor Scripting | Sequencer Tools | Control Rig | Constraints")
+	UE_DEPRECATED(5.3, "Please use UConstraintSubsystem delegates.")
 	static CONSTRAINTS_API UConstraintsManager* GetManager(UWorld* InWorld);
 
 	/**
diff --git a/Engine/Source/Runtime/Experimental/Animation/Constraints/Public/TransformConstraint.h b/Engine/Source/Runtime/Experimental/Animation/Constraints/Public/TransformConstraint.h
--- a/Engine/Source/Runtime/Experimental/Animation/Constraints/Public/TransformConstraint.h
+++ b/Engine/Source/Runtime/Experimental/Animation/Constraints/Public/TransformConstraint.h
@@ -16,7 +16,7 @@ class UTransformableHandle;
 class UTransformableComponentHandle;
 class USceneComponent;
 enum class EMovieSceneTransformChannel : uint32;
-
+class UWorld;
 /** 
  * UTickableTransformConstraint
  **/
@@ -46,7 +46,10 @@ public:
 	/** Resolve the bound objects so that any object it references are resovled and correctly set up*/
 	CONSTRAINTS_API virtual void ResolveBoundObjects(FMovieSceneSequenceID LocalSequenceID, IMovieScenePlayer& Player,UObject* SubObject) override;
 
-	/** If Active and the handles and targets are valid*/
+	/** Whether or not it's valid for example it may not be fully loaded*/
+	virtual bool IsValid() const override;
+
+	/** If Active and the handles and targets are valid, and tick function is registered*/
 	CONSTRAINTS_API virtual bool IsFullyActive() const override;
 
 	/** If that constraint needs to be handled by the compensation system. */
@@ -114,7 +117,7 @@ public:
 	/**
 	* Sets up dependencies with the first primary prerequisite available if the parent does not tick.   
 	*/
-	CONSTRAINTS_API void EnsurePrimaryDependency();
+	CONSTRAINTS_API void EnsurePrimaryDependency(UWorld* InWorld);
 	
 protected:
 
@@ -133,7 +136,7 @@ protected:
 	 * Sets up dependencies between the parent, the constraint and the child using their respective tick functions.
 	 * It creates a dependency graph between them so that they tick in the right order when evaluated.   
 	*/
-	CONSTRAINTS_API void SetupDependencies();
+	CONSTRAINTS_API void SetupDependencies(UWorld* InWorld);
 
 	/** Set the current child's global transform. */
 	CONSTRAINTS_API void SetChildGlobalTransform(const FTransform& InGlobal) const;
@@ -151,8 +154,10 @@ protected:
 	/** Returns the handle's tick function (ensuring it lives in the same world). */
 	CONSTRAINTS_API FTickFunction* GetHandleTickFunction(const TObjectPtr<UTransformableHandle>& InHandle) const;
 
+public:
 	/** (Re-)Registers the constraint function and (re-)binds the required delegates*/
-	CONSTRAINTS_API void InitConstraint();
+	CONSTRAINTS_API virtual void InitConstraint(UWorld * InWorld) override;
+	CONSTRAINTS_API virtual void TeardownConstraint(UWorld * InWorld) override;
 
 #if WITH_EDITOR
 public:
@@ -441,13 +446,12 @@ struct FTransformConstraintUtils
 	static CONSTRAINTS_API void GetParentConstraints(
 		UWorld* World,
 		const AActor* InChild,
-		TArray< TObjectPtr<UTickableConstraint> >& OutConstraints);
+		TArray< TWeakObjectPtr<UTickableConstraint> >& OutConstraints);
 
 	/** Create a handle for the scene component.*/
 	static CONSTRAINTS_API UTransformableComponentHandle* CreateHandleForSceneComponent(
 		USceneComponent* InSceneComponent,
-		const FName& InSocketName,
-		UObject* Outer);
+		const FName& InSocketName);
 
 	/** Creates a new transform constraint based on the InType. */
 	static CONSTRAINTS_API UTickableTransformConstraint* CreateFromType(
@@ -480,15 +484,15 @@ struct FTransformConstraintUtils
 	/** Computes the current constraint space local transform. */
 	static CONSTRAINTS_API TOptional<FTransform> GetRelativeTransform(UWorld* InWorld, const uint32 InHandleHash);
 	static CONSTRAINTS_API TOptional<FTransform> GetConstraintsRelativeTransform(
-		const TArray< TObjectPtr<UTickableConstraint> >& InConstraints,
+		const TArray< TWeakObjectPtr<UTickableConstraint> >& InConstraints,
 		const FTransform& InChildLocal, const FTransform& InChildWorld);
 
 	/** Get the last active constraint that has dynamic offset. */
-	static CONSTRAINTS_API int32 GetLastActiveConstraintIndex(const TArray< TObjectPtr<UTickableConstraint> >& InConstraints);
+	static CONSTRAINTS_API int32 GetLastActiveConstraintIndex(const TArray< TWeakObjectPtr<UTickableConstraint> >& InConstraints);
 
 	/** Fills a constraint array that InParentHandle is the parent of. */
 	static CONSTRAINTS_API void GetChildrenConstraints(
 		UWorld* World,
 		const UTransformableHandle* InParentHandle,
-		TArray< TObjectPtr<UTickableConstraint> >& OutConstraints);
+		TArray< TWeakObjectPtr<UTickableConstraint> >& OutConstraints);
 };
diff --git a/Engine/Source/Runtime/Experimental/Animation/Constraints/Public/TransformableHandle.h b/Engine/Source/Runtime/Experimental/Animation/Constraints/Public/TransformableHandle.h
--- a/Engine/Source/Runtime/Experimental/Animation/Constraints/Public/TransformableHandle.h
+++ b/Engine/Source/Runtime/Experimental/Animation/Constraints/Public/TransformableHandle.h
@@ -72,7 +72,7 @@ public:
 	CONSTRAINTS_API void OnBindingIDsUpdated(const TMap<UE::MovieScene::FFixedObjectBindingID, UE::MovieScene::FFixedObjectBindingID>& OldFixedToNewFixedMap, FMovieSceneSequenceID LocalSequenceID, const FMovieSceneSequenceHierarchy* Hierarchy, IMovieScenePlayer& Player);
 	
 	/** Perform any special ticking needed for this handle, by default it does nothing, todo need to see if we need to tick control rig also*/
-	virtual void TickForBaking() {};
+	virtual void TickForBaking() const {};
 	/** Get the array of float channels for the specified section*/
 	CONSTRAINTS_API virtual TArrayView<FMovieSceneFloatChannel*>  GetFloatChannels(const UMovieSceneSection* InSection) const PURE_VIRTUAL(GetFloatChannels, return TArrayView<FMovieSceneFloatChannel*>(); );
 	/** Get the array of double channels for the specified section*/
@@ -149,7 +149,7 @@ public:
 	/** Gets the local transform of Component in it's attachment. */
 	CONSTRAINTS_API virtual FTransform GetLocalTransform() const override;
 	/** Tick the component*/
-	CONSTRAINTS_API virtual void TickForBaking() override;
+	CONSTRAINTS_API virtual void TickForBaking() const override;
 	/** Returns the target object containing the tick function (e.i. Component). */
 	CONSTRAINTS_API virtual UObject* GetPrerequisiteObject() const override;
 	/** Returns Component's tick function. */
diff --git a/Engine/Source/Runtime/Experimental/Animation/Constraints/Public/TransformableRegistry.h b/Engine/Source/Runtime/Experimental/Animation/Constraints/Public/TransformableRegistry.h
--- a/Engine/Source/Runtime/Experimental/Animation/Constraints/Public/TransformableRegistry.h
+++ b/Engine/Source/Runtime/Experimental/Animation/Constraints/Public/TransformableRegistry.h
@@ -20,7 +20,7 @@ class FTransformableRegistry
 public:
 
 	/** Defines the function that will allocate a new transformable handle to wrap the object. */
-	using CreateHandleFuncT = TFunction<UTransformableHandle*(UObject*, UObject*, const FName&)>;
+	using CreateHandleFuncT = TFunction<UTransformableHandle*(UObject*, const FName&)>;
 	/** Defines the function that return a hash value from the object and the attachment name. */
 	using GetHashFuncT = TFunction<uint32(const UObject*, const FName&)>;
 	
diff --git a/Engine/Source/Runtime/MovieSceneTracks/Private/Sections/MovieScene3DTransformSection.cpp b/Engine/Source/Runtime/MovieSceneTracks/Private/Sections/MovieScene3DTransformSection.cpp
--- a/Engine/Source/Runtime/MovieSceneTracks/Private/Sections/MovieScene3DTransformSection.cpp
+++ b/Engine/Source/Runtime/MovieSceneTracks/Private/Sections/MovieScene3DTransformSection.cpp
@@ -400,7 +400,7 @@ void UMovieScene3DTransformSection::OnBindingIDsUpdated(const TMap<UE::MovieScen
 	{
 		for (FConstraintAndActiveChannel& ConstraintChannel : Constraints->ConstraintsChannels)
 		{
-			if (UTickableTransformConstraint* TransformConstraint = Cast< UTickableTransformConstraint>(ConstraintChannel.Constraint.Get()))
+			if (UTickableTransformConstraint* TransformConstraint = Cast< UTickableTransformConstraint>(ConstraintChannel.GetConstraint()))
 			{
 				//Don't do child's we do that in the system, needed for duplication
 				if (TransformConstraint->ParentTRSHandle)
@@ -408,14 +408,6 @@ void UMovieScene3DTransformSection::OnBindingIDsUpdated(const TMap<UE::MovieScen
 					TransformConstraint->ParentTRSHandle->OnBindingIDsUpdated(OldFixedToNewFixedMap, LocalSequenceID, Hierarchy, Player);
 				}
 			}
-			if (UTickableTransformConstraint* SpawnCopy = Cast< UTickableTransformConstraint>(ConstraintChannel.ConstraintCopyToSpawn))
-			{
-				//Don't do child's we do that in the system, needed for duplication
-				if (SpawnCopy->ParentTRSHandle)
-				{
-					SpawnCopy->ParentTRSHandle->OnBindingIDsUpdated(OldFixedToNewFixedMap, LocalSequenceID, Hierarchy, Player);
-				}
-			}
 		}
 	}
 }
@@ -426,7 +418,7 @@ void UMovieScene3DTransformSection::GetReferencedBindings(TArray<FGuid>& OutBind
 	{
 		for (FConstraintAndActiveChannel& ConstraintChannel : Constraints->ConstraintsChannels)
 		{
-			if (UTickableTransformConstraint* TransformConstraint = Cast< UTickableTransformConstraint>(ConstraintChannel.Constraint.Get()))
+			if (UTickableTransformConstraint* TransformConstraint = Cast< UTickableTransformConstraint>(ConstraintChannel.GetConstraint().Get()))
 			{
 				if (TransformConstraint->ChildTRSHandle && TransformConstraint->ChildTRSHandle->ConstraintBindingID.IsValid())
 				{
@@ -444,16 +436,6 @@ void UMovieScene3DTransformSection::GetReferencedBindings(TArray<FGuid>& OutBind
 void UMovieScene3DTransformSection::PreSave(FObjectPreSaveContext SaveContext)
 {
 	Super::PreSave(SaveContext);
-	if (Constraints)
-	{
-		for (FConstraintAndActiveChannel& ActiveChannel : Constraints->ConstraintsChannels)
-		{
-			if (ActiveChannel.Constraint.IsValid())
-			{
-				ActiveChannel.ConstraintCopyToSpawn = ActiveChannel.Constraint->Duplicate(this);
-			}
-		}
-	}
 }
 
 template<typename BaseBuilderType>
@@ -622,14 +604,13 @@ void UMovieScene3DTransformSection::ImportConstraintEntity(UMovieSceneEntitySyst
 
 		checkf(Constraints->ConstraintsChannels.IsValidIndex(ConstraintIndex), TEXT("Encoded constraint (%d) index is not valid within array size %d. Data must have been manipulated without re-compilaition."), ConstraintIndex, Constraints->ConstraintsChannels.Num());
 		//add if constraint or spawn copy is valid
-		if (Constraints->ConstraintsChannels[ConstraintIndex].Constraint.IsValid() ||
-			Constraints->ConstraintsChannels[ConstraintIndex].ConstraintCopyToSpawn)
+		if (Constraints->ConstraintsChannels[ConstraintIndex].GetConstraint().Get())
 		{
-			FName ConstraintName = Constraints->ConstraintsChannels[ConstraintIndex].Constraint.IsValid() ? Constraints->ConstraintsChannels[ConstraintIndex].Constraint->GetFName()
-				: Constraints->ConstraintsChannels[ConstraintIndex].ConstraintCopyToSpawn->GetFName();
+			FGuid ConstraintID = Constraints->ConstraintsChannels[ConstraintIndex].GetConstraint()->ConstraintID;
+			//ID's should be the same!
 
 			FConstraintComponentData ComponentData;
-			ComponentData.ConstraintName = ConstraintName;
+			ComponentData.ConstraintID = ConstraintID;
 			ComponentData.Section = this;
 			OutImportedEntity->AddBuilder(
 				FEntityBuilder()
@@ -759,11 +740,7 @@ EMovieSceneChannelProxyType UMovieScene3DTransformSection::CacheChannelProxy()
 
 	auto GetValidConstraint = [](FConstraintAndActiveChannel& ConstraintChannel)
 	{
-		if (ConstraintChannel.Constraint.IsValid())
-		{
-			return ConstraintChannel.Constraint.Get();
-		}
-		return (ConstraintChannel.ConstraintCopyToSpawn.Get());
+		return ConstraintChannel.GetConstraint().Get();
 	};
 #endif
 		
@@ -980,26 +957,25 @@ void UMovieScene3DTransformSection::SetBlendType(EMovieSceneBlendType InBlendTyp
 	}
 }
 
-bool UMovieScene3DTransformSection::HasConstraintChannel(const FName& InConstraintName) const
+bool UMovieScene3DTransformSection::HasConstraintChannel(const FGuid& InGuid) const
 {
 	if (Constraints)
 	{
-		return Constraints->ConstraintsChannels.ContainsByPredicate([InConstraintName](const FConstraintAndActiveChannel& InChannel)
+		return Constraints->ConstraintsChannels.ContainsByPredicate([InGuid](const FConstraintAndActiveChannel& InChannel)
 			{
-				return InChannel.Constraint.IsValid() ? InChannel.Constraint->GetFName() == InConstraintName : false;
+				return InChannel.GetConstraint().Get() ? InChannel.GetConstraint()->ConstraintID == InGuid : false;
 			});
 	}
 	return false;
 }
 
-FConstraintAndActiveChannel* UMovieScene3DTransformSection::GetConstraintChannel(const FName& InConstraintName)
+FConstraintAndActiveChannel* UMovieScene3DTransformSection::GetConstraintChannel(const FGuid& InConstraintID)
 {
 	if (Constraints)
 	{
-		const int32 Index = Constraints->ConstraintsChannels.IndexOfByPredicate([InConstraintName](const FConstraintAndActiveChannel& InChannel)
+		const int32 Index = Constraints->ConstraintsChannels.IndexOfByPredicate([InConstraintID](const FConstraintAndActiveChannel& InChannel)
 			{
-				return InChannel.Constraint.IsValid() ? InChannel.Constraint->GetFName() == InConstraintName : (InChannel.ConstraintCopyToSpawn ?
-				InChannel.ConstraintCopyToSpawn->GetFName() == InConstraintName : false);
+				return InChannel.GetConstraint().Get() ? InChannel.GetConstraint()->ConstraintID == InConstraintID : false;
 			});
 		return (Index != INDEX_NONE) ? &(Constraints->ConstraintsChannels[Index]) : nullptr;
 	}
@@ -1019,15 +995,7 @@ TArray<FConstraintAndActiveChannel>& UMovieScene3DTransformSection::GetConstrain
 void UMovieScene3DTransformSection::PostLoad()
 {
 	Super::PostLoad();
-	//for spawnables the control rig saved in our channels may have changed so we need to update thaem
-	if (Constraints)
-	{
-		for (FConstraintAndActiveChannel& ConstraintChannel : Constraints->ConstraintsChannels)
-		{
-			ConstraintChannel.Constraint.Reset(); //clear it out may be referencing wrong level
-		}
-	}
-}
+}	
 
 #if WITH_EDITOR
 bool UMovieScene3DTransformSection::Modify(bool bAlwaysMarkDirty)
@@ -1046,20 +1014,6 @@ bool UMovieScene3DTransformSection::Modify(bool bAlwaysMarkDirty)
 void UMovieScene3DTransformSection::PostDuplicate(bool bDuplicateForPIE)
 {
 	Super::PostDuplicate(bDuplicateForPIE);
-	//if not duplicating for PIE clear out the soft object ptrs to the constraints we dont' want to hold pointers to the old constraints.
-	//also make the copy to spawn so we can duplicate it if it doesn't exist in the level
-	//to the old ones but make new duplicates
-	if (bDuplicateForPIE == false && Constraints)
-	{
-		for (FConstraintAndActiveChannel& ConstraintChannel : Constraints->ConstraintsChannels)
-		{
-			if (ConstraintChannel.Constraint.IsValid() && !ConstraintChannel.ConstraintCopyToSpawn)
-			{
-				ConstraintChannel.ConstraintCopyToSpawn = ConstraintChannel.Constraint->Duplicate(this);
-			}
-			ConstraintChannel.Constraint.Reset();
-		}
-	}
 }
 #endif
 void UMovieScene3DTransformSection::AddConstraintChannel(UTickableConstraint* InConstraint)
@@ -1069,7 +1023,7 @@ void UMovieScene3DTransformSection::AddConstraintChannel(UTickableConstraint* In
 	{
 		Constraints = NewObject<UMovieScene3DTransformSectionConstraints>(this, NAME_None, RF_Public| RF_Transactional);
 	}
-	if (InConstraint && !HasConstraintChannel(InConstraint->GetFName()))
+	if (InConstraint && !HasConstraintChannel(InConstraint->ConstraintID))
 	{
 		Constraints->SetFlags(RF_Transactional);
 		Constraints->Modify();
@@ -1079,8 +1033,9 @@ void UMovieScene3DTransformSection::AddConstraintChannel(UTickableConstraint* In
 		ExistingChannel->SetDefault(false);
 
 		//make copy that we can spawn if it doesn't exist
-		Constraints->ConstraintsChannels[NewIndex].ConstraintCopyToSpawn = InConstraint->Duplicate(this);
-
+		//the rename changes the outer to this section (from any actor manager)
+		InConstraint->Rename(nullptr, this, REN_ForceNoResetLoaders | REN_DontCreateRedirectors);
+	
 		CacheChannelProxy();
 
 		if (OnConstraintChannelAdded.IsBound())
@@ -1096,12 +1051,12 @@ void UMovieScene3DTransformSection::ReplaceConstraint(const FName InConstraintNa
 	{
 		const int32 Index = Constraints->ConstraintsChannels.IndexOfByPredicate([InConstraintName](const FConstraintAndActiveChannel& InChannel)
 			{
-				return InChannel.ConstraintCopyToSpawn ? InChannel.ConstraintCopyToSpawn->GetFName() == InConstraintName : false;
+				return InChannel.GetConstraint().Get() ? InChannel.GetConstraint()->GetFName() == InConstraintName : false;
 			});
 		if (Index != INDEX_NONE)
 		{
 			Modify();
-			Constraints->ConstraintsChannels[Index].Constraint = InConstraint;
+			Constraints->ConstraintsChannels[Index].GetConstraint() = InConstraint;
 			CacheChannelProxy();
 		}
 	}
@@ -1117,7 +1072,7 @@ void UMovieScene3DTransformSection::RemoveConstraintChannel(const UTickableConst
 	{
 		const int32 Index = Constraints->ConstraintsChannels.IndexOfByPredicate([InConstraint](const FConstraintAndActiveChannel& InChannel)
 			{
-				return InChannel.Constraint.IsValid() ? InChannel.Constraint == InConstraint : false;
+				return InChannel.GetConstraint().Get() ? InChannel.GetConstraint().Get() == InConstraint : false;
 			});
 
 		if (Constraints->ConstraintsChannels.IsValidIndex(Index))
diff --git a/Engine/Source/Runtime/MovieSceneTracks/Private/Systems/MovieSceneConstraintSystem.cpp b/Engine/Source/Runtime/MovieSceneTracks/Private/Systems/MovieSceneConstraintSystem.cpp
--- a/Engine/Source/Runtime/MovieSceneTracks/Private/Systems/MovieSceneConstraintSystem.cpp
+++ b/Engine/Source/Runtime/MovieSceneTracks/Private/Systems/MovieSceneConstraintSystem.cpp
@@ -20,35 +20,20 @@
 
 #include UE_INLINE_GENERATED_CPP_BY_NAME(MovieSceneConstraintSystem)
 
-class UTickableConstraint;
-
-
-
 
 static UTickableConstraint* CreateConstraintIfNeeded(FConstraintsManagerController& Controller, FConstraintAndActiveChannel* ConstraintAndActiveChannel, const  FConstraintComponentData& ConstraintChannel)
 {
-	if (ConstraintAndActiveChannel->Constraint.IsPending())
-	{
-		ConstraintAndActiveChannel->Constraint.LoadSynchronous();
-	}
-	UTickableConstraint* Constraint = ConstraintAndActiveChannel->Constraint.Get();
+	UTickableConstraint* Constraint = ConstraintAndActiveChannel->GetConstraint();
 	if (!Constraint) //if constraint doesn't exist it probably got unspawned so recreate it and add it
 	{
-		Constraint = Controller.AddConstraintFromCopy(ConstraintAndActiveChannel->ConstraintCopyToSpawn);
-		ConstraintChannel.Section->ReplaceConstraint(ConstraintAndActiveChannel->ConstraintCopyToSpawn->GetFName(), Constraint);
+		return nullptr;
 	}
 	else // it's possible that we have it but it's not in the manager, due to manager not being saved with it (due to spawning or undo/redo).
 	{
-		const TArray< TObjectPtr<UTickableConstraint>>& ConstraintsArray = Controller.GetConstraintsArray();
-		if (ConstraintsArray.Find(Constraint) == INDEX_NONE)
+		const TArray< TWeakObjectPtr<UTickableConstraint>>& ConstraintsArray = Controller.GetConstraintsArray();
+		if (Controller.GetConstraint(Constraint->ConstraintID) == nullptr)
 		{
-			//it's possible it belongs to another Manager
-			if (Controller.DoesExistInAnyWorld(Constraint) == false) //it's no where make a copy and add it
-			{
-				Constraint = Controller.AddConstraintFromCopy(ConstraintAndActiveChannel->ConstraintCopyToSpawn);
-				ConstraintChannel.Section->ReplaceConstraint(ConstraintAndActiveChannel->ConstraintCopyToSpawn->GetFName(), Constraint);
-			}
-
+			Controller.AddConstraint(Constraint);
 		}
 	}
 	return Constraint;
@@ -61,7 +46,7 @@ namespace UE::MovieScene
 
 struct FPreAnimatedConstraint
 {
-	FName ConstraintName;
+	FGuid ConstraintID;
 	bool bDeleteIt;
 	bool bPreviouslyEnabled;
 	TWeakObjectPtr<USceneComponent> SceneComponent;
@@ -74,22 +59,22 @@ struct FPreAnimatedConstraintTraits : FBoundObjectPreAnimatedStateTraits
 	struct FKeyType
 	{
 		FObjectKey Object;
-		FName ConstraintName;
+		FGuid ConstraintID;
 
 		/** Constructor that takes a BoundObject and ConstraintChannel component */
 		FKeyType(UObject* InObject, const FConstraintComponentData& ComponentData)
 			: Object(InObject)
-			, ConstraintName(ComponentData.ConstraintName)
+			, ConstraintID(ComponentData.ConstraintID)
 		{}
 
 		/** Hashing and equality required for storage within a map */
 		friend uint32 GetTypeHash(const FKeyType& InKey)
 		{
-			return HashCombine(GetTypeHash(InKey.Object), GetTypeHash(InKey.ConstraintName));
+			return HashCombine(GetTypeHash(InKey.Object), GetTypeHash(InKey.ConstraintID));
 		}
 		friend bool operator==(const FKeyType& A, const FKeyType& B)
 		{
-			return A.Object == B.Object && A.ConstraintName == B.ConstraintName;
+			return A.Object == B.Object && A.ConstraintID == B.ConstraintID;
 		}
 	};
 
@@ -100,20 +85,20 @@ struct FPreAnimatedConstraintTraits : FBoundObjectPreAnimatedStateTraits
 	{
 		USceneComponent* SceneComponent = CastChecked<USceneComponent>(InBoundObject);
 		FConstraintsManagerController& Controller = FConstraintsManagerController::Get(SceneComponent->GetWorld());
-		if (FConstraintAndActiveChannel* ConstraintAndActiveChannel = ConstraintData.Section->GetConstraintChannel(ConstraintData.ConstraintName))
+		if (FConstraintAndActiveChannel* ConstraintAndActiveChannel = ConstraintData.Section->GetConstraintChannel(ConstraintData.ConstraintID))
 		{
-			if (ConstraintAndActiveChannel->Constraint.IsPending())
+			UTickableConstraint* Constraint = ConstraintAndActiveChannel->GetConstraint().Get();
+			const bool bIsActive = false; //we set static mesh constraints to be inactive to matchup how CR's work in not being valid
+			const FGuid ID = Constraint ? Constraint->ConstraintID : ConstraintData.ConstraintID;
+			//also need to create it
+			Constraint = CreateConstraintIfNeeded(Controller, ConstraintAndActiveChannel, ConstraintData);
+			if (Constraint)
 			{
-				ConstraintAndActiveChannel->Constraint.LoadSynchronous();
+				Controller.AddConstraint(Constraint);
 			}
-			UTickableConstraint* Constraint = ConstraintAndActiveChannel->Constraint.Get();
-			const bool bIsActive = Constraint ? Constraint->Active : false;
-			const FName Name = Constraint ? Constraint->GetFName() : ConstraintData.ConstraintName;
-			//also need to create it
-			CreateConstraintIfNeeded(Controller, ConstraintAndActiveChannel, ConstraintData);
-			return FPreAnimatedConstraint{ Name, (Constraint == nullptr), bIsActive, SceneComponent, ConstraintData.Section };
+			return FPreAnimatedConstraint{ ID, (Constraint == nullptr), bIsActive, SceneComponent, ConstraintData.Section };
 		}
-		return FPreAnimatedConstraint{ ConstraintData.ConstraintName, false,false, SceneComponent, ConstraintData.Section };
+		return FPreAnimatedConstraint{ ConstraintData.ConstraintID, false,false, SceneComponent, ConstraintData.Section };
 	}
 
 	static void RestorePreAnimatedValue(const FKeyType& InKey, const FPreAnimatedConstraint& OldValue, const FRestoreStateParams& Params)
@@ -124,17 +109,13 @@ struct FPreAnimatedConstraintTraits : FBoundObjectPreAnimatedStateTraits
 			{
 				return;
 			}
-			if (FConstraintAndActiveChannel* ConstraintAndActiveChannel = OldValue.Section->GetConstraintChannel(OldValue.ConstraintName))
+			if (FConstraintAndActiveChannel* ConstraintAndActiveChannel = OldValue.Section->GetConstraintChannel(OldValue.ConstraintID))
 			{				
-				if (ConstraintAndActiveChannel->Constraint.IsPending())
-				{
-					ConstraintAndActiveChannel->Constraint.LoadSynchronous();
-				}
-				UTickableConstraint* Constraint = ConstraintAndActiveChannel->Constraint.Get();
+				UTickableConstraint* Constraint = ConstraintAndActiveChannel->GetConstraint().Get();
 				if (Constraint)
 				{
 					Constraint->SetActive(OldValue.bPreviouslyEnabled);
-
+					Constraint->bValid = false; 
 					//we don't delete it since for some reason on first creation we get a restore state and so it immediately get's deleted
 					if (OldValue.bDeleteIt)
 					{
@@ -178,7 +159,7 @@ struct FEvaluateConstraintChannels
 
 	void ForEachEntity(UObject* BoundObject, FInstanceHandle InstanceHandle, const FConstraintComponentData& ConstraintChannel, FFrameTime FrameTime) const
 	{
-		FConstraintAndActiveChannel* ConstraintAndActiveChannel = ConstraintChannel.Section->GetConstraintChannel(ConstraintChannel.ConstraintName);
+		FConstraintAndActiveChannel* ConstraintAndActiveChannel = ConstraintChannel.Section->GetConstraintChannel(ConstraintChannel.ConstraintID);
 		if (!ConstraintAndActiveChannel)
 		{
 			return;
@@ -198,6 +179,8 @@ struct FEvaluateConstraintChannels
 		Constraint->ResolveBoundObjects(TargetInstance.GetSequenceID(), *TargetInstance.GetPlayer());
 		if (UTickableTransformConstraint* TransformConstraint = Cast< UTickableTransformConstraint>(Constraint))
 		{
+			TransformConstraint->InitConstraint(World);
+			
 			if (UTransformableComponentHandle* ComponentHandle = Cast<UTransformableComponentHandle>(TransformConstraint->ChildTRSHandle))
 			{
 				//bound component may change so need to update constraint
@@ -206,6 +189,8 @@ struct FEvaluateConstraintChannels
 				UpdateHandle.Constraint = TransformConstraint;
 				UpdateHandle.TransformHandle = ComponentHandle;
 				System->DynamicOffsets.Add(UpdateHandle);
+				TransformConstraint->EnsurePrimaryDependency(BoundObject->GetWorld());
+
 			}
 		}
 	}
diff --git a/Engine/Source/Runtime/MovieSceneTracks/Public/MovieSceneTracksComponentTypes.h b/Engine/Source/Runtime/MovieSceneTracks/Public/MovieSceneTracksComponentTypes.h
--- a/Engine/Source/Runtime/MovieSceneTracks/Public/MovieSceneTracksComponentTypes.h
+++ b/Engine/Source/Runtime/MovieSceneTracks/Public/MovieSceneTracksComponentTypes.h
@@ -17,6 +17,7 @@
 #include "ConstraintChannel.h"
 #include "MaterialTypes.h"
 #include "Sections/MovieSceneCameraShakeSection.h"
+#include "Misc/Guid.h"
 #include "Tracks/MovieSceneMaterialTrack.h"
 #include "MovieSceneTracksComponentTypes.generated.h"
 
@@ -81,7 +82,7 @@ struct FConstraintComponentData
 	GENERATED_BODY()
 
 	UPROPERTY()
-	FName ConstraintName;
+	FGuid ConstraintID;
 	UMovieScene3DTransformSection* Section;
 };
 
diff --git a/Engine/Source/Runtime/MovieSceneTracks/Public/Sections/MovieScene3DTransformSection.h b/Engine/Source/Runtime/MovieSceneTracks/Public/Sections/MovieScene3DTransformSection.h
--- a/Engine/Source/Runtime/MovieSceneTracks/Public/Sections/MovieScene3DTransformSection.h
+++ b/Engine/Source/Runtime/MovieSceneTracks/Public/Sections/MovieScene3DTransformSection.h
@@ -229,6 +229,7 @@ public:
 	virtual bool Modify(bool bAlwaysMarkDirty = true) override;
 	virtual void PostDuplicate(bool bDuplicateForPIE) override;
 #endif
+
 	/* From UMovieSection*/
 	
 	MOVIESCENETRACKS_API virtual bool ShowCurveForChannel(const void *Channel) const override;
@@ -331,12 +332,12 @@ public:
 	/*
 	* If it has that constraint with that Name
 	*/
-	virtual  bool HasConstraintChannel(const FName& InConstraintName) const override;
+	virtual  bool HasConstraintChannel(const FGuid& InGuid) const override;
 
 	/*
 	* Get constraint with that name
 	*/
-	virtual FConstraintAndActiveChannel* GetConstraintChannel(const FName& InConstraintName) override;
+	virtual FConstraintAndActiveChannel* GetConstraintChannel(const FGuid& InConstraintGuid) override;
 
 	/*
 	*  Add Constraint channel
diff --git a/Engine/Source/Runtime/MovieSceneTracks/Public/Sections/MovieSceneConstrainedSection.h b/Engine/Source/Runtime/MovieSceneTracks/Public/Sections/MovieSceneConstrainedSection.h
--- a/Engine/Source/Runtime/MovieSceneTracks/Public/Sections/MovieSceneConstrainedSection.h
+++ b/Engine/Source/Runtime/MovieSceneTracks/Public/Sections/MovieSceneConstrainedSection.h
@@ -11,7 +11,7 @@
 #include "ConstraintChannel.h"
 #include "MovieSceneConstrainedSection.generated.h"
 
-
+struct FGuid;
 
 /**
  * Functionality to add to sections that contain constraints
@@ -40,12 +40,12 @@ public:
 	/*
 	* If it has that constraint with that Name
 	*/
-	virtual bool HasConstraintChannel(const FName& InConstraintName) const = 0;
+	virtual bool HasConstraintChannel(const FGuid& InGuid) const = 0;
 
 	/*
 	* Get constraint with that name
 	*/
-	virtual FConstraintAndActiveChannel* GetConstraintChannel(const FName& InConstraintName) = 0;
+	virtual FConstraintAndActiveChannel* GetConstraintChannel(const FGuid& InConstraintGuid) = 0;
 
 	/*
 	*  Add Constraint channel
