Add cache reading support to the IAS cache io service

#rnx

(git:4f35389 - Martin Ridgers - [REDACTED_PHONE]:01:22 +0200)

diff --git a/Engine/Source/Runtime/Experimental/IoStoreOnDemand/Private/JournaledCache.cpp b/Engine/Source/Runtime/Experimental/IoStoreOnDemand/Private/JournaledCache.cpp
--- a/Engine/Source/Runtime/Experimental/IoStoreOnDemand/Private/JournaledCache.cpp
+++ b/Engine/Source/Runtime/Experimental/IoStoreOnDemand/Private/JournaledCache.cpp
@@ -1408,11 +1408,26 @@ class FServiceThread
 	: public FRunnable
 {
 public:
+	struct FReadSink
+	{
+		struct FReadResult	{ uint16 ReadId; uint16 Status; };
+		virtual void		OnRead(const FReadResult* Results, int32 Num) = 0;
+	};
+
+	struct FReadRequest
+	{
+		uint64				Key;
+		FIoBuffer*			Dest;
+		FReadSink*			Sink;
+		uint32				Offset = 0;
+	};
+
 	static FServiceThread&	Get();
 	void					RegisterCache(TUniquePtr<FCache> Cache);
 	void					UnregisterCache(FCache* Cache);
 	void					SetGovernorRate(uint32 Allowance, uint32 Ops, uint32 Seconds);
 	void					SetGovernorDemand(uint32 Threshold, uint32 Boost, uint32 SuperBoost);
+	uint32					BeginRead(const FCache* Cache, const FReadRequest& Request);
 
 private:
 	struct FWork
@@ -1422,16 +1437,17 @@ private:
 			Work_Unregister,
 			Work_GovDemand,
 			Work_GovRate,
+			Work_Read,
 		};
 
 		void	SetPtr(const void* In)	{ Ptr = UPTRINT(In) >> 3; check((UPTRINT(In) & 0x7) == 0); }
 		void*	GetPtr() const			{ return (FCache*)(Ptr << 3); }
 
 		union {
-			struct {			// reg / unreg
+			struct {			// reg / unreg / read
 				UPTRINT			What : 3;
 				UPTRINT			Ptr : 45;
-				UPTRINT			_Unused : 16;
+				UPTRINT			ReadId : 16;
 			};
 			struct {			// rate
 				uint16			_What0 : 3;
@@ -1446,9 +1462,14 @@ private:
 				uint16			Boost;
 				uint16			SuperBoost;
 			};
+			FReadSink*			Sink;
+		};
+		union {
+			uint64				Key;	// read
+			FIoBuffer*			Dest;	// read.param
 		};
 	};
-	static_assert(sizeof(FWork) == sizeof(UPTRINT));
+	static_assert(sizeof(FWork) == sizeof(UPTRINT) * 2);
 
 	void						StartThread();
 	int32						Update();
@@ -1463,9 +1484,11 @@ private:
 	TArray<TUniquePtr<FCache>>	Caches;
 	std::atomic_int				RunCount = 0;
 	FCriticalSection			Lock;
+	TArray<FWork>				ActiveReads;
 	TArray<FWork>				PendingWork;
 	std::atomic_int				PendingCount = 0;
 	int32						PendingPrev;
+	uint32						ReadIdCounter = 0;
 };
 
 ////////////////////////////////////////////////////////////////////////////////
@@ -1553,6 +1576,23 @@ void FServiceThread::SetGovernorDemand(uint32 Threshold, uint32 Boost, uint32 Su
 	check(Work.Threshold == Threshold && Work.Boost == Boost && Work.SuperBoost == SuperBoost);
 }
 
+////////////////////////////////////////////////////////////////////////////////
+uint32 FServiceThread::BeginRead(const FCache* Cache, const FReadRequest& Request)
+{
+	FWork Works[2];
+
+	Works[0].SetPtr(Cache);
+	Works[0].What = FWork::Work_Read;
+	Works[0].Key = Request.Key;
+	Works[0].ReadId = ReadIdCounter++;
+
+	Works[1].Sink = Request.Sink;
+	Works[1].Dest = Request.Dest;
+
+	SubmitWork(Works, 2);
+
+	return Works[0].ReadId;
+}
 
 ////////////////////////////////////////////////////////////////////////////////
 void FServiceThread::StartThread()
@@ -1619,7 +1659,45 @@ int32 FServiceThread::Update()
 		CycleSlice = FMath::Min(CyclesTillActive, CycleSlice);
 	}
 
-	return CycleSlice;
+	// Now we've a slice of time to process reads until caches need another tick
+	int64 Cycle = FPlatformTime::Cycles64();
+	int64 StopReadsCycle = Cycle + CycleSlice;
+
+	// Early out
+	if (ActiveReads.Num() == 0)
+	{
+		return CycleSlice;
+	}
+
+	TRACE_CPUPROFILER_EVENT_SCOPE(IasCache::ProcessReads);
+
+	uint32 Index = 0;
+	for (uint32 n = ActiveReads.Num(); Index < n;)
+	{
+		TRACE_CPUPROFILER_EVENT_SCOPE(IasCache::ActiveRead);
+
+		// Lets always do at least one to make progress.
+		FWork& Work = ActiveReads[Index];
+		FWork& Param = ActiveReads[Index + 1];
+		Index += 2;
+
+		auto* Cache = (FCache*)(Work.GetPtr());
+		EIoErrorCode Status = Cache->Materialize(Work.Key, *Param.Dest);
+
+		FReadSink::FReadResult Result = { uint16(Work.ReadId), uint16(Status) };
+		Param.Sink->OnRead(&Result, 1);
+
+		Cycle = FPlatformTime::Cycles64();
+		if (Cycle >= StopReadsCycle)
+		{
+			break;
+		}
+	}
+
+	check(Index > 0);
+	ActiveReads.RemoveAt(0, Index);
+
+	return CycleSlice - int32(Cycle - StopReadsCycle);
 }
 
 ////////////////////////////////////////////////////////////////////////////////
@@ -1658,8 +1736,10 @@ void FServiceThread::ReceiveWork()
 	}
 
 	// Then the rest
-	for (const FWork& Work : InboundWork)
+	for (uint32 i = 0, n = InboundWork.Num(); i < n; ++i)
 	{
+		const FWork& Work = InboundWork[i];
+
 		if (Work.What == FWork::Work_Unregister)
 		{
 			continue;
@@ -1684,7 +1764,13 @@ void FServiceThread::ReceiveWork()
 			Governor.SetDemands(Work.Threshold, Work.Boost, Work.SuperBoost);
 			continue;
 		}
+
+		check(Work.What == FWork::Work_Read);
+		ActiveReads.Add(Work);
+		ActiveReads.Add(InboundWork[++i]); // the read's params
 	}
+
+	check((ActiveReads.Num() & 1) == 0);
 }
 
 ////////////////////////////////////////////////////////////////////////////////
