Consolidated common state changing into a function

#rnx

(git:3770129 - Martin Ridgers - [REDACTED_PHONE]:39:55 +0200)

diff --git a/Engine/Source/Runtime/Experimental/IoStoreOnDemand/Private/CoreHttp/Http.cpp b/Engine/Source/Runtime/Experimental/IoStoreOnDemand/Private/CoreHttp/Http.cpp
--- a/Engine/Source/Runtime/Experimental/IoStoreOnDemand/Private/CoreHttp/Http.cpp
+++ b/Engine/Source/Runtime/Experimental/IoStoreOnDemand/Private/CoreHttp/Http.cpp
@@ -94,7 +94,6 @@ namespace UE::IO::IAS::HTTP
 	x(ActivityDestroy) \
 	x(RequestBegin) \
 	x(StateChange) \
-	x(StateChangeWait) \
 	x(Wait) \
 	x(Unwait) \
 	x(Connect) \
@@ -1427,6 +1426,7 @@ struct alignas(16) FActivity
 	enum class EWait : uint8 { None, Read, Write, Pool };
 	enum class EState : uint8
 	{
+		None,
 		Build,
 		Resolve,
 		Connect,
@@ -1442,7 +1442,7 @@ struct alignas(16) FActivity
 	};
 
 	int8				Slot = -1;
-	EState				State;
+	EState				State = EState::None;
 	EWait				SocketWait;
 	uint8				IsKeepAlive : 1;
 	uint8				NoContent : 1;
@@ -1463,15 +1463,25 @@ struct alignas(16) FActivity
 static void Activity_TraceStateNames()
 {
 	FAnsiStringView StateNames[] = {
-		"Build",		"Resolve",		"Connect",		"Send",
-		"RecvMessage",	"RecvContent",	"RecvStream",	"RecvDone",
-		"Completed",	"Cancelled",	"Failed",		"$",
+		"None", "Build", "Resolve", "Connect",	"Send", "RecvMessage",
+		"RecvContent", "RecvStream", "RecvDone", "Completed", "Cancelled",
+		"Failed", "$",
 	};
 	static_assert(UE_ARRAY_COUNT(StateNames) == int32(FActivity::EState::_Num) + 1);
 
 	TraceEnum(StateNames);
 }
 
+////////////////////////////////////////////////////////////////////////////////
+static void Activity_ChangeState(FActivity* Activity, FActivity::EState InState, uint32 Param=0)
+{
+	Trace(Activity, ETrace::StateChange, InState);
+
+	check(Activity->State != InState);
+	Activity->State = InState;
+	Activity->StateParam = Param;
+}
+
 ////////////////////////////////////////////////////////////////////////////////
 static FActivity* Activity_Alloc(uint32 BufferSize)
 {
@@ -1486,7 +1496,8 @@ static FActivity* Activity_Alloc(uint32 BufferSize)
 	uint32 ScratchSize = BufferSize;
 	Activity->Buffer = FBuffer(Scratch, ScratchSize);
 
-	Activity->State = FActivity::EState::Build;
+	Activity_ChangeState(Activity, FActivity::EState::Build);
+
 	return Activity;
 }
 
@@ -1519,8 +1530,8 @@ static void Activity_SetError(FActivity* Activity, const char* Reason)
 	Activity->IsKeepAlive = 0;
 	Activity->SocketWait = FActivity::EWait::None;
 	Activity->ErrorReason = Reason;
-	Activity->State = FActivity::EState::Failed;
-	Activity->StateParam = LastSocketResult();
+
+	Activity_ChangeState(Activity, FActivity::EState::Failed, LastSocketResult());
 
 	Trace(Activity, ETrace::Error, Activity->StateParam);
 }
@@ -1860,16 +1871,15 @@ static int32 DoResolve(FActivity* Activity)
 	int32 Result = Pool->IsResolved();
 	if (Result > 0)
 	{
-		Activity->State = FActivity::EState::Connect;
-		Trace(Activity, ETrace::StateChange, Activity->State);
+		Activity_ChangeState(Activity, FActivity::EState::Connect);
 		return 0;
 	}
 
 	if (Result < 0)
 	{
 		Activity->SocketWait = FActivity::EWait::Pool;
-		Activity->State = FActivity::EState::Connect;
-		Trace(Activity, ETrace::StateChangeWait, Activity->State);
+		Activity_ChangeState(Activity, FActivity::EState::Connect);
+		Trace(Activity, ETrace::Wait, Activity->State);
 		return 1;
 	}
 
@@ -1882,8 +1892,7 @@ static int32 DoResolve(FActivity* Activity)
 	case -2: Activity_SetError(Activity, "Unexpected address family during resolve"); return -1;
 	}
 
-	Activity->State = FActivity::EState::Connect;
-	Trace(Activity, ETrace::StateChange, Activity->State);
+	Activity_ChangeState(Activity, FActivity::EState::Connect);
 	return 0;
 }
 
@@ -1915,9 +1924,7 @@ static int32 DoConnect(FActivity* Activity)
 	{
 		Activity->Socket = MoveTemp(Candidate);
 		Activity->SocketWait = FActivity::EWait::None;
-		Activity->State = FActivity::EState::Send;
-		Activity->StateParam = 0;
-		Trace(Activity, ETrace::StateChange, Activity->State);
+		Activity_ChangeState(Activity, FActivity::EState::Send);
 		return 0;
 	}
 
@@ -1969,9 +1976,8 @@ static int32 DoConnect(FActivity* Activity)
 	Activity->Socket = MoveTemp(Candidate);
 	Activity->SocketWait = FActivity::EWait::Write;
 
-	Activity->StateParam = 0;
-	Activity->State = FActivity::EState::Send;
-	Trace(Activity, ETrace::StateChangeWait, Activity->State);
+	Activity_ChangeState(Activity, FActivity::EState::Send);
+	Trace(Activity, ETrace::Wait, Activity->State);
 	return 1;
 }
 
@@ -2016,10 +2022,9 @@ static int32 DoSend(FActivity* Activity)
 	Buffer.Reset();
 	Buffer.AdvanceUsed(sizeof(FResponseInternal));
 
-	Activity->StateParam = 0;
-	Activity->State = FActivity::EState::RecvMessage;
+	Activity_ChangeState(Activity, FActivity::EState::RecvMessage);
 	Activity->SocketWait = FActivity::EWait::Read;
-	Trace(Activity, ETrace::StateChangeWait, Activity->State);
+	Trace(Activity, ETrace::Wait, Activity->State);
 	return 1;
 }
 
@@ -2183,8 +2188,7 @@ static int32 DoRecvMessage(FActivity* Activity)
 			Activity_SetError(Activity, "Received content when none was expected");
 			return -1;
 		}
-		Activity->State = FActivity::EState::RecvDone;
-		Trace(Activity, ETrace::StateChange, Activity->State);
+		Activity_ChangeState(Activity, FActivity::EState::RecvDone);
 		return 0;
 	}
 
@@ -2192,9 +2196,8 @@ static int32 DoRecvMessage(FActivity* Activity)
 
 	const bool bStreamed = Internal.Dest->GetSize() < ContentLength;
 
-	Activity->State = bStreamed ? FActivity::EState::RecvStream : FActivity::EState::RecvContent;
-	Activity->StateParam = AlreadyReceived;
-	Trace(Activity, ETrace::StateChange, Activity->State);
+	auto NextState = bStreamed ? FActivity::EState::RecvStream : FActivity::EState::RecvContent;
+	Activity_ChangeState(Activity, NextState, AlreadyReceived);
 
 	if (AlreadyReceived == 0)
 	{
@@ -2287,8 +2290,7 @@ static FHandlerResult DoRecvContent(FActivity* Activity, uint32 MaxRecvSize)
 		return { -1 };
 	}
 
-	Activity->State = FActivity::EState::RecvDone;
-	Trace(Activity, ETrace::StateChange, Activity->State);
+	Activity_ChangeState(Activity, FActivity::EState::RecvDone);
 	return { 0, RecvSize };
 }
 
@@ -2304,8 +2306,7 @@ static int32 DoRecvDone(FActivity* Activity)
 	FTicketStatus& SinkArg = *(FTicketStatus*)Activity;
 	Activity->Sink(SinkArg);
 
-	Activity->State = FActivity::EState::Completed;
-	Trace(Activity, ETrace::StateChange, Activity->State);
+	Activity_ChangeState(Activity, FActivity::EState::Completed);
 	return 0;
 }
 
@@ -2326,8 +2327,7 @@ static void DoCancel(FActivity* Activity)
 		return;
 	}
 
-	Activity->State = FActivity::EState::Cancelled;
-	Trace(Activity, ETrace::StateChange, Activity->State);
+	Activity_ChangeState(Activity, FActivity::EState::Cancelled);
 
 	FTicketStatus& SinkArg = *(FTicketStatus*)Activity;
 	Activity->Sink(SinkArg);
@@ -2820,10 +2820,11 @@ FTicket FEventLoop::Send(FRequest&& Request, FTicketSink Sink, UPTRINT SinkParam
 {
 	FActivity* Activity = nullptr;
 	Swap(Activity, Request.Ptr);
-	Activity->State = FActivity::EState::Resolve;
 	Activity->SinkParam = SinkParam;
 	Activity->Sink = Sink;
 
+	Activity_ChangeState(Activity, FActivity::EState::Resolve);
+
 	return Impl->Send(Activity);
 }
 
