Removed unused event

#rnx

(git:1db0e5d - Martin Ridgers - [REDACTED_PHONE]:57:53 +0200)

diff --git a/Engine/Source/Runtime/Experimental/IoStoreOnDemand/Private/CoreHttp/Http.cpp b/Engine/Source/Runtime/Experimental/IoStoreOnDemand/Private/CoreHttp/Http.cpp
--- a/Engine/Source/Runtime/Experimental/IoStoreOnDemand/Private/CoreHttp/Http.cpp
+++ b/Engine/Source/Runtime/Experimental/IoStoreOnDemand/Private/CoreHttp/Http.cpp
@@ -99,7 +99,6 @@ namespace UE::IO::IAS::HTTP
 	x(Connect) \
 	x(Send) \
 	x(Recv) \
-	x(Error) \
 	x($)
 
 enum class ETrace
@@ -1532,8 +1531,6 @@ static void Activity_SetError(FActivity* Activity, const char* Reason)
 	Activity->ErrorReason = Reason;
 
 	Activity_ChangeState(Activity, FActivity::EState::Failed, LastSocketResult());
-
-	Trace(Activity, ETrace::Error, Activity->StateParam);
 }
 
 // {{{1 request ................................................................
