Wrapped activities' begin/end waiting

#rnx

(git:500e91c - Martin Ridgers - [REDACTED_PHONE]:00:20 +0200)

diff --git a/Engine/Source/Runtime/Experimental/IoStoreOnDemand/Private/CoreHttp/Http.cpp b/Engine/Source/Runtime/Experimental/IoStoreOnDemand/Private/CoreHttp/Http.cpp
--- a/Engine/Source/Runtime/Experimental/IoStoreOnDemand/Private/CoreHttp/Http.cpp
+++ b/Engine/Source/Runtime/Experimental/IoStoreOnDemand/Private/CoreHttp/Http.cpp
@@ -1442,7 +1442,7 @@ struct alignas(16) FActivity
 
 	int8				Slot = -1;
 	EState				State = EState::None;
-	EWait				SocketWait;
+	EWait				SocketWait = EWait::None;
 	uint8				IsKeepAlive : 1;
 	uint8				NoContent : 1;
 	uint8				_Unused0 : 6;
@@ -1481,6 +1481,24 @@ static void Activity_ChangeState(FActivity* Activity, FActivity::EState InState,
 	Activity->StateParam = Param;
 }
 
+////////////////////////////////////////////////////////////////////////////////
+static void Activity_BeginWait(FActivity* Activity, FActivity::EWait What)
+{
+	check(Activity->SocketWait == FActivity::EWait::None);
+	Activity->SocketWait = What;
+	Trace(Activity, ETrace::Wait);
+}
+
+////////////////////////////////////////////////////////////////////////////////
+static void Activity_EndWait(FActivity* Activity)
+{
+	if (Activity->SocketWait != FActivity::EWait::None)
+	{
+		Trace(Activity, ETrace::Unwait);
+	}
+	Activity->SocketWait = FActivity::EWait::None;
+}
+
 ////////////////////////////////////////////////////////////////////////////////
 static FActivity* Activity_Alloc(uint32 BufferSize)
 {
@@ -1527,9 +1545,9 @@ static void Activity_Free(FActivity* Activity)
 static void Activity_SetError(FActivity* Activity, const char* Reason)
 {
 	Activity->IsKeepAlive = 0;
-	Activity->SocketWait = FActivity::EWait::None;
 	Activity->ErrorReason = Reason;
 
+	Activity_EndWait(Activity);
 	Activity_ChangeState(Activity, FActivity::EState::Failed, LastSocketResult());
 }
 
@@ -1807,9 +1825,8 @@ static uint64 ReadyCheck(FActivity** Activities, uint32 Num, uint32 TimeoutMs)
 		case EWait::Pool:
 			if (Activity->Pool->GetIpAddress() > 0x00ff'ffff)
 			{
-				Activities[i]->SocketWait = EWait::None;
+				Activity_EndWait(Activities[i]);
 				Ret |= (1ull << Activity->Slot);
-				Trace(Activity, ETrace::Unwait);
 			}
 			break;
 
@@ -1841,9 +1858,8 @@ static uint64 ReadyCheck(FActivity** Activities, uint32 Num, uint32 TimeoutMs)
 		}
 
 		FActivity* Activity = *Cursor;
-		Activity->SocketWait = EWait::None;
+		Activity_EndWait(Activity);
 		Ret |= (1ull << Activity->Slot);
-		Trace(Activity, ETrace::Unwait);
 	}
 
 	return Ret;
@@ -1874,9 +1890,8 @@ static int32 DoResolve(FActivity* Activity)
 
 	if (Result < 0)
 	{
-		Activity->SocketWait = FActivity::EWait::Pool;
 		Activity_ChangeState(Activity, FActivity::EState::Connect);
-		Trace(Activity, ETrace::Wait, Activity->State);
+		Activity_BeginWait(Activity, FActivity::EWait::Pool);
 		return 1;
 	}
 
@@ -1920,7 +1935,7 @@ static int32 DoConnect(FActivity* Activity)
 	if (Candidate.IsValid())
 	{
 		Activity->Socket = MoveTemp(Candidate);
-		Activity->SocketWait = FActivity::EWait::None;
+		Activity_EndWait(Activity);
 		Activity_ChangeState(Activity, FActivity::EState::Send);
 		return 0;
 	}
@@ -1971,10 +1986,9 @@ static int32 DoConnect(FActivity* Activity)
 	}
 
 	Activity->Socket = MoveTemp(Candidate);
-	Activity->SocketWait = FActivity::EWait::Write;
 
 	Activity_ChangeState(Activity, FActivity::EState::Send);
-	Trace(Activity, ETrace::Wait, Activity->State);
+	Activity_BeginWait(Activity, FActivity::EWait::Write);
 	return 1;
 }
 
@@ -2003,8 +2017,7 @@ static int32 DoSend(FActivity* Activity)
 	case FSocket::EResult::ConnectError:Activity_SetError(Activity, "Connection error"); return -1;
 	case FSocket::EResult::Wait:
 		Activity->StateParam = Remaining;
-		Activity->SocketWait = FActivity::EWait::Write;
-		Trace(Activity, ETrace::Wait);
+		Activity_BeginWait(Activity, FActivity::EWait::Write);
 		return 1;
 	}
 
@@ -2020,8 +2033,7 @@ static int32 DoSend(FActivity* Activity)
 	Buffer.AdvanceUsed(sizeof(FResponseInternal));
 
 	Activity_ChangeState(Activity, FActivity::EState::RecvMessage);
-	Activity->SocketWait = FActivity::EWait::Read;
-	Trace(Activity, ETrace::Wait, Activity->State);
+	Activity_BeginWait(Activity, FActivity::EWait::Read);
 	return 1;
 }
 
@@ -2045,8 +2057,7 @@ static int32 DoRecvMessage(FActivity* Activity)
 
 		if (Result == int32(FSocket::EResult::Wait))
 		{
-			Activity->SocketWait = FActivity::EWait::Read;
-			Trace(Activity, ETrace::Wait);
+			Activity_BeginWait(Activity, FActivity::EWait::Read);
 			return 1;
 		}
 
@@ -2260,8 +2271,7 @@ static FHandlerResult DoRecvContent(FActivity* Activity, uint32 MaxRecvSize)
 
 		if (Result == int32(FSocket::EResult::Wait))
 		{
-			Activity->SocketWait = FActivity::EWait::Read;
-			Trace(Activity, ETrace::Wait);
+			Activity_BeginWait(Activity, FActivity::EWait::Read);
 			return { 1, RecvSize };
 		}
 
