Added simple reporting of IAS activity time

#rnx

(git:4e7a2b5 - Martin Ridgers - [REDACTED_PHONE]:06:54 +0200)

diff --git a/Engine/Source/Runtime/Experimental/IoStoreOnDemand/Private/CoreHttp/Client.h b/Engine/Source/Runtime/Experimental/IoStoreOnDemand/Private/CoreHttp/Client.h
--- a/Engine/Source/Runtime/Experimental/IoStoreOnDemand/Private/CoreHttp/Client.h
+++ b/Engine/Source/Runtime/Experimental/IoStoreOnDemand/Private/CoreHttp/Client.h
@@ -6,6 +6,10 @@
 #include "Containers/StringView.h"
 #endif
 
+#if !defined(IAS_HTTP_WITH_PERF)
+#	define IAS_HTTP_WITH_PERF !UE_BUILD_SHIPPING
+#endif
+
 #define UE_API
 
 ////////////////////////////////////////////////////////////////////////////////
@@ -127,6 +131,22 @@ private:
 	FResponse&			operator = (FResponse&&) = delete;
 };
 
+////////////////////////////////////////////////////////////////////////////////
+class UE_API FTicketPerf
+{
+#if IAS_HTTP_WITH_PERF
+public:
+	struct FSample
+	{
+		uint32			TotalMs;
+		uint32			WaitMs;
+	};
+
+	FSample				GetSendSample() const;
+	FSample				GetRecvSample() const;
+#endif // IAS_HTTP_WITH_PERF
+};
+
 ////////////////////////////////////////////////////////////////////////////////
 class UE_API FTicketStatus
 {
@@ -139,6 +159,7 @@ public:
 	FResponse&			GetResponse() const;		// if GetId() == EId::Response
 	const FIoBuffer&	GetContent() const;			// if GetId() == EId::Content
 	uint32				GetContentLength() const;	//  |
+	const FTicketPerf&	GetPerf() const;			// _|_
 	const char*			GetErrorReason() const;		// if GetId() == EId::Error
 
 private:
diff --git a/Engine/Source/Runtime/Experimental/IoStoreOnDemand/Private/CoreHttp/Http.cpp b/Engine/Source/Runtime/Experimental/IoStoreOnDemand/Private/CoreHttp/Http.cpp
--- a/Engine/Source/Runtime/Experimental/IoStoreOnDemand/Private/CoreHttp/Http.cpp
+++ b/Engine/Source/Runtime/Experimental/IoStoreOnDemand/Private/CoreHttp/Http.cpp
@@ -1407,6 +1407,60 @@ bool FConnectionPool::Resolve()
 
 // {{{1 activity ...............................................................
 
+#if IAS_HTTP_WITH_PERF
+
+////////////////////////////////////////////////////////////////////////////////
+class FStopwatch
+{
+public:
+	struct FInterval
+	{
+					FInterval() : Elapsed(0), Counter(0) {}
+		int64		Elapsed : 48;
+		int64		Counter : 16;
+	};
+
+	struct FLapTime
+	{
+		FInterval	Total;
+		FInterval	Wait;
+	};
+
+	void		Start()		{ Laps[Index].Total.Elapsed -= Sample(); }
+	void		Stop()		{ Laps[Index].Total.Elapsed += Sample(); }
+	void		Wait()		{ Laps[Index].Wait.Elapsed -= Sample(); Laps[Index].Wait.Counter++; }
+	void		Unwait()	{ Laps[Index].Wait.Elapsed += Sample(); }
+	void		Lap()		{ ++Index; check(Index < UE_ARRAY_COUNT(Laps)); }
+	void		AddCount()	{ Laps[Index].Total.Counter++; }
+	int64		Sample();
+
+	const FLapTime&	GetLap(uint32 i) const
+	{
+		check(i < UE_ARRAY_COUNT(Laps));
+		return Laps[i];
+	}
+
+private:
+	FLapTime	Laps[2];
+	uint32		Index = 0;
+};
+
+////////////////////////////////////////////////////////////////////////////////
+int64 FStopwatch::Sample()
+{
+	int64 Value = FPlatformTime::Cycles64();
+	static int64 Base = 0;
+	if (Base == 0)
+	{
+		Base = Value;
+		return 0;
+	}
+
+	return Value - Base;
+}
+
+#endif // IAS_HTTP_WITH_PERF
+
 ////////////////////////////////////////////////////////////////////////////////
 struct FResponseInternal
 {
@@ -1447,14 +1501,14 @@ struct alignas(16) FActivity
 	uint8				NoContent : 1;
 	uint8				_Unused0 : 6;
 	uint32				StateParam = 0;
-
+#if IAS_HTTP_WITH_PERF
+	FStopwatch			Stopwatch;
+#endif
 	FSocketPool*		Pool;
 	const char*			ErrorReason;
 	UPTRINT				SinkParam;
 	FTicketSink			Sink;
-
 	FSocket				Socket;
-
 	FBuffer				Buffer;
 };
 
@@ -1476,6 +1530,29 @@ static void Activity_ChangeState(FActivity* Activity, FActivity::EState InState,
 {
 	Trace(Activity, ETrace::StateChange, InState);
 
+#if IAS_HTTP_WITH_PERF
+	using EState = FActivity::EState;
+
+	FStopwatch& Stopwatch = Activity->Stopwatch;
+	if (InState == EState::Send)
+	{
+		Stopwatch.Start();
+	}
+	else if (Activity->State == EState::Send)
+	{
+		Stopwatch.Stop();
+		Stopwatch.Lap();
+	}
+	else if (InState == EState::RecvContent || InState == EState::RecvStream)
+	{
+		Stopwatch.Start();
+	}
+	else if (Activity->State == EState::RecvContent || Activity->State == EState::RecvStream)
+	{
+		Stopwatch.Stop();
+	}
+#endif // IAS_HTTP_WITH_PERF
+
 	check(Activity->State != InState);
 	Activity->State = InState;
 	Activity->StateParam = Param;
@@ -1485,8 +1562,16 @@ static void Activity_ChangeState(FActivity* Activity, FActivity::EState InState,
 static void Activity_BeginWait(FActivity* Activity, FActivity::EWait What)
 {
 	check(Activity->SocketWait == FActivity::EWait::None);
-	Activity->SocketWait = What;
+
+#if IAS_HTTP_WITH_PERF
+	if (Activity->State >= FActivity::EState::Send)
+	{
+		Activity->Stopwatch.Wait();
+	}
+#endif
 	Trace(Activity, ETrace::Wait);
+
+	Activity->SocketWait = What;
 }
 
 ////////////////////////////////////////////////////////////////////////////////
@@ -1494,8 +1579,15 @@ static void Activity_EndWait(FActivity* Activity)
 {
 	if (Activity->SocketWait != FActivity::EWait::None)
 	{
+#if IAS_HTTP_WITH_PERF
+		if (Activity->State >= FActivity::EState::Send)
+		{
+			Activity->Stopwatch.Unwait();
+		}
+#endif
 		Trace(Activity, ETrace::Unwait);
 	}
+
 	Activity->SocketWait = FActivity::EWait::None;
 }
 
@@ -1781,6 +1873,14 @@ uint32 FTicketStatus::GetContentLength() const
 	return Response.ContentLength;
 }
 
+////////////////////////////////////////////////////////////////////////////////
+const FTicketPerf& FTicketStatus::GetPerf() const
+{
+	check(GetId() == EId::Content);
+	const auto* Activity = (FActivity*)this;
+	return *(FTicketPerf*)Activity;
+}
+
 ////////////////////////////////////////////////////////////////////////////////
 const FIoBuffer& FTicketStatus::GetContent() const
 {
@@ -1800,6 +1900,46 @@ const char* FTicketStatus::GetErrorReason() const
 
 
 
+// {{{1 perf ...................................................................
+
+#if IAS_HTTP_WITH_PERF
+
+////////////////////////////////////////////////////////////////////////////////
+static FTicketPerf::FSample GetPerfSample(const FActivity* Activity, uint32 Index)
+{
+	static uint64 Freq;
+	if (Freq == 0)
+	{
+		Freq = uint64(1.0 / FPlatformTime::GetSecondsPerCycle());
+	}
+
+	auto ToMs = [] (uint64 Value) { return uint32((Value * 1000ull) / Freq); };
+
+	const FStopwatch::FLapTime& LapTime = Activity->Stopwatch.GetLap(Index);
+	return {
+		ToMs(LapTime.Total.Elapsed),
+		ToMs(LapTime.Wait.Elapsed),
+	};
+}
+
+////////////////////////////////////////////////////////////////////////////////
+FTicketPerf::FSample FTicketPerf::GetSendSample() const
+{
+	const auto* Activity = (FActivity*)this;
+	return GetPerfSample(Activity, 0);
+}
+
+////////////////////////////////////////////////////////////////////////////////
+FTicketPerf::FSample FTicketPerf::GetRecvSample() const
+{
+	const auto* Activity = (FActivity*)this;
+	return GetPerfSample(Activity, 1);
+}
+
+#endif // IAS_HTTP_WITH_PERF
+
+
+
 // {{{1 event-loop-int .........................................................
 
 ////////////////////////////////////////////////////////////////////////////////
@@ -2021,6 +2161,10 @@ static int32 DoSend(FActivity* Activity)
 		return 1;
 	}
 
+#if IAS_HTTP_WITH_PERF
+	Activity->Stopwatch.AddCount();
+#endif
+
 	checkf(Result > 0, TEXT("Result wasn't caught by switch statement so it is expected to be a positive amount of bytes sent"));
 	Remaining = SendSize - Result;
 	if (Remaining != 0)
@@ -2287,6 +2431,10 @@ static FHandlerResult DoRecvContent(FActivity* Activity, uint32 MaxRecvSize)
 			return { -1 };
 		}
 
+#if IAS_HTTP_WITH_PERF
+		Activity->Stopwatch.AddCount();
+#endif
+
 		Activity->StateParam += Result;
 		RecvSize += Result;
 	}
