Apease static analysis

#rnx

diff --git a/Engine/Source/Runtime/Experimental/IoStoreOnDemand/Private/JournaledCache.cpp b/Engine/Source/Runtime/Experimental/IoStoreOnDemand/Private/JournaledCache.cpp
--- a/Engine/Source/Runtime/Experimental/IoStoreOnDemand/Private/JournaledCache.cpp
+++ b/Engine/Source/Runtime/Experimental/IoStoreOnDemand/Private/JournaledCache.cpp
@@ -1579,7 +1579,7 @@ void FServiceThread::SetGovernorDemand(uint32 Threshold, uint32 Boost, uint32 Su
 ////////////////////////////////////////////////////////////////////////////////
 uint32 FServiceThread::BeginRead(const FCache* Cache, const FReadRequest& Request)
 {
-	FWork Works[2];
+	FWork Works[2] = {};
 
 	Works[0].SetPtr(Cache);
 	Works[0].What = FWork::Work_Read;
