Add a timeout value when testing latency to avoid long waits when network trouble occurs.

- Originally UE::IO::IAS::HTTP::LatencyTest was only being issued as an async task, so long delays did not matter but now it is being called in a blocking manner to test if an IAS connection to a CDN is valid. This could cause the game to effectively lock up for very long lengths of time under certain bad network conditions.

diff --git a/Engine/Source/Runtime/Experimental/IoStoreOnDemand/Private/CoreHttp/LatencyTesting.cpp b/Engine/Source/Runtime/Experimental/IoStoreOnDemand/Private/CoreHttp/LatencyTesting.cpp
--- a/Engine/Source/Runtime/Experimental/IoStoreOnDemand/Private/CoreHttp/LatencyTesting.cpp
+++ b/Engine/Source/Runtime/Experimental/IoStoreOnDemand/Private/CoreHttp/LatencyTesting.cpp
@@ -12,7 +12,7 @@
 namespace UE::IO::IAS::HTTP
 {
 
-void LatencyTest(FStringView InUrl, FStringView InPath, TArrayView<int32> OutResults)
+void LatencyTest(FStringView InUrl, FStringView InPath, uint32 InTimeOutMs, TArrayView<int32> OutResults)
 {
 	auto AnsiUrl = StringCast<ANSICHAR>(InUrl.GetData(), InUrl.Len());
 
@@ -45,7 +45,7 @@ void LatencyTest(FStringView InUrl, FStringView InPath, TArrayView<int32> OutRes
 			});
 
 		uint64 Cycles = FPlatformTime::Cycles64();
-		while (Loop.Tick(-1) != 0);
+		while (Loop.Tick(InTimeOutMs) != 0);
 		Cycles = FPlatformTime::Cycles64() - Cycles;
 
 		Result = Ok ? int32(Cycles) : -1;
diff --git a/Engine/Source/Runtime/Experimental/IoStoreOnDemand/Private/CoreHttp/LatencyTesting.h b/Engine/Source/Runtime/Experimental/IoStoreOnDemand/Private/CoreHttp/LatencyTesting.h
--- a/Engine/Source/Runtime/Experimental/IoStoreOnDemand/Private/CoreHttp/LatencyTesting.h
+++ b/Engine/Source/Runtime/Experimental/IoStoreOnDemand/Private/CoreHttp/LatencyTesting.h
@@ -8,6 +8,6 @@
 namespace UE::IO::IAS::HTTP
 {
 
-void LatencyTest(FStringView InUrl, FStringView InPath, TArrayView<int32> OutResults);
+void LatencyTest(FStringView InUrl, FStringView InPath, uint32 InTimeOutMs, TArrayView<int32> OutResults);
 
 } // namespace UE::IO::IAS::HTTP
diff --git a/Engine/Source/Runtime/Experimental/IoStoreOnDemand/Private/OnDemandIoDispatcherBackend.cpp b/Engine/Source/Runtime/Experimental/IoStoreOnDemand/Private/OnDemandIoDispatcherBackend.cpp
--- a/Engine/Source/Runtime/Experimental/IoStoreOnDemand/Private/OnDemandIoDispatcherBackend.cpp
+++ b/Engine/Source/Runtime/Experimental/IoStoreOnDemand/Private/OnDemandIoDispatcherBackend.cpp
@@ -78,6 +78,13 @@ static FAutoConsoleVariableRef CVar_IasMaxHttpRetryCount(
 	TEXT("Max number of HTTP request retries before failing the I/O request.")
 );
 
+int32 GIasHttpTimeOutMs = 30 * 1000;
+static FAutoConsoleVariableRef CVar_IasHttpTimeOutMs(
+	TEXT("ias.HttpTimeOutMs"),
+	GIasHttpTimeOutMs,
+	TEXT("Time out value for HTTP requests in milliseconds")
+);
+
 int32 GIasHttpHealthCheckWaitTime = 3000;
 static FAutoConsoleVariableRef CVar_IasHttpHealthCheckWaitTime(
 	TEXT("ias.HttpHealthCheckWaitTime"),
@@ -154,8 +161,10 @@ static FAutoConsoleCommand CVar_IasAbandonCache(
 #if !UE_BUILD_SHIPPING
 static void LatencyTest(FStringView Url, FStringView Path)
 {
+	TRACE_CPUPROFILER_EVENT_SCOPE(IasBackend::LatencyTest);
+
 	int32 Results[4] = {};
-	UE::IO::IAS::HTTP::LatencyTest(Url, Path, MakeArrayView(Results));
+	UE::IO::IAS::HTTP::LatencyTest(Url, Path, GIasHttpTimeOutMs, MakeArrayView(Results));
 	UE_LOG(LogIas, Log, TEXT("Endpoint '%s' latency test (ms): %d %d %d %d"),
 		Url.GetData(), Results[0], Results[1], Results[2], Results[3]);
 }
@@ -163,10 +172,12 @@ static void LatencyTest(FStringView Url, FStringView Path)
 ///////////////////////////////////////////////////////////////////////////////
 static int32 LatencyTest(TConstArrayView<FString> Urls, FStringView Path, std::atomic_bool& bCancel)
 {
+	TRACE_CPUPROFILER_EVENT_SCOPE(IasBackend::LatencyTest);
+
 	for (int32 Idx = 0; Idx < Urls.Num() && !bCancel.load(std::memory_order_relaxed); ++Idx)
 	{
 		int32 LatencyMs = -1;
-		UE::IO::IAS::HTTP::LatencyTest(Urls[Idx], Path, MakeArrayView(&LatencyMs, 1));
+		UE::IO::IAS::HTTP::LatencyTest(Urls[Idx], Path, GIasHttpTimeOutMs, MakeArrayView(&LatencyMs, 1));
 		if (LatencyMs > 0)
 		{
 			return Idx;
diff --git a/Engine/Source/Runtime/Experimental/IoStoreOnDemand/Private/PrimeEndpoint.cpp b/Engine/Source/Runtime/Experimental/IoStoreOnDemand/Private/PrimeEndpoint.cpp
--- a/Engine/Source/Runtime/Experimental/IoStoreOnDemand/Private/PrimeEndpoint.cpp
+++ b/Engine/Source/Runtime/Experimental/IoStoreOnDemand/Private/PrimeEndpoint.cpp
@@ -93,8 +93,10 @@ static void PrintHelp()
 static int32 PingCDN(const FString& Url, const FString& TocPath)
 {
 #if !UE_BUILD_SHIPPING
+	const uint32 TimeoutMs = 30 * 1000;
+
 	int32 Results[4] = {};
-	UE::IO::IAS::HTTP::LatencyTest(Url, TocPath, MakeArrayView(Results));
+	UE::IO::IAS::HTTP::LatencyTest(Url, TocPath, TimeoutMs, MakeArrayView(Results));
 
 	// Ignore the first result as it might be wildly higher due to caching on the CDN
 	return (Results[1] + Results[2] + Results[3]) / 3;
