Cancellable IAS cache reads

#rnx

(git:f0667a4 - Martin Ridgers - [REDACTED_PHONE]:03:31 +0200)

diff --git a/Engine/Source/Runtime/Experimental/IoStoreOnDemand/Private/JournaledCache.cpp b/Engine/Source/Runtime/Experimental/IoStoreOnDemand/Private/JournaledCache.cpp
--- a/Engine/Source/Runtime/Experimental/IoStoreOnDemand/Private/JournaledCache.cpp
+++ b/Engine/Source/Runtime/Experimental/IoStoreOnDemand/Private/JournaledCache.cpp
@@ -1428,6 +1428,7 @@ public:
 	void					SetGovernorRate(uint32 Allowance, uint32 Ops, uint32 Seconds);
 	void					SetGovernorDemand(uint32 Threshold, uint32 Boost, uint32 SuperBoost);
 	uint32					BeginRead(const FCache* Cache, const FReadRequest& Request);
+	void					CancelRead(const void* GivenDest);
 
 private:
 	struct FWork
@@ -1438,13 +1439,14 @@ private:
 			Work_GovDemand,
 			Work_GovRate,
 			Work_Read,
+			Work_Cancel,
 		};
 
 		void	SetPtr(const void* In)	{ Ptr = UPTRINT(In) >> 3; check((UPTRINT(In) & 0x7) == 0); }
 		void*	GetPtr() const			{ return (FCache*)(Ptr << 3); }
 
 		union {
-			struct {			// reg / unreg / read
+			struct {			// reg / unreg / read / cancel
 				UPTRINT			What : 3;
 				UPTRINT			Ptr : 45;
 				UPTRINT			ReadId : 16;
@@ -1594,6 +1596,15 @@ uint32 FServiceThread::BeginRead(const FCache* Cache, const FReadRequest& Reques
 	return Works[0].ReadId;
 }
 
+////////////////////////////////////////////////////////////////////////////////
+void FServiceThread::CancelRead(const void* GivenData)
+{
+	FWork Work;
+	Work.What = FWork::Work_Cancel;
+	Work.SetPtr(GivenData);
+	SubmitWork(&Work, 1);
+}
+
 ////////////////////////////////////////////////////////////////////////////////
 void FServiceThread::StartThread()
 {
@@ -1681,8 +1692,13 @@ int32 FServiceThread::Update()
 		FWork& Param = ActiveReads[Index + 1];
 		Index += 2;
 
-		auto* Cache = (FCache*)(Work.GetPtr());
-		EIoErrorCode Status = Cache->Materialize(Work.Key, *Param.Dest);
+		// A read is marked as cancelled by setting its destination to a nullptr
+		EIoErrorCode Status = EIoErrorCode::Cancelled;
+		if (Param.Dest != nullptr)
+		{
+			auto* Cache = (FCache*)(Work.GetPtr());
+			Status = Cache->Materialize(Work.Key, *Param.Dest);
+		}
 
 		FReadSink::FReadResult Result = { uint16(Work.ReadId), uint16(Status) };
 		Param.Sink->OnRead(&Result, 1);
@@ -1735,12 +1751,12 @@ void FServiceThread::ReceiveWork()
 		}
 	}
 
-	// Then the rest
+	// ...then the rest
 	for (uint32 i = 0, n = InboundWork.Num(); i < n; ++i)
 	{
 		const FWork& Work = InboundWork[i];
 
-		if (Work.What == FWork::Work_Unregister)
+		if (Work.What == FWork::Work_Unregister || Work.What == FWork::Work_Cancel)
 		{
 			continue;
 		}
@@ -1770,6 +1786,25 @@ void FServiceThread::ReceiveWork()
 		ActiveReads.Add(InboundWork[++i]); // the read's params
 	}
 
+	// ...and finally the cancels
+	for (const FWork& Work : InboundWork)
+	{
+		if (Work.What != FWork::Work_Cancel)
+		{
+			continue;
+		}
+
+		const void* ToCancel = Work.GetPtr();
+		for (uint32 i = 1, n = ActiveReads.Num(); i < n; i += 2)
+		{
+			if (ActiveReads[i].Dest == ToCancel)
+			{
+				ActiveReads[i].Dest = nullptr;
+				break;
+			}
+		}
+	}
+
 	check((ActiveReads.Num() & 1) == 0);
 }
 
