Constraints in LS follow-up
    - fixed UConstraintSubsystem::Get when GEngine null or uninitialized

#rnx

diff --git a/Engine/Source/Runtime/Experimental/Animation/Constraints/Private/ConstraintSubsystem.cpp b/Engine/Source/Runtime/Experimental/Animation/Constraints/Private/ConstraintSubsystem.cpp
--- a/Engine/Source/Runtime/Experimental/Animation/Constraints/Private/ConstraintSubsystem.cpp
+++ b/Engine/Source/Runtime/Experimental/Animation/Constraints/Private/ConstraintSubsystem.cpp
@@ -5,6 +5,7 @@
 #include "Engine/Engine.h"
 #include "ConstraintsManager.h"
 #include "TransformConstraint.h"
+#include "Misc/CoreDelegates.h"
 
 //needs to be static to avoid system getting deleted with dangling handles.
 FDelegateHandle UConstraintSubsystem::OnWorldInitHandle;
@@ -18,10 +19,24 @@ void UConstraintSubsystem::Initialize(FSubsystemCollectionBase& Collection)
 {
 	Super::Initialize(Collection);
 
+	if (GEngine && GEngine->IsInitialized())
+	{
+		RegisterWorldDelegates();
+	}
+	else
+	{
+		FCoreDelegates::OnPostEngineInit.AddUObject(this, &UConstraintSubsystem::RegisterWorldDelegates);
+	}
+	
+	SetFlags(RF_Transactional);
+}
+
+void UConstraintSubsystem::RegisterWorldDelegates()
+{
 	OnWorldInitHandle = FWorldDelegates::OnPreWorldInitialization.AddStatic(&UConstraintSubsystem::OnWorldInit);
 	OnWorldCleanupHandle = FWorldDelegates::OnWorldCleanup.AddStatic(&UConstraintSubsystem::OnWorldCleanup);
 	
-	SetFlags(RF_Transactional);
+	FCoreDelegates::OnPostEngineInit.RemoveAll(this);
 }
 
 void UConstraintSubsystem::Deinitialize()
@@ -38,7 +53,11 @@ void UConstraintSubsystem::Deinitialize()
 
 UConstraintSubsystem* UConstraintSubsystem::Get()
 {
-	return GEngine->GetEngineSubsystem<UConstraintSubsystem>();
+	if (GEngine && GEngine->IsInitialized())
+	{
+		return GEngine->GetEngineSubsystem<UConstraintSubsystem>();
+	}
+	return nullptr;
 }
 
 const FConstraintsInWorld* UConstraintSubsystem::ConstraintsInWorldFind(UWorld* InWorld) const
@@ -182,22 +201,25 @@ void UConstraintSubsystem::PostEditUndo()
 
 void UConstraintSubsystem::OnWorldInit(UWorld* InWorld, const UWorld::InitializationValues IVS)
 {
-	UConstraintSubsystem* System = UConstraintSubsystem::Get();
-	FConstraintsInWorld &ConstraintInWorld = System->ConstraintsInWorldFindOrAdd(InWorld);
-	ConstraintInWorld.Init(InWorld);
+	if (UConstraintSubsystem* System = UConstraintSubsystem::Get())
+	{
+		FConstraintsInWorld &ConstraintInWorld = System->ConstraintsInWorldFindOrAdd(InWorld);
+		ConstraintInWorld.Init(InWorld);
+	}
 }
 
 void UConstraintSubsystem::OnWorldCleanup(UWorld* InWorld, bool bSessionEnded, bool bCleanupResources)
 {
-	UConstraintSubsystem* System = UConstraintSubsystem::Get();
-
-	for (int32 Index = System->ConstraintsInWorld.Num() - 1; Index >= 0; --Index)
+	if(UConstraintSubsystem* System = UConstraintSubsystem::Get())
 	{
-		if (System->ConstraintsInWorld[Index].World.Get() == InWorld)
+		for (int32 Index = System->ConstraintsInWorld.Num() - 1; Index >= 0; --Index)
 		{
-			System->ConstraintsInWorld[Index].RemoveConstraints(InWorld);
-			System->ConstraintsInWorld.RemoveAt(Index);
-			break;
+			if (System->ConstraintsInWorld[Index].World.Get() == InWorld)
+			{
+				System->ConstraintsInWorld[Index].RemoveConstraints(InWorld);
+				System->ConstraintsInWorld.RemoveAt(Index);
+				break;
+			}
 		}
 	}
 }
diff --git a/Engine/Source/Runtime/Experimental/Animation/Constraints/Private/ConstraintSubsystem.h b/Engine/Source/Runtime/Experimental/Animation/Constraints/Private/ConstraintSubsystem.h
--- a/Engine/Source/Runtime/Experimental/Animation/Constraints/Private/ConstraintSubsystem.h
+++ b/Engine/Source/Runtime/Experimental/Animation/Constraints/Private/ConstraintSubsystem.h
@@ -82,6 +82,7 @@ private:
 	//handles for handling world creation init and teardown
 	static FDelegateHandle OnWorldInitHandle;
 	static FDelegateHandle OnWorldCleanupHandle;
+	void RegisterWorldDelegates();
 
 	const FConstraintsInWorld* ConstraintsInWorldFind(UWorld* InWorld) const;
 	FConstraintsInWorld* ConstraintsInWorldFind(UWorld* InWorld);
