Split host-handling parts of the pool into a separate class

#rnx

(git:919b9b4 - Martin Ridgers - [REDACTED_PHONE]:39:43 +0200)

diff --git a/Engine/Source/Runtime/Experimental/IoStoreOnDemand/Private/CoreHttp/Http.cpp b/Engine/Source/Runtime/Experimental/IoStoreOnDemand/Private/CoreHttp/Http.cpp
--- a/Engine/Source/Runtime/Experimental/IoStoreOnDemand/Private/CoreHttp/Http.cpp
+++ b/Engine/Source/Runtime/Experimental/IoStoreOnDemand/Private/CoreHttp/Http.cpp
@@ -1170,23 +1170,19 @@
 // {{{1 connection-pool ........................................................
 
 ////////////////////////////////////////////////////////////////////////////////
-class FSocketPool
+class FHost
 {
 public:
 	enum class EDirection : uint8 { Send, Recv };
 
-					FSocketPool(const ANSICHAR* InHostName, uint32 InPort, uint32 InMaxLeases);
-					~FSocketPool();
-	static uint32	GetAllocSize(uint32 MaxLeases);
-	bool			LeaseSocket(FSocket& Out);
-	void			ReturnLease(FSocket&& Socket);
+					FHost(const ANSICHAR* InHostName, uint32 InPort);
 	void			SetBufferSize(EDirection Dir, int32 Size);
 	int32			GetBufferSize(EDirection Dir) const;
 	int32			IsResolved() const;
 	int32			ResolveHostName();
-	uint32			GetIpAddress() const	{ return IpAddresses[0]; }
-	FAnsiStringView	GetHostName() const		{ return HostName; }
-	uint32			GetPort() const			{ return Port; }
+	uint32			GetIpAddress() const		{ return IpAddresses[0]; }
+	FAnsiStringView	GetHostName() const			{ return HostName; }
+	uint32			GetPort() const				{ return Port; }
 
 private:
 	const ANSICHAR*	HostName;
@@ -1194,89 +1190,29 @@
 	int16			SendBufKb = -1;
 	int16			RecvBufKb = -1;
 	uint16			Port;
-	uint8			LeaseCount = 0;
-	uint8			MaxLeases;
-	FSocket			Sockets[1/*...N*/]; // this should be the last member
 };
 
 ////////////////////////////////////////////////////////////////////////////////
-FSocketPool::FSocketPool(const ANSICHAR* InHostName, uint32 InPort, uint32 InMaxLeases)
+FHost::FHost(const ANSICHAR* InHostName, uint32 InPort)
 : HostName(InHostName)
 , Port(uint16(InPort))
-, MaxLeases(uint8(InMaxLeases))
 {
-	check(MaxLeases == InMaxLeases); // field overflow
-
-	for (uint32 i = 0; i < InMaxLeases; ++i)
-	{
-		new (Sockets + i) FSocket();
-	}
 }
 
 ////////////////////////////////////////////////////////////////////////////////
-FSocketPool::~FSocketPool()
+void FHost::SetBufferSize(EDirection Dir, int32 Size)
 {
-	for (uint32 i = 0; i < MaxLeases; ++i)
-	{
-		Sockets[i].Destroy();
-	}
-}
-
-////////////////////////////////////////////////////////////////////////////////
-uint32 FSocketPool::GetAllocSize(uint32 MaxLeases)
-{
-	return sizeof(FSocketPool) + (sizeof(Sockets[0]) * (MaxLeases - 1));
-}
-
-////////////////////////////////////////////////////////////////////////////////
-bool FSocketPool::LeaseSocket(FSocket& Out)
-{
-	check(LeaseCount <= MaxLeases);
-
-	if (LeaseCount == MaxLeases)
-	{
-		return false;
-	}
-
-	Out = MoveTemp(Sockets[LeaseCount]);
-	++LeaseCount;
-
-	return true;
-}
-
-////////////////////////////////////////////////////////////////////////////////
-void FSocketPool::ReturnLease(FSocket&& Socket)
-{
-	check(LeaseCount > 0);
-	--LeaseCount;
-	Sockets[LeaseCount] = MoveTemp(Socket);
-}
-
-////////////////////////////////////////////////////////////////////////////////
-void FSocketPool::SetBufferSize(EDirection Dir, int32 Size)
-{
 	(Dir == EDirection::Send) ? SendBufKb : RecvBufKb = uint16(Size >> 10);
 }
 
 ////////////////////////////////////////////////////////////////////////////////
-int32 FSocketPool::GetBufferSize(EDirection Dir) const
+int32 FHost::GetBufferSize(EDirection Dir) const
 {
 	return int32((Dir == EDirection::Send) ? SendBufKb : RecvBufKb) << 10;
 }
 
 ////////////////////////////////////////////////////////////////////////////////
-int32 FSocketPool::IsResolved() const
+int32 FHost::ResolveHostName()
 {
-	switch (IpAddresses[0])
-	{
-	case 0:  return 0;
-	case 1:  return -1;
-	default: return 1;
-	}
-}
-
-////////////////////////////////////////////////////////////////////////////////
-int32 FSocketPool::ResolveHostName()
-{
 	// todo: GetAddrInfoW() for async resolve on Windows
 
@@ -1294,13 +1230,11 @@
 	auto Result = getaddrinfo(HostName, nullptr, &Hints, &Info);
 	if (uint32(Result) || Info == nullptr)
 	{
-		IpAddresses[0] = 2;
 		return -1;
 	}
 
 	if (Info->ai_family != AF_INET)
 	{
-		IpAddresses[0] = 2;
 		return -2;
 	}
 
@@ -1333,10 +1267,92 @@
 		return AddressCount;
 	}
 
-	IpAddresses[0] = 2;
 	return 0;
 }
 
+////////////////////////////////////////////////////////////////////////////////
+int32 FHost::IsResolved() const
+{
+	switch (IpAddresses[0])
+	{
+	case 0:  return 0;
+	case 1:  return -1;
+	default: return 1;
+	}
+}
+
 
 
 ////////////////////////////////////////////////////////////////////////////////
+class FSocketPool
+{
+public:
+					FSocketPool(const ANSICHAR* InHostName, uint32 InPort, uint32 InMaxLeases);
+					~FSocketPool();
+	static uint32	GetAllocSize(uint32 MaxLeases);
+	FHost&			GetHost() { return Host; }
+	bool			LeaseSocket(FSocket& Out);
+	void			ReturnLease(FSocket&& Socket);
+
+private:
+	FHost			Host;
+	uint8			LeaseCount = 0;
+	uint8			MaxLeases;
+	FSocket			Sockets[1/*...N*/]; // this should be the last member
+};
+
+////////////////////////////////////////////////////////////////////////////////
+FSocketPool::FSocketPool(const ANSICHAR* InHostName, uint32 InPort, uint32 InMaxLeases)
+: Host(InHostName, InPort)
+, MaxLeases(uint8(InMaxLeases))
+{
+	check(MaxLeases == InMaxLeases); // field overflow
+
+	for (uint32 i = 0; i < InMaxLeases; ++i)
+	{
+		new (Sockets + i) FSocket();
+	}
+}
+
+////////////////////////////////////////////////////////////////////////////////
+FSocketPool::~FSocketPool()
+{
+	for (uint32 i = 0; i < MaxLeases; ++i)
+	{
+		Sockets[i].Destroy();
+	}
+}
+
+////////////////////////////////////////////////////////////////////////////////
+uint32 FSocketPool::GetAllocSize(uint32 MaxLeases)
+{
+	return sizeof(FSocketPool) + (sizeof(Sockets[0]) * (MaxLeases - 1));
+}
+
+////////////////////////////////////////////////////////////////////////////////
+bool FSocketPool::LeaseSocket(FSocket& Out)
+{
+	check(LeaseCount <= MaxLeases);
+
+	if (LeaseCount == MaxLeases)
+	{
+		return false;
+	}
+
+	Out = MoveTemp(Sockets[LeaseCount]);
+	++LeaseCount;
+
+	return true;
+}
+
+////////////////////////////////////////////////////////////////////////////////
+void FSocketPool::ReturnLease(FSocket&& Socket)
+{
+	check(LeaseCount > 0);
+	--LeaseCount;
+	Sockets[LeaseCount] = MoveTemp(Socket);
+}
+
+
+
+////////////////////////////////////////////////////////////////////////////////
@@ -1382,8 +1398,8 @@
 		Params.Host.Port,
 		Params.ConnectionCount
 	);
-	Internal->SetBufferSize(FSocketPool::EDirection::Send, Params.SendBufSize);
-	Internal->SetBufferSize(FSocketPool::EDirection::Recv, Params.RecvBufSize);
+	Internal->GetHost().SetBufferSize(FHost::EDirection::Send, Params.SendBufSize);
+	Internal->GetHost().SetBufferSize(FHost::EDirection::Recv, Params.RecvBufSize);
 
 	Ptr = Internal;
 }
@@ -1400,7 +1416,7 @@
 ////////////////////////////////////////////////////////////////////////////////
 bool FConnectionPool::Resolve()
 {
-	return (Ptr->ResolveHostName() > 0);
+	return (Ptr->GetHost().ResolveHostName() > 0);
 }
 
 
@@ -1621,7 +1637,7 @@
 			Socket = FSocket();
 		}
 
-		if (Activity->Pool->GetIpAddress() > 0x00ff'ffff)
+		if (Activity->Pool->GetHost().GetIpAddress() > 0x00ff'ffff)
 		{
 			Activity->Pool->ReturnLease(MoveTemp(Socket));
 		}
@@ -1963,7 +1979,7 @@
 			break;
 
 		case EWait::Pool:
-			if (Activity->Pool->GetIpAddress() > 0x00ff'ffff)
+			if (Activity->Pool->GetHost().GetIpAddress() > 0x00ff'ffff)
 			{
 				Activity_EndWait(Activities[i]);
 				Ret |= (1ull << Activity->Slot);
@@ -2021,7 +2037,7 @@
 	// needs to happen once, the first activity in can do the honours. Everyone
 	// else can wait.
 	FSocketPool* Pool = Activity->Pool;
-	int32 Result = Pool->IsResolved();
+	int32 Result = Pool->GetHost().IsResolved();
 	if (Result > 0)
 	{
 		Activity_ChangeState(Activity, FActivity::EState::Connect);
@@ -2036,7 +2052,7 @@
 	}
 
 	// We won! We WON!
-	Result = Pool->ResolveHostName();
+	Result = Pool->GetHost().ResolveHostName();
 	switch (Result)
 	{
 	case 0:  Activity_SetError(Activity, "Unable to resolve host"); return -1;
@@ -2054,10 +2070,10 @@
 	TRACE_CPUPROFILER_EVENT_SCOPE(IasHttp::DoConnect);
 
 	FSocketPool* Pool = Activity->Pool;
-	check(Pool->IsResolved() > 0);
+	check(Pool->GetHost().IsResolved() > 0);
 
-	uint32 IpAddress = Pool->GetIpAddress();
-	uint32 Port = Pool->GetPort();
+	uint32 IpAddress = Pool->GetHost().GetIpAddress();
+	uint32 Port = Pool->GetHost().GetPort();
 	if (IpAddress <= 0x00ff'ffff)
 	{
 		Activity_SetError(Activity, "Unresolved host");
@@ -2108,12 +2124,12 @@
 		return -1;
 	}
 
-	if (int32 OptValue = Pool->GetBufferSize(FSocketPool::EDirection::Send); OptValue >= 0)
+	if (int32 OptValue = Pool->GetHost().GetBufferSize(FHost::EDirection::Send); OptValue >= 0)
 	{
 		Candidate.SetSendBufSize(OptValue);
 	}
 
-	if (int32 OptValue = Pool->GetBufferSize(FSocketPool::EDirection::Recv); OptValue >= 0)
+	if (int32 OptValue = Pool->GetHost().GetBufferSize(FHost::EDirection::Recv); OptValue >= 0)
 	{
 		Candidate.SetRecvBufSize(OptValue);
 	}
@@ -2650,7 +2666,7 @@
 	FMessageBuilder Builder(Activity->Buffer);
 
 	Builder << Method << " " << Path << " HTTP/1.1" "\r\n"
-		"Host: " << Activity->Pool->GetHostName() << "\r\n";
+		"Host: " << Activity->Pool->GetHost().GetHostName() << "\r\n";
 
 	// HTTP/1.1 is persistent by default thus "Connection" header isn't required
 	if (!Activity->IsKeepAlive)
