Allow custom initialization for IAS.

Allow users to control when IAS is initalized by setting UE_IAS_CUSTOM_INITIALIZATION in a target and manually call FIoStoreOnDemandModule::Initialize where needed.

diff --git a/Engine/Source/Runtime/Experimental/IoStoreOnDemand/Private/IoStoreOnDemand.cpp b/Engine/Source/Runtime/Experimental/IoStoreOnDemand/Private/IoStoreOnDemand.cpp
--- a/Engine/Source/Runtime/Experimental/IoStoreOnDemand/Private/IoStoreOnDemand.cpp
+++ b/Engine/Source/Runtime/Experimental/IoStoreOnDemand/Private/IoStoreOnDemand.cpp
@@ -1518,7 +1518,7 @@ void FIoStoreOnDemandModule::ReportAnalytics(TArray<FAnalyticsEventAttribute>& O
 	}
 }
 
-void FIoStoreOnDemandModule::StartupModule()
+void FIoStoreOnDemandModule::InitializeInternal()
 {
 	LLM_SCOPE_BYTAG(Ias);
 
@@ -1531,8 +1531,18 @@ void FIoStoreOnDemandModule::StartupModule()
 		return;
 	}
 #endif //WITH_EDITOR
-
+	
 	const TCHAR* CommandLine = FCommandLine::Get();
+	
+#if !UE_BUILD_SHIPPING
+	if (FParse::Param(CommandLine, TEXT("NoIas")))
+	{
+		return;
+	}
+#endif
+
+	// Make sure we haven't called initialize before
+	check(!Backend.IsValid());
 
 	FOnDemandEndpoint Endpoint;
 	
@@ -1568,12 +1578,8 @@ void FIoStoreOnDemandModule::StartupModule()
 		}
 	}
 
-#if !UE_BUILD_SHIPPING
-	if (FParse::Param(CommandLine, TEXT("NoIas")))
-	{
-		return;
-	}
-#endif
+	UE_LOG(LogIas, Display, TEXT("Initializing IoStoreOnDemand on service %s from %s."),
+		*Endpoint.ServiceUrl, *Endpoint.DistributionUrl);
 
 	FLatencyInjector::Initialize(CommandLine);
 
@@ -1604,6 +1610,13 @@ void FIoStoreOnDemandModule::StartupModule()
 #endif
 	FIoDispatcher::Get().Mount(Backend.ToSharedRef(), BackendPriority);
 }
+	
+void FIoStoreOnDemandModule::StartupModule()
+{
+#if !UE_IAS_CUSTOM_INITIALIZATION
+	InitializeInternal();
+#endif
+}
 
 void FIoStoreOnDemandModule::ShutdownModule()
 {
diff --git a/Engine/Source/Runtime/Experimental/IoStoreOnDemand/Public/IO/IoStoreOnDemand.h b/Engine/Source/Runtime/Experimental/IoStoreOnDemand/Public/IO/IoStoreOnDemand.h
--- a/Engine/Source/Runtime/Experimental/IoStoreOnDemand/Public/IO/IoStoreOnDemand.h
+++ b/Engine/Source/Runtime/Experimental/IoStoreOnDemand/Public/IO/IoStoreOnDemand.h
@@ -28,6 +28,12 @@ struct FIoStoreWriterSettings;
 namespace UE::IO::IAS { struct FOnDemandEndpoint; }
 using FIoBlockHash = uint32;
 
+// Custom initialization allows users to control when
+// the system should be initialized.
+#if !defined(UE_IAS_CUSTOM_INITIALIZATION)
+	#define UE_IAS_CUSTOM_INITIALIZATION 0
+#endif
+
 UE_API DECLARE_LOG_CATEGORY_EXTERN(LogIas, VeryVerbose, All);
 
 namespace UE::IO::IAS
@@ -209,6 +215,7 @@ class FIoStoreOnDemandModule
 	: public IModuleInterface
 {
 private:
+	void InitializeInternal();
 	TSharedPtr<IOnDemandIoDispatcherBackend> Backend;
 
 public:
@@ -220,6 +227,10 @@ public:
 
 	virtual void StartupModule() override;
 	virtual void ShutdownModule() override;
+	
+#if UE_IAS_CUSTOM_INITIALIZATION
+	UE_API void Initialize() { InitializeInternal(); };
+#endif
 };
 
 } // namespace UE::IO::IAS
