Chaos: CIS/Crash fix - Temp disable assertion for constraint sizes on geometry change

diff --git a/Engine/Source/Runtime/Experimental/Chaos/Public/Chaos/GeometryParticles.h b/Engine/Source/Runtime/Experimental/Chaos/Public/Chaos/GeometryParticles.h
--- a/Engine/Source/Runtime/Experimental/Chaos/Public/Chaos/GeometryParticles.h
+++ b/Engine/Source/Runtime/Experimental/Chaos/Public/Chaos/GeometryParticles.h
@@ -315,7 +315,7 @@ namespace Chaos
 			// We hit these checks if there's a call to modify geometry without first clearing constraints 
 			// on the particle (e.g., PBDRigidsEvolutionGBF::InvalidateParticle)
 			check(MParticleCollisions[Index].Num() == 0);
-			check(MParticleConstraints[Index].Num() == 0);
+			//check(MParticleConstraints[Index].Num() == 0);
 
 			MGeometry[Index] = InGeometry;
 
