geometry collection X/R/V/W result interpolator when pulling data from the PT to the GT should not interpolate if only Next Data is available because we run the possibility of interpolating from the current GT particle position which may be incorrect since it may not have yet updated due to its parent's movement

diff --git a/Engine/Source/Runtime/Experimental/Chaos/Private/PhysicsProxy/GeometryCollectionPhysicsProxy.cpp b/Engine/Source/Runtime/Experimental/Chaos/Private/PhysicsProxy/GeometryCollectionPhysicsProxy.cpp
--- a/Engine/Source/Runtime/Experimental/Chaos/Private/PhysicsProxy/GeometryCollectionPhysicsProxy.cpp
+++ b/Engine/Source/Runtime/Experimental/Chaos/Private/PhysicsProxy/GeometryCollectionPhysicsProxy.cpp
@@ -3645,10 +3645,10 @@ public:
 		}
 		else if (EntryIndex == INDEX_NONE)
 		{
-			// not present in current results, we need to interpolate from the actual GTparticle
+			// not present in current results, just use next results
 			const FGeometryCollectionResults::FPositionData& NextData = NextResults->GetPositions(NextEntryIndex);
-			PosOut = FMath::Lerp(GTParticle.X(), NextData.ParticleX , Alpha);
-			RotOut = FQuat::Slerp(GTParticle.R(), NextData.ParticleR, Alpha);
+			PosOut = NextData.ParticleX;
+			RotOut = NextData.ParticleR;
 		}
 		else
 		{
@@ -3671,10 +3671,10 @@ public:
 		}
 		else if (EntryIndex == INDEX_NONE)
 		{
-			// not present in current results, we need to interpolate from the actual GTparticle
+			// not present in current results, just use next results
 			const FGeometryCollectionResults::FVelocityData& NextData = NextResults->GetVelocities(NextEntryIndex);
-			LinearVelOut = FMath::Lerp(Chaos::FVec3f(GTParticle.V()), NextData.ParticleV, Alpha);
-			AngularVelOut = FMath::Lerp(Chaos::FVec3f(GTParticle.W()), NextData.ParticleW, Alpha);
+			LinearVelOut = NextData.ParticleV;
+			AngularVelOut = NextData.ParticleW;
 		}
 		else
 		{
