Iris
* Deleted single-repsystem support. Now all configs support multi-systems
* Renamed IrisDebugHelper commands so they use the InternalNetRefIndex name properly
* NetRefHandle::ToString now returns RepSystem=? instead of Unknown

#rnx

diff --git a/Engine/Source/Runtime/Experimental/Iris/Core/Private/Iris/Core/IrisDebugging.cpp b/Engine/Source/Runtime/Experimental/Iris/Core/Private/Iris/Core/IrisDebugging.cpp
--- a/Engine/Source/Runtime/Experimental/Iris/Core/Private/Iris/Core/IrisDebugging.cpp
+++ b/Engine/Source/Runtime/Experimental/Iris/Core/Private/Iris/Core/IrisDebugging.cpp
@@ -120,7 +120,7 @@ bool BreakOnRPCName(FName RPCName)
 }
 
 
-bool BreakOnNetInternalIndex(UE::Net::Private::FInternalNetRefIndex InternalIndex)
+bool BreakOnInternalNetRefIndex(UE::Net::Private::FInternalNetRefIndex InternalIndex)
 {
 	if (IrisDebugHelperInternal::GIrisDebugInternalIndex != 0 && IrisDebugHelperInternal::GIrisDebugInternalIndex == InternalIndex)
 	{
@@ -131,7 +131,7 @@ bool BreakOnNetInternalIndex(UE::Net::Private::FInternalNetRefIndex InternalInde
 	return false;
 }
 
-UE::Net::Private::FInternalNetRefIndex GetDebugNetInternalIndex()
+UE::Net::Private::FInternalNetRefIndex GetDebugInternalNetRefIndex()
 {
 	return IrisDebugHelperInternal::GIrisDebugInternalIndex;
 }
@@ -170,12 +170,12 @@ void SetIrisDebugRPCName(const ANSICHAR* NameBuffer)
 	}
 }
 
-void SetIrisDebugInternalIndex(UE::Net::Private::FInternalNetRefIndex InternalIndex)
+void SetIrisDebugInternalNetRefIndex(UE::Net::Private::FInternalNetRefIndex InternalIndex)
 {
 	IrisDebugHelperInternal::GIrisDebugInternalIndex = InternalIndex;
 }
 
-void SetIrisDebugInternalIndexViaNetHandle(FNetRefHandle RefHandle)
+void SetIrisDebugInternalNetRefIndexViaNetHandle(FNetRefHandle RefHandle)
 {
 	using namespace UE::Net::Private;
 
@@ -205,6 +205,40 @@ void SetIrisDebugInternalIndexViaNetHandle(FNetRefHandle RefHandle)
 	IrisDebugHelperInternal::GIrisDebugInternalIndex = InternalNetRefIndex;
 }
 
+void SetIrisDebugInternalNetRefIndexViaObject(UObject* Instance)
+{
+	using namespace UE::Net::Private;
+
+	if (!Instance)
+	{
+		IrisDebugHelperInternal::GIrisDebugNetRefHandle = FNetRefHandle();
+		IrisDebugHelperInternal::GIrisDebugInternalIndex = FNetRefHandleManager::InvalidInternalIndex;
+		return;
+	}
+
+	// See if we can find the instance in any replication system
+	for (uint32 RepSystemIt = 0; RepSystemIt < FReplicationSystemFactory::MaxReplicationSystemCount; ++RepSystemIt)
+	{
+		const UReplicationSystem* ReplicationSystem = GetReplicationSystemForDebug(RepSystemIt);
+		if (ReplicationSystem)
+		{
+			if (const UObjectReplicationBridge* Bridge = ReplicationSystem->GetReplicationBridgeAs<UObjectReplicationBridge>())
+			{
+				FNetRefHandle Handle = Bridge->GetReplicatedRefHandle(Instance);
+				if (Handle.IsValid())
+				{
+					const FNetRefHandleManager& NetRefHandleManager = ReplicationSystem->GetReplicationSystemInternal()->GetNetRefHandleManager();
+					FInternalNetRefIndex InternalIndex = NetRefHandleManager.GetInternalIndex(Handle);
+
+					IrisDebugHelperInternal::GIrisDebugNetRefHandle = Handle;
+					IrisDebugHelperInternal::GIrisDebugInternalIndex = InternalIndex;
+					break;
+				}
+			}
+		}
+	}
+}
+
 
 UReplicationSystem* GetReplicationSystemForDebug(uint32 Id)
 {
@@ -225,12 +259,13 @@ uint64 Init()
 	UE_NET_FORCE_REFERENCE_DEBUGFUNCTION(DebugNetObjectById);
 	UE_NET_FORCE_REFERENCE_DEBUGFUNCTION(DebugNetRefHandle);
 	UE_NET_FORCE_REFERENCE_DEBUGFUNCTION(DebugNetRefHandleById);
-	UE_NET_FORCE_REFERENCE_DEBUGFUNCTION(DebugNetInternalIndex);
+	UE_NET_FORCE_REFERENCE_DEBUGFUNCTION(DebugInternalNetRefIndex);
 	UE_NET_FORCE_REFERENCE_DEBUGFUNCTION(DebugNetObjectProtocolReferencesToString);
 	UE_NET_FORCE_REFERENCE_DEBUGFUNCTION(SetIrisDebugObjectName);
 	UE_NET_FORCE_REFERENCE_DEBUGFUNCTION(SetIrisDebugNetRefHandle);
-	UE_NET_FORCE_REFERENCE_DEBUGFUNCTION(SetIrisDebugInternalIndex);
-	UE_NET_FORCE_REFERENCE_DEBUGFUNCTION(SetIrisDebugInternalIndexViaNetHandle);
+	UE_NET_FORCE_REFERENCE_DEBUGFUNCTION(SetIrisDebugInternalNetRefIndex);
+	UE_NET_FORCE_REFERENCE_DEBUGFUNCTION(SetIrisDebugInternalNetRefIndexViaNetHandle);
+	UE_NET_FORCE_REFERENCE_DEBUGFUNCTION(SetIrisDebugInternalNetRefIndexViaObject);
 	UE_NET_FORCE_REFERENCE_DEBUGFUNCTION(SetIrisDebugRPCName);
 	
 	return FunctionReferenceAccumulator;
@@ -426,7 +461,7 @@ FNetReplicatedObjectDebugInfo DebugNetRefHandleById(uint64 NetRefHandleId, uint3
 	return Info;
 }
 
-FNetReplicatedObjectDebugInfo DebugNetInternalIndex(uint32 InternalIndex, uint32 ReplicationSystemId)
+FNetReplicatedObjectDebugInfo DebugInternalNetRefIndex(uint32 InternalIndex, uint32 ReplicationSystemId)
 {
 	using namespace UE::Net::Private;
 	FNetReplicatedObjectDebugInfo Info = {};
diff --git a/Engine/Source/Runtime/Experimental/Iris/Core/Private/Iris/ReplicationSystem/NetRefHandle.cpp b/Engine/Source/Runtime/Experimental/Iris/Core/Private/Iris/ReplicationSystem/NetRefHandle.cpp
--- a/Engine/Source/Runtime/Experimental/Iris/Core/Private/Iris/ReplicationSystem/NetRefHandle.cpp
+++ b/Engine/Source/Runtime/Experimental/Iris/Core/Private/Iris/ReplicationSystem/NetRefHandle.cpp
@@ -9,28 +9,21 @@ namespace UE::Net
 
 FString FNetRefHandle::ToString() const
 {
-	FString Result;
-#if UE_NET_ALLOW_MULTIPLE_REPLICATION_SYSTEMS
 	if (IsCompleteHandle())
 	{
 		const uint32 ReplicationSystemIdToDisplay = GetReplicationSystemId();
-		Result = FString::Printf(TEXT("NetRefHandle (Id=%" UINT64_FMT "):(RepSystemId=%u)"), GetId(), ReplicationSystemIdToDisplay);
+		return FString::Printf(TEXT("NetRefHandle (Id=%" UINT64_FMT "):(RepSystemId=%u)"), GetId(), ReplicationSystemIdToDisplay);
 	}
 	else
 	{
-		Result = FString::Printf(TEXT("NetRefHandle (Id=%" UINT64_FMT "):(Incomplete)"), GetId());
+		return FString::Printf(TEXT("NetRefHandle (Id=%" UINT64_FMT "):(RepSystemId=?)"), GetId());
 	}
-#else
-	Result = FString::Printf(TEXT("NetRefHandle (Id=%" UINT64_FMT ")"), GetId());
-#endif
-	return Result;
 }
 
 }
 
 FStringBuilderBase& operator<<(FStringBuilderBase& Builder, const UE::Net::FNetRefHandle& NetRefHandle) 
 { 	
-#if UE_NET_ALLOW_MULTIPLE_REPLICATION_SYSTEMS
 	if (NetRefHandle.IsCompleteHandle())
 	{
 		const uint32 ReplicationSystemIdToDisplay = NetRefHandle.GetReplicationSystemId();
@@ -38,17 +31,14 @@ FStringBuilderBase& operator<<(FStringBuilderBase& Builder, const UE::Net::FNetR
 	}
 	else
 	{
-		Builder.Appendf(TEXT("NetRefHandle (Id=%" UINT64_FMT "):(Incomplete)"), NetRefHandle.GetId());
+		Builder.Appendf(TEXT("NetRefHandle (Id=%" UINT64_FMT "):(RepSystemId=?)"), NetRefHandle.GetId());
 	}
-#else
-	Builder.Appendf(TEXT("NetRefHandle (Id=%" UINT64_FMT ")"), NetRefHandle.GetId());
-#endif
+
 	return Builder;
 }
 
 FAnsiStringBuilderBase& operator<<(FAnsiStringBuilderBase& Builder, const UE::Net::FNetRefHandle& NetRefHandle)
 {
-#if UE_NET_ALLOW_MULTIPLE_REPLICATION_SYSTEMS
 	if (NetRefHandle.IsCompleteHandle())
 	{
 		const uint32 ReplicationSystemIdToDisplay = NetRefHandle.GetReplicationSystemId();
@@ -56,10 +46,8 @@ FAnsiStringBuilderBase& operator<<(FAnsiStringBuilderBase& Builder, const UE::Ne
 	}
 	else
 	{
-		Builder.Appendf("NetRefHandle (Id=%" UINT64_FMT "):(Incomplete)", NetRefHandle.GetId());
+		Builder.Appendf("NetRefHandle (Id=%" UINT64_FMT "):(RepSystemId=?)", NetRefHandle.GetId());
 	}
-#else
-	Builder.Appendf("NetRefHandle (Id=%" UINT64_FMT ")", NetRefHandle.GetId());
-#endif
+
 	return Builder;
 }
diff --git a/Engine/Source/Runtime/Experimental/Iris/Core/Private/Iris/ReplicationSystem/NetRefHandleManager.cpp b/Engine/Source/Runtime/Experimental/Iris/Core/Private/Iris/ReplicationSystem/NetRefHandleManager.cpp
--- a/Engine/Source/Runtime/Experimental/Iris/Core/Private/Iris/ReplicationSystem/NetRefHandleManager.cpp
+++ b/Engine/Source/Runtime/Experimental/Iris/Core/Private/Iris/ReplicationSystem/NetRefHandleManager.cpp
@@ -79,13 +79,15 @@ FInternalNetRefIndex FNetRefHandleManager::InternalCreateNetObject(const FNetRef
 	// Verify that the handle is free
 	if (RefHandleToInternalIndex.Contains(NetRefHandle))
 	{
-		ensureAlwaysMsgf(false, TEXT("NetRefHandleManager::InternalCreateNetObject %s already exists"), *NetRefHandle.ToString());
+		ensureMsgf(false, TEXT("NetRefHandleManager::InternalCreateNetObject %s already exists"), *NetRefHandle.ToString());
 		return InvalidInternalIndex;
 	}
 
 	const uint32 InternalIndex = GetNextFreeInternalIndex();
 	if (InternalIndex != InvalidInternalIndex)
 	{
+		UE_LOG(LogIris, Verbose, TEXT("FNetRefHandleManager::InternalCreateNetObject ( InternalIndex: %u ) %s"), InternalIndex, *NetRefHandle.ToString());
+
 		// Store data;
 		FReplicatedObjectData& Data = ReplicatedObjectData[InternalIndex];
 
diff --git a/Engine/Source/Runtime/Experimental/Iris/Core/Public/Iris/Core/IrisDebugging.h b/Engine/Source/Runtime/Experimental/Iris/Core/Public/Iris/Core/IrisDebugging.h
--- a/Engine/Source/Runtime/Experimental/Iris/Core/Public/Iris/Core/IrisDebugging.h
+++ b/Engine/Source/Runtime/Experimental/Iris/Core/Public/Iris/Core/IrisDebugging.h
@@ -36,7 +36,7 @@ IRISCORE_API bool BreakOnNetRefHandle(FNetRefHandle NetRefHandle);
 IRISCORE_API bool BreakOnRPCName(FName RPCName);
 
 /** Trigger a breakpoint and return true if the index is the current debug index */
-IRISCORE_API bool BreakOnNetInternalIndex(UE::Net::Private::FInternalNetRefIndex InternalIndex);
+IRISCORE_API bool BreakOnInternalNetRefIndex(UE::Net::Private::FInternalNetRefIndex InternalIndex);
 
 /** Returns true if the object name contains the current debug name, will return true if no debug name is set */
 IRISCORE_API bool FilterDebuggedObject(UObject* Object);
@@ -65,6 +65,7 @@ extern "C" IRISCORE_API void DebugOutputNetObjectState(uint32 NetRefHandleId, ui
  * Variant of NetObjectStateToString that can be used from breakpoints and in watch window to print the current state of a NetRefHandle.
  * NOTE: Use only for debugging as this variant uses a static buffer which is not thread safe.
  */
+ //$IRIS TODO: Make console command to log these for object X or class X
 extern "C" IRISCORE_API const TCHAR* DebugNetObjectStateToString(uint32 NetRefHandleId, uint32 ReplicationSystemId);
 
 /** Find all handles references registered for a protocol and output to DebugOutput in debugger */
@@ -94,7 +95,7 @@ extern "C" IRISCORE_API FNetReplicatedObjectDebugInfo DebugNetRefHandle(FNetRefH
 extern "C" IRISCORE_API FNetReplicatedObjectDebugInfo DebugNetRefHandleById(uint64 NetRefHandleId, uint32 ReplicationSystemId);
 
 /** Look up replicated handle specified by an internal index and replicationsystem id and return debug information */
-extern "C" IRISCORE_API FNetReplicatedObjectDebugInfo DebugNetInternalIndex(UE::Net::Private::FInternalNetRefIndex InternalIndex, uint32 ReplicationSystemId);
+extern "C" IRISCORE_API FNetReplicatedObjectDebugInfo DebugInternalNetRefIndex(UE::Net::Private::FInternalNetRefIndex InternalIndex, uint32 ReplicationSystemId);
 
 /** Variant of DebugOutputNetObjectProtocolReferences that can be used from breakpoints or in watch window to find all handles references registered for a protocol and output to DebugOutput in debugger
  * NOTE: Use only for debugging as this variant uses a static buffer which is not thread safe	
@@ -108,10 +109,13 @@ extern "C" IRISCORE_API void SetIrisDebugObjectName(const ANSICHAR* NameBuffer);
 extern "C" IRISCORE_API void SetIrisDebugNetRefHandle(uint64 NetHandleId);
 
 /** Set the InternalIndex to break on */
-extern "C" IRISCORE_API void SetIrisDebugInternalIndex(UE::Net::Private::FInternalNetRefIndex InternalIndex);
+extern "C" IRISCORE_API void SetIrisDebugInternalNetRefIndex(UE::Net::Private::FInternalNetRefIndex InternalIndex);
 
 /** Set the InternalIndex to break on via it's NetHandle */
-extern "C" IRISCORE_API void SetIrisDebugInternalIndexViaNetHandle(FNetRefHandle RefHandle);
+extern "C" IRISCORE_API void SetIrisDebugInternalNetRefIndexViaNetHandle(FNetRefHandle RefHandle);
+
+/** Set the InternalIndex to break on via an object pointer*/
+extern "C" IRISCORE_API void SetIrisDebugInternalNetRefIndexViaObject(UObject* Instance);
 
 /** Set the RPC Name to break on */
 extern "C" IRISCORE_API void SetIrisDebugRPCName(const ANSICHAR* NameBuffer);
diff --git a/Engine/Source/Runtime/Experimental/Iris/Core/Public/Iris/IrisConfig.h b/Engine/Source/Runtime/Experimental/Iris/Core/Public/Iris/IrisConfig.h
--- a/Engine/Source/Runtime/Experimental/Iris/Core/Public/Iris/IrisConfig.h
+++ b/Engine/Source/Runtime/Experimental/Iris/Core/Public/Iris/IrisConfig.h
@@ -12,15 +12,6 @@ IRISCORE_API void SetUseIrisReplication(bool EnableIrisReplication);
 
 }
 
-/** Allow multiple replication systems inside the same process. Ex: PIE support */
-#ifndef UE_NET_ALLOW_MULTIPLE_REPLICATION_SYSTEMS
-#if (WITH_EDITOR || WITH_DEV_AUTOMATION_TESTS || WITH_AUTOMATION_WORKER)
-#	define UE_NET_ALLOW_MULTIPLE_REPLICATION_SYSTEMS 1
-#else
-#	define UE_NET_ALLOW_MULTIPLE_REPLICATION_SYSTEMS 0
-#endif
-#endif
-
 /* NetBitStreamReader/Writer validation support */
 #ifndef UE_NETBITSTREAMWRITER_VALIDATE
 #if !(UE_BUILD_SHIPPING || UE_BUILD_TEST)
diff --git a/Engine/Source/Runtime/Experimental/Iris/Core/Public/Iris/ReplicationSystem/ReplicationSystem.h b/Engine/Source/Runtime/Experimental/Iris/Core/Public/Iris/ReplicationSystem/ReplicationSystem.h
--- a/Engine/Source/Runtime/Experimental/Iris/Core/Public/Iris/ReplicationSystem/ReplicationSystem.h
+++ b/Engine/Source/Runtime/Experimental/Iris/Core/Public/Iris/ReplicationSystem/ReplicationSystem.h
@@ -628,11 +628,7 @@ public:
 
 	enum ReplicationSystemConstants : uint32
 	{
-#if UE_NET_ALLOW_MULTIPLE_REPLICATION_SYSTEMS
 		MaxReplicationSystemCount = 16,
-#else
-		MaxReplicationSystemCount = 1,
-#endif
 	};
 
 private:
