Fix typo checking the wrong variable before an early out.

diff --git a/Engine/Source/Runtime/Experimental/GeometryCollectionEngine/Private/GeometryCollection/GeometryCollectionISMPoolComponent.cpp b/Engine/Source/Runtime/Experimental/GeometryCollectionEngine/Private/GeometryCollection/GeometryCollectionISMPoolComponent.cpp
--- a/Engine/Source/Runtime/Experimental/GeometryCollectionEngine/Private/GeometryCollection/GeometryCollectionISMPoolComponent.cpp
+++ b/Engine/Source/Runtime/Experimental/GeometryCollectionEngine/Private/GeometryCollection/GeometryCollectionISMPoolComponent.cpp
@@ -281,7 +281,7 @@ void FGeometryCollectionISMPool::BatchUpdateInstanceCustomData(FGeometryCollecti
 	}
 	
 	FGeometryCollectionISM& ISM = ISMs[MeshInfo.ISMIndex];
-	if (!ensure(CustomFloatValue < ISM.MeshInstance.Desc.NumCustomDataFloats))
+	if (!ensure(CustomFloatIndex < ISM.MeshInstance.Desc.NumCustomDataFloats))
 	{
 		return;
 	}
