Geometry collection : fix race condition where setting the mass override scale right after creating the physics thread may result in apply the scale twice on the mass of the particles

diff --git a/Engine/Source/Runtime/Experimental/Chaos/Private/PhysicsProxy/GeometryCollectionPhysicsProxy.cpp b/Engine/Source/Runtime/Experimental/Chaos/Private/PhysicsProxy/GeometryCollectionPhysicsProxy.cpp
--- a/Engine/Source/Runtime/Experimental/Chaos/Private/PhysicsProxy/GeometryCollectionPhysicsProxy.cpp
+++ b/Engine/Source/Runtime/Experimental/Chaos/Private/PhysicsProxy/GeometryCollectionPhysicsProxy.cpp
@@ -1330,7 +1330,7 @@ void FGeometryCollectionPhysicsProxy::InitializeBodiesPT(Chaos::FPBDRigidsSolver
 					}
 					else
 					{
-						Handle = BuildNonClusters_Internal(TransformGroupIndex, RigidsSolver, AdjustMassForScale(Masses[TransformGroupIndex]), Inertias[TransformGroupIndex]);
+						Handle = BuildNonClusters_Internal(TransformGroupIndex, RigidsSolver, Masses[TransformGroupIndex], Inertias[TransformGroupIndex]);
 					}
 					Handle->GTGeometryParticle() = GTParticle;
 
@@ -3051,7 +3051,15 @@ void FGeometryCollectionPhysicsProxy::PushStateOnGameThread(Chaos::FPBDRigidsSol
 	if (GameThreadPerFrameData.MaterialOverrideMassScaleMultiplier.IsSet())
 	{
 		const float NewValue = GameThreadPerFrameData.MaterialOverrideMassScaleMultiplier.GetValue();
-		MaterialOverrideMassScaleMultiplierChange = NewValue / Parameters.MaterialOverrideMassScaleMultiplier;
+
+		// Because we need to send a change in scale , we need to make sure the physics state has been created
+		// otherwise we set the change to 1.0 and let the PT pick up the right Parameters.MaterialOverrideMassScaleMultiplier during InitializeBodiesPT
+		MaterialOverrideMassScaleMultiplierChange = 1.0;
+		if (bIsInitializedOnPhysicsThread)
+		{
+			MaterialOverrideMassScaleMultiplierChange = NewValue / Parameters.MaterialOverrideMassScaleMultiplier;
+		}
+
 		Parameters.MaterialOverrideMassScaleMultiplier = NewValue;
 		GameThreadPerFrameData.MaterialOverrideMassScaleMultiplier.Reset();
 	}
