Iris - Fix constness.

diff --git a/Engine/Source/Runtime/Experimental/Iris/Core/Private/Iris/ReplicationSystem/NetExports.cpp b/Engine/Source/Runtime/Experimental/Iris/Core/Private/Iris/ReplicationSystem/NetExports.cpp
--- a/Engine/Source/Runtime/Experimental/Iris/Core/Private/Iris/ReplicationSystem/NetExports.cpp
+++ b/Engine/Source/Runtime/Experimental/Iris/Core/Private/Iris/ReplicationSystem/NetExports.cpp
@@ -33,7 +33,7 @@ void FNetExports::PushExportRecordForPacket()
 	ExportRecord.Enqueue(Info);
 }
 
-void FNetExports::AcknowledgeBatchExports(FNetExportContext::FBatchExports& BatchExports)
+void FNetExports::AcknowledgeBatchExports(const FNetExportContext::FBatchExports& BatchExports)
 {
 	AcknowledgedExports.AcknowledgedExportedHandles.Append(BatchExports.HandlesExportedInCurrentBatch);
 	AcknowledgedExports.AcknowledgedExportedNetTokens.Append(BatchExports.NetTokensExportedInCurrentBatch);
diff --git a/Engine/Source/Runtime/Experimental/Iris/Core/Private/Iris/ReplicationSystem/NetExports.h b/Engine/Source/Runtime/Experimental/Iris/Core/Private/Iris/ReplicationSystem/NetExports.h
--- a/Engine/Source/Runtime/Experimental/Iris/Core/Private/Iris/ReplicationSystem/NetExports.h
+++ b/Engine/Source/Runtime/Experimental/Iris/Core/Private/Iris/ReplicationSystem/NetExports.h
@@ -47,7 +47,7 @@ public:
 	void ProcessPacketDeliveryStatus(UE::Net::EPacketDeliveryStatus Status);
 
 	// Explicitly acknowledge exports, this is used for exports originating from out of band batches 
-	void AcknowledgeBatchExports(FNetExportContext::FBatchExports& BatchExports);
+	void AcknowledgeBatchExports(const FNetExportContext::FBatchExports& BatchExports);
 
 	// Construct a new export scope
 	FExportScope MakeExportScope(UE::Net::FNetSerializationContext& Context, FNetExportContext::FBatchExports& BatchExports) { return FExportScope(Context, AcknowledgedExports, BatchExports); }
