[NetPhysics] CU and GC resim error correction through render interpolation.

- Implemented the linear error offset correction on CU and GCs, same as for SingleParticles.
- When a CU is resimulated we also need to send the error offset to its child GCs

diff --git a/Engine/Source/Runtime/Experimental/Chaos/Private/PhysicsProxy/ClusterUnionPhysicsProxy.cpp b/Engine/Source/Runtime/Experimental/Chaos/Private/PhysicsProxy/ClusterUnionPhysicsProxy.cpp
--- a/Engine/Source/Runtime/Experimental/Chaos/Private/PhysicsProxy/ClusterUnionPhysicsProxy.cpp
+++ b/Engine/Source/Runtime/Experimental/Chaos/Private/PhysicsProxy/ClusterUnionPhysicsProxy.cpp
@@ -8,6 +8,7 @@
 #include "Chaos/PullPhysicsDataImp.h"
 #include "Math/UnrealMathUtility.h"
 #include "PBDRigidsSolver.h"
+#include "Chaos/DebugDrawQueue.h"
 
 namespace Chaos
 {
@@ -381,7 +382,7 @@ namespace Chaos
 	}
 
 	DECLARE_CYCLE_STAT(TEXT("FClusterUnionPhysicsProxy::PullFromPhysicsState"), STAT_ClusterUnionPhysicsProxyPullFromPhysicsState, STATGROUP_Chaos);
-	bool FClusterUnionPhysicsProxy::PullFromPhysicsState(const FDirtyClusterUnionData& PullData, int32 SolverSyncTimestamp, const FDirtyClusterUnionData* NextPullData, const FRealSingle* Alpha)
+	bool FClusterUnionPhysicsProxy::PullFromPhysicsState(const FDirtyClusterUnionData& PullData, int32 SolverSyncTimestamp, const FDirtyClusterUnionData* NextPullData, const FRealSingle* Alpha, const FDirtyRigidParticleReplicationErrorData* Error, const Chaos::FReal AsyncFixedTimeStep)
 	{
 		SCOPE_CYCLE_COUNTER(STAT_ClusterUnionPhysicsProxyPullFromPhysicsState);
 		if (!ensure(Particle_External))
@@ -396,6 +397,12 @@ namespace Chaos
 			return false;
 		}
 
+		if (Error)
+		{
+			const int32 RenderInterpErrorCorrectionDurationTicks = FMath::FloorToInt32(GetRenderInterpErrorCorrectionDuration() / AsyncFixedTimeStep); // Convert duration from seconds to simulation ticks
+			InterpolationData.AccumlateErrorXR(Error->ErrorX, Error->ErrorR, SolverSyncTimestamp, RenderInterpErrorCorrectionDurationTicks);
+		}
+
 		SyncedData_External.bIsAnchored = CurrentPullData.bIsAnchored;
 		SyncedData_External.bDidSyncGeometry = false;
 		SyncedData_External.ChildParticles.Empty(CurrentPullData.ChildParticles.Num());
@@ -453,15 +460,24 @@ namespace Chaos
 				return OverwriteProperty.Timestamp <= SolverSyncTimestamp ? (OverwriteProperty.Timestamp < SolverSyncTimestamp ? &Prev : &OverwriteProperty.Value) : nullptr;
 			};
 
+			const bool bIsReplicationErrorSmoothing = InterpolationData.IsErrorSmoothing();
+			InterpolationData.UpdateError(SolverSyncTimestamp, AsyncFixedTimeStep);
+
 			if (const FVec3* Prev = LerpHelper(PullData.X, ProxyTimestamp->OverWriteX))
 			{
-				const FVec3 NewX = FMath::Lerp(*Prev, NextPullData->X, *Alpha);
+				const FVec3 NewX = bIsReplicationErrorSmoothing ?
+					FMath::Lerp(*Prev, NextPullData->X, *Alpha) + InterpolationData.GetErrorX(*Alpha) :
+					FMath::Lerp(*Prev, NextPullData->X, *Alpha);
+
 				Particle_External->SetX(NewX, false);
 			}
 
 			if (const FQuat* Prev = LerpHelper(PullData.R, ProxyTimestamp->OverWriteR))
 			{
-				const FQuat NewR = FMath::Lerp(*Prev, NextPullData->R, *Alpha);
+				const FQuat NewR = bIsReplicationErrorSmoothing ? 
+					InterpolationData.GetErrorR(*Alpha) * FMath::Lerp(*Prev, NextPullData->R, *Alpha) : 
+					FMath::Lerp(*Prev, NextPullData->R, *Alpha);
+		
 				Particle_External->SetR(NewR, false);
 			}
 
@@ -476,6 +492,26 @@ namespace Chaos
 				const FVec3 NewW = FMath::Lerp(*Prev, NextPullData->W, *Alpha);
 				Particle_External->SetW(NewW, false);
 			}
+
+#if CHAOS_DEBUG_DRAW
+			if (GetRenderInterpDebugDraw())
+			{
+				Chaos::FDebugDrawQueue::GetInstance().DrawDebugBox(NextPullData->X, FVector(2, 1, 1), NextPullData->R, FColor::Yellow, false, 5.f, 0, 0.5f);
+				Chaos::FDebugDrawQueue::GetInstance().DrawDebugDirectionalArrow(PullData.X, NextPullData->X, 0.5f, FColor::Yellow, false, 5.0f, 0, 0.5f);
+				Chaos::FDebugDrawQueue::GetInstance().DrawDebugBox(Particle_External->X(), FVector(2, 1, 1), Particle_External->R(), FColor::Green, false, 5.f, 0, 0.5f);
+
+				if (bIsReplicationErrorSmoothing)
+				{
+					if (Error)
+					{
+						Chaos::FDebugDrawQueue::GetInstance().DrawDebugBox(PullData.X, FVector(4, 2, 2), PullData.R, FColor::Red, false, 5.f, 0, 0.5f);
+						Chaos::FDebugDrawQueue::GetInstance().DrawDebugDirectionalArrow(PullData.X, (PullData.X + InterpolationData.GetErrorX(0)), 1, FColor::Red, false, 5.0f, 0, 0.5f);
+					}
+
+					Chaos::FDebugDrawQueue::GetInstance().DrawDebugDirectionalArrow((Particle_External->X() - InterpolationData.GetErrorX(*Alpha)), Particle_External->X(), 1, FColor::Blue, false, 5.0f, 0, 0.5f);
+				}
+			}
+#endif // CHAOS_DEBUG_DRAW
 		}
 		else
 		{
diff --git a/Engine/Source/Runtime/Experimental/Chaos/Private/PhysicsProxy/GeometryCollectionPhysicsProxy.cpp b/Engine/Source/Runtime/Experimental/Chaos/Private/PhysicsProxy/GeometryCollectionPhysicsProxy.cpp
--- a/Engine/Source/Runtime/Experimental/Chaos/Private/PhysicsProxy/GeometryCollectionPhysicsProxy.cpp
+++ b/Engine/Source/Runtime/Experimental/Chaos/Private/PhysicsProxy/GeometryCollectionPhysicsProxy.cpp
@@ -3810,7 +3810,7 @@ bool FGeometryCollectionPhysicsProxy::PullNonInterpolatableDataFromSinglePhysics
 }
 
 // Called from FPhysScene_ChaosInterface::SyncBodies(), NOT the solver.
-bool FGeometryCollectionPhysicsProxy::PullFromPhysicsState(const Chaos::FDirtyGeometryCollectionData& PullData, const int32 SolverSyncTimestamp, const Chaos::FDirtyGeometryCollectionData* NextPullData, const Chaos::FRealSingle* Alpha)
+bool FGeometryCollectionPhysicsProxy::PullFromPhysicsState(const Chaos::FDirtyGeometryCollectionData& PullData, const int32 SolverSyncTimestamp, const Chaos::FDirtyGeometryCollectionData* NextPullData, const Chaos::FRealSingle* Alpha, const Chaos::FDirtyRigidParticleReplicationErrorData* Error, const Chaos::FReal AsyncFixedTimeStep)
 {
 	if(IsObjectDeleting) return false;
 
@@ -3834,6 +3834,12 @@ bool FGeometryCollectionPhysicsProxy::PullFromPhysicsState(const Chaos::FDirtyGe
 		return false;
 	}
 
+	if (Error)
+	{
+		const int32 RenderInterpErrorCorrectionDurationTicks = FMath::FloorToInt32(GetRenderInterpErrorCorrectionDuration() / AsyncFixedTimeStep); // Convert duration from seconds to simulation ticks
+		InterpolationData.AccumlateErrorXR(Error->ErrorX, Error->ErrorR, SolverSyncTimestamp, RenderInterpErrorCorrectionDurationTicks);
+	}
+
 	const int32 NumTransforms = GameThreadCollection.Transform.Num();
 	const bool bNeedInterpolation = (NextPullData != nullptr);
 
@@ -3857,6 +3863,8 @@ bool FGeometryCollectionPhysicsProxy::PullFromPhysicsState(const Chaos::FDirtyGe
 		// second : interpolate-able ones
 		if (bNeedInterpolation)
 		{
+			InterpolationData.UpdateError(SolverSyncTimestamp, AsyncFixedTimeStep);
+
 			const FGeometryCollectionResults& PrevResults = PullData.Results();
 			const FGeometryCollectionResults& NextResults = NextPullData->Results();
 
@@ -3888,6 +3896,13 @@ bool FGeometryCollectionPhysicsProxy::PullFromPhysicsState(const Chaos::FDirtyGe
 					Chaos::FVec3 NewX;
 					Chaos::FRotation3 NewR;
 					ResultInterpolator.GetPositions(NewX, NewR);
+
+					if (InterpolationData.IsErrorSmoothing())
+					{
+						NewX = NewX + InterpolationData.GetErrorX(*Alpha);
+						NewR = InterpolationData.GetErrorR(*Alpha) * NewR;
+					}
+
 					const bool XRModified = UpdateGTParticleXR(GTParticle, NewX, NewR);
 					bIsCollectionDirty |= XRModified;
 
diff --git a/Engine/Source/Runtime/Experimental/Chaos/Private/PhysicsProxy/SingleParticlePhysicsProxy.cpp b/Engine/Source/Runtime/Experimental/Chaos/Private/PhysicsProxy/SingleParticlePhysicsProxy.cpp
--- a/Engine/Source/Runtime/Experimental/Chaos/Private/PhysicsProxy/SingleParticlePhysicsProxy.cpp
+++ b/Engine/Source/Runtime/Experimental/Chaos/Private/PhysicsProxy/SingleParticlePhysicsProxy.cpp
@@ -254,16 +254,6 @@ void FSingleParticlePhysicsProxy::BufferPhysicsResults_External(Chaos::FDirtyRig
 	}
 }
 
-
-float RenderInterpErrorCorrectionDuration = 0.5f;
-FAutoConsoleVariableRef CVarRenderInterpErrorCorrectionDuration(TEXT("p.RenderInterp.ErrorCorrectionDuration"), RenderInterpErrorCorrectionDuration, TEXT("How long in seconds to apply error correction over."));
-
-float RenderInterpErrorVelocitySmoothingDuration = 0.5f;
-FAutoConsoleVariableRef CVarRenderInterpErrorVelocitySmoothingDuration(TEXT("p.RenderInterp.ErrorVelocitySmoothingDuration"), RenderInterpErrorVelocitySmoothingDuration, TEXT("How long in seconds to apply error velocity smoothing correction over, should be smaller than or equal to p.RenderInterp.ErrorCorrectionDuration. RENDERINTERPOLATION_VELOCITYSMOOTHING needs to be defined."));
-
-int32 RenderInterpDebugDraw = 0;
-FAutoConsoleVariableRef CVarRenderInterpDebugDraw(TEXT("p.RenderInterp.DebugDraw"), RenderInterpDebugDraw, TEXT("Draw debug lines for physics render interpolation, also needs p.Chaos.DebugDraw.Enabled set"));
-
 bool ShouldUpdateTransformFromSimulation(const Chaos::FPBDRigidParticle& Rigid)
 {
 	if (Rigid.ObjectState() == Chaos::EObjectStateType::Kinematic)
@@ -284,6 +274,7 @@ bool ShouldUpdateTransformFromSimulation(const Chaos::FPBDRigidParticle& Rigid)
 bool FSingleParticlePhysicsProxy::PullFromPhysicsState(const Chaos::FDirtyRigidParticleData& PullData,int32 SolverSyncTimestamp, const Chaos::FDirtyRigidParticleData* NextPullData, const Chaos::FRealSingle* Alpha, const FDirtyRigidParticleReplicationErrorData* Error, const Chaos::FReal AsyncFixedTimeStep)
 {
 	using namespace Chaos;
+
 	// Move buffered data into the TPBDRigidParticle without triggering invalidation of the physics state.
 	Chaos::FPBDRigidParticle* Rigid = Particle ? Particle->CastToRigidParticle() : nullptr;
 	if(Rigid)
@@ -295,12 +286,12 @@ bool FSingleParticlePhysicsProxy::PullFromPhysicsState(const Chaos::FDirtyRigidP
 		const FSingleParticleProxyTimestamp* ProxyTimestamp = PullData.GetTimestamp();
 		
 #if RENDERINTERP_ERRORVELOCITYSMOOTHING
-		const int32 RenderInterpErrorVelocitySmoothingDurationTicks = FMath::FloorToInt32(RenderInterpErrorVelocitySmoothingDuration / AsyncFixedTimeStep); // Convert duration from seconds to simulation ticks
+		const int32 RenderInterpErrorVelocitySmoothingDurationTicks = FMath::FloorToInt32(GetRenderInterpErrorVelocitySmoothingDuration() / AsyncFixedTimeStep); // Convert duration from seconds to simulation ticks
 #endif
 
 		if (Error)
 		{
-			const int32 RenderInterpErrorCorrectionDurationTicks = FMath::FloorToInt32(RenderInterpErrorCorrectionDuration / AsyncFixedTimeStep); // Convert duration from seconds to simulation ticks
+			const int32 RenderInterpErrorCorrectionDurationTicks = FMath::FloorToInt32(GetRenderInterpErrorCorrectionDuration() / AsyncFixedTimeStep); // Convert duration from seconds to simulation ticks
 			InterpolationData.AccumlateErrorXR(Error->ErrorX, Error->ErrorR, SolverSyncTimestamp, RenderInterpErrorCorrectionDurationTicks);
 #if RENDERINTERP_ERRORVELOCITYSMOOTHING
 			InterpolationData.SetVelocitySmoothing(Rigid->V(), Rigid->X(), RenderInterpErrorVelocitySmoothingDurationTicks);
@@ -365,7 +356,7 @@ bool FSingleParticlePhysicsProxy::PullFromPhysicsState(const Chaos::FDirtyRigidP
 				}
 				
 #if CHAOS_DEBUG_DRAW
-				if (!!RenderInterpDebugDraw)
+				if (GetRenderInterpDebugDraw())
 				{
 					Chaos::FDebugDrawQueue::GetInstance().DrawDebugBox(NextPullData->X, FVector(2, 1, 1), NextPullData->R, FColor::Yellow, false, 5.f, 0, 0.5f);
 					Chaos::FDebugDrawQueue::GetInstance().DrawDebugDirectionalArrow(PullData.X, NextPullData->X, 0.5f, FColor::Yellow, false, 5.0f, 0, 0.5f);
@@ -390,7 +381,6 @@ bool FSingleParticlePhysicsProxy::PullFromPhysicsState(const Chaos::FDirtyRigidP
 							Chaos::FDebugDrawQueue::GetInstance().DrawDebugDirectionalArrow((Rigid->X() - InterpolationData.GetErrorX(*Alpha)), Rigid->X(), 1, FColor::Blue, false, 5.0f, 0, 0.5f);
 						}
 					}
-
 				}
 #endif // CHAOS_DEBUG_DRAW
 			}
diff --git a/Engine/Source/Runtime/Experimental/Chaos/Private/RewindData.cpp b/Engine/Source/Runtime/Experimental/Chaos/Private/RewindData.cpp
--- a/Engine/Source/Runtime/Experimental/Chaos/Private/RewindData.cpp
+++ b/Engine/Source/Runtime/Experimental/Chaos/Private/RewindData.cpp
@@ -3,6 +3,7 @@
 #include "RewindData.h"
 #include "PBDRigidsSolver.h"
 #include "PhysicsProxy/SingleParticlePhysicsProxy.h"
+#include "PhysicsProxy/ClusterUnionPhysicsProxy.h"
 #include "HAL/IConsoleManager.h"
 #include "Chaos/PBDJointConstraints.h"
 #include "Misc/Paths.h"
@@ -749,7 +750,7 @@ void FRewindData::SpawnProxyIfNeeded(FSingleParticlePhysicsProxy& Proxy)
 }
 
 // Hand over error data per particle from RewindData to a solver collection instead which gets marshalled to GT
-void FRewindData::BufferPhysicsResults(TMap<const FSingleParticlePhysicsProxy*, struct FDirtyRigidParticleReplicationErrorData>& DirtyRigidErrors)
+void FRewindData::BufferPhysicsResults(TMap<const IPhysicsProxyBase*, struct FDirtyRigidParticleReplicationErrorData>& DirtyRigidErrors)
 {
 	DirtyRigidErrors.Reserve(DirtyParticleErrors.Num());
 
@@ -758,7 +759,26 @@ void FRewindData::BufferPhysicsResults(TMap<const FSingleParticlePhysicsProxy*,
 		FDirtyRigidParticleReplicationErrorData ErrorData;
 		ErrorData.ErrorX = ErrorInfo.GetErrorX();
 		ErrorData.ErrorR = ErrorInfo.GetErrorR();
-		DirtyRigidErrors.Add(static_cast<const FSingleParticlePhysicsProxy*>(ErrorInfo.GetObjectPtr()->PhysicsProxy()), ErrorData);
+
+		const IPhysicsProxyBase* PhysicsProxy = static_cast<const IPhysicsProxyBase*>(ErrorInfo.GetObjectPtr()->PhysicsProxy());
+		DirtyRigidErrors.Add(PhysicsProxy, ErrorData);
+
+		// todo, move out of RewindData.cpp, possibly to ChaosResultManager.cpp
+		// If ClusterUnion, hand over error to child GeometryCollections which do their own render interpolation and need the error offset individually
+		if (PhysicsProxy->GetType() == EPhysicsProxyType::ClusterUnionProxy)
+		{
+			const FClusterUnionPhysicsProxy* ClusterProxy = static_cast<const FClusterUnionPhysicsProxy*>(PhysicsProxy);
+			if (ClusterProxy)
+			{
+				for (IPhysicsProxyBase* ChildProxy : ClusterProxy->GetParticle_Internal()->PhysicsProxies())
+				{
+					if (ChildProxy->GetType() == EPhysicsProxyType::GeometryCollectionType)
+					{
+						DirtyRigidErrors.Add(ChildProxy, ErrorData);
+					}
+				}
+			}
+		}
 	}
 
 	DirtyParticleErrors.Reset();
diff --git a/Engine/Source/Runtime/Experimental/Chaos/Public/Chaos/Framework/PhysicsProxyBase.h b/Engine/Source/Runtime/Experimental/Chaos/Public/Chaos/Framework/PhysicsProxyBase.h
--- a/Engine/Source/Runtime/Experimental/Chaos/Public/Chaos/Framework/PhysicsProxyBase.h
+++ b/Engine/Source/Runtime/Experimental/Chaos/Public/Chaos/Framework/PhysicsProxyBase.h
@@ -121,6 +121,26 @@ public:
 	IPhysicsProxyBase* GetParentProxy() const { return ParentProxy; }
 	void SetParentProxy(IPhysicsProxyBase* InProxy) { ParentProxy = InProxy; }
 
+	// Render Interpolation CVars
+	static float GetRenderInterpErrorCorrectionDuration()
+	{
+		static float RenderInterpErrorCorrectionDuration = 0.5f;
+		static FAutoConsoleVariableRef CVarRenderInterpErrorCorrectionDuration(TEXT("p.RenderInterp.ErrorCorrectionDuration"), RenderInterpErrorCorrectionDuration, TEXT("How long in seconds to apply error correction over."));
+		return RenderInterpErrorCorrectionDuration;
+	}
+	static float GetRenderInterpErrorVelocitySmoothingDuration()
+	{
+		static float RenderInterpErrorVelocitySmoothingDuration = 0.5f;
+		static FAutoConsoleVariableRef CVarRenderInterpErrorVelocitySmoothingDuration(TEXT("p.RenderInterp.ErrorVelocitySmoothingDuration"), RenderInterpErrorVelocitySmoothingDuration, TEXT("How long in seconds to apply error velocity smoothing correction over, should be smaller than or equal to p.RenderInterp.ErrorCorrectionDuration. RENDERINTERPOLATION_VELOCITYSMOOTHING needs to be defined."));
+		return RenderInterpErrorVelocitySmoothingDuration;
+	}
+	static bool GetRenderInterpDebugDraw()
+	{
+		static bool RenderInterpDebugDraw = false;
+		static FAutoConsoleVariableRef CVarRenderInterpDebugDraw(TEXT("p.RenderInterp.DebugDraw"), RenderInterpDebugDraw, TEXT("Draw debug lines for physics render interpolation, also needs p.Chaos.DebugDraw.Enabled set"));
+		return RenderInterpDebugDraw;
+	}
+
 protected:
 	// Ensures that derived classes can successfully call this destructor
 	// but no one can delete using a IPhysicsProxyBase*
@@ -304,4 +324,4 @@ protected:
 	Chaos::FVec3 ErrorVelocitySmoothingX = { 0,0,0 };
 	Chaos::FVec3 ErrorVelocitySmoothingXPrev = { 0,0,0 };
 	int32 ErrorVelocitySmoothingCount = 0;
-};
+};
\ No newline at end of file
diff --git a/Engine/Source/Runtime/Experimental/Chaos/Public/Chaos/PhysicsSolverBaseImpl.h b/Engine/Source/Runtime/Experimental/Chaos/Public/Chaos/PhysicsSolverBaseImpl.h
--- a/Engine/Source/Runtime/Experimental/Chaos/Public/Chaos/PhysicsSolverBaseImpl.h
+++ b/Engine/Source/Runtime/Experimental/Chaos/Public/Chaos/PhysicsSolverBaseImpl.h
@@ -77,12 +77,12 @@ namespace Chaos
 				// single particles
 				{
 					SCOPE_CYCLE_COUNTER(STAT_ProcessSingleProxy);
-					for(const FChaosRigidInterpolationData& RigidInterp : Results.RigidInterpolations)
+					for (const FChaosRigidInterpolationData& RigidInterp : Results.RigidInterpolations)
 					{
-						if(FSingleParticlePhysicsProxy* Proxy = RigidInterp.Prev.GetProxy())
+						if (FSingleParticlePhysicsProxy* Proxy = RigidInterp.Prev.GetProxy())
 						{
 							const FDirtyRigidParticleReplicationErrorData* ErrorData = LatestData->DirtyRigidErrors.Find(Proxy);
-							if(Proxy->PullFromPhysicsState(RigidInterp.Prev, SolverTimestamp, &RigidInterp.Next, &Results.Alpha, ErrorData, GetAsyncDeltaTime()))
+							if (Proxy->PullFromPhysicsState(RigidInterp.Prev, SolverTimestamp, &RigidInterp.Next, &Results.Alpha, ErrorData, GetAsyncDeltaTime()))
 							{
 								Private::TPullPhysicsStateDispatchHelper<TDispatcher>::Apply(Dispatcher, Proxy);
 							}
@@ -94,14 +94,16 @@ namespace Chaos
 				// geometry collections
 				{
 					SCOPE_CYCLE_COUNTER(STAT_ProcessGCProxy);
-					for(const FChaosGeometryCollectionInterpolationData& GCInterp : Results.GeometryCollectionInterpolations)
+					for (const FChaosGeometryCollectionInterpolationData& GCInterp : Results.GeometryCollectionInterpolations)
 					{
-						if(FGeometryCollectionPhysicsProxy* Proxy = GCInterp.Prev.GetProxy())
+						if (FGeometryCollectionPhysicsProxy* Proxy = GCInterp.Prev.GetProxy())
 						{
-							if(Proxy->PullFromPhysicsState(GCInterp.Prev, SolverTimestamp, &GCInterp.Next, &Results.Alpha))
+							const FDirtyRigidParticleReplicationErrorData* ErrorData = LatestData->DirtyRigidErrors.Find(Proxy);
+							if (Proxy->PullFromPhysicsState(GCInterp.Prev, SolverTimestamp, &GCInterp.Next, &Results.Alpha, ErrorData, GetAsyncDeltaTime()))
 							{
 								Private::TPullPhysicsStateDispatchHelper<TDispatcher>::Apply(Dispatcher, Proxy);
 							}
+							LatestData->DirtyRigidErrors.Remove(Proxy);
 						}
 					}
 				}
@@ -113,10 +115,12 @@ namespace Chaos
 					{
 						if(FClusterUnionPhysicsProxy* Proxy = ClusterInterp.Prev.GetProxy())
 						{
-							if(Proxy->PullFromPhysicsState(ClusterInterp.Prev, SolverTimestamp, &ClusterInterp.Next, &Results.Alpha))
+							const FDirtyRigidParticleReplicationErrorData* ErrorData = LatestData->DirtyRigidErrors.Find(Proxy);
+							if (Proxy->PullFromPhysicsState(ClusterInterp.Prev, SolverTimestamp, &ClusterInterp.Next, &Results.Alpha, ErrorData, GetAsyncDeltaTime()))
 							{
 								Private::TPullPhysicsStateDispatchHelper<TDispatcher>::Apply(Dispatcher, Proxy);
 							}
+							LatestData->DirtyRigidErrors.Remove(Proxy);
 						}
 					}
 				}
diff --git a/Engine/Source/Runtime/Experimental/Chaos/Public/Chaos/PullPhysicsDataImp.h b/Engine/Source/Runtime/Experimental/Chaos/Public/Chaos/PullPhysicsDataImp.h
--- a/Engine/Source/Runtime/Experimental/Chaos/Public/Chaos/PullPhysicsDataImp.h
+++ b/Engine/Source/Runtime/Experimental/Chaos/Public/Chaos/PullPhysicsDataImp.h
@@ -159,7 +159,7 @@ class FPullPhysicsData
 {
 public:
 	TArray<FDirtyRigidParticleData> DirtyRigids;
-	TMap<const FSingleParticlePhysicsProxy*, FDirtyRigidParticleReplicationErrorData> DirtyRigidErrors;
+	TMap<const IPhysicsProxyBase*, FDirtyRigidParticleReplicationErrorData> DirtyRigidErrors;
 	TArray<FDirtyGeometryCollectionData> DirtyGeometryCollections;
 	TArray<FDirtyClusterUnionData> DirtyClusterUnions;
 	TArray<FDirtyJointConstraintData> DirtyJointConstraints;
diff --git a/Engine/Source/Runtime/Experimental/Chaos/Public/PhysicsProxy/ClusterUnionPhysicsProxy.h b/Engine/Source/Runtime/Experimental/Chaos/Public/PhysicsProxy/ClusterUnionPhysicsProxy.h
--- a/Engine/Source/Runtime/Experimental/Chaos/Public/PhysicsProxy/ClusterUnionPhysicsProxy.h
+++ b/Engine/Source/Runtime/Experimental/Chaos/Public/PhysicsProxy/ClusterUnionPhysicsProxy.h
@@ -124,7 +124,7 @@ namespace Chaos
 		// and the physics thread.
 		//
 		CHAOS_API void PushToPhysicsState(const FDirtyPropertiesManager& Manager, int32 DataIdx, const FDirtyProxy& Dirty);
-		CHAOS_API bool PullFromPhysicsState(const FDirtyClusterUnionData& PullData, int32 SolverSyncTimestamp, const FDirtyClusterUnionData* NextPullData = nullptr, const FRealSingle* Alpha = nullptr);
+		CHAOS_API bool PullFromPhysicsState(const FDirtyClusterUnionData& PullData, int32 SolverSyncTimestamp, const FDirtyClusterUnionData* NextPullData = nullptr, const FRealSingle* Alpha = nullptr, const FDirtyRigidParticleReplicationErrorData* Error = nullptr, const Chaos::FReal AsyncFixedTimeStep = 0);
 
 		CHAOS_API void BufferPhysicsResults_Internal(FDirtyClusterUnionData& BufferData);
 		CHAOS_API void BufferPhysicsResults_External(FDirtyClusterUnionData& BufferData);
@@ -151,7 +151,7 @@ namespace Chaos
 		FInternalParticle* Particle_Internal = nullptr;
 		FClusterUnionIndex ClusterUnionIndex = INDEX_NONE;
 
-		FProxyInterpolationBase InterpolationData;
+		FProxyInterpolationError InterpolationData;
 
 		// An array of a particles that exist in the external implicit object union.
 		// Note that this array should only be used for book-keeping. It is generally
diff --git a/Engine/Source/Runtime/Experimental/Chaos/Public/PhysicsProxy/GeometryCollectionPhysicsProxy.h b/Engine/Source/Runtime/Experimental/Chaos/Public/PhysicsProxy/GeometryCollectionPhysicsProxy.h
--- a/Engine/Source/Runtime/Experimental/Chaos/Public/PhysicsProxy/GeometryCollectionPhysicsProxy.h
+++ b/Engine/Source/Runtime/Experimental/Chaos/Public/PhysicsProxy/GeometryCollectionPhysicsProxy.h
@@ -223,7 +223,7 @@ public:
 	 * Pulls data out of the PhysToGameInterchange and updates \c GTDynamicCollection. 
 	 * Called from FPhysScene_ChaosInterface::SyncBodies(), NOT the solver.
 	 */
-	CHAOS_API bool PullFromPhysicsState(const Chaos::FDirtyGeometryCollectionData& BufferData, const int32 SolverSyncTimestamp, const Chaos::FDirtyGeometryCollectionData* NextPullData = nullptr, const Chaos::FRealSingle* Alpha= nullptr);
+	CHAOS_API bool PullFromPhysicsState(const Chaos::FDirtyGeometryCollectionData& BufferData, const int32 SolverSyncTimestamp, const Chaos::FDirtyGeometryCollectionData* NextPullData = nullptr, const Chaos::FRealSingle* Alpha = nullptr, const Chaos::FDirtyRigidParticleReplicationErrorData* Error = nullptr, const Chaos::FReal AsyncFixedTimeStep = 0);
 
 	bool IsDirty() { return false; }
 
@@ -654,7 +654,7 @@ private:
 	// paradigm, at least for this component of the handshake.
 	Chaos::FGuardedTripleBuffer<FGeometryCollectionResults> PhysToGameInterchange;
 
-	FProxyInterpolationBase InterpolationData;
+	FProxyInterpolationError InterpolationData;
 
 	// this is used as a unique ID when collecting data from runtime
 	FGuid CollectorGuid;
diff --git a/Engine/Source/Runtime/Experimental/Chaos/Public/RewindData.h b/Engine/Source/Runtime/Experimental/Chaos/Public/RewindData.h
--- a/Engine/Source/Runtime/Experimental/Chaos/Public/RewindData.h
+++ b/Engine/Source/Runtime/Experimental/Chaos/Public/RewindData.h
@@ -2,7 +2,6 @@
 #pragma once
 #include "Chaos/Core.h"
 #include "Chaos/ParticleHandle.h"
-#include "PhysicsProxy/SingleParticlePhysicsProxyFwd.h"
 #include "Chaos/Framework/PhysicsSolverBase.h"
 #include "Serialization/BufferArchive.h"
 #include "Containers/CircularBuffer.h"
@@ -1443,7 +1442,7 @@ public:
 	/** Rewind the states from all the history states datas */
 	void RewindStates(const int32 RewindFrame, const bool bResetSolver);
 
-	void BufferPhysicsResults(TMap<const FSingleParticlePhysicsProxy*, struct FDirtyRigidParticleReplicationErrorData>& DirtyRigidErrors);
+	void BufferPhysicsResults(TMap<const IPhysicsProxyBase*, struct FDirtyRigidParticleReplicationErrorData>& DirtyRigidErrors);
 
 	/** Return the rewind data solver */
 	const FPBDRigidsSolver* GetSolver() const { return Solver; }
