[NetPhysics] Enable RenderInterpolation CVars

diff --git a/Engine/Source/Runtime/Experimental/Chaos/Public/Chaos/Framework/PhysicsProxyBase.h b/Engine/Source/Runtime/Experimental/Chaos/Public/Chaos/Framework/PhysicsProxyBase.h
--- a/Engine/Source/Runtime/Experimental/Chaos/Public/Chaos/Framework/PhysicsProxyBase.h
+++ b/Engine/Source/Runtime/Experimental/Chaos/Public/Chaos/Framework/PhysicsProxyBase.h
@@ -5,6 +5,7 @@
 #include "Chaos/Declares.h"
 #include "UObject/GCObject.h"
 #include "Chaos/Core.h"
+#include "HAL/IConsoleManager.h"
 
 enum class EPhysicsProxyType : uint32
 {
@@ -125,19 +126,19 @@ public:
 	static float GetRenderInterpErrorCorrectionDuration()
 	{
 		static float RenderInterpErrorCorrectionDuration = 0.5f;
-	//	temp disable, static FAutoConsoleVariableRef CVarRenderInterpErrorCorrectionDuration(TEXT("p.RenderInterp.ErrorCorrectionDuration"), RenderInterpErrorCorrectionDuration, TEXT("How long in seconds to apply error correction over."));
+		static FAutoConsoleVariableRef CVarRenderInterpErrorCorrectionDuration(TEXT("p.RenderInterp.ErrorCorrectionDuration"), RenderInterpErrorCorrectionDuration, TEXT("How long in seconds to apply error correction over."));
 		return RenderInterpErrorCorrectionDuration;
 	}
 	static float GetRenderInterpErrorVelocitySmoothingDuration()
 	{
 		static float RenderInterpErrorVelocitySmoothingDuration = 0.5f;
-	//	temp disable, static FAutoConsoleVariableRef CVarRenderInterpErrorVelocitySmoothingDuration(TEXT("p.RenderInterp.ErrorVelocitySmoothingDuration"), RenderInterpErrorVelocitySmoothingDuration, TEXT("How long in seconds to apply error velocity smoothing correction over, should be smaller than or equal to p.RenderInterp.ErrorCorrectionDuration. RENDERINTERPOLATION_VELOCITYSMOOTHING needs to be defined."));
+		static FAutoConsoleVariableRef CVarRenderInterpErrorVelocitySmoothingDuration(TEXT("p.RenderInterp.ErrorVelocitySmoothingDuration"), RenderInterpErrorVelocitySmoothingDuration, TEXT("How long in seconds to apply error velocity smoothing correction over, should be smaller than or equal to p.RenderInterp.ErrorCorrectionDuration. RENDERINTERPOLATION_VELOCITYSMOOTHING needs to be defined."));
 		return RenderInterpErrorVelocitySmoothingDuration;
 	}
 	static bool GetRenderInterpDebugDraw()
 	{
 		static bool RenderInterpDebugDraw = false;
-	//	temp disable, static FAutoConsoleVariableRef CVarRenderInterpDebugDraw(TEXT("p.RenderInterp.DebugDraw"), RenderInterpDebugDraw, TEXT("Draw debug lines for physics render interpolation, also needs p.Chaos.DebugDraw.Enabled set"));
+		static FAutoConsoleVariableRef CVarRenderInterpDebugDraw(TEXT("p.RenderInterp.DebugDraw"), RenderInterpDebugDraw, TEXT("Draw debug lines for physics render interpolation, also needs p.Chaos.DebugDraw.Enabled set"));
 		return RenderInterpDebugDraw;
 	}
 
