Chaos - fix CollisionConstraint initial penetration MinInitialPhi

E.g., Objects would pop out of the ground even when InitialDepenetrationVelocity was set to zero. We don't save solver results for contacts that did not generate an impulse so we were throwing away initial penetration data. We must calculate and store MinimumInitialPhi on the previous frame when the data is still available.

diff --git a/Engine/Source/Runtime/Experimental/Chaos/Private/Chaos/Collision/PBDCollisionConstraint.cpp b/Engine/Source/Runtime/Experimental/Chaos/Private/Chaos/Collision/PBDCollisionConstraint.cpp
--- a/Engine/Source/Runtime/Experimental/Chaos/Private/Chaos/Collision/PBDCollisionConstraint.cpp
+++ b/Engine/Source/Runtime/Experimental/Chaos/Private/Chaos/Collision/PBDCollisionConstraint.cpp
@@ -436,9 +436,6 @@ namespace Chaos
 		Flags.bIsCurrent = true;
 		Flags.bDisabled = false;
 
-		// If we have data from the previous tick, update the maximum allowed initial penetration
-		CalculateMinInitialPhi();
-
 		// Match new manifold points with data from previous tick
 		AssignSavedManifoldPoints();
 
@@ -1135,16 +1132,6 @@ namespace Chaos
 		}
 	}
 
-	void FPBDCollisionConstraint::CalculateMinInitialPhi()
-	{
-		MinInitialPhi = 0;
-
-		for (const FSavedManifoldPoint& SavedManifoldPoint : SavedManifoldPoints)
-		{
-			MinInitialPhi = FMath::Min(MinInitialPhi, SavedManifoldPoint.InitialPhi);
-		}
-	}
-
 	ECollisionConstraintDirection FPBDCollisionConstraint::GetConstraintDirection(const FReal Dt) const
 	{
 		if (GetDisabled() || GetIsProbe())
diff --git a/Engine/Source/Runtime/Experimental/Chaos/Public/Chaos/Collision/PBDCollisionConstraint.h b/Engine/Source/Runtime/Experimental/Chaos/Public/Chaos/Collision/PBDCollisionConstraint.h
--- a/Engine/Source/Runtime/Experimental/Chaos/Public/Chaos/Collision/PBDCollisionConstraint.h
+++ b/Engine/Source/Runtime/Experimental/Chaos/Public/Chaos/Collision/PBDCollisionConstraint.h
@@ -692,6 +692,7 @@ namespace Chaos
 			//SavedManifoldPoints.SetNum(ManifoldPoints.Num());
 			SavedManifoldPoints.Reset(ManifoldPoints.Num());
 			ManifoldPointResults.SetNum(ManifoldPoints.Num());
+			MinInitialPhi = 0;
 		}
 
 		/**
@@ -704,7 +705,7 @@ namespace Chaos
 			const FRealSingle StaticFrictionRatio,
 			const FRealSingle Dt)
 		{
-			FManifoldPoint& ManifoldPoint = ManifoldPoints[ManifoldPointIndex];
+			const FManifoldPoint& ManifoldPoint = ManifoldPoints[ManifoldPointIndex];
 			FManifoldPointResult& ManifoldPointResult = ManifoldPointResults[ManifoldPointIndex];
 			FSavedManifoldPoint* SavedManifoldPoint = nullptr;
 
@@ -755,6 +756,8 @@ namespace Chaos
 				SavedManifoldPoint->ShapeContactPoints[1] = Anchor1;
 				SavedManifoldPoint->InitialPhi = ManifoldPoint.InitialPhi;
 			}
+
+			MinInitialPhi = FMath::Min(MinInitialPhi, ManifoldPoint.InitialPhi);
 		}
 
 		/**
@@ -821,7 +824,6 @@ namespace Chaos
 
 		CHAOS_API int32 FindSavedManifoldPoint(const int32 ManifoldPointIndex, int32* InOutAllowedSavedPointIndices, int32& InOutNumAllowedSavedPoints) const;
 		CHAOS_API void AssignSavedManifoldPoints();
-		CHAOS_API void CalculateMinInitialPhi();
 
 		inline void InitManifoldPoint(const int32 ManifoldPointIndex, const FContactPoint& ContactPoint)
 		{
