Iris - Multiple huge objects support.

The max number of huge objects in transit is controlled via the cvar net.Iris.ReplicationWriterMaxHugeObjectsInTransit, defaulting to 16, but will be at least one. There's a tradeoff mainly between the connection characteristics to support and normal object replication scheduling when tweaking this value. On one hand you don't want to end up stalling object replication because the top priority objects are huge. So you want to be able to keep replicating huge objects during the maximum latency, including latency variation, and packet loss scenarios you want to provide the best experience possible for. On the other hand object deletion cannot be performed once the object is in the huge object queue. Consider this and how long time it will take to replicate the huge object queue depending on the average payload of a huge object.

diff --git a/Engine/Plugins/Runtime/ReplicationSystemTestPlugin/Source/Private/Tests/ReplicationSystem/TestObjectSplitting.cpp b/Engine/Plugins/Runtime/ReplicationSystemTestPlugin/Source/Private/Tests/ReplicationSystem/TestObjectSplitting.cpp
--- a/Engine/Plugins/Runtime/ReplicationSystemTestPlugin/Source/Private/Tests/ReplicationSystem/TestObjectSplitting.cpp
+++ b/Engine/Plugins/Runtime/ReplicationSystemTestPlugin/Source/Private/Tests/ReplicationSystem/TestObjectSplitting.cpp
@@ -125,7 +125,7 @@ UE_NET_TEST_FIXTURE(FSplitObjectTestFixture, SplitHugeObjectOnCreation)
 
 	// As the payload is huge we don't expect the whole payload to arrive the first frame
 	const UTestReplicatedIrisObject* ClientObject = Cast<UTestReplicatedIrisObject>(Client->GetReplicationBridge()->GetReplicatedObject(ServerObject->NetRefHandle));
-	UE_NET_ASSERT_EQ(ClientObject, nullptr);
+	UE_NET_EXPECT_EQ(ClientObject, nullptr);
 
 	for (uint32 RetryIt = 0; RetryIt != HugeObjectMaxNetTickCountToArrive && ClientObject == nullptr; ++RetryIt)
 	{
@@ -188,7 +188,7 @@ UE_NET_TEST_FIXTURE(FSplitObjectTestFixture, SplitObjectWithHugeSubObjectsOnCrea
 
 	// As the payload is huge we don't expect the whole payload to arrive the first frame
 	const UTestReplicatedIrisObject* ClientObject = Cast<UTestReplicatedIrisObject>(Client->GetReplicationBridge()->GetReplicatedObject(ServerObject->NetRefHandle));
-	UE_NET_ASSERT_EQ(ClientObject, nullptr);
+	UE_NET_EXPECT_EQ(ClientObject, nullptr);
 
 	for (uint32 RetryIt = 0; RetryIt != HugeObjectMaxNetTickCountToArrive && ClientObject == nullptr; ++RetryIt)
 	{
@@ -248,7 +248,7 @@ UE_NET_TEST_FIXTURE(FSplitObjectTestFixture, SplitObjectWithSubObjectsWithHugeAt
 
 	// As the payload is huge we don't expect the whole payload to arrive the first frame
 	const UTestReplicatedIrisObject* ClientObject = Cast<UTestReplicatedIrisObject>(Client->GetReplicationBridge()->GetReplicatedObject(ServerObject->NetRefHandle));
-	UE_NET_ASSERT_EQ(ClientObject, nullptr);
+	UE_NET_EXPECT_EQ(ClientObject, nullptr);
 
 	for (uint32 RetryIt = 0; RetryIt != HugeObjectMaxNetTickCountToArrive && ClientObject == nullptr; ++RetryIt)
 	{
@@ -1001,6 +1001,110 @@ UE_NET_TEST_FIXTURE(FSplitObjectTestFixture, TestHugeObjectIsFlushedAndNotDestro
 	UE_NET_ASSERT_EQ(ClientObject->IntA, ServerObject->IntA);
 }
 
+// Test that huge object state can be replicated on creation.
+UE_NET_TEST_FIXTURE(FSplitObjectTestFixture, SplitManyHugeObjectsOnCreation)
+{
+	FReplicationSystemTestClient* Client = CreateClient();
+
+	constexpr uint32 HugeObjectCount = 37U;
+	UTestReplicatedIrisObject* ServerObjects[HugeObjectCount];
+	for (UTestReplicatedIrisObject*& ServerObject : ServerObjects)
+	{
+		ServerObject = CreateHugeObject(Server);
+	}
+
+	// Send and deliver packets until all huge objects have arrived.
+	UTestReplicatedIrisObject* ClientObjects[HugeObjectCount] = {};
+	uint32 ClientObjectCount = 0;
+	for (uint32 RetryIt = 0; RetryIt != (HugeObjectMaxNetTickCountToArrive*HugeObjectCount) && ClientObjectCount < HugeObjectCount; ++RetryIt)
+	{
+		Server->PreSendUpdate();
+		Server->SendAndDeliverTo(Client, DeliverPacket);
+		Server->PostSendUpdate();
+
+		for (UTestReplicatedIrisObject*& ClientObject : ClientObjects)
+		{
+			if (ClientObject)
+			{
+				continue;
+			}
+
+			const SIZE_T ObjectIndex = &ClientObject - ClientObjects;
+			ClientObject = Cast<UTestReplicatedIrisObject>(Client->GetReplicationBridge()->GetReplicatedObject(ServerObjects[ObjectIndex]->NetRefHandle));
+			ClientObjectCount += (ClientObject != nullptr);
+		}
+	}
+
+	UE_NET_ASSERT_EQ(ClientObjectCount, HugeObjectCount);
+}
+
+// Test that huge object state can be replicated after an object has been created.
+UE_NET_TEST_FIXTURE(FSplitObjectTestFixture, SplitManyHugeObjectsAfterCreation)
+{
+	FReplicationSystemTestClient* Client = CreateClient();
+
+	constexpr uint32 HugeObjectCount = 37U;
+	UTestReplicatedIrisObject* ServerObjects[HugeObjectCount];
+	for (UTestReplicatedIrisObject*& ServerObject : ServerObjects)
+	{
+		ServerObject = CreateObject(Server);
+	}
+
+	// Send and deliver packets until all huge objects have been received on the client.
+	UTestReplicatedIrisObject* ClientObjects[HugeObjectCount] = {};
+	uint32 ClientObjectCount = 0;
+	for (uint32 RetryIt = 0; RetryIt != HugeObjectCount && ClientObjectCount < HugeObjectCount; ++RetryIt)
+	{
+		Server->PreSendUpdate();
+		Server->SendAndDeliverTo(Client, DeliverPacket);
+		Server->PostSendUpdate();
+
+		for (UTestReplicatedIrisObject*& ClientObject : ClientObjects)
+		{
+			if (ClientObject)
+			{
+				continue;
+			}
+
+			const SIZE_T ObjectIndex = &ClientObject - ClientObjects;
+			ClientObject = Cast<UTestReplicatedIrisObject>(Client->GetReplicationBridge()->GetReplicatedObject(ServerObjects[ObjectIndex]->NetRefHandle));
+			ClientObjectCount += (ClientObject != nullptr);
+		}
+	}
+
+	UE_NET_ASSERT_EQ(ClientObjectCount, HugeObjectCount);
+
+	// Make all objects huge.
+	for (UTestReplicatedIrisObject*& ServerObject : ServerObjects)
+	{
+		SetObjectPayloadByteCount(ServerObject, HugeObjectPayloadByteCount);
+	}
+
+	// Clear function call status so we can easily verify we get the huge payload.
+	for (UTestReplicatedIrisObject* ClientObject : ClientObjects)
+	{
+		UTestReplicatedIrisDynamicStatePropertyComponent* Component = ClientObject->DynamicStateComponents[0].Get();
+		Component->CallCounts = {};
+	}
+
+	uint32 ClientObjectsWithHugeArraysCount = 0;
+	for (uint32 RetryIt = 0; RetryIt != (HugeObjectMaxNetTickCountToArrive*HugeObjectCount) && ClientObjectsWithHugeArraysCount < HugeObjectCount; ++RetryIt)
+	{
+		Server->PreSendUpdate();
+		Server->SendAndDeliverTo(Client, DeliverPacket);
+		Server->PostSendUpdate();
+
+		ClientObjectsWithHugeArraysCount = 0;
+		for (UTestReplicatedIrisObject* ClientObject : ClientObjects)
+		{
+			UTestReplicatedIrisDynamicStatePropertyComponent* Component = ClientObject->DynamicStateComponents[0].Get();
+			ClientObjectsWithHugeArraysCount += (Component->CallCounts.IntArrayRepNotifyCounter > 0);
+		}
+	}
+
+	UE_NET_ASSERT_EQ(ClientObjectsWithHugeArraysCount, HugeObjectCount);
+}
+
 // Below test will fail as we only have a special path for reliable attachments for objects that stopped replicating, not for being filtered out.
 #if 0
 UE_NET_TEST_FIXTURE(FSplitObjectTestFixture, TestReliableAttachmentAddedAfterSplittingHugeObjectIsDeliveredBeforeObjectIsFilteredOut)
diff --git a/Engine/Source/Runtime/Experimental/Iris/Core/Private/Iris/ReplicationSystem/AttachmentReplication.cpp b/Engine/Source/Runtime/Experimental/Iris/Core/Private/Iris/ReplicationSystem/AttachmentReplication.cpp
--- a/Engine/Source/Runtime/Experimental/Iris/Core/Private/Iris/ReplicationSystem/AttachmentReplication.cpp
+++ b/Engine/Source/Runtime/Experimental/Iris/Core/Private/Iris/ReplicationSystem/AttachmentReplication.cpp
@@ -66,6 +66,11 @@ public:
 		return PreQueue.Num() == 0 && ReliableQueue.IsAllSentAndAcked();
 	}
 
+	bool IsSendWindowFull() const
+	{
+		return ReliableQueue.IsSendWindowFull();
+	}
+
 	uint32 GetUnsentBlobCount() const
 	{
 		return ReliableQueue.GetUnsentBlobCount() + static_cast<uint32>(PreQueue.Num());
@@ -106,7 +111,7 @@ public:
 		return true;
 	}
 
-	uint32 Serialize(FNetSerializationContext& Context, FNetRefHandle RefHandle, FReliableNetBlobQueue::ReplicationRecord& OutRecord)
+	uint32 Serialize(FNetSerializationContext& Context, FNetRefHandle RefHandle, FReliableNetBlobQueue::FReplicationRecord& OutRecord)
 	{
 		if (RefHandle.IsValid())
 		{
@@ -118,16 +123,16 @@ public:
 		}
 	}
 
-	void ProcessPacketDeliveryStatus(EPacketDeliveryStatus Status, FReliableNetBlobQueue::ReplicationRecord Record)
+	void ProcessPacketDeliveryStatus(EPacketDeliveryStatus Status, const FReliableNetBlobQueue::FReplicationRecord& Record)
 	{
 		ReliableQueue.ProcessPacketDeliveryStatus(Status, Record);
-		if (Status == EPacketDeliveryStatus::Delivered && Record != FReliableNetBlobQueue::InvalidReplicationRecord)
+		if (Status == EPacketDeliveryStatus::Delivered && Record.IsValid())
 		{
 			PopulateQueueFromPreQueue();
 		}
 	}
 
-	void CommitReplicationRecord(FReliableNetBlobQueue::ReplicationRecord Record)
+	void CommitReplicationRecord(const FReliableNetBlobQueue::FReplicationRecord& Record)
 	{
 		ReliableQueue.CommitReplicationRecord(Record);
 	}
@@ -230,6 +235,11 @@ bool FNetObjectAttachmentSendQueue::IsAllReliableSentAndAcked() const
 	return ReliableQueue == nullptr || ReliableQueue->IsAllSentAndAcked();
 }
 
+bool FNetObjectAttachmentSendQueue::CanSendMoreReliableAttachments() const
+{
+	return ReliableQueue == nullptr || !ReliableQueue->IsSendWindowFull();
+}
+
 bool FNetObjectAttachmentSendQueue::IsSafeToDestroy() const
 {
 	return UnreliableQueue.IsEmpty() && (ReliableQueue == nullptr || ReliableQueue->IsSafeToDestroy());
@@ -251,11 +261,11 @@ void FNetObjectAttachmentSendQueue::SetUnreliableQueueCapacity(uint32 QueueCapac
 	UnreliableQueue.PopNoCheck(DropCount);
 }
 
-EAttachmentWriteStatus FNetObjectAttachmentSendQueue::Serialize(FNetSerializationContext& Context, FNetRefHandle RefHandle, FNetObjectAttachmentSendQueue::ReplicationRecord& OutRecord, bool& bOutHasUnsentAttachments)
+EAttachmentWriteStatus FNetObjectAttachmentSendQueue::Serialize(FNetSerializationContext& Context, FNetRefHandle RefHandle, FNetObjectAttachmentSendQueue::FCommitRecord& OutRecord, bool& bOutHasUnsentAttachments)
 {
 	FNetBitStreamWriter& Writer = *Context.GetBitStreamWriter();
 
-	FInternalRecord ReplicationRecord;
+	FCommitRecord ReplicationRecord;
 	// This count is the total number of unsent reliable blobs. If the reliable window is full no blobs can be sentuntil some have been acked.
 	const uint32 UnsentReliableCount = (ReliableQueue != nullptr ? ReliableQueue->GetUnsentBlobCount() : 0U);
 	const bool bCanSendReliableAttachments = UnsentReliableCount > 0 && ReliableQueue->CanSendBlobs();
@@ -264,7 +274,7 @@ EAttachmentWriteStatus FNetObjectAttachmentSendQueue::Serialize(FNetSerializatio
 	{
 		// Ideally we shouldn't get here, but we can handle it. Important to overflow to report a soft error.
 		Writer.DoOverflow();
-		OutRecord = ReplicationRecord.CombinedRecord;
+		OutRecord = ReplicationRecord;
 		bOutHasUnsentAttachments = false;
 
 		if (UnsentReliableCount > 0)
@@ -286,7 +296,7 @@ EAttachmentWriteStatus FNetObjectAttachmentSendQueue::Serialize(FNetSerializatio
 
 	if (Writer.IsOverflown())
 	{
-		OutRecord = ReplicationRecord.CombinedRecord;
+		OutRecord = ReplicationRecord;
 		bOutHasUnsentAttachments = true;
 		return EAttachmentWriteStatus::BitstreamOverflow;
 	}
@@ -295,12 +305,12 @@ EAttachmentWriteStatus FNetObjectAttachmentSendQueue::Serialize(FNetSerializatio
 	if (bCanSendReliableAttachments)
 	{
 		UE_NET_TRACE_SCOPE(Reliable, Writer, Context.GetTraceCollector(), ENetTraceVerbosity::Trace);
-		SerializedReliableCount = SerializeReliable(Context, RefHandle, ReplicationRecord.ReliableRecord);
+		SerializedReliableCount = SerializeReliable(Context, RefHandle, ReplicationRecord.ReliableReplicationRecord);
 		// If we couldn't fit any reliable attachments then don't even try unreliable
 		if (SerializedReliableCount == 0)
 		{
 			Writer.DoOverflow();
-			OutRecord = ReplicationRecord.CombinedRecord;
+			OutRecord = ReplicationRecord;
 			bOutHasUnsentAttachments = true;
 			return EAttachmentWriteStatus::BitstreamOverflow;
 		}
@@ -313,7 +323,7 @@ EAttachmentWriteStatus FNetObjectAttachmentSendQueue::Serialize(FNetSerializatio
 	if (bHasUnreliableAttachments)
 	{
 		UE_NET_TRACE_SCOPE(Unreliable, Writer, Context.GetTraceCollector(), ENetTraceVerbosity::Trace);
-		SerializedUnreliableCount = SerializeUnreliable(Context, RefHandle, ReplicationRecord.UnreliableRecord);
+		SerializedUnreliableCount = SerializeUnreliable(Context, RefHandle, ReplicationRecord.UnreliableCommitRecord);
 		if (SerializedUnreliableCount == 0)
 		{
 			// If we didn't manage to send anything then inform the caller of this via overflowing the bitstream
@@ -341,17 +351,17 @@ EAttachmentWriteStatus FNetObjectAttachmentSendQueue::Serialize(FNetSerializatio
 		}
 	}
 
-	OutRecord = ReplicationRecord.CombinedRecord;
+	OutRecord = ReplicationRecord;
 	bOutHasUnsentAttachments = (SerializedReliableCount < UnsentReliableCount) || (SerializedUnreliableCount < UnreliableQueue.Count());
 	return WriteStatus;
 }
 
-uint32 FNetObjectAttachmentSendQueue::SerializeReliable(FNetSerializationContext& Context, FNetRefHandle RefHandle, FReliableNetBlobQueue::ReplicationRecord& OutRecord)
+uint32 FNetObjectAttachmentSendQueue::SerializeReliable(FNetSerializationContext& Context, FNetRefHandle RefHandle, FReliableNetBlobQueue::FReplicationRecord& OutRecord)
 {
 	return ReliableQueue->Serialize(Context, RefHandle, OutRecord);
 }
 
-uint32 FNetObjectAttachmentSendQueue::SerializeUnreliable(FNetSerializationContext& Context, FNetRefHandle RefHandle, uint32& OutRecord)
+uint32 FNetObjectAttachmentSendQueue::SerializeUnreliable(FNetSerializationContext& Context, FNetRefHandle RefHandle, FUnreliableReplicationRecord& OutRecord)
 {
 	FNetBitStreamWriter* Writer = Context.GetBitStreamWriter();
 	const FObjectReferenceCache* ObjectReferenceCache = Context.GetInternalContext()->ObjectReferenceCache;
@@ -404,47 +414,32 @@ uint32 FNetObjectAttachmentSendQueue::SerializeUnreliable(FNetSerializationConte
 			PrevHasMoreAttachmentsWritePos = HasMoreAttachmentsWritePos;
 		}
 	}
-
-	OutRecord = SerializedUnreliableCount;
+	FUnreliableReplicationRecord ReplicationRecord;
+	ReplicationRecord.Record = SerializedUnreliableCount;
+	OutRecord = ReplicationRecord;
 	return SerializedUnreliableCount;
 }
 
-void FNetObjectAttachmentSendQueue::CommitReplicationRecord(FNetObjectAttachmentSendQueue::ReplicationRecord Record)
+void FNetObjectAttachmentSendQueue::CommitReplicationRecord(const FNetObjectAttachmentSendQueue::FCommitRecord& Record)
 {
-	FInternalRecord InternalRecord;
-	InternalRecord.CombinedRecord = Record;
-	if (InternalRecord.UnreliableRecord)
+	if (Record.UnreliableCommitRecord.IsValid())
 	{
-		UnreliableQueue.PopNoCheck(InternalRecord.UnreliableRecord);
+		UnreliableQueue.PopNoCheck(Record.UnreliableCommitRecord.Record);
 	}
-	if (InternalRecord.ReliableRecord)
+	if (Record.ReliableReplicationRecord.IsValid())
 	{
-		ReliableQueue->CommitReplicationRecord(InternalRecord.ReliableRecord);
+		ReliableQueue->CommitReplicationRecord(Record.ReliableReplicationRecord);
 	}
 }
 
-void FNetObjectAttachmentSendQueue::OnPacketDelivered(FNetObjectAttachmentSendQueue::ReplicationRecord Record)
+void FNetObjectAttachmentSendQueue::ProcessPacketDeliveryStatus(EPacketDeliveryStatus Status, const FNetObjectAttachmentSendQueue::FReliableReplicationRecord& Record)
 {
 	if (ReliableQueue == nullptr)
 	{
 		return;
 	}
 
-	FInternalRecord InternalRecord;
-	InternalRecord.CombinedRecord = Record;
-	ReliableQueue->ProcessPacketDeliveryStatus(EPacketDeliveryStatus::Delivered, InternalRecord.ReliableRecord);
-}
-
-void FNetObjectAttachmentSendQueue::OnPacketLost(FNetObjectAttachmentSendQueue::ReplicationRecord Record)
-{
-	if (ReliableQueue == nullptr)
-	{
-		return;
-	}
-
-	FInternalRecord InternalRecord;
-	InternalRecord.CombinedRecord = Record;
-	ReliableQueue->ProcessPacketDeliveryStatus(EPacketDeliveryStatus::Lost, InternalRecord.ReliableRecord);
+	ReliableQueue->ProcessPacketDeliveryStatus(Status, Record);
 }
 
 bool FNetObjectAttachmentsWriter::Enqueue(ENetObjectAttachmentType Type, uint32 ObjectIndex, TArrayView<const TRefCountPtr<FNetBlob>> Attachments)
@@ -491,6 +486,17 @@ bool FNetObjectAttachmentsWriter::IsAllReliableSentAndAcked(ENetObjectAttachment
 	return Queue->IsAllSentAndAcked();
 }
 
+bool FNetObjectAttachmentsWriter::CanSendMoreReliableAttachments(ENetObjectAttachmentType Type, uint32 ObjectIndex) const
+{
+	const FNetObjectAttachmentSendQueue* Queue = GetQueue(Type, ObjectIndex);
+	if (Queue == nullptr)
+	{
+		return true;
+	}
+
+	return Queue->CanSendMoreReliableAttachments();
+}
+
 bool FNetObjectAttachmentsWriter::IsSafeToDestroy(ENetObjectAttachmentType Type, uint32 ObjectIndex) const
 {
 	const FNetObjectAttachmentSendQueue* Queue = GetQueue(Type, ObjectIndex);
@@ -523,38 +529,29 @@ void FNetObjectAttachmentsWriter::DropUnreliableAttachments(ENetObjectAttachment
 	}
 }
 
-EAttachmentWriteStatus FNetObjectAttachmentsWriter::Serialize(FNetSerializationContext& Context, ENetObjectAttachmentType Type, uint32 ObjectIndex, const FNetRefHandle RefHandle,  FNetObjectAttachmentsWriter::ReplicationRecord& OutRecord, bool& bOutHasUnsentAttachments)
+EAttachmentWriteStatus FNetObjectAttachmentsWriter::Serialize(FNetSerializationContext& Context, ENetObjectAttachmentType Type, uint32 ObjectIndex, const FNetRefHandle RefHandle,  FNetObjectAttachmentsWriter::FCommitRecord& OutRecord, bool& bOutHasUnsentAttachments)
 {
 	FNetObjectAttachmentSendQueue* Queue = GetQueue(Type, ObjectIndex);
 	// If this ensure fires we have bad logic for keeping track of whether there are attachments or not
-	if (!ensure(Queue != nullptr))
+	if (ensure(Queue != nullptr))
 	{
-		OutRecord = 0;
+		return Queue->Serialize(Context, RefHandle, OutRecord, bOutHasUnsentAttachments);
+	}
+	else
+	{
+		OutRecord = FCommitRecord();
 		bOutHasUnsentAttachments = false;
 		return EAttachmentWriteStatus::NoAttachments;
 	}
-
-	return Queue->Serialize(Context, RefHandle, OutRecord, bOutHasUnsentAttachments);
 }
 
-void FNetObjectAttachmentsWriter::CommitReplicationRecord(ENetObjectAttachmentType Type, uint32 ObjectIndex, FNetObjectAttachmentsWriter::ReplicationRecord Record)
+void FNetObjectAttachmentsWriter::CommitReplicationRecord(ENetObjectAttachmentType Type, uint32 ObjectIndex, const FNetObjectAttachmentsWriter::FCommitRecord& Record)
 {
 	FNetObjectAttachmentSendQueue* Queue = GetQueue(Type, ObjectIndex);
 	Queue->CommitReplicationRecord(Record);
 }
 
-void FNetObjectAttachmentsWriter::OnPacketDelivered(ENetObjectAttachmentType Type, uint32 ObjectIndex, FNetObjectAttachmentsWriter::ReplicationRecord Record)
-{
-	FNetObjectAttachmentSendQueue* Queue = GetQueue(Type, ObjectIndex);
-	if (Queue == nullptr)
-	{
-		return;
-	}
-
-	Queue->OnPacketDelivered(Record);
-}
-
-void FNetObjectAttachmentsWriter::OnPacketLost(ENetObjectAttachmentType Type, uint32 ObjectIndex, ReplicationRecord Record)
+void FNetObjectAttachmentsWriter::ProcessPacketDeliveryStatus(EPacketDeliveryStatus Status, ENetObjectAttachmentType Type, uint32 ObjectIndex, const FReliableReplicationRecord& Record)
 {
 	FNetObjectAttachmentSendQueue* Queue = GetQueue(Type, ObjectIndex);
 	if (Queue == nullptr)
@@ -562,7 +559,7 @@ void FNetObjectAttachmentsWriter::OnPacketLost(ENetObjectAttachmentType Type, ui
 		return;
 	}
 
-	Queue->OnPacketLost(Record);
+	Queue->ProcessPacketDeliveryStatus(Status, Record);
 }
 
 FNetObjectAttachmentSendQueue* FNetObjectAttachmentsWriter::GetQueue(ENetObjectAttachmentType Type, uint32 ObjectIndex)
diff --git a/Engine/Source/Runtime/Experimental/Iris/Core/Private/Iris/ReplicationSystem/AttachmentReplication.h b/Engine/Source/Runtime/Experimental/Iris/Core/Private/Iris/ReplicationSystem/AttachmentReplication.h
--- a/Engine/Source/Runtime/Experimental/Iris/Core/Private/Iris/ReplicationSystem/AttachmentReplication.h
+++ b/Engine/Source/Runtime/Experimental/Iris/Core/Private/Iris/ReplicationSystem/AttachmentReplication.h
@@ -46,8 +46,24 @@ enum class EAttachmentWriteStatus : unsigned
 class FNetObjectAttachmentSendQueue
 {
 public:
-	typedef uint64 ReplicationRecord;
-	static constexpr uint64 InvalidReplicationRecord = uint64(0);
+	typedef FReliableNetBlobQueue::FReplicationRecord FReliableReplicationRecord;
+	
+	struct FUnreliableReplicationRecord
+	{
+		bool IsValid() const { return Record != InvalidReplicationRecord; }
+
+		static constexpr uint32 InvalidReplicationRecord = 0;
+		uint32 Record = InvalidReplicationRecord;
+	};
+
+	// Commit record contains all data to be committed after serialization is committed to, i.e. will be part of a packet intended to be sent. The ReliableReplicationRecord needs to be part of the packet replication record so that we can act on packet notifications.
+	struct FCommitRecord
+	{
+		bool IsValid() const { return UnreliableCommitRecord.IsValid() || ReliableReplicationRecord.IsValid(); }
+
+		FReliableReplicationRecord ReliableReplicationRecord;
+		FUnreliableReplicationRecord UnreliableCommitRecord;
+	};
 
 public:
 	FNetObjectAttachmentSendQueue();
@@ -65,35 +81,21 @@ public:
 
 	bool IsAllReliableSentAndAcked() const;
 
+	bool CanSendMoreReliableAttachments() const;
+
 	void SetUnreliableQueueCapacity(uint32 QueueCapacity);
 
 private:
 	friend FNetObjectAttachmentsWriter;
 	class FReliableSendQueue;
 
-	struct FInternalRecord
-	{
-		union
-		{
-			struct
-			{
-				uint32 UnreliableRecord;
-				FReliableNetBlobQueue::ReplicationRecord ReliableRecord;
-			};
-			ReplicationRecord CombinedRecord;
-		};
-
-		FInternalRecord() { CombinedRecord = InvalidReplicationRecord; };
-	};
-
-	EAttachmentWriteStatus Serialize(FNetSerializationContext& Context, FNetRefHandle RefHandle, ReplicationRecord& OutRecord, bool& bOutHasUnprocessedAttachments);
-	uint32 SerializeReliable(FNetSerializationContext& Context, FNetRefHandle RefHandle, FReliableNetBlobQueue::ReplicationRecord& OutRecord);
-	uint32 SerializeUnreliable(FNetSerializationContext& Context, FNetRefHandle RefHandle, uint32& OutRecord);
+	EAttachmentWriteStatus Serialize(FNetSerializationContext& Context, FNetRefHandle RefHandle, FCommitRecord& OutRecord, bool& bOutHasUnprocessedAttachments);
+	uint32 SerializeReliable(FNetSerializationContext& Context, FNetRefHandle RefHandle, FReliableReplicationRecord& OutRecord);
+	uint32 SerializeUnreliable(FNetSerializationContext& Context, FNetRefHandle RefHandle, FUnreliableReplicationRecord& OutRecord);
 
-	void CommitReplicationRecord(ReplicationRecord Record);
+	void CommitReplicationRecord(const FCommitRecord& Record);
 
-	void OnPacketDelivered(ReplicationRecord Record);
-	void OnPacketLost(ReplicationRecord Record);
+	void ProcessPacketDeliveryStatus(EPacketDeliveryStatus Status, const FReliableReplicationRecord& Record);
 
 	FReliableSendQueue* ReliableQueue;
 	TResizableCircularQueue<TRefCountPtr<FNetBlob>> UnreliableQueue;
@@ -103,7 +105,8 @@ private:
 class FNetObjectAttachmentsWriter
 {
 public:
-	typedef FNetObjectAttachmentSendQueue::ReplicationRecord ReplicationRecord;
+	typedef FNetObjectAttachmentSendQueue::FCommitRecord FCommitRecord;
+	typedef FNetObjectAttachmentSendQueue::FReliableReplicationRecord FReliableReplicationRecord;
 
 public:
 	bool Enqueue(ENetObjectAttachmentType Type, uint32 ObjectIndex, TArrayView<const TRefCountPtr<FNetBlob>> Attachments);
@@ -115,18 +118,20 @@ public:
 	// Whether all queued reliable attachments have been sent and acked
 	bool IsAllReliableSentAndAcked(ENetObjectAttachmentType Type, uint32 ObjectIndex) const;
 
+	// Whether more reliable attachments can be sent now. It's possible to queue up as many attachments as you see fit, but if the queue is full it can take a while before more attachments will be replicated.
+	bool CanSendMoreReliableAttachments(ENetObjectAttachmentType Type, uint32 ObjectIndex) const;
+
 	// Whether the queue can be destroyed without causing issues if more attachments are queued to this instance.
 	bool IsSafeToDestroy(ENetObjectAttachmentType Type, uint32 ObjectIndex) const;
 
 	void DropAllAttachments(ENetObjectAttachmentType Type, uint32 ObjectIndex);
 	void DropUnreliableAttachments(ENetObjectAttachmentType Type, uint32 ObjectIndex, bool& bOutHasUnsentAttachments);
 
-	EAttachmentWriteStatus Serialize(FNetSerializationContext& Context, ENetObjectAttachmentType Type, uint32 ObjectIndex, FNetRefHandle RefHandle, ReplicationRecord& OutRecord, bool& bOutHasUnsentAttachments);
+	EAttachmentWriteStatus Serialize(FNetSerializationContext& Context, ENetObjectAttachmentType Type, uint32 ObjectIndex, FNetRefHandle RefHandle, FCommitRecord& OutRecord, bool& bOutHasUnsentAttachments);
 
-	void CommitReplicationRecord(ENetObjectAttachmentType Type, uint32 ObjectIndex, ReplicationRecord Record);
+	void CommitReplicationRecord(ENetObjectAttachmentType Type, uint32 ObjectIndex, const FCommitRecord& Record);
 
-	void OnPacketDelivered(ENetObjectAttachmentType Type, uint32 ObjectIndex, ReplicationRecord Record);
-	void OnPacketLost(ENetObjectAttachmentType Type, uint32 ObjectIndex, ReplicationRecord Record);
+	void ProcessPacketDeliveryStatus(EPacketDeliveryStatus Status, ENetObjectAttachmentType Type, uint32 ObjectIndex, const FReliableReplicationRecord& Record);
 
 private:
 	bool NetBlobMightNeedSplitting(const TRefCountPtr<FNetObjectAttachment>& Attachment) const;
diff --git a/Engine/Source/Runtime/Experimental/Iris/Core/Private/Iris/ReplicationSystem/NetBlob/ReliableNetBlobQueue.cpp b/Engine/Source/Runtime/Experimental/Iris/Core/Private/Iris/ReplicationSystem/NetBlob/ReliableNetBlobQueue.cpp
--- a/Engine/Source/Runtime/Experimental/Iris/Core/Private/Iris/ReplicationSystem/NetBlob/ReliableNetBlobQueue.cpp
+++ b/Engine/Source/Runtime/Experimental/Iris/Core/Private/Iris/ReplicationSystem/NetBlob/ReliableNetBlobQueue.cpp
@@ -47,18 +47,20 @@ bool FReliableNetBlobQueue::IsSafeToDestroy() const
 	return true;
 }
 
-uint32 FReliableNetBlobQueue::Serialize(FNetSerializationContext& Context, FReliableNetBlobQueue::ReplicationRecord& OutRecord)
+uint32 FReliableNetBlobQueue::Serialize(FNetSerializationContext& Context, FReliableNetBlobQueue::FReplicationRecord& OutRecord)
 {
 	FNetRefHandle InvalidNetHandle;
-	return SerializeInternal(Context, InvalidNetHandle, OutRecord, false);
+	constexpr bool bSerializeWithObject = false;
+	return SerializeInternal(Context, InvalidNetHandle, OutRecord, bSerializeWithObject);
 }
 
-uint32 FReliableNetBlobQueue::SerializeWithObject(FNetSerializationContext& Context, FNetRefHandle RefHandle, FReliableNetBlobQueue::ReplicationRecord& OutRecord)
+uint32 FReliableNetBlobQueue::SerializeWithObject(FNetSerializationContext& Context, FNetRefHandle RefHandle, FReliableNetBlobQueue::FReplicationRecord& OutRecord)
 {
-	return SerializeInternal(Context, RefHandle, OutRecord, true);
+	constexpr bool bSerializeWithObject = true;
+	return SerializeInternal(Context, RefHandle, OutRecord, bSerializeWithObject);
 }
 
-uint32 FReliableNetBlobQueue::SerializeInternal(FNetSerializationContext& Context, FNetRefHandle RefHandle, FReliableNetBlobQueue::ReplicationRecord& OutRecord, const bool bSerializeWithObject)
+uint32 FReliableNetBlobQueue::SerializeInternal(FNetSerializationContext& Context, FNetRefHandle RefHandle, FReliableNetBlobQueue::FReplicationRecord& OutRecord, const bool bSerializeWithObject)
 {
 	FNetBitStreamWriter* Writer = Context.GetBitStreamWriter();
 	const FObjectReferenceCache* ObjectReferenceCache = Context.GetInternalContext()->ObjectReferenceCache;
@@ -66,9 +68,14 @@ uint32 FReliableNetBlobQueue::SerializeInternal(FNetSerializationContext& Contex
 	uint32 PrevHasMoreBlobsWritePos = 0;
 	uint32 SerializedCount = 0;
 	uint32 PrevWrittenSeq = ~0U;
-	uint32 WrittenSeq[2] = {~0U, ~0U};
-	uint32 WrittenCount[2] = {0U, 0U};
 	uint32 WrittenIndex = 0;
+	uint32 WrittenCount[MaxWriteSequenceCount] = {};
+	uint32 WrittenSeq[MaxWriteSequenceCount];
+	for (uint32& Sequence : WrittenSeq)
+	{
+		Sequence = ~0U;
+	}
+
 	for (uint32 Seq = FirstSeq, EndSeq = LastSeq; Seq < EndSeq; ++Seq)
 	{
 		const uint32 Index = SequenceToIndex(Seq);
@@ -87,7 +94,7 @@ uint32 FReliableNetBlobQueue::SerializeInternal(FNetSerializationContext& Contex
 		else if (PrevWrittenSeq + 1U != Seq)
 		{
 			++WrittenIndex;
-			// Our ReplicationRecord can't hold a lot of info.
+			// There's limited support for disjoint sequences in the replication record.
 			if (WrittenIndex >= UE_ARRAY_COUNT(WrittenCount))
 			{
 				break;
@@ -139,7 +146,6 @@ uint32 FReliableNetBlobQueue::SerializeInternal(FNetSerializationContext& Contex
 
 			++SerializedCount;
 
-			// Our ReplicationRecord currently only holds two sequences with number and count, one byte each.
 			++WrittenCount[WrittenIndex];
 			if (WrittenCount[WrittenIndex] == 255)
 			{
@@ -155,22 +161,25 @@ uint32 FReliableNetBlobQueue::SerializeInternal(FNetSerializationContext& Contex
 		}
 	}
 
-	//
-	const ReplicationRecord Record = ((WrittenSeq[0] & 255) << 24U) | ((WrittenCount[0] & 255) << 16U) | ((WrittenSeq[1] & 255) << 8U) | ((WrittenCount[1] & 255) << 0U);
-	OutRecord = (SerializedCount ? Record : FReliableNetBlobQueue::InvalidReplicationRecord);
+	// Assemble replication record
+	FReplicationRecord Record;
+	for (uint32 Index = 0; Index < MaxWriteSequenceCount; ++Index)
+	{
+		Record.Sequences[Index] = WrittenSeq[Index] & 255U;
+		Record.Counts[Index] = WrittenCount[Index] & 255U;
+	}
+	OutRecord = Record;
 
 	return SerializedCount;
 }
 
-void FReliableNetBlobQueue::CommitReplicationRecord(FReliableNetBlobQueue::ReplicationRecord Record)
+void FReliableNetBlobQueue::CommitReplicationRecord(const FReliableNetBlobQueue::FReplicationRecord& Record)
 {
-	uint32 RecordSeq[2] = {(Record >> 24U) & 255, (Record >> 8U) & 255U};
-	uint32 RecordCount[2] = {(Record >> 16U) & 255, (Record >> 0U) & 255U};
-
-	UnsentBlobCount -= RecordCount[0] + RecordCount[1];
-	for (const uint32 Index : {0U, 1U})
+	for (uint32 Index = 0, EndIndex = MaxWriteSequenceCount; Index != EndIndex; ++Index)
 	{
-		for (uint32 Seq = RecordSeq[Index], EndSeq = Seq + RecordCount[Index]; Seq != EndSeq; ++Seq)
+		const uint32 Count = Record.Counts[Index];
+		UnsentBlobCount -= Count;
+		for (uint32 Seq = Record.Sequences[Index], EndSeq = Seq + Count; Seq != EndSeq; ++Seq)
 		{
 			SetSequenceIsSent(Seq);
 		}
@@ -230,7 +239,7 @@ uint32 FReliableNetBlobQueue::DeserializeInternal(FNetSerializationContext& Cont
 		const TRefCountPtr<FNetBlob>& Blob = BlobReceiver->CreateNetBlob(CreationInfo);
 		if (!Blob.IsValid())
 		{
-			UE_LOG(LogIris, Warning, TEXT("%s"), TEXT("Unable to create blob."));
+			UE_LOG(LogIris, Warning, TEXT("%hs"), "Unable to create blob.");
 			Context.SetError(GNetError_UnsupportedNetBlob);
 
 			return DeserializedCount;
@@ -266,7 +275,7 @@ uint32 FReliableNetBlobQueue::DeserializeInternal(FNetSerializationContext& Cont
 
 bool FReliableNetBlobQueue::Enqueue(const TRefCountPtr<FNetBlob>& Blob)
 {
-	if (IsFull())
+	if (IsSendWindowFull())
 	{
 		return false;
 	}
@@ -302,7 +311,7 @@ void FReliableNetBlobQueue::Pop()
 	// $TODO. For a memory optimization one can change the storage implementation and free the memory here if everything is acked.
 }
 
-void FReliableNetBlobQueue::ProcessPacketDeliveryStatus(EPacketDeliveryStatus Status, FReliableNetBlobQueue::ReplicationRecord Record)
+void FReliableNetBlobQueue::ProcessPacketDeliveryStatus(EPacketDeliveryStatus Status, const FReliableNetBlobQueue::FReplicationRecord& Record)
 {
 	switch (Status)
 	{
@@ -316,6 +325,12 @@ void FReliableNetBlobQueue::ProcessPacketDeliveryStatus(EPacketDeliveryStatus St
 		OnPacketDropped(Record);
 		break;
 	}
+	case EPacketDeliveryStatus::Discard:
+	{
+		// Pretend that it was delivered.
+		OnPacketDelivered(Record);
+		break;
+	}
 	default:
 	{
 		break;
@@ -323,17 +338,20 @@ void FReliableNetBlobQueue::ProcessPacketDeliveryStatus(EPacketDeliveryStatus St
 	}
 }
 
-void FReliableNetBlobQueue::OnPacketDelivered(FReliableNetBlobQueue::ReplicationRecord Record)
+void FReliableNetBlobQueue::OnPacketDelivered(const FReliableNetBlobQueue::FReplicationRecord& Record)
 {
-	uint32 RecordSeq[2] = {(Record >> 24U) & 255, (Record >> 8U) & 255U};
-	uint32 RecordCount[2] = {(Record >> 16U) & 255, (Record >> 0U) & 255U};
-
 	// Mark blobs as acked
-	for (const uint32 Index : {0U, 1U})
+	for (uint32 SeqIt = 0, EndSeqIt = MaxWriteSequenceCount; SeqIt != EndSeqIt; ++SeqIt)
 	{
-		for (uint32 Seq = RecordSeq[Index], EndSeq = Seq + RecordCount[Index]; Seq != EndSeq; ++Seq)
+		const uint32 Count = Record.Counts[SeqIt];
+		for (uint32 Seq = Record.Sequences[SeqIt], EndSeq = Seq + Count; Seq != EndSeq; ++Seq)
 		{
-			SetSequenceIsAcked(Seq);
+			const uint32 Index = SequenceToIndex(Seq);
+
+			SetIndexIsAcked(Index);
+
+			// Release blob as quickly as possible.
+			NetBlobs[Index].SafeRelease();
 		}
 	}
 
@@ -341,16 +359,14 @@ void FReliableNetBlobQueue::OnPacketDelivered(FReliableNetBlobQueue::Replication
 	PopInOrderAckedBlobs();
 }
 
-void FReliableNetBlobQueue::OnPacketDropped(FReliableNetBlobQueue::ReplicationRecord Record)
+void FReliableNetBlobQueue::OnPacketDropped(const FReliableNetBlobQueue::FReplicationRecord& Record)
 {
-	uint32 RecordSeq[2] = {(Record >> 24U) & 255, (Record >> 8U) & 255U};
-	uint32 RecordCount[2] = {(Record >> 16U) & 255, (Record >> 0U) & 255U};
-
 	// Mark blobs as unsent
-	UnsentBlobCount += RecordCount[0] + RecordCount[1];
-	for (const uint32 Index : {0U, 1U})
+	for (uint32 SeqIt = 0, EndSeqIt = MaxWriteSequenceCount; SeqIt != EndSeqIt; ++SeqIt)
 	{
-		for (uint32 Seq = RecordSeq[Index], EndSeq = Seq + RecordCount[Index]; Seq != EndSeq; ++Seq)
+		const uint32 Count = Record.Counts[SeqIt];
+		UnsentBlobCount += Count;
+		for (uint32 Seq = Record.Sequences[SeqIt], EndSeq = Seq + Count; Seq != EndSeq; ++Seq)
 		{
 			ClearSequenceIsSent(Seq);
 		}
@@ -367,7 +383,6 @@ void FReliableNetBlobQueue::PopInOrderAckedBlobs()
 			return;
 		}
 
-		NetBlobs[Index].SafeRelease();
 		ClearIndexIsAcked(Index);
 		ClearIndexIsSent(Index);
 	}
diff --git a/Engine/Source/Runtime/Experimental/Iris/Core/Private/Iris/ReplicationSystem/ReplicationWriter.cpp b/Engine/Source/Runtime/Experimental/Iris/Core/Private/Iris/ReplicationSystem/ReplicationWriter.cpp
--- a/Engine/Source/Runtime/Experimental/Iris/Core/Private/Iris/ReplicationSystem/ReplicationWriter.cpp
+++ b/Engine/Source/Runtime/Experimental/Iris/Core/Private/Iris/ReplicationSystem/ReplicationWriter.cpp
@@ -51,9 +51,19 @@ static FAutoConsoleVariableRef CVarWarnAboutDroppedAttachmentsToObjectsNotInScop
 	));
 
 static int32 GReplicationWriterMaxAllowedPacketsIfNotHugeObject = 3;
-static FAutoConsoleVariableRef CVarReplicationWriterMaxAllowedPacketsIfNotHugeObject(TEXT("net.iris.ReplicationWriterMaxAllowedPacketsIfNotHugeObject"), GReplicationWriterMaxAllowedPacketsIfNotHugeObject,
+static FAutoConsoleVariableRef CVarReplicationWriterMaxAllowedPacketsIfNotHugeObject(TEXT("net.Iris.ReplicationWriterMaxAllowedPacketsIfNotHugeObject"), GReplicationWriterMaxAllowedPacketsIfNotHugeObject,
 	TEXT("Allow ReplicationWriter to overcommit data if we have more data to write."));
 
+/*
+ * net.Iris.ReplicationWriterMaxHugeObjectsInTransit
+ * There's a tradeoff mainly between the connection characteristics to support and normal object replication scheduling when tweaking this value.
+ * On one hand you don't want to end up stalling object replication because the top priority objects are huge. So you want to be able to keep replicating huge objects during the maximum latency, including latency variation, and packet loss scenarios 
+ * you want to provide the best experience possible for. On the other hand object deletion cannot be performed once the object is in the huge object queue. Consider this and how long time it will take to replicate the huge object queue depending on the average payload of a huge object.
+ */
+static int32 GReplicationWriterMaxHugeObjectsInTransit = 16;
+static FAutoConsoleVariableRef CVarReplicationWriterMaxHugeObjectsInTransit(TEXT("net.Iris.ReplicationWriterMaxHugeObjectsInTransit"), GReplicationWriterMaxHugeObjectsInTransit,
+	TEXT("How many very large objects, one whose payload doesn't fit in a single packet, is allowed to be scheduled for send. Needs to be at least 1."));
+
 static bool bValidateObjectsWithDirtyChanges = false;
 static FAutoConsoleVariableRef CvarValidateObjectsWithDirtyChanges(TEXT("net.Iris.ReplicationWriter.ValidateObjectsWithDirtyChanges"), bValidateObjectsWithDirtyChanges, TEXT("Ensure that we don't try to mark invalid objects as dirty when they shouldn't."));
 
@@ -175,26 +185,6 @@ void FReplicationWriter::FReplicationInfo::SetState(EReplicatedObjectState NewSt
 }
 
 
-FReplicationWriter::FHugeObjectContext::FHugeObjectContext()
-: SendStatus(EHugeObjectSendStatus::Idle)
-, InternalIndex(0)
-, TraceCollector(nullptr)
-, DebugName(CreatePersistentNetDebugName(TEXT("HugeObjectState"), UE_ARRAY_COUNT(TEXT("HugeObjectState"))))
-, StartSendingTime(0)
-, EndSendingTime(0)
-, StartStallTime(0)
-{
-#if UE_NET_TRACE_ENABLED
-	DebugName->DebugNameId = FNetTrace::TraceName(DebugName->Name);
-#endif
-}
-
-FReplicationWriter::FHugeObjectContext::~FHugeObjectContext()
-{
-	UE_NET_TRACE_DESTROY_COLLECTOR(TraceCollector);
-	TraceCollector = nullptr;
-}
-
 // Default allocator for changemasks
 static FGlobalChangeMaskAllocator s_DefaultChangeMaskAllocator;
 
@@ -222,7 +212,7 @@ struct TReplicationRecordHelper
 		for (uint32 It = 0; It < RecordInfoCount; ++It)
 		{
  			const FReplicationRecord::FRecordInfo& RecordInfo = ReplicationRecord->PeekInfo();
-			FNetObjectAttachmentsWriter::ReplicationRecord AttachmentRecord = (RecordInfo.HasAttachments ? ReplicationRecord->DequeueAttachmentRecord() : 0U);
+			FNetObjectAttachmentsWriter::FReliableReplicationRecord AttachmentRecord(RecordInfo.HasAttachments ? ReplicationRecord->DequeueAttachmentRecord() : uint64(0));
 			FReplicationInfo& Info = ReplicationInfos[RecordInfo.Index];
 			FReplicationRecord::FRecordInfoList& RecordInfoList = ReplicationInfosRecordInfoLists[RecordInfo.Index];
 
@@ -299,6 +289,10 @@ static bool s_ValidateReplicationRecord(const FReplicationRecord* ReplicationRec
 FReplicationWriter::~FReplicationWriter()
 {
 	DiscardAllRecords();
+
+	// Freeing the huge object queue needs to be done before calling StopAllReplication() in order to be able to free any changemask allocations.
+	FreeHugeObjectSendQueue();
+	
 	StopAllReplication();
 }
 
@@ -950,7 +944,7 @@ uint32 FReplicationWriter::SortScheduledObjects(FScheduleObjectInfo* ScheduledOb
 	return FMath::Min(ScheduledObjectCount - StartIndex, PartialSortObjectCount);
 }
 
-void FReplicationWriter::HandleDeliveredRecord(const FReplicationRecord::FRecordInfo& RecordInfo, FReplicationInfo& Info, FNetObjectAttachmentsWriter::ReplicationRecord AttachmentRecord)
+void FReplicationWriter::HandleDeliveredRecord(const FReplicationRecord::FRecordInfo& RecordInfo, FReplicationInfo& Info, const FNetObjectAttachmentsWriter::FReliableReplicationRecord& AttachmentRecord)
 {
 	EReplicatedObjectState DeliveredState = (EReplicatedObjectState)RecordInfo.ReplicatedObjectState;
 	EReplicatedObjectState CurrentState = Info.GetState();
@@ -1028,24 +1022,24 @@ void FReplicationWriter::HandleDeliveredRecord(const FReplicationRecord::FRecord
 		case EReplicatedObjectState::AttachmentToObjectNotInScope:
 		{
 			check(IsObjectIndexForOOBAttachment(InternalIndex));
-			Attachments.OnPacketDelivered(ENetObjectAttachmentType::OutOfBand, ObjectIndexForOOBAttachment, AttachmentRecord);
+			Attachments.ProcessPacketDeliveryStatus(EPacketDeliveryStatus::Delivered, ENetObjectAttachmentType::OutOfBand, ObjectIndexForOOBAttachment, AttachmentRecord);
 		}
 		return;
 
 		case EReplicatedObjectState::HugeObject:
 		{
 			check(IsObjectIndexForOOBAttachment(InternalIndex));
-			Attachments.OnPacketDelivered(ENetObjectAttachmentType::HugeObject, ObjectIndexForOOBAttachment, AttachmentRecord);
+			Attachments.ProcessPacketDeliveryStatus(EPacketDeliveryStatus::Delivered, ENetObjectAttachmentType::HugeObject, ObjectIndexForOOBAttachment, AttachmentRecord);
 
-			// If we've sent the entire state now we can clear the huge object state and proceed as normal.
-			if (Attachments.IsAllSentAndAcked(ENetObjectAttachmentType::HugeObject, ObjectIndexForOOBAttachment))
+			HugeObjectSendQueue.AckObjects([this](const FHugeObjectContext& HugeObjectContext)
 			{
-				FReplicationInfo& ReplicationInfo = GetReplicationInfo(HugeObjectContext.InternalIndex);
+				// If we've sent an entire huge objects we can ack everything in the payload and continue replicating this object using normal means.
+				const FReplicationInfo& HugeObjectReplicationInfo = this->GetReplicationInfo(HugeObjectContext.RootObjectInternalIndex);
 				for (const FObjectRecord& ObjectRecord : HugeObjectContext.BatchRecord.ObjectReplicationRecords)
 				{
-					FReplicationInfo& HugeObjectReplicationInfo =  GetReplicationInfo(ObjectRecord.Record.Index);
-					const uint32 ChangeMaskBitCount = HugeObjectReplicationInfo.ChangeMaskBitCount;
-					HandleDeliveredRecord(ObjectRecord.Record, HugeObjectReplicationInfo, ObjectRecord.AttachmentRecord);
+					FReplicationInfo& ReplicationInfo = this->GetReplicationInfo(ObjectRecord.Record.Index);
+					const uint32 ChangeMaskBitCount = ReplicationInfo.ChangeMaskBitCount;
+					this->HandleDeliveredRecord(ObjectRecord.Record, ReplicationInfo, ObjectRecord.AttachmentRecord);
 					if (ObjectRecord.Record.HasChangeMask)
 					{
 						FChangeMaskStorageOrPointer::Free(ObjectRecord.Record.ChangeMaskOrPtr, ChangeMaskBitCount, s_DefaultChangeMaskAllocator);
@@ -1053,10 +1047,8 @@ void FReplicationWriter::HandleDeliveredRecord(const FReplicationRecord::FRecord
 				}
 
 				// We need to explicitly acknowledge exports made through the huge object batch
-				NetExports->AcknowledgeBatchExports(HugeObjectContext.BatchExports);
-				
-				ClearHugeObjectContext(HugeObjectContext);
-			}
+				this->NetExports->AcknowledgeBatchExports(HugeObjectContext.BatchExports);
+			});
 		}
 		return;
 
@@ -1066,7 +1058,7 @@ void FReplicationWriter::HandleDeliveredRecord(const FReplicationRecord::FRecord
 
 	if (RecordInfo.HasAttachments)
 	{
-		Attachments.OnPacketDelivered(ENetObjectAttachmentType::Normal, InternalIndex, AttachmentRecord);
+		Attachments.ProcessPacketDeliveryStatus(EPacketDeliveryStatus::Delivered, ENetObjectAttachmentType::Normal, InternalIndex, AttachmentRecord);
 	}
 
 	// Must process WaitOnflush after attachments in order to correctly evaluate flush-status if needed
@@ -1125,7 +1117,7 @@ void FReplicationWriter::HandleDeliveredRecord(const FReplicationRecord::FRecord
 	}
 }
 
-void FReplicationWriter::HandleDiscardedRecord(const FReplicationRecord::FRecordInfo& RecordInfo, FReplicationInfo& Info, FNetObjectAttachmentsWriter::ReplicationRecord AttachmentRecord)
+void FReplicationWriter::HandleDiscardedRecord(const FReplicationRecord::FRecordInfo& RecordInfo, FReplicationInfo& Info, const FNetObjectAttachmentsWriter::FReliableReplicationRecord& AttachmentRecord)
 {
 	const EReplicatedObjectState DeliveredState = (EReplicatedObjectState)RecordInfo.ReplicatedObjectState;
 	const uint32 InternalIndex = RecordInfo.Index;
@@ -1144,28 +1136,29 @@ void FReplicationWriter::HandleDiscardedRecord(const FReplicationRecord::FRecord
 			check(IsObjectIndexForOOBAttachment(InternalIndex));
 
 			// Deal with it similar to if the entire state has been sent as we need to go through all records.
+			Attachments.ProcessPacketDeliveryStatus(EPacketDeliveryStatus::Discard, ENetObjectAttachmentType::HugeObject, ObjectIndexForOOBAttachment, AttachmentRecord);
+
+			HugeObjectSendQueue.AckObjects([this](const FHugeObjectContext& HugeObjectContext)
 			{
-				FReplicationInfo& ReplicationInfo = GetReplicationInfo(HugeObjectContext.InternalIndex);
+				const FReplicationInfo& HugeObjectReplicationInfo = this->GetReplicationInfo(HugeObjectContext.RootObjectInternalIndex);
 				for (const FObjectRecord& ObjectRecord : HugeObjectContext.BatchRecord.ObjectReplicationRecords)
 				{
-					FReplicationInfo& HugeObjectReplicationInfo = GetReplicationInfo(ObjectRecord.Record.Index);
-					const uint32 ChangeMaskBitCount = HugeObjectReplicationInfo.ChangeMaskBitCount;
-					HandleDiscardedRecord(ObjectRecord.Record, HugeObjectReplicationInfo, ObjectRecord.AttachmentRecord);
+					FReplicationInfo& ReplicationInfo = this->GetReplicationInfo(ObjectRecord.Record.Index);
+					const uint32 ChangeMaskBitCount = ReplicationInfo.ChangeMaskBitCount;
+					this->HandleDiscardedRecord(ObjectRecord.Record, ReplicationInfo, ObjectRecord.AttachmentRecord);
 					if (ObjectRecord.Record.HasChangeMask)
 					{
 						FChangeMaskStorageOrPointer::Free(ObjectRecord.Record.ChangeMaskOrPtr, ChangeMaskBitCount, s_DefaultChangeMaskAllocator);
 					}
 				}
-
-				ClearHugeObjectContext(HugeObjectContext);
-			}
+			});
 		}
 		return;
 	}
 }
 
 template<>
-void FReplicationWriter::HandleDroppedRecord<FReplicationWriter::EReplicatedObjectState::WaitOnCreateConfirmation>(FReplicationWriter::EReplicatedObjectState CurrentState, const FReplicationRecord::FRecordInfo& RecordInfo, FReplicationInfo& Info, FNetObjectAttachmentsWriter::ReplicationRecord AttachmentRecord)
+void FReplicationWriter::HandleDroppedRecord<FReplicationWriter::EReplicatedObjectState::WaitOnCreateConfirmation>(FReplicationWriter::EReplicatedObjectState CurrentState, const FReplicationRecord::FRecordInfo& RecordInfo, FReplicationInfo& Info, const FNetObjectAttachmentsWriter::FReliableReplicationRecord& AttachmentRecord)
 {
 	const uint32 InternalIndex = RecordInfo.Index;
 
@@ -1221,7 +1214,7 @@ void FReplicationWriter::HandleDroppedRecord<FReplicationWriter::EReplicatedObje
 }
 
 template<>
-void FReplicationWriter::HandleDroppedRecord<FReplicationWriter::EReplicatedObjectState::Created>(FReplicationWriter::EReplicatedObjectState CurrentState, const FReplicationRecord::FRecordInfo& RecordInfo, FReplicationInfo& Info, FNetObjectAttachmentsWriter::ReplicationRecord AttachmentRecord)
+void FReplicationWriter::HandleDroppedRecord<FReplicationWriter::EReplicatedObjectState::Created>(FReplicationWriter::EReplicatedObjectState CurrentState, const FReplicationRecord::FRecordInfo& RecordInfo, FReplicationInfo& Info, const FNetObjectAttachmentsWriter::FReliableReplicationRecord& AttachmentRecord)
 {
 	const uint32 InternalIndex = RecordInfo.Index;
 
@@ -1298,7 +1291,7 @@ void FReplicationWriter::HandleDroppedRecord<FReplicationWriter::EReplicatedObje
 }
 
 template<>
-void FReplicationWriter::HandleDroppedRecord<FReplicationWriter::EReplicatedObjectState::WaitOnDestroyConfirmation>(FReplicationWriter::EReplicatedObjectState CurrentState, const FReplicationRecord::FRecordInfo& RecordInfo, FReplicationInfo& Info, FNetObjectAttachmentsWriter::ReplicationRecord AttachmentRecord)
+void FReplicationWriter::HandleDroppedRecord<FReplicationWriter::EReplicatedObjectState::WaitOnDestroyConfirmation>(FReplicationWriter::EReplicatedObjectState CurrentState, const FReplicationRecord::FRecordInfo& RecordInfo, FReplicationInfo& Info, const FNetObjectAttachmentsWriter::FReliableReplicationRecord& AttachmentRecord)
 {
 	const uint32 InternalIndex = RecordInfo.Index;
 
@@ -1424,7 +1417,7 @@ void FReplicationWriter::HandleDroppedRecord<FReplicationWriter::EReplicatedObje
 	}
 }
 
-void FReplicationWriter::HandleDroppedRecord(const FReplicationRecord::FRecordInfo& RecordInfo, FReplicationInfo& Info, FNetObjectAttachmentsWriter::ReplicationRecord AttachmentRecord)
+void FReplicationWriter::HandleDroppedRecord(const FReplicationRecord::FRecordInfo& RecordInfo, FReplicationInfo& Info, const FNetObjectAttachmentsWriter::FReliableReplicationRecord& AttachmentRecord)
 {
 	EReplicatedObjectState LostObjectState = (EReplicatedObjectState)RecordInfo.ReplicatedObjectState;
 	EReplicatedObjectState CurrentState = Info.GetState();
@@ -1476,14 +1469,14 @@ void FReplicationWriter::HandleDroppedRecord(const FReplicationRecord::FRecordIn
 		case EReplicatedObjectState::AttachmentToObjectNotInScope:
 		{
 			check(IsObjectIndexForOOBAttachment(InternalIndex));
-			Attachments.OnPacketLost(ENetObjectAttachmentType::OutOfBand, ObjectIndexForOOBAttachment, AttachmentRecord);
+			Attachments.ProcessPacketDeliveryStatus(EPacketDeliveryStatus::Lost, ENetObjectAttachmentType::OutOfBand, ObjectIndexForOOBAttachment, AttachmentRecord);
 		}
 		return;
 
 		case EReplicatedObjectState::HugeObject:
 		{
 			check(IsObjectIndexForOOBAttachment(InternalIndex));
-			Attachments.OnPacketLost(ENetObjectAttachmentType::HugeObject, ObjectIndexForOOBAttachment, AttachmentRecord);
+			Attachments.ProcessPacketDeliveryStatus(EPacketDeliveryStatus::Lost, ENetObjectAttachmentType::HugeObject, ObjectIndexForOOBAttachment, AttachmentRecord);
 		}
 		return;
 
@@ -1494,7 +1487,7 @@ void FReplicationWriter::HandleDroppedRecord(const FReplicationRecord::FRecordIn
 
 	if (RecordInfo.HasAttachments)
 	{
-		Attachments.OnPacketLost(ENetObjectAttachmentType::Normal, InternalIndex, AttachmentRecord);
+		Attachments.ProcessPacketDeliveryStatus(EPacketDeliveryStatus::Lost, ENetObjectAttachmentType::Normal, InternalIndex, AttachmentRecord);
 	}
 }
 
@@ -1513,7 +1506,7 @@ void FReplicationWriter::ProcessDeliveryNotification(EPacketDeliveryStatus Packe
 		if (PacketDeliveryStatus == EPacketDeliveryStatus::Delivered)
 		{
 			Helper.Process(RecordCount,
-				[this](const FReplicationRecord::FRecordInfo& RecordInfo, FReplicationInfo& Info, FNetObjectAttachmentsWriter::ReplicationRecord AttachmentRecord)
+				[this](const FReplicationRecord::FRecordInfo& RecordInfo, FReplicationInfo& Info, const FNetObjectAttachmentsWriter::FReliableReplicationRecord& AttachmentRecord)
 				{ 
 					HandleDeliveredRecord(RecordInfo, Info, AttachmentRecord);
 				}
@@ -1522,7 +1515,7 @@ void FReplicationWriter::ProcessDeliveryNotification(EPacketDeliveryStatus Packe
 		else if (PacketDeliveryStatus == EPacketDeliveryStatus::Lost)
 		{
 			Helper.Process(RecordCount,
-				[this](const FReplicationRecord::FRecordInfo& RecordInfo, FReplicationInfo& Info, FNetObjectAttachmentsWriter::ReplicationRecord AttachmentRecord)
+				[this](const FReplicationRecord::FRecordInfo& RecordInfo, FReplicationInfo& Info, const FNetObjectAttachmentsWriter::FReliableReplicationRecord& AttachmentRecord)
 				{
 					HandleDroppedRecord(RecordInfo, Info, AttachmentRecord);
 				}
@@ -1531,7 +1524,7 @@ void FReplicationWriter::ProcessDeliveryNotification(EPacketDeliveryStatus Packe
 		else if (PacketDeliveryStatus == EPacketDeliveryStatus::Discard)
 		{
 			Helper.Process(RecordCount,
-				[this](const FReplicationRecord::FRecordInfo& RecordInfo, FReplicationInfo& Info, FNetObjectAttachmentsWriter::ReplicationRecord AttachmentRecord)
+				[this](const FReplicationRecord::FRecordInfo& RecordInfo, FReplicationInfo& Info, const FNetObjectAttachmentsWriter::FReliableReplicationRecord& AttachmentRecord)
 				{
 					HandleDiscardedRecord(RecordInfo, Info, AttachmentRecord);
 				}
@@ -1546,14 +1539,14 @@ void FReplicationWriter::ProcessDeliveryNotification(EPacketDeliveryStatus Packe
 
 void FReplicationWriter::CreateObjectRecord(const FNetBitArrayView* ChangeMask, const FReplicationInfo& Info, const FBatchObjectInfo& ObjectInfo, FReplicationWriter::FObjectRecord& OutRecord)
 {
-	OutRecord.AttachmentRecord = ObjectInfo.AttachmentRecord;
+	OutRecord.AttachmentRecord = ObjectInfo.AttachmentRecord.ReliableReplicationRecord;
 
 	FReplicationRecord::FRecordInfo& RecordInfo = OutRecord.Record;
 
 	RecordInfo.Index = ObjectInfo.InternalIndex;
 	RecordInfo.ReplicatedObjectState = ObjectInfo.AttachmentType == ENetObjectAttachmentType::HugeObject ? uint8(EReplicatedObjectState::HugeObject) : (uint8)Info.GetState();
 	RecordInfo.HasChangeMask = ChangeMask ? 1U : 0U;
-	RecordInfo.HasAttachments = (ObjectInfo.AttachmentRecord != 0 ? 1U : 0U);
+	RecordInfo.HasAttachments = (OutRecord.AttachmentRecord.IsValid() ? 1U : 0U);
 	RecordInfo.WroteTearOff = ObjectInfo.bSentTearOff;
 	RecordInfo.WroteDestroySubObject = Info.SubObjectPendingDestroy;
 	
@@ -1578,7 +1571,7 @@ void FReplicationWriter::CreateObjectRecord(const FNetBitArrayView* ChangeMask,
 void FReplicationWriter::CommitObjectRecord(uint32 InternalObjectIndex, const FObjectRecord& ObjectRecord)
 {
 	// Push and link replication record to data already in-flight
-	ReplicationRecord.PushInfoAndAddToList(ReplicatedObjectsRecordInfoLists[InternalObjectIndex], ObjectRecord.Record, ObjectRecord.AttachmentRecord);
+	ReplicationRecord.PushInfoAndAddToList(ReplicatedObjectsRecordInfoLists[InternalObjectIndex], ObjectRecord.Record, ObjectRecord.AttachmentRecord.ToUint64());
 }
 
 void FReplicationWriter::CommitBatchRecord(const FBatchRecord& BatchRecord)
@@ -1850,7 +1843,15 @@ FReplicationWriter::EWriteObjectStatus FReplicationWriter::WriteObjectAndSubObje
 	const FNetRefHandle NetRefHandle = ObjectData.RefHandle;
 
 #if UE_NET_TRACE_ENABLED
-	const FNetRefHandle NetRefHandleForTraceScope = (WriteObjectFlags & EWriteObjectFlag::WriteObjectFlag_HugeObject ? NetRefHandleManager->GetReplicatedObjectDataNoCheck(HugeObjectContext.InternalIndex).RefHandle : NetRefHandle);
+	FNetRefHandle NetRefHandleForTraceScope = NetRefHandle;
+	if (WriteObjectFlags & EWriteObjectFlag::WriteObjectFlag_HugeObject)
+	{
+		const FInternalNetRefIndex HugeObjectInternalIndex = HugeObjectSendQueue.GetRootObjectInternalIndexForTrace();
+		if (HugeObjectInternalIndex != FNetRefHandleManager::InvalidInternalIndex)
+		{
+			NetRefHandleForTraceScope = NetRefHandleManager->GetReplicatedObjectDataNoCheck(HugeObjectInternalIndex).RefHandle;
+		}
+	}
 	UE_NET_TRACE_OBJECT_SCOPE(NetRefHandleForTraceScope, Writer, Context.GetTraceCollector(), ENetTraceVerbosity::Trace);
 #endif
 
@@ -2099,12 +2100,12 @@ FReplicationWriter::EWriteObjectStatus FReplicationWriter::WriteObjectAndSubObje
 				{
 					if (AttachmentWriteStatus == EAttachmentWriteStatus::ReliableWindowFull)
 					{
-						HugeObjectContext.StartStallTime = FPlatformTime::Cycles64();
+						HugeObjectSendQueue.Stats.StartStallTime = FPlatformTime::Cycles64();
 					}
 					else
 					{
 						// Clear stall time now that we were theoretically able to send something.
-						HugeObjectContext.StartStallTime = 0;
+						HugeObjectSendQueue.Stats.StartStallTime = 0;
 					}
 				}
 
@@ -2334,6 +2335,14 @@ FReplicationWriter::EWriteObjectStatus FReplicationWriter::WriteObjectInBatch(FN
 int FReplicationWriter::PrepareAndSendHugeObjectPayload(FNetSerializationContext& Context, uint32 InternalIndex)
 {
 	IRIS_PROFILER_SCOPE(FReplicationWriter_PrepareAndSendHugeObjectPayload);
+
+	// Sanity check
+	if (HugeObjectSendQueue.IsFull() || HugeObjectSendQueue.IsObjectInQueue(InternalIndex, false))
+	{
+		ensureMsgf(false, TEXT("HugeObjectSendQueue should not be full or already transmitting. ( InternalIndex: %u )"), InternalIndex);
+		return 0;
+	}
+
 	typedef uint32 HugeObjectStorageType;
 	const uint32 BitsPerStorageWord = sizeof(HugeObjectStorageType) * 8;
 
@@ -2347,15 +2356,15 @@ int FReplicationWriter::PrepareAndSendHugeObjectPayload(FNetSerializationContext
 	FNetSerializationContext HugeObjectSerializationContext = Context.MakeSubContext(&HugeObjectWriter);
 
 #if UE_NET_TRACE_ENABLED
-	if (!HugeObjectContext.TraceCollector)
+	if (!HugeObjectSendQueue.TraceCollector)
 	{
-		HugeObjectContext.TraceCollector = UE_NET_TRACE_CREATE_COLLECTOR(ENetTraceVerbosity::Trace);
+		HugeObjectSendQueue.TraceCollector = UE_NET_TRACE_CREATE_COLLECTOR(ENetTraceVerbosity::Trace);
 	}
 	else
 	{
-		HugeObjectContext.TraceCollector->Reset();
+		HugeObjectSendQueue.TraceCollector->Reset();
 	}
-	HugeObjectSerializationContext.SetTraceCollector(HugeObjectContext.TraceCollector);
+	HugeObjectSerializationContext.SetTraceCollector(HugeObjectSendQueue.TraceCollector);
 #endif
 
 	// Huge object header needed for the receiving side to be able to process this correctly.
@@ -2364,6 +2373,8 @@ int FReplicationWriter::PrepareAndSendHugeObjectPayload(FNetSerializationContext
 	FNetObjectBlob::SerializeHeader(HugeObjectSerializationContext, HugeObjectHeader);
 	const uint32 PastHeaderPos = HugeObjectWriter.GetPosBits();
 
+	FHugeObjectContext HugeObjectContext;
+
 	FBatchInfo BatchInfo;
 	BatchInfo.Type = EBatchInfoType::Internal;
 	BatchInfo.ParentInternalIndex = InternalIndex;
@@ -2406,9 +2417,13 @@ int FReplicationWriter::PrepareAndSendHugeObjectPayload(FNetSerializationContext
 		}
 	}
 
-	HugeObjectContext.InternalIndex = InternalIndex;
-	HugeObjectContext.SendStatus = EHugeObjectSendStatus::Sending;
-	HugeObjectContext.StartSendingTime = FPlatformTime::Cycles64();
+	if (HugeObjectSendQueue.IsEmpty())
+	{
+		HugeObjectSendQueue.Stats.StartSendingTime = FPlatformTime::Cycles64();
+	}
+
+	HugeObjectContext.RootObjectInternalIndex = InternalIndex;
+
 	// Store batch record for later processing once the whole state is acked.
 	HandleObjectBatchSuccess(BatchInfo, HugeObjectContext.BatchRecord);
 	// We want to track the number of Batches
@@ -2424,14 +2439,14 @@ int FReplicationWriter::PrepareAndSendHugeObjectPayload(FNetSerializationContext
 
 	// Create a NetObjectBlob from the temporary buffer and split it into multiple smaller pieces.
 	const uint32 PayLoadBitCount = HugeObjectWriter.GetPosBits();
-	const uint32 StorageWordsWritten = (PayLoadBitCount + (BitsPerStorageWord - 1)) / BitsPerStorageWord;
+	const uint32 StorageWordsWritten = (PayLoadBitCount + (BitsPerStorageWord - 1))/BitsPerStorageWord;
 
 	check(StorageWordsWritten <= (uint32)HugeObjectPayload.Num());
 
 	TArrayView<HugeObjectStorageType> PayloadView(HugeObjectPayload.GetData(), StorageWordsWritten);
 	TRefCountPtr<FNetObjectBlob> NetObjectBlob = NetObjectBlobHandler->CreateNetObjectBlob(PayloadView, PayLoadBitCount);
 	TArray<TRefCountPtr<FNetBlob>> PartialNetBlobs;
-	const bool bSplitSuccess = PartialNetObjectAttachmentHandler->SplitRawDataNetBlob(TRefCountPtr<FRawDataNetBlob>(NetObjectBlob.GetReference()), PartialNetBlobs, HugeObjectContext.DebugName);
+	const bool bSplitSuccess = PartialNetObjectAttachmentHandler->SplitRawDataNetBlob(TRefCountPtr<FRawDataNetBlob>(NetObjectBlob.GetReference()), PartialNetBlobs, HugeObjectSendQueue.DebugName);
 	if (!bSplitSuccess)
 	{
 		UE_LOG(LogIris, Error, TEXT("Unable to split huge object ( InternalIndex: %u ) payload. Connection %u will be disconnected."), InternalIndex, Context.GetLocalConnectionId());
@@ -2442,6 +2457,23 @@ int FReplicationWriter::PrepareAndSendHugeObjectPayload(FNetSerializationContext
 	// Enqueue attachments
 	const bool bEnqueueSuccess = Attachments.Enqueue(ENetObjectAttachmentType::HugeObject, ObjectIndexForOOBAttachment, MakeArrayView(PartialNetBlobs.GetData(), PartialNetBlobs.Num()));
 	check(bEnqueueSuccess);
+	if (!bEnqueueSuccess)
+	{
+		UE_LOG(LogIris, Error, TEXT("Unable to enqueue huge object attachments ( InternalIndex: %u ). Connection %u will be disconnected."), InternalIndex, Context.GetLocalConnectionId());
+		Context.SetError(GNetError_InternalError);
+		return -1;
+	}
+
+	// Add huge object to queue
+	HugeObjectContext.Blobs = MoveTemp(PartialNetBlobs);
+	const bool bHugeObjectWasEnqueued = HugeObjectSendQueue.EnqueueHugeObject(HugeObjectContext);
+	check(bHugeObjectWasEnqueued);
+	if (!bHugeObjectWasEnqueued)
+	{
+		UE_LOG(LogIris, Error, TEXT("Unable to enqueue huge object ( InternalIndex: %u ). Connection %u will be disconnected."), InternalIndex, Context.GetLocalConnectionId());
+		Context.SetError(GNetError_InternalError);
+		return -1;
+	}
 
 	// Write huge object attachment(s)
 	{
@@ -2479,7 +2511,7 @@ int FReplicationWriter::PrepareAndSendHugeObjectPayload(FNetSerializationContext
 		WriteContext.bHasHugeObjectToSend = bHasHugeObjectToSend;
 		if (!bHasHugeObjectToSend)
 		{
-			HugeObjectContext.EndSendingTime = FPlatformTime::Cycles64();
+			HugeObjectSendQueue.Stats.EndSendingTime = FPlatformTime::Cycles64();
 		}
 
 		return 1;
@@ -2516,7 +2548,7 @@ int FReplicationWriter::WriteObjectBatch(FNetSerializationContext& Context, uint
 		BatchInfo.Type = (InternalIndex == ObjectIndexForOOBAttachment ? (WriteObjectFlags & WriteObjectFlag_HugeObject ? EBatchInfoType::HugeObject : EBatchInfoType::OOBAttachment) : EBatchInfoType::Object);
 		BatchInfo.ParentInternalIndex = InternalIndex;
 
-		// Write an objects and its SubObjects, if object has dependent objects pending creation we currently write them as well but they will be as an individual batch
+		// Write an object and its subobjects. If object has dependent objects pending creation we currently write them as well as an individual batch.
 		const EWriteObjectStatus WriteObjectStatus = WriteObjectInBatch(Context, InternalIndex, WriteObjectFlags, BatchInfo);
 
 		if (IsWriteObjectSuccess(WriteObjectStatus))
@@ -2663,7 +2695,7 @@ uint32 FReplicationWriter::WriteOOBAttachments(FNetSerializationContext& Context
 		WriteContext.bHasHugeObjectToSend = bHasHugeObjectToSend;
 		if (!bHasHugeObjectToSend)
 		{
-			HugeObjectContext.EndSendingTime = FPlatformTime::Cycles64();
+			HugeObjectSendQueue.Stats.EndSendingTime = FPlatformTime::Cycles64();
 		}
 
 		WrittenObjectCount += Result;
@@ -2741,7 +2773,7 @@ uint32 FReplicationWriter::WriteObjects(FNetSerializationContext& Context)
 
 				const uint32 InternalIndex = WriteContext.DependentObjectsPendingSend.Pop();
 				checkSlow(InternalIndex != ObjectIndexForOOBAttachment);
-				ensureAlwaysMsgf(GetReplicationInfo(InternalIndex).GetState() != EReplicatedObjectState::Invalid, TEXT("DependentObject with internalIndex %u is not in scope"), InternalIndex);
+				ensureAlwaysMsgf(GetReplicationInfo(InternalIndex).GetState() != EReplicatedObjectState::Invalid, TEXT("DependentObject with InternalIndex %u is not in scope"), InternalIndex);
 				if (!SendObjectFunction(InternalIndex))
 				{
 					// If we fail, we put the object back on the pending send stack and try again in the next packet of the batch
@@ -2823,7 +2855,7 @@ int FReplicationWriter::HandleObjectBatchSuccess(const FBatchInfo& BatchInfo, FR
 		}
 
 		// We're now committing to what we wrote so inform the attachments writer.
-		if (BatchObjectInfo.AttachmentRecord)
+		if (BatchObjectInfo.AttachmentRecord.IsValid())
 		{
 			Attachments.CommitReplicationRecord(BatchObjectInfo.AttachmentType, BatchObjectInfo.InternalIndex, BatchObjectInfo.AttachmentRecord);
 		}
@@ -2846,7 +2878,7 @@ int FReplicationWriter::HandleObjectBatchSuccess(const FBatchInfo& BatchInfo, FR
 				++DeltaCompressedObjectCount;
 			}
 		}
-		else if (BatchObjectInfo.AttachmentRecord != 0 || BatchObjectInfo.bSentTearOff || BatchObjectInfo.bSentDestroySubObject)
+		else if (BatchObjectInfo.AttachmentRecord.IsValid() || BatchObjectInfo.bSentTearOff || BatchObjectInfo.bSentDestroySubObject)
 		{
 			FObjectRecord& ObjectRecord = OutRecord.ObjectReplicationRecords.AddDefaulted_GetRef();
 			CreateObjectRecord(nullptr, Info, BatchObjectInfo, ObjectRecord);
@@ -2948,7 +2980,7 @@ FReplicationWriter::EWriteObjectRetryMode FReplicationWriter::HandleObjectBatchF
 
 	// If there are more bits left than the split threshold we treat it as a huge object and proceed with splitting.
 	// We expect at least one part of the payload to be sendable if there are more bits left than the split threshold.
-	if (HugeObjectContext.SendStatus == EHugeObjectSendStatus::Idle && PartialNetObjectAttachmentHandler != nullptr)
+	if (CanQueueHugeObject() && PartialNetObjectAttachmentHandler != nullptr)
 	{
 		const uint32 SplitThreshold = PartialNetObjectAttachmentHandler->GetConfig()->GetBitCountSplitThreshold();
 		if (BitsLeft > SplitThreshold)
@@ -3038,17 +3070,17 @@ void FReplicationWriter::EndWrite()
 		// Update stats
 		{
 			FNetSendStats& Stats = WriteContext.Stats;
-			if (HugeObjectContext.SendStatus == EHugeObjectSendStatus::Sending)
+			if (!HugeObjectSendQueue.IsEmpty())
 			{
-				Stats.SetNumberOfActiveHugeObjects(1U);
+				Stats.SetNumberOfActiveHugeObjects(HugeObjectSendQueue.NumRootObjectsInTransit());
 
-				if (HugeObjectContext.EndSendingTime != 0)
+				if (HugeObjectSendQueue.Stats.EndSendingTime != 0)
 				{
-					Stats.AddHugeObjectWaitingTime(FPlatformTime::ToSeconds64(FPlatformTime::Cycles64() - HugeObjectContext.EndSendingTime));
+					Stats.AddHugeObjectWaitingTime(FPlatformTime::ToSeconds64(FPlatformTime::Cycles64() - HugeObjectSendQueue.Stats.EndSendingTime));
 				}
-				if (HugeObjectContext.StartStallTime != 0)
+				if (HugeObjectSendQueue.Stats.StartStallTime != 0)
 				{
-					Stats.AddHugeObjectStallTime(FPlatformTime::ToSeconds64(FPlatformTime::Cycles64() - HugeObjectContext.StartStallTime));
+					Stats.AddHugeObjectStallTime(FPlatformTime::ToSeconds64(FPlatformTime::Cycles64() - HugeObjectSendQueue.Stats.StartStallTime));
 				}
 			}
 
@@ -3173,7 +3205,7 @@ UDataStream::EWriteResult FReplicationWriter::Write(FNetSerializationContext& Co
 #endif
 
 #if UE_NET_TRACE_ENABLED
-		if (FNetTraceCollector* Collector = HugeObjectContext.TraceCollector)
+		if (FNetTraceCollector* Collector = HugeObjectSendQueue.TraceCollector)
 		{
 			FNetTrace::FoldTraceCollector(Context.GetTraceCollector(), Collector, GetBitStreamPositionForNetTrace(Writer));
 			Collector->Reset();
@@ -3286,43 +3318,47 @@ void FReplicationWriter::SetNetExports(FNetExports& InNetExports)
 	NetExports = &InNetExports;
 }
 
+bool FReplicationWriter::IsActiveHugeObject(uint32 InternalIndex) const
+{
+	constexpr bool bIncludeSubObjects = false;
+	return HugeObjectSendQueue.IsObjectInQueue(InternalIndex, bIncludeSubObjects);
+}
+
 bool FReplicationWriter::IsObjectPartOfActiveHugeObject(uint32 InternalIndex, const FReplicationInfo& Info) const
 {
-	if (HugeObjectContext.SendStatus == EHugeObjectSendStatus::Idle)
+	return HugeObjectSendQueue.IsObjectInQueue(InternalIndex, Info.IsSubObject);
+}
+
+bool FReplicationWriter::CanQueueHugeObject() const
+{
+	if (HugeObjectSendQueue.IsFull())
 	{
 		return false;
 	}
 
-	if (InternalIndex == HugeObjectContext.InternalIndex)
+	// Check whether the reliable queue is full in which case there's no point in queueing additional huge objects.
+	if (!Attachments.CanSendMoreReliableAttachments(ENetObjectAttachmentType::HugeObject, ObjectIndexForOOBAttachment))
 	{
-		return true;
+		return false;
 	}
 
-	if (Info.IsSubObject)
+	return true;
+}
+
+void FReplicationWriter::FreeHugeObjectSendQueue()
+{
+	HugeObjectSendQueue.FreeContexts([this](const FHugeObjectContext& HugeObjectContext)
 	{
-		// At this state any information regarding the parent has been cleared. Check if the index is part of the payload.
-		const FBatchRecord& BatchRecord = HugeObjectContext.BatchRecord;
-		for (const FObjectRecord& ObjectRecord : MakeArrayView(BatchRecord.ObjectReplicationRecords.GetData(), BatchRecord.ObjectReplicationRecords.Num()))
+		for (const FObjectRecord& ObjectRecord : HugeObjectContext.BatchRecord.ObjectReplicationRecords)
 		{
-			if (InternalIndex == ObjectRecord.Record.Index)
+			FReplicationInfo& ReplicationInfo = this->GetReplicationInfo(ObjectRecord.Record.Index);
+			const uint32 ChangeMaskBitCount = ReplicationInfo.ChangeMaskBitCount;
+			if (ObjectRecord.Record.HasChangeMask)
 			{
-				return true;
+				FChangeMaskStorageOrPointer::Free(ObjectRecord.Record.ChangeMaskOrPtr, ChangeMaskBitCount, s_DefaultChangeMaskAllocator);
 			}
 		}
-	}
-
-	return false;
-}
-
-void FReplicationWriter::ClearHugeObjectContext(FHugeObjectContext& Context) const
-{
-	Context.InternalIndex = 0;
-	Context.SendStatus = EHugeObjectSendStatus::Idle;
-	Context.BatchRecord = FBatchRecord();
-	Context.BatchExports.Reset();
-	Context.StartSendingTime = 0;
-	Context.EndSendingTime = 0;
-	Context.StartStallTime = 0;
+	});
 }
 
 void FReplicationWriter::CollectAndAppendExports(FNetSerializationContext& Context, uint8* RESTRICT InternalBuffer, const FReplicationProtocol* Protocol) const
@@ -3357,7 +3393,7 @@ void FReplicationWriter::DiscardAllRecords()
 		if (const uint32 RecordInfoCount = ReplicationRecord.PopRecord())
 		{
 			Helper.Process(RecordInfoCount,
-				[this](const FReplicationRecord::FRecordInfo& RecordInfo, FReplicationInfo& Info, FNetObjectAttachmentsWriter::ReplicationRecord AttachmentRecord)
+				[this](const FReplicationRecord::FRecordInfo& RecordInfo, FReplicationInfo& Info, const FNetObjectAttachmentsWriter::FReliableReplicationRecord& AttachmentRecord)
 				{
 					HandleDiscardedRecord(RecordInfo, Info, AttachmentRecord);
 				}
@@ -3411,4 +3447,161 @@ void FReplicationWriter::MarkObjectDirty(FInternalNetRefIndex InternalIndex, con
 }
 
 
+FReplicationWriter::FHugeObjectContext::FHugeObjectContext() = default;
+
+FReplicationWriter::FHugeObjectContext::~FHugeObjectContext() = default;
+
+// HugeObjectSendQueue implementation
+FReplicationWriter::FHugeObjectSendQueue::FHugeObjectSendQueue()
+: DebugName(CreatePersistentNetDebugName(TEXT("HugeObjectState"), UE_ARRAY_COUNT(TEXT("HugeObjectState"))))
+{
+#if UE_NET_TRACE_ENABLED
+	DebugName->DebugNameId = FNetTrace::TraceName(DebugName->Name);
+#endif
+}
+
+FReplicationWriter::FHugeObjectSendQueue::~FHugeObjectSendQueue()
+{
+	UE_NET_TRACE_DESTROY_COLLECTOR(TraceCollector);
+	TraceCollector = nullptr;
+}
+
+// TODO: If reliable queue is full should we keep on filling up?
+bool FReplicationWriter::FHugeObjectSendQueue::IsFull() const
+{
+	const int32 QueueSize = FPlatformMath::Max<int32>(GReplicationWriterMaxHugeObjectsInTransit, 1);
+	return RootObjectsInTransit.Num() > QueueSize;
+}
+
+bool FReplicationWriter::FHugeObjectSendQueue::IsEmpty() const
+{
+	return RootObjectsInTransit.IsEmpty();
+}
+
+uint32 FReplicationWriter::FHugeObjectSendQueue::NumRootObjectsInTransit() const
+{
+	return static_cast<uint32>(RootObjectsInTransit.Num());
+}
+
+bool FReplicationWriter::FHugeObjectSendQueue::EnqueueHugeObject(const FHugeObjectContext& Context)
+{
+	if (IsFull())
+	{
+		return false;
+	}
+
+	if (RootObjectsInTransit.Find(Context.RootObjectInternalIndex))
+	{
+		ensureMsgf(false, TEXT("An object that is already in the huge object queue should not try replicating again ( InternalIndex: %u )"), Context.RootObjectInternalIndex);
+		return false;
+	}
+
+	RootObjectsInTransit.Add(Context.RootObjectInternalIndex);
+	// Note: Lists don't have methods to perform moving of an element.
+	SendContexts.AddTail(Context);
+	return true;
+}
+
+// Returns true if the object is a huge object root object or part of any huge object's payload. The latter is an expensive operation.
+bool FReplicationWriter::FHugeObjectSendQueue::IsObjectInQueue(FInternalNetRefIndex ObjectIndex, bool bIncludeSubObjects) const
+{
+	if (IsEmpty())
+	{
+		return false;
+	}
+
+	if (RootObjectsInTransit.Find(ObjectIndex))
+	{
+		return true;
+	}
+
+	if (!bIncludeSubObjects)
+	{
+		return false;
+	}
+
+	for (const FHugeObjectContext& Context : SendContexts)
+	{
+		const FBatchRecord& BatchRecord = Context.BatchRecord;
+		for (const FObjectRecord& ObjectRecord : MakeArrayView(BatchRecord.ObjectReplicationRecords.GetData(), BatchRecord.ObjectReplicationRecords.Num()))
+		{
+			if (ObjectIndex == ObjectRecord.Record.Index)
+			{
+				return true;
+			}
+		}
+	}
+
+	return false;
+}
+
+ FInternalNetRefIndex FReplicationWriter::FHugeObjectSendQueue::GetRootObjectInternalIndexForTrace() const
+ {
+	 const TDoubleLinkedList<FHugeObjectContext>::TDoubleLinkedListNode* TailNode = SendContexts.GetTail();
+	 if (TailNode)
+	 {
+		 return TailNode->GetValue().RootObjectInternalIndex;
+	 }
+
+	return FNetRefHandleManager::InvalidInternalIndex;
+ }
+
+void FReplicationWriter::FHugeObjectSendQueue::AckObjects(TFunctionRef<void (const FHugeObjectContext& Context)> AckHugeObject)
+{
+	for (TDoubleLinkedList<FHugeObjectContext>::TDoubleLinkedListNode* Node = SendContexts.GetHead(), *NextNode = nullptr; Node != nullptr; Node = NextNode)
+	{
+		NextNode = Node->GetNextNode();
+
+		FHugeObjectContext& Context = Node->GetValue();
+
+		// Iterate over the blobs backwards to break out of the loop as quickly as possible.
+		bool bObjectIsAcked = true;
+		for (TRefCountPtr<FNetBlob>& Blob : ReverseIterate(Context.Blobs))
+		{
+			const uint32 RefCount = Blob.GetRefCount();		
+			if (RefCount > 1)
+			{
+				bObjectIsAcked = false;
+				break;
+			}
+			else if (RefCount == 1)
+			{
+				// We no longer need to keep this blob around as we're the only thing referencing it.
+				Blob.SafeRelease();
+			}
+		}
+
+		if (!bObjectIsAcked)
+		{
+			continue;
+		}
+
+		AckHugeObject(Context);
+
+		// Remove from fast lookup set.
+		RootObjectsInTransit.Remove(Context.RootObjectInternalIndex);
+
+		// Remove from queue.
+		SendContexts.RemoveNode(Node);
+
+		ensure(SendContexts.IsEmpty() == RootObjectsInTransit.IsEmpty());
+	}
+
+	if (RootObjectsInTransit.IsEmpty())
+	{
+		Stats = FStats();
+	}
+}
+
+void FReplicationWriter::FHugeObjectSendQueue::FreeContexts(TFunctionRef<void (const FHugeObjectContext& Context)> FreeHugeObject)
+{
+	for (const FHugeObjectContext& Context : SendContexts)
+	{
+		FreeHugeObject(Context);
+	}
+
+	SendContexts.Empty();
+	RootObjectsInTransit.Empty();
+}
+
 }
diff --git a/Engine/Source/Runtime/Experimental/Iris/Core/Private/Iris/ReplicationSystem/ReplicationWriter.h b/Engine/Source/Runtime/Experimental/Iris/Core/Private/Iris/ReplicationSystem/ReplicationWriter.h
--- a/Engine/Source/Runtime/Experimental/Iris/Core/Private/Iris/ReplicationSystem/ReplicationWriter.h
+++ b/Engine/Source/Runtime/Experimental/Iris/Core/Private/Iris/ReplicationSystem/ReplicationWriter.h
@@ -15,6 +15,8 @@
 #include "Iris/ReplicationSystem/ObjectReferenceCache.h"
 #include "Iris/Stats/NetStats.h"
 #include "Containers/Array.h"
+#include "Containers/List.h"
+#include "Containers/Set.h"
 #include "Misc/EnumClassFlags.h"
 
 // Forward declaration
@@ -51,7 +53,7 @@ class FReplicationWriter
 {
 public:
 	// Scheduling constants
-	static constexpr float CreatePriority = 0.f;
+	static constexpr float CreatePriority = 1.f;
 	static constexpr float TearOffPriority = 1.f;
 	static constexpr float LostStatePriorityBump = 1.f;
 	static constexpr float SchedulingThresholdPriority = 1.f;
@@ -242,7 +244,7 @@ private:
 	{
 		FNetRefHandle Handle;
 		uint32 InternalIndex;
-		FNetObjectAttachmentsWriter::ReplicationRecord AttachmentRecord;
+		FNetObjectAttachmentsWriter::FCommitRecord AttachmentRecord;
 		ENetObjectAttachmentType AttachmentType;
 		bool bHasUnsentAttachments;
 		uint32 NewBaselineIndex : 2;
@@ -275,7 +277,7 @@ private:
 	struct FObjectRecord
 	{
 		FReplicationRecord::FRecordInfo Record;
-		FNetObjectAttachmentsWriter::ReplicationRecord AttachmentRecord;
+		FNetObjectAttachmentsWriter::FReliableReplicationRecord AttachmentRecord;
 	};
 
 	struct FBatchRecord
@@ -302,18 +304,59 @@ private:
 		FHugeObjectContext();
 		~FHugeObjectContext();
 
-		EHugeObjectSendStatus SendStatus;
-		uint32 InternalIndex;
+		FInternalNetRefIndex RootObjectInternalIndex = 0;
 		FBatchRecord BatchRecord;
 		FNetExportContext::FBatchExports BatchExports;
-		FNetTraceCollector* TraceCollector;
-		const FNetDebugName* DebugName;
-		// Cycle counter for when the huge object context went from idle to sending.
-		uint64 StartSendingTime;
-		// Cycle counter for when the last part of huge object was sent.
-		uint64 EndSendingTime;
-		// Cycle counter for when it was detected that no more parts of the huge object could be sent until some of the first parts have been acked.
-		uint64 StartStallTime;
+
+		// The entire payload. When refcount reaches one for all blobs the object has been fully acked.
+		TArray<TRefCountPtr<FNetBlob>> Blobs;
+	};
+
+	class FHugeObjectSendQueue
+	{
+	public:
+		FHugeObjectSendQueue();
+		~FHugeObjectSendQueue();
+
+		// If the queue is full we can't start another send.
+		bool IsFull() const;
+		bool IsEmpty() const;
+		uint32 NumRootObjectsInTransit() const;
+
+		// Enqueue huge object info and return true if it can be sent.
+		bool EnqueueHugeObject(const FHugeObjectContext& Context);
+
+		// Returns true if the object is a huge object root object or part of any huge object's payload. The latter is an expensive operation.
+		bool IsObjectInQueue(FInternalNetRefIndex ObjectIndex, bool bIncludeSubObjects) const;
+
+		// Best effort implementation of getting a valid index for trace.
+		FInternalNetRefIndex GetRootObjectInternalIndexForTrace() const;
+
+		// Call AckHugeObject on all objects determined to have been fully processed.
+		void AckObjects(TFunctionRef<void (const FHugeObjectContext& Context)> AckHugeObject);
+
+		void FreeContexts(TFunctionRef<void (const FHugeObjectContext& Context)> FreeHugeObject);
+
+	public:
+		// Public members
+		struct FStats
+		{
+			// Cycle counter for when the huge object context went from idle to sending.
+			uint64 StartSendingTime = 0;
+			// Cycle counter for when the last part of huge object was sent.
+			uint64 EndSendingTime = 0;
+			// Cycle counter for when it was detected that no more parts of the huge object could be sent until some of the first parts have been acked.
+			uint64 StartStallTime = 0;
+		};
+
+		FStats Stats;
+
+		FNetTraceCollector* TraceCollector = nullptr;
+		const FNetDebugName* DebugName = nullptr;
+
+	private:
+		TSet<FInternalNetRefIndex> RootObjectsInTransit;
+		TDoubleLinkedList<FHugeObjectContext> SendContexts;
 	};
 
 	enum EWriteObjectFlag : unsigned
@@ -428,14 +471,14 @@ private:
 	EWriteObjectRetryMode HandleObjectBatchFailure(EWriteObjectStatus WriteObjectStatus, const FBatchInfo& BatchInfo, const FBitStreamInfo& BatchBitStreamInfo) const;
 
 	// Update logic for dropped RecordInfo
-	void HandleDroppedRecord(const FReplicationRecord::FRecordInfo& RecordInfo, FReplicationInfo& Info, FNetObjectAttachmentsWriter::ReplicationRecord AttachmentRecord);
-	template<EReplicatedObjectState LostState> void HandleDroppedRecord(EReplicatedObjectState CurrentState, const FReplicationRecord::FRecordInfo& RecordInfo, FReplicationInfo& Info, FNetObjectAttachmentsWriter::ReplicationRecord AttachmentRecord);
+	void HandleDroppedRecord(const FReplicationRecord::FRecordInfo& RecordInfo, FReplicationInfo& Info, const FNetObjectAttachmentsWriter::FReliableReplicationRecord& AttachmentRecord);
+	template<EReplicatedObjectState LostState> void HandleDroppedRecord(EReplicatedObjectState CurrentState, const FReplicationRecord::FRecordInfo& RecordInfo, FReplicationInfo& Info, const FNetObjectAttachmentsWriter::FReliableReplicationRecord& AttachmentRecord);
 
 	// Update logic for delivered RecordInfo
-	void HandleDeliveredRecord(const FReplicationRecord::FRecordInfo& RecordInfo, FReplicationInfo& Info, FNetObjectAttachmentsWriter::ReplicationRecord AttachmentRecord);
+	void HandleDeliveredRecord(const FReplicationRecord::FRecordInfo& RecordInfo, FReplicationInfo& Info, const FNetObjectAttachmentsWriter::FReliableReplicationRecord& AttachmentRecord);
 
 	// Update logic for discarded RecordInfo, for preventing memory leaks on disconnect and shutdown.
-	void HandleDiscardedRecord(const FReplicationRecord::FRecordInfo& RecordInfo, FReplicationInfo& Info, FNetObjectAttachmentsWriter::ReplicationRecord AttachmentRecord);
+	void HandleDiscardedRecord(const FReplicationRecord::FRecordInfo& RecordInfo, FReplicationInfo& Info, const FNetObjectAttachmentsWriter::FReliableReplicationRecord& AttachmentRecord);
 
 	// Setup replication info to be able to send attachments to objects not in scope
 	void SetupReplicationInfoForAttachmentsToObjectsNotInScope();
@@ -461,10 +504,12 @@ private:
 
 	inline bool IsInitialState(const EReplicatedObjectState State) const { return State == EReplicatedObjectState::PendingCreate || (bHighPrioCreate && State == EReplicatedObjectState::WaitOnCreateConfirmation); }
 
-	bool IsActiveHugeObject(uint32 InternalIndex) const { return HugeObjectContext.InternalIndex == InternalIndex && HugeObjectContext.SendStatus != EHugeObjectSendStatus::Idle; }
+	bool IsActiveHugeObject(uint32 InternalIndex) const;
 	bool IsObjectPartOfActiveHugeObject(uint32 InternalIndex, const FReplicationInfo& Info) const;
 
-	void ClearHugeObjectContext(FHugeObjectContext& Context) const;
+	bool CanQueueHugeObject() const;
+
+	void FreeHugeObjectSendQueue();
 
 	bool HasDataToSend(const FWriteContext& Context) const;
 
@@ -521,7 +566,7 @@ private:
 
 	FWriteContext WriteContext;
 	FBitStreamInfo WriteBitStreamInfo;
-	FHugeObjectContext HugeObjectContext;
+	FHugeObjectSendQueue HugeObjectSendQueue;
 
 	// Is replication enabled?
 	bool bReplicationEnabled = false;
@@ -535,7 +580,7 @@ inline FReplicationWriter::FReplicationInfo::FReplicationInfo()
 {
 }
 
-template<FReplicationWriter::EReplicatedObjectState LostState> void FReplicationWriter::HandleDroppedRecord(FReplicationWriter::EReplicatedObjectState CurrentState, const FReplicationRecord::FRecordInfo& RecordInfo, FReplicationWriter::FReplicationInfo& Info, FNetObjectAttachmentsWriter::ReplicationRecord AttachmentRecord)
+template<FReplicationWriter::EReplicatedObjectState LostState> void FReplicationWriter::HandleDroppedRecord(FReplicationWriter::EReplicatedObjectState CurrentState, const FReplicationRecord::FRecordInfo& RecordInfo, FReplicationWriter::FReplicationInfo& Info, const FNetObjectAttachmentsWriter::FReliableReplicationRecord& AttachmentRecord)
 {
 	//static_assert(false, "Expected specialization to exist.");
 }
diff --git a/Engine/Source/Runtime/Experimental/Iris/Core/Private/Iris/Serialization/NetErrorContext.cpp b/Engine/Source/Runtime/Experimental/Iris/Core/Private/Iris/Serialization/NetErrorContext.cpp
--- a/Engine/Source/Runtime/Experimental/Iris/Core/Private/Iris/Serialization/NetErrorContext.cpp
+++ b/Engine/Source/Runtime/Experimental/Iris/Core/Private/Iris/Serialization/NetErrorContext.cpp
@@ -11,6 +11,7 @@ const FName GNetError_ArraySizeTooLarge("Array size is too large");
 const FName GNetError_InvalidNetHandle("Invalid NetHandle");
 const FName GNetError_BrokenNetHandle("Broken NetHandle");
 const FName GNetError_InvalidValue("Invalid value");
+const FName GNetError_InternalError("Internal error");
 
 void FNetErrorContext::SetError(const FName InError)
 {
diff --git a/Engine/Source/Runtime/Experimental/Iris/Core/Public/Iris/ReplicationSystem/NetBlob/ReliableNetBlobQueue.h b/Engine/Source/Runtime/Experimental/Iris/Core/Public/Iris/ReplicationSystem/NetBlob/ReliableNetBlobQueue.h
--- a/Engine/Source/Runtime/Experimental/Iris/Core/Public/Iris/ReplicationSystem/NetBlob/ReliableNetBlobQueue.h
+++ b/Engine/Source/Runtime/Experimental/Iris/Core/Public/Iris/ReplicationSystem/NetBlob/ReliableNetBlobQueue.h
@@ -24,10 +24,37 @@ namespace UE::Net::Private
 class FReliableNetBlobQueue
 {
 public:
-	typedef uint32 ReplicationRecord;
+	struct FReplicationRecord
+	{
+		FReplicationRecord() = default;
+		explicit FReplicationRecord(uint64 Value)
+		{
+			for (unsigned Index : {3U, 2U, 1U, 0U})
+			{
+				Counts[Index] = Value & 255U;
+				Sequences[Index] = (Value >> 8U) & 255U;
+				Value >>= 16U;
+			}
+		}
+
+		uint64 ToUint64() const
+		{
+			uint64 Value = 0;
+			for (unsigned Index : {0U, 1U, 2U, 3U})
+			{
+				Value = (Value << 16U) | (uint64(Sequences[Index]) << 8U) | Counts[Index];
+			}
+
+			return Value;
+		}
+
+		bool IsValid() const { return Counts[0] | Counts[1] | Counts[2] | Counts[3]; }
+
+		uint8 Sequences[4] = {};
+		uint8 Counts[4] = {};
+	};
 
 	/** This represents a ReplicationRecord where nothing was serialized. */
-	static constexpr uint32 InvalidReplicationRecord = 0U;
 
 	/** How many blobs can be sent before an ACK/NAK is required to continue sending. */
 	static constexpr uint32 MaxUnackedBlobCount = 256U;
@@ -44,6 +71,9 @@ public:
 	/** Returns whether all blobs have been sent and acknowledged as received. */
 	bool IsAllSentAndAcked() const { return FirstSeq == LastSeq && GetUnsentBlobCount() == 0; }
 
+	/** Returns whether the send window is full or not. */
+	bool IsSendWindowFull() const;
+
 	/** Returns true if it's safe to destroy this queue. */
 	IRISCORE_API bool IsSafeToDestroy() const;
 
@@ -59,13 +89,13 @@ public:
 	/**
 	 * Serializes as many blobs as possible using their respective SerializeWithObject() method. It is assumed the NetRefHandle will be 
 	 * reconstructed somehow on the receiving end and passed to DeserializeWithObject().
-	 * This provides an opportunity fot FNetObjectAttachments, such as FNetRPCs, to avoid serializing the same NetRefHandle redundantly.
+	 * This provides an opportunity for FNetObjectAttachments, such as FNetRPCs, to avoid serializing the same NetRefHandle redundantly.
 	 * @param Context A FNetSerializationContext.
 	 * @param RefHandle The handle for the blobs' target object.
 	 * @param OutRecord The record to pass to CommitReplicationRecord() if a packet containing the serialized data was sent.
 	 * @return The number of blobs that were serialized.
 	 */
-	IRISCORE_API uint32 SerializeWithObject(FNetSerializationContext& Context, FNetRefHandle RefHandle, ReplicationRecord& OutRecord);
+	IRISCORE_API uint32 SerializeWithObject(FNetSerializationContext& Context, FNetRefHandle RefHandle, FReplicationRecord& OutRecord);
 
 	/**
 	 * Deserializes blobs with object using their respective DeserializeWithObject() method.
@@ -80,7 +110,7 @@ public:
 	 * @param OutRecord The record to pass to CommitReplicationRecord() if a packet containing the serialized data was sent.
 	 * @return The number of blobs that were serialized.
 	 */
-	IRISCORE_API uint32 Serialize(FNetSerializationContext& Context, ReplicationRecord& OutRecord);
+	IRISCORE_API uint32 Serialize(FNetSerializationContext& Context, FReplicationRecord& OutRecord);
 
 	/**
 	 * Deserializes blobs with object using their respective Deserialize() method.
@@ -93,7 +123,7 @@ public:
 	 * Call after a packet containing serialized data was sent.
 	 * @see SerializeWithObject, Serialize
 	 */
-	IRISCORE_API void CommitReplicationRecord(ReplicationRecord Record);
+	IRISCORE_API void CommitReplicationRecord(const FReplicationRecord& Record);
 
 	/**
 	 * For each packet for which CommitReplicationRecord() was called ProcessPacketDeliveryStatus() needs
@@ -102,16 +132,15 @@ public:
 	 * @param Record The record that was obtained via a Serialize/SerializeWithObject call and passed to CommitReplicationRecord.
 	 * @see CommitReplicationRecord
 	 */
-	IRISCORE_API void ProcessPacketDeliveryStatus(EPacketDeliveryStatus Status, ReplicationRecord Record);
+	IRISCORE_API void ProcessPacketDeliveryStatus(EPacketDeliveryStatus Status, const FReplicationRecord& Record);
 
 private:
 	enum Constants : uint32
 	{
 		IndexBitCount = 8U,
+		MaxWriteSequenceCount = UE_ARRAY_COUNT(FReplicationRecord::Sequences),
 	};
 
-	bool IsFull() const;
-
 	uint32 SequenceToIndex(uint32 Seq) const;
 
 	bool IsSequenceAcked(uint32 Seq) const;
@@ -129,11 +158,11 @@ private:
 
 	bool IsValidReceiveSequence(uint32 Seq) const;
 
-	uint32 SerializeInternal(FNetSerializationContext& Context, FNetRefHandle RefHandle, FReliableNetBlobQueue::ReplicationRecord& OutRecord, const bool bSerializeWithObject);
+	uint32 SerializeInternal(FNetSerializationContext& Context, FNetRefHandle RefHandle, FReliableNetBlobQueue::FReplicationRecord& OutRecord, const bool bSerializeWithObject);
 	uint32 DeserializeInternal(FNetSerializationContext& Context, FNetRefHandle RefHandle, const bool bSerializeWithObject);
 
-	void OnPacketDelivered(ReplicationRecord Record);
-	void OnPacketDropped(ReplicationRecord Record);
+	void OnPacketDelivered(const FReplicationRecord& Record);
+	void OnPacketDropped(const FReplicationRecord& Record);
 
 	void PopInOrderAckedBlobs();
 
@@ -146,7 +175,7 @@ private:
 };
 
 //
-inline bool FReliableNetBlobQueue::IsFull() const
+inline bool FReliableNetBlobQueue::IsSendWindowFull() const
 {
 	return (LastSeq - FirstSeq) >= MaxUnackedBlobCount;
 }
diff --git a/Engine/Source/Runtime/Experimental/Iris/Core/Public/Iris/Serialization/NetErrorContext.h b/Engine/Source/Runtime/Experimental/Iris/Core/Public/Iris/Serialization/NetErrorContext.h
--- a/Engine/Source/Runtime/Experimental/Iris/Core/Public/Iris/Serialization/NetErrorContext.h
+++ b/Engine/Source/Runtime/Experimental/Iris/Core/Public/Iris/Serialization/NetErrorContext.h
@@ -14,6 +14,7 @@ IRISCORE_API extern const FName GNetError_ArraySizeTooLarge;
 IRISCORE_API extern const FName GNetError_InvalidNetHandle;
 IRISCORE_API extern const FName GNetError_BrokenNetHandle;
 IRISCORE_API extern const FName GNetError_InvalidValue;
+IRISCORE_API extern const FName GNetError_InternalError;
 
 class FNetErrorContext
 {
