Chaos - add show queryonly, simonly, probeonly options to debugdraw

diff --git a/Engine/Source/Runtime/Experimental/Chaos/Private/Chaos/ChaosDebugDraw.cpp b/Engine/Source/Runtime/Experimental/Chaos/Private/Chaos/ChaosDebugDraw.cpp
--- a/Engine/Source/Runtime/Experimental/Chaos/Private/Chaos/ChaosDebugDraw.cpp
+++ b/Engine/Source/Runtime/Experimental/Chaos/Private/Chaos/ChaosDebugDraw.cpp
@@ -67,6 +67,15 @@ namespace Chaos
 		bool bChaosDebugDebugDrawColorBoundsByShapeType = false;
 		FAutoConsoleVariableRef CVarChaosDebugDebugDrawColorBoundsByShapeType(TEXT("p.Chaos.DebugDraw.ColorBoundsByShapeType"), bChaosDebugDebugDrawColorBoundsByShapeType, TEXT("Whether to use shape type to define the color of the bounds instead of using the particle state (if multiple shapes , will use the first one)"));
 
+		bool bChaosDebugDebugDrawShowQueryOnlyShapes = true;
+		FAutoConsoleVariableRef CVarChaosDebugDebugDrawShowQueryOnlyShapes(TEXT("p.Chaos.DebugDraw.ShowQueryOnly"), bChaosDebugDebugDrawShowQueryOnlyShapes, TEXT("Whether to show QueryOnly shapes"));
+
+		bool bChaosDebugDebugDrawShowSimOnlyShapes = true;
+		FAutoConsoleVariableRef CVarChaosDebugDebugDrawShowSimOnlyShapes(TEXT("p.Chaos.DebugDraw.ShowSimOnly"), bChaosDebugDebugDrawShowSimOnlyShapes, TEXT("Whether to show SimOnly shapes"));
+
+		bool bChaosDebugDebugDrawShowProbeOnlyShapes = true;
+		FAutoConsoleVariableRef CVarChaosDebugDebugDrawShowProbeOnlyShapes(TEXT("p.Chaos.DebugDraw.ShowProbeOnly"), bChaosDebugDebugDrawShowProbeOnlyShapes, TEXT("Whether to show ProbeOnly shapes"));
+
 		bool bChaosDebugDebugDrawConvexVertices = false;
 		bool bChaosDebugDebugDrawCoreShapes = false;
 		bool bChaosDebugDebugDrawExactCoreShapes = false;
@@ -313,16 +322,11 @@ namespace Chaos
 		//
 		//
 
-		void DrawShapesImpl(const FGeometryParticleHandle* Particle, const FRigidTransform3& ShapeTransform, const FImplicitObject* Implicit, const FShapeOrShapesArray& Shapes, const FReal Margin, const FColor& Color, const FRealSingle Duration, const FChaosDebugDrawSettings& Settings);
-
-		void DrawShape(const FRigidTransform3& ShapeTransform, const FImplicitObject* Implicit, const FShapeOrShapesArray& Shapes, const FColor& Color, const FChaosDebugDrawSettings* Settings)
-		{
-			DrawShapesImpl(nullptr, ShapeTransform, Implicit, Shapes, 0.0f, Color, 0.0f, GetChaosDebugDrawSettings(Settings));
-		}
+		void DrawShapesImpl(const FGeometryParticleHandle* Particle, const FRigidTransform3& ShapeTransform, const FImplicitObject* Implicit, const FPerShapeData* Shape, const FReal Margin, const FColor& Color, const FRealSingle Duration, const FChaosDebugDrawSettings& Settings);
 
-		void DrawShape(const FRigidTransform3& ShapeTransform, const FImplicitObject* Implicit, const FShapeOrShapesArray& Shapes, const FColor& Color, const float Duration, const FChaosDebugDrawSettings* Settings)
+		void DrawShape(const FRigidTransform3& ShapeTransform, const FImplicitObject* Implicit, const FPerShapeData* Shape, const FColor& Color, const float Duration, const FChaosDebugDrawSettings* Settings)
 		{
-			DrawShapesImpl(nullptr, ShapeTransform, Implicit, Shapes, 0.0f, Color, Duration, GetChaosDebugDrawSettings(Settings));
+			DrawShapesImpl(nullptr, ShapeTransform, Implicit, Shape, 0.0f, Color, Duration, GetChaosDebugDrawSettings(Settings));
 		}
 
 		void DrawShapesConvexImpl(const TGeometryParticleHandle<FReal, 3>* Particle, const FRigidTransform3& ShapeTransform, const FConvex* Shape, const FReal InMargin, const FColor& Color, const FRealSingle Duration, const FChaosDebugDrawSettings& Settings)
@@ -431,7 +435,7 @@ namespace Chaos
 		}
 
 		template <bool bInstanced>
-		void DrawShapesScaledImpl(const FGeometryParticleHandle* Particle, const FRigidTransform3& ShapeTransform, const FImplicitObject* Implicit, const FShapeOrShapesArray& Shapes, const FReal Margin, const FColor& Color, const FRealSingle Duration, const FChaosDebugDrawSettings& Settings)
+		void DrawShapesScaledImpl(const FGeometryParticleHandle* Particle, const FRigidTransform3& ShapeTransform, const FImplicitObject* Implicit, const FPerShapeData* Shape, const FReal Margin, const FColor& Color, const FRealSingle Duration, const FChaosDebugDrawSettings& Settings)
 		{
 			const EImplicitObjectType PackedType = Implicit->GetType();
 			const EImplicitObjectType InnerType = GetInnerType(PackedType);
@@ -458,7 +462,7 @@ namespace Chaos
 				const TImplicitObjectScaled<FLevelSet, bInstanced>* Scaled = Implicit->template GetObject<TImplicitObjectScaled<FLevelSet, bInstanced>>();
 				// even though thhe levelset is scaled, the debugdraw uses the collisionParticles  that are pre-scaled
 				// so no need to pass the scaled transform and just extract the wrapped LevelSet
-				DrawShapesImpl(Particle, ShapeTransform, Scaled->GetUnscaledObject(), Shapes, Scaled->GetMargin(), Color, Duration, Settings);
+				DrawShapesImpl(Particle, ShapeTransform, Scaled->GetUnscaledObject(), Shape, Scaled->GetMargin(), Color, Duration, Settings);
 				break;
 			}
 			case ImplicitObjectType::Unknown:
@@ -467,7 +471,7 @@ namespace Chaos
 			{
 				const TImplicitObjectScaled<FConvex, bInstanced>* Scaled = Implicit->template GetObject<TImplicitObjectScaled<FConvex, bInstanced>>();
 				ScaleTM.SetScale3D(Scaled->GetScale());
-				DrawShapesImpl(Particle, ScaleTM * ShapeTransform, Scaled->GetUnscaledObject(), Shapes, Scaled->GetMargin(), Color, Duration, Settings);
+				DrawShapesImpl(Particle, ScaleTM * ShapeTransform, Scaled->GetUnscaledObject(), Shape, Scaled->GetMargin(), Color, Duration, Settings);
 				break;
 			}
 			case ImplicitObjectType::TaperedCylinder:
@@ -478,14 +482,14 @@ namespace Chaos
 			{
 				const TImplicitObjectScaled<FTriangleMeshImplicitObject, bInstanced>* Scaled = Implicit->template GetObject<TImplicitObjectScaled<FTriangleMeshImplicitObject, bInstanced>>();
 				ScaleTM.SetScale3D(Scaled->GetScale());
-				DrawShapesImpl(Particle, ScaleTM * ShapeTransform, Scaled->GetUnscaledObject(), Shapes, 0.0f, Color, Duration, Settings);
+				DrawShapesImpl(Particle, ScaleTM * ShapeTransform, Scaled->GetUnscaledObject(), Shape, 0.0f, Color, Duration, Settings);
 				break;
 			}
 			case ImplicitObjectType::HeightField:
 			{
 				const TImplicitObjectScaled<FHeightField, bInstanced>* Scaled = Implicit->template GetObject<TImplicitObjectScaled<FHeightField, bInstanced>>();
 				ScaleTM.SetScale3D(Scaled->GetScale());
-				DrawShapesImpl(Particle, ScaleTM * ShapeTransform, Scaled->GetUnscaledObject(), Shapes, 0.0f, Color, Duration, Settings);
+				DrawShapesImpl(Particle, ScaleTM * ShapeTransform, Scaled->GetUnscaledObject(), Shape, 0.0f, Color, Duration, Settings);
 				break;
 			}
 			default:
@@ -493,7 +497,7 @@ namespace Chaos
 			}
 		}
 
-		void DrawShapesInstancedImpl(const FGeometryParticleHandle* Particle, const FRigidTransform3& ShapeTransform, const FImplicitObject* Implicit, const FShapeOrShapesArray& Shapes, const FReal Margin, const FColor& Color, const FRealSingle Duration, const FChaosDebugDrawSettings& Settings)
+		void DrawShapesInstancedImpl(const FGeometryParticleHandle* Particle, const FRigidTransform3& ShapeTransform, const FImplicitObject* Implicit, const FPerShapeData* Shape, const FReal Margin, const FColor& Color, const FRealSingle Duration, const FChaosDebugDrawSettings& Settings)
 		{
 			const EImplicitObjectType PackedType = Implicit->GetType();
 			const EImplicitObjectType InnerType = GetInnerType(PackedType);
@@ -521,7 +525,7 @@ namespace Chaos
 			case ImplicitObjectType::Convex:
 			{
 				const TImplicitObjectInstanced<FConvex>* Instanced = Implicit->template GetObject<TImplicitObjectInstanced<FConvex>>();
-				DrawShapesImpl(Particle, ShapeTransform, Instanced->GetInstancedObject(), Shapes, Instanced->GetMargin(), Color, Duration, Settings);
+				DrawShapesImpl(Particle, ShapeTransform, Instanced->GetInstancedObject(), Shape, Instanced->GetMargin(), Color, Duration, Settings);
 				break;
 			}
 			case ImplicitObjectType::TaperedCylinder:
@@ -531,13 +535,13 @@ namespace Chaos
 			case ImplicitObjectType::TriangleMesh:
 			{
 				const TImplicitObjectInstanced<FTriangleMeshImplicitObject>* Scaled = Implicit->template GetObject<TImplicitObjectInstanced<FTriangleMeshImplicitObject>>();
-				DrawShapesImpl(Particle, ShapeTransform, Scaled->GetInstancedObject(), Shapes, 0.0f, Color, Duration, Settings);
+				DrawShapesImpl(Particle, ShapeTransform, Scaled->GetInstancedObject(), Shape, 0.0f, Color, Duration, Settings);
 				break;
 			}
 			case ImplicitObjectType::HeightField:
 			{
 				const TImplicitObjectInstanced<FHeightField>* Scaled = Implicit->template GetObject<TImplicitObjectInstanced<FHeightField>>();
-				DrawShapesImpl(Particle, ShapeTransform, Scaled->GetInstancedObject(), Shapes, 0.0f, Color, Duration, Settings);
+				DrawShapesImpl(Particle, ShapeTransform, Scaled->GetInstancedObject(), Shape, 0.0f, Color, Duration, Settings);
 				break;
 			}
 			default:
@@ -545,7 +549,7 @@ namespace Chaos
 			}
 		}
 
-		void DrawShapesImpl(const FGeometryParticleHandle* Particle, const FRigidTransform3& ShapeTransform, const FImplicitObject* Implicit, const FShapeOrShapesArray& Shapes, const FReal Margin, const FColor& Color, const FRealSingle Duration, const FChaosDebugDrawSettings& Settings)
+		void DrawShapesImpl(const FGeometryParticleHandle* Particle, const FRigidTransform3& ShapeTransform, const FImplicitObject* Implicit, const FPerShapeData* Shape, const FReal Margin, const FColor& Color, const FRealSingle Duration, const FChaosDebugDrawSettings& Settings)
 		{
 			if (Implicit == nullptr)
 			{
@@ -566,24 +570,24 @@ namespace Chaos
 			{
 				if (IsInstanced(PackedType))
 				{
-					DrawShapesScaledImpl<true>(Particle, ShapeTransform, Implicit, Shapes, Margin, Color, Duration, Settings);
+					DrawShapesScaledImpl<true>(Particle, ShapeTransform, Implicit, Shape, Margin, Color, Duration, Settings);
 				}
 				else
 				{
-					DrawShapesScaledImpl<false>(Particle, ShapeTransform, Implicit, Shapes, Margin, Color, Duration, Settings);
+					DrawShapesScaledImpl<false>(Particle, ShapeTransform, Implicit, Shape, Margin, Color, Duration, Settings);
 				}
 				return;
 			}
 			else if (IsInstanced(PackedType))
 			{
-				DrawShapesInstancedImpl(Particle, ShapeTransform, Implicit, Shapes, Margin, Color, Duration, Settings);
+				DrawShapesInstancedImpl(Particle, ShapeTransform, Implicit, Shape, Margin, Color, Duration, Settings);
 				return;
 			}
 			else if (InnerType == ImplicitObjectType::Transformed)
 			{
 				const TImplicitObjectTransformed<FReal, 3>* Transformed = Implicit->template GetObject<TImplicitObjectTransformed<FReal, 3>>();
 				FRigidTransform3 TransformedTransform = FRigidTransform3(ShapeTransform.TransformPosition(Transformed->GetTransform().GetLocation()), ShapeTransform.GetRotation() * Transformed->GetTransform().GetRotation());
-				DrawShapesImpl(Particle, TransformedTransform, Transformed->GetTransformedObject(), Shapes, Margin, Color, Duration, Settings);
+				DrawShapesImpl(Particle, TransformedTransform, Transformed->GetTransformedObject(), Shape, Margin, Color, Duration, Settings);
 				return;
 			}
 			else if ((InnerType == ImplicitObjectType::Union) || (InnerType == ImplicitObjectType::UnionClustered))
@@ -593,14 +597,7 @@ namespace Chaos
 					const FImplicitObjectUnion* Union = Implicit->template GetObject<FImplicitObjectUnion>();
 					for (int32 UnionIdx = 0; UnionIdx < Union->GetObjects().Num(); ++UnionIdx)
 					{
-						// Retrieve shape from union's shapes array
-						const FPerShapeData* PerShapeData = nullptr;
-						if (!Shapes.IsSingleShape())
-						{
-							const FShapesArray& ShapesArray = *Shapes.GetShapesArray();
-							PerShapeData = ShapesArray[UnionIdx].Get();
-						}
-						DrawShapesImpl(Particle, ShapeTransform, Union->GetObjects()[UnionIdx].GetReference(), FShapeOrShapesArray(PerShapeData), Margin, Color, Duration, Settings);
+						DrawShapesImpl(Particle, ShapeTransform, Union->GetObjects()[UnionIdx].GetReference(), Shape, Margin, Color, Duration, Settings);
 					}
 				}
 				else
@@ -613,7 +610,7 @@ namespace Chaos
 						const TPBDRigidParticleHandle<FReal, 3>* OriginalParticle = Union->FindParticleForImplicitObject(UnionImplicit.GetReference());
 						if (ensure(OriginalParticle))
 						{
-							DrawShapesImpl(Particle, ShapeTransform, UnionImplicit.GetReference(), FShapeOrShapesArray(OriginalParticle), Margin, Color, Duration, Settings);
+							DrawShapesImpl(Particle, ShapeTransform, UnionImplicit.GetReference(), Shape, Margin, Color, Duration, Settings);
 						}
 					}
 				}
@@ -631,7 +628,7 @@ namespace Chaos
 								{
 									for(auto& UnionConvex : Union->GetObjects())
 									{
-										DrawShapesImpl(Particle, ShapeTransform, UnionConvex, FShapeOrShapesArray(PerShapeData.Get()), Margin, Color, Duration, Settings);
+										DrawShapesImpl(Particle, ShapeTransform, UnionConvex, PerShapeData.Get(), Margin, Color, Duration, Settings);
 									}
 								}
 							}
@@ -645,18 +642,12 @@ namespace Chaos
 			bool bShowMeshes = Settings.bShowComplexCollision;
 			bool bShowNonMeshes = Settings.bShowSimpleCollision;
 
-			const FPerShapeData* ShapeData = nullptr;
-			if (Shapes.IsValid() && ensure(Shapes.IsSingleShape()))
-			{
-				ShapeData = Shapes.GetShape();
-			}
-			
-			if (ShapeData != nullptr)
+			if (Shape != nullptr)
 			{
-				bShowMeshes = (Settings.bShowComplexCollision && (ShapeData->GetCollisionTraceType() != EChaosCollisionTraceFlag::Chaos_CTF_UseSimpleAsComplex))
-					|| (Settings.bShowSimpleCollision && (ShapeData->GetCollisionTraceType() == EChaosCollisionTraceFlag::Chaos_CTF_UseComplexAsSimple));
-				bShowNonMeshes = (Settings.bShowSimpleCollision && (ShapeData->GetCollisionTraceType() != EChaosCollisionTraceFlag::Chaos_CTF_UseComplexAsSimple))
-					|| (Settings.bShowComplexCollision && (ShapeData->GetCollisionTraceType() == EChaosCollisionTraceFlag::Chaos_CTF_UseSimpleAsComplex));
+				bShowMeshes = (Settings.bShowComplexCollision && (Shape->GetCollisionTraceType() != EChaosCollisionTraceFlag::Chaos_CTF_UseSimpleAsComplex))
+					|| (Settings.bShowSimpleCollision && (Shape->GetCollisionTraceType() == EChaosCollisionTraceFlag::Chaos_CTF_UseComplexAsSimple));
+				bShowNonMeshes = (Settings.bShowSimpleCollision && (Shape->GetCollisionTraceType() != EChaosCollisionTraceFlag::Chaos_CTF_UseComplexAsSimple))
+					|| (Settings.bShowComplexCollision && (Shape->GetCollisionTraceType() == EChaosCollisionTraceFlag::Chaos_CTF_UseSimpleAsComplex));
 			}
 
 			// Quit if we don't want to show this shape
@@ -704,9 +695,9 @@ namespace Chaos
 					}
 				}
 			}
-			if (ShapeData && bChaosDebugDebugDrawColorShapesBySimQueryType)
+			if (Shape && bChaosDebugDebugDrawColorShapesBySimQueryType)
 			{
-				if (ShapeData->GetSimEnabled())
+				if (Shape->GetSimEnabled())
 				{
 					const bool bIsUnion = Particle->GetGeometry()->IsUnderlyingUnion();
 					if (bIsUnion)
@@ -725,11 +716,32 @@ namespace Chaos
 						ShapeColor = FColor::Orange;
 					}
 				}
-				else if (ShapeData->GetQueryEnabled())
+				else if (Shape->GetQueryEnabled())
 				{
 					ShapeColor = FColor::Red;
 				}
 			}
+			if (Shape && !bChaosDebugDebugDrawShowQueryOnlyShapes)
+			{
+				if (!Shape->GetSimEnabled() && !Shape->GetIsProbe())
+				{
+					return;
+				}
+			}
+			if (Shape && !bChaosDebugDebugDrawShowSimOnlyShapes)
+			{
+				if (!Shape->GetQueryEnabled() && !Shape->GetIsProbe())
+				{
+					return;
+				}
+			}
+			if (Shape && !bChaosDebugDebugDrawShowProbeOnlyShapes)
+			{
+				if (!Shape->GetSimEnabled() && !Shape->GetQueryEnabled())
+				{
+					return;
+				}
+			}
 
 			// If we get here, we have an actual shape to render
 			switch (InnerType)
@@ -840,7 +852,10 @@ namespace Chaos
 			FVec3 P = SpaceTransform.TransformPosition(Particle->ObjectState() == EObjectStateType::Dynamic ? Particle->CastToRigidParticle()->P() : Particle->X());
 			FRotation3 Q = SpaceTransform.GetRotation() * (Particle->ObjectState() == EObjectStateType::Dynamic ? Particle->CastToRigidParticle()->Q() : Particle->R());
 
-			DrawShapesImpl(Particle, FRigidTransform3(P, Q), Particle->GetGeometry(), FShapeOrShapesArray(Particle), 0.0f, InColor, 0.0f, Settings);
+			for (const FShapeInstancePtr& ShapeInstance : Particle->ShapeInstances())
+			{
+				DrawShapesImpl(Particle, FRigidTransform3(P, Q), ShapeInstance->GetGeometry(), ShapeInstance.Get(), 0.0f, InColor, 0.0f, Settings);
+			}
 		}
 
 		void DrawParticleShapesImpl(const FRigidTransform3& SpaceTransform, const FGeometryParticle* Particle, const FColor& InColor, const FChaosDebugDrawSettings& Settings)
@@ -848,7 +863,10 @@ namespace Chaos
 			FVec3 P = SpaceTransform.TransformPosition(Particle->X());
 			FRotation3 Q = SpaceTransform.GetRotation() * (Particle->R());
 
-			DrawShapesImpl(Particle->Handle(), FRigidTransform3(P, Q), Particle->GetGeometry(), FShapeOrShapesArray(Particle->Handle()), 0.0f, InColor, 0.0f, Settings);
+			for (const FShapeInstanceProxyPtr& ShapeInstance : Particle->ShapeInstances())
+			{
+				DrawShapesImpl(Particle->Handle(), FRigidTransform3(P, Q), ShapeInstance->GetGeometry(), ShapeInstance.Get(), 0.0f, InColor, 0.0f, Settings);
+			}
 		}
 
 		void DrawBVHImpl(const FGeometryParticleHandle* Particle, const FRigidTransform3& ShapeTransform, const Private::FImplicitBVH* BVH, const FColor& UnusedColor, const FRealSingle Duration, const FChaosDebugDrawSettings& Settings)
@@ -878,7 +896,7 @@ namespace Chaos
 								{
 									const int32 ShapeIndex = (RootObjectIndex != INDEX_NONE) ? RootObjectIndex : 0;
 									const FShapeInstance* Shape = Particle->ShapeInstances()[ShapeIndex].Get();
-									DrawShapesImpl(Particle, FRigidTransform3(RelativeTransformf) * ShapeTransform, Implicit, FShapeOrShapesArray(Shape), 0, Color, Duration, Settings);
+									DrawShapesImpl(Particle, FRigidTransform3(RelativeTransformf) * ShapeTransform, Implicit, Shape, 0, Color, Duration, Settings);
 								});
 						}
 					}
@@ -1271,12 +1289,12 @@ namespace Chaos
 					const FRigidTransform3 ShapeWorldTransform1 = Collision.GetShapeRelativeTransform1() * WorldActorTransform1;
 					DrawShapesImpl(
 						Particle0->Handle(), ShapeWorldTransform0,
-						Implicit0, FShapeOrShapesArray(Shape0),
+						Implicit0, Shape0,
 						0.0f, Particle0->IsDynamic() ? FColor::Yellow : FColor::Red,
 						Duration, Settings);
 					DrawShapesImpl(
 						Particle1->Handle(), ShapeWorldTransform1,
-						Implicit1, FShapeOrShapesArray(Shape1),
+						Implicit1, Shape1,
 						0.0f, Particle1->IsDynamic() ? FColor::Yellow : FColor::Red,
 						Duration, Settings);
 				}
@@ -1345,6 +1363,23 @@ namespace Chaos
 		{
 			if (const auto RigidParticle = InParticle->CastToRigidParticle())
 			{
+				if (!bChaosDebugDebugDrawShowQueryOnlyShapes)
+				{
+					bool bIsQueryOnly = true;
+					for (const FShapeInstancePtr& ShapeInstance : InParticle->ShapeInstances())
+					{
+						if (ShapeInstance->GetSimEnabled())
+						{
+							bIsQueryOnly = false;
+							break;
+						}
+					}
+					if (bIsQueryOnly)
+					{
+						return;
+					}
+				}
+
 				const FColor Color = RigidParticle->IsDynamic()? FColor::Yellow: FColor::Cyan;
 
 				const FVec3 PCOM = SpaceTransform.TransformPosition(FParticleUtilities::GetCoMWorldPosition(RigidParticle));
diff --git a/Engine/Source/Runtime/Experimental/Chaos/Public/Chaos/ChaosDebugDraw.h b/Engine/Source/Runtime/Experimental/Chaos/Public/Chaos/ChaosDebugDraw.h
--- a/Engine/Source/Runtime/Experimental/Chaos/Public/Chaos/ChaosDebugDraw.h
+++ b/Engine/Source/Runtime/Experimental/Chaos/Public/Chaos/ChaosDebugDraw.h
@@ -222,9 +222,8 @@ namespace Chaos
 		CHAOS_API void DrawJointConstraints(const FRigidTransform3& SpaceTransform, const TArray<FPBDJointConstraintHandle*>& ConstraintHandles, FRealSingle ColorScale, const FChaosDebugDrawJointFeatures& FeatureMask = FChaosDebugDrawJointFeatures::MakeDefault(), const FChaosDebugDrawSettings* Settings = nullptr);
 		CHAOS_API void DrawJointConstraints(const FRigidTransform3& SpaceTransform, const FPBDJointConstraints& Constraints, FRealSingle ColorScale, const FChaosDebugDrawJointFeatures& FeatureMask = FChaosDebugDrawJointFeatures::MakeDefault(), const FChaosDebugDrawSettings* Settings = nullptr);
 		CHAOS_API void DrawCharacterGroundConstraints(const FRigidTransform3& SpaceTransform, const FCharacterGroundConstraintContainer& Constraints, const FChaosDebugDrawSettings* Settings = nullptr);
+		CHAOS_API void DrawShape(const FRigidTransform3& ShapeTransform, const FImplicitObject* Implicit, const FPerShapeData* Shape, const FColor& Color, const float Duration, const FChaosDebugDrawSettings* Settings = nullptr);
 		CHAOS_API void DrawSimulationSpace(const FSimulationSpace& SimSpace, const FChaosDebugDrawSettings* Settings = nullptr);
-		CHAOS_API void DrawShape(const FRigidTransform3& ShapeTransform, const FImplicitObject* Implicit, const FShapeOrShapesArray& Shapes, const FColor& Color, const FChaosDebugDrawSettings* Settings = nullptr);
-		CHAOS_API void DrawShape(const FRigidTransform3& ShapeTransform, const FImplicitObject* Implicit, const FShapeOrShapesArray& Shapes, const FColor& Color, const float Duration, const FChaosDebugDrawSettings* Settings = nullptr);
 		CHAOS_API void DrawConstraintGraph(const FRigidTransform3& ShapeTransform, const Private::FPBDIslandManager& Graph, const FChaosDebugDrawSettings* Settings = nullptr);
 		CHAOS_API void DrawConnectionGraph(const FRigidClustering& Clustering, const FChaosDebugDrawSettings* Settings = nullptr);
 		CHAOS_API void DrawCollidingShapes(const FRigidTransform3& SpaceTransform, const FPBDCollisionConstraints& Collisions, FRealSingle ColorScale, const FRealSingle Duration = 0.f, const FChaosDebugDrawSettings* Settings = nullptr);
@@ -233,6 +232,19 @@ namespace Chaos
 		CHAOS_API void DrawSuspensionConstraints(const FRigidTransform3& SpaceTransform, const FPBDSuspensionConstraints& Constraints, const FChaosDebugDrawSettings* Settings = nullptr);
 		CHAOS_API void DrawCCDAxisThreshold(const FVec3& X, const FVec3& AxisThreshold, const FVec3& DeltaX, const FQuat& R);
 		CHAOS_API void DrawCCDCollisionImpulse(const FRigidTransform3& SpaceTransform, const FCCDConstraint& CCDConstraint, const int32 ManifoldPointIndex, const FVec3& Impulse, const FChaosDebugDrawSettings* Settings = nullptr);
+
+
+		UE_DEPRECATED(5.4, "Pass in the ShapeInstance or call DrawParticleShapes")
+		inline void DrawShape(const FRigidTransform3& ShapeTransform, const FImplicitObject* Implicit, const FShapeOrShapesArray& Shapes, const FColor& Color, const FChaosDebugDrawSettings* Settings = nullptr)
+		{
+			DrawShape(ShapeTransform, Implicit, Shapes.IsSingleShape() ? Shapes.GetShape() : nullptr, Color, 0.0f, Settings);
+		}
+
+		UE_DEPRECATED(5.4, "Pass in the ShapeInstance or call DrawParticleShapes")
+		inline void DrawShape(const FRigidTransform3& ShapeTransform, const FImplicitObject* Implicit, const FShapeOrShapesArray& Shapes, const FColor& Color, const float Duration, const FChaosDebugDrawSettings* Settings = nullptr)
+		{
+			DrawShape(ShapeTransform, Implicit, Shapes.IsSingleShape() ? Shapes.GetShape() : nullptr, Color, Duration, Settings);
+		}
 #endif
 	}
 }
diff --git a/Engine/Source/Runtime/PhysicsCore/Public/SQVisitor.h b/Engine/Source/Runtime/PhysicsCore/Public/SQVisitor.h
--- a/Engine/Source/Runtime/PhysicsCore/Public/SQVisitor.h
+++ b/Engine/Source/Runtime/PhysicsCore/Public/SQVisitor.h
@@ -397,7 +397,7 @@ private:
 		}
 		else if (SQ == ESQType::Overlap)
 		{
-			Chaos::DebugDraw::DrawShape(StartTM, QueryGeom, Chaos::FShapeOrShapesArray(), bHit ? FColor::Red : FColor::Green);
+			Chaos::DebugDraw::DrawShape(StartTM, QueryGeom, nullptr, bHit ? FColor::Red : FColor::Green, 0.0f);
 		}
 
 		if (Instance.bHasBounds)
