Smaller alloc size increases for request builds and message receives
(git:a0ad1a6 - Martin Ridgers - [REDACTED_PHONE]:30:53 +0200)

#rnx

diff --git a/Engine/Source/Runtime/Experimental/IoStoreOnDemand/Private/CoreHttp/Http.cpp b/Engine/Source/Runtime/Experimental/IoStoreOnDemand/Private/CoreHttp/Http.cpp
--- a/Engine/Source/Runtime/Experimental/IoStoreOnDemand/Private/CoreHttp/Http.cpp
+++ b/Engine/Source/Runtime/Experimental/IoStoreOnDemand/Private/CoreHttp/Http.cpp
@@ -491,12 +491,12 @@ public:
 	uint32						GetSize() const;
 	uint32						GetCapacity() const;
 	template <typename T> T*	Alloc(uint32 Count=1);
-	FMutableSection				GetMutableFree(uint32 MinSize, uint32 PageSize=0);
+	FMutableSection				GetMutableFree(uint32 MinSize, uint32 PageSize=256);
 	void						AdvanceUsed(uint32 Delta);
 
 private:
 	char*						GetDataPtr();
-	void						Extend(uint32 AtLeast, uint32 PageSize=1024);
+	void						Extend(uint32 AtLeast, uint32 PageSize);
 	union
 	{
 		struct
@@ -593,7 +593,7 @@ T* FBuffer::Alloc(uint32 Count)
 	uint32 PotentialUsed = Used + AlignBias + (sizeof(T) * Count);
 	if (PotentialUsed > Max)
 	{
-		Extend(PotentialUsed);
+		Extend(PotentialUsed, 256);
 	}
 
 	void* Ret = GetDataPtr() + Used + AlignBias;
@@ -609,7 +609,7 @@ FBuffer::FMutableSection FBuffer::GetMutableFree(uint32 MinSize, uint32 PageSize
 	uint32 PotentialUsed = Used + MinSize;
 	if (PotentialUsed > Max)
 	{
-		Extend(PotentialUsed);
+		Extend(PotentialUsed, PageSize);
 	}
 
 	return FMutableSection{ GetDataPtr() + Used, Max - Used };
@@ -2202,7 +2202,7 @@ static int32 DoRecvMessage(FActivity* Activity)
 {
 	TRACE_CPUPROFILER_EVENT_SCOPE(IasHttp::DoRecvMessage);
 
-	static const uint32 PageSize = 2048;
+	static const uint32 PageSize = 384;
 
 	FBuffer& Buffer = Activity->Buffer;
 
