Moved work reception into its own method
(git:2ce4019 - Martin Ridgers - [REDACTED_PHONE]:18:55 +0200)
#rnx

diff --git a/Engine/Source/Runtime/Experimental/IoStoreOnDemand/Private/CoreHttp/Http.cpp b/Engine/Source/Runtime/Experimental/IoStoreOnDemand/Private/CoreHttp/Http.cpp
--- a/Engine/Source/Runtime/Experimental/IoStoreOnDemand/Private/CoreHttp/Http.cpp
+++ b/Engine/Source/Runtime/Experimental/IoStoreOnDemand/Private/CoreHttp/Http.cpp
@@ -2630,6 +2630,7 @@ public:
 	FTicket					Send(FActivity* Activity);
 
 private:
+	void					ReceiveWork();
 	FCriticalSection		Lock;
 	std::atomic<uint64>		FreeSlots		= ~0ull;
 	std::atomic<uint64>		Cancels			= 0;
@@ -2732,28 +2733,31 @@ void FEventLoop::FImpl::Cancel(FTicket Ticket)
 }
 
 ////////////////////////////////////////////////////////////////////////////////
-uint32 FEventLoop::FImpl::Tick(uint32 PollTimeoutMs)
+void FEventLoop::FImpl::ReceiveWork()
 {
-	TRACE_CPUPROFILER_EVENT_SCOPE(IasHttp::Tick);
-
-	// Collect activity changes
 	uint64 FreeSlotsLoad = FreeSlots.load(std::memory_order_relaxed);
-	if (FreeSlots != PrevFreeSlots)
+	if (FreeSlots == PrevFreeSlots)
 	{
-		TArray<FActivity*> NewActive;
-		{
-			FScopeLock _(&Lock);
-			NewActive = MoveTemp(Pending);
-		}
-
-		for (FActivity* Activity : NewActive)
-		{
-			Active.Add(Activity);
-		}
+		return;
+	}
+	PrevFreeSlots = FreeSlotsLoad;
 
-		PrevFreeSlots = FreeSlotsLoad;
+	TArray<FActivity*> NewActive;
+	{
+		FScopeLock _(&Lock);
+		NewActive = MoveTemp(Pending);
 	}
 
+	Active.Append(NewActive);
+}
+
+////////////////////////////////////////////////////////////////////////////////
+uint32 FEventLoop::FImpl::Tick(uint32 PollTimeoutMs)
+{
+	TRACE_CPUPROFILER_EVENT_SCOPE(IasHttp::Tick);
+
+	ReceiveWork();
+
 	uint32 BusyCount = Active.Num();
 
 	// Work out which activities are ready
