Make sure IAS' service thread exits correctly on shutdown

#rnx

diff --git a/Engine/Source/Runtime/Experimental/IoStoreOnDemand/Private/JournaledCache.cpp b/Engine/Source/Runtime/Experimental/IoStoreOnDemand/Private/JournaledCache.cpp
--- a/Engine/Source/Runtime/Experimental/IoStoreOnDemand/Private/JournaledCache.cpp
+++ b/Engine/Source/Runtime/Experimental/IoStoreOnDemand/Private/JournaledCache.cpp
@@ -1423,6 +1423,7 @@ public:
 	};
 
 	static FServiceThread&	Get();
+							~FServiceThread();
 	void					RegisterCache(TUniquePtr<FCache> Cache);
 	void					UnregisterCache(FCache* Cache);
 	void					SetGovernorRate(uint32 Allowance, uint32 Ops, uint32 Seconds);
@@ -1505,6 +1506,14 @@ FServiceThread& FServiceThread::Get()
 	static FServiceThread Instance;
 	Ptr = &Instance;
 	return *Ptr;
+};
+
+////////////////////////////////////////////////////////////////////////////////
+FServiceThread::~FServiceThread()
+{
+	uint32 PrevRunCount = RunCount.fetch_sub(1, std::memory_order_relaxed);
+	check(PrevRunCount == 1);
+	WakeEvent->Trigger();
 }
 
 ////////////////////////////////////////////////////////////////////////////////
