Resolve() now uses the FResult object

#rnx
(git:f28d447 - Martin Ridgers - [REDACTED_PHONE]:06:49 +0200)

diff --git a/Engine/Source/Runtime/Experimental/IoStoreOnDemand/Private/CoreHttp/Http.cpp b/Engine/Source/Runtime/Experimental/IoStoreOnDemand/Private/CoreHttp/Http.cpp
--- a/Engine/Source/Runtime/Experimental/IoStoreOnDemand/Private/CoreHttp/Http.cpp
+++ b/Engine/Source/Runtime/Experimental/IoStoreOnDemand/Private/CoreHttp/Http.cpp
@@ -174,6 +174,22 @@ static void Trace(const void* Id, ETrace Action, T Param=0)
 
 // {{{1 misc ...................................................................
 
+////////////////////////////////////////////////////////////////////////////////
+class FResult
+{
+public:
+				FResult() : FResult(0, "") {}
+				FResult(const char* Msg="") : FResult(-1, Msg) {}
+				FResult(int32 Val, const char* Msg="") : Message(UPTRINT(Msg)), Value(Val) {}
+	const char*	GetMessage() const		{ return (const char*)Message; }
+	int32		GetValue() const		{ return int16(Value); }
+
+private:
+	UPTRINT		Message : 48;
+	PTRINT		Value : 16;
+};
+static_assert(sizeof(FResult) == sizeof(void*));
+
 ////////////////////////////////////////////////////////////////////////////////
 template <typename LambdaType>
 static void EnumerateHeaders(FAnsiStringView Headers, LambdaType&& Lambda)
@@ -760,6 +776,8 @@ bool FSocket::Connect(uint32 IpAddress, uint32 Port)
 {
 	check(IsValid());
 
+	IpAddress = htonl(IpAddress);
+
 	sockaddr_in AddrInet = { sizeof(sockaddr_in) };
 	AddrInet.sin_family = AF_INET;
 	AddrInet.sin_port = htons(uint16(Port));
@@ -1174,12 +1192,14 @@ class FHost
 {
 public:
 	enum class EDirection : uint8 { Send, Recv };
+	static const uint32 InvalidIp = 0x00ff'ffff;
 
 					FHost(const ANSICHAR* InHostName, uint32 InPort);
 	void			SetBufferSize(EDirection Dir, int32 Size);
 	int32			GetBufferSize(EDirection Dir) const;
+	FResult			Connect(FSocket& Socket);
 	int32			IsResolved() const;
-	int32			ResolveHostName();
+	FResult			ResolveHostName();
 	uint32			GetIpAddress() const		{ return IpAddresses[0]; }
 	FAnsiStringView	GetHostName() const			{ return HostName; }
 	uint32			GetPort() const				{ return Port; }
@@ -1212,7 +1232,7 @@ int32 FHost::GetBufferSize(EDirection Dir) const
 }
 
 ////////////////////////////////////////////////////////////////////////////////
-int32 FHost::ResolveHostName()
+FResult FHost::ResolveHostName()
 {
 	// todo: GetAddrInfoW() for async resolve on Windows
 
@@ -1230,12 +1250,12 @@ int32 FHost::ResolveHostName()
 	auto Result = getaddrinfo(HostName, nullptr, &Hints, &Info);
 	if (uint32(Result) || Info == nullptr)
 	{
-		return -1;
+		return FResult(-1, "Error encountered resolving");
 	}
 
 	if (Info->ai_family != AF_INET)
 	{
-		return -2;
+		return FResult(-2, "Unexpected address family during resolve");
 	}
 
 	uint32 AddressCount = 0;
@@ -1255,7 +1275,7 @@ int32 FHost::ResolveHostName()
 			break;
 		}
 
-		IpAddresses[AddressCount] = IpAddress;
+		IpAddresses[AddressCount] = htonl(IpAddress);
 		if (++AddressCount >= UE_ARRAY_COUNT(IpAddresses))
 		{
 			break;
@@ -1267,7 +1287,7 @@ int32 FHost::ResolveHostName()
 		return AddressCount;
 	}
 
-	return 0;
+	return FResult(0, "Unable to resolve host");
 }
 
 ////////////////////////////////////////////////////////////////////////////////
@@ -1281,21 +1301,81 @@ int32 FHost::IsResolved() const
 	}
 }
 
+////////////////////////////////////////////////////////////////////////////////
+FResult FHost::Connect(FSocket& Socket)
+{
+	if (IsResolved() <= 0)
+	{
+		FResult Result = ResolveHostName();
+		if (Result.GetValue() <= 0)
+		{
+			return Result;
+		}
+	}
+
+	check(IsResolved() > 0);
+	check(!Socket.IsValid());
+
+	uint32 IpAddress = GetIpAddress();
+
+	FSocket Candidate;
+	if (!Candidate.Create())
+	{
+		return FResult(-1, "Failed to create socket");
+	}
+
+	// Attempt a SOCKS connect
+	bool bSocksConnected = false;
+	if (int32 Result = MaybeConnectSocks(Candidate, IpAddress, Port); Result)
+	{
+		if (Result < 0)
+		{
+			return FResult("Failed establishing SOCKS connection");
+		}
+
+		bSocksConnected = true;
+	}
+
+	// Condition the socket
+	if (!Candidate.SetBlocking(false))
+	{
+		return FResult("Unable to set socket non-blocking");
+	}
+
+	if (int32 OptValue = GetBufferSize(FHost::EDirection::Send); OptValue >= 0)
+	{
+		Candidate.SetSendBufSize(OptValue);
+	}
+
+	if (int32 OptValue = GetBufferSize(FHost::EDirection::Recv); OptValue >= 0)
+	{
+		Candidate.SetRecvBufSize(OptValue);
+	}
+
+	// Adjust socket send and recv buffer sizes
+	if (!bSocksConnected && !Candidate.Connect(IpAddress, Port))
+	{
+		return FResult("Socket connect failed");
+	}
+
+	Socket = MoveTemp(Candidate);
+	return 1;
+}
+
 
 
 ////////////////////////////////////////////////////////////////////////////////
 class FSocketPool
+	: public FHost
 {
 public:
 					FSocketPool(const ANSICHAR* InHostName, uint32 InPort, uint32 InMaxLeases);
 					~FSocketPool();
 	static uint32	GetAllocSize(uint32 MaxLeases);
-	FHost&			GetHost() { return Host; }
 	bool			LeaseSocket(FSocket& Out);
 	void			ReturnLease(FSocket&& Socket);
 
 private:
-	FHost			Host;
 	uint8			LeaseCount = 0;
 	uint8			MaxLeases;
 	FSocket			Sockets[1/*...N*/]; // this should be the last member
@@ -1303,7 +1383,7 @@ private:
 
 ////////////////////////////////////////////////////////////////////////////////
 FSocketPool::FSocketPool(const ANSICHAR* InHostName, uint32 InPort, uint32 InMaxLeases)
-: Host(InHostName, InPort)
+: FHost(InHostName, InPort)
 , MaxLeases(uint8(InMaxLeases))
 {
 	check(MaxLeases == InMaxLeases); // field overflow
@@ -1398,8 +1478,8 @@ FConnectionPool::FConnectionPool(const FParams& Params)
 		Params.Host.Port,
 		Params.ConnectionCount
 	);
-	Internal->GetHost().SetBufferSize(FHost::EDirection::Send, Params.SendBufSize);
-	Internal->GetHost().SetBufferSize(FHost::EDirection::Recv, Params.RecvBufSize);
+	Internal->SetBufferSize(FHost::EDirection::Send, Params.SendBufSize);
+	Internal->SetBufferSize(FHost::EDirection::Recv, Params.RecvBufSize);
 
 	Ptr = Internal;
 }
@@ -1416,7 +1496,7 @@ FConnectionPool::~FConnectionPool()
 ////////////////////////////////////////////////////////////////////////////////
 bool FConnectionPool::Resolve()
 {
-	return (Ptr->GetHost().ResolveHostName() > 0);
+	return (Ptr->ResolveHostName().GetValue() > 0);
 }
 
 
@@ -1637,7 +1717,7 @@ static void Activity_Free(FActivity* Activity)
 			Socket = FSocket();
 		}
 
-		if (Activity->Pool->GetHost().GetIpAddress() > 0x00ff'ffff)
+		if (Activity->Pool->GetIpAddress() > FHost::InvalidIp)
 		{
 			Activity->Pool->ReturnLease(MoveTemp(Socket));
 		}
@@ -1979,7 +2059,7 @@ static uint64 ReadyCheck(FActivity** Activities, uint32 Num, uint32 TimeoutMs)
 			break;
 
 		case EWait::Pool:
-			if (Activity->Pool->GetHost().GetIpAddress() > 0x00ff'ffff)
+			if (Activity->Pool->GetIpAddress() > FHost::InvalidIp)
 			{
 				Activity_EndWait(Activities[i]);
 				Ret |= (1ull << Activity->Slot);
@@ -2037,7 +2117,7 @@ static int32 DoResolve(FActivity* Activity)
 	// needs to happen once, the first activity in can do the honours. Everyone
 	// else can wait.
 	FSocketPool* Pool = Activity->Pool;
-	int32 Result = Pool->GetHost().IsResolved();
+	int32 Result = Pool->IsResolved();
 	if (Result > 0)
 	{
 		Activity_ChangeState(Activity, FActivity::EState::Connect);
@@ -2052,7 +2132,7 @@ static int32 DoResolve(FActivity* Activity)
 	}
 
 	// We won! We WON!
-	Result = Pool->GetHost().ResolveHostName();
+	Result = Pool->ResolveHostName().GetValue();
 	switch (Result)
 	{
 	case 0:  Activity_SetError(Activity, "Unable to resolve host"); return -1;
@@ -2070,15 +2150,6 @@ static int32 DoConnect(FActivity* Activity)
 	TRACE_CPUPROFILER_EVENT_SCOPE(IasHttp::DoConnect);
 
 	FSocketPool* Pool = Activity->Pool;
-	check(Pool->GetHost().IsResolved() > 0);
-
-	uint32 IpAddress = Pool->GetHost().GetIpAddress();
-	uint32 Port = Pool->GetHost().GetPort();
-	if (IpAddress <= 0x00ff'ffff)
-	{
-		Activity_SetError(Activity, "Unresolved host");
-		return -1;
-	}
 
 	// Claim an existing socket from the pool.
 	FSocket Candidate;
@@ -2096,55 +2167,20 @@ static int32 DoConnect(FActivity* Activity)
 		return 0;
 	}
 
-	// Leased socket isn't valid so we'll create and connect one
-	if (!Candidate.Create())
-	{
-		Activity_SetError(Activity, "Failed to create socket");
-		return -1;
-	}
+	Trace(Activity, ETrace::Connect);
 
-	// Attempt a SOCKS connect
-	Trace(Activity, ETrace::Connect, IpAddress);
-	bool bSocksConnected = false;
-	if (int32 Result = MaybeConnectSocks(Candidate, IpAddress, Port); Result)
+	FResult Result = Pool->Connect(Candidate);
+	if (Result.GetValue() < 0)
 	{
-		if (Result < 0)
-		{
-			Activity_SetError(Activity, "Failed establishing SOCKS connection");
-			return -1;
-		}
-
-		bSocksConnected = true;
+		Activity_SetError(Activity, Result.GetMessage());
 	}
-
-	// Condition the socket
-	if (!Candidate.SetBlocking(false))
-	{
-		Activity_SetError(Activity, "Unable to set socket non-blocking");
-		return -1;
-	}
-
-	if (int32 OptValue = Pool->GetHost().GetBufferSize(FHost::EDirection::Send); OptValue >= 0)
-	{
-		Candidate.SetSendBufSize(OptValue);
-	}
-
-	if (int32 OptValue = Pool->GetHost().GetBufferSize(FHost::EDirection::Recv); OptValue >= 0)
-	{
-		Candidate.SetRecvBufSize(OptValue);
-	}
-
-	// Adjust socket send and recv buffer sizes
-	if (!bSocksConnected && !Candidate.Connect(IpAddress, Port))
-	{
-		Activity_SetError(Activity, "Socket connect failed");
-		return -1;
-	}
-
 	Activity->Socket = MoveTemp(Candidate);
 
 	Activity_ChangeState(Activity, FActivity::EState::Send);
-	Activity_BeginWait(Activity, FActivity::EWait::Write);
+	if (Result.GetValue() == 0)
+	{
+		Activity_BeginWait(Activity, FActivity::EWait::Write);
+	}
 	return 1;
 }
 
@@ -2667,7 +2703,7 @@ FRequest FEventLoop::FImpl::Request(
 	FMessageBuilder Builder(Activity->Buffer);
 
 	Builder << Method << " " << Path << " HTTP/1.1" "\r\n"
-		"Host: " << Activity->Pool->GetHost().GetHostName() << "\r\n";
+		"Host: " << Activity->Pool->GetHostName() << "\r\n";
 
 	// HTTP/1.1 is persistent by default thus "Connection" header isn't required
 	if (!Activity->IsKeepAlive)
@@ -3121,6 +3157,16 @@ static void MiscTest()
 	check(UrlOut.Port.Get(Url) == "999");
 	check(UrlOut.Path == 25);
 #undef CRLF
+
+	check(FResult(-5).GetValue()		== -5);
+	check(FResult(-1).GetValue()		== -1);
+	check(FResult( 0).GetValue()		==  0);
+	check(FResult( 1).GetValue()		==  1);
+	check(FResult(19).GetValue()		== 19);
+	check(FResult(0xffff).GetValue()	== -1);
+	check(FResult(1, "yes").GetValue()	==  1);
+	check(FResult(0, "?").GetValue()	==  0);
+	check(FResult(-1, "no").GetValue()	== -1);
 }
 
 ////////////////////////////////////////////////////////////////////////////////
