Change IAS cache to use read/cancel ops on its IO thread

#rnx
(git:2acec25 - Martin Ridgers - [REDACTED_PHONE]:09:39 +0200)

diff --git a/Engine/Source/Runtime/Experimental/IoStoreOnDemand/Private/IasCache.h b/Engine/Source/Runtime/Experimental/IoStoreOnDemand/Private/IasCache.h
--- a/Engine/Source/Runtime/Experimental/IoStoreOnDemand/Private/IasCache.h
+++ b/Engine/Source/Runtime/Experimental/IoStoreOnDemand/Private/IasCache.h
@@ -19,9 +19,6 @@ namespace UE::IO::IAS {
 class IIasCache
 {
 public:
-	using FGetToken = UPTRINT;
-	using FGetWork = UE::Tasks::TTask<TIoStatusOr<FIoBuffer>>;
-
 	virtual ~IIasCache() = default;
 
 	/** Deletes the IAS object, dropping all data persisted to disk and releasing
@@ -34,15 +31,20 @@ public:
 	virtual bool ContainsChunk(const FIoHash& Key) const = 0;
 
 	/** Get the chunk associated with the specified cache key. If the data is
-	 already in memory it is return in OutData. Otherwise a FGetToken value is
-	 returned; zero if Key is not found, or non-zero for use with Materialize. */
-	virtual FGetToken Get(const FIoHash& Key, FIoBuffer& OutData) = 0;
-
-	/** Materialize the data for a Get() if it was not immediately available */
-	virtual FGetWork Materialize(
-		FGetToken Token,
-		const FIoReadOptions& Options,
-		const FIoCancellationToken* CancellationToken) = 0;
+	 already in memory it is return in OutData. Otherwise the returned status
+	 indicates if the key can be materialized or if it does not exist */
+	virtual EIoErrorCode Get(const FIoHash& Key, FIoBuffer& OutData) = 0;
+
+	/** Materializes a cached items data from disk. The data is read into Dest
+	 so Dest must remain valid throughout. The result of the disk read is returned
+	 in Status (same lifetime needs as Dest). DoneEvent is triggered when the read
+	 succeeds, is not found, or if an IO error occurred. */
+	virtual void Materialize(const FIoHash& Key, FIoBuffer& Dest, EIoErrorCode& Status, UE::Tasks::FTaskEvent DoneEvent) = 0;
+
+	/** Cancels a previously request to Materialize(). Note that the materialize
+	 already can still complete as the read operation may already be in flight. If
+	 the cancel succeeds, the materialize's DoneEvent is not triggered. */
+	virtual void Cancel(FIoBuffer& GivenDest) = 0;
 
 	/** Insert a new chunk into the cache. */
 	virtual FIoStatus Put(const FIoHash& Key, FIoBuffer& Data) = 0;
diff --git a/Engine/Source/Runtime/Experimental/IoStoreOnDemand/Private/JournaledCache.cpp b/Engine/Source/Runtime/Experimental/IoStoreOnDemand/Private/JournaledCache.cpp
--- a/Engine/Source/Runtime/Experimental/IoStoreOnDemand/Private/JournaledCache.cpp
+++ b/Engine/Source/Runtime/Experimental/IoStoreOnDemand/Private/JournaledCache.cpp
@@ -1864,8 +1864,9 @@ namespace UE::IO::IAS {
 // {{{1 journaled-cache ........................................................
 
 ////////////////////////////////////////////////////////////////////////////////
-class FJournaledCache
+class FJournaledCache final
 	: public IIasCache
+	, public JournaledCache::FServiceThread::FReadSink
 {
 public:
 								FJournaledCache() = default;
@@ -1873,15 +1874,25 @@ public:
 	bool						Initialize(const TCHAR* RootDir, const FIasCacheConfig& Config);
 	virtual void				Abandon() override;
 	virtual bool				ContainsChunk(const FIoHash& Key) const override;
-	virtual FGetToken			Get(const FIoHash& Key, FIoBuffer& OutData) override;
-	virtual FGetWork			Materialize(FGetToken Token, const FIoReadOptions& Options, const FIoCancellationToken* CancelToken) override;
+	virtual EIoErrorCode		Get(const FIoHash& Key, FIoBuffer& OutData) override;
+	virtual void				Materialize(const FIoHash& Key, FIoBuffer& Dest, EIoErrorCode& Status, UE::Tasks::FTaskEvent DoneEvent) override;
+	virtual void				Cancel(FIoBuffer& GivenDest) override;
 	virtual FIoStatus			Put(const FIoHash& Key, FIoBuffer& Data) override;
 
 private:
+	struct FMaterialOp
+	{
+		UE::Tasks::FTaskEvent	DoneEvent;
+		EIoErrorCode*			Status;
+	};
+
 	void						Update();
+	virtual void				OnRead(const FReadResult* Results, int32 Num) override;
 	static uint64				ReduceKey(const FIoHash& Key);
 	JournaledCache::FCache*		Cache = nullptr;
 	UE::Tasks::FPipe			GetPipe = UE::Tasks::FPipe(TEXT("IasCacheGetPipe"));
+	FCriticalSection			Lock;
+	TMap<int32, FMaterialOp>	PendingMaterializes;
 };
 
 ////////////////////////////////////////////////////////////////////////////////
@@ -1958,57 +1969,66 @@ bool FJournaledCache::ContainsChunk(const FIoHash& Key) const
 }
 
 ////////////////////////////////////////////////////////////////////////////////
-FJournaledCache::FGetToken FJournaledCache::Get(const FIoHash& Key, FIoBuffer& OutData)
+EIoErrorCode FJournaledCache::Get(const FIoHash& Key, FIoBuffer& OutData)
 {
+	using namespace JournaledCache;
+
+	check(OutData.GetData() == nullptr);
+
 	uint64 InnerKey = ReduceKey(Key);
 
 	auto Entry = Cache->Get(InnerKey, OutData);
 	if (Entry == 0)
 	{
-		return 0;
+		return EIoErrorCode::NotFound;
 	}
+	
+	return (OutData.GetData() != nullptr) ? EIoErrorCode::Ok : EIoErrorCode::FileNotOpen;
+}
+
+////////////////////////////////////////////////////////////////////////////////
+void FJournaledCache::Materialize(
+	const FIoHash& Key,
+	FIoBuffer& Dest,
+	EIoErrorCode& Status,
+	UE::Tasks::FTaskEvent DoneEvent)
+{
+	using namespace JournaledCache;
+
+	FServiceThread::FReadRequest Request = {
+		.Key = ReduceKey(Key),
+		.Dest = &Dest,
+		.Sink = this,
+	};
+	uint32 ReadId = FServiceThread::Get().BeginRead(Cache, Request);
 
-	static_assert(sizeof(InnerKey) == sizeof(FGetToken));
-	return FGetToken(InnerKey);
+	FScopeLock _(&Lock);
+	PendingMaterializes.Add(ReadId, { MoveTemp(DoneEvent), &Status });
 }
 
 ////////////////////////////////////////////////////////////////////////////////
-FJournaledCache::FGetWork FJournaledCache::Materialize(
-	FGetToken Token,
-	const FIoReadOptions& Options,
-	const FIoCancellationToken* CancelToken)
+void FJournaledCache::Cancel(FIoBuffer& GivenDest)
 {
-	return GetPipe.Launch(TEXT("IasCacheGet"), [this, Token, Options, CancelToken] ()
-	{
-		LLM_SCOPE_BYTAG(Ias);
+	JournaledCache::FServiceThread::Get().CancelRead(&GivenDest);
+}
 
-		if (CancelToken != nullptr && CancelToken->IsCancelled())
-		{
-			return TIoStatusOr<FIoBuffer>(FIoStatus(EIoErrorCode::Cancelled));
-		}
+////////////////////////////////////////////////////////////////////////////////
+void FJournaledCache::OnRead(const FReadResult* Results, int32 Num)
+{
+	FScopeLock _(&Lock);
 
-		FIoBuffer Buffer;
+	for (; Num-- > 0; ++Results)
+	{
+		const FReadResult& Result = Results[0];
 
-		int64 Size = Options.GetSize();
-		if (void* DestAddr = Options.GetTargetVa(); DestAddr != nullptr)
-		{
-			check(Size > 0);
-			Buffer = FIoBuffer(FIoBuffer::Wrap, DestAddr, Size);
-		}
-		else if (Size >= 0)
-		{
-			Buffer = FIoBuffer(Size);
-		}
+		FMaterialOp* Pend = PendingMaterializes.Find(Result.ReadId);
+		check(Pend != nullptr);
 
-		uint32 Offset = uint32(Options.GetOffset());
-		EIoErrorCode Ret = Cache->Materialize(Token, Buffer, Offset);
-		if (Ret != EIoErrorCode::Ok)
-		{
-			return TIoStatusOr<FIoBuffer>(FIoStatus(Ret));
-		}
+		Pend->Status[0] = EIoErrorCode(Result.Status);
+		Pend->DoneEvent.Trigger();
 
-		return TIoStatusOr<FIoBuffer>(Buffer);
-	});
+		PendingMaterializes.Remove(Result.ReadId);
+	}
 }
 
 ////////////////////////////////////////////////////////////////////////////////
diff --git a/Engine/Source/Runtime/Experimental/IoStoreOnDemand/Private/OnDemandIoDispatcherBackend.cpp b/Engine/Source/Runtime/Experimental/IoStoreOnDemand/Private/OnDemandIoDispatcherBackend.cpp
--- a/Engine/Source/Runtime/Experimental/IoStoreOnDemand/Private/OnDemandIoDispatcherBackend.cpp
+++ b/Engine/Source/Runtime/Experimental/IoStoreOnDemand/Private/OnDemandIoDispatcherBackend.cpp
@@ -762,14 +762,14 @@ struct FChunkRequest
 	FIoRequestImpl* RequestHead;
 	FIoRequestImpl* RequestTail;
 	FIoBuffer Chunk;
-	UE::Tasks::TTask<TIoStatusOr<FIoBuffer>> CacheTask;
 	UE::Tasks::FTask DecodeTask;
-	FIoCancellationToken CancellationToken;
 	uint64 StartTime;
 	int32 Priority;
 	uint16 RequestCount;
 	uint16 HttpRetryCount;
 	bool bCached;
+	bool bCancelled = false;
+	EIoErrorCode CacheGetStatus;
 };
 
 ///////////////////////////////////////////////////////////////////////////////
@@ -1227,7 +1227,7 @@ class FOnDemandIoBackend final
 			if (FChunkRequest** InflightRequest = Inflight.Find(Params.ChunkKey))
 			{
 				FChunkRequest* ChunkRequest = *InflightRequest;
-				check(!ChunkRequest->CancellationToken.IsCancelled());
+				check(!ChunkRequest->bCancelled);
 				bOutPending = true;
 				bOutUpdatePriority = ChunkRequest->AddDispatcherRequest(Request);
 
@@ -1241,7 +1241,7 @@ class FOnDemandIoBackend final
 			return ChunkRequest;
 		}
 
-		bool Cancel(FIoRequestImpl* Request)
+		bool Cancel(FIoRequestImpl* Request, IIasCache* TheCache)
 		{
 			FScopeLock _(&Mutex);
 
@@ -1257,7 +1257,11 @@ class FOnDemandIoBackend final
 
 				if (RemainingCount == 0)
 				{
-					ChunkRequest.CancellationToken.Cancel();
+					ChunkRequest.bCancelled = true;
+					if (TheCache != nullptr)
+					{
+						TheCache->Cancel(ChunkRequest.Chunk);
+					}
 					Inflight.Remove(BackendData.ChunkKey);
 				}
 
@@ -1342,6 +1346,7 @@ private:
 	FString GetEndpointTestPath() const;
 	void ConditionallyStartBackendThread();
 	void CompleteRequest(FChunkRequest* ChunkRequest);
+	void CompleteMaterialize(FChunkRequest* ChunkRequest);
 
 	FIoStatus ApplyGeneratedOnDemandToc(const FString& CdnUrl, const FString& TocPath);
 	FIoStatus DownloadoadOnDemandToc(const FString& CdnUrl, const FString& TocPath);
@@ -1358,6 +1363,7 @@ private:
 	FChunkRequests ChunkRequests;
 	FIoRequestQueue CompletedRequests;
 	FChunkRequestQueue HttpRequests;
+	TArray<FChunkRequest*> PendingCacheGets;
 	FOnDemandIoBackendStats Stats;
 	FBackendStatus BackendStatus;
 	FAvailableEps AvailableEps;
@@ -1498,9 +1504,8 @@ void FOnDemandIoBackend::CompleteRequest(FChunkRequest* ChunkRequest)
 	LLM_SCOPE_BYTAG(Ias);
 	TRACE_CPUPROFILER_EVENT_SCOPE(IasBackend::CompleteRequest);
 	check(ChunkRequest != nullptr);
-	const bool bCancelled = ChunkRequest->CancellationToken.IsCancelled();
 
-	if (bCancelled)
+	if (ChunkRequest->bCancelled)
 	{
 		check(ChunkRequest->RequestHead == nullptr);
 		check(ChunkRequest->RequestTail == nullptr);
@@ -1570,6 +1575,43 @@ void FOnDemandIoBackend::CompleteRequest(FChunkRequest* ChunkRequest)
 	}
 }
 
+void FOnDemandIoBackend::CompleteMaterialize(FChunkRequest* ChunkRequest)
+{
+	TRACE_CPUPROFILER_EVENT_SCOPE(IasBackend::CompleteMaterialize);
+
+	bool bWasCancelled = false;
+	switch (ChunkRequest->CacheGetStatus)
+	{
+		case EIoErrorCode::Ok:
+		check(ChunkRequest->Chunk.GetData() != nullptr);
+		ChunkRequest->bCached = true;
+		CompleteRequest(ChunkRequest);
+		return;
+
+	case EIoErrorCode::ReadError:
+		FOnDemandIoBackendStats::Get()->OnCacheError();
+		break;
+
+	case EIoErrorCode::Cancelled:
+		bWasCancelled = true;
+		break;
+
+	case EIoErrorCode::NotFound:
+		break;
+	}
+
+	if (bWasCancelled || BackendStatus.IsHttpEnabled() == false)
+	{
+		UE_CLOG(BackendStatus.IsHttpEnabled() == false, LogIas, Log, TEXT("Chunk was not found in the cache and HTTP is disabled"));
+		CompleteRequest(ChunkRequest);
+		return;
+	}
+
+	Stats.OnHttpEnqueue();
+	HttpRequests.EnqueueByPriority(ChunkRequest);
+	TickBackendEvent->Trigger();
+}
+
 bool FOnDemandIoBackend::Resolve(FIoRequestImpl* Request)
 {
 	using namespace UE::Tasks;
@@ -1611,65 +1653,48 @@ bool FOnDemandIoBackend::Resolve(FIoRequestImpl* Request)
 	{
 		const FIoHash& Key = ChunkRequest->Params.ChunkKey;
 		FIoBuffer& Buffer = ChunkRequest->Chunk;
-		IIasCache::FGetToken GetToken = Cache->Get(Key, Buffer);
 
-		if (Buffer.GetData() != nullptr)
+		//TODO: Pass priority to cache
+		EIoErrorCode GetStatus = Cache->Get(Key, Buffer);
+
+		if (GetStatus == EIoErrorCode::Ok)
 		{
+			check(Buffer.GetData() != nullptr);
 			ChunkRequest->bCached = true;
-			Launch(UE_SOURCE_LOCATION, [this, ChunkRequest] { CompleteRequest(ChunkRequest); });
+			Launch(UE_SOURCE_LOCATION, [this, ChunkRequest] {
+				CompleteRequest(ChunkRequest);
+			});
 			return true;
 		}
 
-		if (GetToken == 0)
+		if (GetStatus == EIoErrorCode::FileNotOpen)
 		{
-			Stats.OnHttpEnqueue();
-			HttpRequests.EnqueueByPriority(ChunkRequest);
-			TickBackendEvent->Trigger();
-			return true;
-		}
+			InflightCacheRequestCount.fetch_add(1, std::memory_order_relaxed);
 
-		//TODO: Pass priority to cache
-		InflightCacheRequestCount.fetch_add(1, std::memory_order_relaxed);
-		ChunkRequest->CacheTask = Cache->Materialize(GetToken, FIoReadOptions(), &ChunkRequest->CancellationToken);
-	}
+			FTaskEvent OnReadyEvent(TEXT("IasCacheMaterializeDone"));
 
-	const ETaskPriority TaskPriority = ChunkRequest->Priority > IoDispatcherPriority_Medium ? ETaskPriority::High : ETaskPriority::Normal;
-	Launch(UE_SOURCE_LOCATION, [this, ChunkRequest]()
-	{
-		LLM_SCOPE_BYTAG(Ias);
-		TRACE_CPUPROFILER_EVENT_SCOPE(FOnDemandIoBackend::CompleteOrEnqueueHttpRequest);
-		if (ChunkRequest->CacheTask.IsValid())
-		{
-			InflightCacheRequestCount.fetch_sub(1, std::memory_order_relaxed);
-			if (TIoStatusOr<FIoBuffer> Status = ChunkRequest->CacheTask.GetResult(); Status.IsOk())
-			{
-				ChunkRequest->Chunk = Status.ConsumeValueOrDie();
-				ChunkRequest->bCached = true;
-				return CompleteRequest(ChunkRequest);
-			}
-			else if (Status.Status().GetErrorCode() == EIoErrorCode::ReadError)
-			{
-				FOnDemandIoBackendStats::Get()->OnCacheError();
-			}
-		}
+			Launch(UE_SOURCE_LOCATION, [this, ChunkRequest] {
+				InflightCacheRequestCount.fetch_sub(1, std::memory_order_relaxed);
+				CompleteMaterialize(ChunkRequest);
+			}, OnReadyEvent);
 
-		if (ChunkRequest->CancellationToken.IsCancelled() || BackendStatus.IsHttpEnabled() == false)
-		{
-			UE_CLOG(BackendStatus.IsHttpEnabled() == false, LogIas, Log, TEXT("Chunk was not found in the cache and HTTP is disabled"));
-			return CompleteRequest(ChunkRequest);
+			EIoErrorCode& OutStatus = ChunkRequest->CacheGetStatus;
+			Cache->Materialize(Key, Buffer, OutStatus, MoveTemp(OnReadyEvent));
+			return true;
 		}
 
-		Stats.OnHttpEnqueue();
-		HttpRequests.EnqueueByPriority(ChunkRequest);
-		TickBackendEvent->Trigger();
-	}, ChunkRequest->CacheTask, TaskPriority);
+		check(GetStatus == EIoErrorCode::NotFound);
+	}
 
+	Stats.OnHttpEnqueue();
+	HttpRequests.EnqueueByPriority(ChunkRequest);
+	TickBackendEvent->Trigger();
 	return true;
 }
 
 void FOnDemandIoBackend::CancelIoRequest(FIoRequestImpl* Request)
 {
-	if (ChunkRequests.Cancel(Request))
+	if (ChunkRequests.Cancel(Request, Cache.Get()))
 	{
 		CompletedRequests.Enqueue(Request);
 		BackendContext->WakeUpDispatcherThreadDelegate.Execute();
