Convex optimization LODs

diff --git a/Engine/Source/Runtime/Experimental/Chaos/Private/Chaos/ConvexOptimizer.cpp b/Engine/Source/Runtime/Experimental/Chaos/Private/Chaos/ConvexOptimizer.cpp
--- a/Engine/Source/Runtime/Experimental/Chaos/Private/Chaos/ConvexOptimizer.cpp
+++ b/Engine/Source/Runtime/Experimental/Chaos/Private/Chaos/ConvexOptimizer.cpp
@@ -1,53 +1,263 @@
 // Copyright Epic Games, Inc. All Rights Reserved.
 
-
 #include "Chaos/ConvexOptimizer.h"
 #include "Chaos/Tribox.h"
+#include "ChaosStats.h"
 #include "Chaos/ImplicitObjectBVH.h"
+#include "Chaos/ImplicitObjectUnion.h"
 #include "Chaos/PhysicsObjectCollisionInterface.h"
-#include "Chaos/Collision/CollisionContext.h"
 
 namespace Chaos
 {
-	namespace CVars
-	{
-		// Replace all the convexes within an implicit hierarchy with a simplified one (kdop18 tribox for now) for collision
-		bool bChaosSimplifyConvexes = false;
-		FAutoConsoleVariableRef CVarChaosSimplifyConvexes(TEXT("p.Chaos.Collision.SimplifyConvexes"), bChaosSimplifyConvexes, TEXT("If true replace all the convexes within an implcit hierarchy with a simplified one (kdop18 tribox for now) for collision"));
+namespace CVars
+{
+	// Replace all the convexes within an implicit hierarchy with a simplified one (kdop18 tribox for now) for collision
+	bool bChaosConvexSimplifyUnion = false;
+	FAutoConsoleVariableRef CVarChaosSimplifyUnion(TEXT("p.Chaos.Convex.SimplifyUnion"), bChaosConvexSimplifyUnion, TEXT("If true replace all the convexes within an implcit hierarchy with a simplified one (kdop18 tribox for now) for collision"));
 
-		// Num LODs for the simplification
-		int32 ChaosNumConvexLODs = 0;
-		FAutoConsoleVariableRef CVarChaosNumConvexLODs(TEXT("p.Chaos.Collision.NumConvexLODs"), ChaosNumConvexLODs, TEXT("Max number of collisions LODs required during simplication"));
-	}
+	// Max number of convex LODs used during simplication
+	int32 ChaosConvexMaxLODs = 4;
+	FAutoConsoleVariableRef CVarChaosConvexMaxLODs(TEXT("p.Chaos.Convex.MaxLODs"), ChaosConvexMaxLODs, TEXT("Max number of convex LODs used during simplication"));
+	
+	// Tribox volume / convex hull threshold to trigger a volume splitting during tree construction
+	float ChaosConvexSplittingThreshold = 1.2f;
+	FAutoConsoleVariableRef CVarChaosConvexSplittingThreshold(TEXT("p.Chaos.Convex.SplittingThreshold"), ChaosConvexSplittingThreshold, TEXT("Tribox volume / convex hull threshold to trigger a volume splitting during tree construction"));
+
+	// Min volume of the simplified convexes
+	float ChaosConvexMinVolume = 0.2f;
+	FAutoConsoleVariableRef CVarChaosConvexMinVolume(TEXT("p.Chaos.Convex.MinVolume"), ChaosConvexMinVolume, TEXT("Min volume of the simplified convexes"));
+	
+	extern int32 ChaosUnionBVHMaxDepth;
+	extern int32 ChaosUnionBVHMinShapes;
+}
 
 namespace Private
 {
+/**
+ * @brief Simple tribox node to store tribox hierarchy information
+ */
+struct FTriboxNode
+{
+	// Constructor with a given node tribox and the underling objects range 
+	FTriboxNode(const FTribox& NodeBounds, const TArray<FTribox>& NodeObjects) : BoundsTribox(NodeBounds), ObjectTriboxes(NodeObjects)
+	{}
 
-	FConvexOptimizer::FConvexOptimizer() :
-		SimplifiedConvexes(), CollisionObjects(MakeUnique<Private::FCollisionObjects>()), ShapesArray()
+	// Base constructor
+	FTriboxNode() : BoundsTribox(), ObjectTriboxes()
 	{}
 
-	FConvexOptimizer::~FConvexOptimizer() = default;
+	// Get the closest plane to the node center 
+	void FindClosestPlane(int32& PlaneAxis, FTribox::FRealType& PlaneProjection) const;
+
+	// Split the object triboxes into a left and a right tribox
+	void SplitTriboxNode(const int32& PlaneAxis, const FTribox::FRealType& PlaneProjection,
+							FTriboxNode& LeftTribox, FTribox::FRealType& LeftVolume,
+							FTriboxNode& RightTribox, FTribox::FRealType& RightVolume) const;
+	
+	// Enclosing node tribox 
+	FTribox BoundsTribox;
+
+	// List of objects that are part of that node
+	TArray<FTribox> ObjectTriboxes;
+};
 
-	void FConvexOptimizer::VisitCollisionObjects(const FImplicitHierarchyVisitor& VisitorFunc) const
+FORCEINLINE void FTriboxNode::FindClosestPlane(int32& PlaneAxis, FTribox::FRealType& PlaneProjection) const
+{
+	const FTribox::FVec3Type PointPosition = BoundsTribox.GetCenter();
+	
+	int32 LocalAxis = INDEX_NONE;
+	FTribox::FRealType LocalProjection = 0.0;
+	FTribox::FRealType MinDistance = FLT_MAX;
+			
+	for(const FTribox& ObjectTribox : ObjectTriboxes)
 	{
-		if(CVars::bChaosSimplifyConvexes)
+		const FTribox::FRealType ClosestDistance = ObjectTribox.GetClosestPlane(PointPosition, LocalAxis, LocalProjection);
+		if(ClosestDistance < MinDistance)
 		{
-			int32 ObjectIndex = 0;
-			for(const Private::FImplicitBVHObject& CollisionObject : CollisionObjects->ImplicitObjects)
+			MinDistance = ClosestDistance;
+			PlaneAxis = LocalAxis;
+			PlaneProjection = LocalProjection;
+		}
+	}
+}
+	
+FORCEINLINE void FTriboxNode::SplitTriboxNode(const int32& PlaneAxis, const FTribox::FRealType& PlaneProjection,
+	FTriboxNode& LeftTribox, FTribox::FRealType& LeftVolume, FTriboxNode& RightTribox,  FTribox::FRealType& RightVolume) const
+{
+	for(const FTribox& ObjectTribox : ObjectTriboxes)
+	{
+		FTribox LocalLeftTribox, LocalRightTribox;
+		ObjectTribox.SplitTriboxSlab(PlaneAxis, PlaneProjection, LocalLeftTribox, LocalRightTribox);
+				
+		if(LocalLeftTribox.HasDatas())
+		{
+			LeftTribox.BoundsTribox += LocalLeftTribox;
+			LeftTribox.ObjectTriboxes.Add(LocalLeftTribox);
+			LeftVolume += LocalLeftTribox.ComputeVolume();
+		}
+		if(LocalRightTribox.HasDatas())
+		{
+			RightTribox.BoundsTribox += LocalRightTribox;
+			RightTribox.ObjectTriboxes.Add(LocalRightTribox);
+			RightVolume += LocalRightTribox.ComputeVolume();
+		}
+	}
+}
+
+/**
+ * @brief Simple tribox binary tree to store tribox hierarchy
+ */
+struct FTriboxTree
+{
+	// Constructor given a LOD and a number of objects
+	FTriboxTree(const int32 NumLODs, const FConvexOptimizer::FCachedTriboxes& RootTriboxes);
+	
+	// Build all the tree LODs
+	void BuildTreeLODs(TArray<FImplicitObjectPtr>& SimplifiedConvexes, const TUniquePtr<Private::FCollisionObjects>& CollisionObjects);
+
+	// Add a node to the tree nodes
+	void AddTreeNode(const int32 LODIndex, const FTriboxNode& TriboxNode, const FTribox::FRealType& ObjectsVolume,
+		TArray<FImplicitObjectPtr>& SimplifiedConvexes, const TUniquePtr<Private::FCollisionObjects>& CollisionObjects);
+	
+	// Swap the 2 buffers
+	void SwapNodeBuffers()
+	{
+		BufferIndex = 1 - BufferIndex;
+		TreeNodes[BufferIndex].Reset();
+	}
+
+	// Current buffer used to build tthe tree
+	int32 CurrentBuffer() const {return BufferIndex;}
+
+	// Previous buffer used to build the tree
+	int32 PreviousBuffer() const {return 1-BufferIndex;}
+	
+	// List of tree nodes 
+	TArray<FTriboxNode> TreeNodes[2];
+
+	// Tree buffer index 
+	int32 BufferIndex = 0;
+};
+	
+FORCEINLINE FTriboxTree::FTriboxTree(const int32 NumLODs, const FConvexOptimizer::FCachedTriboxes& RootTriboxes) : TreeNodes(), BufferIndex(0)
+{
+	const int32 NumNodes = 1 << NumLODs;
+	
+	TreeNodes[0].Reserve(NumNodes);
+	TreeNodes[1].Reserve(NumNodes);
+	
+	FTriboxNode TriboxNode;
+	for(auto& RootTribox : RootTriboxes)
+	{
+		if(RootTribox.Value.Tribox.HasDatas())
+		{
+			TriboxNode.BoundsTribox += RootTribox.Value.Tribox;
+			TriboxNode.ObjectTriboxes.Add(RootTribox.Value.Tribox);
+		}
+	}
+	TreeNodes[CurrentBuffer()].Emplace(TriboxNode);
+}
+
+FORCEINLINE void FTriboxTree::BuildTreeLODs(TArray<FImplicitObjectPtr>& SimplifiedConvexes, const TUniquePtr<Private::FCollisionObjects>& CollisionObjects)
+{
+	//UE_LOG(LogTemp, Log, TEXT("	Building tribox binary tree"))
+	for(int32 LODIndex = 1; LODIndex < CVars::ChaosConvexMaxLODs; ++LODIndex)
+	{
+		// Swap the node buffers (working/result)
+		SwapNodeBuffers();
+
+		//int32 NodeIndex = 0;
+		for(FTriboxNode& TriboxNode : TreeNodes[PreviousBuffer()])
+		{
+			//UE_LOG(LogTemp, Log, TEXT("			Iterating over node %d with num objects = %d"), NodeIndex++, TriboxNode.ObjectTriboxes.Num())
+
+			int32 PlaneAxis = INDEX_NONE;
+			FTribox::FRealType PlaneProjection = 0.0;
+
+			// Find the closest plane to the node center
+			TriboxNode.FindClosestPlane(PlaneAxis, PlaneProjection);
+					
+			FTriboxNode LeftTribox, RightTribox;
+			FTribox::FRealType LeftVolume = 0.0, RightVolume = 0.0;
+
+			// Split the node along the cutting plane
+			TriboxNode.SplitTriboxNode(PlaneAxis, PlaneProjection, LeftTribox, LeftVolume, RightTribox, RightVolume);
+
+			// Add left node to the tree
+			AddTreeNode(LODIndex, LeftTribox, LeftVolume, SimplifiedConvexes, CollisionObjects);
+
+			// Add right node to the tree
+			AddTreeNode(LODIndex, RightTribox, RightVolume, SimplifiedConvexes, CollisionObjects);
+		}
+	}
+}
+
+FORCEINLINE void FTriboxTree::AddTreeNode(const int32 LODIndex, const FTriboxNode& TriboxNode, const FTribox::FRealType& ObjectsVolume,
+	TArray<FImplicitObjectPtr>& SimplifiedConvexes, const TUniquePtr<Private::FCollisionObjects>& CollisionObjects)
+{
+	if(TriboxNode.BoundsTribox.HasDatas())
+	{
+		// Disgard all the objects with small volumes
+		if(ObjectsVolume > CVars::ChaosConvexMinVolume)
+		{
+			const FTribox::FRealType NodeVolume = TriboxNode.BoundsTribox.ComputeVolume();
+			//UE_LOG(LogTemp, Log, TEXT("				Node volume = [%f %f] -> %f"), NodeVolume, ObjectsVolume, NodeVolume/ObjectsVolume)
+			
+			if((NodeVolume/ObjectsVolume > CVars::ChaosConvexSplittingThreshold) &&
+				(LODIndex < (CVars::ChaosConvexMaxLODs-1)))
 			{
-				int32 LeafObjectIndex = CollisionObject.GetObjectIndex();
-				CollisionObject.GetGeometry()->VisitLeafObjectsImpl(CollisionObject.GetTransform(),
-					CollisionObject.GetRootObjectIndex(), ObjectIndex, LeafObjectIndex, VisitorFunc);
+				// If volume ratio (concavity) is big enough add this node to the current list to be processed by the next LOD
+				TreeNodes[CurrentBuffer()].Add(TriboxNode);
 			}
-		}	
+			else
+			{
+				// Otherwise add the node to the final convexes list
+				SimplifiedConvexes.Add(TriboxNode.BoundsTribox.MakeConvex());
+				Private::FImplicitBVH::CollectLeafObject(SimplifiedConvexes.Last(),
+						FRigidTransform3::Identity, INDEX_NONE, CollisionObjects->ImplicitObjects);
+			}
+		}
 	}
+}
+
+FConvexOptimizer::FConvexOptimizer() :
+	SimplifiedConvexes(), CollisionObjects(MakeUnique<Private::FCollisionObjects>()), ShapesArray()
+{}
+
+FConvexOptimizer::~FConvexOptimizer() = default;
 
-	void FConvexOptimizer::VisitOverlappingObjects(const FAABB3& LocalBounds, const FImplicitHierarchyVisitor& VisitorFunc) const
+void FConvexOptimizer::VisitCollisionObjects(const FImplicitHierarchyVisitor& VisitorFunc) const
+{
+	if(CVars::bChaosConvexSimplifyUnion)
 	{
-		if(CVars::bChaosSimplifyConvexes)
+		const TArray<Private::FImplicitBVHObject>& ImplicitObjects = BVH.IsValid() ? BVH->GetObjects() : CollisionObjects->ImplicitObjects;
+		
+		int32 ObjectIndex = 0;
+		for(const Private::FImplicitBVHObject& CollisionObject : ImplicitObjects)
+		{
+			int32 LeafObjectIndex = CollisionObject.GetObjectIndex();
+			CollisionObject.GetGeometry()->VisitLeafObjectsImpl(CollisionObject.GetTransform(),
+				CollisionObject.GetRootObjectIndex(), ObjectIndex, LeafObjectIndex, VisitorFunc);
+		}
+	}	
+}
+
+void FConvexOptimizer::VisitOverlappingObjects(const FAABB3& LocalBounds, const FImplicitHierarchyVisitor& VisitorFunc) const
+{
+	if(CVars::bChaosConvexSimplifyUnion)
+	{
+		int32 ObjectIndex = 0;
+		if(BVH.IsValid())
+		{
+			BVH->VisitAllIntersections(LocalBounds,
+			[this, &ObjectIndex, &VisitorFunc](const FImplicitObject* Implicit, const FRigidTransform3f& RelativeTransformf,
+				const FAABB3f& RelativeBoundsf, const int32 RootObjectIndex, const int32 LeafObjectIndex)
+			{
+				VisitorFunc(Implicit, FRigidTransform3(RelativeTransformf), RootObjectIndex, ObjectIndex++, LeafObjectIndex);
+			});
+		}
+		else
 		{
-			int32 ObjectIndex = 0;
 			for(const Private::FImplicitBVHObject& CollisionObject : CollisionObjects->ImplicitObjects)
 			{
 				int32 LeafObjectIndex = CollisionObject.GetObjectIndex();
@@ -56,111 +266,241 @@ namespace Private
 			}
 		}
 	}
+}
 
-	void FConvexOptimizer::BuildConvexShapes(const FShapesArray& UnionShapes)
-	{
-		ShapesArray.Reset();
-		TArray<FImplicitObjectPtr> ImplicitObjects = SimplifiedConvexes;
-		ShapesArray.Add(FShapeInstance::Make(UnionShapes.Num(),
-			MakeImplicitObjectPtr<FImplicitObjectUnion>(MoveTemp(ImplicitObjects))));
+DECLARE_CYCLE_STAT(TEXT("Collisions::BuildConvexShapes"), STAT_BuildConvexShapes, STATGROUP_ChaosCollision);
+void FConvexOptimizer::BuildConvexShapes(const FShapesArray& UnionShapes)
+{
+	SCOPE_CYCLE_COUNTER(STAT_BuildConvexShapes);
+	
+	ShapesArray.Reset();
+	TArray<FImplicitObjectPtr> ImplicitObjects = SimplifiedConvexes;
+	ShapesArray.Add(FShapeInstance::Make(UnionShapes.Num(),
+		MakeImplicitObjectPtr<FImplicitObjectUnion>(MoveTemp(ImplicitObjects))));
 
-		const TUniquePtr<FPerShapeData>& TemplateData = UnionShapes[0];
-		const TUniquePtr<FShapeInstance>& ShapeData = ShapesArray.Last();
+	const TUniquePtr<FPerShapeData>& TemplateData = UnionShapes[0];
+	const TUniquePtr<FShapeInstance>& ShapeData = ShapesArray.Last();
 
-		ShapeData->SetQueryData(TemplateData->GetQueryData());
-		ShapeData->SetSimData(TemplateData->GetSimData());
-		ShapeData->SetCollisionTraceType(TemplateData->GetCollisionTraceType());
+	ShapeData->SetQueryData(TemplateData->GetQueryData());
+	ShapeData->SetSimData(TemplateData->GetSimData());
+	ShapeData->SetCollisionTraceType(TemplateData->GetCollisionTraceType());
 
-		// Only sim enabled if all the underlying shapes could be used for physics
-		ShapeData->SetSimEnabled(true);
+	// Only sim enabled if all the underlying shapes could be used for physics
+	ShapeData->SetSimEnabled(true);
 
-		// Disable simple shapes for query 
-		ShapeData->SetQueryEnabled(false);
-	}
+	// Disable simple shapes for query 
+	ShapeData->SetQueryEnabled(false);
+}
 
-	DECLARE_CYCLE_STAT(TEXT("FConvexOptimizer::SimplifyRootConvexes"), STAT_SimplifyRootConvexes, STATGROUP_ChaosCollision);
-	void FConvexOptimizer::SimplifyRootConvexes(const Chaos::FImplicitObjectUnionPtr& UnionGeometry, const FShapesArray& UnionShapes)
+void FConvexOptimizer::SimplifyRootConvexes(const Chaos::FImplicitObjectUnionPtr& UnionGeometry, const FShapesArray& UnionShapes)
+{
+	SCOPE_CYCLE_COUNTER(STAT_Collisions_SimplifyConvexes);
+	if(CVars::bChaosConvexSimplifyUnion && UnionGeometry && (UnionShapes.Num() > 0) &&
+		(UnionShapes.Num() == UnionGeometry->GetObjects().Num()))
 	{
-		if(CVars::ChaosNumConvexLODs > 0)
+		SimplifiedConvexes.Reset();
+		CollisionObjects->ImplicitObjects.Reset();
+		ShapesArray.Reset();
+
+		if(CVars::ChaosConvexMaxLODs == 0)
 		{
-			BuildConvexLODs(UnionGeometry, UnionShapes);
+			BuildSingleConvex(UnionGeometry, UnionShapes);
 		}
 		else
 		{
-			BuildSingleConvex(UnionGeometry, UnionShapes);
+			BuildMultipleConvex(UnionGeometry, UnionShapes);
+		}
+		if(!SimplifiedConvexes.IsEmpty())
+		{
+			// Build the simplified shapes
+			BuildConvexShapes(UnionShapes);
 		}
 	}
 
-	void FConvexOptimizer::BuildSingleConvex(const Chaos::FImplicitObjectUnionPtr& UnionGeometry, const FShapesArray& UnionShapes)
+	if(CVars::bChaosUnionBVHEnabled)
 	{
-		if(CVars::bChaosSimplifyConvexes && UnionGeometry && (UnionShapes.Num() == UnionGeometry->GetObjects().Num()))
+		if(BVH.IsValid())
 		{
-			FTribox Tribox;
+			BVH.Reset();
+		}
+		BVH = FImplicitBVH::TryMakeFromLeaves(MoveTemp(CollisionObjects->ImplicitObjects), CVars::ChaosUnionBVHMinShapes, CVars::ChaosUnionBVHMaxDepth);
+	}
+}
+
+FORCEINLINE void InvalidateCachedTriboxes(FConvexOptimizer::FCachedTriboxes& RootTriboxes)
+{
+	for(auto& RootTribox : RootTriboxes)
+	{
+		RootTribox.Value.Tribox.SetValid(false);
+	}
+}
 
-			SimplifiedConvexes.Reset();
-			CollisionObjects->ImplicitObjects.Reset();
+FORCEINLINE void ResizeCachedTriboxes(FConvexOptimizer::FCachedTriboxes& RootTriboxes)
+{
+	TArray<FImplicitObject*> KeysToRemove;
+	for(auto& RootTribox : RootTriboxes)
+	{
+		if(!RootTribox.Value.Tribox.IsValid())
+		{
+			KeysToRemove.Add(RootTribox.Key);
+		}
+	}
+	for(auto& KeyToRemove : KeysToRemove)
+	{
+		RootTriboxes.Remove(KeyToRemove);
+	}
+}
+
+DECLARE_CYCLE_STAT(TEXT("Collisions::BuildConvexTriboxes"), STAT_BuildConvexTriboxes, STATGROUP_ChaosCollision);
+FORCEINLINE void BuildConvexTriboxes(const Chaos::FImplicitObjectUnionPtr& UnionGeometry, const FShapesArray& UnionShapes,
+	TUniquePtr<Private::FCollisionObjects>& CollisionObjects, FConvexOptimizer::FCachedTriboxes& RootTriboxes)
+{
+	SCOPE_CYCLE_COUNTER(STAT_BuildConvexTriboxes);
+
+	// Invalidate the cached root triboxes
+	InvalidateCachedTriboxes(RootTriboxes);
 	
-			UnionGeometry->VisitLeafObjects([&Tribox, this, &UnionShapes](
-				const FImplicitObject* ImplicitObject, const FRigidTransform3& RelativeTransform,
-				const int32 RootObjectIndex, const int32 ObjectIndex, const int32 LeafObjectIndex)
+	for(int32 RootObjectIndex = 0, NumRootObjects = UnionGeometry->GetObjects().Num(); RootObjectIndex < NumRootObjects; ++RootObjectIndex)
+	{
+		if(UnionShapes[RootObjectIndex]->GetSimEnabled())
+		{
+			if(FImplicitObject* RootObject = UnionGeometry->GetObjects()[RootObjectIndex].GetReference())
 			{
-				if(UnionShapes[RootObjectIndex]->GetSimEnabled())
+				FConvexOptimizer::FCachedTribox* RootTribox = RootTriboxes.Find(RootObject);
+				const bool bHasRootTribox = (RootTribox != nullptr);
+				
+				FTribox LocalTribox;
+				int32 LeafObjectIndex = 0, ObjectIndex = 0;
+				
+				RootObject->VisitLeafObjectsImpl(FRigidTransform3::Identity,RootObjectIndex, ObjectIndex, LeafObjectIndex,
+					[&CollisionObjects, &LocalTribox, &bHasRootTribox](
+					const FImplicitObject* ImplicitObject, const FRigidTransform3& RelativeTransform,
+					const int32 RootObjectIndex, const int32, const int32)
 				{
-					const FTribox::FRigidTransform3Type ConvexTransform(RelativeTransform);
 					if(const FConvex* Convex = ImplicitObject->AsA<FConvex>())
 					{
-						// Disable collision for all the convexes that are going to be used to build the tribox
-						// For now only used for debug draw
-						const_cast<FImplicitObject*>(ImplicitObject)->SetDoCollide(false);
-						Tribox.AddConvex(Convex, ConvexTransform);
+						if(!bHasRootTribox)
+						{
+							const FTribox::FRigidTransform3Type ConvexTransform(RelativeTransform);
+							// Disable collision for all the convexes that are going to be used to build the tribox
+							// For now only used for debug draw
+							const_cast<FImplicitObject*>(ImplicitObject)->SetDoCollide(false);
+							LocalTribox.AddConvex(Convex, ConvexTransform);
+						}
 					}
 					else
 					{
 						// Add non convex implicits to the list of collision objects
-						Private::FImplicitBVH::CollectLeafObject(ImplicitObject, RelativeTransform, RootObjectIndex, CollisionObjects->ImplicitObjects);	  
+						Private::FImplicitBVH::CollectLeafObject(ImplicitObject,
+							RelativeTransform, RootObjectIndex, CollisionObjects->ImplicitObjects);	  
 					}
+				});
+				if(!bHasRootTribox && LocalTribox.HasDatas())
+				{
+					LocalTribox.BuildTribox();
+					RootTribox = &(RootTriboxes.Add(RootObject, {LocalTribox,nullptr}));
+				}
+				if(RootTribox)
+				{
+					RootTribox->Tribox.SetValid(true);
 				}
-			});
-			if(Tribox.BuildTribox())
-			{
-				// Build the tribox and add it to the list of optimized convexes
-				SimplifiedConvexes.Add(Tribox.MakeConvex());
-				Private::FImplicitBVH::CollectLeafObject(SimplifiedConvexes.Last(), FRigidTransform3::Identity, INDEX_NONE, CollisionObjects->ImplicitObjects);	  
-
-				// Build the simplified shapes
-				BuildConvexShapes(UnionShapes);
 			}
 		}
 	}
+	ResizeCachedTriboxes(RootTriboxes);
+}
 
-	void FConvexOptimizer::BuildConvexLODs(const Chaos::FImplicitObjectUnionPtr& UnionGeometry, const FShapesArray& UnionShapes)
+void FConvexOptimizer::BuildSingleConvex(const Chaos::FImplicitObjectUnionPtr& UnionGeometry, const FShapesArray& UnionShapes)
+{
+	BuildConvexTriboxes(UnionGeometry, UnionShapes, CollisionObjects, RootTriboxes);
+	
+	FTribox MainTribox;
+	for(auto& RootTribox : RootTriboxes)
 	{
-		// WIP Another CL
+		if(RootTribox.Value.Tribox.HasDatas())
+		{
+			MainTribox += RootTribox.Value.Tribox;
+		}
 	}
 
-	void VisitCollisionObjects(const FConvexOptimizer* ConvexOptimizer, const FImplicitObject* ImplicitObject, const FImplicitHierarchyVisitor& VisitorFunc)
+	if(MainTribox.HasDatas())
 	{
-		if(ConvexOptimizer && ConvexOptimizer->IsValid())
-		{
-			ConvexOptimizer->VisitCollisionObjects(VisitorFunc);
-		}
-		else
-		{
-			ImplicitObject->VisitLeafObjects(VisitorFunc);
-		}
+		// Build the tribox and add it to the list of optimized convexes
+		SimplifiedConvexes.Add(MainTribox.MakeConvex());
+		Private::FImplicitBVH::CollectLeafObject(SimplifiedConvexes.Last(), FRigidTransform3::Identity, INDEX_NONE, CollisionObjects->ImplicitObjects);	
 	}
+}
 
-	void VisitOverlappingObjects(const FConvexOptimizer* ConvexOptimizer, const FImplicitObject* ImplicitObject, const FAABB3& LocalBounds, const FImplicitHierarchyVisitor& VisitorFunc)
+FORCEINLINE void FindClosestPlane(const FTribox::FVec3Type& PointPosition, const TArray<FTribox>& Triboxes,
+						int32& PlaneAxis, FTribox::FRealType& PlaneProjection)
+{
+	int32 LocalAxis = INDEX_NONE;
+	FTribox::FRealType LocalProjection = 0.0;
+	FTribox::FRealType MinDistance = FLT_MAX;
+				
+	for(auto& Tribox : Triboxes)
 	{
-		if(ConvexOptimizer && ConvexOptimizer->IsValid())
+		const FTribox::FRealType ClosestDistance = Tribox.GetClosestPlane(PointPosition, LocalAxis, LocalProjection);
+		if(ClosestDistance < MinDistance)
 		{
-			ConvexOptimizer->VisitOverlappingObjects(LocalBounds, VisitorFunc);
-		}
-		else
-		{
-			ImplicitObject->VisitOverlappingLeafObjects(LocalBounds, VisitorFunc);
+			MinDistance = ClosestDistance;
+			PlaneAxis = LocalAxis;
+			PlaneProjection = LocalProjection;
 		}
 	}
+}
+	
+void FConvexOptimizer::BuildMultipleConvex(const Chaos::FImplicitObjectUnionPtr& UnionGeometry, const FShapesArray& UnionShapes)
+{
+	BuildConvexTriboxes(UnionGeometry, UnionShapes, CollisionObjects, RootTriboxes);
+
+	if(CVars::ChaosConvexMaxLODs < 0)
+	{
+		for(auto& RootTribox : RootTriboxes)
+        {
+			if(RootTribox.Value.Tribox.HasDatas() && (RootTribox.Value.Tribox.ComputeVolume() > CVars::ChaosConvexMinVolume))
+			{
+				// Build the tribox and add it to the list of optimized convexes
+				if(!RootTribox.Value.Convex)
+				{
+					RootTribox.Value.Convex = RootTribox.Value.Tribox.MakeConvex();
+				}
+				SimplifiedConvexes.Add(RootTribox.Value.Convex);
+				Private::FImplicitBVH::CollectLeafObject(SimplifiedConvexes.Last(), FRigidTransform3::Identity, INDEX_NONE, CollisionObjects->ImplicitObjects);
+			}
+        }
+	}
+	else
+	{
+		// Build the binary tree and add the leaves to the cimplified convexes
+		FTriboxTree TriboxTree(CVars::ChaosConvexMaxLODs, RootTriboxes);
+		TriboxTree.BuildTreeLODs(SimplifiedConvexes, CollisionObjects);
+	}
+}
+
+void VisitCollisionObjects(const FConvexOptimizer* ConvexOptimizer, const FImplicitObject* ImplicitObject, const FImplicitHierarchyVisitor& VisitorFunc)
+{
+	if(ConvexOptimizer && ConvexOptimizer->IsValid())
+	{
+		ConvexOptimizer->VisitCollisionObjects(VisitorFunc);
+	}
+	else
+	{
+		ImplicitObject->VisitLeafObjects(VisitorFunc);
+	}
+}
+
+void VisitOverlappingObjects(const FConvexOptimizer* ConvexOptimizer, const FImplicitObject* ImplicitObject, const FAABB3& LocalBounds, const FImplicitHierarchyVisitor& VisitorFunc)
+{
+	if(ConvexOptimizer && ConvexOptimizer->IsValid())
+	{
+		ConvexOptimizer->VisitOverlappingObjects(LocalBounds, VisitorFunc);
+	}
+	else
+	{
+		ImplicitObject->VisitOverlappingLeafObjects(LocalBounds, VisitorFunc);
+	}
+}
 
 }
 }
diff --git a/Engine/Source/Runtime/Experimental/Chaos/Private/Chaos/ImplicitObjectBVH.cpp b/Engine/Source/Runtime/Experimental/Chaos/Private/Chaos/ImplicitObjectBVH.cpp
--- a/Engine/Source/Runtime/Experimental/Chaos/Private/Chaos/ImplicitObjectBVH.cpp
+++ b/Engine/Source/Runtime/Experimental/Chaos/Private/Chaos/ImplicitObjectBVH.cpp
@@ -100,6 +100,11 @@ namespace Chaos
 		TUniquePtr<FImplicitBVH> FImplicitBVH::TryMake(const TArrayView<const Chaos::FImplicitObjectPtr>& InRootObjects, const int32 InMinObjects, const int32 InMaxBVHDepth)
 		{
 			TArray<FImplicitBVHObject> LeafObjects = CollectLeafObjects(InRootObjects);
+			return TryMakeFromLeaves(MoveTemp(LeafObjects), InMinObjects, InMaxBVHDepth);
+		}
+		
+		TUniquePtr<FImplicitBVH> FImplicitBVH::TryMakeFromLeaves(TArray<FImplicitBVHObject>&& LeafObjects, const int32 InMinObjects, const int32 InMaxBVHDepth)
+		{
 			if (LeafObjects.Num() > InMinObjects)
 			{
 				return TUniquePtr<FImplicitBVH>(new FImplicitBVH(MoveTemp(LeafObjects), InMaxBVHDepth));
diff --git a/Engine/Source/Runtime/Experimental/Chaos/Private/Chaos/PBDRigidClustering.cpp b/Engine/Source/Runtime/Experimental/Chaos/Private/Chaos/PBDRigidClustering.cpp
--- a/Engine/Source/Runtime/Experimental/Chaos/Private/Chaos/PBDRigidClustering.cpp
+++ b/Engine/Source/Runtime/Experimental/Chaos/Private/Chaos/PBDRigidClustering.cpp
@@ -229,9 +229,6 @@ namespace Chaos
 		UpdateKinematicProperties(NewParticle, MChildren, MEvolution);
 		UpdateGeometry(NewParticle, ChildrenSet, MChildren, ProxyGeometry, Parameters);
 		
-		// Build the convex optimizer if required
-		BuildConvexOptimizer(NewParticle);
-		
 		GenerateConnectionGraph(NewParticle, Parameters);
 
 		NewParticle->SetSleeping(bClusterIsAsleep);
@@ -317,7 +314,10 @@ namespace Chaos
 
 	void FRigidClustering::BuildConvexOptimizer(FPBDRigidClusteredParticleHandle* Particle)
 	{
-		Particle->ConvexOptimizer() = MakePimpl<Private::FConvexOptimizer>();
+		if(!Particle->ConvexOptimizer())
+		{
+			Particle->ConvexOptimizer() = MakePimpl<Private::FConvexOptimizer>();
+		}
 		Particle->ConvexOptimizer()->SimplifyRootConvexes( Particle->GetGeometry()->template AsA<FImplicitObjectUnion>(),
 					Particle->ShapesArray());
 	}
@@ -509,9 +509,6 @@ namespace Chaos
 		UpdateKinematicProperties(NewParticle, MChildren, MEvolution);
 
 		UpdateGeometry(NewParticle, ChildrenSet, MChildren, FImplicitObjectPtr(nullptr), NoCleanParams);
-		
-		// Build the convex optimizer if required
-		BuildConvexOptimizer(NewParticle);
 
 		return NewParticle;
 	}
diff --git a/Engine/Source/Runtime/Experimental/Chaos/Private/Chaos/Tribox.cpp b/Engine/Source/Runtime/Experimental/Chaos/Private/Chaos/Tribox.cpp
--- a/Engine/Source/Runtime/Experimental/Chaos/Private/Chaos/Tribox.cpp
+++ b/Engine/Source/Runtime/Experimental/Chaos/Private/Chaos/Tribox.cpp
@@ -196,81 +196,81 @@ static const int32 PrincipalOrder[8][9][3] =
 		{3, 0, 0}
 	},
 	{
-	{2, 0, 7},
-	{2, 2, 6},
-	{1, 1, 3},
-	{7, 1, 5},
-	{3, 6, 4},
-	{2, 2, 6},
-	{0, 2, 6},
-	{1, 1, 5},
-	{0, 3, 0}
+		{2, 0, 7},
+		{2, 2, 6},
+		{1, 1, 3},
+		{7, 1, 5},
+		{3, 6, 4},
+		{2, 2, 6},
+		{0, 2, 6},
+		{1, 1, 5},
+		{0, 3, 0}
 	},
 	{
-	{6, 4, 3},
-	{2, 6, 2},
-	{1, 5, 7},
-	{1, 5, 1},
-	{3, 0, 0},
-	{2, 6, 0},
-	{2, 6, 2},
-	{1, 3, 1},
-	{0, 7, 2}
+		{6, 4, 3},
+		{2, 6, 2},
+		{1, 5, 7},
+		{1, 5, 1},
+		{3, 0, 0},
+		{2, 6, 0},
+		{2, 6, 2},
+		{1, 3, 1},
+		{0, 7, 2}
 	},
 	{
-	{0, 0, 3},
-	{6, 0, 2},
-	{5, 1, 1},
-	{3, 1, 1},
-	{7, 2, 0},
-	{6, 2, 2},
-	{6, 2, 2},
-	{5, 7, 1},
-	{4, 3, 6}
+		{0, 0, 3},
+		{6, 0, 2},
+		{5, 1, 1},
+		{3, 1, 1},
+		{7, 2, 0},
+		{6, 2, 2},
+		{6, 2, 2},
+		{5, 7, 1},
+		{4, 3, 6}
 	},
 	{
-	{2, 7, 0},
-	{5, 3, 1},
-	{6, 6, 2},
-	{6, 6, 2},
-	{4, 6, 3},
-	{5, 5, 7},
-	{5, 5, 1},
-	{6, 4, 2},
-	{7, 4, 4}
+		{2, 7, 0},
+		{5, 3, 1},
+		{6, 6, 2},
+		{6, 6, 2},
+		{4, 6, 3},
+		{5, 5, 7},
+		{5, 5, 1},
+		{6, 4, 2},
+		{7, 4, 4}
 	},
 	{
-	{6, 3, 4},
-	{5, 7, 5},
-	{6, 2, 6},
-	{4, 2, 6},
-	{4, 4, 7},
-	{5, 1, 5},
-	{3, 1, 5},
-	{6, 2, 6},
-	{7, 0, 2}
+		{6, 3, 4},
+		{5, 7, 5},
+		{6, 2, 6},
+		{4, 2, 6},
+		{4, 4, 7},
+		{5, 1, 5},
+		{3, 1, 5},
+		{6, 2, 6},
+		{7, 0, 2}
 	},
 	{
-	{4, 7, 4},
-	{1, 5, 5},
-	{2, 6, 4},
-	{2, 6, 6},
-	{0, 2, 7},
-	{1, 5, 3},
-	{7, 5, 5},
-	{2, 6, 6},
-	{3, 4, 6}
+		{4, 7, 4},
+		{1, 5, 5},
+		{2, 6, 4},
+		{2, 6, 6},
+		{0, 2, 7},
+		{1, 5, 3},
+		{7, 5, 5},
+		{2, 6, 6},
+		{3, 4, 6}
 	},
 	{
-	{4, 4, 7},
-	{6, 4, 6},
-	{5, 5, 5},
-	{5, 5, 5},
-	{7, 4, 4},
-	{6, 6, 4},
-	{4, 6, 6},
-	{5, 5, 5},
-	{4, 7, 4}
+		{4, 4, 7},
+		{6, 4, 6},
+		{5, 5, 5},
+		{5, 5, 5},
+		{7, 4, 4},
+		{6, 6, 4},
+		{4, 6, 6},
+		{5, 5, 5},
+		{4, 7, 4}
 	}
 };
 
@@ -303,11 +303,13 @@ void FTribox::AddPoint(const FVec3Type& PointPosition)
 	BuildChamferDist(P, 0, 2, 10, 11, 12, 13);
 	BuildChamferDist(P, 0, 1, 14, 15, 16, 17);
 
-	++NumPoints;
+	bHasDatas = true;
 }
 
+DECLARE_CYCLE_STAT(TEXT("Collisions::AddConvex"), STAT_AddConvexToTribox, STATGROUP_ChaosCollision);
 void FTribox::AddConvex(const FConvex* Convex, const FRigidTransform3Type& RelativeTransform)
 {
+	SCOPE_CYCLE_COUNTER(STAT_AddConvexToTribox);
 	if(Convex)
 	{
 		for(int32 VertexIndex = 0; VertexIndex < Convex->NumVertices(); ++VertexIndex)
@@ -322,7 +324,7 @@ void FTribox::AddConvex(const FConvex* Convex, const FRigidTransform3Type& Relat
 
 bool FTribox::BuildTribox()
 {
-	if(NumPoints > 0)
+	if(bHasDatas)
 	{
 		// Inflate the principal planes max distance by a given distance
 		for(int32 PlaneIndex = 0; PlaneIndex < NumPrincipalPlanes; ++PlaneIndex)
@@ -393,58 +395,62 @@ FORCEINLINE void CompressFaces(TArray<TArray<int32>>& FaceIndices)
 	}
 }
 
+DECLARE_CYCLE_STAT(TEXT("Collisions::MakeConvex"), STAT_MakeTriboxConvex, STATGROUP_ChaosCollision);
 FImplicitObjectPtr FTribox::MakeConvex() const
 {
 	TArray<TArray<int32>> FaceIndices;
 	TArray<FConvex::FVec3Type> ConvexVertices;
 	TArray<FConvex::FPlaneType> ConvexPlanes;
-
-	FaceIndices.SetNum(NumPlanes);
-	ConvexVertices.Reserve(32);
-	ConvexPlanes.SetNum(NumPlanes);
-
-	for(int32 PlaneIndex = 0; PlaneIndex < NumPlanes; ++PlaneIndex)
-	{
-		ConvexPlanes[PlaneIndex] = FConvex::FPlaneType(
-			PlanesDirs[PlaneIndex] * MaxDists[PlaneIndex], PlanesDirs[PlaneIndex]);
-		FaceIndices[PlaneIndex].Init(INDEX_NONE, 8);
-	}
-	for(int32 CornerIndex = 0; CornerIndex < 8; ++CornerIndex)
 	{
-		// Compute the intersection in between 3 corner chamfer planes
-		const FVec3Type P = SolveIntersection(ChamferIndices[CornerIndex], ChamferMatrices[CornerIndex]);
+		SCOPE_CYCLE_COUNTER(STAT_MakeTriboxConvex);
 
-		if(P.Dot(PlanesDirs[PrincipalIndices[CornerIndex][0]]) > MaxDists[PrincipalIndices[CornerIndex][0]])
-		{
-			ComputeIntersection(CornerIndex, 0,1, FaceIndices, ConvexVertices);
-			ComputeIntersection(CornerIndex, 0,2, FaceIndices, ConvexVertices);
-			ComputeIntersection(CornerIndex, 1,1, FaceIndices, ConvexVertices);
-			ComputeIntersection(CornerIndex, 2,2, FaceIndices, ConvexVertices);
-		}
-		else if(P.Dot(PlanesDirs[PrincipalIndices[CornerIndex][1]]) > MaxDists[PrincipalIndices[CornerIndex][1]])
-		{
-			ComputeIntersection(CornerIndex, 1,0, FaceIndices, ConvexVertices);
-			ComputeIntersection(CornerIndex, 1,2, FaceIndices, ConvexVertices);
-			ComputeIntersection(CornerIndex, 0,0, FaceIndices, ConvexVertices);
-			ComputeIntersection(CornerIndex, 2,2, FaceIndices, ConvexVertices);
-		}
-		else if(P.Dot(PlanesDirs[PrincipalIndices[CornerIndex][2]]) > MaxDists[PrincipalIndices[CornerIndex][2]])
+		FaceIndices.SetNum(NumPlanes);
+		ConvexVertices.Reserve(32);
+		ConvexPlanes.SetNum(NumPlanes);
+
+		for(int32 PlaneIndex = 0; PlaneIndex < NumPlanes; ++PlaneIndex)
 		{
-			ComputeIntersection(CornerIndex, 2,0, FaceIndices, ConvexVertices);
-			ComputeIntersection(CornerIndex, 2,1, FaceIndices, ConvexVertices);
-			ComputeIntersection(CornerIndex, 0,0, FaceIndices, ConvexVertices);
-			ComputeIntersection(CornerIndex, 1,1, FaceIndices, ConvexVertices);
+			ConvexPlanes[PlaneIndex] = FConvex::FPlaneType(
+				PlanesDirs[PlaneIndex] * MaxDists[PlaneIndex], PlanesDirs[PlaneIndex]);
+			FaceIndices[PlaneIndex].Init(INDEX_NONE, 8);
 		}
-		else
+		for(int32 CornerIndex = 0; CornerIndex < 8; ++CornerIndex)
 		{
-			ComputeIntersection(CornerIndex, 0,0, FaceIndices, ConvexVertices);
-			ComputeIntersection(CornerIndex, 1,1, FaceIndices, ConvexVertices);
-			ComputeIntersection(CornerIndex, 2,2, FaceIndices, ConvexVertices);
+			// Compute the intersection in between 3 corner chamfer planes
+			const FVec3Type P = SolveIntersection(ChamferIndices[CornerIndex], ChamferMatrices[CornerIndex]);
+
+			if(P.Dot(PlanesDirs[PrincipalIndices[CornerIndex][0]]) > MaxDists[PrincipalIndices[CornerIndex][0]])
+			{
+				ComputeIntersection(CornerIndex, 0,1, FaceIndices, ConvexVertices);
+				ComputeIntersection(CornerIndex, 0,2, FaceIndices, ConvexVertices);
+				ComputeIntersection(CornerIndex, 1,1, FaceIndices, ConvexVertices);
+				ComputeIntersection(CornerIndex, 2,2, FaceIndices, ConvexVertices);
+			}
+			else if(P.Dot(PlanesDirs[PrincipalIndices[CornerIndex][1]]) > MaxDists[PrincipalIndices[CornerIndex][1]])
+			{
+				ComputeIntersection(CornerIndex, 1,0, FaceIndices, ConvexVertices);
+				ComputeIntersection(CornerIndex, 1,2, FaceIndices, ConvexVertices);
+				ComputeIntersection(CornerIndex, 0,0, FaceIndices, ConvexVertices);
+				ComputeIntersection(CornerIndex, 2,2, FaceIndices, ConvexVertices);
+			}
+			else if(P.Dot(PlanesDirs[PrincipalIndices[CornerIndex][2]]) > MaxDists[PrincipalIndices[CornerIndex][2]])
+			{
+				ComputeIntersection(CornerIndex, 2,0, FaceIndices, ConvexVertices);
+				ComputeIntersection(CornerIndex, 2,1, FaceIndices, ConvexVertices);
+				ComputeIntersection(CornerIndex, 0,0, FaceIndices, ConvexVertices);
+				ComputeIntersection(CornerIndex, 1,1, FaceIndices, ConvexVertices);
+			}
+			else
+			{
+				ComputeIntersection(CornerIndex, 0,0, FaceIndices, ConvexVertices);
+				ComputeIntersection(CornerIndex, 1,1, FaceIndices, ConvexVertices);
+				ComputeIntersection(CornerIndex, 2,2, FaceIndices, ConvexVertices);
 
-			AddIntersection(ChamferIndices[CornerIndex], ChamferOrder[CornerIndex], P, FaceIndices, ConvexVertices);
+				AddIntersection(ChamferIndices[CornerIndex], ChamferOrder[CornerIndex], P, FaceIndices, ConvexVertices);
+			}
 		}
+		CompressFaces(FaceIndices);
 	}
-	CompressFaces(FaceIndices);
 	return MakeImplicitObjectPtr<Chaos::FConvex>(MoveTemp(ConvexPlanes), MoveTemp(FaceIndices), MoveTemp(ConvexVertices));
 }
 	
@@ -471,17 +477,157 @@ bool FTribox::OverlapTribox(const FTribox& OtherTribox, FTribox& OverlapTribox)
 	return bOverlapTriboxes;
 }
 
+bool FTribox::SplitTriboxSlab(const int32 PlaneAxis, const FRealType& PlaneDistance,
+	FTribox& LeftTribox, FTribox& RightTribox) const
+{
+	if(PlaneDistance < MaxDists[PlaneAxis])
+	{
+		if(PlaneDistance > -MaxDists[PlaneAxis+1])
+		{
+			LeftTribox = *this;
+			RightTribox = *this;
+
+			auto CutTribox = [](const int32 PX, const int32 MX, const int32 PY, const int32 MY,
+						const int32 PZPX, const int32 PZMX, const int32 PZPY, const int32 PZMY,
+						const int32 MZPX, const int32 MZMX, const int32 MZPY, const int32 MZMY,
+						const int32 PXPY, const int32 PXMY, const int32 MXPY, const int32 MXMY,
+						const int32 PZ, const FRealType MaxDistance, FTribox& SideTribox)
+			{
+				SideTribox.MaxDists[PZ] = MaxDistance;
+				
+				SideTribox.MaxDists[PX] = FMath::Min(SideTribox.MaxDists[PX], SideTribox.MaxDists[MZPX] * Sqrt2 + MaxDistance);
+				SideTribox.MaxDists[MX] = FMath::Min(SideTribox.MaxDists[MX], SideTribox.MaxDists[MZMX] * Sqrt2 + MaxDistance);
+				SideTribox.MaxDists[PY] = FMath::Min(SideTribox.MaxDists[PY], SideTribox.MaxDists[MZPY] * Sqrt2 + MaxDistance);
+				SideTribox.MaxDists[MY] = FMath::Min(SideTribox.MaxDists[MY], SideTribox.MaxDists[MZMY] * Sqrt2 + MaxDistance);
+
+				SideTribox.MaxDists[PZPX] = FMath::Min(SideTribox.MaxDists[PZPX], (SideTribox.MaxDists[PX] + MaxDistance) * InvSqrt2);
+				SideTribox.MaxDists[PZMX] = FMath::Min(SideTribox.MaxDists[PZMX], (SideTribox.MaxDists[MX] + MaxDistance) * InvSqrt2);
+				SideTribox.MaxDists[PZPY] = FMath::Min(SideTribox.MaxDists[PZPY], (SideTribox.MaxDists[PY] + MaxDistance) * InvSqrt2);
+				SideTribox.MaxDists[PZMY] = FMath::Min(SideTribox.MaxDists[PZMY], (SideTribox.MaxDists[MY] + MaxDistance) * InvSqrt2);
+				
+				SideTribox.MaxDists[PXPY] = FMath::Min(SideTribox.MaxDists[PXPY], (SideTribox.MaxDists[PX] + SideTribox.MaxDists[PY]) * InvSqrt2);
+				SideTribox.MaxDists[PXMY] = FMath::Min(SideTribox.MaxDists[PXMY], (SideTribox.MaxDists[PX] + SideTribox.MaxDists[MY]) * InvSqrt2);
+				SideTribox.MaxDists[MXPY] = FMath::Min(SideTribox.MaxDists[MXPY], (SideTribox.MaxDists[MX] + SideTribox.MaxDists[PY]) * InvSqrt2);
+				SideTribox.MaxDists[MXMY] = FMath::Min(SideTribox.MaxDists[MXMY], (SideTribox.MaxDists[MX] + SideTribox.MaxDists[MY]) * InvSqrt2);
+			};
+			
+			if(PlaneAxis == 0)
+			{
+				CutTribox(4,5,2,3,10,12,14,16,13,11,17,15,6,9,8,7,0, PlaneDistance, LeftTribox);
+				CutTribox(4,5,2,3,13,11,17,15,10,12,14,16,6,9,8,7,1,-PlaneDistance, RightTribox);
+			}
+			else if(PlaneAxis == 2)
+			{
+				CutTribox(4,5,0,1,6,8,14,17,9,7,16,15,10,13,12,11,2, PlaneDistance, LeftTribox);
+				CutTribox(4,5,0,1,9,7,16,15,6,8,14,17,10,13,12,11,3,-PlaneDistance, RightTribox);
+			}
+			else if(PlaneAxis == 4)
+			{
+				CutTribox(2,3,0,1,6,9,10,13,8,7,12,11,14,17,16,15,4, PlaneDistance, LeftTribox);
+				CutTribox(2,3,0,1,8,7,12,11,6,9,10,13,14,17,16,15,5,-PlaneDistance, RightTribox);
+			}
+			
+			return true;
+		}
+		else  
+		{
+			RightTribox = *this;
+		}
+	}
+	else
+	{
+		LeftTribox = *this;
+	}
+	return false;
+}
+
+int32 FTribox::GetThickestSlab() const
+{
+	int32 ThickestSlab = INDEX_NONE;
+	FRealType MaxThickness = 0.0;
+	for(int32 AxisIndex = 0; AxisIndex < NumPrincipalPlanes; AxisIndex += 2)
+	{
+		const FRealType SlabThickness = (MaxDists[AxisIndex]+MaxDists[AxisIndex+1]);
+		if(SlabThickness > MaxThickness)
+		{
+			ThickestSlab = AxisIndex;
+			MaxThickness = SlabThickness;
+		}
+	}
+	return ThickestSlab;
+}
+
+FTribox::FRealType FTribox::SampleSlabPoint(const int32 PlaneAxis, const FRealType& LocalDistance) const
+{
+	return -MaxDists[PlaneAxis+1] + FMath::Max(0.0f, FMath::Min(1.0f, LocalDistance)) * (MaxDists[PlaneAxis]+MaxDists[PlaneAxis+1]);
+}
+	
+FTribox::FRealType FTribox::GetClosestPlane(const FVec3Type& PointPosition, int32& PlaneAxis, FRealType& PlaneProjection) const
+{
+	FRealType ClosestDistance = -FLT_MAX;
+	for(int32 AxisIndex = 0; AxisIndex < 3; ++AxisIndex)
+	{
+		const int32 PlaneIndex = 2 * AxisIndex;
+		FRealType PlaneDistance = PointPosition[AxisIndex]-MaxDists[PlaneIndex];
+		if(PlaneDistance > ClosestDistance)
+		{
+			ClosestDistance = PlaneDistance;
+			PlaneAxis = PlaneIndex;
+			PlaneProjection = MaxDists[PlaneIndex];
+		}
+		PlaneDistance = -PointPosition[AxisIndex]-MaxDists[PlaneIndex+1];
+		if(PlaneDistance > ClosestDistance)
+		{
+			ClosestDistance = PlaneDistance;
+			PlaneAxis = PlaneIndex;
+			PlaneProjection = -MaxDists[PlaneIndex+1];
+		}
+	}
+	return ClosestDistance;
+}
+
 FTribox::FRealType FTribox::ComputeVolume() const
 {
-	return (MaxDists[0]+MaxDists[1]) * (MaxDists[2]+MaxDists[3]) * (MaxDists[4]+MaxDists[5]);
+	const FVec3Type BoxSize(MaxDists[0]+MaxDists[1], MaxDists[2]+MaxDists[3], MaxDists[4]+MaxDists[5]);
+	FRealType Volume = BoxSize[0] * BoxSize[1] * BoxSize[2];
+
+	auto RemovePrism = [&Volume, &BoxSize, this](const int32 AxisX, const int32 AxisY, const int32 AxisZ,
+		const int32 ChamferA, const int32 ChamferB, const int32 ChamferC, const int32 ChamferD)
+	{
+		FRealType CornerDistance = (MaxDists[AxisX] + MaxDists[AxisY]) * InvSqrt2 - MaxDists[ChamferA];
+		Volume -= CornerDistance * CornerDistance * BoxSize[AxisZ];
+		CornerDistance = (MaxDists[AxisX] + MaxDists[AxisY+1]) * InvSqrt2 - MaxDists[ChamferB];
+		Volume -= CornerDistance * CornerDistance * BoxSize[AxisZ];
+		CornerDistance = (MaxDists[AxisX+1] + MaxDists[AxisY]) * InvSqrt2 - MaxDists[ChamferC];
+		Volume -= CornerDistance * CornerDistance * BoxSize[AxisZ];
+		CornerDistance = (MaxDists[AxisX+1] + MaxDists[AxisY+1]) * InvSqrt2 - MaxDists[ChamferD];
+		Volume -= CornerDistance * CornerDistance * BoxSize[AxisZ];
+	};
+
+	// Remove all th prisms along the z directions
+	RemovePrism(0,2,2,14,16,17,15);
+	
+	// Remove all th prisms along the y directions
+	RemovePrism(0,4,1,10,12,13,11);
+	
+	// Remove all th prisms along the x directions
+	RemovePrism(2,4,0,6,8,9,7);
+	
+	return Volume;
+}
+	
+FTribox::FVec3Type FTribox::GetCenter() const
+{
+	return 0.5f * FVec3Type(MaxDists[0]-MaxDists[1], MaxDists[2]-MaxDists[3], MaxDists[4]-MaxDists[5]);
 }
 	
 FTribox& FTribox::operator+=(const FTribox& OtherTribox)
 {
 	for(int32 PlaneIndex = 0; PlaneIndex < NumPlanes; ++PlaneIndex)
 	{
-		MaxDists[PlaneIndex]  = FMath::Max(OtherTribox.MaxDists[PlaneIndex], MaxDists[PlaneIndex]);
+		MaxDists[PlaneIndex] = FMath::Max(OtherTribox.MaxDists[PlaneIndex], MaxDists[PlaneIndex]);
 	}
+	bHasDatas |= OtherTribox.bHasDatas;
 	return *this;
 }
 
diff --git a/Engine/Source/Runtime/Experimental/Chaos/Private/ChaosStats.cpp b/Engine/Source/Runtime/Experimental/Chaos/Private/ChaosStats.cpp
--- a/Engine/Source/Runtime/Experimental/Chaos/Private/ChaosStats.cpp
+++ b/Engine/Source/Runtime/Experimental/Chaos/Private/ChaosStats.cpp
@@ -112,3 +112,4 @@ DEFINE_STAT(STAT_Collisions_Gather);
 DEFINE_STAT(STAT_Collisions_Scatter);
 DEFINE_STAT(STAT_Collisions_Apply);
 DEFINE_STAT(STAT_Collisions_ApplyPushOut);
+DEFINE_STAT(STAT_Collisions_SimplifyConvexes);
diff --git a/Engine/Source/Runtime/Experimental/Chaos/Public/Chaos/ConvexOptimizer.h b/Engine/Source/Runtime/Experimental/Chaos/Public/Chaos/ConvexOptimizer.h
--- a/Engine/Source/Runtime/Experimental/Chaos/Public/Chaos/ConvexOptimizer.h
+++ b/Engine/Source/Runtime/Experimental/Chaos/Public/Chaos/ConvexOptimizer.h
@@ -4,6 +4,7 @@
 
 #include "Chaos/Core.h"
 #include "Chaos/ImplicitObject.h"
+#include "Chaos/Tribox.h"
 
 namespace Chaos
 {
@@ -23,6 +24,18 @@ namespace Private
 	class FConvexOptimizer
 	{
 		public :
+		
+		struct FCachedTribox
+		{
+			// Cached tribox
+			Private::FTribox Tribox;
+
+			// Cached convex
+			FImplicitObjectPtr Convex;
+		};
+
+		using FCachedTriboxes = TMap<FImplicitObject*,FCachedTribox>;
+		
 		CHAOS_API FConvexOptimizer();
 
 		// Default destructor
@@ -48,8 +61,8 @@ namespace Private
 		// Build a single convex
 		void BuildSingleConvex(const Chaos::FImplicitObjectUnionPtr& UnionGeometry, const FShapesArray& UnionShapes);
 
-		// Build several convexes for a given LOD
-		void BuildConvexLODs(const Chaos::FImplicitObjectUnionPtr& UnionGeometry, const FShapesArray& UnionShapes);
+		// Build several convexes 
+		void BuildMultipleConvex(const Chaos::FImplicitObjectUnionPtr& UnionGeometry, const FShapesArray& UnionShapes);
 
 		// Build the simplified shapes
 		void BuildConvexShapes(const FShapesArray& UnionShapes);
@@ -62,6 +75,12 @@ namespace Private
 
 		// Additional shapes array that could be used during collision midphase
 		FShapeInstanceArray ShapesArray;
+
+		// Intermediate root triboxes to reuse the intermediate computation
+		FCachedTriboxes RootTriboxes;
+
+		// BVH used to accelerate the collisions queries
+		TUniquePtr<Private::FImplicitBVH> BVH;
 	};
 
 	// Visit all the collision objects if they exist / otherwise forward it to the RootHierarchy
diff --git a/Engine/Source/Runtime/Experimental/Chaos/Public/Chaos/ImplicitObjectBVH.h b/Engine/Source/Runtime/Experimental/Chaos/Public/Chaos/ImplicitObjectBVH.h
--- a/Engine/Source/Runtime/Experimental/Chaos/Public/Chaos/ImplicitObjectBVH.h
+++ b/Engine/Source/Runtime/Experimental/Chaos/Public/Chaos/ImplicitObjectBVH.h
@@ -117,14 +117,19 @@ namespace Chaos
 
 			// Create a BVH around a set of ImplicitObjects. Usually these are the immediate child elements of an FImplcitObjectUnion
 			// TryMake will then recurse into the geometry hierachy and add all descendents to the BVH. Will return null if the 
-			// number of descendents is less that MinObjscts.
+			// number of descendents is less that MinObjects.
 			static TUniquePtr<FImplicitBVH> TryMake(const TArrayView<const Chaos::FImplicitObjectPtr>& InRootObjects, const int32 MinObjects, const int32 InMaxBVHDepth);
+			
+			// Create a BVH around given a list of leaves.  Will return null if the  number of descendents is less that MinObjects.
+			static TUniquePtr<FImplicitBVH> TryMakeFromLeaves(TArray<FImplicitBVHObject>&& LeafObjects, const int32 InMinObjects, const int32 InMaxBVHDepth);
 
 			~FImplicitBVH();
 
 			int32 GetNumObjects() const { return Objects.Num(); }
 			int32 GetDepth() const { return TreeDepth; }
 			
+			const TArray<FImplicitBVHObject>& GetObjects() const { return Objects; }
+			
 			const FImplicitBVHObject& GetObject(const int32 ObjectIndex) const { return Objects[ObjectIndex]; }
 
 			const FImplicitObject* GetGeometry(const int32 ObjectIndex) const { return Objects[ObjectIndex].GetGeometry(); }
diff --git a/Engine/Source/Runtime/Experimental/Chaos/Public/Chaos/Tribox.h b/Engine/Source/Runtime/Experimental/Chaos/Public/Chaos/Tribox.h
--- a/Engine/Source/Runtime/Experimental/Chaos/Public/Chaos/Tribox.h
+++ b/Engine/Source/Runtime/Experimental/Chaos/Public/Chaos/Tribox.h
@@ -26,6 +26,7 @@ public :
 	using FVec3Type = TVec3<FRealType>;
 	using FRigidTransform3Type = TRigidTransform<FRealType, 3>;
 	using FMatrix33Type = PMatrix<FRealType, 3, 3>;
+	using FPlaneType = TPlaneConcrete<FRealType, 3>;
 
 	// Number of planes that will be used to define the Tribox 
 	static constexpr int32 NumPlanes = 18;
@@ -40,7 +41,7 @@ public :
 	static constexpr FRealType InflateDistance = 0.5;
 
 	// Base Constructor
-	FORCEINLINE FTribox() : MaxDists()
+	FORCEINLINE FTribox() : MaxDists(), bIsValid(false), bHasDatas(false)
 	{
 		for(int32 DistsIndex = 0; DistsIndex< NumPlanes; ++DistsIndex)
 		{
@@ -48,6 +49,23 @@ public :
 		}
 	};
 
+	// Base Constructor
+	FORCEINLINE FTribox(const FTribox& OtherTribox) : MaxDists(), bIsValid(false), bHasDatas(false)
+	{
+		for(int32 DistsIndex = 0; DistsIndex < NumPlanes; ++DistsIndex)
+		{
+			MaxDists[DistsIndex] = OtherTribox.MaxDists[DistsIndex];
+		}
+		bIsValid = OtherTribox.bIsValid;
+		bHasDatas = OtherTribox.bHasDatas;
+	};
+	
+	// Get the tribox center
+	FVec3Type GetCenter() const;
+	
+	// Get the closest plane along the +X,-X,+Y,-Y,+Z,-Z directions
+	FRealType GetClosestPlane(const FVec3Type& PointPosition, int32& PlaneAxis, FRealType& PlaneProjection) const;
+
 	// Add a point position to the Tribox
 	void AddPoint(const FVec3Type& PointPosition);
 
@@ -60,6 +78,16 @@ public :
 	// Find the overlapping tribox 
 	bool OverlapTribox(const FTribox& OtherTribox, FTribox& OverlapTribox) const;
 
+	// Split the tribox in 2 along a defined cuttng plane
+	bool SplitTriboxSlab(const int32 PlaneAxis, const FRealType& PlaneDistance,
+					FTribox& LeftTribox, FTribox& RightTribox) const;
+
+	// Get the thickest tribox slab
+	int32 GetThickestSlab() const;
+
+	// Sample a point along the plane direction in betwen min and max
+	FRealType SampleSlabPoint(const int32 PlaneAxis, const FRealType& LocalDistance) const;
+
 	// Compute the tribox volume
 	FRealType ComputeVolume() const;
 
@@ -72,6 +100,24 @@ public :
 	// Add a tribox to this and return a new one
 	FTribox operator+( const FTribox& OtherTribox) const;
 
+	// Check ihe tribox is valid
+	bool IsValid() const {return bIsValid;}
+
+	// Set the valid flag
+	void SetValid(const bool bValid) {bIsValid = bValid;}
+
+	// Check ihe tribox have been built with datas
+	bool HasDatas() const {return bHasDatas;}
+
+	// Reset the tribox max distances
+	void ResetDists()
+	{
+		for(int32 DistsIndex = 0; DistsIndex< NumPlanes; ++DistsIndex)
+		{
+			MaxDists[DistsIndex] = TNumericLimits<FRealType>::Lowest();
+		}
+	}
+
 private : 
 	// Solve the intersection point position
 	FVec3 SolveIntersection(const int32 FaceIndex[3], const FMatrix33Type& A) const;
@@ -96,8 +142,11 @@ private :
 									 TNumericLimits<FRealType>::Lowest(), TNumericLimits<FRealType>::Lowest(),
 									 TNumericLimits<FRealType>::Lowest()};
 
-	// Number of points used to build the simplified convexes
-	int32 NumPoints = 0;
+	// Boolean to specify if the tribox is valid or not 
+	bool bIsValid = false;
+
+	// Boolean to specify if the tribox has been built with datas or not
+	bool bHasDatas = false;
 };
 
 }
diff --git a/Engine/Source/Runtime/Experimental/Chaos/Public/ChaosStats.h b/Engine/Source/Runtime/Experimental/Chaos/Public/ChaosStats.h
--- a/Engine/Source/Runtime/Experimental/Chaos/Public/ChaosStats.h
+++ b/Engine/Source/Runtime/Experimental/Chaos/Public/ChaosStats.h
@@ -137,6 +137,7 @@ DECLARE_CYCLE_STAT_EXTERN(TEXT("Collisions::Gather"), STAT_Collisions_Gather, ST
 DECLARE_CYCLE_STAT_EXTERN(TEXT("Collisions::Scatter"), STAT_Collisions_Scatter, STATGROUP_ChaosCollision, CHAOS_API);
 DECLARE_CYCLE_STAT_EXTERN(TEXT("Collisions::Apply"), STAT_Collisions_Apply, STATGROUP_ChaosCollision, CHAOS_API);
 DECLARE_CYCLE_STAT_EXTERN(TEXT("Collisions::ApplyPushOut"), STAT_Collisions_ApplyPushOut, STATGROUP_ChaosCollision, CHAOS_API);
+DECLARE_CYCLE_STAT_EXTERN(TEXT("Collisions::SimplifyConvexes"), STAT_Collisions_SimplifyConvexes, STATGROUP_ChaosCollision, CHAOS_API);
 
 
 #if 0
