Iris
- Added support for sending attachments flagged as SendImmediate from PostDispatchTick

diff --git a/Engine/Plugins/Runtime/ReplicationSystemTestPlugin/Source/Private/Tests/ReplicationSystem/RPC/ReplicatedTestObjectWithRPC.cpp b/Engine/Plugins/Runtime/ReplicationSystemTestPlugin/Source/Private/Tests/ReplicationSystem/RPC/ReplicatedTestObjectWithRPC.cpp
--- a/Engine/Plugins/Runtime/ReplicationSystemTestPlugin/Source/Private/Tests/ReplicationSystem/RPC/ReplicatedTestObjectWithRPC.cpp
+++ b/Engine/Plugins/Runtime/ReplicationSystemTestPlugin/Source/Private/Tests/ReplicationSystem/RPC/ReplicatedTestObjectWithRPC.cpp
@@ -107,4 +107,12 @@ void UTestReplicatedObjectWithRPC::ServerRPCWithParam_Implementation(int32 IntPa
 	ServerRPCWithParamCalled = IntParam;
 }
 
+void UTestReplicatedObjectWithRPC::NetMulticast_MultiCastRPCSendImmediate_Implementation()
+{
+	NetMulticast_MultiCastRPCSendImmediateCallOrder = ++CallOrder;
+}
 
+void UTestReplicatedObjectWithRPC::NetMulticast_MultiCastRPC_Implementation()
+{
+	NetMulticast_MultiCastRPCCallOrder = ++CallOrder;
+}
diff --git a/Engine/Plugins/Runtime/ReplicationSystemTestPlugin/Source/Private/Tests/ReplicationSystem/RPC/ReplicatedTestObjectWithRPC.h b/Engine/Plugins/Runtime/ReplicationSystemTestPlugin/Source/Private/Tests/ReplicationSystem/RPC/ReplicatedTestObjectWithRPC.h
--- a/Engine/Plugins/Runtime/ReplicationSystemTestPlugin/Source/Private/Tests/ReplicationSystem/RPC/ReplicatedTestObjectWithRPC.h
+++ b/Engine/Plugins/Runtime/ReplicationSystemTestPlugin/Source/Private/Tests/ReplicationSystem/RPC/ReplicatedTestObjectWithRPC.h
@@ -46,5 +46,16 @@ public:
 	UFUNCTION(Reliable, Server)
 	void ServerRPCWithParam(int32 IntParam);
 	int32 ServerRPCWithParamCalled = 0;
+
+	UFUNCTION(NetMulticast, unreliable)
+	void NetMulticast_MultiCastRPCSendImmediate();
+	int32 NetMulticast_MultiCastRPCSendImmediateCallOrder = 0;
+
+	UFUNCTION(NetMulticast, unreliable)
+	void NetMulticast_MultiCastRPC();
+	int32 NetMulticast_MultiCastRPCCallOrder = 0;
+
+	int32 CallOrder = 0;
+
 };
 
diff --git a/Engine/Plugins/Runtime/ReplicationSystemTestPlugin/Source/Private/Tests/ReplicationSystem/RPC/TestObjectRPC.cpp b/Engine/Plugins/Runtime/ReplicationSystemTestPlugin/Source/Private/Tests/ReplicationSystem/RPC/TestObjectRPC.cpp
--- a/Engine/Plugins/Runtime/ReplicationSystemTestPlugin/Source/Private/Tests/ReplicationSystem/RPC/TestObjectRPC.cpp
+++ b/Engine/Plugins/Runtime/ReplicationSystemTestPlugin/Source/Private/Tests/ReplicationSystem/RPC/TestObjectRPC.cpp
@@ -82,4 +82,58 @@ namespace UE::Net::Private
 		// Verify RPC reception
 		UE_NET_ASSERT_TRUE(ServerObject->ServerRPCWithParamCalled == IntParam);
 	}
+
+	UE_NET_TEST_FIXTURE(FRPCTestFixture, TestMultiCastSendImmediateRPC)
+	{
+		// Add a client
+		FReplicationSystemTestClient* Client = CreateClient();
+
+		// Spawn object on server
+		UTestReplicatedObjectWithRPC* ServerObject = Server->CreateObject<UTestReplicatedObjectWithRPC>();
+
+		ServerObject->bIsServerObject = true;
+		ServerObject->ReplicationSystem = Server->GetReplicationSystem();
+		Server->ReplicationSystem->SetOwningNetConnection(ServerObject->NetRefHandle, 0x01);
+
+		// Setup NetMulticast_MultiCastRPCSendImmediateCallOrder to be sent immediately
+		ServerObject->ReplicationSystem->SetRPCSendPolicyFlags(UTestReplicatedObjectWithRPC::StaticClass()->FindFunctionByName(GET_FUNCTION_NAME_CHECKED(UTestReplicatedObjectWithRPC, NetMulticast_MultiCastRPCSendImmediate)), ENetObjectAttachmentSendPolicyFlags::SendImmediate);
+
+		// Send and deliver packet
+		Server->PreSendUpdate();
+		Server->SendAndDeliverTo(Client, true);
+		Server->PostSendUpdate();
+
+		UTestReplicatedObjectWithRPC* ClientObject = Cast<UTestReplicatedObjectWithRPC>(Client->GetReplicationBridge()->GetReplicatedObject(ServerObject->NetRefHandle));
+		
+		// Verify that created server handle now also exists on client
+		UE_NET_ASSERT_TRUE(ClientObject != nullptr);
+
+		ClientObject->ReplicationSystem = Client->GetReplicationSystem();
+
+		// Send multicast RPCs Server->Client
+		// This is a normal multicast rpc, it will be scheduled with replication of object
+		ServerObject->NetMulticast_MultiCastRPC();
+
+		// This is a send immediate multicast rpc, it should be scheduled using the OOB replication channel so it should be received before the normally flagged rpc.
+		ServerObject->NetMulticast_MultiCastRPCSendImmediate();
+
+		// Send and deliver packet, simulating a send from PostTickDispatch
+		Server->PreSendUpdate(UReplicationSystem::FSendUpdateParams {.SendPass = UE::Net::EReplicationSystemSendPass::PostTickDispatch});
+		Server->SendAndDeliverTo(Client, true);
+		Server->PostSendUpdate();
+
+		// Verify RPC reception, at this point we expect only the RPC flagged as ENetObjectAttachmentSendPolicyFlags::SendImmediate to have been received
+		UE_NET_ASSERT_EQ(ClientObject->NetMulticast_MultiCastRPCCallOrder, 0);
+		UE_NET_ASSERT_EQ(ClientObject->NetMulticast_MultiCastRPCSendImmediateCallOrder, 1);
+
+		// Send and deliver packet
+		Server->PreSendUpdate();
+		Server->SendAndDeliverTo(Client, true);
+		Server->PostSendUpdate();
+		
+		// Verify RPC reception, at this point we expect the other RPC to have been received as well.
+		UE_NET_ASSERT_EQ(ClientObject->NetMulticast_MultiCastRPCSendImmediateCallOrder, 1);
+		UE_NET_ASSERT_EQ(ClientObject->NetMulticast_MultiCastRPCCallOrder, 2);
+	}
+
 }
\ No newline at end of file
diff --git a/Engine/Plugins/Runtime/ReplicationSystemTestPlugin/Source/Private/Tests/ReplicationSystem/ReplicationSystemServerClientTestFixture.cpp b/Engine/Plugins/Runtime/ReplicationSystemTestPlugin/Source/Private/Tests/ReplicationSystem/ReplicationSystemServerClientTestFixture.cpp
--- a/Engine/Plugins/Runtime/ReplicationSystemTestPlugin/Source/Private/Tests/ReplicationSystem/ReplicationSystemServerClientTestFixture.cpp
+++ b/Engine/Plugins/Runtime/ReplicationSystemTestPlugin/Source/Private/Tests/ReplicationSystem/ReplicationSystemServerClientTestFixture.cpp
@@ -231,9 +231,15 @@ uint32 FReplicationSystemTestNode::AddConnection()
 	return Connection.ConnectionId;
 }
 
+void FReplicationSystemTestNode::PreSendUpdate(const UReplicationSystem::FSendUpdateParams& Params)
+{
+	ReplicationSystem->PreSendUpdate(Params);
+	CurrentSendPass = Params.SendPass;
+}
+
 void FReplicationSystemTestNode::PreSendUpdate()
 {
-	ReplicationSystem->PreSendUpdate(1.f);
+	PreSendUpdate(UReplicationSystem::FSendUpdateParams {.SendPass = EReplicationSystemSendPass::TickFlush, .DeltaSeconds = 1.f});
 }
 
 bool FReplicationSystemTestNode::SendUpdate(uint32 ConnectionId, const TCHAR* Desc)
@@ -252,6 +258,11 @@ bool FReplicationSystemTestNode::SendUpdate(uint32 ConnectionId, const TCHAR* De
 	const FDataStreamRecord* Record = nullptr;
 	UDataStream::FBeginWriteParameters BeginWriteParameters;
 
+	if (CurrentSendPass == EReplicationSystemSendPass::PostTickDispatch)
+	{
+		BeginWriteParameters.WriteMode = EDataStreamWriteMode::PostTickDispatch;
+	}
+
 	const bool bResult = (Connection.DataStreamManager->BeginWrite(BeginWriteParameters) != UDataStream::EWriteResult::NoData) && (Connection.DataStreamManager->WriteData(Context, Record)  != UDataStream::EWriteResult::NoData);
 	if (bResult)
 	{
@@ -284,6 +295,7 @@ bool FReplicationSystemTestNode::SendUpdate(uint32 ConnectionId, const TCHAR* De
 void FReplicationSystemTestNode::PostSendUpdate()
 {
 	ReplicationSystem->PostSendUpdate();
+	CurrentSendPass =	EReplicationSystemSendPass::Invalid;
 }
 
 void FReplicationSystemTestNode::DeliverTo(FReplicationSystemTestNode& Dest, uint32 LocalConnectionId, uint32 RemoteConnectionId, bool bDeliver)
diff --git a/Engine/Plugins/Runtime/ReplicationSystemTestPlugin/Source/Private/Tests/ReplicationSystem/ReplicationSystemServerClientTestFixture.h b/Engine/Plugins/Runtime/ReplicationSystemTestPlugin/Source/Private/Tests/ReplicationSystem/ReplicationSystemServerClientTestFixture.h
--- a/Engine/Plugins/Runtime/ReplicationSystemTestPlugin/Source/Private/Tests/ReplicationSystem/ReplicationSystemServerClientTestFixture.h
+++ b/Engine/Plugins/Runtime/ReplicationSystemTestPlugin/Source/Private/Tests/ReplicationSystem/ReplicationSystemServerClientTestFixture.h
@@ -13,6 +13,7 @@
 #include "Iris/DataStream/DataStreamManager.h"
 #include "Iris/ReplicationSystem/NetTokenDataStream.h"
 #include "Iris/ReplicationSystem/ReplicationDataStream.h"
+#include "Iris/ReplicationSystem/ReplicationSystem.h"
 
 #include "Net/Core/Misc/ResizableCircularQueue.h"
 
@@ -128,6 +129,7 @@ public:
 	uint32 AddConnection();
 
 	// System Update
+	void PreSendUpdate(const UReplicationSystem::FSendUpdateParams& Params);
 	void PreSendUpdate();
 	void PostSendUpdate();
 
@@ -155,6 +157,7 @@ public:
 	UReplicationSystem* ReplicationSystem;
 	UReplicatedTestObjectBridge* ReplicationBridge;
 	TArray<TStrongObjectPtr<UObject>> CreatedObjects;
+	EReplicationSystemSendPass CurrentSendPass = EReplicationSystemSendPass::Invalid;
 
 private:
 	TArray<FConnectionInfo> Connections;
diff --git a/Engine/Plugins/Runtime/ReplicationSystemTestPlugin/Source/Private/Tests/ReplicationSystem/TestReplicationSystem.cpp b/Engine/Plugins/Runtime/ReplicationSystemTestPlugin/Source/Private/Tests/ReplicationSystem/TestReplicationSystem.cpp
--- a/Engine/Plugins/Runtime/ReplicationSystemTestPlugin/Source/Private/Tests/ReplicationSystem/TestReplicationSystem.cpp
+++ b/Engine/Plugins/Runtime/ReplicationSystemTestPlugin/Source/Private/Tests/ReplicationSystem/TestReplicationSystem.cpp
@@ -1590,6 +1590,7 @@ UE_NET_TEST_FIXTURE(FReplicationSystemServerClientTestFixture, TestTearOffNextUp
 
 	// Trigger presend without send to add the objects to scope
 	Server->PreSendUpdate();
+	Server->PostSendUpdate();
 
 	UE_NET_ASSERT_EQ(uint16(1), NetRefHandleManager->GetNetObjectRefCount(ServerObjectInternalIndex));
 	UE_NET_ASSERT_EQ(uint16(1), NetRefHandleManager->GetNetObjectRefCount(SubObjectObjectInternalIndex));
@@ -1599,10 +1600,12 @@ UE_NET_TEST_FIXTURE(FReplicationSystemServerClientTestFixture, TestTearOffNextUp
 
 	// Update logic, object should be removed from scope but still exist as pending create in
 	Server->PreSendUpdate();
+	Server->PostSendUpdate();
 
 	UE_NET_ASSERT_EQ(uint16(1), NetRefHandleManager->GetNetObjectRefCount(ServerObjectInternalIndex));
 	UE_NET_ASSERT_EQ(uint16(1), NetRefHandleManager->GetNetObjectRefCount(SubObjectObjectInternalIndex));
 
+	Server->PreSendUpdate();
 	Server->SendAndDeliverTo(Client, true);
 	Server->PostSendUpdate();
 
diff --git a/Engine/Source/Runtime/Engine/Classes/Engine/NetDriver.h b/Engine/Source/Runtime/Engine/Classes/Engine/NetDriver.h
--- a/Engine/Source/Runtime/Engine/Classes/Engine/NetDriver.h
+++ b/Engine/Source/Runtime/Engine/Classes/Engine/NetDriver.h
@@ -2127,6 +2127,7 @@ private:
 	void CreateReplicationSystem(bool bInitAsClient);
 	void UpdateIrisReplicationViews() const;
 	void SendClientMoveAdjustments();
+	void PostDispatchSendUpdate();
 #endif
 
 	void InitNetTraceId();
diff --git a/Engine/Source/Runtime/Engine/Private/Net/Experimental/Iris/DataStreamChannel.cpp b/Engine/Source/Runtime/Engine/Private/Net/Experimental/Iris/DataStreamChannel.cpp
--- a/Engine/Source/Runtime/Engine/Private/Net/Experimental/Iris/DataStreamChannel.cpp
+++ b/Engine/Source/Runtime/Engine/Private/Net/Experimental/Iris/DataStreamChannel.cpp
@@ -155,6 +155,39 @@ void UDataStreamChannel::SendOpenBunch()
 #endif // UE_WITH_IRIS
 }
 
+void UDataStreamChannel::PostTickDispatch()
+{
+#if UE_WITH_IRIS
+	using namespace UE::Net;
+
+	if (!Connection->Driver->IsUsingIrisReplication() || !bHandshakeComplete)
+	{
+		return;
+	}
+
+	if (IsPacketWindowFull() || !Connection->HasReceivedClientPacket() || (Connection->Handler != nullptr && !Connection->Handler->IsFullyInitialized()))
+	{
+		return;
+	}
+
+	// We probably want separate bandwidth management for iris as we are not pre-filling sendbuffer before call to NetReady.
+	if (!IsNetReady(UE::Net::Private::bIrisSaturateBandwidth))
+	{
+		return;
+	}
+
+#if UE_NET_IRIS_CSV_STATS
+	CSV_SCOPED_TIMING_STAT_EXCLUSIVE(UDataStreamChannel_Tick);
+#endif
+
+	IRIS_PROFILER_SCOPE(UDataStreamChannel_Tick);
+	LLM_SCOPE_BYTAG(Iris);
+
+	WriteData(EDataStreamWriteMode::PostTickDispatch);
+
+#endif
+}
+
 void UDataStreamChannel::Tick()
 {
 #if UE_WITH_IRIS
@@ -191,6 +224,16 @@ void UDataStreamChannel::Tick()
 	IRIS_PROFILER_SCOPE(UDataStreamChannel_Tick);
 	LLM_SCOPE_BYTAG(Iris);
 
+	WriteData(EDataStreamWriteMode::Full);
+
+#endif
+}
+
+void UDataStreamChannel::WriteData(UE::Net::EDataStreamWriteMode WriteMode)
+{
+#if UE_WITH_IRIS
+	using namespace UE::Net;
+
 	// Limit the amount of bits to minimum of a bunch and our buffer. NetBitStreamWriter requires the number of bytes to be a multiple of 4.
 	const uint32 MaxBitCount = uint32(Connection->GetMaxSingleBunchSizeBits());
 	const uint32 MaxBytes = FPlatformMath::Min((MaxBitCount/32U)*4U, (uint32)sizeof(BitStreamBuffer));
@@ -198,6 +241,9 @@ void UDataStreamChannel::Tick()
 
 	// Try to determine if we have headroom to write more than a single packet if needed.
 	UDataStream::FBeginWriteParameters BeginWriteParams;
+	BeginWriteParams.WriteMode = WriteMode;
+
+	if (WriteMode == EDataStreamWriteMode::Full)
 	{
 		int32 CurrentQueuedBits = Connection->QueuedBits + Connection->SendBuffer.GetNumBits();
 		if (CurrentQueuedBits < 0)
@@ -310,6 +356,13 @@ void UDataStreamChannel::Tick()
 	// call end write to cleanup data initialized in BeginWrite	
 	DataStreamManager->EndWrite();
 
+	// If we did write data and the current WriteMode is PostTickDispatch we flush the packet here.
+	if (WriteMode == EDataStreamWriteMode::PostTickDispatch && bNeedsPreSendFlush)
+	{
+		IRIS_PROFILER_SCOPE(UDataStreamChannel_FlushNet);
+		Connection->FlushNet();	
+	}
+
 #endif // UE_WITH_IRIS
 }
 
diff --git a/Engine/Source/Runtime/Engine/Private/Net/Experimental/Iris/DataStreamChannel.h b/Engine/Source/Runtime/Engine/Private/Net/Experimental/Iris/DataStreamChannel.h
--- a/Engine/Source/Runtime/Engine/Private/Net/Experimental/Iris/DataStreamChannel.h
+++ b/Engine/Source/Runtime/Engine/Private/Net/Experimental/Iris/DataStreamChannel.h
@@ -12,6 +12,11 @@
 
 class UDataStreamManager;
 
+namespace UE::Net
+{
+	enum class EDataStreamWriteMode : unsigned;
+};
+
 UCLASS(transient, customConstructor, MinimalAPI)
 class UDataStreamChannel final : public UChannel
 {
@@ -19,6 +24,10 @@ class UDataStreamChannel final : public UChannel
 
 	UDataStreamChannel(const FObjectInitializer& ObjectInitializer = FObjectInitializer::Get());
 
+public:
+	/** Invoked from NetDriver::PostTickDispatch if we have any data that should be written from PostTickDispatch */
+	ENGINE_API void PostTickDispatch();
+
 private:
 	static void AddReferencedObjects(UObject* Object, FReferenceCollector& Collector);
 
@@ -62,6 +71,8 @@ private:
 
 	void SendOpenBunch();
 
+	void WriteData(UE::Net::EDataStreamWriteMode WriteMode);
+
 	TObjectPtr<UDataStreamManager> DataStreamManager = nullptr;
 
 	TResizableCircularQueue<FDataStreamChannelRecord> WriteRecords;
diff --git a/Engine/Source/Runtime/Engine/Private/NetDriver.cpp b/Engine/Source/Runtime/Engine/Private/NetDriver.cpp
--- a/Engine/Source/Runtime/Engine/Private/NetDriver.cpp
+++ b/Engine/Source/Runtime/Engine/Private/NetDriver.cpp
@@ -49,6 +49,7 @@
 #include "Iris/Core/IrisDebugging.h"
 #include "Iris/Core/IrisProfiler.h"
 #include "Iris/Core/IrisMemoryTracker.h"
+#include "Net/Experimental/Iris/DataStreamChannel.h"
 #include "Iris/ReplicationSystem/ReplicationSystem.h"
 #include "Iris/ReplicationSystem/ReplicationView.h"
 #include "Iris/ReplicationSystem/Filtering/NetObjectFilter.h"
@@ -949,7 +950,7 @@ void UNetDriver::TickFlush(float DeltaSeconds)
 		{
 			UpdateIrisReplicationViews();
 			SendClientMoveAdjustments();
-			ReplicationSystem->PreSendUpdate(DeltaSeconds);
+			ReplicationSystem->PreSendUpdate(UReplicationSystem::FSendUpdateParams { .SendPass = UE::Net::EReplicationSystemSendPass::TickFlush, .DeltaSeconds = DeltaSeconds });
 		}
 		else
 #endif // UE_WITH_IRIS
@@ -972,7 +973,7 @@ void UNetDriver::TickFlush(float DeltaSeconds)
 		if (ReplicationSystem)
 		{
 			UpdateIrisReplicationViews();
-			ReplicationSystem->PreSendUpdate(DeltaSeconds);
+			ReplicationSystem->PreSendUpdate(UReplicationSystem::FSendUpdateParams { .SendPass = UE::Net::EReplicationSystemSendPass::TickFlush, .DeltaSeconds = DeltaSeconds });
 		}
 	}
 #endif // UE_WITH_IRIS
@@ -1933,6 +1934,10 @@ void UNetDriver::PostTickDispatch()
 		}
 	}
 
+#if UE_WITH_IRIS
+	PostDispatchSendUpdate();
+#endif
+
 	if (ReplicationDriver)
 	{
 		ReplicationDriver->PostTickDispatch();
@@ -6567,6 +6572,39 @@ void UNetDriver::SendClientMoveAdjustments()
 		}
 	}
 }
+
+void UNetDriver::PostDispatchSendUpdate()
+{
+	if (ReplicationSystem)
+	{
+		if (!IsKnownChannelName(NAME_DataStream))
+		{
+			return;
+		}
+
+		ReplicationSystem->PreSendUpdate(UReplicationSystem::FSendUpdateParams { .SendPass = UE::Net::EReplicationSystemSendPass::PostTickDispatch });
+
+		ReplicationSystem->SendUpdate([this](TArrayView<uint32> ConnectionsToSend)
+		{
+			const int32 DataStreamChannelIndex = ChannelDefinitionMap[NAME_DataStream].StaticChannelIndex;
+
+			for (uint32 ConnId : ConnectionsToSend)
+			{
+				UNetConnection* NetConnection = GetConnectionById(ConnId);
+				if (NetConnection && NetConnection->Channels.IsValidIndex(DataStreamChannelIndex))
+				{
+					if (UDataStreamChannel* DataStreamChannel = Cast<UDataStreamChannel>(NetConnection->Channels[DataStreamChannelIndex]))
+					{
+						DataStreamChannel->PostTickDispatch();
+					}
+				}
+			}
+		});
+
+		ReplicationSystem->PostSendUpdate();
+	}
+}
+
 #endif // UE_WITH_IRIS
 
 void UNetDriver::InitNetTraceId()
diff --git a/Engine/Source/Runtime/Experimental/Iris/Core/Private/Iris/ReplicationSystem/AttachmentReplication.cpp b/Engine/Source/Runtime/Experimental/Iris/Core/Private/Iris/ReplicationSystem/AttachmentReplication.cpp
--- a/Engine/Source/Runtime/Experimental/Iris/Core/Private/Iris/ReplicationSystem/AttachmentReplication.cpp
+++ b/Engine/Source/Runtime/Experimental/Iris/Core/Private/Iris/ReplicationSystem/AttachmentReplication.cpp
@@ -225,6 +225,11 @@ bool FNetObjectAttachmentSendQueue::HasUnsent() const
 	return !UnreliableQueue.IsEmpty() || (ReliableQueue != nullptr && ReliableQueue->HasUnsentBlobs());
 }
 
+bool FNetObjectAttachmentSendQueue::HasUnsentUnreliable() const
+{
+	return !UnreliableQueue.IsEmpty();
+}
+
 bool FNetObjectAttachmentSendQueue::IsAllSentAndAcked() const
 {
 	return UnreliableQueue.IsEmpty() && (ReliableQueue == nullptr || ReliableQueue->IsAllSentAndAcked());
@@ -464,6 +469,17 @@ bool FNetObjectAttachmentsWriter::HasUnsentAttachments(ENetObjectAttachmentType
 	return Queue->HasUnsent();
 }
 
+bool FNetObjectAttachmentsWriter::HasUnsentUnreliableAttachments(ENetObjectAttachmentType Type, uint32 ObjectIndex) const
+{
+	const FNetObjectAttachmentSendQueue* Queue = GetQueue(Type, ObjectIndex);
+	if (Queue == nullptr)
+	{
+		return false;
+	}
+
+	return Queue->HasUnsentUnreliable();
+}
+
 bool FNetObjectAttachmentsWriter::IsAllSentAndAcked(ENetObjectAttachmentType Type, uint32 ObjectIndex) const
 {
 	const FNetObjectAttachmentSendQueue* Queue = GetQueue(Type, ObjectIndex);
diff --git a/Engine/Source/Runtime/Experimental/Iris/Core/Private/Iris/ReplicationSystem/AttachmentReplication.h b/Engine/Source/Runtime/Experimental/Iris/Core/Private/Iris/ReplicationSystem/AttachmentReplication.h
--- a/Engine/Source/Runtime/Experimental/Iris/Core/Private/Iris/ReplicationSystem/AttachmentReplication.h
+++ b/Engine/Source/Runtime/Experimental/Iris/Core/Private/Iris/ReplicationSystem/AttachmentReplication.h
@@ -75,6 +75,8 @@ public:
 
 	bool HasUnsent() const;
 
+	bool HasUnsentUnreliable() const;
+
 	bool IsSafeToDestroy() const;
 
 	bool IsAllSentAndAcked() const;
@@ -112,6 +114,9 @@ public:
 	bool Enqueue(ENetObjectAttachmentType Type, uint32 ObjectIndex, TArrayView<const TRefCountPtr<FNetBlob>> Attachments);
 
 	bool HasUnsentAttachments(ENetObjectAttachmentType Type, uint32 ObjectIndex) const;
+
+	bool HasUnsentUnreliableAttachments(ENetObjectAttachmentType Type, uint32 ObjectIndex) const;
+
 	// Whether all queued attachments have been sent and that all reliable ones have been acked.
 	bool IsAllSentAndAcked(ENetObjectAttachmentType Type, uint32 ObjectIndex) const;
 
diff --git a/Engine/Source/Runtime/Experimental/Iris/Core/Private/Iris/ReplicationSystem/NetBlob/NetBlobManager.cpp b/Engine/Source/Runtime/Experimental/Iris/Core/Private/Iris/ReplicationSystem/NetBlob/NetBlobManager.cpp
--- a/Engine/Source/Runtime/Experimental/Iris/Core/Private/Iris/ReplicationSystem/NetBlob/NetBlobManager.cpp
+++ b/Engine/Source/Runtime/Experimental/Iris/Core/Private/Iris/ReplicationSystem/NetBlob/NetBlobManager.cpp
@@ -236,6 +236,11 @@ bool FNetBlobManager::GetOwnerAndSubObjectIndicesFromHandle(FNetRefHandle RefHan
 	return true;
 }
 
+void FNetBlobManager::ProcessOOBNetObjectAttachmentSendQueue(FNetBitArray& OutConnetionsPendingImmediateSend)
+{
+	AttachmentSendQueue.PrepareAndProcessOOBAttachmentQueue(Connections, NetRefHandleManager, OutConnetionsPendingImmediateSend);
+}
+
 void FNetBlobManager::ProcessNetObjectAttachmentSendQueue(EProcessMode ProcessMode)
 {
 	AttachmentSendQueue.PrepareProcessQueue(Connections, NetRefHandleManager);
@@ -316,7 +321,10 @@ void FNetBlobManager::FNetObjectAttachmentSendQueue::Init(FNetBlobManager* InMan
 
 void FNetBlobManager::FNetObjectAttachmentSendQueue::Enqueue(uint32 ConnectionId, FInternalNetRefIndex OwnerIndex, FInternalNetRefIndex SubObjectIndex, const TRefCountPtr<FNetObjectAttachment>& Attachment, ENetObjectAttachmentSendPolicyFlags SendFlags)
 {
-	FNetObjectAttachmentQueueEntry& QueueEntry = AttachmentQueue.AddDefaulted_GetRef();
+	const bool bScheduleUsingOOBAttachmentQueue = EnumHasAnyFlags(SendFlags, ENetObjectAttachmentSendPolicyFlags::ScheduleAsOOB);
+	FQueue& TargetQueue = bScheduleUsingOOBAttachmentQueue ? ScheduleAsOOBAttachmentQueue : AttachmentQueue;
+
+	FNetObjectAttachmentQueueEntry& QueueEntry = TargetQueue.AddDefaulted_GetRef();
 	QueueEntry.ConnectionId = ConnectionId;
 	QueueEntry.OwnerIndex = OwnerIndex;
 	QueueEntry.SubObjectIndex = SubObjectIndex;
@@ -326,7 +334,10 @@ void FNetBlobManager::FNetObjectAttachmentSendQueue::Enqueue(uint32 ConnectionId
 
 void FNetBlobManager::FNetObjectAttachmentSendQueue::Enqueue(FInternalNetRefIndex OwnerIndex, FInternalNetRefIndex SubObjectIndex, const TRefCountPtr<FNetObjectAttachment>& Attachment, ENetObjectAttachmentSendPolicyFlags SendFlags)
 {
-	FNetObjectAttachmentQueueEntry& QueueEntry = AttachmentQueue.AddDefaulted_GetRef();
+	const bool bScheduleUsingOOBAttachmentQueue = EnumHasAnyFlags(SendFlags, ENetObjectAttachmentSendPolicyFlags::ScheduleAsOOB);
+	FQueue& TargetQueue = bScheduleUsingOOBAttachmentQueue ? ScheduleAsOOBAttachmentQueue : AttachmentQueue;
+
+	FNetObjectAttachmentQueueEntry& QueueEntry = TargetQueue.AddDefaulted_GetRef();
 	QueueEntry.ConnectionId = 0;
 	QueueEntry.OwnerIndex = OwnerIndex;
 	QueueEntry.SubObjectIndex = SubObjectIndex;
@@ -336,24 +347,77 @@ void FNetBlobManager::FNetObjectAttachmentSendQueue::Enqueue(FInternalNetRefInde
 	bHasMulticastAttachments = true;
 }
 
-void FNetBlobManager::FNetObjectAttachmentSendQueue::PrepareProcessQueue(FReplicationConnections* InConnections, const FNetRefHandleManager* InNetRefHandleManager)
+void FNetBlobManager::FNetObjectAttachmentSendQueue::PrepareAndProcessOOBAttachmentQueue(FReplicationConnections* InConnections, const FNetRefHandleManager* InNetRefHandleManager, FNetBitArray& OutConnectionsPendingImmediateSend)
 {
+	check(!ProcessContext.IsValid());
+
 	if (ProcessContext.IsValid())
 	{
 		return;
 	}
 
+	if (ScheduleAsOOBAttachmentQueue.Num() <= 0)
+	{
+		OutConnectionsPendingImmediateSend.Reset();
+		return;
+	}
+
+	// Init context to process OOBAttachmentQueue
+	ProcessContext.QueueToProcess = &ScheduleAsOOBAttachmentQueue;
 	ProcessContext.Connections = InConnections;
-	ProcessContext.NetRefHandleManager = InNetRefHandleManager;
+	ProcessContext.NetRefHandleManager = InNetRefHandleManager;	
+	ProcessContext.AttachmentsToObjectsInScope.Init(ScheduleAsOOBAttachmentQueue.Num());
+	ProcessContext.ConnectionsPendingSendInPostDispatch.Init(InConnections->GetValidConnections().GetNumBits());
 
-	ProcessContext.AttachmentsToObjectsGoingOutOfScope.Init(AttachmentQueue.Num());
-	ProcessContext.AttachmentsToObjectsInScope.Init(AttachmentQueue.Num());
+	if (bHasMulticastAttachments)
+	{
+		const FNetBitArray& ValidConnections = InConnections->GetValidConnections();
+		const FNetBitArrayView ReplicatingConnections = MakeNetBitArrayView(ValidConnections);
+
+		ProcessContext.ConnectionIds.SetNum(ReplicatingConnections.CountSetBits());
+		ReplicatingConnections.GetSetBitIndices(0, ~0, ProcessContext.ConnectionIds.GetData(), ProcessContext.ConnectionIds.Num());
+	}
+
+	uint32 CurrentEntryIndex = 0U;
+	for (const FNetObjectAttachmentQueueEntry& Entry : MakeArrayView(ScheduleAsOOBAttachmentQueue))
+	{
+		ProcessContext.AttachmentsToObjectsInScope.SetBit(CurrentEntryIndex);
+		++CurrentEntryIndex;
+	}
+
+	ProcessQueue(EProcessMode::ProcessObjectsInScope);
+
+	OutConnectionsPendingImmediateSend.InitAndCopy(ProcessContext.ConnectionsPendingSendInPostDispatch);
+
+	// Reset Context
+	ProcessContext.Reset();
+	ScheduleAsOOBAttachmentQueue.Reset();	
+}
+
+void FNetBlobManager::FNetObjectAttachmentSendQueue::PrepareProcessQueue(FReplicationConnections* InConnections, const FNetRefHandleManager* InNetRefHandleManager)
+{
+	if (ProcessContext.IsValid())
+	{
+		return;
+	}
+
+	// If we have entries in the ScheduleAsOOBAttachmentQueue that was not processed during PostTickDispatch (might have been posted after PostTickDispatch) we make sure to process them now.
+	AttachmentQueue.Append(ScheduleAsOOBAttachmentQueue);
+	ScheduleAsOOBAttachmentQueue.Reset();
 
 	if (AttachmentQueue.Num() <= 0)
 	{
 		return;
 	}
 
+	// Init context
+	ProcessContext.QueueToProcess = &AttachmentQueue;
+	ProcessContext.Connections = InConnections;
+	ProcessContext.NetRefHandleManager = InNetRefHandleManager;	
+	ProcessContext.AttachmentsToObjectsGoingOutOfScope.Init(AttachmentQueue.Num());
+	ProcessContext.AttachmentsToObjectsInScope.Init(AttachmentQueue.Num());
+	ProcessContext.ConnectionsPendingSendInPostDispatch.Init(InConnections->GetValidConnections().GetNumBits());
+
 	if (bHasMulticastAttachments)
 	{
 		const FNetBitArray& ValidConnections = InConnections->GetValidConnections();
@@ -407,13 +471,19 @@ void FNetBlobManager::FNetObjectAttachmentSendQueue::ResetProcessQueue()
 
 void FNetBlobManager::FNetObjectAttachmentSendQueue::ProcessQueue(EProcessMode ProcessMode)
 {
+	if (!ProcessContext.IsValid())
+	{
+		return;
+	}
+
 	const FNetBitArray& IndicesToProcess = (ProcessMode == EProcessMode::ProcessObjectsGoingOutOfScope) ? ProcessContext.AttachmentsToObjectsGoingOutOfScope : ProcessContext.AttachmentsToObjectsInScope;
 
 	TArray<TRefCountPtr<FNetBlob>> PartialNetBlobs;
-	const FNetObjectAttachmentQueueEntry* Entries = AttachmentQueue.GetData();
+	const FNetObjectAttachmentQueueEntry* Entries = ProcessContext.QueueToProcess->GetData();
+	const uint32 NumEntries = ProcessContext.QueueToProcess->Num();
 
 	// Verify that we have not missed to prepare the process context
-	check(ProcessContext.IsValid() && IndicesToProcess.GetNumBits() == AttachmentQueue.Num());
+	check(ProcessContext.IsValid() && IndicesToProcess.GetNumBits() == NumEntries);
 
 	IndicesToProcess.ForAllSetBits([this, Entries, &PartialNetBlobs](uint32 Index)
 	{
@@ -426,6 +496,7 @@ void FNetBlobManager::FNetObjectAttachmentSendQueue::ProcessQueue(EProcessMode P
 
 		const bool bMulticast = Entry.ConnectionId == 0;
 		const bool bHasConnectionSpecificSerialization = ReplicationStateDescriptor && EnumHasAnyFlags(ReplicationStateDescriptor->Traits, EReplicationStateTraits::HasConnectionSpecificSerialization);
+		const bool bSendInPostTickDispatch = EnumHasAnyFlags(Entry.SendFlags, ENetObjectAttachmentSendPolicyFlags::SendInPostTickDispatch);
 
 		const bool bShouldSendAttachmentsWithObject = EnumHasAnyFlags(Entry.SendFlags, ENetObjectAttachmentSendPolicyFlags::ScheduleAsOOB) ? false : Manager->bSendAttachmentsWithObject;
 
@@ -468,14 +539,22 @@ void FNetBlobManager::FNetObjectAttachmentSendQueue::ProcessQueue(EProcessMode P
 
 				FReplicationConnection* Connection = ProcessContext.Connections->GetConnection(ConnectionId);
 				// We're only iterating over valid connections so the Connection pointer must be valid.
-				Connection->ReplicationWriter->QueueNetObjectAttachments(Entry.OwnerIndex, Entry.SubObjectIndex, AttachmentsView, Entry.SendFlags);
+				const bool bWasEnqueued = Connection->ReplicationWriter->QueueNetObjectAttachments(Entry.OwnerIndex, Entry.SubObjectIndex, AttachmentsView, Entry.SendFlags);
+				if (bWasEnqueued && bSendInPostTickDispatch)
+				{
+					ProcessContext.ConnectionsPendingSendInPostDispatch.SetBit(ConnectionId);
+				}
 			}
 		}
 		else
 		{
 			if (FReplicationConnection* Connection = ProcessContext.Connections->GetConnection(Entry.ConnectionId))
 			{
-				Connection->ReplicationWriter->QueueNetObjectAttachments(Entry.OwnerIndex, Entry.SubObjectIndex, AttachmentsView, Entry.SendFlags);
+				const bool bWasEnqueued = Connection->ReplicationWriter->QueueNetObjectAttachments(Entry.OwnerIndex, Entry.SubObjectIndex, AttachmentsView, Entry.SendFlags);
+				if (bWasEnqueued && bSendInPostTickDispatch)
+				{
+					ProcessContext.ConnectionsPendingSendInPostDispatch.SetBit(Entry.ConnectionId);
+				}	
 			}
 		}
 	});
diff --git a/Engine/Source/Runtime/Experimental/Iris/Core/Private/Iris/ReplicationSystem/NetBlob/NetBlobManager.h b/Engine/Source/Runtime/Experimental/Iris/Core/Private/Iris/ReplicationSystem/NetBlob/NetBlobManager.h
--- a/Engine/Source/Runtime/Experimental/Iris/Core/Private/Iris/ReplicationSystem/NetBlob/NetBlobManager.h
+++ b/Engine/Source/Runtime/Experimental/Iris/Core/Private/Iris/ReplicationSystem/NetBlob/NetBlobManager.h
@@ -55,6 +55,9 @@ public:
 		ProcessObjectsInScope,
 	};
 
+	FNetBitArrayView GetConnectionsPendingImmediateSend() const;
+
+	void ProcessOOBNetObjectAttachmentSendQueue(FNetBitArray& OutConnectionsPendingImmediateSend);
 	void ProcessNetObjectAttachmentSendQueue(EProcessMode ProcessMode);
 	void ResetNetObjectAttachmentSendQueue();
 
@@ -89,6 +92,7 @@ private:
 		void PrepareProcessQueue(FReplicationConnections* InConnections, const FNetRefHandleManager* InNetRefHandleManager);
 		void ProcessQueue(EProcessMode ProcessMode);
 		void ResetProcessQueue();
+		void PrepareAndProcessOOBAttachmentQueue(FReplicationConnections* InConnections, const FNetRefHandleManager* InNetRefHandleManager, FNetBitArray& OutConnetionsPendingImmediateSend);
 	
 	private:
 		struct FNetObjectAttachmentQueueEntry
@@ -105,21 +109,26 @@ private:
 
 		FNetBlobManager* Manager;
 		FQueue AttachmentQueue;
+		FQueue ScheduleAsOOBAttachmentQueue;		
 		bool bHasMulticastAttachments;
 
 		struct FProcessQueueContext
 		{
 			FNetBitArray AttachmentsToObjectsGoingOutOfScope;
 			FNetBitArray AttachmentsToObjectsInScope;
+			FNetBitArray ConnectionsPendingSendInPostDispatch;
 
 			TArray<uint32> ConnectionIds;
 			FReplicationConnections* Connections = nullptr;
 			const FNetRefHandleManager* NetRefHandleManager = nullptr;
+			FQueue* QueueToProcess = nullptr;
 
 			void Reset()
 			{
 				Connections = nullptr;
 				NetRefHandleManager = nullptr;
+				QueueToProcess = nullptr;
+				ConnectionsPendingSendInPostDispatch.Reset();
 			}
 
 			bool IsValid() const { return NetRefHandleManager != nullptr; }
diff --git a/Engine/Source/Runtime/Experimental/Iris/Core/Private/Iris/ReplicationSystem/ReplicationSystem.cpp b/Engine/Source/Runtime/Experimental/Iris/Core/Private/Iris/ReplicationSystem/ReplicationSystem.cpp
--- a/Engine/Source/Runtime/Experimental/Iris/Core/Private/Iris/ReplicationSystem/ReplicationSystem.cpp
+++ b/Engine/Source/Runtime/Experimental/Iris/Core/Private/Iris/ReplicationSystem/ReplicationSystem.cpp
@@ -45,14 +45,15 @@ namespace UE::Net::Private
 class FReplicationSystemImpl
 {
 public:
+	TMap<FObjectKey, ENetObjectAttachmentSendPolicyFlags> AttachmentSendPolicyFlags;
 	UReplicationSystem* ReplicationSystem;
 	FReplicationSystemInternal ReplicationSystemInternal;
 	uint64 IrisDebugHelperDummy = 0U;
 	FNetObjectGroupHandle NotReplicatedNetObjectGroupHandle;
 	FNetObjectGroupHandle NetGroupOwnerNetObjectGroupHandle;
 	FNetObjectGroupHandle NetGroupReplayNetObjectGroupHandle;
-
-	TMap<FObjectKey, ENetObjectAttachmentSendPolicyFlags> AttachmentSendPolicyFlags;
+	FNetBitArray ConnectionsPendingPostTickDispatchSend;
+	EReplicationSystemSendPass CurrentSendPass = EReplicationSystemSendPass::Invalid;
 
 	explicit FReplicationSystemImpl(UReplicationSystem* InReplicationSystem, const UReplicationSystem::FReplicationSystemParams& Params)
 	: ReplicationSystem(InReplicationSystem)
@@ -231,6 +232,8 @@ public:
 		{
 			ReplicationSystemInternal.GetForwardNetRPCCallMulticastDelegate().Add(Params.ForwardNetRPCCallDelegate);
 		}
+
+		ConnectionsPendingPostTickDispatchSend.Init(ReplicationSystemInternal.GetConnections().GetMaxConnectionCount());
 	}
 
 	void Deinit()
@@ -443,15 +446,18 @@ public:
 		NetBlobManager.ProcessNetObjectAttachmentSendQueue(ProcessMode);
 	}
 
-	void ResetNetObjectAttachmentSendQueue()
+	void ProcessOOBNetObjectAttachmentSendQueue()
 	{
+		IRIS_PROFILER_SCOPE(FReplicationSystem_ProcessOOBNetObjectAttachmentSendQueue);
+
 		FNetBlobManager& NetBlobManager = ReplicationSystemInternal.GetNetBlobManager();
-		NetBlobManager.ResetNetObjectAttachmentSendQueue();
+		NetBlobManager.ProcessOOBNetObjectAttachmentSendQueue(ConnectionsPendingPostTickDispatchSend);
 	}
 
-	// send data
-	void SendUpdate()
+	void ResetNetObjectAttachmentSendQueue()
 	{
+		FNetBlobManager& NetBlobManager = ReplicationSystemInternal.GetNetBlobManager();
+		NetBlobManager.ResetNetObjectAttachmentSendQueue();
 	}
 
 	void AddConnection(uint32 ConnectionId)
@@ -604,125 +610,183 @@ const UE::Net::Private::FReplicationSystemInternal* UReplicationSystem::GetRepli
 	return &Impl->ReplicationSystemInternal;
 }
 
-void UReplicationSystem::PreSendUpdate(float DeltaSeconds)
+void UReplicationSystem::PreSendUpdate(const FSendUpdateParams& Params)
 {
 	using namespace UE::Net;
 	using namespace UE::Net::Private;
 
 	IRIS_PROFILER_SCOPE(FReplicationSystem_PreSendUpdate);
 
-	ElapsedTime += DeltaSeconds;
-
-#if !UE_BUILD_SHIPPING
-	// Force a integrity check of all replicated instances
-	if (bDoCollectGarbage || ReplicationSystemCVars::bForcePruneBeforeUpdate)
-	{
-		CollectGarbage();
-	}
-#endif
+	ensureAlways(Impl->CurrentSendPass == EReplicationSystemSendPass::Invalid);
+	Impl->CurrentSendPass = Params.SendPass;
 
 	FReplicationSystemInternal& InternalSys = Impl->ReplicationSystemInternal;
 
-	// $IRIS TODO. There may be some throttling of connections to tick that we should take into account.
-	const FNetBitArrayView& ReplicatingConnections = MakeNetBitArrayView(Impl->ReplicationSystemInternal.GetConnections().GetValidConnections());
+	// Most systems are only updated during the normal TickFlush
+	if (Impl->CurrentSendPass ==  EReplicationSystemSendPass::TickFlush)
+	{
+		ElapsedTime += Params.DeltaSeconds;
+
+		// $IRIS TODO. There may be some throttling of connections to tick that we should take into account.
+		const FNetBitArrayView& ReplicatingConnections = MakeNetBitArrayView(Impl->ReplicationSystemInternal.GetConnections().GetValidConnections());
 
 #if UE_NET_IRIS_CSV_STATS && CSV_PROFILER
-	{
-		FNetSendStats& SendStats = InternalSys.GetSendStats();
-		SendStats.Reset();
-		SendStats.SetNumberOfReplicatingConnections(ReplicatingConnections.CountSetBits());
-	}
+		{
+			FNetSendStats& SendStats = InternalSys.GetSendStats();
+			SendStats.Reset();
+			SendStats.SetNumberOfReplicatingConnections(ReplicatingConnections.CountSetBits());
+		}
+#endif
+	
+#if !UE_BUILD_SHIPPING
+		// Force a integrity check of all replicated instances
+		if (bDoCollectGarbage || ReplicationSystemCVars::bForcePruneBeforeUpdate)
+		{
+			CollectGarbage();
+		}
 #endif
 
-	if (bAllowObjectReplication)
-	{
-		UE_NET_TRACE_FRAME_STATSCOUNTER(GetId(), ReplicationSystem.ReplicatedObjectCount, InternalSys.GetNetRefHandleManager().GetActiveObjectCount(), ENetTraceVerbosity::Verbose);
+		if (bAllowObjectReplication)
+		{
+			UE_NET_TRACE_FRAME_STATSCOUNTER(GetId(), ReplicationSystem.ReplicatedObjectCount, InternalSys.GetNetRefHandleManager().GetActiveObjectCount(), ENetTraceVerbosity::Verbose);
 
-		// Tell systems we are starting PreSendUpdate
-		Impl->StartPreSendUpdate();
+			// Tell systems we are starting PreSendUpdate
+			Impl->StartPreSendUpdate();
 
-		// Refresh the dirty objects we were told about.
-		Impl->UpdateDirtyObjectList();
+			// Refresh the dirty objects we were told about.
+			Impl->UpdateDirtyObjectList();
 
-		// Update world locations. We need this to happen before both filtering and prioritization.
-		Impl->UpdateWorldLocations();
+			// Update world locations. We need this to happen before both filtering and prioritization.
+			Impl->UpdateWorldLocations();
 
-		// Update filters, reduce the top-level scoped object list and set each connection's scope.
-		Impl->UpdateFilterPrePoll();
+			// Update filters, reduce the top-level scoped object list and set each connection's scope.
+			Impl->UpdateFilterPrePoll();
 
-		// Invoke any operations we need to do before copying state data
-		Impl->CallPreSendUpdate(DeltaSeconds);
+			// Invoke any operations we need to do before copying state data
+			Impl->CallPreSendUpdate(Params.DeltaSeconds);
 
-		// Finalize the dirty list with objects set dirty during the poll phase
-		Impl->UpdateDirtyListPostPoll();
+			// Finalize the dirty list with objects set dirty during the poll phase
+			Impl->UpdateDirtyListPostPoll();
 
-		// Update conditionals
-		Impl->UpdateConditionals();
+			// Update conditionals
+			Impl->UpdateConditionals();
 
-		// Copy dirty state data. We need this to happen before both filtering and prioritization
-		Impl->CopyDirtyStateData();
+			// Copy dirty state data. We need this to happen before both filtering and prioritization
+			Impl->CopyDirtyStateData();
 
-		// We must process all attachments to objects going out of scope before we update the scope
-		Impl->ProcessNetObjectAttachmentSendQueue(FNetBlobManager::EProcessMode::ProcessObjectsGoingOutOfScope);
+			// We must process all attachments to objects going out of scope before we update the scope
+			Impl->ProcessNetObjectAttachmentSendQueue(FNetBlobManager::EProcessMode::ProcessObjectsGoingOutOfScope);
 
-		// Update filtering and scope for all connections
-		Impl->UpdateFilterPostPoll();
+			// Update filtering and scope for all connections
+			Impl->UpdateFilterPostPoll();
 
-		// Propagate dirty changes to all connections
-		Impl->PropagateDirtyChanges();
-	}
+			// Propagate dirty changes to all connections
+			Impl->PropagateDirtyChanges();
+		}
 
-	// Forward attachments to the connections after scope update
-	Impl->ProcessNetObjectAttachmentSendQueue(FNetBlobManager::EProcessMode::ProcessObjectsInScope);
-	Impl->ResetNetObjectAttachmentSendQueue();
+		// Forward attachments to the connections after scope update
+		Impl->ProcessNetObjectAttachmentSendQueue(FNetBlobManager::EProcessMode::ProcessObjectsInScope);
+		Impl->ResetNetObjectAttachmentSendQueue();
 
-	if (bAllowObjectReplication)
-	{
-		// Update object priorities
-		Impl->UpdatePrioritization(ReplicatingConnections);
+		if (bAllowObjectReplication)
+		{
+			// Update object priorities
+			Impl->UpdatePrioritization(ReplicatingConnections);
 
-		// Delta compression preparations before send
+			// Delta compression preparations before send
+			{
+				FDeltaCompressionBaselineManagerPreSendUpdateParams UpdateParams;
+				UpdateParams.ChangeMaskCache = &InternalSys.GetChangeMaskCache();
+				InternalSys.GetDeltaCompressionBaselineManager().PreSendUpdate(UpdateParams);
+			}
+		}
+
+		// Destroy objects pending destroy
 		{
-			FDeltaCompressionBaselineManagerPreSendUpdateParams UpdateParams;
-			UpdateParams.ChangeMaskCache = &InternalSys.GetChangeMaskCache();
-			InternalSys.GetDeltaCompressionBaselineManager().PreSendUpdate(UpdateParams);
+			Impl->UpdateUnresolvableReferenceTracking();
+			InternalSys.GetNetRefHandleManager().DestroyObjectsPendingDestroy();
 		}
 	}
-
-	// Destroy objects pending destroy
+	else if (Impl->CurrentSendPass == EReplicationSystemSendPass::PostTickDispatch)
 	{
-		Impl->UpdateUnresolvableReferenceTracking();
-		InternalSys.GetNetRefHandleManager().DestroyObjectsPendingDestroy();
+		// Forward attachments scheduled to use the OOBChannel and mark connetions needing immediate send
+		Impl->ProcessOOBNetObjectAttachmentSendQueue();
 	}
 }
 
-void UReplicationSystem::SendUpdate()
+void UReplicationSystem::PreSendUpdate(float DeltaSeconds)
 {
-	Impl->SendUpdate();
+	PreSendUpdate(FSendUpdateParams {.SendPass = UE::Net::EReplicationSystemSendPass::TickFlush, .DeltaSeconds = DeltaSeconds });
 }
 
-void UReplicationSystem::PostSendUpdate()
+IRISCORE_API void UReplicationSystem::SendUpdate(TFunctionRef<void(TArrayView<uint32>)> SendFunction)
 {
+	using namespace UE::Net;
 	using namespace UE::Net::Private;
 
-	IRIS_PROFILER_SCOPE(FReplicationSystem_PostSendUpdate);
-	
-	Impl->ResetObjectStateDirtiness();
+	if (!ensure(Impl->CurrentSendPass != EReplicationSystemSendPass::Invalid))
+	{
+		return;
+	}
 
-	Impl->EndPostSendUpdate();
+	UE::Net::Private::FReplicationConnections& Connections = Impl->ReplicationSystemInternal.GetConnections();
+	const FNetBitArray& ReplicatingConnections = Connections.GetValidConnections();
 
-	if (bAllowObjectReplication)
+	TArray<uint32, TInlineAllocator<128>> ConnectionToUpdate;
+	
+	if (Impl->CurrentSendPass == EReplicationSystemSendPass::TickFlush)
 	{
-		FDeltaCompressionBaselineManagerPostSendUpdateParams UpdateParams;
-		Impl->ReplicationSystemInternal.GetDeltaCompressionBaselineManager().PostSendUpdate(UpdateParams);
+		// This is currently handled when ticking NetDriver->NetConnection->Channels.
+
+		ConnectionToUpdate.SetNum(ReplicatingConnections.CountSetBits());
+		Connections.GetValidConnections().GetSetBitIndices(0U, ~0U, ConnectionToUpdate.GetData(), ConnectionToUpdate.Num());
 	}
+	else if (Impl->CurrentSendPass == EReplicationSystemSendPass::PostTickDispatch)
+	{
+		// We only need to send data to connections that has data to send in PostTickDispatch
 
-#if UE_NET_IRIS_CSV_STATS && CSV_PROFILER
+		FNetBitArray::ForAllSetBits(Impl->ConnectionsPendingPostTickDispatchSend, ReplicatingConnections, FNetBitArray::AndOp, [&ConnectionToUpdate](uint32 ConnId) { ConnectionToUpdate.Add(ConnId);});
+		Impl->ConnectionsPendingPostTickDispatchSend.Reset();
+	}
+
+	SendFunction(MakeArrayView(ConnectionToUpdate));
+}
+
+void UReplicationSystem::PostSendUpdate()
+{
+	using namespace UE::Net;
+	using namespace UE::Net::Private;
+
+	IRIS_PROFILER_SCOPE(FReplicationSystem_PostSendUpdate);
+
+	if (!ensure(Impl->CurrentSendPass != EReplicationSystemSendPass::Invalid))
 	{
-		UE::Net::FNetSendStats& SendStats = Impl->ReplicationSystemInternal.GetSendStats();
-		SendStats.ReportCsvStats();
+		return;
 	}
+
+	// Most systems are only updated during the normal TickFlush
+	if (Impl->CurrentSendPass == EReplicationSystemSendPass::TickFlush)
+	{
+		Impl->ResetObjectStateDirtiness();
+
+		Impl->EndPostSendUpdate();
+
+		if (bAllowObjectReplication)
+		{
+			FDeltaCompressionBaselineManagerPostSendUpdateParams UpdateParams;
+			Impl->ReplicationSystemInternal.GetDeltaCompressionBaselineManager().PostSendUpdate(UpdateParams);
+		}
+
+#if UE_NET_IRIS_CSV_STATS && CSV_PROFILER
+		{
+			UE::Net::FNetSendStats& SendStats = Impl->ReplicationSystemInternal.GetSendStats();
+			SendStats.ReportCsvStats();
+		}
 #endif
+
+	}
+
+	Impl->CurrentSendPass = EReplicationSystemSendPass::Invalid;
 }
 
 void UReplicationSystem::PostGarbageCollection()
diff --git a/Engine/Source/Runtime/Experimental/Iris/Core/Private/Iris/ReplicationSystem/ReplicationWriter.cpp b/Engine/Source/Runtime/Experimental/Iris/Core/Private/Iris/ReplicationSystem/ReplicationWriter.cpp
--- a/Engine/Source/Runtime/Experimental/Iris/Core/Private/Iris/ReplicationSystem/ReplicationWriter.cpp
+++ b/Engine/Source/Runtime/Experimental/Iris/Core/Private/Iris/ReplicationSystem/ReplicationWriter.cpp
@@ -308,19 +308,19 @@ bool FReplicationWriter::IsReplicationEnabled() const
 
 // $IRIS TODO : May need to introduce queue and send behaviors. For example one may want to send only with object.
 // One may not want to send unless the object is replicated very soon etc.
-void FReplicationWriter::QueueNetObjectAttachments(FInternalNetRefIndex OwnerInternalIndex, FInternalNetRefIndex SubObjectInternalIndex, TArrayView<const TRefCountPtr<FNetBlob>> InAttachments, ENetObjectAttachmentSendPolicyFlags SendFlags)
+bool FReplicationWriter::QueueNetObjectAttachments(FInternalNetRefIndex OwnerInternalIndex, FInternalNetRefIndex SubObjectInternalIndex, TArrayView<const TRefCountPtr<FNetBlob>> InAttachments, ENetObjectAttachmentSendPolicyFlags SendFlags)
 {
 	if (InAttachments.Num() <= 0)
 	{
 		ensureMsgf(false, TEXT("%s"), TEXT("QueueNetObjectAttachments expects at least one attachment."));
-		return;
+		return false;
 	}
 
 	const bool bObjectInScope = ObjectsInScope.GetBit(OwnerInternalIndex);
 	if (!bObjectInScope && !Parameters.bAllowSendingAttachmentsToObjectsNotInScope)
 	{
 		UE_CLOG_REPLICATIONWRITER_WARNING(bWarnAboutDroppedAttachmentsToObjectsNotInScope, TEXT("Dropping %s attachment due to object ( InternalIndex: %u ) not in scope."), (EnumHasAnyFlags(InAttachments[0]->GetCreationInfo().Flags, ENetBlobFlags::Reliable) ? TEXT("reliable") : TEXT("unreliable")), OwnerInternalIndex);
-		return;
+		return false;
 	}
 	
 	// Route attachments flagged with ScheduleAsOOB through OOB channel if we have started replicating the owner.
@@ -328,20 +328,20 @@ void FReplicationWriter::QueueNetObjectAttachments(FInternalNetRefIndex OwnerInt
 	if (bScheduleUsingOOBChannel && (GetReplicationInfo(OwnerInternalIndex).GetState() < EReplicatedObjectState::WaitOnCreateConfirmation || GetReplicationInfo(OwnerInternalIndex).GetState() >= EReplicatedObjectState::PendingDestroy))
 	{
 		UE_CLOG_REPLICATIONWRITER_WARNING(bWarnAboutDroppedAttachmentsToObjectsNotInScope, TEXT("Dropping attachment scheduled as ScheduleAsOOB due to object ( InternalIndex: %u ) not in replicated state."),  OwnerInternalIndex);
-		return;
+		return false;
 	}
 
 	const uint32 TargetIndex = (bObjectInScope && !bScheduleUsingOOBChannel) ? (SubObjectInternalIndex != FNetRefHandleManager::InvalidInternalIndex ? SubObjectInternalIndex : OwnerInternalIndex) : ObjectIndexForOOBAttachment;
 	ENetObjectAttachmentType AttachmentType = ((bObjectInScope && !bScheduleUsingOOBChannel) ? ENetObjectAttachmentType::Normal : ENetObjectAttachmentType::OutOfBand);
 	if (!Attachments.Enqueue(AttachmentType, TargetIndex, InAttachments))
 	{
-		return;
+		return false;
 	}
 
 	// There's a special case for out of band attachments, we don't need to mark anything dirty.
 	if (IsObjectIndexForOOBAttachment(TargetIndex))
 	{
-		return;
+		return true;
 	}
 
 	FReplicationInfo& TargetInfo = GetReplicationInfo(TargetIndex);
@@ -355,6 +355,8 @@ void FReplicationWriter::QueueNetObjectAttachments(FInternalNetRefIndex OwnerInt
 		FReplicationInfo& OwnerInfo = GetReplicationInfo(OwnerInternalIndex);
 		OwnerInfo.HasDirtySubObjects = 1;
 	}
+
+	return true;
 }
 
 void FReplicationWriter::SetState(uint32 InternalIndex, EReplicatedObjectState NewState)
@@ -2681,37 +2683,56 @@ int FReplicationWriter::WriteDestructionInfo(FNetSerializationContext& Context,
 
 uint32 FReplicationWriter::WriteOOBAttachments(FNetSerializationContext& Context)
 {
-	uint32 WrittenObjectCount = 0U;	
-	if (WriteContext.bHasHugeObjectToSend)
+	uint32 WrittenObjectCount = 0U;
+
+	if (WriteContext.WriteMode == EDataStreamWriteMode::PostTickDispatch)
 	{
-		IRIS_PROFILER_SCOPE(FReplicationWriter_WriteHugeObjectAttachments);
-		const int32 Result = WriteObjectBatch(Context, ObjectIndexForOOBAttachment, WriteObjectFlag_Attachments | WriteObjectFlag_HugeObject);
-		if (Result == -1)
+		if (WriteContext.bHasOOBAttachmentsToSend && CanSendObject(ObjectIndexForOOBAttachment))
 		{
-			return WrittenObjectCount;
-		}
+			IRIS_PROFILER_SCOPE(FReplicationWriter_WriteOOBAttachments);
+			const int32 Result = WriteObjectBatch(Context, ObjectIndexForOOBAttachment, WriteObjectFlag_Attachments);
+			if (Result == -1)
+			{
+				return WrittenObjectCount;
+			}
 
-		const bool bHasHugeObjectToSend = Attachments.HasUnsentAttachments(ENetObjectAttachmentType::HugeObject, ObjectIndexForOOBAttachment);
-		WriteContext.bHasHugeObjectToSend = bHasHugeObjectToSend;
-		if (!bHasHugeObjectToSend)
-		{
-			HugeObjectSendQueue.Stats.EndSendingTime = FPlatformTime::Cycles64();
+			WriteContext.bHasOOBAttachmentsToSend = Attachments.HasUnsentUnreliableAttachments(ENetObjectAttachmentType::OutOfBand, ObjectIndexForOOBAttachment);
+			WrittenObjectCount += Result;
 		}
-
-		WrittenObjectCount += Result;
 	}
-
-	if (WriteContext.bHasOOBAttachmentsToSend && CanSendObject(ObjectIndexForOOBAttachment))
+	else
 	{
-		IRIS_PROFILER_SCOPE(FReplicationWriter_WriteOOBAttachments);
-		const int32 Result = WriteObjectBatch(Context, ObjectIndexForOOBAttachment, WriteObjectFlag_Attachments);
-		if (Result == -1)
+		if (WriteContext.bHasHugeObjectToSend)
 		{
-			return WrittenObjectCount;
+			IRIS_PROFILER_SCOPE(FReplicationWriter_WriteHugeObjectAttachments);
+			const int32 Result = WriteObjectBatch(Context, ObjectIndexForOOBAttachment, WriteObjectFlag_Attachments | WriteObjectFlag_HugeObject);
+			if (Result == -1)
+			{
+				return WrittenObjectCount;
+			}
+
+			const bool bHasHugeObjectToSend = Attachments.HasUnsentAttachments(ENetObjectAttachmentType::HugeObject, ObjectIndexForOOBAttachment);
+			WriteContext.bHasHugeObjectToSend = bHasHugeObjectToSend;
+			if (!bHasHugeObjectToSend)
+			{
+				HugeObjectSendQueue.Stats.EndSendingTime = FPlatformTime::Cycles64();
+			}
+
+			WrittenObjectCount += Result;
 		}
 
-		WriteContext.bHasOOBAttachmentsToSend = Attachments.HasUnsentAttachments(ENetObjectAttachmentType::OutOfBand, ObjectIndexForOOBAttachment);
-		WrittenObjectCount += Result;
+		if (WriteContext.bHasOOBAttachmentsToSend && CanSendObject(ObjectIndexForOOBAttachment))
+		{
+			IRIS_PROFILER_SCOPE(FReplicationWriter_WriteOOBAttachments);
+			const int32 Result = WriteObjectBatch(Context, ObjectIndexForOOBAttachment, WriteObjectFlag_Attachments);
+			if (Result == -1)
+			{
+				return WrittenObjectCount;
+			}
+
+			WriteContext.bHasOOBAttachmentsToSend = Attachments.HasUnsentAttachments(ENetObjectAttachmentType::OutOfBand, ObjectIndexForOOBAttachment);
+			WrittenObjectCount += Result;
+		}
 	}
 
 	return WrittenObjectCount;
@@ -3018,39 +3039,63 @@ UDataStream::EWriteResult FReplicationWriter::BeginWrite(const UDataStream::FBeg
 		return UDataStream::EWriteResult::NoData;
 	}
 
-	// See if we have any work to do
-	const bool bHasUpdatedObjectsToSend = ObjectsWithDirtyChanges.IsAnyBitSet();
-	const bool bHasDestroyedObjectsToSend = ObjectsPendingDestroy.IsAnyBitSet();
-	const bool bHasUnsentOOBAttachments = Attachments.HasUnsentAttachments(ENetObjectAttachmentType::OutOfBand, ObjectIndexForOOBAttachment);
-	const bool bHasUnsentHugeObject = Attachments.HasUnsentAttachments(ENetObjectAttachmentType::HugeObject, ObjectIndexForOOBAttachment);
+	// Initialize context which can be used over multiple calls to WriteData
+	WriteContext.bHasUpdatedObjectsToSend  = 0U;
+	WriteContext.bHasDestroyedObjectsToSend = 0U;
+	WriteContext.bHasHugeObjectToSend = 0U;
+	WriteContext.bHasOOBAttachmentsToSend = 0U;
+	WriteContext.ScheduledObjectCount = 0u;
+
+	WriteContext.WriteMode = Params.WriteMode;
 
-	// Nothing to send
-	if (!(bHasUpdatedObjectsToSend | bHasDestroyedObjectsToSend | bHasUnsentOOBAttachments | bHasUnsentHugeObject))
+	// Setup for writing PostTickDispatch data, currently this is only writing unreliable OOBAttachments.
+	if (WriteContext.WriteMode == EDataStreamWriteMode::PostTickDispatch)
 	{
-		return UDataStream::EWriteResult::NoData;
+		const bool bHasUnsentOOBAttachments = Attachments.HasUnsentUnreliableAttachments(ENetObjectAttachmentType::OutOfBand, ObjectIndexForOOBAttachment);
+		if (!bHasUnsentOOBAttachments)
+		{
+			return UDataStream::EWriteResult::NoData;
+		}
+		WriteContext.bHasOOBAttachmentsToSend = bHasUnsentOOBAttachments;
+		WriteContext.bCanWriteMoreData = Params.bCanWriteMoreData;
 	}
+	else
+	{
+		// See if we have any work to do
+		const bool bHasUpdatedObjectsToSend = ObjectsWithDirtyChanges.IsAnyBitSet();
+		const bool bHasDestroyedObjectsToSend = ObjectsPendingDestroy.IsAnyBitSet();
+		const bool bHasUnsentOOBAttachments = Attachments.HasUnsentAttachments(ENetObjectAttachmentType::OutOfBand, ObjectIndexForOOBAttachment);
+		const bool bHasUnsentHugeObject = Attachments.HasUnsentAttachments(ENetObjectAttachmentType::HugeObject, ObjectIndexForOOBAttachment);
+
+		// Nothing to send
+		if (!(bHasUpdatedObjectsToSend | bHasDestroyedObjectsToSend | bHasUnsentOOBAttachments | bHasUnsentHugeObject))
+		{
+			return UDataStream::EWriteResult::NoData;
+		}
+
+		// Initialize context which can be used over multiple calls to WriteData
+		WriteContext.bHasUpdatedObjectsToSend  = bHasUpdatedObjectsToSend | bHasUnsentOOBAttachments | bHasUnsentHugeObject;
+		WriteContext.bHasDestroyedObjectsToSend = bHasDestroyedObjectsToSend;
+		WriteContext.bHasHugeObjectToSend = bHasUnsentHugeObject;
+		WriteContext.bHasOOBAttachmentsToSend = bHasUnsentOOBAttachments;
+
+		// $IRIS TODO: LinearAllocator/ScratchPad?
+		// Allocate space for indices to send
+		// This should be allocated from frame temp allocator and be cleaned up end of frame, we might want this data to persist over multiple write calls but not over multiple frames 
+		// https://jira.it.epicgames.com/browse/UE-127374	
+		WriteContext.ScheduledObjectInfos = reinterpret_cast<FScheduleObjectInfo*>(FMemory::Malloc(sizeof(FScheduleObjectInfo) * Parameters.MaxActiveReplicatedObjectCount));
+		WriteContext.ScheduledObjectCount = ScheduleObjects(WriteContext.ScheduledObjectInfos);
+	}
+
+	// Reset dependent object array
+	WriteContext.DependentObjectsPendingSend.Reset();
 
-	// Initialize context which can be used over multiple calls to WriteData
-	WriteContext.bHasUpdatedObjectsToSend  = bHasUpdatedObjectsToSend | bHasUnsentOOBAttachments | bHasUnsentHugeObject;
-	WriteContext.bHasDestroyedObjectsToSend = bHasDestroyedObjectsToSend;
-	WriteContext.bHasHugeObjectToSend = bHasUnsentHugeObject;
-	WriteContext.bHasOOBAttachmentsToSend = bHasUnsentOOBAttachments;
 	WriteContext.CurrentIndex = 0U;
 	WriteContext.FailedToWriteSmallObjectCount = 0U;
 	WriteContext.SortedObjectCount = 0U;
 	WriteContext.NumWrittenPacketsInThisBatch = 0U;
 	WriteContext.bCanWriteMoreData = Params.bCanWriteMoreData;
 
-	// Reset dependent object array
-	WriteContext.DependentObjectsPendingSend.Reset();
-
-	// $IRIS TODO: LinearAllocator/ScratchPad?
-	// Allocate space for indices to send
-	// This should be allocated from frame temp allocator and be cleaned up end of frame, we might want this data to persist over multiple write calls but not over multiple frames 
-	// https://jira.it.epicgames.com/browse/UE-127374	
-	WriteContext.ScheduledObjectInfos = reinterpret_cast<FScheduleObjectInfo*>(FMemory::Malloc(sizeof(FScheduleObjectInfo) * Parameters.MaxActiveReplicatedObjectCount));
-	WriteContext.ScheduledObjectCount = ScheduleObjects(WriteContext.ScheduledObjectInfos);
-
 	// Clear net stats. Used for CVS and Network Insights stats.
 	WriteContext.Stats.Reset();
 	WriteContext.Stats.SetNumberOfRootObjectsScheduledForReplication(WriteContext.ScheduledObjectCount);
diff --git a/Engine/Source/Runtime/Experimental/Iris/Core/Private/Iris/ReplicationSystem/ReplicationWriter.h b/Engine/Source/Runtime/Experimental/Iris/Core/Private/Iris/ReplicationSystem/ReplicationWriter.h
--- a/Engine/Source/Runtime/Experimental/Iris/Core/Private/Iris/ReplicationSystem/ReplicationWriter.h
+++ b/Engine/Source/Runtime/Experimental/Iris/Core/Private/Iris/ReplicationSystem/ReplicationWriter.h
@@ -184,7 +184,8 @@ public:
 	void SetNetExports(FNetExports& InNetExports);
 
 	// Attachments
-	void QueueNetObjectAttachments(FInternalNetRefIndex OwnerInternalIndex, FInternalNetRefIndex SubObjectInternalIndex, TArrayView<const TRefCountPtr<FNetBlob>> Attachments, ENetObjectAttachmentSendPolicyFlags SendFlags);
+	// Queue NetObjectAttachments, returns whether the attachments was enqueued or not.
+	bool QueueNetObjectAttachments(FInternalNetRefIndex OwnerInternalIndex, FInternalNetRefIndex SubObjectInternalIndex, TArrayView<const TRefCountPtr<FNetBlob>> Attachments, ENetObjectAttachmentSendPolicyFlags SendFlags);
 
 private:
 	// Various types
@@ -230,7 +231,9 @@ private:
 		// How many packets have we written to?
 		uint32 NumWrittenPacketsInThisBatch = 0U;
 		bool bCanWriteMoreData = false;
-
+		
+		EDataStreamWriteMode WriteMode = EDataStreamWriteMode::Full;	
+	
 		uint32 bHasDestroyedObjectsToSend : 1;
 		uint32 bHasUpdatedObjectsToSend : 1;
 		uint32 bHasHugeObjectToSend : 1;
diff --git a/Engine/Source/Runtime/Experimental/Iris/Core/Public/Iris/DataStream/DataStream.h b/Engine/Source/Runtime/Experimental/Iris/Core/Public/Iris/DataStream/DataStream.h
--- a/Engine/Source/Runtime/Experimental/Iris/Core/Public/Iris/DataStream/DataStream.h
+++ b/Engine/Source/Runtime/Experimental/Iris/Core/Public/Iris/DataStream/DataStream.h
@@ -9,6 +9,16 @@ namespace UE::Net
 {
 	class FNetSerializationContext;
 	enum class EPacketDeliveryStatus : uint8;
+
+enum class EDataStreamWriteMode : unsigned
+{
+	// Allowed to write all data, this is the default WriteMode
+	Full,
+
+	// Only write data that should be sent after PostTickDispatch
+	PostTickDispatch,
+};
+
 }
 
 /**
@@ -55,6 +65,7 @@ public:
 
 	struct FBeginWriteParameters
 	{
+		UE::Net::EDataStreamWriteMode WriteMode = UE::Net::EDataStreamWriteMode::Full;
 		bool bCanWriteMoreData = false;
 	};
 
diff --git a/Engine/Source/Runtime/Experimental/Iris/Core/Public/Iris/ReplicationSystem/ReplicationSystem.h b/Engine/Source/Runtime/Experimental/Iris/Core/Public/Iris/ReplicationSystem/ReplicationSystem.h
--- a/Engine/Source/Runtime/Experimental/Iris/Core/Public/Iris/ReplicationSystem/ReplicationSystem.h
+++ b/Engine/Source/Runtime/Experimental/Iris/Core/Public/Iris/ReplicationSystem/ReplicationSystem.h
@@ -81,6 +81,22 @@ public:
 	/** @return Is this system configured to replicate object properties. */
 	bool AllowObjectReplication() { return bAllowObjectReplication; }
 
+	struct FSendUpdateParams
+	{
+		// Type of SendPass we want to do @see EReplicationSystemSendPass
+		UE::Net::EReplicationSystemSendPass SendPass = UE::Net::EReplicationSystemSendPass::TickFlush;
+
+		// DeltaTime, only relevant for EReplicationSystemSendPass::TickFlush
+		float DeltaSeconds = 0.f;		
+	};
+
+	/**
+	 * PreSendUpdate performs all the necessary work, such as filtering and prioritization of objects,
+	 * so that each connection will be properly updated with all the information needed in order to replicate.
+	 * @param Params, Parameters for the update pass, to @see FSendUpateParams 
+	 */
+	IRISCORE_API void PreSendUpdate(const FSendUpdateParams& Params);	
+
 	/**
 	 * PreSendUpdate performs all the necessary work, such as filtering and prioritization of objects,
 	 * so that each connection will be properly updated with all the information needed in order to replicate.
@@ -89,10 +105,11 @@ public:
 
 	/**
 	 * SendUpdate is currently more of a placeholder for a future where the ReplicationSystem itself is responsible for
-	 * the low level protocol and sending, rather than having the DataStreamChannel.
+	 * the low level protocol and sending, rather than having the DataStreamChannel write data when ticked
 	 * @see UDataStreamChannel.
+	 * @param SendFunction, Function taking an array of ConnectionId`s that has data to send
 	 */
-	IRISCORE_API void SendUpdate();
+	IRISCORE_API void SendUpdate(TFunctionRef<void(TArrayView<uint32>)> SendFunction);
 
 	/**
 	 * Cleanup temporaries and prepare for the next send update.
diff --git a/Engine/Source/Runtime/Experimental/Iris/Core/Public/Iris/ReplicationSystem/ReplicationSystemTypes.h b/Engine/Source/Runtime/Experimental/Iris/Core/Public/Iris/ReplicationSystem/ReplicationSystemTypes.h
--- a/Engine/Source/Runtime/Experimental/Iris/Core/Public/Iris/ReplicationSystem/ReplicationSystemTypes.h
+++ b/Engine/Source/Runtime/Experimental/Iris/Core/Public/Iris/ReplicationSystem/ReplicationSystemTypes.h
@@ -14,6 +14,17 @@ enum class ENetObjectDeltaCompressionStatus : unsigned
 	Allow,
 };
 
+enum class EReplicationSystemSendPass : unsigned
+{
+	Invalid,
+
+	// Sending data directly after PostDispatch, this is a partial update only, not updating scope and filtering and will only process RPC/Attachmeents
+	PostTickDispatch,
+
+	// Sending data part of TickFlush, this will do a full update and replicate data to all connections
+	TickFlush,
+};
+
 enum class EDependentObjectSchedulingHint : uint8
 {
 	// Default behavior, dependent object will be scheduled to replicate if parent is replicated, if the dependent object has not yet been replicated it will be replicated in the same batch as the parent
@@ -38,7 +49,7 @@ enum class ENetObjectAttachmentSendPolicyFlags : uint32
 	ScheduleAsOOB = 1U << 0U,
 
 	// Hint that this attachment like to be sent during PostTickDispatch. 
-	// If one attachment has this flag all unreliable attachments scheduled to use the OOB channel they all will be sent during PostTickDispatch.
+	// If one enqueued attachment has this flag all currently unreliable attachments scheduled to use the OOB channel will be sent during PostTickDispatch.
 	SendInPostTickDispatch = ScheduleAsOOB << 1U,	
 
 	// SendImmediate, Attachment should be sent using OOB channel and from PostTickDispatch.
