Chaos - collision detection fixes
- add setting struct to FMeshContactGenerator
- change capsule collision to use FMeshContactGenerator with custom settings

diff --git a/Engine/Source/Runtime/Experimental/Chaos/Private/Chaos/Collision/CapsuleTriangleContactPoint.cpp b/Engine/Source/Runtime/Experimental/Chaos/Private/Chaos/Collision/CapsuleTriangleContactPoint.cpp
--- a/Engine/Source/Runtime/Experimental/Chaos/Private/Chaos/Collision/CapsuleTriangleContactPoint.cpp
+++ b/Engine/Source/Runtime/Experimental/Chaos/Private/Chaos/Collision/CapsuleTriangleContactPoint.cpp
@@ -264,12 +264,12 @@ namespace Chaos
 		const bool bPreferFaceContact = (FMath::Abs(SinAxisFaceAngle) < FaceContactSinAngleThreshold);
 
 		// Generate contacts for the cylinder ends near the face. Only consider the ends where we did not generate an end-cap contact, 
-		// and only when we are at a low angle to the face, or the cylinder end is close to the face. 
+		// and only when we are at a low angle to the face. 
 		// The cylinder points can only be inside the edge planes if the segment points are within Radius of the edge planes
 		const bool bNearAll0 = ((EdgeD0s[0] <= R + DistanceTolerance) & (EdgeD0s[1] <= R + DistanceTolerance) & (EdgeD0s[2] <= R + DistanceTolerance)) != 0;
 		const bool bNearAll1 = ((EdgeD1s[0] <= R + DistanceTolerance) & (EdgeD1s[1] <= R + DistanceTolerance) & (EdgeD1s[2] <= R + DistanceTolerance)) != 0;
-		const bool bCheckCylinder0 = ((!bCollided0) & (bPreferFaceContact | (FaceD0 < RejectDistance)) & bNearAll0) != 0;
-		const bool bCheckCylinder1 = ((!bCollided1) & (bPreferFaceContact | (FaceD1 < RejectDistance)) & bNearAll1) != 0;
+		const bool bCheckCylinder0 = ((!bCollided0) & bPreferFaceContact & bNearAll0) != 0;
+		const bool bCheckCylinder1 = ((!bCollided1) & bPreferFaceContact & bNearAll1) != 0;
 		if ((bCheckCylinder0 | bCheckCylinder1) != 0)
 		{
 			FVec3 RadialAxis = FVec3::CrossProduct(FVec3::CrossProduct(Axis, FaceN), Axis);
diff --git a/Engine/Source/Runtime/Experimental/Chaos/Private/Chaos/Collision/MeshContactGenerator.cpp b/Engine/Source/Runtime/Experimental/Chaos/Private/Chaos/Collision/MeshContactGenerator.cpp
--- a/Engine/Source/Runtime/Experimental/Chaos/Private/Chaos/Collision/MeshContactGenerator.cpp
+++ b/Engine/Source/Runtime/Experimental/Chaos/Private/Chaos/Collision/MeshContactGenerator.cpp
@@ -5,6 +5,8 @@
 namespace Chaos
 {
 	extern FRealSingle Chaos_Collision_MeshContactNormalThreshold;
+	extern FRealSingle Chaos_Collision_MeshContactNormalRejectionThreshold;
+	extern int32 Chaos_Collision_MeshManifoldHashSize;
 }
 
 namespace Chaos::CVars
@@ -14,9 +16,21 @@ namespace Chaos::CVars
 
 namespace Chaos::Private
 {
-	FMeshContactGenerator::FMeshContactGenerator(const int32 InHashSize)
-		: EdgeTriangleIndicesMap(InHashSize)
-		, VertexContactIndicesMap(InHashSize)
+	FMeshContactGeneratorSettings::FMeshContactGeneratorSettings()
+	{
+		HashSize = FMath::RoundUpToPowerOfTwo(Chaos_Collision_MeshManifoldHashSize);
+		FaceNormalDotThreshold = Chaos_Collision_MeshContactNormalThreshold;
+		EdgeNormalDotRejectTolerance = Chaos_Collision_MeshContactNormalRejectionThreshold;
+		BarycentricTolerance = FReal(1.e-3);
+		MaxContactsBufferSize = 1000;
+		bCullBackFaces = true;
+		bUseTwoPassLoop = true;
+	}
+
+	FMeshContactGenerator::FMeshContactGenerator(const FMeshContactGeneratorSettings& InSettings)
+		: Settings(InSettings)
+		, EdgeTriangleIndicesMap(InSettings.HashSize)
+		, VertexContactIndicesMap(InSettings.HashSize)
 	{
 	}
 
@@ -48,7 +62,6 @@ namespace Chaos::Private
 		// We need to know how many vertex and edge collisions a triangle has on it (from collision with other triangle 
 		// that share edges and vertices) so we avoid visiting triangles if we already know the contacts. Check all
 		// contacts and assigne edge/vertex status as required. See GenerateMeshContacts()
-		const FReal BarycentricTolerance = 1.e-3;
 		for (const FContactPoint& ContactPoint : TriangleContactPoints)
 		{
 			// See if we have a vertex or edge contact
@@ -58,7 +71,7 @@ namespace Chaos::Private
 				ContactPoint.ShapeContactPoints[1],
 				Triangle.GetVertex(0), Triangle.GetVertex(1), Triangle.GetVertex(2),
 				LocalVertexID0, LocalVertexID1,
-				BarycentricTolerance);
+				Settings.BarycentricTolerance);
 
 			const int32 VertexID0 = (LocalVertexID0 != INDEX_NONE) ? Triangle.GetVertexIndex(LocalVertexID0) : INDEX_NONE;
 			const int32 VertexID1 = (LocalVertexID1 != INDEX_NONE) ? Triangle.GetVertexIndex(LocalVertexID1) : INDEX_NONE;
@@ -67,8 +80,7 @@ namespace Chaos::Private
 
 			const FRealSingle ContactNormalDotTriangleNormal = FRealSingle(FVec3::DotProduct(ContactPoint.ShapeContactNormal, Triangle.GetNormal()));
 
-			const FReal FaceNormalThreshold = FReal(0.998);	// 3deg
-			const bool bIsFaceContact = (ContactNormalDotTriangleNormal > FaceNormalThreshold);
+			const bool bIsFaceContact = (ContactNormalDotTriangleNormal > Settings.FaceNormalDotThreshold);
 
 			// Register vertex and edge collisions
 			if ((VertexID0 != INDEX_NONE) && (VertexID1 != INDEX_NONE))
@@ -148,21 +160,19 @@ namespace Chaos::Private
 
 	void FMeshContactGenerator::PruneAndCorrectContacts()
 	{
-		const FRealSingle NormalThreshold = FRealSingle(0.998);
-
 		for (int32 ContactIndex = 0; ContactIndex < Contacts.Num(); ++ContactIndex)
 		{
 			FTriangleContactPointData& ContactPointData = ContactDatas[ContactIndex];
 
 			// Reject back-faces
-			if (ContactPointData.GetContactNormalDotTriangleNormal() < 0)
+			if (!!Settings.bCullBackFaces && (ContactPointData.GetContactNormalDotTriangleNormal() < 0))
 			{
 				ContactPointData.SetDisabled();
 				continue;
 			}
 
-			// Fix edge normals
-			if (ContactPointData.GetContactNormalDotTriangleNormal() < Chaos_Collision_MeshContactNormalThreshold)
+			// Fix edge normals that aren't already close to a face normal
+			if (ContactPointData.GetContactNormalDotTriangleNormal() < Settings.FaceNormalDotThreshold)
 			{
 				FixContactNormal(ContactIndex);
 			}
@@ -181,6 +191,8 @@ namespace Chaos::Private
 		const FTriangleExt& Triangle = Triangles[LocalTriangleIndex];
 		const FVec3& TriangleNormal = Triangle.GetNormal();
 
+		const FReal ContactDotNormal = ContactPointData.GetContactNormalDotTriangleNormal();
+
 		// If we have an edge or vertex contact, make sure that the normal is in a valiid range, based
 		// on the triangles that share that edge or vertex.
 		if (ContactPointData.IsEdge())
@@ -192,11 +204,17 @@ namespace Chaos::Private
 				const FTriangleExt& OtherTriangle = Triangles[OtherLocalTriangleIndex];
 				const FVec3& OtherTriangleNormal = OtherTriangle.GetNormal();
 
-				const FReal ContactDotNormal = ContactPointData.GetContactNormalDotTriangleNormal();
 				const FReal MinContactDotNormal = FRealSingle(FVec3::DotProduct(OtherTriangleNormal, TriangleNormal));
 				if (ContactDotNormal < MinContactDotNormal)
 				{
 					// We are outside the valid normal range for this edge
+					// If our normal was very far away from a valid normal, drop the contact
+					if (MinContactDotNormal - ContactDotNormal > Settings.EdgeNormalDotRejectTolerance)
+					{
+						ContactPointData.SetDisabled();
+						return;
+					}
+
 					// Convert the edge collision to a face collision on one of the faces, selected to get the smallest depth
 					FVec3 CorrectedContactNormal;
 					int32 CorrectedTriangleIndex;
@@ -228,7 +246,7 @@ namespace Chaos::Private
 			for (int32 TriangleIndex = 0; TriangleIndex < Triangles.Num(); ++TriangleIndex)
 			{
 				const FTriangleExt& OtherTriangle = Triangles[TriangleIndex];
-				if (OtherTriangle.HasVertexIndex(ContactPointData.GetVertexID()))
+				if ((TriangleIndex != LocalTriangleIndex) && OtherTriangle.HasVertexIndex(ContactPointData.GetVertexID()))
 				{
 					FVec3 VertexA, VertexB, VertexC;
 					if (OtherTriangle.GetVertexPosition(ContactPointData.GetVertexID(), VertexA) && OtherTriangle.GetOtherVertexPositions(ContactPointData.GetVertexID(), VertexB, VertexC))
@@ -238,6 +256,7 @@ namespace Chaos::Private
 						const FVec3 EdgeDelta0 = VertexB - VertexA;
 						const FVec3 EdgeDelta1 = VertexC - VertexA;
 						const FVec3& OtherTriangleNormal = OtherTriangle.GetNormal();
+
 						const FReal EdgeSign0 = FVec3::DotProduct(FVec3::CrossProduct(ContactPoint.ShapeContactNormal, VertexB - VertexA), OtherTriangleNormal);
 						const FReal EdgeSign1 = FVec3::DotProduct(FVec3::CrossProduct(ContactPoint.ShapeContactNormal, VertexC - VertexA), OtherTriangleNormal);
 						if (FMath::Sign(EdgeSign0) == FMath::Sign(EdgeSign1))
@@ -246,6 +265,14 @@ namespace Chaos::Private
 							const FReal NormalDotCentroid = FVec3::DotProduct(ContactPoint.ShapeContactNormal, Centroid - ContactPoint.ShapeContactPoints[1]);
 							if (NormalDotCentroid > 0)
 							{
+								// If our normal was very far away from a valid normal, drop the contact
+								const FReal MinContactDotNormal = FRealSingle(FVec3::DotProduct(OtherTriangleNormal, TriangleNormal));
+								if (MinContactDotNormal - ContactDotNormal > Settings.EdgeNormalDotRejectTolerance)
+								{
+									ContactPointData.SetDisabled();
+									return;
+								}
+
 								const FReal OtherContactDotNormal = FVec3::DotProduct(ContactPoint.ShapeContactNormal, OtherTriangleNormal);
 								const FVec3 CorrectedContactNormal = OtherTriangleNormal;
 
@@ -316,6 +343,7 @@ namespace Chaos::Private
 #if CHAOS_DEBUG_DRAW
 		if (CVars::ChaosSolverDebugDrawMeshContacts && FDebugDrawQueue::GetInstance().IsDebugDrawingEnabled())
 		{
+			check(LineScale > 0);
 			const FReal Duration = 0;
 			const int32 DrawPriority = 10;
 			for (int32 ContactIndex = 0; ContactIndex < Contacts.Num(); ++ContactIndex)
@@ -328,7 +356,7 @@ namespace Chaos::Private
 					const FVec3 N = ConvexTransform.TransformVectorNoScale(ContactPoint.ShapeContactNormal);
 
 					// Draw the normal from the triangle face
-					FDebugDrawQueue::GetInstance().DrawDebugLine(P1, P1 + LineScale * FReal(50) * N, Color, false, FRealSingle(Duration), DrawPriority, FRealSingle(LineScale * 2));
+					FDebugDrawQueue::GetInstance().DrawDebugLine(P1, P1 + LineScale * FReal(50) * N, Color, false, FRealSingle(Duration), DrawPriority, FRealSingle(1.0 / LineScale));
 				}
 			}
 		}
@@ -365,7 +393,7 @@ namespace Chaos::Private
 	{
 #if CHAOS_DEBUG_DRAW
 		const FReal Duration = 0;
-		const FReal LineScale = 1;
+		const FReal LineScale = 0.5;
 		const int8 DrawPriority = 10;
 
 		const FVec3 V0 = ConvexTransform.TransformPosition(Triangle.GetVertex(0));
diff --git a/Engine/Source/Runtime/Experimental/Chaos/Private/Chaos/CollisionOneShotManifolds.cpp b/Engine/Source/Runtime/Experimental/Chaos/Private/Chaos/CollisionOneShotManifolds.cpp
--- a/Engine/Source/Runtime/Experimental/Chaos/Private/Chaos/CollisionOneShotManifolds.cpp
+++ b/Engine/Source/Runtime/Experimental/Chaos/Private/Chaos/CollisionOneShotManifolds.cpp
@@ -88,9 +88,12 @@ namespace Chaos
 	bool bChaos_Collision_EnableLargeMeshManifolds = 1;
 	FAutoConsoleVariableRef CVarChaos_Collision_EnableLargeMeshManifolds(TEXT("p.Chaos.Collision.EnableLargeMeshManifolds"), bChaos_Collision_EnableLargeMeshManifolds, TEXT("Whether to allow large mesh manifolds for collisions against meshes (required for good behaviour)"));
 
-	FRealSingle Chaos_Collision_MeshContactNormalThreshold = 0.996f;	// ~5deg
+	FRealSingle Chaos_Collision_MeshContactNormalThreshold = 0.998f;	// ~3deg
 	FAutoConsoleVariableRef CVarChaos_Collision_MeshContactNormalThreshold(TEXT("p.Chaos.Collision.MeshContactNormalThreshold"), Chaos_Collision_MeshContactNormalThreshold, TEXT("Treat contact with a dot product between the normal and the triangle face greater than this as face collisions"));
 
+	FRealSingle Chaos_Collision_MeshContactNormalRejectionThreshold = 0.7f;	// ~45deg
+	FAutoConsoleVariableRef CVarChaos_Collision_MeshContactNormalRejectionThreshold(TEXT("p.Chaos.Collision.MeshContactNormalRejectionThreshold"), Chaos_Collision_MeshContactNormalRejectionThreshold, TEXT("Don't correct edge and vertex normals if they are beyond the valid range by more than this"));
+
 	bool bChaos_Collision_MeshManifoldSortByDistance = false;
 	FAutoConsoleVariableRef CVarChaos_Collision_LargeMeshManifoldSortByDistance(TEXT("p.Chaos.Collision.SortMeshManifoldByDistance"), bChaos_Collision_MeshManifoldSortByDistance, TEXT("Sort large mesh manifold points by |RxN| for improved solver stability (less rotation in first iteration)"));
 
diff --git a/Engine/Source/Runtime/Experimental/Chaos/Private/Chaos/CollisionOneShotManifoldsMeshShapes.cpp b/Engine/Source/Runtime/Experimental/Chaos/Private/Chaos/CollisionOneShotManifoldsMeshShapes.cpp
--- a/Engine/Source/Runtime/Experimental/Chaos/Private/Chaos/CollisionOneShotManifoldsMeshShapes.cpp
+++ b/Engine/Source/Runtime/Experimental/Chaos/Private/Chaos/CollisionOneShotManifoldsMeshShapes.cpp
@@ -125,7 +125,7 @@ namespace Chaos
 			const auto& GenerateConvexTriangleContacts =
 				[&Convex, &MeshToConvexTransform, CullDistance](const FTriangle& Triangle, FContactPointManifold& OutContactPoints)
 				{
-					ConstructConvexTriangleOneShotManifold2(Convex, Triangle, CullDistance, OutContactPoints);
+					GenerateConvexTriangleOneShotManifold(Convex, Triangle, CullDistance, OutContactPoints);
 				};
 
 			// Collect all the triangles that overlap our convex, transformed into Convex space
@@ -229,25 +229,53 @@ namespace Chaos
 			ensure(QuadraticTransform.GetScale3D() == FVec3(1));
 			ensure(MeshTransform.GetScale3D() == FVec3(1));
 
+			const FVec3 MeshScale = FVec3(1);	// Scale is built into heightfield
 			const FReal CullDistance = Constraint.GetCullDistance();
 			const FReal PhiTolerance = CalculateTriMeshPhiTolerance(CullDistance);
 			const FReal DistanceTolerance = Chaos_Collision_TriMeshDistanceTolerance;
-			FContactTriangleCollector MeshContacts(bChaos_Collision_OneSidedHeightField, PhiTolerance, DistanceTolerance, QuadraticTransform);
-
-			if (const FImplicitSphere3* Sphere = Quadratic.template GetObject<FImplicitSphere3>())
-			{
-				ConstructConvexMeshOneShotManifold(*Sphere, QuadraticTransform, Mesh, MeshTransform, FVec3(1), CullDistance, MeshContacts);
-			}
-			else if (const FImplicitCapsule3* Capsule = Quadratic.template GetObject<FImplicitCapsule3>())
+			
+			if (bChaos_Collision_EnableMeshManifoldOptimizedLoop)
 			{
-				ConstructConvexMeshOneShotManifold(*Capsule, QuadraticTransform, Mesh, MeshTransform, FVec3(1), CullDistance, MeshContacts);
+				// New version uses a two-pass loop over triangles to avoid visiting triangles whose vertices are all colliding as a result of checking adjacent triangles
+				Private::FMeshContactGeneratorSettings ContactGeneratorSettings;
+				ContactGeneratorSettings.FaceNormalDotThreshold = 0.9999;	// ~0.8deg Normals must be accurate or rolling will not work correctly
+				ContactGeneratorSettings.bUseTwoPassLoop = false;			// two-pass loop is not helpful for capsules and spheres
+				Private::FMeshContactGenerator ContactGenerator(ContactGeneratorSettings);
+
+				if (const FImplicitSphere3* Sphere = Quadratic.template GetObject<FImplicitSphere3>())
+				{
+					ConstructConvexMeshOneShotManifold2(*Sphere, QuadraticTransform, Mesh, MeshTransform, MeshScale, CullDistance, ContactGenerator);
+				}
+				else if (const FImplicitCapsule3* Capsule = Quadratic.template GetObject<FImplicitCapsule3>())
+				{
+					ConstructConvexMeshOneShotManifold2(*Capsule, QuadraticTransform, Mesh, MeshTransform, MeshScale, CullDistance, ContactGenerator);
+				}
+				else
+				{
+					check(false);
+				}
+
+				Constraint.SetOneShotManifoldContacts(ContactGenerator.GetContactPoints());
 			}
 			else
 			{
-				check(false);
-			}
+				FContactTriangleCollector MeshContacts(bChaos_Collision_OneSidedHeightField, PhiTolerance, DistanceTolerance, QuadraticTransform);
 
-			Constraint.SetOneShotManifoldContacts(MeshContacts.GetContactPoints());
+				if (const FImplicitSphere3* Sphere = Quadratic.template GetObject<FImplicitSphere3>())
+				{
+					ConstructConvexMeshOneShotManifold(*Sphere, QuadraticTransform, Mesh, MeshTransform, MeshScale, CullDistance, MeshContacts);
+				}
+				else if (const FImplicitCapsule3* Capsule = Quadratic.template GetObject<FImplicitCapsule3>())
+				{
+					ConstructConvexMeshOneShotManifold(*Capsule, QuadraticTransform, Mesh, MeshTransform, MeshScale, CullDistance, MeshContacts);
+				}
+				else
+				{
+					check(false);
+				}
+
+				Constraint.SetOneShotManifoldContacts(MeshContacts.GetContactPoints());
+			}
 		}
 
 		/**
@@ -314,8 +342,8 @@ namespace Chaos
 			if (bChaos_Collision_EnableMeshManifoldOptimizedLoop)
 			{
 				// New version uses a two-pass loop over triangles to avoid visiting triangles whose vertices are all colliding as a result of checking adjacent triangles
-				const int32 HashSize = FMath::RoundUpToPowerOfTwo(Chaos_Collision_MeshManifoldHashSize);
-				Private::FMeshContactGenerator ContactGenerator(HashSize);
+				Private::FMeshContactGeneratorSettings ContactGeneratorSettings;
+				Private::FMeshContactGenerator ContactGenerator(ContactGeneratorSettings);
 
 				if (const FImplicitBox3* RawBox = Convex.template GetObject<FImplicitBox3>())
 				{
diff --git a/Engine/Source/Runtime/Experimental/Chaos/Public/Chaos/Collision/MeshContactGenerator.h b/Engine/Source/Runtime/Experimental/Chaos/Public/Chaos/Collision/MeshContactGenerator.h
--- a/Engine/Source/Runtime/Experimental/Chaos/Public/Chaos/Collision/MeshContactGenerator.h
+++ b/Engine/Source/Runtime/Experimental/Chaos/Public/Chaos/Collision/MeshContactGenerator.h
@@ -22,18 +22,48 @@ namespace Chaos::Private
 		}
 	};
 
+	class FMeshContactGeneratorSettings
+	{
+	public:
+		FMeshContactGeneratorSettings();
+
+		// Contacts with a dot product against the face normal above this value will not be processed in FixContactNormal
+		FReal FaceNormalDotThreshold;
+
+		// Triangle edge/vertex contacts that are more than this far from a valid normal (dot product) will be rejected rather than corrected
+		FReal EdgeNormalDotRejectTolerance;
+
+		// Used to determine whether a contact is on an edge or vertex
+		FReal BarycentricTolerance;
+
+		// We don't allow more (pre-filtered) contacts than this. Any extras will be lost.
+		int32 MaxContactsBufferSize;
+
+		// Size of the hash table used to store/lookup triangle data
+		int32 HashSize;
+
+		// Whether to ignore inside normals
+		// @todo(chaos): the non-culled option is not well tested and probably broken
+		uint32 bCullBackFaces : 1;
+
+		// Whether to use the optimized two-pass loop over triangles in GenerateMeshContacts which skips triangles
+		// that have contacts on all vertices in the second pass. This is only useful when this case occurs a lot
+		// which is does for large convexes against many triangles, but rarely for capsules and spheres.
+		uint32 bUseTwoPassLoop : 1;
+	};
+
 	/**
 	* Generate contacts between a collision shape and the triangles from a mesh.
 	*/
 	class FMeshContactGenerator
 	{
 	public:
-		FMeshContactGenerator(const int32 InHashSize);
+		FMeshContactGenerator(const FMeshContactGeneratorSettings& InSettings);
 
 		// Clear and initialize buffers
 		void BeginCollect(const int32 InNumTriangles)
 		{
-			const int32 ExpectedNumContacts = FMath::Min(InNumTriangles * 4, 1000);
+			const int32 ExpectedNumContacts = FMath::Min(InNumTriangles * 4, Settings.MaxContactsBufferSize);
 			Reset(InNumTriangles, ExpectedNumContacts);
 		}
 
@@ -43,6 +73,7 @@ namespace Chaos::Private
 			Triangles.Emplace(MeshTriangle, MeshTriangleIndex, VertexIndex0, VertexIndex1, VertexIndex2);
 		}
 
+		// Process all the added triangles to generate connectivity metadata etc
 		void EndCollect()
 		{
 			for (int32 LocalTriangleIndex = 0; LocalTriangleIndex < GetNumTriangles(); ++LocalTriangleIndex)
@@ -59,37 +90,13 @@ namespace Chaos::Private
 		template<typename TriangleContactGeneratorType>
 		void GenerateMeshContacts(const TriangleContactGeneratorType& TriangleContactGenerator)
 		{
-			FContactPointManifold TriangleContactPoints;
-
-			// First loop: Visit triangles that do not have any collisions on any of their vertices or edges.
-			// This will skip all triangles whose neighbours have already been processed and generated a contact
-			// on a shared edge/vertex.
-			for (int32 LocalTriangleIndex = 0; LocalTriangleIndex < GetNumTriangles(); ++LocalTriangleIndex)
+			if (!!Settings.bUseTwoPassLoop)
 			{
-				if (GetNumTriangleFaceCollisions(LocalTriangleIndex) == 0)
-				{
-					TriangleContactPoints.Reset();
-					TriangleContactGenerator(Triangles[LocalTriangleIndex].GetTriangle(), TriangleContactPoints);
-
-					AddTriangleContacts(TriangleContactPoints, LocalTriangleIndex);
-
-					SetTriangleVisited(LocalTriangleIndex, 0);
-				}
+				GenerateMeshContactsTwoPass<TriangleContactGeneratorType>(TriangleContactGenerator);
 			}
-
-			// Second loop: Visit remaining triangles that have less than 3 contacts on them. This will skip all triangles
-			// that have a full manifold as a result of collisions on shared edges/vertices from adjacent triangles.
-			for (int32 LocalTriangleIndex = 0; LocalTriangleIndex < GetNumTriangles(); ++LocalTriangleIndex)
+			else
 			{
-				if (!IsTriangleVisited(LocalTriangleIndex) && (GetNumTriangleFaceCollisions(LocalTriangleIndex) < 3))
-				{
-					TriangleContactPoints.Reset();
-					TriangleContactGenerator(Triangles[LocalTriangleIndex].GetTriangle(), TriangleContactPoints);
-
-					AddTriangleContacts(TriangleContactPoints, LocalTriangleIndex);
-
-					SetTriangleVisited(LocalTriangleIndex, 1);
-				}
+				GenerateMeshContactsOnePass<TriangleContactGeneratorType>(TriangleContactGenerator);
 			}
 		}
 
@@ -277,6 +284,59 @@ namespace Chaos::Private
 			int32 LocalTriangleIndices[2];
 		};
 
+		template<typename TriangleContactGeneratorType>
+		void GenerateMeshContactsOnePass(const TriangleContactGeneratorType& TriangleContactGenerator)
+		{
+			FContactPointManifold TriangleContactPoints;
+
+			for (int32 LocalTriangleIndex = 0; LocalTriangleIndex < GetNumTriangles(); ++LocalTriangleIndex)
+			{
+				TriangleContactPoints.Reset();
+				TriangleContactGenerator(Triangles[LocalTriangleIndex].GetTriangle(), TriangleContactPoints);
+
+				AddTriangleContacts(TriangleContactPoints, LocalTriangleIndex);
+
+				SetTriangleVisited(LocalTriangleIndex, 0);
+			}
+		}
+
+		template<typename TriangleContactGeneratorType>
+		void GenerateMeshContactsTwoPass(const TriangleContactGeneratorType& TriangleContactGenerator)
+		{
+			FContactPointManifold TriangleContactPoints;
+
+			// First loop: Visit triangles that do not have any collisions on any of their vertices or edges.
+			// This will skip all triangles whose neighbours have already been processed and generated a contact
+			// on a shared edge/vertex.
+			for (int32 LocalTriangleIndex = 0; LocalTriangleIndex < GetNumTriangles(); ++LocalTriangleIndex)
+			{
+				if (GetNumTriangleFaceCollisions(LocalTriangleIndex) == 0)
+				{
+					TriangleContactPoints.Reset();
+					TriangleContactGenerator(Triangles[LocalTriangleIndex].GetTriangle(), TriangleContactPoints);
+
+					AddTriangleContacts(TriangleContactPoints, LocalTriangleIndex);
+
+					SetTriangleVisited(LocalTriangleIndex, 0);
+				}
+			}
+
+			// Second loop: Visit remaining triangles that have less than 3 contacts on them. This will skip all triangles
+			// that have a full manifold as a result of collisions on shared edges/vertices from adjacent triangles.
+			for (int32 LocalTriangleIndex = 0; LocalTriangleIndex < GetNumTriangles(); ++LocalTriangleIndex)
+			{
+				if (!IsTriangleVisited(LocalTriangleIndex) && (GetNumTriangleFaceCollisions(LocalTriangleIndex) < 3))
+				{
+					TriangleContactPoints.Reset();
+					TriangleContactGenerator(Triangles[LocalTriangleIndex].GetTriangle(), TriangleContactPoints);
+
+					AddTriangleContacts(TriangleContactPoints, LocalTriangleIndex);
+
+					SetTriangleVisited(LocalTriangleIndex, 1);
+				}
+			}
+		}
+
 		void Reset(const int32 InMaxTriangles, const int32 InMaxContacts);
 
 		void AddTriangleEdge(const int32 LocalTriangleIndex, const int32 VertexIndex0, const int32 VertexIndex1)
@@ -350,6 +410,8 @@ namespace Chaos::Private
 		void DebugDrawTriangle(const FRigidTransform3& ConvexTransform, const FTriangleExt& TriangleData, const FColor& Color);
 
 	private:
+		FMeshContactGeneratorSettings Settings;
+
 		// All the triangles we might collide with
 		TArray<FTriangleExt> Triangles;
 
