Reduce memory usage in Dynamic Collection
Remove children array of set from Dynamic Collection, it is now using the bit array from its parent.

diff --git a/Engine/Source/Programs/HeadlessChaos/Private/GeometryCollection/GeometryCollectionTestClustering.cpp b/Engine/Source/Programs/HeadlessChaos/Private/GeometryCollection/GeometryCollectionTestClustering.cpp
--- a/Engine/Source/Programs/HeadlessChaos/Private/GeometryCollection/GeometryCollectionTestClustering.cpp
+++ b/Engine/Source/Programs/HeadlessChaos/Private/GeometryCollection/GeometryCollectionTestClustering.cpp
@@ -128,8 +128,63 @@ namespace GeometryCollectionTest
 
 	}
 
+	GTEST_TEST(AllTraits, GeometryCollection_DynamicCollection_ChildrenAccess)
+	{
+		FFramework UnitTest;
+
+		RigidBodyWrapper* Floor = TNewSimulationObject<GeometryType::RigidFloor>::Init()->template As<RigidBodyWrapper>();
+		UnitTest.AddSimulationObject(Floor);
+
+		TSharedPtr<FGeometryCollection> RestCollection = GeometryCollection::MakeCubeElement(FTransform(FQuat::MakeFromEuler(FVector(0, 0, 0.)), FVector(0, -10, 10)), FVector(1.0));
+		RestCollection->AppendGeometry(*GeometryCollection::MakeCubeElement(FTransform(FQuat::MakeFromEuler(FVector(0, 0, 0.)), FVector(0, 10, 10)), FVector(1.0)));
+		EXPECT_EQ(RestCollection->Transform.Num(), 2);
+
+		FGeometryCollectionClusteringUtility::ClusterAllBonesUnderNewRoot(RestCollection.Get());
+		EXPECT_EQ(RestCollection->Transform.Num(), 3);
+		RestCollection->Transform[2] = FTransform(FQuat::MakeFromEuler(FVector(90.0, 0, 0.)), FVector(0, 0, 40));
+
+		//GeometryCollectionAlgo::PrintParentHierarchy(RestCollection.Get());
+
+		CreationParameters Params;
+		Params.RestCollection = RestCollection;
+		Params.DynamicState = EObjectStateTypeEnum::Chaos_Object_Dynamic;
+		Params.CollisionType = ECollisionTypeEnum::Chaos_Surface_Volumetric;
+		Params.Simulating = true;
+		Params.EnableClustering = true;
+		Params.DamageThreshold = { 1000.f };
+		FGeometryCollectionWrapper* Collection = TNewSimulationObject<GeometryType::GeometryCollectionWithSuppliedRestCollection>::Init(Params)->template As<FGeometryCollectionWrapper>();
+
+		UnitTest.AddSimulationObject(Collection);
+		UnitTest.Initialize();
+
+		FReal StartingRigidDistance = (Collection->DynamicCollection->GetTransform(1).GetTranslation() - Collection->DynamicCollection->GetTransform(0).GetTranslation()).Size(), CurrentRigidDistance = 0.f;
+
+		TManagedArray<bool>& Active = Collection->DynamicCollection->Active;
+
+		EXPECT_FALSE(Active[0]);
+		EXPECT_FALSE(Active[1]);
+		EXPECT_TRUE(Active[2]); // only the root cluster should be active when using clustering 
+		UnitTest.Advance();
+		EXPECT_FALSE(Active[0]);
+		EXPECT_FALSE(Active[1]);
+		EXPECT_TRUE(Active[2]);
+
+
+		int32 ChildCount[3] = {0, 0, 0};
+		for (int32 Index = 0; Index < 3; ++Index)
+		{
+			Collection->DynamicCollection->IterateThroughChildren(Index, [&](int32 ChildIndex)
+			{
+				ChildCount[Index]++;
+				return true;
+			});
+		}
+		EXPECT_EQ(ChildCount[0], 0);
+		EXPECT_EQ(ChildCount[1], 0);
+		EXPECT_EQ(ChildCount[2], 2);
+	}
+
 
-	
 	GTEST_TEST(AllTraits, GeometryCollection_RigidBodies_ClusterTest_DeactivateClusterParticle)
 	{
 		FFramework UnitTest;
diff --git a/Engine/Source/Runtime/Experimental/Chaos/Private/Chaos/PhysicsObjectInternal.h b/Engine/Source/Runtime/Experimental/Chaos/Private/Chaos/PhysicsObjectInternal.h
--- a/Engine/Source/Runtime/Experimental/Chaos/Private/Chaos/PhysicsObjectInternal.h
+++ b/Engine/Source/Runtime/Experimental/Chaos/Private/Chaos/PhysicsObjectInternal.h
@@ -173,7 +173,7 @@ namespace Chaos
 			case EPhysicsProxyType::GeometryCollectionType:
 			{
 				FGeometryDynamicCollection& Collection = GetGeometryCollectionDynamicCollection<Id>(*static_cast<FGeometryCollectionPhysicsProxy*>(Proxy));
-				return !Collection.Children[BodyIndex].IsEmpty();
+				return Collection.HasChildren(BodyIndex);
 			}
 			case EPhysicsProxyType::ClusterUnionProxy:
 			{
diff --git a/Engine/Source/Runtime/Experimental/Chaos/Private/GeometryCollectionProxyData.cpp b/Engine/Source/Runtime/Experimental/Chaos/Private/GeometryCollectionProxyData.cpp
--- a/Engine/Source/Runtime/Experimental/Chaos/Private/GeometryCollectionProxyData.cpp
+++ b/Engine/Source/Runtime/Experimental/Chaos/Private/GeometryCollectionProxyData.cpp
@@ -27,7 +27,6 @@ void FTransformDynamicCollection::Construct()
 	// Transform Group
 	AddExternalAttribute<FTransform3f>(FTransformCollection::TransformAttribute, FTransformCollection::TransformGroup, Transform);
 	AddExternalAttribute<bool>(FTransformCollection::ParentAttribute, FTransformCollection::TransformGroup, HasParent);
-	AddExternalAttribute<TSet<int32>>(FTransformCollection::ChildrenAttribute, FTransformCollection::TransformGroup, Children);
 }
 
 const FTransform3f& FTransformDynamicCollection::GetTransform(int32 Index) const
@@ -66,6 +65,19 @@ int32 FTransformDynamicCollection::GetParent(int32 Index) const
 	return HasParent[Index] ? RestCollection->Parent[Index] : INDEX_NONE;
 }
 
+bool FTransformDynamicCollection::HasChildren(int32 Index) const
+{
+	const TSet<int32>& Children(RestCollection->Children[Index]);
+	for (int32 Child : Children)
+	{
+		if (HasParent[Child])
+		{
+			return true;
+		}
+	}
+	return false;
+}
+
 
 /*
 * FGeometryDynamicCollection (FTransformDynamicCollection)
@@ -143,12 +155,11 @@ void FGeometryDynamicCollection::FInitialVelocityFacade::CopyFrom(const FGeometr
 
 ///////////////////////////////////////////////////////////////////////////////////////////////////////////////////////////
 
-FGeometryCollectionDynamicStateFacade::FGeometryCollectionDynamicStateFacade(FManagedArrayCollection& InCollection)
+FGeometryCollectionDynamicStateFacade::FGeometryCollectionDynamicStateFacade(FGeometryDynamicCollection& InCollection)
 	: ActiveAttribute(InCollection, FGeometryDynamicCollection::ActiveAttribute,  FTransformCollection::TransformGroup)
 	, DynamicStateAttribute(InCollection, FGeometryDynamicCollection::DynamicStateAttribute,  FTransformCollection::TransformGroup)
-	, ChildrenAttribute(InCollection, FTransformCollection::ChildrenAttribute,  FTransformCollection::TransformGroup)
-	, HasParentAttribute(InCollection, FTransformCollection::ParentAttribute, FTransformCollection::TransformGroup)
 	, InternalClusterParentTypeAttribute(InCollection, "InternalClusterParentTypeArray", FGeometryCollection::TransformGroup)
+	, DynamicCollection(InCollection)
 {
 }
 
@@ -156,8 +167,6 @@ bool FGeometryCollectionDynamicStateFacade::IsValid() const
 {
 	return ActiveAttribute.IsValid()
 		&& DynamicStateAttribute.IsValid()
-		&& ChildrenAttribute.IsValid()
-		&& HasParentAttribute.IsValid()
 		&& InternalClusterParentTypeAttribute.IsValid();
 }
 
@@ -180,13 +189,13 @@ bool FGeometryCollectionDynamicStateFacade::IsSleeping(int32 TransformIndex) con
 
 bool FGeometryCollectionDynamicStateFacade::HasChildren(int32 TransformIndex) const
 {
-	return (ChildrenAttribute.Get()[TransformIndex].Num() > 0);
+	return DynamicCollection.HasChildren(TransformIndex);
 }
 
 bool FGeometryCollectionDynamicStateFacade::HasBrokenOff(int32 TransformIndex) const
 {
 	const bool bIsActive = IsActive(TransformIndex);
-	const bool bHasParent = HasParentAttribute.Get()[TransformIndex];
+	const bool bHasParent = DynamicCollection.GetHasParent(TransformIndex);
 	return bIsActive && (!bHasParent) && IsDynamicOrSleeping(TransformIndex);
 }
 
diff --git a/Engine/Source/Runtime/Experimental/Chaos/Private/PhysicsProxy/GeometryCollectionPhysicsProxy.cpp b/Engine/Source/Runtime/Experimental/Chaos/Private/PhysicsProxy/GeometryCollectionPhysicsProxy.cpp
--- a/Engine/Source/Runtime/Experimental/Chaos/Private/PhysicsProxy/GeometryCollectionPhysicsProxy.cpp
+++ b/Engine/Source/Runtime/Experimental/Chaos/Private/PhysicsProxy/GeometryCollectionPhysicsProxy.cpp
@@ -631,7 +631,7 @@ void FGeometryCollectionPhysicsProxy::Initialize(Chaos::FPBDRigidsEvolutionBase
 #endif
 			}
 			// this step is necessary for Phase 2 where we need to walk back the hierarchy from children to parent 
-			if (bGeometryCollectionAlwaysGenerateGTCollisionForClusters && GameThreadCollection.Children[Index].Num() == 0)
+			if (bGeometryCollectionAlwaysGenerateGTCollisionForClusters && !GameThreadCollection.HasChildren(Index))
 			{
 				ChildrenToCheckForParentFix.Add(Index);
 			}
@@ -663,21 +663,23 @@ void FGeometryCollectionPhysicsProxy::Initialize(Chaos::FPBDRigidsEvolutionBase
 				
 						// let's make sure all our children have an implicit defined, other wise, postpone to next iteration 
 						bool bAllChildrenHaveCollision = true;
-						for (const int32& ChildIndex : GameThreadCollection.Children[ParentToFixIndex])
+						
+						GameThreadCollection.IterateThroughChildren(ParentToFixIndex, [&](int32 ChildIndex)
 						{
 							// defer if any of the children is a cluster with no collision yet generated 
-							if (GameThreadCollection.Implicits[ChildIndex] == nullptr && GameThreadCollection.Children[ChildIndex].Num() > 0)
+							if (GameThreadCollection.Implicits[ChildIndex] == nullptr && GameThreadCollection.HasChildren(ChildIndex))
 							{
 								bAllChildrenHaveCollision = false;
-								break;
+								return false;
 							}
-						}
+							return true;
+						});
 
 						if (bAllChildrenHaveCollision)
 						{
 							// Make a union of the children geometry
 							TArray<Chaos::FImplicitObjectPtr> ChildImplicits;
-							for (const int32& ChildIndex : GameThreadCollection.Children[ParentToFixIndex])
+							GameThreadCollection.IterateThroughChildren(ParentToFixIndex, [&](int32 ChildIndex)
 							{
 								const Chaos::FImplicitObjectPtr& ChildImplicit = GameThreadCollection.Implicits[ChildIndex];
 								if (ChildImplicit)
@@ -690,7 +692,7 @@ void FGeometryCollectionPhysicsProxy::Initialize(Chaos::FPBDRigidsEvolutionBase
 									// if this remains a union we need to unpack it 
 									if (TransformedChildImplicit->IsUnderlyingUnion())
 									{
-										// we just move the array fo children in the new array 
+										// we just move the array for children in the new array 
 										Chaos::FImplicitObjectUnion& Union = static_cast<Chaos::FImplicitObjectUnion&>(*TransformedChildImplicit);
 										ChildImplicits.Append(MoveTemp(Union.GetObjects()));
 									}
@@ -699,7 +701,8 @@ void FGeometryCollectionPhysicsProxy::Initialize(Chaos::FPBDRigidsEvolutionBase
 										ChildImplicits.Add(MoveTemp(TransformedChildImplicit));
 									}
 								}
-							}
+								return true;
+							});
 							if (ChildImplicits.Num() > 0)
 							{
 								Chaos::FImplicitObject* UnionImplicit = new Chaos::FImplicitObjectUnion(MoveTemp(ChildImplicits));
@@ -864,7 +867,7 @@ void FGeometryCollectionPhysicsProxy::InitializeDynamicCollection(FGeometryDynam
 			// If no simulation data is available then default to the simulation of just the rigid geometry.
 			for (int32 TransformIdx = 0; TransformIdx < NumTransforms; TransformIdx++)
 			{
-				if (DynamicCollection.Children[TransformIdx].Num())
+				if (DynamicCollection.HasChildren(TransformIdx))
 				{
 					DynamicCollection.SimulatableParticles[TransformIdx] = false;
 				}
@@ -1064,7 +1067,6 @@ void FGeometryCollectionPhysicsProxy::InitializeBodiesPT(Chaos::FPBDRigidsSolver
 		const TManagedArray<FTransform>& MassToLocal = RestCollection->GetAttribute<FTransform>(MassToLocalAttributeName, FTransformCollection::TransformGroup);
 		const TManagedArray<Chaos::FImplicitObjectPtr>& Implicits = DynamicCollection.Implicits;
 		const TManagedArray<TUniquePtr<FCollisionStructureManager::FSimplicial>>& Simplicials = DynamicCollection.Simplicials;
-		const TManagedArray<TSet<int32>>& Children = DynamicCollection.Children;
 
 		// In PushToPhysicsState, we're going to compute a relative transform from Parameters.PrevWorldTransform
 		// to a particle's current world transform to get its relative transform. Then, that relative transform is
@@ -1244,15 +1246,15 @@ void FGeometryCollectionPhysicsProxy::InitializeBodiesPT(Chaos::FPBDRigidsSolver
 				else
 				{
 					// Cluster parent
-					const TSet<int32>& ChildIndices = Children[TransformGroupIndex];
-					for (const int32 ChildIndex : ChildIndices)
+					GameThreadCollection.IterateThroughChildren(TransformGroupIndex, [&](int32 ChildIndex)
 					{
-						if(SubTreeContainsSimulatableParticle[ChildIndex])
+						if (SubTreeContainsSimulatableParticle[ChildIndex])
 						{
 							SubTreeContainsSimulatableParticle[TransformGroupIndex] = true;
-							break;
+							return false;
 						}
-					}
+						return true;
+					});
 				}
 			}
 
@@ -1275,35 +1277,36 @@ void FGeometryCollectionPhysicsProxy::InitializeBodiesPT(Chaos::FPBDRigidsSolver
 					continue;
 				}
 
-				RigidChildren.Reset(Children.Num());
-				RigidChildrenTransformGroupIndex.Reset(Children.Num());
+				RigidChildren.Reset(NumTransforms);
+				RigidChildrenTransformGroupIndex.Reset(NumTransforms);
 
 				float ParentMass = 0.0f;
 				Chaos::FMatrix33 ParentInertia = Chaos::FMatrix33(0);
 
-				if (!Children[TransformGroupIndex].IsEmpty())
+				if (DynamicCollection.HasChildren(TransformGroupIndex))
 				{
 					Masses[TransformGroupIndex] = 0;
 					Inertias[TransformGroupIndex] = Chaos::FVec3(0.0);
 				}
 
 				Chaos::FMatrix33 FullInertia = Chaos::FMatrix33(0);
-				for (const int32 ChildIndex : Children[TransformGroupIndex])
-				{
-					Masses[TransformGroupIndex] += Masses[ChildIndex];
-					const Chaos::FMatrix33 ChildWorldSpaceI = Chaos::Utilities::ComputeWorldSpaceInertia(Transforms[ChildIndex].GetRotation(), Inertias[ChildIndex]);
-					FullInertia += ChildWorldSpaceI;
-					Inertias[TransformGroupIndex] = FullInertia.GetDiagonal();
-
-					if (EffectiveParticles[ChildIndex])
+				GameThreadCollection.IterateThroughChildren(TransformGroupIndex, [&](int32 ChildIndex)
 					{
-						if (Chaos::TPBDRigidClusteredParticleHandle<Chaos::FReal, 3>*Handle = SolverParticleHandles[ChildIndex])
+						Masses[TransformGroupIndex] += Masses[ChildIndex];
+						const Chaos::FMatrix33 ChildWorldSpaceI = Chaos::Utilities::ComputeWorldSpaceInertia(Transforms[ChildIndex].GetRotation(), Inertias[ChildIndex]);
+						FullInertia += ChildWorldSpaceI;
+						Inertias[TransformGroupIndex] = FullInertia.GetDiagonal();
+
+						if (EffectiveParticles[ChildIndex])
 						{
-							RigidChildren.Add(Handle);
-							RigidChildrenTransformGroupIndex.Add(ChildIndex);
+							if (Chaos::TPBDRigidClusteredParticleHandle<Chaos::FReal, 3>*Handle = SolverParticleHandles[ChildIndex])
+							{
+								RigidChildren.Add(Handle);
+								RigidChildrenTransformGroupIndex.Add(ChildIndex);
+							}
 						}
-					}
-				}
+						return true;
+					});
 
 				if (SolverParticleHandles[TransformGroupIndex] == nullptr)
 				{
@@ -2642,7 +2645,7 @@ TBitArray<> FGeometryCollectionPhysicsProxy::CalculateClustersToCreateFromChildr
 		TArray<int32> ChildrenToCheckForParentFix;
 		for (int32 Index = 0; Index < NumTransforms; ++Index)
 		{
-			if (DynamicCollection.Children[Index].Num() == 0)
+			if (!DynamicCollection.HasChildren(Index))
 			{
 				ChildrenToCheckForParentFix.Add(Index);
 			}
@@ -2668,16 +2671,17 @@ TBitArray<> FGeometryCollectionPhysicsProxy::CalculateClustersToCreateFromChildr
 				{
 					// let's make sure all our children have an implicit defined, otherwise, postpone to next iteration 
 					bool bAllChildrenHaveCollision = true;
-					for (int32 ChildIndex : DynamicCollection.Children[ParentToFixIndex])
+					DynamicCollection.IterateThroughChildren(ParentToFixIndex, [&](int32 ChildIndex)
 					{
 						// defer if any of the children is a cluster with no collision yet generated 
 						const bool bChildHasCollision = (DynamicCollection.Implicits[ChildIndex] != nullptr) || ClustersToGenerate[ChildIndex];
-						if (!bChildHasCollision && (DynamicCollection.Children[ChildIndex].Num() > 0))
+						if (!bChildHasCollision && DynamicCollection.HasChildren(ChildIndex))
 						{
 							bAllChildrenHaveCollision = false;
-							break;
+							return false;
 						}
-					}
+						return true;
+					});
 
 					ClustersToGenerate[ParentToFixIndex] = bAllChildrenHaveCollision;
 				}
@@ -3338,7 +3342,6 @@ inline static bool RemoveFromParentInCollection_Internal(FGeometryDynamicCollect
 	// CONTEXT : PHYSICS THREAD
 	if (int32 ParentIndex = PhysicsThreadCollection.GetParent(TransformGroupIndex); ParentIndex != INDEX_NONE)
 	{
-		PhysicsThreadCollection.Children[ParentIndex].Remove(TransformGroupIndex);
 		PhysicsThreadCollection.SetHasParent(TransformGroupIndex, false);
 		return true;
 	}
diff --git a/Engine/Source/Runtime/Experimental/Chaos/Public/GeometryCollectionProxyData.h b/Engine/Source/Runtime/Experimental/Chaos/Public/GeometryCollectionProxyData.h
--- a/Engine/Source/Runtime/Experimental/Chaos/Public/GeometryCollectionProxyData.h
+++ b/Engine/Source/Runtime/Experimental/Chaos/Public/GeometryCollectionProxyData.h
@@ -37,19 +37,36 @@ public:
 	CHAOS_API void SetHasParent(int32 Index, bool Value);
 	CHAOS_API int32 GetParent(int32 Index) const;
 
-	// Transform Group
-	TManagedArray<TSet<int32>>  Children;
+	CHAOS_API bool HasChildren(int32 Index) const;
+
+	template<typename Lambda>
+	void IterateThroughChildren(int32 Index, Lambda&& LambdaIt) const
+	{
+		const TSet<int32>& Children = RestCollection->Children[Index];
+		for (const int32 Child : Children)
+		{
+			if (GetHasParent(Child))
+			{
+				bool bContinue = LambdaIt(Child);
+				if (!bContinue)
+				{
+					break;
+				}
+			}
+		}
+	}
+
 
 private:
 	TManagedArray<bool>         HasParent;
 	TManagedArray<FTransform3f> Transform;
+	const FGeometryCollection* RestCollection;
+
 	/** Construct */
 	CHAOS_API void Construct();
 
 	friend class FGeometryCollectionPhysicsProxy;
 	friend class UGeometryCollectionComponent;
-
-	const FGeometryCollection* RestCollection;
 };
 
 
@@ -131,7 +148,7 @@ public:
 class FGeometryCollectionDynamicStateFacade
 {
 public:
-	CHAOS_API FGeometryCollectionDynamicStateFacade(FManagedArrayCollection& InCollection);
+	CHAOS_API FGeometryCollectionDynamicStateFacade(FGeometryDynamicCollection& InCollection);
 
 	/** returns true if all the necessary attributes are present */
 	CHAOS_API bool IsValid() const;
@@ -145,7 +162,7 @@ public:
 	/** return true if the transform is in a sleeping state */
 	CHAOS_API bool IsSleeping(int32 TransformIndex) const;
 
-	/** whether there's children attached to this transfom (Cluster) */
+	/** whether there's children attached to this transform (Cluster) */
 	CHAOS_API bool HasChildren(int32 TransformIndex) const;
 	
 	/** return true if the transform has broken off its parent */
@@ -167,14 +184,10 @@ private:
 	/** physics state of the transform (Dynamic, kinematic, static, sleeping) */
 	TManagedArrayAccessor<int32> DynamicStateAttribute;
 
-	/** currently attached children (potentially different from the initial children setup) */
-	TManagedArrayAccessor<TSet<int32>> ChildrenAttribute;
-	
-	/** Current parent (potentially different from the initial parent) */
-	TManagedArrayAccessor<bool> HasParentAttribute;
-
 	/** type of internal state parent */
 	TManagedArrayAccessor<uint8> InternalClusterParentTypeAttribute;
+
+	FGeometryDynamicCollection& DynamicCollection;
 };
 
 class FGeometryCollectioPerFrameData
diff --git a/Engine/Source/Runtime/Experimental/GeometryCollectionEngine/Private/GeometryCollection/GeometryCollectionComponent.cpp b/Engine/Source/Runtime/Experimental/GeometryCollectionEngine/Private/GeometryCollection/GeometryCollectionComponent.cpp
--- a/Engine/Source/Runtime/Experimental/GeometryCollectionEngine/Private/GeometryCollection/GeometryCollectionComponent.cpp
+++ b/Engine/Source/Runtime/Experimental/GeometryCollectionEngine/Private/GeometryCollection/GeometryCollectionComponent.cpp
@@ -2016,7 +2016,6 @@ void UGeometryCollectionComponent::UpdateRepData()
 		const FRigidClustering& RigidClustering = Solver->GetEvolution()->GetRigidClustering();
 
 		const TManagedArray<int32>* InitialLevels = PhysicsProxy->GetPhysicsCollection().FindAttribute<int32>("InitialLevel", FGeometryCollection::TransformGroup);
-		const TManagedArray<TSet<int32>>& InitialChildren = PhysicsProxy->GetPhysicsCollection().Children;
 		const TArray<FPBDRigidClusteredParticleHandle*>& ParticleHandles = PhysicsProxy->GetParticles();
 
 		// Replicate the anchored state of the root particle
@@ -2944,7 +2943,6 @@ void UGeometryCollectionComponent::SetInitialClusterBreaks(const TArray<int32>&
 {
 	if (DynamicCollection)
 	{
-		TManagedArray <TSet<int32>>& Children = DynamicCollection->Children;
 		const int32 NumTransforms = DynamicCollection->GetTransforms().Num();
 
 		for (int32 ReleaseIndex : ReleaseIndices)
@@ -2953,7 +2951,6 @@ void UGeometryCollectionComponent::SetInitialClusterBreaks(const TArray<int32>&
 			{
 				if (int32 ParentIndex = DynamicCollection->GetParent(ReleaseIndex); ParentIndex > INDEX_NONE)
 				{
-					Children[ParentIndex].Remove(ReleaseIndex);
 					DynamicCollection->SetHasParent(ReleaseIndex, false);
 				}
 			}
@@ -3504,7 +3501,6 @@ void UGeometryCollectionComponent::ResetDynamicCollection()
 
 		GetTransformArrayCopyOnWrite();
 		GetParentArrayCopyOnWrite();
-		GetChildrenArrayCopyOnWrite();
 
 		if (bStoreVelocities || bNotifyTrailing)
 		{
diff --git a/Engine/Source/Runtime/Experimental/GeometryCollectionEngine/Private/GeometryCollection/GeometryCollectionEngineRemoval.cpp b/Engine/Source/Runtime/Experimental/GeometryCollectionEngine/Private/GeometryCollection/GeometryCollectionEngineRemoval.cpp
--- a/Engine/Source/Runtime/Experimental/GeometryCollectionEngine/Private/GeometryCollection/GeometryCollectionEngineRemoval.cpp
+++ b/Engine/Source/Runtime/Experimental/GeometryCollectionEngine/Private/GeometryCollection/GeometryCollectionEngineRemoval.cpp
@@ -3,6 +3,7 @@
 #include "GeometryCollection/GeometryCollectionEngineRemoval.h"
 #include "GeometryCollection/GeometryCollection.h"
 #include "GeometryCollection/Facades/CollectionRemoveOnBreakFacade.h"
+#include "GeometryCollectionProxyData.h"
 
 ///////////////////////////////////////////////////////////////////////////////////////////////////////////////////////////
 
@@ -20,12 +21,11 @@ namespace
 	const FName DecayAttributeName = "Decay";
 }
 
-FGeometryCollectionRemoveOnBreakDynamicFacade::FGeometryCollectionRemoveOnBreakDynamicFacade(FManagedArrayCollection& InCollection)
+FGeometryCollectionRemoveOnBreakDynamicFacade::FGeometryCollectionRemoveOnBreakDynamicFacade(FGeometryDynamicCollection& InCollection)
 	: BreakTimerAttribute(InCollection, BreakTimerAttributeName, FGeometryCollection::TransformGroup)
 	, PostBreakDurationAttribute(InCollection, PostBreakDurationAttributeName, FGeometryCollection::TransformGroup)
 	, BreakRemovalDurationAttribute(InCollection, BreakRemovalDurationAttributeName, FGeometryCollection::TransformGroup)
-	, ChildrenAttribute(InCollection, FTransformCollection::ChildrenAttribute, FTransformCollection::TransformGroup)
-
+	, DynamicCollection(InCollection)
 {
 }
 
@@ -33,9 +33,7 @@ bool FGeometryCollectionRemoveOnBreakDynamicFacade::IsValid() const
 {
 	return BreakTimerAttribute.IsValid()
 		&& PostBreakDurationAttribute.IsValid()
-		&& BreakRemovalDurationAttribute.IsValid()
-		&& ChildrenAttribute.IsValid();
-	;
+		&& BreakRemovalDurationAttribute.IsValid();
 }
 
 bool FGeometryCollectionRemoveOnBreakDynamicFacade::IsConst() const
@@ -56,13 +54,11 @@ void FGeometryCollectionRemoveOnBreakDynamicFacade::SetAttributeValues(const Geo
 	check(!IsConst());
 	if (IsValid())
 	{
-		const TManagedArray<TSet<int32>>& Children = ChildrenAttribute.Get();
-
 		BreakTimerAttribute.Fill(DisabledBreakTimer);
 
 		// make sure we generate random value consistently between client and server 
 		// we can use the length of the transform group for that
-		const int32 RandomSeed = Children.Num();
+		const int32 RandomSeed = DynamicCollection.GetTransforms().Num();
 		FRandomStream Random(RandomSeed);
 
 		TManagedArray<float>& PostBreakDuration = PostBreakDurationAttribute.Modify();
@@ -82,7 +78,7 @@ void FGeometryCollectionRemoveOnBreakDynamicFacade::SetAttributeValues(const Geo
 			const FVector2f RemovalTimer = RemoveOnBreakData.GetRemovalTimer();
 			const float MinRemovalTime = FMath::Max(0.0f, RemovalTimer.X);
 			const float MaxRemovalTime = FMath::Max(MinRemovalTime, RemovalTimer.Y);
-			const bool bIsCluster = (Children[Idx].Num() > 0);
+			const bool bIsCluster = DynamicCollection.HasChildren(Idx);
 			const bool bUseClusterCrumbling = (bIsCluster && RemoveOnBreakData.GetClusterCrumbling());
 			BreakRemovalDuration[Idx] = bUseClusterCrumbling ? CrumblingRemovalTimer : Random.FRandRange(MinRemovalTime, MaxRemovalTime);
 		}
diff --git a/Engine/Source/Runtime/Experimental/GeometryCollectionEngine/Public/GeometryCollection/GeometryCollectionComponent.h b/Engine/Source/Runtime/Experimental/GeometryCollectionEngine/Public/GeometryCollection/GeometryCollectionComponent.h
--- a/Engine/Source/Runtime/Experimental/GeometryCollectionEngine/Public/GeometryCollection/GeometryCollectionComponent.h
+++ b/Engine/Source/Runtime/Experimental/GeometryCollectionEngine/Public/GeometryCollection/GeometryCollectionComponent.h
@@ -885,7 +885,6 @@ public:
 	/* Transform group */																				\
 	COPY_ON_WRITE_ATTRIBUTE(FString, BoneName, FTransformCollection::TransformGroup)					\
 	COPY_ON_WRITE_ATTRIBUTE(FLinearColor, BoneColor, FTransformCollection::TransformGroup)				\
-	COPY_ON_WRITE_ATTRIBUTE(TSet<int32>, Children, FTransformCollection::TransformGroup)				\
 	COPY_ON_WRITE_ATTRIBUTE(int32, TransformToGeometryIndex, FTransformCollection::TransformGroup)		\
 	COPY_ON_WRITE_ATTRIBUTE(int32, ExemplarIndex, FTransformCollection::TransformGroup)					\
 
diff --git a/Engine/Source/Runtime/Experimental/GeometryCollectionEngine/Public/GeometryCollection/GeometryCollectionEngineRemoval.h b/Engine/Source/Runtime/Experimental/GeometryCollectionEngine/Public/GeometryCollection/GeometryCollectionEngineRemoval.h
--- a/Engine/Source/Runtime/Experimental/GeometryCollectionEngine/Public/GeometryCollection/GeometryCollectionEngineRemoval.h
+++ b/Engine/Source/Runtime/Experimental/GeometryCollectionEngine/Public/GeometryCollection/GeometryCollectionEngineRemoval.h
@@ -6,6 +6,7 @@
 #include "GeometryCollection/ManagedArrayAccessor.h"
 
 struct FManagedArrayCollection;
+class FGeometryDynamicCollection;
 
 namespace GeometryCollection::Facades
 {
@@ -23,7 +24,7 @@ public:
 	static constexpr float DisabledPostBreakDuration = -1;
 	static constexpr float CrumblingRemovalTimer = -1;
 	
-	FGeometryCollectionRemoveOnBreakDynamicFacade(FManagedArrayCollection& InCollection);
+	FGeometryCollectionRemoveOnBreakDynamicFacade(FGeometryDynamicCollection& InCollection);
 
 	/** 
 	 * returns true if all the necessary attributes are present
@@ -46,7 +47,7 @@ public:
 	/** true if the removal is active for a specific piece */
 	bool IsRemovalActive(int32 TransformIndex) const;
 
-	/** true if a specific transfom uses cluster crumbling */
+	/** true if a specific transform uses cluster crumbling */
 	bool UseClusterCrumbling(int32 TransformIndex) const;
 	
 	/**
@@ -67,8 +68,7 @@ private:
 	/** removal duration */
 	TManagedArrayAccessor<float> BreakRemovalDurationAttribute;
 
-	/** utility attribute needed to compute the values and know what's a cluster */
-	TManagedArrayAccessor<TSet<int32>> ChildrenAttribute;
+	const FGeometryDynamicCollection& DynamicCollection;
 };
 
 /**
@@ -118,7 +118,7 @@ public:
 	 */
 	void UpdateSleepTimer(int32 TransformIndex, float DeltaTime);
 
-	/** Compute decay from elapsed timer and durantion attributes */
+	/** Compute decay from elapsed timer and duration attributes */
 	float ComputeDecay(int32 TransformIndex) const;
 
 	
