geometry collection should return false for CanEditSimulatePhysics if no rest collection is set

diff --git a/Engine/Source/Runtime/Experimental/GeometryCollectionEngine/Private/GeometryCollection/GeometryCollectionComponent.cpp b/Engine/Source/Runtime/Experimental/GeometryCollectionEngine/Private/GeometryCollection/GeometryCollectionComponent.cpp
--- a/Engine/Source/Runtime/Experimental/GeometryCollectionEngine/Private/GeometryCollection/GeometryCollectionComponent.cpp
+++ b/Engine/Source/Runtime/Experimental/GeometryCollectionEngine/Private/GeometryCollection/GeometryCollectionComponent.cpp
@@ -1170,7 +1170,7 @@ void UGeometryCollectionComponent::SetNotifyRigidBodyCollision(bool bNewNotifyRi
 
 bool UGeometryCollectionComponent::CanEditSimulatePhysics()
 {
-	return true;
+	return (RestCollection != nullptr);
 }
 
 void UGeometryCollectionComponent::SetSimulatePhysics(bool bEnabled)
