Ias will now generate the OnDemandToc via utoc files on disk rather than downloading it from a CDN.

#rnx

- Flipped cvar 's.IasGenerateOnDemandToc' to default to true.
- Added a specific line of logging when generating the OnDemandToc from disk.

diff --git a/Engine/Source/Runtime/Experimental/IoStoreOnDemand/Private/OnDemandIoDispatcherBackend.cpp b/Engine/Source/Runtime/Experimental/IoStoreOnDemand/Private/OnDemandIoDispatcherBackend.cpp
--- a/Engine/Source/Runtime/Experimental/IoStoreOnDemand/Private/OnDemandIoDispatcherBackend.cpp
+++ b/Engine/Source/Runtime/Experimental/IoStoreOnDemand/Private/OnDemandIoDispatcherBackend.cpp
@@ -134,7 +134,7 @@ static TAutoConsoleVariable<bool> CVar_IoReportAnalytics(
 
 static FAutoConsoleVariable CVar_IasGenerateOnDemandToc(
 	TEXT("s.IasGenerateOnDemandToc"),
-	false,
+	true,
 	TEXT("Enables generating the FOnDemandToc from utoc files on disk rather than downloading them"),
 	ECVF_ReadOnly
 );
@@ -860,6 +860,8 @@ TIoStatusOr<FOnDemandToc> GenerateOnDemandTocFromDisk()
 {
 	TRACE_CPUPROFILER_EVENT_SCOPE(IasBackend::GenerateOnDemandTocFromDisk);
 
+	UE_LOG(LogIas, Log, TEXT("Generating OnDemandToc via utoc files from disk"));
+
 	FOnDemandToc OutToc;
 
 	TArray<FString> UtocFilePaths = FindOnDemandUtocFilesOnDisk();
