The FS lock is no longer required but we need to lock RWs to DataMap

#rnx

diff --git a/Engine/Source/Runtime/Experimental/IoStoreOnDemand/Private/JournaledCache.cpp b/Engine/Source/Runtime/Experimental/IoStoreOnDemand/Private/JournaledCache.cpp
--- a/Engine/Source/Runtime/Experimental/IoStoreOnDemand/Private/JournaledCache.cpp
+++ b/Engine/Source/Runtime/Experimental/IoStoreOnDemand/Private/JournaledCache.cpp
@@ -558,6 +558,7 @@ private:
 	uint64					Insert(uint64 DataBase, const FDataEntry& Entry);
 	uint64					Insert(uint64 DataBase, const FDataEntry* Entries, uint32 EntryCount);
 	void					Prune(uint64 DataBase, uint32 Size);
+	mutable FRWLock			Lock;
 	FString					BinPath;
 	FDataMap				DataMap;
 	uint64					MappedBytes = 0;
@@ -642,8 +643,11 @@ void FDiskCache::ClosePhrase(FDiskPhrase&& Phrase)
 		DataHandle->Write(Buffer, WriteSize);
 	}
 
-	Prune(DataCursor, WriteSize);
-	Insert(DataCursor, Phrase.GetEntries(), EntryCount);
+	{
+		FWriteScopeLock _(Lock);
+		Prune(DataCursor, WriteSize);
+		Insert(DataCursor, Phrase.GetEntries(), EntryCount);
+	}
 
 	Journal.ClosePhrase(MoveTemp(Phrase), DataCursor);
 	DataCursor += WriteSize;
@@ -652,6 +656,7 @@ void FDiskCache::ClosePhrase(FDiskPhrase&& Phrase)
 ////////////////////////////////////////////////////////////////////////////////
 bool FDiskCache::Has(uint64 Key) const
 {
+	FReadScopeLock _(Lock);
 	return (DataMap.Find(Key) != nullptr);
 }
 
@@ -665,13 +670,20 @@ EIoErrorCode FDiskCache::Materialize(uint64 Key, FIoBuffer& Out, uint32 Offset)
 		return EIoErrorCode::FileNotOpen;
 	}
 
-	const FMapEntry* Entry = DataMap.Find(Key);
-	if (Entry == nullptr)
+	uint32 ReadSize;
+	uint64 EntryDataCursor;
 	{
-		return EIoErrorCode::NotFound;
-	}
+		FReadScopeLock _(Lock);
 
-	uint32 ReadSize = uint32(Entry->Size) - Offset;
+		const FMapEntry* Entry = DataMap.Find(Key);
+		if (Entry == nullptr)
+		{
+			return EIoErrorCode::NotFound;
+		}
+
+		ReadSize = uint32(Entry->Size) - Offset;
+		EntryDataCursor = Entry->DataCursor;
+	}
 
 	if (Out.GetData() == nullptr)
 	{
@@ -680,9 +692,9 @@ EIoErrorCode FDiskCache::Materialize(uint64 Key, FIoBuffer& Out, uint32 Offset)
 
 	ReadSize = FMath::Min<uint32>(uint32(Out.GetSize()), ReadSize);
 
-	TRACE_COUNTER_SET(IasReadCursor, Entry->DataCursor + Offset);
+	TRACE_COUNTER_SET(IasReadCursor, EntryDataCursor + Offset);
 
-	DataHandle->Seek(Entry->DataCursor + Offset);
+	DataHandle->Seek(EntryDataCursor + Offset);
 	bool bOk = DataHandle->Read(Out.GetData(), ReadSize);
 	return bOk ? EIoErrorCode::Ok : EIoErrorCode::ReadError;
 }
@@ -788,6 +800,8 @@ void FDiskCache::Drop()
 ////////////////////////////////////////////////////////////////////////////////
 void FDiskCache::Spam()
 {
+	FReadScopeLock _(Lock);
+
 	UE_LOG(LogIas, VeryVerbose,
 		TEXT("JournaledCache: MappedKiB=%llu Entries=%d DataCur=%llu JournalCur=%u Marker=%u)"),
 		(MappedBytes >> 10),
@@ -801,6 +815,8 @@ void FDiskCache::Spam()
 ////////////////////////////////////////////////////////////////////////////////
 uint32 FDiskCache::DebugVisit(void* Param, FDebugCacheEntry::Callback* Callback)
 {
+	FReadScopeLock _(Lock);
+
 	FDebugCacheEntry Out = {};
 	for (const auto& Entry : DataMap)
 	{
@@ -818,6 +834,8 @@ uint32 FDiskCache::DebugVisit(void* Param, FDebugCacheEntry::Callback* Callback)
 ////////////////////////////////////////////////////////////////////////////////
 static bool LoadCache(FDiskCache& DiskCache)
 {
+	FWriteScopeLock _(DiskCache.Lock);
+
 	FDiskJournal& Journal = DiskCache.Journal;
 
 	uint32 DataSize = 0;
@@ -1090,7 +1108,6 @@ public:
 	uint32			DebugVisit(void* Param, FDebugCacheEntry::Callback* Callback);
 
 private:
-	mutable FRWLock	FsLock;
 	mutable FRWLock	MemLock;
 	FMemCache		MemCache;
 	FDiskCache		DiskCache;
@@ -1119,14 +1136,12 @@ uint32 FCache::GetAilments() const
 ////////////////////////////////////////////////////////////////////////////////
 bool FCache::Load()
 {
-	FWriteScopeLock _(FsLock);
 	return LoadCache(DiskCache);
 }
 
 ////////////////////////////////////////////////////////////////////////////////
 void FCache::Drop()
 {
-	FWriteScopeLock _(FsLock);
 	DiskCache.Drop();
 }
 
@@ -1139,7 +1154,7 @@ uint32 FCache::GetDemand() const
 ////////////////////////////////////////////////////////////////////////////////
 bool FCache::Has(uint64 Key) const
 {
-	if (FReadScopeLock _(FsLock); DiskCache.Has(Key))
+	if (DiskCache.Has(Key))
 	{
 		return true;
 	}
@@ -1152,7 +1167,7 @@ bool FCache::Has(uint64 Key) const
 FCache::FGetToken FCache::Get(uint64 Key, FIoBuffer& OutData) const
 {
 	// Disk first as that will have more data and is more likely to hit
-	if (FReadScopeLock _(FsLock); DiskCache.Has(Key))
+	if (DiskCache.Has(Key))
 	{
 		return FGetToken(Key);
 	}
@@ -1184,8 +1199,6 @@ bool FCache::Put(uint64 Key, FIoBuffer& Data)
 ////////////////////////////////////////////////////////////////////////////////
 EIoErrorCode FCache::Materialize(FGetToken Token, FIoBuffer& OutData, uint32 Offset) const
 {
-	FReadScopeLock _(FsLock);
-
 	uint64 Key = Token;
 	EIoErrorCode Ret = DiskCache.Materialize(Key, OutData, Offset);
 	if (Ret == EIoErrorCode::Ok)
@@ -1199,7 +1212,6 @@ EIoErrorCode FCache::Materialize(FGetToken Token, FIoBuffer& OutData, uint32 Off
 ////////////////////////////////////////////////////////////////////////////////
 uint32 FCache::Flush()
 {
-	FWriteScopeLock _(FsLock);
 	return DiskCache.Flush();
 }
 
@@ -1237,10 +1249,7 @@ uint32 FCache::WriteMemToDisk(int32 Allowance)
 		MemCacheSize = Phrase.GetDataSize();
 	}
 
-	{
-		FWriteScopeLock _(FsLock);
-		DiskCache.ClosePhrase(MoveTemp(Phrase));
-	}
+	DiskCache.ClosePhrase(MoveTemp(Phrase));
 
 	return MemCacheSize;
 }
@@ -1248,7 +1257,7 @@ uint32 FCache::WriteMemToDisk(int32 Allowance)
 ////////////////////////////////////////////////////////////////////////////////
 uint32 FCache::DebugVisit(void* Param, FDebugCacheEntry::Callback* Callback)
 {
-	FReadScopeLock _[] = { FReadScopeLock(MemLock), FReadScopeLock(FsLock) };
+	FReadScopeLock _(MemLock);
 	uint32 Count = 0;
 	Count += MemCache.DebugVisit(Param, Callback);
 	Count += DiskCache.DebugVisit(Param, Callback);
