Removed unused cancellation token object

#rnx

(git:ae45cf1 - Martin Ridgers - [REDACTED_PHONE]:10:32 +0200)

diff --git a/Engine/Source/Runtime/Experimental/IoStoreOnDemand/Private/CancellationToken.h b/Engine/Source/Runtime/Experimental/IoStoreOnDemand/Private/CancellationToken.h
deleted file mode 100644
diff --git a/Engine/Source/Runtime/Experimental/IoStoreOnDemand/Private/IasCache.h b/Engine/Source/Runtime/Experimental/IoStoreOnDemand/Private/IasCache.h
--- a/Engine/Source/Runtime/Experimental/IoStoreOnDemand/Private/IasCache.h
+++ b/Engine/Source/Runtime/Experimental/IoStoreOnDemand/Private/IasCache.h
@@ -2,7 +2,6 @@
 
 #pragma once
 
-#include "CancellationToken.h"
 #include "Containers/StringView.h"
 #include "IO/IoStatus.h"
 #include "Memory/MemoryFwd.h"
diff --git a/Engine/Source/Runtime/Experimental/IoStoreOnDemand/Private/JournaledCache.cpp b/Engine/Source/Runtime/Experimental/IoStoreOnDemand/Private/JournaledCache.cpp
--- a/Engine/Source/Runtime/Experimental/IoStoreOnDemand/Private/JournaledCache.cpp
+++ b/Engine/Source/Runtime/Experimental/IoStoreOnDemand/Private/JournaledCache.cpp
@@ -4,7 +4,6 @@
 #include "IO/IoStoreOnDemand.h"
 #include "Statistics.h"
 
-#include "CancellationToken.h"
 #include "Containers/Array.h"
 #include "Containers/Map.h"
 #include "Containers/UnrealString.h"
diff --git a/Engine/Source/Runtime/Experimental/IoStoreOnDemand/Private/OnDemandIoDispatcherBackend.cpp b/Engine/Source/Runtime/Experimental/IoStoreOnDemand/Private/OnDemandIoDispatcherBackend.cpp
--- a/Engine/Source/Runtime/Experimental/IoStoreOnDemand/Private/OnDemandIoDispatcherBackend.cpp
+++ b/Engine/Source/Runtime/Experimental/IoStoreOnDemand/Private/OnDemandIoDispatcherBackend.cpp
@@ -3,7 +3,6 @@
 #include "OnDemandIoDispatcherBackend.h"
 
 #include "AnalyticsEventAttribute.h"
-#include "CancellationToken.h"
 #include "Containers/BitArray.h"
 #include "Containers/StringView.h"
 #include "CoreHttp/LatencyTesting.h"
