Fix up some of the Ias cvars to match the style of the rest of the module.

#rnx

diff --git a/Engine/Source/Runtime/Experimental/IoStoreOnDemand/Private/OnDemandIoDispatcherBackend.cpp b/Engine/Source/Runtime/Experimental/IoStoreOnDemand/Private/OnDemandIoDispatcherBackend.cpp
--- a/Engine/Source/Runtime/Experimental/IoStoreOnDemand/Private/OnDemandIoDispatcherBackend.cpp
+++ b/Engine/Source/Runtime/Experimental/IoStoreOnDemand/Private/OnDemandIoDispatcherBackend.cpp
@@ -126,21 +126,24 @@ static FAutoConsoleVariableRef CVar_IasHttpOptionalBulkDataEnabled(
 	TEXT("Enables optional bulk data via HTTP")
 );
 
-static TAutoConsoleVariable<bool> CVar_IoReportAnalytics(
+bool GIasReportAnalyticsEnabled = true;
+static FAutoConsoleVariableRef CVar_IoReportAnalytics(
 	TEXT("ias.ReportAnalytics"),
-	true,
+	GIasReportAnalyticsEnabled,
 	TEXT("Enables reporting statics to the analytics system"));
 
-static FAutoConsoleVariable CVar_IasGenerateOnDemandToc(
+bool GIasGenerateOnDemandToc = true;
+static FAutoConsoleVariableRef CVar_IasGenerateOnDemandToc(
 	TEXT("s.IasGenerateOnDemandToc"),
-	true,
+	GIasGenerateOnDemandToc,
 	TEXT("Enables generating the FOnDemandToc from utoc files on disk rather than downloading them"),
 	ECVF_ReadOnly
 );
 
-static FAutoConsoleVariable CVar_IasEnableAsyncTocGeneration(
+bool GIasAsyncTocGenerationEnabled = true;
+static FAutoConsoleVariableRef CVar_IasAsyncTocGeneration(
 	TEXT("s.IasEnableThreadedTocGeneration"),
-	true,
+	GIasAsyncTocGenerationEnabled,
 	TEXT("Enables pushing the work FOnDemandToc generation work to the task system"),
 	ECVF_ReadOnly
 );
@@ -1393,10 +1396,10 @@ FOnDemandIoBackend::FOnDemandIoBackend(TUniquePtr<IIasCache>&& InCache)
 	DynamicConsoleCommands.Emplace(
 	IConsoleManager::Get().RegisterConsoleCommand(
 		TEXT("ias.InvokeHttpFailure"),
-		TEXT(""),
+		TEXT("Marks the current ias http connection as failed forcing the system to try to reconnect"),
 		FConsoleCommandDelegate::CreateLambda([this]()
 			{
-				UE_LOG(LogIas, Display, TEXT("Marks the current ias http connection as failed forcing the system to try to reconnect"));
+				UE_LOG(LogIas, Display, TEXT("User invoked http error via 'ias.InvokeHttpFailure'"));
 				BackendStatus.SetHttpError(true);
 
 				TickBackendEvent->Trigger();
@@ -1788,7 +1791,7 @@ FIoStatus FOnDemandIoBackend::ApplyGeneratedOnDemandToc(const FString& CdnUrl, c
 {
 	TRACE_CPUPROFILER_EVENT_SCOPE(IasBackend::ApplyGeneratedOnDemandToc);
 
-	if (bGeneratedOnDemandToc || !CVar_IasGenerateOnDemandToc->GetBool())
+	if (bGeneratedOnDemandToc || !GIasGenerateOnDemandToc)
 	{
 		return FIoStatus::Ok;
 	}
@@ -1871,7 +1874,7 @@ void FOnDemandIoBackend::AddDeferredTocs()
 	// So keeping that in mind we need to prevent the toc being downloaded if we are supposed to be using the disk generated version 
 	// instead.
 	// This can all be cleaned up when we either remove the toggle or drop the validation requirements AND clean up the initialization flow.
-	if (!CVar_IasGenerateOnDemandToc->GetBool())
+	if (!GIasGenerateOnDemandToc)
 	{
 		for (const FString& TocPath : TocPaths)
 		{
@@ -1895,7 +1898,7 @@ void FOnDemandIoBackend::Mount(const FOnDemandEndpoint& Endpoint)
 		return;
 	}
 
-	if (CVar_IasGenerateOnDemandToc->GetBool() && CVar_IasEnableAsyncTocGeneration->GetBool())
+	if (GIasGenerateOnDemandToc && GIasAsyncTocGenerationEnabled)
 	{
 		OnDemandTocTask = UE::Tasks::Launch(UE_SOURCE_LOCATION, [this]() -> TIoStatusOr<FOnDemandToc>
 			{
@@ -1934,7 +1937,7 @@ void FOnDemandIoBackend::Mount(const FOnDemandEndpoint& Endpoint)
 			UE_LOG(LogIas, Error, TEXT("Failed to add generated toc', reason '%s'"), *GeneratedResult.ToString());
 		}
 
-		if (!CVar_IasGenerateOnDemandToc->GetBool())
+		if (!GIasGenerateOnDemandToc)
 		{
 			FIoStatus Result = DownloadoadOnDemandToc(AvailableEps.GetCurrent(), Endpoint.TocPath);
 			if (!Result.IsOk())
@@ -1972,7 +1975,7 @@ void FOnDemandIoBackend::AbandonCache()
 
 void FOnDemandIoBackend::ReportAnalytics(TArray<FAnalyticsEventAttribute>& OutAnalyticsArray) const
 {
-	if (!CVar_IoReportAnalytics.GetValueOnAnyThread())
+	if (!GIasReportAnalyticsEnabled)
 	{
 		return;
 	}
diff --git a/Engine/Source/Runtime/Experimental/IoStoreOnDemand/Private/Statistics.cpp b/Engine/Source/Runtime/Experimental/IoStoreOnDemand/Private/Statistics.cpp
--- a/Engine/Source/Runtime/Experimental/IoStoreOnDemand/Private/Statistics.cpp
+++ b/Engine/Source/Runtime/Experimental/IoStoreOnDemand/Private/Statistics.cpp
@@ -13,14 +13,16 @@ LLM_DEFINE_TAG(Ias);
 namespace UE::IO::IAS
 {
 
-static TAutoConsoleVariable<bool> CVar_ReportHttpAnalytics(
+bool GIasReportHttpAnalyticsEnabled = true;
+static FAutoConsoleVariableRef CVar_ReportHttpAnalytics(
 	TEXT("ias.ReportHttpAnalytics"),
-	true,
+	GIasReportHttpAnalyticsEnabled,
 	TEXT("Enables reporting statics on our http traffic to the analytics system"));
 
-static TAutoConsoleVariable<bool> CVar_ReportCacheAnalytics(
+bool GIasReportCacheAnalyticsEnabled = true;
+static FAutoConsoleVariableRef CVar_ReportCacheAnalytics(
 	TEXT("ias.ReportCacheAnalytics"),
-	true,
+	GIasReportCacheAnalyticsEnabled,
 	TEXT("Enables reporting statics on our file cache usage to the analytics system"));
 
 ////////////////////////////////////////////////////////////////////////////////
@@ -308,7 +310,7 @@ FOnDemandIoBackendStats* FOnDemandIoBackendStats::Get()
 
 void FOnDemandIoBackendStats::ReportAnalytics(TArray<FAnalyticsEventAttribute>& OutAnalyticsArray) const
 {
-	if (CVar_ReportHttpAnalytics.GetValueOnAnyThread())
+	if (GIasReportHttpAnalyticsEnabled)
 	{
 		AppendAnalyticsEventAttributeArray(OutAnalyticsArray
 			,TEXT("IasHttpErrorCount"), GHttpErrorCount.Get()
@@ -326,7 +328,7 @@ void FOnDemandIoBackendStats::ReportAnalytics(TArray<FAnalyticsEventAttribute>&
 		);
 	}
 
-	if (CVar_ReportCacheAnalytics.GetValueOnAnyThread())
+	if (GIasReportCacheAnalyticsEnabled)
 	{
 		AppendAnalyticsEventAttributeArray(OutAnalyticsArray
 
