Corrected page size causing check()

#rnx

diff --git a/Engine/Source/Runtime/Experimental/IoStoreOnDemand/Private/CoreHttp/Http.cpp b/Engine/Source/Runtime/Experimental/IoStoreOnDemand/Private/CoreHttp/Http.cpp
--- a/Engine/Source/Runtime/Experimental/IoStoreOnDemand/Private/CoreHttp/Http.cpp
+++ b/Engine/Source/Runtime/Experimental/IoStoreOnDemand/Private/CoreHttp/Http.cpp
@@ -2238,7 +2238,7 @@ static int32 DoRecvMessage(FActivity* Activity)
 {
 	TRACE_CPUPROFILER_EVENT_SCOPE(IasHttp::DoRecvMessage);
 
-	static const uint32 PageSize = 384;
+	static const uint32 PageSize = 256;
 
 	FBuffer& Buffer = Activity->Buffer;
 
