Iris
- Fixed assert in ReplicationWriter

diff --git a/Engine/Plugins/Runtime/ReplicationSystemTestPlugin/Source/Private/Tests/ReplicationSystem/TestDependentObjects.cpp b/Engine/Plugins/Runtime/ReplicationSystemTestPlugin/Source/Private/Tests/ReplicationSystem/TestDependentObjects.cpp
--- a/Engine/Plugins/Runtime/ReplicationSystemTestPlugin/Source/Private/Tests/ReplicationSystem/TestDependentObjects.cpp
+++ b/Engine/Plugins/Runtime/ReplicationSystemTestPlugin/Source/Private/Tests/ReplicationSystem/TestDependentObjects.cpp
@@ -721,4 +721,42 @@ UE_NET_TEST_FIXTURE(FReplicationSystemServerClientTestFixture, TestLateAddedNest
 	UE_NET_ASSERT_LT(ClientNestedDependentObject->LastRepOrderCounter, ClientObject->LastRepOrderCounter);
 }
 
+UE_NET_TEST_FIXTURE(FReplicationSystemServerClientTestFixture, TestLateAddedNestedDependentObjectPendingWaitOnCreateConfirmation)
+{
+	UReplicationSystem* ReplicationSystem = Server->ReplicationSystem;
+	UReplicatedTestObjectBridge* Bridge = Server->GetReplicationBridge();
+
+	// Add a client
+	FReplicationSystemTestClient* Client = CreateClient();
+
+	// Spawn objects on server
+	UReplicatedSubObjectOrderObject* ServerSharedDependentObject = Server->CreateObject<UReplicatedSubObjectOrderObject>();
+	UReplicatedSubObjectOrderObject* ServerDependentObject = Server->CreateObject<UReplicatedSubObjectOrderObject>();
+
+	// Set static prio to only replicate with parent
+	ReplicationSystem->SetStaticPriority(ServerDependentObject->NetRefHandle, 1.f);
+
+	// Reset RepOrderCounter
+	UReplicatedSubObjectOrderObject::RepOrderCounter = 0U;
+
+	// Create server object
+	UReplicatedSubObjectOrderObject* ServerObject = Server->CreateObject<UReplicatedSubObjectOrderObject>();
+
+	// Add same dependent to both ServerObject and ServerDependentObject
+	Bridge->AddDependentObject(ServerObject->NetRefHandle, ServerDependentObject->NetRefHandle);
+	Bridge->AddDependentObject(ServerObject->NetRefHandle, ServerSharedDependentObject->NetRefHandle);
+	Bridge->AddDependentObject(ServerDependentObject->NetRefHandle, ServerSharedDependentObject->NetRefHandle);
+
+	// Send data
+	Server->PreSendUpdate();
+	Server->SendAndDeliverTo(Client, true);
+	Server->PostSendUpdate();
+
+	// Verify that objects have replicated
+	UReplicatedSubObjectOrderObject* ClientObject = Cast<UReplicatedSubObjectOrderObject>(Client->GetReplicationBridge()->GetReplicatedObject(ServerObject->NetRefHandle));
+	UReplicatedSubObjectOrderObject* ClientDependentObject = Cast<UReplicatedSubObjectOrderObject>(Client->GetReplicationBridge()->GetReplicatedObject(ServerDependentObject->NetRefHandle));
+	UReplicatedSubObjectOrderObject* ClientSharedDependentObject = Cast<UReplicatedSubObjectOrderObject>(Client->GetReplicationBridge()->GetReplicatedObject(ServerSharedDependentObject->NetRefHandle));
+}
+
+
 }
diff --git a/Engine/Source/Runtime/Experimental/Iris/Core/Private/Iris/ReplicationSystem/ReplicationWriter.cpp b/Engine/Source/Runtime/Experimental/Iris/Core/Private/Iris/ReplicationSystem/ReplicationWriter.cpp
--- a/Engine/Source/Runtime/Experimental/Iris/Core/Private/Iris/ReplicationSystem/ReplicationWriter.cpp
+++ b/Engine/Source/Runtime/Experimental/Iris/Core/Private/Iris/ReplicationSystem/ReplicationWriter.cpp
@@ -2198,6 +2198,11 @@ FReplicationWriter::EWriteObjectStatus FReplicationWriter::WriteObjectAndSubObje
 		FinalBatchEntry.bSentTearOff = bSentTearOff;
 		FinalBatchEntry.bSentDestroySubObject = Info.SubObjectPendingDestroy;
 		FinalBatchEntry.NewBaselineIndex = CreatedBaselineIndex;
+		if (InternalIndex != ObjectIndexForOOBAttachment)
+		{
+			// Mark this object as written this tick to avoid sending it multiple times
+			WriteContext.ObjectsWrittenThisPacket.SetBit(InternalIndex);
+		}
 	}
 
 	// Reset CreatedBaselineIndex to avoid it being released on scope exit
@@ -2305,7 +2310,7 @@ FReplicationWriter::EWriteObjectStatus FReplicationWriter::WriteObjectInBatch(FN
 		return WriteObjectStatus;
 	}
 
-	// Include dependent objects as separate batch
+	// Include dependent objects as separate batch, (for hugeobjects they will be included as they are written to a separate bitstream)
 	{
 		const uint32 OldBatchInfoParentInternalIndex = OutBatchInfo.ParentInternalIndex;
 		for (const FDependentObjectInfo DependentObjectInfo : NetRefHandleManager->GetDependentObjectInfos(InternalIndex))
@@ -2942,13 +2947,6 @@ int FReplicationWriter::HandleObjectBatchSuccess(const FBatchInfo& BatchInfo, FR
 		{
 			SchedulingPriorities[BatchObjectInfo.InternalIndex] = 0.0f;
 		}
-
-		if (BatchObjectInfo.InternalIndex != ObjectIndexForOOBAttachment)
-		{
-			// Mark this object as written this tick to avoid sending it multiple times
-			WriteContext.ObjectsWrittenThisPacket.SetBit(BatchObjectInfo.InternalIndex);
-		}
-
 	}
 
 #if UE_NET_IRIS_CSV_STATS && CSV_PROFILER
@@ -2974,7 +2972,7 @@ int FReplicationWriter::HandleObjectBatchSuccess(const FBatchInfo& BatchInfo, FR
 	return WrittenObjectCount;
 }
 
-FReplicationWriter::EWriteObjectRetryMode FReplicationWriter::HandleObjectBatchFailure(FReplicationWriter::EWriteObjectStatus WriteObjectStatus, const FBatchInfo& BatchInfo, const FReplicationWriter::FBitStreamInfo& BatchBitStreamInfo) const
+FReplicationWriter::EWriteObjectRetryMode FReplicationWriter::HandleObjectBatchFailure(FReplicationWriter::EWriteObjectStatus WriteObjectStatus, const FBatchInfo& BatchInfo, const FReplicationWriter::FBitStreamInfo& BatchBitStreamInfo)
 {
 	IRIS_PROFILER_SCOPE(FReplicationWriter_HandleObjectBatchFailure);
 	
@@ -2986,6 +2984,12 @@ FReplicationWriter::EWriteObjectRetryMode FReplicationWriter::HandleObjectBatchF
 		{
 			BaselineManager->LostBaseline(Parameters.ConnectionId, BatchObjectInfo.InternalIndex, BatchObjectInfo.NewBaselineIndex);
 		}
+		
+		if (BatchObjectInfo.InternalIndex != ObjectIndexForOOBAttachment)
+		{
+			// If we failed to write the batch and we wrote data for an object we need to mark it as not written, if we want to try again
+			WriteContext.ObjectsWrittenThisPacket.ClearBit(BatchObjectInfo.InternalIndex);
+		}
 	}
 
 	if (WriteObjectStatus == EWriteObjectStatus::NoInstanceProtocol || WriteObjectStatus == EWriteObjectStatus::InvalidOwner)
diff --git a/Engine/Source/Runtime/Experimental/Iris/Core/Private/Iris/ReplicationSystem/ReplicationWriter.h b/Engine/Source/Runtime/Experimental/Iris/Core/Private/Iris/ReplicationSystem/ReplicationWriter.h
--- a/Engine/Source/Runtime/Experimental/Iris/Core/Private/Iris/ReplicationSystem/ReplicationWriter.h
+++ b/Engine/Source/Runtime/Experimental/Iris/Core/Private/Iris/ReplicationSystem/ReplicationWriter.h
@@ -471,7 +471,7 @@ private:
 	int HandleObjectBatchSuccess(const FBatchInfo& BatchInfo, FBatchRecord& OutRecord);
 
 	// Determines the best course of action after a WriteObjectBatch() call failed.
-	EWriteObjectRetryMode HandleObjectBatchFailure(EWriteObjectStatus WriteObjectStatus, const FBatchInfo& BatchInfo, const FBitStreamInfo& BatchBitStreamInfo) const;
+	EWriteObjectRetryMode HandleObjectBatchFailure(EWriteObjectStatus WriteObjectStatus, const FBatchInfo& BatchInfo, const FBitStreamInfo& BatchBitStreamInfo);
 
 	// Update logic for dropped RecordInfo
 	void HandleDroppedRecord(const FReplicationRecord::FRecordInfo& RecordInfo, FReplicationInfo& Info, const FNetObjectAttachmentsWriter::FReliableReplicationRecord& AttachmentRecord);
