Clustering optimization : add two cvar (soft and hard ) to control the number of released particles based on the number of active geometry collection particle

diff --git a/Engine/Source/Runtime/Experimental/Chaos/Private/Chaos/PBDRigidClustering.cpp b/Engine/Source/Runtime/Experimental/Chaos/Private/Chaos/PBDRigidClustering.cpp
--- a/Engine/Source/Runtime/Experimental/Chaos/Private/Chaos/PBDRigidClustering.cpp
+++ b/Engine/Source/Runtime/Experimental/Chaos/Private/Chaos/PBDRigidClustering.cpp
@@ -64,6 +64,13 @@ namespace Chaos
 	float RestoreBreakingMomentumPercent = .5;
 	FAutoConsoleVariableRef CVarRestoreBreakingMomentumPercent(TEXT("p.RestoreBreakingMomentumPercent"), RestoreBreakingMomentumPercent, TEXT("When a rigid cluster is broken, objects that its in contact with will receive an impulse to restore this percent of their momentum prior to the break."));
 
+	int32 ClusteringParticleReleaseThrottlingMinCount = INDEX_NONE;
+	FAutoConsoleVariableRef CVarClusteringParticleReleaseThrottlingMinCount(TEXT("p.Clustering.ParticleReleaseThrottlingMinCount"), ClusteringParticleReleaseThrottlingMinCount, TEXT("Minimum number of active geometry collection to reach before clustering start to disable a percentage of the released particle per cluster"));
+
+	int32 ClusteringParticleReleaseThrottlingMaxCount = INDEX_NONE;
+	FAutoConsoleVariableRef CVarClusteringParticleReleaseThrottlingMaxCount(TEXT("p.Clustering.ParticleReleaseThrottlingMaxCount"), ClusteringParticleReleaseThrottlingMaxCount, TEXT("Maximum number of active geometry collection to reach before all released clustering disable all released particle instantly"));
+
+
 	template <typename TProxy=FGeometryCollectionPhysicsProxy>
 	TProxy* GetConcreteProxy(FPBDRigidClusteredParticleHandle* ClusteredParticle)
 	{
@@ -119,6 +126,29 @@ namespace Chaos
 
 			return Current;
 		}
+
+		// compute a ratio (between 0 and 1) of released particles to release
+		float GetRatioOfReleasedParticlesToDisable(const FRigidClustering::FRigidEvolution& Evolution)
+		{
+			float Percentage = 1.0f;
+			if (ClusteringParticleReleaseThrottlingMinCount >= 0 && ClusteringParticleReleaseThrottlingMaxCount >= 0)
+			{
+				const FPBDRigidsSOAs& ParticleStructures = Evolution.GetParticles();
+				int32 NumActiveParticles = 0;
+				NumActiveParticles += ParticleStructures.GetSleepingGeometryCollectionArray().Num();
+				NumActiveParticles += ParticleStructures.GetDynamicGeometryCollectionArray().Num();
+
+				const int32 Range = FMath::Max(0, (ClusteringParticleReleaseThrottlingMaxCount - ClusteringParticleReleaseThrottlingMinCount));
+				const int32 OverMinCount = FMath::Max(0, (NumActiveParticles - ClusteringParticleReleaseThrottlingMinCount));
+
+				if (Range > 0)
+				{
+					// clamp to 1, as OverMinCount can get larger than Range
+					Percentage = FMath::Min(1.f, ((float)OverMinCount / (float)Range));
+				}
+			}
+			return Percentage;
+		}
 	}
 	
 	//==========================================================================
@@ -1159,6 +1189,27 @@ namespace Chaos
 			}
 		}
 
+		// optimization : start disabling activated children if the number of active released particle is too high
+		const float RatioOfParticlesToDisable = GetRatioOfReleasedParticlesToDisable(MEvolution);
+		const int32 NumberOfParticlesToDisable = (int32)((float)ActivatedChildren.Num() * RatioOfParticlesToDisable);
+		if (NumberOfParticlesToDisable > 0)
+		{ 
+			int32 DisabledParticleCount = 0;
+			for (auto ChildIt = ActivatedChildren.CreateIterator(); ChildIt; ++ChildIt)
+			{
+				if (FPBDRigidParticleHandle* Child = *ChildIt)
+				{
+					DisabledParticleCount++;
+					MEvolution.DisableParticle(Child);
+					ChildIt.RemoveCurrent();
+				}
+				if (DisabledParticleCount >= NumberOfParticlesToDisable)
+				{
+					break;
+				}
+			}
+		}
+
 		FrameReleasedChildren += ActivatedChildren.Num();
 
 		return ActivatedChildren;
diff --git a/Engine/Source/Runtime/Experimental/Chaos/Public/Chaos/PBDRigidsSOAs.h b/Engine/Source/Runtime/Experimental/Chaos/Public/Chaos/PBDRigidsSOAs.h
--- a/Engine/Source/Runtime/Experimental/Chaos/Public/Chaos/PBDRigidsSOAs.h
+++ b/Engine/Source/Runtime/Experimental/Chaos/Public/Chaos/PBDRigidsSOAs.h
@@ -969,6 +969,9 @@ public:
 	const TPBDGeometryCollectionParticles<FReal, 3>& GetGeometryCollectionParticles() const { return *GeometryCollectionParticles; }
 	TPBDGeometryCollectionParticles<FReal, 3>& GetGeometryCollectionParticles() { return *GeometryCollectionParticles; }
 
+	const TArray<FPBDGeometryCollectionParticleHandle*>& GetSleepingGeometryCollectionArray() const { return SleepingGeometryCollectionArray.GetArray(); }
+	const TArray<FPBDGeometryCollectionParticleHandle*>& GetDynamicGeometryCollectionArray() const { return DynamicGeometryCollectionArray.GetArray(); }
+
 	void InsertGeometryCollectionParticle(TPBDGeometryCollectionParticleHandle<FReal, 3>* GCParticle)
 	{
 		if (!GCParticle->Disabled())
