Chaos Cloth:
- Added a p.ChaosCloth.Solver.MaximumVelocity console variable to clamp the relative particles' velocity to the skinned cloth mesh velocity.
- Added a p.Chaos.VelocityField.MaximumVelocity console variable to clamp the relative particles' velocity to the aerodynamcis velocity field.
- Added a p.Cloth.MaxVelocityScale console variable to clamp the maximum velocity scale values allowed on all cloth simulations.

diff --git a/Engine/Plugins/Experimental/ChaosCloth/Source/Chaos/Private/ChaosCloth/ChaosClothingSimulationSolver.cpp b/Engine/Plugins/Experimental/ChaosCloth/Source/Chaos/Private/ChaosCloth/ChaosClothingSimulationSolver.cpp
--- a/Engine/Plugins/Experimental/ChaosCloth/Source/Chaos/Private/ChaosCloth/ChaosClothingSimulationSolver.cpp
+++ b/Engine/Plugins/Experimental/ChaosCloth/Source/Chaos/Private/ChaosCloth/ChaosClothingSimulationSolver.cpp
@@ -53,6 +53,7 @@ static bool bClothSolverParallelClothUpdate = true;
 static bool bClothSolverParallelClothPostUpdate = true;
 static bool bClothSolverDisableTimeDependentNumIterations = false;
 static bool bClothSolverUseVelocityScale = true;
+static float ClothSolverMaxVelocity = 0.f;
 
 #if !UE_BUILD_SHIPPING
 static int32 ClothSolverDebugHitchLength = 0;
@@ -70,6 +71,7 @@ FAutoConsoleVariableRef CVarClothSolverDisableCollision(TEXT("p.ChaosCloth.Solve
 
 FAutoConsoleVariableRef CVarClothSolverDisableTimeDependentNumIterations(TEXT("p.ChaosCloth.Solver.DisableTimeDependentNumIterations"), bClothSolverDisableTimeDependentNumIterations, TEXT("Make the number of iterations independent from the time step."));
 FAutoConsoleVariableRef CVarClothSolverUseVelocityScale(TEXT("p.ChaosCloth.Solver.UseVelocityScale"), bClothSolverUseVelocityScale, TEXT("Use the velocity scale to compensate for clamping to MaxPhysicsDelta, in order to avoid miscalculating velocities during hitches."));
+FAutoConsoleVariableRef CVarClothSolverMaxVelocity(TEXT("p.ChaosCloth.Solver.MaxVelocity"), ClothSolverMaxVelocity, TEXT("Maximum relative velocity of the cloth particles relatively to their animated positions equivalent. 0 to disable."));
 
 namespace ClothingSimulationSolverDefault
 {
@@ -944,8 +946,10 @@ void FClothingSimulationSolver::ApplyPreSimulationTransforms()
 	const TPBDActiveView<Softs::FSolverParticles>& ParticlesActiveView = Evolution->ParticlesActiveView();
 	const TArray<uint32>& ParticleGroupIds = Evolution->ParticleGroupIds();
 
+	const FSolverReal MaxVelocitySquared = (ClothSolverMaxVelocity > 0.f) ? FMath::Square((FSolverReal)ClothSolverMaxVelocity) : TNumericLimits<FSolverReal>::Max();
+
 	ParticlesActiveView.RangeFor(
-		[this, &ParticleGroupIds, &DeltaLocalSpaceLocation](Softs::FSolverParticles& Particles, int32 Offset, int32 Range)
+		[this, &ParticleGroupIds, &DeltaLocalSpaceLocation, MaxVelocitySquared](Softs::FSolverParticles& Particles, int32 Offset, int32 Range)
 		{
 			TRACE_CPUPROFILER_EVENT_SCOPE(FClothingSimulationSolver_ParticlePreSimulationTransforms);
 			SCOPE_CYCLE_COUNTER(STAT_ChaosClothParticlePreSimulationTransforms);
@@ -955,26 +959,47 @@ void FClothingSimulationSolver::ApplyPreSimulationTransforms()
 #if INTEL_ISPC
 			if (bRealTypeCompatibleWithISPC && bChaos_PreSimulationTransforms_ISPC_Enabled)  // TODO: Make the ISPC works with both Single and Double depending on the FSolverReal type
 			{
-				ispc::ApplyPreSimulationTransforms(
-					(ispc::FVector4f*)Particles.GetPAndInvM().GetData(),
-					(ispc::FVector3f*)Particles.GetV().GetData(),
-					(ispc::FVector3f*)Particles.XArray().GetData(),
-					(ispc::FVector3f*)OldAnimationPositions.GetData(),
-					(ispc::FVector3f*)AnimationVelocities.GetData(),
-					Particles.GetInvM().GetData(),
-					(const ispc::FVector3f*)AnimationPositions.GetData(),
-					ParticleGroupIds.GetData(),
-					(ispc::FTransform3f*)PreSimulationTransforms.GetData(),
-					(ispc::FVector3f&)DeltaLocalSpaceLocation,
-					DeltaTime,
-					Offset,
-					Range);
+				if (MaxVelocitySquared == TNumericLimits<FSolverReal>::Max())
+				{
+					ispc::ApplyPreSimulationTransforms(
+						(ispc::FVector4f*)Particles.GetPAndInvM().GetData(),
+						(ispc::FVector3f*)Particles.GetV().GetData(),
+						(ispc::FVector3f*)Particles.XArray().GetData(),
+						(ispc::FVector3f*)OldAnimationPositions.GetData(),
+						(ispc::FVector3f*)AnimationVelocities.GetData(),
+						Particles.GetInvM().GetData(),
+						(const ispc::FVector3f*)AnimationPositions.GetData(),
+						ParticleGroupIds.GetData(),
+						(ispc::FTransform3f*)PreSimulationTransforms.GetData(),
+						(ispc::FVector3f&)DeltaLocalSpaceLocation,
+						DeltaTime,
+						Offset,
+						Range);
+				}
+				else
+				{
+					ispc::ApplyPreSimulationTransformsAndClampVelocity(
+						(ispc::FVector4f*)Particles.GetPAndInvM().GetData(),
+						(ispc::FVector3f*)Particles.GetV().GetData(),
+						(ispc::FVector3f*)Particles.XArray().GetData(),
+						(ispc::FVector3f*)OldAnimationPositions.GetData(),
+						(ispc::FVector3f*)AnimationVelocities.GetData(),
+						Particles.GetInvM().GetData(),
+						(const ispc::FVector3f*)AnimationPositions.GetData(),
+						ParticleGroupIds.GetData(),
+						(ispc::FTransform3f*)PreSimulationTransforms.GetData(),
+						(ispc::FVector3f&)DeltaLocalSpaceLocation,
+						DeltaTime,
+						Offset,
+						Range,
+						MaxVelocitySquared);
+				}
 			}
 			else
 #endif
 			{
 				PhysicsParallelFor(RangeSize,
-					[this, &ParticleGroupIds, &DeltaLocalSpaceLocation, &Particles, Offset](int32 i)
+					[this, &ParticleGroupIds, &DeltaLocalSpaceLocation, &Particles, Offset, MaxVelocitySquared](int32 i)
 					{
 						const int32 Index = Offset + i;
 						const Softs::FSolverRigidTransform3& GroupSpaceTransform = PreSimulationTransforms[ParticleGroupIds[Index]];
@@ -991,6 +1016,14 @@ void FClothingSimulationSolver::ApplyPreSimulationTransforms()
 
 						// Update Animation velocities
 						AnimationVelocities[Index] = (AnimationPositions[Index] - OldAnimationPositions[Index]) / DeltaTime;
+
+						// Clamp relative velocity
+						const FSolverVec3 RelVelocity = Particles.V(Index) - AnimationVelocities[Index];
+						const FSolverReal RelVelocitySquaredLength = RelVelocity.SquaredLength();
+						if (RelVelocitySquaredLength > MaxVelocitySquared)
+						{
+							Particles.V(Index) = AnimationVelocities[Index] + RelVelocity * FMath::Sqrt(MaxVelocitySquared / RelVelocitySquaredLength);
+						}
 					}, RangeSize < ClothSolverMinParallelBatchSize);
 			}
 		}, /*bForceSingleThreaded =*/ !bClothSolverParallelClothPreUpdate);
diff --git a/Engine/Plugins/Experimental/ChaosCloth/Source/Chaos/Private/ChaosCloth/ChaosClothingSimulationSolver.ispc b/Engine/Plugins/Experimental/ChaosCloth/Source/Chaos/Private/ChaosCloth/ChaosClothingSimulationSolver.ispc
--- a/Engine/Plugins/Experimental/ChaosCloth/Source/Chaos/Private/ChaosCloth/ChaosClothingSimulationSolver.ispc
+++ b/Engine/Plugins/Experimental/ChaosCloth/Source/Chaos/Private/ChaosCloth/ChaosClothingSimulationSolver.ispc
@@ -49,6 +49,59 @@ export void ApplyPreSimulationTransforms(uniform FVector4f ParticlesPandInvM[],
 	}
 }
 
+export void ApplyPreSimulationTransformsAndClampVelocity(uniform FVector4f ParticlesPandInvM[],
+										uniform FVector3f ParticlesV[],
+										uniform FVector3f ParticlesX[],
+										uniform FVector3f OldAnimationPositions[],
+										uniform FVector3f AnimationVelocities[],
+										const uniform float InvM[],
+										const uniform FVector3f AnimationPositions[],
+										const uniform uint32 ParticleGroupIds[],
+										const uniform FTransform3f PreSimulationTransforms[],
+										const uniform FVector3f &DeltaLocalSpaceLocation,
+										const uniform float DeltaTime,
+										const uniform int32 Offset,
+										const uniform int32 Range,
+										const uniform float MaxVelocitySquared)
+{
+	foreach(i = Offset ... Range)
+	{
+		const varying uint32 ParticleGroupId = ParticleGroupIds[i];
+		#pragma ignore warning(perf)
+		const varying FTransform3f GroupSpaceTransform = PreSimulationTransforms[ParticleGroupId];
+
+		const uniform int32 Index = extract(i, 0);
+
+		const varying FVector3f ParticleX = VectorLoad(&ParticlesX[Index]);
+		const varying FVector3f PX = TransformPosition(GroupSpaceTransform, ParticleX) - DeltaLocalSpaceLocation;
+
+		VectorStore(&ParticlesPandInvM[Index], SetVector4(PX, InvM[i]));
+		VectorStore(&ParticlesX[Index], PX);
+
+		const varying FVector3f ParticleV = VectorLoad(&ParticlesV[Index]);
+		varying FVector3f V = TransformVector(GroupSpaceTransform, ParticleV);
+
+		const varying FVector3f OldAnimationPosition = VectorLoad(&OldAnimationPositions[Index]);
+		const varying FVector3f AP = TransformPosition(GroupSpaceTransform, OldAnimationPosition) - DeltaLocalSpaceLocation;
+
+		VectorStore(&OldAnimationPositions[Index], AP);
+
+		const varying FVector3f AnimationPosition = VectorLoad(&AnimationPositions[Index]);
+		const varying FVector3f AV = (AnimationPosition - AP) / DeltaTime;
+
+		VectorStore(&AnimationVelocities[Index], AV);
+
+		const varying FVector3f RelVelocity = V - AV;
+		const varying float RelVelocitySquared = VectorSizeSquared(RelVelocity);
+		if (RelVelocitySquared > MaxVelocitySquared)
+		{
+			V = AV + RelVelocity * sqrt(MaxVelocitySquared/RelVelocitySquared);
+		}
+
+		VectorStore(&ParticlesV[Index], V);
+	}
+}
+
 export void PreSubstepInterpolation( uniform FVector3f InterpolatedPositions[],
 									uniform FVector3f InterpolatedNormals[],
 									const uniform FVector3f AnimationPositions[],
diff --git a/Engine/Source/Runtime/ClothingSystemRuntimeCommon/Private/ClothingSimulation.cpp b/Engine/Source/Runtime/ClothingSystemRuntimeCommon/Private/ClothingSimulation.cpp
--- a/Engine/Source/Runtime/ClothingSystemRuntimeCommon/Private/ClothingSimulation.cpp
+++ b/Engine/Source/Runtime/ClothingSystemRuntimeCommon/Private/ClothingSimulation.cpp
@@ -20,6 +20,14 @@ static TAutoConsoleVariable<float> GClothMaxDeltaTimeTeleportMultiplier(
 	TEXT("A multiplier of the MaxPhysicsDelta time at which we will automatically just teleport cloth to its new location\n")
 	TEXT(" default: 1.5"));
 
+static TAutoConsoleVariable<float> GClothMaxVelocityScale(
+	TEXT("p.Cloth.MaxVelocityScale"),
+	1.f,
+	TEXT("The maximum amount of the component induced velocity allowed on all cloths.\n")
+	TEXT("Use 1.0 for fully induced velocity(default), or use 0.0 for no induced velocity, and any other values in between for a reduced induced velocity.\n")
+	TEXT("When set to 0.0, it also provides a way to force the clothing to simulate in local space.\n")
+	TEXT(" default: 1.0"));
+
 PRAGMA_DISABLE_DEPRECATION_WARNINGS // CachedPositions, CachedVelocities
 FClothingSimulationContextCommon::FClothingSimulationContextCommon()
 	: ComponentToWorld(FTransform::Identity)
@@ -172,8 +180,11 @@ void FClothingSimulationContextCommon::FillTeleportMode(const USkeletalMeshCompo
 		EClothingTeleportMode::Teleport :
 		InComponent->ClothTeleportMode;
 
+	const float MaxVelocityScale = FMath::Clamp(GClothMaxVelocityScale.GetValueOnGameThread(), 0.f, 1.f);
+	const float ComponentVelocityScale = InComponent->ClothVelocityScale;
+
 	VelocityScale = (TeleportMode == EClothingTeleportMode::None && InDeltaSeconds > 0.f) ?
-		InComponent->ClothVelocityScale * FMath::Min(InDeltaSeconds, InMaxPhysicsDelta) / InDeltaSeconds : 0.f;
+		FMath::Clamp(ComponentVelocityScale, 0.f, MaxVelocityScale) * FMath::Min(InDeltaSeconds, InMaxPhysicsDelta) / InDeltaSeconds : 0.f;
 }
 
 void FClothingSimulationContextCommon::FillMaxDistanceScale(const USkeletalMeshComponent* InComponent)
diff --git a/Engine/Source/Runtime/Experimental/Chaos/Private/Chaos/VelocityField.cpp b/Engine/Source/Runtime/Experimental/Chaos/Private/Chaos/VelocityField.cpp
--- a/Engine/Source/Runtime/Experimental/Chaos/Private/Chaos/VelocityField.cpp
+++ b/Engine/Source/Runtime/Experimental/Chaos/Private/Chaos/VelocityField.cpp
@@ -16,6 +16,12 @@ FAutoConsoleVariableRef CVarChaosVelocityFieldISPCEnabled(TEXT("p.Chaos.Velocity
 
 namespace Chaos::Softs {
 
+namespace Private
+{
+	static float VelocityFieldMaxVelocity = 0.f;
+	static FAutoConsoleVariableRef CVarChaosVelocityFieldMaxVelocity(TEXT("p.Chaos.VelocityField.MaxVelocity"), VelocityFieldMaxVelocity, TEXT("The maximum relative velocity to process the aerodynamics forces with."));
+}
+
 void FVelocityAndPressureField::SetProperties(
 	const FCollectionPropertyConstFacade& PropertyCollection,
 	const TMap<FString, TConstArrayView<FRealSingle>>& Weightmaps,
@@ -199,29 +205,59 @@ void FVelocityAndPressureField::SetMultipliers(
 
 void FVelocityAndPressureField::UpdateForces(const FSolverParticles& InParticles, const FSolverReal /*Dt*/)
 {
+	const FSolverReal MaxVelocitySquared = (Private::VelocityFieldMaxVelocity > 0.f) ? FMath::Square((FSolverReal)Private::VelocityFieldMaxVelocity) : TNumericLimits<FSolverReal>::Max();
+
 	if (!Multipliers.Num())
 	{
 #if INTEL_ISPC
 		if (bRealTypeCompatibleWithISPC && bChaos_VelocityField_ISPC_Enabled)
 		{
-			ispc::UpdateField(
-				(ispc::FVector3f*)Forces.GetData(),
-				(const ispc::FIntVector*)Elements.GetData(),
-				(const ispc::FVector3f*)InParticles.GetV().GetData(),
-				(const ispc::FVector3f*)InParticles.XArray().GetData(),
-				(const ispc::FVector3f&)Velocity,
-				QuarterRho,
-				DragBase,
-				LiftBase,
-				PressureBase,
-				Elements.Num());
+			if (MaxVelocitySquared == TNumericLimits<FSolverReal>::Max())
+			{
+				ispc::UpdateField(
+					(ispc::FVector3f*)Forces.GetData(),
+					(const ispc::FIntVector*)Elements.GetData(),
+					(const ispc::FVector3f*)InParticles.GetV().GetData(),
+					(const ispc::FVector3f*)InParticles.XArray().GetData(),
+					(const ispc::FVector3f&)Velocity,
+					QuarterRho,
+					DragBase,
+					LiftBase,
+					PressureBase,
+					Elements.Num());
+			}
+			else
+			{
+				ispc::UpdateFieldAndClampVelocity(
+					(ispc::FVector3f*)Forces.GetData(),
+					(const ispc::FIntVector*)Elements.GetData(),
+					(const ispc::FVector3f*)InParticles.GetV().GetData(),
+					(const ispc::FVector3f*)InParticles.XArray().GetData(),
+					(const ispc::FVector3f&)Velocity,
+					QuarterRho,
+					DragBase,
+					LiftBase,
+					PressureBase,
+					Elements.Num(),
+					MaxVelocitySquared);
+			}
 		}
 		else
 #endif
 		{
-			for (int32 ElementIndex = 0; ElementIndex < Elements.Num(); ++ElementIndex)
+			if (MaxVelocitySquared == TNumericLimits<FSolverReal>::Max())
 			{
-				UpdateField(InParticles, ElementIndex, Velocity, DragBase, LiftBase, PressureBase);
+				for (int32 ElementIndex = 0; ElementIndex < Elements.Num(); ++ElementIndex)
+				{
+					UpdateField(InParticles, ElementIndex, Velocity, DragBase, LiftBase, PressureBase);
+				}
+			}
+			else
+			{
+				for (int32 ElementIndex = 0; ElementIndex < Elements.Num(); ++ElementIndex)
+				{
+					UpdateField(InParticles, ElementIndex, Velocity, DragBase, LiftBase, PressureBase, MaxVelocitySquared);
+				}
 			}
 		}
 	}
@@ -230,33 +266,70 @@ void FVelocityAndPressureField::UpdateForces(const FSolverParticles& InParticles
 #if INTEL_ISPC
 		if (bRealTypeCompatibleWithISPC && bChaos_VelocityField_ISPC_Enabled)
 		{
-			ispc::UpdateFieldWithWeightMaps(
-				(ispc::FVector3f*)Forces.GetData(),
-				(const ispc::FIntVector*)Elements.GetData(),
-				(const ispc::FVector3f*)InParticles.GetV().GetData(),
-				(const ispc::FVector3f*)InParticles.XArray().GetData(),
-				(const ispc::FVector3f*)Multipliers.GetData(),
-				(const ispc::FVector3f&)Velocity,
-				QuarterRho,
-				DragBase,
-				DragRange,
-				LiftBase,
-				LiftRange,
-				PressureBase,
-				PressureRange,
-				Elements.Num());
+			if (MaxVelocitySquared == TNumericLimits<FSolverReal>::Max())
+			{
+				ispc::UpdateFieldWithWeightMaps(
+					(ispc::FVector3f*)Forces.GetData(),
+					(const ispc::FIntVector*)Elements.GetData(),
+					(const ispc::FVector3f*)InParticles.GetV().GetData(),
+					(const ispc::FVector3f*)InParticles.XArray().GetData(),
+					(const ispc::FVector3f*)Multipliers.GetData(),
+					(const ispc::FVector3f&)Velocity,
+					QuarterRho,
+					DragBase,
+					DragRange,
+					LiftBase,
+					LiftRange,
+					PressureBase,
+					PressureRange,
+					Elements.Num());
+			}
+			else
+			{
+				ispc::UpdateFieldWithWeightMapsAndClampVelocity(
+					(ispc::FVector3f*)Forces.GetData(),
+					(const ispc::FIntVector*)Elements.GetData(),
+					(const ispc::FVector3f*)InParticles.GetV().GetData(),
+					(const ispc::FVector3f*)InParticles.XArray().GetData(),
+					(const ispc::FVector3f*)Multipliers.GetData(),
+					(const ispc::FVector3f&)Velocity,
+					QuarterRho,
+					DragBase,
+					DragRange,
+					LiftBase,
+					LiftRange,
+					PressureBase,
+					PressureRange,
+					Elements.Num(),
+					MaxVelocitySquared);
+			}
 		}
 		else
 #endif
 		{
-			for (int32 ElementIndex = 0; ElementIndex < Elements.Num(); ++ElementIndex)
+			if (MaxVelocitySquared == TNumericLimits<FSolverReal>::Max())
 			{
-				const FSolverVec3& Multiplier = Multipliers[ElementIndex];
-				const FSolverReal Cd = DragBase + DragRange * Multiplier[0];
-				const FSolverReal Cl = LiftBase + LiftRange * Multiplier[1];
-				const FSolverReal Cp = PressureBase + PressureRange * Multiplier[2];
-
-				UpdateField(InParticles, ElementIndex, Velocity, Cd, Cl, Cp);
+				for (int32 ElementIndex = 0; ElementIndex < Elements.Num(); ++ElementIndex)
+				{
+					const FSolverVec3& Multiplier = Multipliers[ElementIndex];
+					const FSolverReal Cd = DragBase + DragRange * Multiplier[0];
+					const FSolverReal Cl = LiftBase + LiftRange * Multiplier[1];
+					const FSolverReal Cp = PressureBase + PressureRange * Multiplier[2];
+
+					UpdateField(InParticles, ElementIndex, Velocity, Cd, Cl, Cp);
+				}
+			}
+			else
+			{
+				for (int32 ElementIndex = 0; ElementIndex < Elements.Num(); ++ElementIndex)
+				{
+					const FSolverVec3& Multiplier = Multipliers[ElementIndex];
+					const FSolverReal Cd = DragBase + DragRange * Multiplier[0];
+					const FSolverReal Cl = LiftBase + LiftRange * Multiplier[1];
+					const FSolverReal Cp = PressureBase + PressureRange * Multiplier[2];
+
+					UpdateField(InParticles, ElementIndex, Velocity, Cd, Cl, Cp, MaxVelocitySquared);
+				}
 			}
 		}
 	}
diff --git a/Engine/Source/Runtime/Experimental/Chaos/Private/Chaos/VelocityField.ispc b/Engine/Source/Runtime/Experimental/Chaos/Private/Chaos/VelocityField.ispc
--- a/Engine/Source/Runtime/Experimental/Chaos/Private/Chaos/VelocityField.ispc
+++ b/Engine/Source/Runtime/Experimental/Chaos/Private/Chaos/VelocityField.ispc
@@ -151,3 +151,158 @@ export void UpdateFieldWithWeightMaps(uniform FVector3f Forces[],
 		Offset += programCount;
 	}
 }
+
+export void UpdateFieldAndClampVelocity(uniform FVector3f Forces[],
+								const uniform FIntVector Elements[],
+								const uniform FVector3f V[],
+								const uniform FVector3f X[],
+								const uniform FVector3f& Velocity,
+								const uniform float QuarterRho,
+								const uniform float Cd,
+								const uniform float Cl,
+								const uniform float Cp,
+								const uniform int32 NumElements,
+								const uniform float MaxVelocitySquared)
+{
+	varying FIntVector Element;
+	uniform int Offset = 0;
+
+#if HW_GATHER_SUPPORTED == 1
+	if(programIndex < NumElements)
+	{
+		Element = VectorLoad(&Elements[Offset]);
+	}
+#endif
+
+	foreach(ElementIndex = 0 ... NumElements)
+	{
+#if HW_GATHER_SUPPORTED == 0
+		Element = VectorLoad(&Elements[Offset]);
+#endif
+
+		const FVector3f X0 = VectorGather(&X[Element.V[0]]);
+		const FVector3f X1 = VectorGather(&X[Element.V[1]]);
+		const FVector3f X2 = VectorGather(&X[Element.V[2]]);
+
+		// Calculate the normal and the area of the surface exposed to the flow
+		FVector3f N = VectorCross(X1 - X0, X2 - X0);
+		const float DoubleArea = SafeNormalize(N);
+
+		const FVector3f V0 = VectorGather(&V[Element.V[0]]);
+		const FVector3f V1 = VectorGather(&V[Element.V[1]]);
+		const FVector3f V2 = VectorGather(&V[Element.V[2]]);
+
+		// Calculate the direction and the relative velocity of the triangle to the flow
+		const FVector3f SurfaceVelocity = OneThird * (V0 + V1 + V2);
+		FVector3f RelVelocity = Velocity - SurfaceVelocity;
+
+		// Clamp the relative velocity
+		const float RelVelocitySquared = VectorSizeSquared(RelVelocity);
+		if (RelVelocitySquared > MaxVelocitySquared)
+		{
+			RelVelocity = RelVelocity * sqrt(MaxVelocitySquared / RelVelocitySquared);
+		}
+
+		// Set the aerodynamic forces
+		const float VDotN = VectorDot(RelVelocity, N);
+		const float VSquare = VectorDot(RelVelocity, RelVelocity);
+
+#if HW_GATHER_SUPPORTED == 1
+		if(ElementIndex + programCount < NumElements)
+		{
+			Element = VectorLoad(&Elements[Offset + programCount]);
+		}
+#endif
+
+		const FVector3f EvenForce = (Cd - Cl) * VDotN * RelVelocity + Cl * VSquare * N;
+		const FVector3f OddForce = (Cl - Cd) * VDotN * RelVelocity - Cl * VSquare * N;
+
+		const FVector3f Force = QuarterRho * DoubleArea * VectorSelect(VDotN >= 0, EvenForce, OddForce) - DoubleArea * 0.5 * Cp * N; // N points in the opposite direction of the actual mesh normals
+
+		VectorStore(&Forces[Offset], Force);
+
+		Offset += programCount;
+	}
+}
+
+export void UpdateFieldWithWeightMapsAndClampVelocity(uniform FVector3f Forces[],
+								const uniform FIntVector Elements[],
+								const uniform FVector3f V[],
+								const uniform FVector3f X[],
+								const uniform FVector3f Multipliers[],
+								const uniform FVector3f& Velocity,
+								const uniform float QuarterRho,
+								const uniform float DragBase,
+								const uniform float DragRange,
+								const uniform float LiftBase,
+								const uniform float LiftRange,
+								const uniform float PressureBase,
+								const uniform float PressureRange,
+								const uniform int32 NumElements,
+								const uniform float MaxVelocitySquared)
+{
+	varying FIntVector Element;
+	uniform int Offset = 0;
+
+#if HW_GATHER_SUPPORTED == 1
+	if(programIndex < NumElements)
+	{
+		Element = VectorLoad(&Elements[Offset]);
+	}
+#endif
+
+	foreach(ElementIndex = 0 ... NumElements)
+	{
+#if HW_GATHER_SUPPORTED == 0
+		Element = VectorLoad(&Elements[Offset]);
+#endif
+		const FVector3f X0 = VectorGather(&X[Element.V[0]]);
+		const FVector3f X1 = VectorGather(&X[Element.V[1]]);
+		const FVector3f X2 = VectorGather(&X[Element.V[2]]);
+
+		const FVector3f Multiplier = VectorGather(&Multipliers[ElementIndex]);
+
+		// Calculate the normal and the area of the surface exposed to the flow
+		FVector3f N = VectorCross(X1 - X0, X2 - X0);
+		const float DoubleArea = SafeNormalize(N);
+
+		const FVector3f V0 = VectorGather(&V[Element.V[0]]);
+		const FVector3f V1 = VectorGather(&V[Element.V[1]]);
+		const FVector3f V2 = VectorGather(&V[Element.V[2]]);
+
+		// Calculate the direction and the relative velocity of the triangle to the flow
+		const FVector3f SurfaceVelocity = OneThird * (V0 + V1 + V2);
+		FVector3f RelVelocity = Velocity - SurfaceVelocity;
+
+		// Clamp the relative velocity
+		const float RelVelocitySquared = VectorSizeSquared(RelVelocity);
+		if (RelVelocitySquared > MaxVelocitySquared)
+		{
+			RelVelocity = RelVelocity * sqrt(MaxVelocitySquared / RelVelocitySquared);
+		}
+
+		// Set the aerodynamic forces
+		const float VDotN = VectorDot(RelVelocity, N);
+		const float VSquare = VectorDot(RelVelocity, RelVelocity);
+
+		const float Cd = DragBase + DragRange * Multiplier.V[0];
+		const float Cl = LiftBase + LiftRange * Multiplier.V[1];
+		const float Cp = PressureBase + PressureRange * Multiplier.V[2];
+
+#if HW_GATHER_SUPPORTED == 1
+		if(ElementIndex + programCount < NumElements)
+		{
+			Element = VectorLoad(&Elements[Offset + programCount]);
+		}
+#endif
+
+		const FVector3f EvenForce = (Cd - Cl) * VDotN * RelVelocity + Cl * VSquare * N;
+		const FVector3f OddForce = (Cl - Cd) * VDotN * RelVelocity - Cl * VSquare * N;
+
+		const FVector3f Force = QuarterRho * DoubleArea * VectorSelect(VDotN >= 0, EvenForce, OddForce) - DoubleArea * 0.5 * Cp * N; // N points in the opposite direction of the actual mesh normals
+
+		VectorStore(&Forces[Offset], Force);
+
+		Offset += programCount;
+	}
+}
diff --git a/Engine/Source/Runtime/Experimental/Chaos/Public/Chaos/VelocityField.h b/Engine/Source/Runtime/Experimental/Chaos/Public/Chaos/VelocityField.h
--- a/Engine/Source/Runtime/Experimental/Chaos/Public/Chaos/VelocityField.h
+++ b/Engine/Source/Runtime/Experimental/Chaos/Public/Chaos/VelocityField.h
@@ -152,6 +152,41 @@ private:
 			(Cl - Cd) * VDotN * V - Cl * VSquare * N) - DoubleArea * (FSolverReal)0.5 * Cp * N; // N points in the opposite direction of the actual mesh normals
 	}
 
+	void UpdateField(const FSolverParticles& InParticles, int32 ElementIndex, const FSolverVec3& InVelocity, const FSolverReal Cd, const FSolverReal Cl, const FSolverReal Cp, const FSolverReal MaxVelocitySquared)
+	{
+		checkSlow(MaxVelocitySquared > (FSolverReal)0);
+
+		const TVec3<int32>& Element = Elements[ElementIndex];
+
+		// Calculate the normal and the area of the surface exposed to the flow
+		FSolverVec3 N = FSolverVec3::CrossProduct(
+			InParticles.X(Element[1]) - InParticles.X(Element[0]),
+			InParticles.X(Element[2]) - InParticles.X(Element[0]));
+		const FSolverReal DoubleArea = N.SafeNormalize();
+
+		// Calculate the direction and the relative velocity of the triangle to the flow
+		const FSolverVec3& SurfaceVelocity = (FSolverReal)(1. / 3.) * (
+			InParticles.V(Element[0]) +
+			InParticles.V(Element[1]) +
+			InParticles.V(Element[2]));
+		FSolverVec3 V = InVelocity - SurfaceVelocity;
+
+		// Clamp the velocity
+		const FSolverReal RelVelocitySquared = V.SquaredLength();
+		if (RelVelocitySquared > MaxVelocitySquared)
+		{
+			V *= FMath::Sqrt(MaxVelocitySquared / RelVelocitySquared);
+		}
+
+		// Set the aerodynamic forces
+		const FSolverReal VDotN = FSolverVec3::DotProduct(V, N);
+		const FSolverReal VSquare = FSolverVec3::DotProduct(V, V);
+
+		Forces[ElementIndex] = QuarterRho * DoubleArea * (VDotN >= (FSolverReal)0. ?  // The flow can hit either side of the triangle, so the normal might need to be reversed
+			(Cd - Cl) * VDotN * V + Cl * VSquare * N :
+			(Cl - Cd) * VDotN * V - Cl * VSquare * N) - DoubleArea * (FSolverReal)0.5 * Cp * N; // N points in the opposite direction of the actual mesh normals
+	}
+
 private:
 	TConstArrayView<TArray<int32>> PointToTriangleMap;
 	TConstArrayView<TVec3<int32>> Elements;
