particle release throttling should not disable all particles by default

diff --git a/Engine/Source/Runtime/Experimental/Chaos/Private/Chaos/PBDRigidClustering.cpp b/Engine/Source/Runtime/Experimental/Chaos/Private/Chaos/PBDRigidClustering.cpp
--- a/Engine/Source/Runtime/Experimental/Chaos/Private/Chaos/PBDRigidClustering.cpp
+++ b/Engine/Source/Runtime/Experimental/Chaos/Private/Chaos/PBDRigidClustering.cpp
@@ -127,27 +127,29 @@ namespace Chaos
 			return Current;
 		}
 
+		bool ShouldThrottleParticleRelease()
+		{
+			return (ClusteringParticleReleaseThrottlingMinCount >= 0 && ClusteringParticleReleaseThrottlingMaxCount >= 0);
+		}
+
 		// compute a ratio (between 0 and 1) of released particles to release
 		float GetRatioOfReleasedParticlesToDisable(const FRigidClustering::FRigidEvolution& Evolution)
 		{
-			float Percentage = 1.0f;
-			if (ClusteringParticleReleaseThrottlingMinCount >= 0 && ClusteringParticleReleaseThrottlingMaxCount >= 0)
-			{
-				const FPBDRigidsSOAs& ParticleStructures = Evolution.GetParticles();
-				int32 NumActiveParticles = 0;
-				NumActiveParticles += ParticleStructures.GetSleepingGeometryCollectionArray().Num();
-				NumActiveParticles += ParticleStructures.GetDynamicGeometryCollectionArray().Num();
+			const FPBDRigidsSOAs& ParticleStructures = Evolution.GetParticles();
+			int32 NumActiveParticles = 0;
+			NumActiveParticles += ParticleStructures.GetSleepingGeometryCollectionArray().Num();
+			NumActiveParticles += ParticleStructures.GetDynamicGeometryCollectionArray().Num();
 
-				const int32 Range = FMath::Max(0, (ClusteringParticleReleaseThrottlingMaxCount - ClusteringParticleReleaseThrottlingMinCount));
-				const int32 OverMinCount = FMath::Max(0, (NumActiveParticles - ClusteringParticleReleaseThrottlingMinCount));
+			const int32 Range = FMath::Max(0, (ClusteringParticleReleaseThrottlingMaxCount - ClusteringParticleReleaseThrottlingMinCount));
+			const int32 OverMinCount = FMath::Max(0, (NumActiveParticles - ClusteringParticleReleaseThrottlingMinCount));
 
-				if (Range > 0)
-				{
-					// clamp to 1, as OverMinCount can get larger than Range
-					Percentage = FMath::Min(1.f, ((float)OverMinCount / (float)Range));
-				}
+			if (Range > 0)
+			{
+				// clamp to 1, as OverMinCount can get larger than Range
+				return FMath::Min(1.f, ((float)OverMinCount / (float)Range));
 			}
-			return Percentage;
+
+			return 1.0f;
 		}
 	}
 	
@@ -1190,26 +1192,29 @@ namespace Chaos
 		}
 
 		// optimization : start disabling activated children if the number of active released particle is too high
-		const float RatioOfParticlesToDisable = GetRatioOfReleasedParticlesToDisable(MEvolution);
-		const int32 NumberOfParticlesToDisable = (int32)((float)ActivatedChildren.Num() * RatioOfParticlesToDisable);
-		if (NumberOfParticlesToDisable > 0)
-		{ 
-			int32 DisabledParticleCount = 0;
-			for (auto ChildIt = ActivatedChildren.CreateIterator(); ChildIt; ++ChildIt)
-			{
-				if (FPBDRigidParticleHandle* Child = *ChildIt)
+		if (ShouldThrottleParticleRelease())
+		{
+			const float RatioOfParticlesToDisable = GetRatioOfReleasedParticlesToDisable(MEvolution);
+			const int32 NumberOfParticlesToDisable = (int32)((float)ActivatedChildren.Num() * RatioOfParticlesToDisable);
+			if (NumberOfParticlesToDisable > 0)
+			{ 
+				int32 DisabledParticleCount = 0;
+				for (auto ChildIt = ActivatedChildren.CreateIterator(); ChildIt; ++ChildIt)
 				{
-					DisabledParticleCount++;
-					MEvolution.DisableParticle(Child);
-					ChildIt.RemoveCurrent();
-				}
-				if (DisabledParticleCount >= NumberOfParticlesToDisable)
-				{
-					break;
+					if (FPBDRigidParticleHandle* Child = *ChildIt)
+					{
+						DisabledParticleCount++;
+						MEvolution.DisableParticle(Child);
+						ChildIt.RemoveCurrent();
+					}
+					if (DisabledParticleCount >= NumberOfParticlesToDisable)
+					{
+						break;
+					}
 				}
 			}
 		}
-
+		
 		FrameReleasedChildren += ActivatedChildren.Num();
 
 		return ActivatedChildren;
