Reduce memory usage in Dynamic Collection
Remove Initial Level array from Dynamic Collection, it is now using the Rest Collection which is the same.

diff --git a/Engine/Source/Runtime/Experimental/Chaos/Private/GeometryCollectionProxyData.cpp b/Engine/Source/Runtime/Experimental/Chaos/Private/GeometryCollectionProxyData.cpp
--- a/Engine/Source/Runtime/Experimental/Chaos/Private/GeometryCollectionProxyData.cpp
+++ b/Engine/Source/Runtime/Experimental/Chaos/Private/GeometryCollectionProxyData.cpp
@@ -108,6 +108,12 @@ FGeometryDynamicCollection::FGeometryDynamicCollection(const FGeometryCollection
 
 }
 
+const TManagedArrayAccessor<int32> FGeometryDynamicCollection::GetInitialLevels() const
+{
+	static const FName LevelAttributeName = "Level";
+	return TManagedArrayAccessor<int32>(*RestCollection, LevelAttributeName, FGeometryCollection::TransformGroup);
+}
+
 void FGeometryDynamicCollection::CopyInitialVelocityAttributesFrom(const FGeometryDynamicCollection& SourceCollection)
 {
 	FInitialVelocityFacade InitialVelocityFacade(*this);
diff --git a/Engine/Source/Runtime/Experimental/Chaos/Private/PhysicsProxy/GeometryCollectionPhysicsProxy.cpp b/Engine/Source/Runtime/Experimental/Chaos/Private/PhysicsProxy/GeometryCollectionPhysicsProxy.cpp
--- a/Engine/Source/Runtime/Experimental/Chaos/Private/PhysicsProxy/GeometryCollectionPhysicsProxy.cpp
+++ b/Engine/Source/Runtime/Experimental/Chaos/Private/PhysicsProxy/GeometryCollectionPhysicsProxy.cpp
@@ -776,9 +776,6 @@ void FGeometryCollectionPhysicsProxy::Initialize(Chaos::FPBDRigidsEvolutionBase
 		}
 	}
 
-	// Add levels to physics thread collection, rename to make it clear these are initial levels from rest collection and not updated during sim.
-	PhysicsThreadCollection.CopyAttribute(*Parameters.RestCollection, /*SrcName=*/"Level", /*DestName=*/"InitialLevel", FTransformCollection::TransformGroup);
-	GameThreadCollection.CopyAttribute(*Parameters.RestCollection, /*SrcName=*/"Level", /*DestName=*/"InitialLevel", FTransformCollection::TransformGroup);
 
 	if (Parameters.EnableClustering)
 	{
@@ -2270,9 +2267,7 @@ void FGeometryCollectionPhysicsProxy::SetAnchoredByTransformedBox_External(const
 				using namespace Chaos;
 				TSet<FPBDRigidClusteredParticleHandle*> TopParentHandles;
 
-				static const FName InitialLevelAttributeName = "InitialLevel";
-
-				const TManagedArrayAccessor<int32> InitialLevelAttribute(PhysicsThreadCollection, InitialLevelAttributeName, FGeometryCollection::TransformGroup);
+				const TManagedArrayAccessor<int32> InitialLevelAttribute = PhysicsThreadCollection.GetInitialLevels();
 				const int32 MaxLevelToCheck = (InitialLevelAttribute.IsValid()) ? MaxLevel : INDEX_NONE;
 
 				FPBDRigidsEvolution* Evolution = RBDSolver->GetEvolution();
diff --git a/Engine/Source/Runtime/Experimental/Chaos/Public/GeometryCollectionProxyData.h b/Engine/Source/Runtime/Experimental/Chaos/Public/GeometryCollectionProxyData.h
--- a/Engine/Source/Runtime/Experimental/Chaos/Public/GeometryCollectionProxyData.h
+++ b/Engine/Source/Runtime/Experimental/Chaos/Public/GeometryCollectionProxyData.h
@@ -56,11 +56,12 @@ public:
 		}
 	}
 
+protected:
+	const FGeometryCollection* RestCollection;
 
 private:
 	TManagedArray<bool>         HasParent;
 	TManagedArray<FTransform3f> Transform;
-	const FGeometryCollection* RestCollection;
 
 	/** Construct */
 	CHAOS_API void Construct();
@@ -120,6 +121,9 @@ public:
 	TManagedArray<int32> CollisionGroup;
 
 public:
+
+	CHAOS_API const TManagedArrayAccessor<int32> GetInitialLevels() const;
+
 	struct FInitialVelocityFacade
 	{
 		FInitialVelocityFacade(FGeometryDynamicCollection& DynamicCollection);
diff --git a/Engine/Source/Runtime/Experimental/GeometryCollectionEngine/Private/GeometryCollection/GeometryCollectionComponent.cpp b/Engine/Source/Runtime/Experimental/GeometryCollectionEngine/Private/GeometryCollection/GeometryCollectionComponent.cpp
--- a/Engine/Source/Runtime/Experimental/GeometryCollectionEngine/Private/GeometryCollection/GeometryCollectionComponent.cpp
+++ b/Engine/Source/Runtime/Experimental/GeometryCollectionEngine/Private/GeometryCollection/GeometryCollectionComponent.cpp
@@ -2015,7 +2015,7 @@ void UGeometryCollectionComponent::UpdateRepData()
 		FPBDRigidsSolver* Solver = PhysicsProxy->GetSolver<Chaos::FPBDRigidsSolver>();
 		const FRigidClustering& RigidClustering = Solver->GetEvolution()->GetRigidClustering();
 
-		const TManagedArray<int32>* InitialLevels = PhysicsProxy->GetPhysicsCollection().FindAttribute<int32>("InitialLevel", FGeometryCollection::TransformGroup);
+		const TManagedArrayAccessor<int32> InitialLevels = PhysicsProxy->GetPhysicsCollection().GetInitialLevels();
 		const TArray<FPBDRigidClusteredParticleHandle*>& ParticleHandles = PhysicsProxy->GetParticles();
 
 		// Replicate the anchored state of the root particle
@@ -2068,14 +2068,14 @@ void UGeometryCollectionComponent::UpdateRepData()
 						ensureMsgf(TransformGroupIdx >= 0, TEXT("Non-internal cluster should always have a group index"));
 						ensureMsgf(TransformGroupIdx < TNumericLimits<uint16>::Max(), TEXT("Trying to replicate GC with more than 65k pieces. We assumed uint16 would suffice"));
 
-						Level = InitialLevels && InitialLevels->Num() > 0 ? (*InitialLevels)[TransformGroupIdx] : INDEX_NONE;
+						Level = InitialLevels.IsValid() && InitialLevels.Num() > 0 ? InitialLevels[TransformGroupIdx] : INDEX_NONE;
 					}
 					else
 					{
 						// Use internal cluster child's index to compute level.
 						const TArray<FPBDRigidParticleHandle*>& Children = RigidClustering.GetChildrenMap()[Root];
 						const int32 ChildTransformGroupIdx = PhysicsProxy->GetTransformGroupIndexFromHandle(Children[0]);
-						Level = InitialLevels && InitialLevels->Num() > 0 ? ((*InitialLevels)[ChildTransformGroupIdx] - 1) : INDEX_NONE;
+						Level = InitialLevels.IsValid() && InitialLevels.Num() > 0 ? (InitialLevels[ChildTransformGroupIdx] - 1) : INDEX_NONE;
 					}
 	
 					if (!bEnableAbandonAfterLevel || Level <= ReplicationAbandonAfterLevel)
@@ -2237,7 +2237,7 @@ void UGeometryCollectionComponent::UpdateRepStateAndDynamicData()
 			const TArray<FPBDRigidClusteredParticleHandle*>& ParticleHandles = PhysicsProxy->GetParticles();
 			const int32 NumTransforms = ParticleHandles.Num();
 
-			const TManagedArray<int32>* InitialLevels = PhysicsProxy->GetPhysicsCollection().FindAttribute<int32>("InitialLevel", FGeometryCollection::TransformGroup);
+			const TManagedArrayAccessor<int32> InitialLevels = PhysicsProxy->GetPhysicsCollection().GetInitialLevels();
 
 			bool bStateChanged = false;
 			bool bDynamicChanged = false;
@@ -2274,7 +2274,7 @@ void UGeometryCollectionComponent::UpdateRepStateAndDynamicData()
 			{
 				if (FPBDRigidClusteredParticleHandle* ParticleHandle = ParticleHandles[TransformIndex])
 				{
-					const int32 Level = (InitialLevels && InitialLevels->IsValidIndex(TransformIndex)) ? (*InitialLevels)[TransformIndex] : INDEX_NONE;
+					const int32 Level = (InitialLevels.IsValid() && InitialLevels.IsValidIndex(TransformIndex)) ? InitialLevels[TransformIndex] : INDEX_NONE;
 
 					const FPBDRigidClusteredParticleHandle* ParentHandle = ParticleHandle->Parent();
 
@@ -5963,7 +5963,7 @@ void UGeometryCollectionComponent::IncrementBreakTimer(float DeltaTime)
 			FGeometryCollectionDecayContext DecayContext(*PhysicsProxy, DecayFacade);
 			const TManagedArray<int32>& OriginalParents = RestCollection->GetGeometryCollection()->Parent;
 
-			const TManagedArray<int32>* InitialLevels = PhysicsProxy->GetPhysicsCollection().FindAttribute<int32>("InitialLevel", FGeometryCollection::TransformGroup);
+			const TManagedArrayAccessor<int32> InitialLevels = PhysicsProxy->GetPhysicsCollection().GetInitialLevels();
 
 			const int32 NumTransforms = OriginalParents.Num();
 			for (int32 TransformIdx = 0; TransformIdx < NumTransforms; ++TransformIdx)
@@ -5978,9 +5978,9 @@ void UGeometryCollectionComponent::IncrementBreakTimer(float DeltaTime)
 					if (OriginalParentIdx > INDEX_NONE && HasDynamicInternalClusterParent && RemoveOnBreakFacade.IsRemovalActive(OriginalParentIdx))
 					{
 						bool bIsAllowedClusterCrumbling = true;
-						if (bIsReplicatedClient && InitialLevels && (InitialLevels->Num() > 0))
+						if (bIsReplicatedClient && InitialLevels.IsValid() && (InitialLevels.Num() > 0))
 						{
-							if (!bEnableAbandonAfterLevel || (*InitialLevels)[OriginalParentIdx] <= ReplicationAbandonAfterLevel)
+							if (!bEnableAbandonAfterLevel || InitialLevels[OriginalParentIdx] <= ReplicationAbandonAfterLevel)
 							{
 								bIsAllowedClusterCrumbling = false;
 							}
@@ -5997,9 +5997,9 @@ void UGeometryCollectionComponent::IncrementBreakTimer(float DeltaTime)
 				else if (RemoveOnBreakFacade.IsRemovalActive(TransformIdx) && DynamicStateFacade.HasBrokenOff(TransformIdx))
 				{
 					bool bIsAllowedClusterCrumbling = true;
-					if (bIsReplicatedClient && InitialLevels && (InitialLevels->Num() > 0))
+					if (bIsReplicatedClient && InitialLevels.IsValid() && (InitialLevels.Num() > 0))
 					{
-						if (!bEnableAbandonAfterLevel || (*InitialLevels)[TransformIdx] <= ReplicationAbandonAfterLevel)
+						if (!bEnableAbandonAfterLevel || InitialLevels[TransformIdx] <= ReplicationAbandonAfterLevel)
 						{
 							bIsAllowedClusterCrumbling = false;
 						}
