Fix Event Manager crash
The read lock was called, then calling the dispatch, the user was callaing the Unregister function, which was calling a write lock, which was blocking.
A deferred unregister function has been put in place to avoid this dead lock issue.

diff --git a/Engine/Source/Runtime/Experimental/Chaos/Public/EventManager.h b/Engine/Source/Runtime/Experimental/Chaos/Public/EventManager.h
--- a/Engine/Source/Runtime/Experimental/Chaos/Public/EventManager.h
+++ b/Engine/Source/Runtime/Experimental/Chaos/Public/EventManager.h
@@ -220,41 +220,51 @@ namespace Chaos
 		 */
 		virtual void UnregisterHandler(const void* InHandler)
 		{
-			HandlerLock.WriteLock();
-			TArray<TPair<UObject*, FEventHandlerPtr>> KeysAndValuesToRemove;
-			for (TPair<UObject*, FEventHandlerPtr>& KeyValue : ProxyOwnerToHandlerMap)
+			// If Unregister Handler is called by the user inside a dispatch event we ll end up with a dead lock. 
+			// so if we cannot lock it we store the event handler to try to unregister it later
+			if (HandlerLock.TryWriteLock())
 			{
-				const FEventHandlerPtr Value = KeyValue.Value;
-				if (Value && Value->GetHandler() == InHandler)
+				TArray<TPair<UObject*, FEventHandlerPtr>> KeysAndValuesToRemove;
+				for (TPair<UObject*, FEventHandlerPtr>& KeyValue : ProxyOwnerToHandlerMap)
 				{
-					KeysAndValuesToRemove.Add(KeyValue);
+					const FEventHandlerPtr Value = KeyValue.Value;
+					if (Value && Value->GetHandler() == InHandler)
+					{
+						KeysAndValuesToRemove.Add(KeyValue);
+					}
 				}
-			}
 
-			for (TPair<UObject*, FEventHandlerPtr>& KeyAndValue : KeysAndValuesToRemove)
-			{
-				ProxyOwnerToHandlerMap.Remove(KeyAndValue.Get<0>(), KeyAndValue.Get<1>());
-			}
+				for (TPair<UObject*, FEventHandlerPtr>& KeyAndValue : KeysAndValuesToRemove)
+				{
+					ProxyOwnerToHandlerMap.Remove(KeyAndValue.Get<0>(), KeyAndValue.Get<1>());
+				}
 
-			for (int i = 0; i < HandlersNotInMap.Num(); i++)
-			{
-				if (HandlersNotInMap[i]->GetHandler() == InHandler)
+				for (int i = 0; i < HandlersNotInMap.Num(); i++)
 				{
-					HandlersNotInMap.RemoveAtSwap(i, 1, false);
-					break;
+					if (HandlersNotInMap[i]->GetHandler() == InHandler)
+					{
+						HandlersNotInMap.RemoveAtSwap(i, 1, false);
+						break;
+					}
 				}
-			}
 
-			for (int i = 0; i < HandlerArray.Num(); i++)
-			{
-				if (HandlerArray[i]->GetHandler() == InHandler)
+				for (int i = 0; i < HandlerArray.Num(); i++)
 				{
-					DeleteHandler(HandlerArray[i]);
-					HandlerArray.RemoveAtSwap(i, 1, false);
-					break;
+					if (HandlerArray[i]->GetHandler() == InHandler)
+					{
+						DeleteHandler(HandlerArray[i]);
+						HandlerArray.RemoveAtSwap(i, 1, false);
+						break;
+					}
 				}
+				HandlerLock.WriteUnlock();
+			}
+			else
+			{
+				DeferredHandlerLock.WriteLock();
+				DeferredUnhandlers.AddUnique(InHandler);
+				DeferredHandlerLock.WriteUnlock();
 			}
-			HandlerLock.WriteUnlock();
 		}
 
 		/*
@@ -338,22 +348,28 @@ namespace Chaos
 				}
 			}
 			HandlerLock.ReadUnlock();
-			RegisterDeferedHandler();
+			UnAndRegisterDeferedHandler();
 		}
 
 private:
 
-		void RegisterDeferedHandler()
+		void UnAndRegisterDeferedHandler()
 		{
 			DeferredHandlerLock.WriteLock();
 			// Move array
 			TArray<FEventHandlerPtr> DeferredHandlersCopy(MoveTemp(DeferredHandlers));
 			check(DeferredHandlers.Num() == 0);
+			TArray<const void*> DeferredUnhandlersCopy(MoveTemp(DeferredUnhandlers));
+			check(DeferredUnhandlers.Num() == 0);
 			DeferredHandlerLock.WriteUnlock();
 			for (const FEventHandlerPtr& HandlerPtr : DeferredHandlersCopy)
 			{
 				RegisterHandler(HandlerPtr);
 			}
+			for (const void* HandlerPtr : DeferredUnhandlersCopy)
+			{
+				UnregisterHandler(HandlerPtr);
+			}
 		}
 
 		void DeleteHandler(FEventHandlerPtr& HandlerPtr)
@@ -382,7 +398,8 @@ private:
 
 		FRWLock HandlerLock; // protect access ProxyOwnerToHandlerMap, HandlersNotInMap, HandlerArray
 		TArray<FEventHandlerPtr> DeferredHandlers; // Store handler to register, they couldn't registered to avoid reentrant lock
-		FRWLock DeferredHandlerLock; // protect access to DeferredHandlers
+		TArray<const void*> DeferredUnhandlers; // Store handler to unregister, they couldn't unregistered to avoid reentrant lock
+		FRWLock DeferredHandlerLock; // protect access to DeferredHandlers and DeferredUnhandlers
 
 	};
 
