Ias: Update bucket name and remove version info from the chunks folder.

#rnx

diff --git a/Engine/Source/Runtime/Experimental/IoStoreOnDemand/Private/IoStoreOnDemand.cpp b/Engine/Source/Runtime/Experimental/IoStoreOnDemand/Private/IoStoreOnDemand.cpp
--- a/Engine/Source/Runtime/Experimental/IoStoreOnDemand/Private/IoStoreOnDemand.cpp
+++ b/Engine/Source/Runtime/Experimental/IoStoreOnDemand/Private/IoStoreOnDemand.cpp
@@ -897,17 +897,14 @@ TIoStatusOr<FIoStoreUploadResult> UploadContainerFiles(
 		}
 	}
 
-	FString ChunksRelativePath = UploadParams.BucketPrefix.IsEmpty()
-		? FString::Printf(TEXT("IoChunksV%u"), EOnDemandChunkVersion::Latest)
-		: FString::Printf(TEXT("%s/IoChunksV%u"), *UploadParams.BucketPrefix, EOnDemandChunkVersion::Latest);
-
+	FString ChunksRelativePath = UploadParams.BucketPrefix.IsEmpty() ? TEXT("Chunks") : FString::Printf(TEXT("%s/Chunks"), *UploadParams.BucketPrefix);
 	ChunksRelativePath.ToLowerInline();
 
 	uint64 TotalUploadedChunks = 0;
 	uint64 TotalUploadedBytes = 0;
 
 	FOnDemandToc OnDemandToc;
-	OnDemandToc.Header.ChunksDirectory = FString::Printf(TEXT("IoChunksV%u"), EOnDemandChunkVersion::Latest).ToLower();
+	OnDemandToc.Header.ChunksDirectory = TEXT("Chunks");
 
 	TArray<FString> FilesToDelete;
 	for (const FString& Path : ContainerFiles)
diff --git a/Engine/Source/Runtime/Experimental/IoStoreOnDemand/Private/OnDemandIoDispatcherBackend.cpp b/Engine/Source/Runtime/Experimental/IoStoreOnDemand/Private/OnDemandIoDispatcherBackend.cpp
--- a/Engine/Source/Runtime/Experimental/IoStoreOnDemand/Private/OnDemandIoDispatcherBackend.cpp
+++ b/Engine/Source/Runtime/Experimental/IoStoreOnDemand/Private/OnDemandIoDispatcherBackend.cpp
@@ -946,7 +946,7 @@ TIoStatusOr<FOnDemandToc> GenerateOnDemandTocFromDisk()
 
 	if (!OutToc.Containers.IsEmpty())
 	{
-		OutToc.Header.ChunksDirectory = FString::Printf(TEXT("IoChunksV%u"), EOnDemandChunkVersion::Latest).ToLower();
+		OutToc.Header.ChunksDirectory = TEXT("chunks");
 	}
 
 	return OutToc;
