Ias: Make sure the service thread exits correctly on shutdown.

#rnx

diff --git a/Engine/Source/Runtime/Experimental/IoStoreOnDemand/Private/JournaledCache.cpp b/Engine/Source/Runtime/Experimental/IoStoreOnDemand/Private/JournaledCache.cpp
--- a/Engine/Source/Runtime/Experimental/IoStoreOnDemand/Private/JournaledCache.cpp
+++ b/Engine/Source/Runtime/Experimental/IoStoreOnDemand/Private/JournaledCache.cpp
@@ -1521,9 +1521,7 @@ FServiceThread& FServiceThread::Get()
 ////////////////////////////////////////////////////////////////////////////////
 FServiceThread::~FServiceThread()
 {
-	uint32 PrevRunCount = RunCount.fetch_sub(1, std::memory_order_relaxed);
-	check(PrevRunCount == 1);
-	WakeEvent->Trigger();
+	Thread.Reset();
 }
 
 ////////////////////////////////////////////////////////////////////////////////
@@ -1677,6 +1675,7 @@ uint32 FServiceThread::Run()
 void FServiceThread::Stop()
 {
 	RunCount.fetch_sub(1, std::memory_order_relaxed); // THREAD_ALIVE
+	WakeEvent->Trigger();
 }
 
 ////////////////////////////////////////////////////////////////////////////////
