Clean up, rename variables

diff --git a/Engine/Source/Runtime/Experimental/Chaos/Public/EventManager.h b/Engine/Source/Runtime/Experimental/Chaos/Public/EventManager.h
--- a/Engine/Source/Runtime/Experimental/Chaos/Public/EventManager.h
+++ b/Engine/Source/Runtime/Experimental/Chaos/Public/EventManager.h
@@ -262,7 +262,7 @@ namespace Chaos
 			else
 			{
 				DeferredHandlerLock.WriteLock();
-				DeferredUnhandlers.AddUnique(InHandler);
+				DeferredUnregisterHandlers.AddUnique(InHandler);
 				DeferredHandlerLock.WriteUnlock();
 			}
 		}
@@ -359,14 +359,14 @@ private:
 			// Move array
 			TArray<FEventHandlerPtr> DeferredHandlersCopy(MoveTemp(DeferredHandlers));
 			check(DeferredHandlers.Num() == 0);
-			TArray<const void*> DeferredUnhandlersCopy(MoveTemp(DeferredUnhandlers));
-			check(DeferredUnhandlers.Num() == 0);
+			TArray<const void*> DeferredUnregisterHandlersCopy(MoveTemp(DeferredUnregisterHandlers));
+			check(DeferredUnregisterHandlers.Num() == 0);
 			DeferredHandlerLock.WriteUnlock();
 			for (const FEventHandlerPtr& HandlerPtr : DeferredHandlersCopy)
 			{
 				RegisterHandler(HandlerPtr);
 			}
-			for (const void* HandlerPtr : DeferredUnhandlersCopy)
+			for (const void* HandlerPtr : DeferredUnregisterHandlersCopy)
 			{
 				UnregisterHandler(HandlerPtr);
 			}
@@ -398,8 +398,8 @@ private:
 
 		FRWLock HandlerLock; // protect access ProxyOwnerToHandlerMap, HandlersNotInMap, HandlerArray
 		TArray<FEventHandlerPtr> DeferredHandlers; // Store handler to register, they couldn't registered to avoid reentrant lock
-		TArray<const void*> DeferredUnhandlers; // Store handler to unregister, they couldn't unregistered to avoid reentrant lock
-		FRWLock DeferredHandlerLock; // protect access to DeferredHandlers and DeferredUnhandlers
+		TArray<const void*> DeferredUnregisterHandlers; // Store handler to unregister, they couldn't unregistered to avoid reentrant lock
+		FRWLock DeferredHandlerLock; // protect access to DeferredHandlers and DeferredUnregisterHandlers
 
 	};
 
