IAS - check for valid backend data before calling cancel and update priority

diff --git a/Engine/Source/Runtime/Experimental/IoStoreOnDemand/Private/OnDemandIoDispatcherBackend.cpp b/Engine/Source/Runtime/Experimental/IoStoreOnDemand/Private/OnDemandIoDispatcherBackend.cpp
--- a/Engine/Source/Runtime/Experimental/IoStoreOnDemand/Private/OnDemandIoDispatcherBackend.cpp
+++ b/Engine/Source/Runtime/Experimental/IoStoreOnDemand/Private/OnDemandIoDispatcherBackend.cpp
@@ -1193,10 +1193,9 @@ class FOnDemandIoBackend final
 			return TUniquePtr<FBackendData>(static_cast<FBackendData*>(BackendData));
 		}
 		
-		static FBackendData& Get(FIoRequestImpl* Request)
+		static FBackendData* Get(FIoRequestImpl* Request)
 		{
-			check(Request->BackendData != nullptr);
-			return *static_cast<FBackendData*>(Request->BackendData);
+			return static_cast<FBackendData*>(Request->BackendData);
 		}
 
 		FIoHash ChunkKey;
@@ -1208,8 +1207,13 @@ class FOnDemandIoBackend final
 		{
 			FScopeLock _(&Mutex);
 
-			const FBackendData& BackendData = FBackendData::Get(Request);
-			if (FChunkRequest** InflightRequest = Inflight.Find(BackendData.ChunkKey))
+			const FBackendData* BackendData = FBackendData::Get(Request);
+			if (BackendData == nullptr)
+			{
+				return nullptr;
+			}
+
+			if (FChunkRequest** InflightRequest = Inflight.Find(BackendData->ChunkKey))
 			{
 				FChunkRequest* ChunkRequest = *InflightRequest;
 				if (Request->Priority > ChunkRequest->Priority)
@@ -1249,11 +1253,16 @@ class FOnDemandIoBackend final
 		{
 			FScopeLock _(&Mutex);
 
-			FBackendData& BackendData = FBackendData::Get(Request);
+			const FBackendData* BackendData = FBackendData::Get(Request);
+			if (BackendData == nullptr)
+			{
+				return false;
+			}
+
 			UE_LOG(LogIas, VeryVerbose, TEXT("%s"),
-				*WriteToString<256>(TEXT("Cancelling I/O request ChunkId='"), LexToString(Request->ChunkId), TEXT("' ChunkKey='"), BackendData.ChunkKey, TEXT("'")));
+				*WriteToString<256>(TEXT("Cancelling I/O request ChunkId='"), Request->ChunkId, TEXT("' ChunkKey='"), BackendData->ChunkKey, TEXT("'")));
 
-			if (FChunkRequest** InflightRequest = Inflight.Find(BackendData.ChunkKey))
+			if (FChunkRequest** InflightRequest = Inflight.Find(BackendData->ChunkKey))
 			{
 				FChunkRequest& ChunkRequest = **InflightRequest;
 				const uint32 RemainingCount = ChunkRequest.RemoveDispatcherRequest(Request);
@@ -1266,7 +1275,7 @@ class FOnDemandIoBackend final
 					{
 						TheCache->Cancel(ChunkRequest.Chunk);
 					}
-					Inflight.Remove(BackendData.ChunkKey);
+					Inflight.Remove(BackendData->ChunkKey);
 				}
 
 				return true;
