Fixed usage of TRACE_CPUPROFILER_EVENT_SCOPE macros in Chaos.

* TRACE_CPUPROFILER_EVENT_SCOPE(ScopeName) --> to be used with a plain text as scope name (not a string!). If the parameter is a string ("abc" or TEXT("abc")), the quotes (and the TEXT prefix) will be included in the actual timer name.
* TRACE_CPUPROFILER_EVENT_SCOPE_STR("scope name") --> to be used with a static string (const ANSICHAR* or const TCHAR*); ex: when the scope name needs to include a space char
* TRACE_CPUPROFILER_EVENT_SCOPE_TEXT(*MyScopeName) --> to be used with a dynamic string (const ANSICHAR* or const TCHAR*)
See also the comments for these macros in ProfilingDebugging\CpuProfilerTrace.h.

diff --git a/Engine/Source/Runtime/Experimental/Chaos/Public/Chaos/BlendedXPBDCorotatedConstraints.h b/Engine/Source/Runtime/Experimental/Chaos/Public/Chaos/BlendedXPBDCorotatedConstraints.h
--- a/Engine/Source/Runtime/Experimental/Chaos/Public/Chaos/BlendedXPBDCorotatedConstraints.h
+++ b/Engine/Source/Runtime/Experimental/Chaos/Public/Chaos/BlendedXPBDCorotatedConstraints.h
@@ -103,7 +103,7 @@ namespace Chaos::Softs
 
 		virtual TVec4<TVector<T, 3>> GetDeterminantDelta(const ParticleType& Particles, const T Dt, const int32 ElementIndex, const T Tol = (T)1e-3) const override
 		{
-			TRACE_CPUPROFILER_EVENT_SCOPE(TEXT("STAT_ChaosBlendedXPBDCorotatedApplyDet"));
+			TRACE_CPUPROFILER_EVENT_SCOPE(STAT_ChaosBlendedXPBDCorotatedApplyDet);
 
 			const PMatrix<T, 3, 3> Fe = Base::F(ElementIndex, Particles);
 			PMatrix<T, 3, 3> DmInvT = Base::ElementDmInv(ElementIndex).GetTransposed();
@@ -165,7 +165,7 @@ namespace Chaos::Softs
 
 		virtual TVec4<TVector<T, 3>> GetPolarDelta(const ParticleType& Particles, const T Dt, const int32 ElementIndex, const T Tol = (T)1e-3) const override
 		{
-			TRACE_CPUPROFILER_EVENT_SCOPE(TEXT("STAT_ChaosBlendedXPBDCorotatedApplyPolar"));
+			TRACE_CPUPROFILER_EVENT_SCOPE(STAT_ChaosBlendedXPBDCorotatedApplyPolar);
 			SCOPE_CYCLE_COUNTER(STAT_ChaosXPBDCorotatedPolar);
 			const PMatrix<T, 3, 3> Fe = Base::F(ElementIndex, Particles);
 
diff --git a/Engine/Source/Runtime/Experimental/Chaos/Public/Chaos/MPMTransfer.h b/Engine/Source/Runtime/Experimental/Chaos/Public/Chaos/MPMTransfer.h
--- a/Engine/Source/Runtime/Experimental/Chaos/Public/Chaos/MPMTransfer.h
+++ b/Engine/Source/Runtime/Experimental/Chaos/Public/Chaos/MPMTransfer.h
@@ -63,7 +63,7 @@ public:
 		int32 N = InParticles.Size();
 
 		{
-			TRACE_CPUPROFILER_EVENT_SCOPE(TEXT("STAT_ChaosMPMTransferInitialBinning"));
+			TRACE_CPUPROFILER_EVENT_SCOPE(STAT_ChaosMPMTransferInitialBinning);
 			Indices.SetNum(N);
 			Weights.SetNum(N);
 
@@ -83,7 +83,7 @@ public:
 		// Computes which particles in the same cell
 		/////////////////////
 		{
-			TRACE_CPUPROFILER_EVENT_SCOPE(TEXT("STAT_ChaosMPMTransferCellMetaCalc"));
+			TRACE_CPUPROFILER_EVENT_SCOPE(STAT_ChaosMPMTransferCellMetaCalc);
 			NumCells = Grid.Size();
 			CellData.SetNum(NumCells);
 			for (int32 c = 0; c < NumCells; c++)
@@ -100,7 +100,7 @@ public:
 		// splat data to cells
 		/////////////////////
 		{
-			TRACE_CPUPROFILER_EVENT_SCOPE(TEXT("STAT_ChaosMPMTransferSplatData"));
+			TRACE_CPUPROFILER_EVENT_SCOPE(STAT_ChaosMPMTransferSplatData);
 			GridData.Init((T)0., CellData.Num() * (NTransfer + 1));
 
 			TVector<int32, 3> GridCells = Grid.GetCells();
diff --git a/Engine/Source/Runtime/Experimental/Chaos/Public/Chaos/PBDAltitudeSpringConstraints.h b/Engine/Source/Runtime/Experimental/Chaos/Public/Chaos/PBDAltitudeSpringConstraints.h
--- a/Engine/Source/Runtime/Experimental/Chaos/Public/Chaos/PBDAltitudeSpringConstraints.h
+++ b/Engine/Source/Runtime/Experimental/Chaos/Public/Chaos/PBDAltitudeSpringConstraints.h
@@ -113,7 +113,7 @@ namespace Chaos::Softs
 
 		virtual void ApplySingleElement(FSolverParticles& Particles, const FSolverReal Dt, const int32 ElementIndex) const
 		{
-			TRACE_CPUPROFILER_EVENT_SCOPE(TEXT("STAT_ChaosPBDAltitudeSpringApplySingle"));
+			TRACE_CPUPROFILER_EVENT_SCOPE(STAT_ChaosPBDAltitudeSpringApplySingle);
 			const TVec4<int32>& Constraint = Constraints[ElementIndex];
 			for (int ie = 0; ie < 4; ++ie)
 			{
@@ -133,7 +133,7 @@ namespace Chaos::Softs
 
 		virtual void ApplySingleElementAllInOne(FSolverParticles& Particles, const FSolverReal Dt, const int32 ElementIndex) const
 		{
-			TRACE_CPUPROFILER_EVENT_SCOPE(TEXT("STAT_ChaosPBDAltitudeSpringApplySingle"));
+			TRACE_CPUPROFILER_EVENT_SCOPE(STAT_ChaosPBDAltitudeSpringApplySingle);
 			const TVec4<int32>& Constraint = Constraints[ElementIndex];
 			for (int ie = 0; ie < 4; ++ie)
 			{
@@ -223,7 +223,7 @@ namespace Chaos::Softs
 
 		virtual void ApplySingleElementShortest(FSolverParticles& Particles, const FSolverReal Dt, const int32 ElementIndex) const
 		{
-			TRACE_CPUPROFILER_EVENT_SCOPE(TEXT("STAT_ChaosPBDAltitudeSpringApplySingleShortest"));
+			TRACE_CPUPROFILER_EVENT_SCOPE(STAT_ChaosPBDAltitudeSpringApplySingleShortest);
 			const TVec4<int32>& Constraint = Constraints[ElementIndex];
 			int32 ie = GetSingleElementShortestLength(Particles, Dt, ElementIndex);
 			
@@ -294,7 +294,7 @@ namespace Chaos::Softs
 		void ApplyInSerial(FSolverParticles& Particles, const FSolverReal Dt) const
 		{
 			//SCOPE_CYCLE_COUNTER(STAT_ChaosXPBDCorotated);
-			TRACE_CPUPROFILER_EVENT_SCOPE(TEXT("STAT_ChaosPBDAltitudeSpringApplySerial"));
+			TRACE_CPUPROFILER_EVENT_SCOPE(STAT_ChaosPBDAltitudeSpringApplySerial);
 			for (int32 ElementIndex = 0; ElementIndex < Constraints.Num(); ++ElementIndex)
 			{
 				ApplySingleElement(Particles, Dt, ElementIndex);
@@ -306,7 +306,7 @@ namespace Chaos::Softs
 		void ApplyInParallel(FSolverParticles& Particles, const FSolverReal Dt) const
 		{
 			{
-				TRACE_CPUPROFILER_EVENT_SCOPE(TEXT("STAT_ChaosPBDAltitudeSpringApply"));
+				TRACE_CPUPROFILER_EVENT_SCOPE(STAT_ChaosPBDAltitudeSpringApply);
 				if ((ConstraintsPerColorStartIndex.Num() > 1))//&& (MeshConstraints.Num() > Chaos_Spring_ParallelConstraintCount))
 				{
 					const int32 ConstraintColorNum = ConstraintsPerColorStartIndex.Num() - 1;
diff --git a/Engine/Source/Runtime/Experimental/Chaos/Public/Chaos/XPBDCorotatedConstraints.h b/Engine/Source/Runtime/Experimental/Chaos/Public/Chaos/XPBDCorotatedConstraints.h
--- a/Engine/Source/Runtime/Experimental/Chaos/Public/Chaos/XPBDCorotatedConstraints.h
+++ b/Engine/Source/Runtime/Experimental/Chaos/Public/Chaos/XPBDCorotatedConstraints.h
@@ -246,7 +246,7 @@ namespace Chaos::Softs
 
 		virtual void ApplyInSerial(ParticleType& Particles, const T Dt, const int32 ElementIndex) const
 		{
-			TRACE_CPUPROFILER_EVENT_SCOPE(TEXT("STAT_ChaosXPBDCorotatedApplySingle"));
+			TRACE_CPUPROFILER_EVENT_SCOPE(STAT_ChaosXPBDCorotatedApplySingle);
 
 			TVec4<TVector<T, 3>> PolarDelta = GetPolarDelta(Particles, Dt, ElementIndex);
 
@@ -267,7 +267,7 @@ namespace Chaos::Softs
 		void ApplyInSerial(ParticleType& Particles, const T Dt) const
 		{
 			SCOPE_CYCLE_COUNTER(STAT_ChaosXPBDCorotated);
-			TRACE_CPUPROFILER_EVENT_SCOPE(TEXT("STAT_ChaosXPBDCorotatedApplySerial"));
+			TRACE_CPUPROFILER_EVENT_SCOPE(STAT_ChaosXPBDCorotatedApplySerial);
 			const int32 ConstraintColorNum = ConstraintsPerColorStartIndex.Num() - 1;
 			for (int32 ConstraintColorIndex = 0; ConstraintColorIndex < ConstraintColorNum; ++ConstraintColorIndex)
 			{
@@ -294,7 +294,7 @@ namespace Chaos::Softs
 			
 			{
 				SCOPE_CYCLE_COUNTER(STAT_ChaosXPBDCorotated);
-				TRACE_CPUPROFILER_EVENT_SCOPE(TEXT("STAT_ChaosXPBDCorotatedApply"));
+				TRACE_CPUPROFILER_EVENT_SCOPE(STAT_ChaosXPBDCorotatedApply);
 				if ((ConstraintsPerColorStartIndex.Num() > 1))//&& (MeshConstraints.Num() > Chaos_Spring_ParallelConstraintCount))
 				{
 					const int32 ConstraintColorNum = ConstraintsPerColorStartIndex.Num() - 1;
@@ -360,7 +360,7 @@ namespace Chaos::Softs
 
 		TVec4<TVector<T, 3>> GetDeterminantGradient(const PMatrix<T, 3, 3>& Fe, const PMatrix<T, 3, 3>& DmInvT) const
 		{
-			TRACE_CPUPROFILER_EVENT_SCOPE(TEXT("STAT_ChaosXPBDCorotatedGetDetGradient"));
+			TRACE_CPUPROFILER_EVENT_SCOPE(STAT_ChaosXPBDCorotatedGetDetGradient);
 			//SCOPE_CYCLE_COUNTER(STAT_ChaosXPBDCorotatedDet);
 			//const TVec4<int32>& Constraint = MeshConstraints[ElementIndex];
 
@@ -444,7 +444,7 @@ namespace Chaos::Softs
 
 		virtual TVec4<TVector<T, 3>> GetDeterminantDelta(const ParticleType& Particles, const T Dt, const int32 ElementIndex, const T Tol = (T)1e-3) const
 		{
-			TRACE_CPUPROFILER_EVENT_SCOPE(TEXT("STAT_ChaosXPBDCorotatedApplyDet"));
+			TRACE_CPUPROFILER_EVENT_SCOPE(STAT_ChaosXPBDCorotatedApplyDet);
 			//SCOPE_CYCLE_COUNTER(STAT_ChaosXPBDCorotatedDet);
 			////const TVec4<int32>& Constraint = MeshConstraints[ElementIndex];
 
@@ -508,7 +508,7 @@ namespace Chaos::Softs
 
 		virtual TVec4<TVector<T, 3>> GetPolarDelta(const ParticleType& Particles, const T Dt, const int32 ElementIndex, const T Tol = (T)1e-3) const
 		{	
-			TRACE_CPUPROFILER_EVENT_SCOPE(TEXT("STAT_ChaosXPBDCorotatedApplyPolar"));
+			TRACE_CPUPROFILER_EVENT_SCOPE(STAT_ChaosXPBDCorotatedApplyPolar);
 			SCOPE_CYCLE_COUNTER(STAT_ChaosXPBDCorotatedPolar);
 			const PMatrix<T, 3, 3> Fe = F(ElementIndex, Particles);
 
diff --git a/Engine/Source/Runtime/Experimental/Chaos/Public/Chaos/XPBDGridBasedCorotatedConstraints.h b/Engine/Source/Runtime/Experimental/Chaos/Public/Chaos/XPBDGridBasedCorotatedConstraints.h
--- a/Engine/Source/Runtime/Experimental/Chaos/Public/Chaos/XPBDGridBasedCorotatedConstraints.h
+++ b/Engine/Source/Runtime/Experimental/Chaos/Public/Chaos/XPBDGridBasedCorotatedConstraints.h
@@ -4,6 +4,7 @@
 // HEADER_UNIT_SKIP - Not included directly
 
 #include "Chaos/PBDSpringConstraintsBase.h"
+#include "Chaos/XPBDCorotatedConstraints.h"
 #include "ChaosStats.h"
 #include "Chaos/ImplicitQRSVD.h"
 #include "Chaos/GraphColoring.h"
@@ -122,23 +123,23 @@ namespace Chaos::Softs
 		{
 			for (T& Lambdas : LambdaArray) { Lambdas = (T)0.; }
 			{
-				TRACE_CPUPROFILER_EVENT_SCOPE(TEXT("STAT_ChaosXPBDGridBasedCorotatedUpdateGrid"));
+				TRACE_CPUPROFILER_EVENT_SCOPE(STAT_ChaosXPBDGridBasedCorotatedUpdateGrid);
 				MPMTransfer.Grid.UpdateGridFromPositions(InParticles);
 			}
 			{
-				TRACE_CPUPROFILER_EVENT_SCOPE(TEXT("STAT_ChaosXPBDGridBasedCorotatedInitialP2G"));
+				TRACE_CPUPROFILER_EVENT_SCOPE(STAT_ChaosXPBDGridBasedCorotatedInitialP2G);
 				MPMTransfer.InitialP2G(InParticles, GridData);
 			}
 			{
-				TRACE_CPUPROFILER_EVENT_SCOPE(TEXT("STAT_ChaosXPBDGridBasedCorotatedMetaDataComputation"));
+				TRACE_CPUPROFILER_EVENT_SCOPE(STAT_ChaosXPBDGridBasedCorotatedMetaDataComputation);
 				MPMTransfer.ComputeElementMetaData(MeshConstraints);
 			}
 			{
-				TRACE_CPUPROFILER_EVENT_SCOPE(TEXT("STAT_ChaosXPBDGridBasedCorotatedSubcoloring"));
+				TRACE_CPUPROFILER_EVENT_SCOPE(STAT_ChaosXPBDGridBasedCorotatedSubcoloring);
 				ComputeGridBasedGraphSubColoringPointer(ElementsPerColor, MPMTransfer.Grid, MPMTransfer.Grid.Size(), PreviousColoring, MPMTransfer.ElementGridNodesSet, ElementsPerSubColors);
 			}
 			{
-				TRACE_CPUPROFILER_EVENT_SCOPE(TEXT("STAT_ChaosXPBDGridBasedCorotatedComputeGridPositions"));
+				TRACE_CPUPROFILER_EVENT_SCOPE(STAT_ChaosXPBDGridBasedCorotatedComputeGridPositions);
 				MPMTransfer.ComputeGridPositions(GridData, Dt, GridPositions);
 			}
 		}
@@ -146,17 +147,17 @@ namespace Chaos::Softs
 		void TimeStepPostprocessing(ParticleType& InParticles, const T Dt)
 		{
 			{
-				TRACE_CPUPROFILER_EVENT_SCOPE(TEXT("STAT_ChaosXPBDGridBasedCorotatedGrid::FinalG2P"));
+				TRACE_CPUPROFILER_EVENT_SCOPE(STAT_ChaosXPBDGridBasedCorotatedGrid::FinalG2P);
 				FinalG2P(InParticles);
 			}
 			
 			{
-				TRACE_CPUPROFILER_EVENT_SCOPE(TEXT("STAT_ChaosXPBDGridBasedCorotated::GridPositions2GridData"));
+				TRACE_CPUPROFILER_EVENT_SCOPE(STAT_ChaosXPBDGridBasedCorotated::GridPositions2GridData);
 				MPMTransfer.GridPositionsToGridData(GridPositions, Dt, GridData);
 			}
 
 			{
-				TRACE_CPUPROFILER_EVENT_SCOPE(TEXT("STAT_ChaosXPBDGridBasedCorotated::ComputeAArray"));
+				TRACE_CPUPROFILER_EVENT_SCOPE(STAT_ChaosXPBDGridBasedCorotated::ComputeAArray);
 				MPMTransfer.ComputeAArray(GridData, InParticles);
 			}
 		
diff --git a/Engine/Source/Runtime/Experimental/Chaos/Public/Chaos/XPBDVolumeConstraints.h b/Engine/Source/Runtime/Experimental/Chaos/Public/Chaos/XPBDVolumeConstraints.h
--- a/Engine/Source/Runtime/Experimental/Chaos/Public/Chaos/XPBDVolumeConstraints.h
+++ b/Engine/Source/Runtime/Experimental/Chaos/Public/Chaos/XPBDVolumeConstraints.h
@@ -53,7 +53,7 @@ namespace Chaos::Softs
 
 		virtual void ApplyOneElement(FSolverParticles& Particles, const FSolverReal Dt, const int32 ElementIndex) const
 		{
-			//TRACE_CPUPROFILER_EVENT_SCOPE(TEXT("STAT_ChaosXPBDCorotatedApplySingle"));
+			//TRACE_CPUPROFILER_EVENT_SCOPE(STAT_ChaosXPBDCorotatedApplySingle);
 			
 			TVec4<FSolverVec3> VolumeDelta = GetVolumeDelta(Particles, Dt, ElementIndex);
 			for (int i = 0; i < 4; i++)
@@ -66,8 +66,8 @@ namespace Chaos::Softs
 
 		void ApplyInSerial(FSolverParticles& Particles, const FSolverReal Dt) const
 		{
-			SCOPE_CYCLE_COUNTER(STAT_ChaosXPBDCorotated);
-			TRACE_CPUPROFILER_EVENT_SCOPE(TEXT("STAT_ChaosXPBDVolumeApplySerial"));
+			//SCOPE_CYCLE_COUNTER(STAT_ChaosXPBDCorotated);
+			TRACE_CPUPROFILER_EVENT_SCOPE(STAT_ChaosXPBDVolumeApplySerial);
 			for (int32 ElementIndex = 0; ElementIndex < MeshConstraints.Num(); ++ElementIndex)
 			{
 				ApplyOneElement(Particles, Dt, ElementIndex);
@@ -79,7 +79,7 @@ namespace Chaos::Softs
 		void ApplyInParallel(FSolverParticles& Particles, const FSolverReal Dt) const
 		{
 			{
-				TRACE_CPUPROFILER_EVENT_SCOPE(TEXT("STAT_ChaosXPBDVolumeApply"));
+				TRACE_CPUPROFILER_EVENT_SCOPE(STAT_ChaosXPBDVolumeApply);
 				if ((ConstraintsPerColorStartIndex.Num() > 1))//&& (MeshConstraints.Num() > Chaos_Spring_ParallelConstraintCount))
 				{
 					const int32 ConstraintColorNum = ConstraintsPerColorStartIndex.Num() - 1;
diff --git a/Engine/Source/Runtime/Experimental/Chaos/Public/Chaos/XPBDWeakConstraints.h b/Engine/Source/Runtime/Experimental/Chaos/Public/Chaos/XPBDWeakConstraints.h
--- a/Engine/Source/Runtime/Experimental/Chaos/Public/Chaos/XPBDWeakConstraints.h
+++ b/Engine/Source/Runtime/Experimental/Chaos/Public/Chaos/XPBDWeakConstraints.h
@@ -80,7 +80,7 @@ namespace Chaos::Softs
 
 		void ApplyInParallel(ParticleType& Particles, const T Dt) const
 		{
-			TRACE_CPUPROFILER_EVENT_SCOPE(TEXT("STAT_ChaosXPBDWeakConstraintApply"));
+			TRACE_CPUPROFILER_EVENT_SCOPE(STAT_ChaosXPBDWeakConstraintApply);
 
 			if ((ConstraintsPerColor.Num() > 0))
 			{
