Timers, debug rendering, user specified raycast filtering, interim vehicle socket solution

diff --git a/Engine/Source/Runtime/Experimental/ChaosVehicles/ChaosVehiclesCore/Private/SimModule/SimModuleTree.cpp b/Engine/Source/Runtime/Experimental/ChaosVehicles/ChaosVehiclesCore/Private/SimModule/SimModuleTree.cpp
--- a/Engine/Source/Runtime/Experimental/ChaosVehicles/ChaosVehiclesCore/Private/SimModule/SimModuleTree.cpp
+++ b/Engine/Source/Runtime/Experimental/ChaosVehicles/ChaosVehiclesCore/Private/SimModule/SimModuleTree.cpp
@@ -4,13 +4,24 @@
 #include "Chaos/ParticleHandleFwd.h"
 #include "PhysicsProxy/GeometryCollectionPhysicsProxy.h"
 #include "PhysicsProxy/ClusterUnionPhysicsProxy.h"
+#include "Chaos/PhysicsObjectInternalInterface.h"
 #include "Chaos/DebugDrawQueue.h"
 #include "VehicleUtility.h"
 
 #if VEHICLE_DEBUGGING_ENABLED
-UE_DISABLE_OPTIMIZATION
+UE_DISABLE_OPTIMIZATION_SHIP
 #endif
 
+DECLARE_CYCLE_STAT(TEXT("ModularVehicle_SimulateTree"), STAT_ModularVehicle_SimulateTree, STATGROUP_ModularVehicleSimTree);
+DECLARE_CYCLE_STAT(TEXT("ModularVehicle_GenerateReplicationStructure"), STAT_ModularVehicle_GenerateReplicationStructure, STATGROUP_ModularVehicleSimTree);
+DECLARE_CYCLE_STAT(TEXT("ModularVehicle_SetNetState"), STAT_ModularVehicle_SetNetState, STATGROUP_ModularVehicleSimTree);
+DECLARE_CYCLE_STAT(TEXT("ModularVehicle_SetSimState"), STAT_ModularVehicle_SetSimState, STATGROUP_ModularVehicleSimTree);
+DECLARE_CYCLE_STAT(TEXT("ModularVehicle_AppendTreeUpdates"), STAT_ModularVehicle_AppendTreeUpdates, STATGROUP_ModularVehicleSimTree);
+
+bool bModularVehicle_NetworkData_Enable = false;
+FAutoConsoleVariableRef CVarModularVehicleNetworkDataEnable(TEXT("p.ModularVehicle.NetworkData.Enable"), bModularVehicle_NetworkData_Enable, TEXT("Enable/Disable additional module network data."));
+
+
 namespace Chaos
 {
 
@@ -63,6 +74,8 @@ int FSimModuleTree::AddNodeBelow(int AtIndex, ISimulationModuleBase* SimModule)
 
 void FSimModuleTree::AppendTreeUpdates(const FSimTreeUpdates& TreeUpdates)
 {
+	SCOPE_CYCLE_COUNTER(STAT_ModularVehicle_AppendTreeUpdates);
+
 	int TreeIndex = -1;
 	TMap<int, int> SimTreeMapping;
 
@@ -90,17 +103,21 @@ void FSimModuleTree::AppendTreeUpdates(const FSimTreeUpdates& TreeUpdates)
 	}
 
 	TArray<int> ComponentIndices;
-	for (const FPendingModuleDeletions& TreeUpdate : TreeUpdates.GetDeletedModules())
+
+	if (!SimulationModuleTree.IsEmpty())
 	{
-		for (int Index = 0; Index < SimulationModuleTree.Num(); Index++)
+		for (const FPendingModuleDeletions& TreeUpdate : TreeUpdates.GetDeletedModules())
 		{
-			if (Chaos::ISimulationModuleBase* SimModule = GetNode(Index).SimModule)
-			{ 
-				if (SimModule->GetGuid() == TreeUpdate.Guid)
+			for (int Index = 0; Index < SimulationModuleTree.Num(); Index++)
+			{
+				if (Chaos::ISimulationModuleBase* SimModule = GetNode(Index).SimModule)
 				{
-					ComponentIndices.AddUnique(SimModule->GetTransformIndex());
-					DeleteNode(Index);
-					break;
+					if (SimModule->GetGuid() == TreeUpdate.Guid)
+					{
+						ComponentIndices.AddUnique(SimModule->GetTransformIndex());
+						DeleteNode(Index);
+						break;
+					}
 				}
 			}
 		}				
@@ -233,9 +250,11 @@ void FSimModuleTree::DeleteNode(int AtIndex)
 
 void FSimModuleTree::Simulate(float DeltaTime, FAllInputs& Inputs, FClusterUnionPhysicsProxy* PhysicsProxy)
 {
+	SCOPE_CYCLE_COUNTER(STAT_ModularVehicle_SimulateTree);
+
 	if (PhysicsProxy)
 	{
-		UpdateModuleVelocites(PhysicsProxy);
+		UpdateModuleVelocites(PhysicsProxy, Inputs.ControlInputs.InputNonZero());
 	}
 
 	TArray<int> RootNodes;
@@ -250,6 +269,8 @@ void FSimModuleTree::Simulate(float DeltaTime, FAllInputs& Inputs, FClusterUnion
 
 void FSimModuleTree::Simulate(float DeltaTime, FAllInputs& Inputs, FGeometryCollectionPhysicsProxy* PhysicsProxy)
 {
+	SCOPE_CYCLE_COUNTER(STAT_ModularVehicle_SimulateTree);
+
 	if (PhysicsProxy)
 	{
 		UpdateModuleVelocites(PhysicsProxy);
@@ -372,16 +393,25 @@ void FSimModuleTree::UpdateModuleVelocites(FGeometryCollectionPhysicsProxy* Phys
 
 }
 
-void FSimModuleTree::UpdateModuleVelocites(FClusterUnionPhysicsProxy* PhysicsProxy)
+void FSimModuleTree::UpdateModuleVelocites(FClusterUnionPhysicsProxy* PhysicsProxy, bool bWake)
 {
 	check(PhysicsProxy);
+	Chaos::EnsureIsInPhysicsThreadContext();
 
-	// capture the velocities at the start of each sim iteration
-	for (int i = 0; i < SimulationModuleTree.Num(); i++)
+	if (Chaos::FClusterUnionPhysicsProxy::FInternalParticle* ParentParticle = PhysicsProxy->GetParticle_Internal())
 	{
-		if (ISimulationModuleBase* Module = SimulationModuleTree[i].SimModule)
+		if (bWake && !SimulationModuleTree.IsEmpty())
 		{
-			if (const Chaos::FClusterUnionPhysicsProxy::FInternalParticle* ParentParticle = PhysicsProxy->GetParticle_Internal())
+			const FPhysicsObjectHandle& PhysicsObject = PhysicsProxy->GetPhysicsObjectHandle();
+			
+			Chaos::FWritePhysicsObjectInterface_Internal WriteInterface = Chaos::FPhysicsObjectInternalInterface::GetWrite();
+			WriteInterface.WakeUp( { &PhysicsObject, 1 });
+		}
+
+		// capture the velocities at the start of each sim iteration
+		for (int i = 0; i < SimulationModuleTree.Num(); i++)
+		{
+			if (ISimulationModuleBase* Module = SimulationModuleTree[i].SimModule)
 			{
 				const FTransform BodyTransform(ParentParticle->R(), ParentParticle->X());
 
@@ -394,7 +424,7 @@ void FSimModuleTree::UpdateModuleVelocites(FClusterUnionPhysicsProxy* PhysicsPro
 					}
 					else
 					{
-		//				Particle = Particles[Module->GetTransformIndex()];
+						//				Particle = Particles[Module->GetTransformIndex()];
 					}
 
 					if (Particle)
@@ -423,11 +453,17 @@ void FSimModuleTree::UpdateModuleVelocites(FClusterUnionPhysicsProxy* PhysicsPro
 			}
 		}
 	}
-
 }
 
 void FSimModuleTree::GenerateReplicationStructure(Chaos::FModuleNetDataArray& NetData)
 {
+	SCOPE_CYCLE_COUNTER(STAT_ModularVehicle_GenerateReplicationStructure);
+
+	if (!bModularVehicle_NetworkData_Enable)
+	{
+		return;
+	}
+
 	const TArray<FSimModuleNode>& Tree = SimulationModuleTree;
 	NetData.Reserve(Tree.Num());
 	for (int Index = 0; Index < Tree.Num(); Index++)
@@ -446,6 +482,13 @@ void FSimModuleTree::GenerateReplicationStructure(Chaos::FModuleNetDataArray& Ne
 
 void FSimModuleTree::SetNetState(Chaos::FModuleNetDataArray& ModuleDatas)
 {
+	SCOPE_CYCLE_COUNTER(STAT_ModularVehicle_SetNetState);
+
+	if (!bModularVehicle_NetworkData_Enable)
+	{
+		return;
+	}
+
 	if (ModuleDatas.IsEmpty())
 	{
 		GenerateReplicationStructure(ModuleDatas);
@@ -465,6 +508,13 @@ void FSimModuleTree::SetNetState(Chaos::FModuleNetDataArray& ModuleDatas)
 
 void FSimModuleTree::SetSimState(const Chaos::FModuleNetDataArray& ModuleDatas)
 {
+	SCOPE_CYCLE_COUNTER(STAT_ModularVehicle_SetSimState);
+
+	if (!bModularVehicle_NetworkData_Enable)
+	{
+		return;
+	}
+
 	for (const TSharedPtr<FModuleNetData>& DataElement : ModuleDatas)
 	{
 		if (!SimulationModuleTree.IsEmpty() && DataElement->SimArrayIndex < SimulationModuleTree.Num())
diff --git a/Engine/Source/Runtime/Experimental/ChaosVehicles/ChaosVehiclesCore/Public/SimModule/SimModuleTree.h b/Engine/Source/Runtime/Experimental/ChaosVehicles/ChaosVehiclesCore/Public/SimModule/SimModuleTree.h
--- a/Engine/Source/Runtime/Experimental/ChaosVehicles/ChaosVehiclesCore/Public/SimModule/SimModuleTree.h
+++ b/Engine/Source/Runtime/Experimental/ChaosVehicles/ChaosVehiclesCore/Public/SimModule/SimModuleTree.h
@@ -5,6 +5,9 @@
 #include "SimModule/DeferredForcesModular.h"
 #include "SimModule/SimulationModuleBase.h"
 
+
+DECLARE_STATS_GROUP(TEXT("ModularVehicle.SimTree"), STATGROUP_ModularVehicleSimTree, STATGROUP_Advanced);
+
 class FGeometryCollectionPhysicsProxy;
 namespace Chaos
 {
@@ -183,7 +186,7 @@ namespace Chaos
 		int GetNextIndex();
 
 		void UpdateModuleVelocites(FGeometryCollectionPhysicsProxy* PhysicsProxy);
-		void UpdateModuleVelocites(FClusterUnionPhysicsProxy* PhysicsProxy);
+		void UpdateModuleVelocites(FClusterUnionPhysicsProxy* PhysicsProxy, bool bWake);
 
 		TArray<FSimModuleNode> SimulationModuleTree;
 		TArray<int> FreeList;
diff --git a/Engine/Source/Runtime/Experimental/ChaosVehicles/ChaosVehiclesCore/Public/SimModule/SimulationModuleBase.h b/Engine/Source/Runtime/Experimental/ChaosVehicles/ChaosVehiclesCore/Public/SimModule/SimulationModuleBase.h
--- a/Engine/Source/Runtime/Experimental/ChaosVehicles/ChaosVehiclesCore/Public/SimModule/SimulationModuleBase.h
+++ b/Engine/Source/Runtime/Experimental/ChaosVehicles/ChaosVehiclesCore/Public/SimModule/SimulationModuleBase.h
@@ -46,6 +46,17 @@ namespace Chaos
 		{
 		}
 
+		bool InputNonZero()
+		{
+			return FMath::Abs(Throttle) > SMALL_NUMBER
+				|| FMath::Abs(Brake) > SMALL_NUMBER
+				|| FMath::Abs(Steering) > SMALL_NUMBER
+				|| FMath::Abs(Roll) > SMALL_NUMBER
+				|| FMath::Abs(Pitch) > SMALL_NUMBER
+				|| FMath::Abs(Yaw) > SMALL_NUMBER
+				|| FMath::Abs(Boost) > SMALL_NUMBER;
+		}
+
 		bool IsValid;
 		float Throttle;
 		float Brake;
