simple CIS Fix

diff --git a/Engine/Source/Runtime/Experimental/ChaosVehicles/ChaosVehiclesCore/Private/SimModule/SimModuleTree.cpp b/Engine/Source/Runtime/Experimental/ChaosVehicles/ChaosVehiclesCore/Private/SimModule/SimModuleTree.cpp
--- a/Engine/Source/Runtime/Experimental/ChaosVehicles/ChaosVehiclesCore/Private/SimModule/SimModuleTree.cpp
+++ b/Engine/Source/Runtime/Experimental/ChaosVehicles/ChaosVehiclesCore/Private/SimModule/SimModuleTree.cpp
@@ -537,6 +537,6 @@ void FSimModuleTree::InterpolateState(const float LerpFactor, Chaos::FModuleNetD
 
 
 #if VEHICLE_DEBUGGING_ENABLED
-UE_ENABLE_OPTIMIZATION
+UE_ENABLE_OPTIMIZATION_SHIP
 #endif
 
