make geometry collections opt out of being controlled by physics replication

diff --git a/Engine/Source/Runtime/Engine/Classes/Components/PrimitiveComponent.h b/Engine/Source/Runtime/Engine/Classes/Components/PrimitiveComponent.h
--- a/Engine/Source/Runtime/Engine/Classes/Components/PrimitiveComponent.h
+++ b/Engine/Source/Runtime/Engine/Classes/Components/PrimitiveComponent.h
@@ -2710,6 +2710,11 @@ public:
 	 */
 	ENGINE_API void SetRigidBodyReplicatedTarget(FRigidBodyState& UpdatedState, const FName BoneName = NAME_None, int32 ServerFrame = 0, int32 ServerHandle = 0);
 
+protected:
+	ENGINE_API virtual bool CanBeUsedInPhysicsReplication(const FName BoneName = NAME_None) const { return true; }
+
+public:
+
 	/** 
 	 *	Get the state of the rigid body responsible for this Actor's physics, and fill in the supplied FRigidBodyState struct based on it.
 	 *
diff --git a/Engine/Source/Runtime/Engine/Private/PrimitiveComponentPhysics.cpp b/Engine/Source/Runtime/Engine/Private/PrimitiveComponentPhysics.cpp
--- a/Engine/Source/Runtime/Engine/Private/PrimitiveComponentPhysics.cpp
+++ b/Engine/Source/Runtime/Engine/Private/PrimitiveComponentPhysics.cpp
@@ -40,6 +40,11 @@ namespace PrimitiveComponentCVars
 
 void UPrimitiveComponent::SetRigidBodyReplicatedTarget(FRigidBodyState& UpdatedState, FName BoneName, int32 ServerFrame, int32 ServerHandle)
 {
+	if (!CanBeUsedInPhysicsReplication(BoneName))
+	{
+		return;
+	}
+
 	if (UWorld* World = GetWorld())
 	{
 		if (FPhysScene* PhysScene = World->GetPhysicsScene())
diff --git a/Engine/Source/Runtime/Experimental/GeometryCollectionEngine/Private/GeometryCollection/GeometryCollectionComponent.cpp b/Engine/Source/Runtime/Experimental/GeometryCollectionEngine/Private/GeometryCollection/GeometryCollectionComponent.cpp
--- a/Engine/Source/Runtime/Experimental/GeometryCollectionEngine/Private/GeometryCollection/GeometryCollectionComponent.cpp
+++ b/Engine/Source/Runtime/Experimental/GeometryCollectionEngine/Private/GeometryCollection/GeometryCollectionComponent.cpp
@@ -6590,6 +6590,11 @@ void UGeometryCollectionComponent::OnComponentCollisionSettingsChanged(bool bUpd
 	LoadCollisionProfiles();
 }
 
+bool UGeometryCollectionComponent::CanBeUsedInPhysicsReplication(const FName BoneName) const
+{
+	return false;
+}
+
 #if WITH_EDITOR
 
 bool UGeometryCollectionComponent::IsHLODRelevant() const
diff --git a/Engine/Source/Runtime/Experimental/GeometryCollectionEngine/Public/GeometryCollection/GeometryCollectionComponent.h b/Engine/Source/Runtime/Experimental/GeometryCollectionEngine/Public/GeometryCollection/GeometryCollectionComponent.h
--- a/Engine/Source/Runtime/Experimental/GeometryCollectionEngine/Public/GeometryCollection/GeometryCollectionComponent.h
+++ b/Engine/Source/Runtime/Experimental/GeometryCollectionEngine/Public/GeometryCollection/GeometryCollectionComponent.h
@@ -644,6 +644,7 @@ public:
 	GEOMETRYCOLLECTIONENGINE_API virtual void PostLoad() override;
 protected:
 	GEOMETRYCOLLECTIONENGINE_API virtual void OnComponentCollisionSettingsChanged(bool bUpdateOverlaps=true) override;
+	GEOMETRYCOLLECTIONENGINE_API virtual bool CanBeUsedInPhysicsReplication(const FName BoneName = NAME_None) const override;
 	//~ End UPrimitiveComponent Interface.
 
 
