Retain IoStoreOnDemand state changes if called before init

Keep the set state if any of the interface methods are called before the backend is initialized.

diff --git a/Engine/Source/Runtime/Experimental/IoStoreOnDemand/Private/IoStoreOnDemand.cpp b/Engine/Source/Runtime/Experimental/IoStoreOnDemand/Private/IoStoreOnDemand.cpp
--- a/Engine/Source/Runtime/Experimental/IoStoreOnDemand/Private/IoStoreOnDemand.cpp
+++ b/Engine/Source/Runtime/Experimental/IoStoreOnDemand/Private/IoStoreOnDemand.cpp
@@ -1489,6 +1489,10 @@ void FIoStoreOnDemandModule::SetBulkOptionalEnabled(bool bInEnabled)
 	{
 		Backend->SetBulkOptionalEnabled(bInEnabled);
 	}
+	else
+	{
+		DeferredBulkOptionalEnabled = bInEnabled;
+	}
 }
 
 void FIoStoreOnDemandModule::SetEnabled(bool bInEnabled)
@@ -1497,6 +1501,10 @@ void FIoStoreOnDemandModule::SetEnabled(bool bInEnabled)
 	{
 		Backend->SetEnabled(bInEnabled);
 	}
+	else
+	{
+		DeferredEnabled = bInEnabled;
+	}
 }
 
 void FIoStoreOnDemandModule::AbandonCache()
@@ -1505,6 +1513,10 @@ void FIoStoreOnDemandModule::AbandonCache()
 	{
 		Backend->AbandonCache();
 	}
+	else
+	{
+		DeferredAbandonCache = true;
+	}
 }
 
 void FIoStoreOnDemandModule::ReportAnalytics(TArray<FAnalyticsEventAttribute>& OutAnalyticsArray) const
@@ -1584,6 +1596,7 @@ void FIoStoreOnDemandModule::InitializeInternal()
 
 	TUniquePtr<IIasCache> Cache;
 	FIasCacheConfig CacheConfig = GetIasCacheConfig(CommandLine);
+	CacheConfig.DropCache = DeferredAbandonCache.Get(CacheConfig.DropCache);
 	if (CacheConfig.DiskQuota > 0)
 	{
 		FString CacheDir = FPaths::ProjectPersistentDownloadDir();
@@ -1605,6 +1618,17 @@ void FIoStoreOnDemandModule::InitializeInternal()
 		BackendPriority = 10;
 	}
 #endif
+
+	// Setup any states changes issued before initialization
+	if (DeferredEnabled.IsSet())
+	{
+		Backend->SetEnabled(*DeferredEnabled);
+	}
+	if (DeferredBulkOptionalEnabled.IsSet())
+	{
+		Backend->SetBulkOptionalEnabled(*DeferredBulkOptionalEnabled);
+	}
+	
 	FIoDispatcher::Get().Mount(Backend.ToSharedRef(), BackendPriority);
 }
 	
diff --git a/Engine/Source/Runtime/Experimental/IoStoreOnDemand/Public/IO/IoStoreOnDemand.h b/Engine/Source/Runtime/Experimental/IoStoreOnDemand/Public/IO/IoStoreOnDemand.h
--- a/Engine/Source/Runtime/Experimental/IoStoreOnDemand/Public/IO/IoStoreOnDemand.h
+++ b/Engine/Source/Runtime/Experimental/IoStoreOnDemand/Public/IO/IoStoreOnDemand.h
@@ -217,6 +217,11 @@ class FIoStoreOnDemandModule
 private:
 	void InitializeInternal();
 	TSharedPtr<IOnDemandIoDispatcherBackend> Backend;
+	// Deferred state requests if called before backend
+	// is initialized
+	TOptional<bool> DeferredEnabled;
+	TOptional<bool> DeferredAbandonCache;
+	TOptional<bool> DeferredBulkOptionalEnabled;
 
 public:
 	UE_API void SetBulkOptionalEnabled(bool bInEnabled);
