Disable the generation on the OnDemand toc and return to downloading it at startup.

#rnx

diff --git a/Engine/Source/Runtime/Experimental/IoStoreOnDemand/Private/OnDemandIoDispatcherBackend.cpp b/Engine/Source/Runtime/Experimental/IoStoreOnDemand/Private/OnDemandIoDispatcherBackend.cpp
--- a/Engine/Source/Runtime/Experimental/IoStoreOnDemand/Private/OnDemandIoDispatcherBackend.cpp
+++ b/Engine/Source/Runtime/Experimental/IoStoreOnDemand/Private/OnDemandIoDispatcherBackend.cpp
@@ -132,7 +132,7 @@ static FAutoConsoleVariableRef CVar_IoReportAnalytics(
 	GIasReportAnalyticsEnabled,
 	TEXT("Enables reporting statics to the analytics system"));
 
-bool GIasGenerateOnDemandToc = true;
+bool GIasGenerateOnDemandToc = false;
 static FAutoConsoleVariableRef CVar_IasGenerateOnDemandToc(
 	TEXT("s.IasGenerateOnDemandToc"),
 	GIasGenerateOnDemandToc,
