IAS - try each available endpoint before failing HTTP requests

- updated HTTP client to retry on all available endpoints before failing a request
- the current primary endpoint will be tested twice before moving on to the next endpoint
- after a successful retry the primary endpoint will be changed to the new endpoint
- updated the on demand backend to reconnect to the first endpoint retrieved from the distribution url when the HTTP client has changed endpoint due to failure

#rnx

diff --git a/Engine/Source/Runtime/Experimental/IoStoreOnDemand/Private/CoreHttp/Client.h b/Engine/Source/Runtime/Experimental/IoStoreOnDemand/Private/CoreHttp/Client.h
--- a/Engine/Source/Runtime/Experimental/IoStoreOnDemand/Private/CoreHttp/Client.h
+++ b/Engine/Source/Runtime/Experimental/IoStoreOnDemand/Private/CoreHttp/Client.h
@@ -77,6 +77,8 @@ public:
 	FConnectionPool&		operator = (FConnectionPool&& Rhs)		{ Swap(Ptr, Rhs.Ptr); return *this; }
 	bool					Resolve();
 
+	static bool				IsValidHostUrl(FAnsiStringView Url);
+
 private:
 	friend					class FEventLoop;
 	class FSocketPool*		Ptr = nullptr;
diff --git a/Engine/Source/Runtime/Experimental/IoStoreOnDemand/Private/CoreHttp/Http.cpp b/Engine/Source/Runtime/Experimental/IoStoreOnDemand/Private/CoreHttp/Http.cpp
--- a/Engine/Source/Runtime/Experimental/IoStoreOnDemand/Private/CoreHttp/Http.cpp
+++ b/Engine/Source/Runtime/Experimental/IoStoreOnDemand/Private/CoreHttp/Http.cpp
@@ -1499,7 +1499,12 @@ bool FConnectionPool::Resolve()
 	return (Ptr->ResolveHostName().GetValue() > 0);
 }
 
-
+////////////////////////////////////////////////////////////////////////////////
+bool FConnectionPool::IsValidHostUrl(FAnsiStringView Url)
+{
+	FUrlOffsets Tmp;
+	return ParseUrl(Url, Tmp) >= 0;
+}
 
 // {{{1 activity ...............................................................
 
diff --git a/Engine/Source/Runtime/Experimental/IoStoreOnDemand/Private/IoStoreOnDemand.cpp b/Engine/Source/Runtime/Experimental/IoStoreOnDemand/Private/IoStoreOnDemand.cpp
--- a/Engine/Source/Runtime/Experimental/IoStoreOnDemand/Private/IoStoreOnDemand.cpp
+++ b/Engine/Source/Runtime/Experimental/IoStoreOnDemand/Private/IoStoreOnDemand.cpp
@@ -124,7 +124,7 @@ static bool TryParseConfigContent(const FString& ConfigContent, const FString& C
 	Config.ProcessInputFileContents(ConfigContent, ConfigFileName);
 
 	Config.GetString(TEXT("Endpoint"), TEXT("DistributionUrl"), OutEndpoint.DistributionUrl);
-	Config.GetString(TEXT("Endpoint"), TEXT("ServiceUrl"), OutEndpoint.ServiceUrl);
+	Config.GetArray(TEXT("Endpoint"), TEXT("ServiceUrl"), OutEndpoint.ServiceUrls);
 	Config.GetString(TEXT("Endpoint"), TEXT("TocPath"), OutEndpoint.TocPath);
 
 	if (OutEndpoint.DistributionUrl.EndsWith(TEXT("/")))
@@ -132,9 +132,12 @@ static bool TryParseConfigContent(const FString& ConfigContent, const FString& C
 		OutEndpoint.DistributionUrl = OutEndpoint.DistributionUrl.Left(OutEndpoint.DistributionUrl.Len() - 1);
 	}
 
-	if (OutEndpoint.ServiceUrl.EndsWith(TEXT("/")))
+	for (FString& ServiceUrl : OutEndpoint.ServiceUrls)
 	{
-		OutEndpoint.ServiceUrl = OutEndpoint.DistributionUrl.Left(OutEndpoint.ServiceUrl.Len() - 1);
+		if (ServiceUrl.EndsWith(TEXT("/")))
+		{
+			ServiceUrl.LeftInline(ServiceUrl.Len() - 1);
+		}
 	}
 
 	if (OutEndpoint.TocPath.StartsWith(TEXT("/")))
@@ -545,7 +548,7 @@ bool LoadFromCompactBinary(FCbFieldView Field, FOnDemandToc& OutToc)
 	return false;
 }
 
-TIoStatusOr<FOnDemandToc> LoadTocFromUrl(const FString& ServiceURL, const FString& TocPath, int32 RetryCount)
+TIoStatusOr<FOnDemandToc> LoadTocFromUrl(const FString& ServiceUrl, const FString& TocPath, int32 RetryCount)
 {
 	TRACE_CPUPROFILER_EVENT_SCOPE(LoadTocFromUrl);
 
@@ -553,11 +556,11 @@ TIoStatusOr<FOnDemandToc> LoadTocFromUrl(const FString& ServiceURL, const FStrin
 
 	for (int32 Attempt = 0; Attempt <= RetryCount; ++Attempt)
 	{
-		TUniquePtr<FOnDemandHttpClient> HttpClient = MakeUnique<FOnDemandHttpClient>(ServiceURL, 1);
+		TUniquePtr<FHttpClient> HttpClient = FHttpClient::Create(ServiceUrl);
 		TAnsiStringBuilder<256> Url;
 
 		Url << "/" << TocPath;
-		UE_LOG(LogIas, Log, TEXT("Fetching TOC '%s/%s' (#%d/%d)"), *HttpClient->ServiceUrl(), *TocPath, Attempt + 1, RetryCount);
+		UE_LOG(LogIas, Log, TEXT("Fetching TOC '%s/%s' (#%d/%d)"), *ServiceUrl, *TocPath, Attempt + 1, RetryCount);
 
 		TIoStatusOr<FOnDemandToc> Toc;
 		HttpClient->Get(Url.ToView(), [&Toc, &ErrorMsg](TIoStatusOr<FIoBuffer> Response, uint64 DurationMs)
@@ -585,8 +588,7 @@ TIoStatusOr<FOnDemandToc> LoadTocFromUrl(const FString& ServiceURL, const FStrin
 				}
 			});
 
-		const bool bBlock = true;
-		while (HttpClient->Tick(bBlock));
+		while (HttpClient->Tick());
 
 		if (Toc.IsOk())
 		{
@@ -1564,8 +1566,8 @@ void FIoStoreOnDemandModule::InitializeInternal()
 			int32 Delim = INDEX_NONE;
 			if (UrlView.RightChop(7).FindChar(TEXT('/'), Delim))
 			{
-				Endpoint.ServiceUrl = UrlView.Left(7 +  Delim);
-				Endpoint.TocPath = UrlView.RightChop(Endpoint.ServiceUrl.Len() + 1);
+				Endpoint.ServiceUrls.Add(FString(UrlView.Left(7 +  Delim)));
+				Endpoint.TocPath = UrlView.RightChop(Endpoint.ServiceUrls[0].Len() + 1);
 			}
 		}
 	}
@@ -1587,9 +1589,6 @@ void FIoStoreOnDemandModule::InitializeInternal()
 		}
 	}
 
-	UE_LOG(LogIas, Display, TEXT("Initializing IoStoreOnDemand on service %s from %s."),
-		*Endpoint.ServiceUrl, *Endpoint.DistributionUrl);
-
 	FLatencyInjector::Initialize(CommandLine);
 
 	Endpoint.EndpointType = EOnDemandEndpointType::CDN;
diff --git a/Engine/Source/Runtime/Experimental/IoStoreOnDemand/Private/OnDemandHttpClient.cpp b/Engine/Source/Runtime/Experimental/IoStoreOnDemand/Private/OnDemandHttpClient.cpp
--- a/Engine/Source/Runtime/Experimental/IoStoreOnDemand/Private/OnDemandHttpClient.cpp
+++ b/Engine/Source/Runtime/Experimental/IoStoreOnDemand/Private/OnDemandHttpClient.cpp
@@ -11,124 +11,220 @@
 namespace UE::IO::IAS
 {
 
-int32 GIasHttpRateLimitKiBPerSecond = 0;
-static FAutoConsoleVariableRef CVar_GIasHttpRateLimitKiBPerSecond(
-	TEXT("ias.HttpRateLimitKiBPerSecond"),
-	GIasHttpRateLimitKiBPerSecond,
-	TEXT("Http throttle limit in KiBPerSecond")
-);
-
-int32 GIasHttpPollTimeoutMs = 17;
-static FAutoConsoleVariableRef CVar_GIasHttpPollTimeoutMs(
-	TEXT("ias.HttpPollTimeoutMs"),
-	GIasHttpPollTimeoutMs,
-	TEXT("Http tick poll timeout in milliseconds")
-);
-
-static int32 GIasHttpRecvBufKiB = -1;
-static FAutoConsoleVariableRef CVar_GIasHttpRecvBufKiB(
-	TEXT("ias.HttpRecvBufKiB"),
-	GIasHttpRecvBufKiB,
-	TEXT("Recv buffer size")
-);
-
-static void LogHttpResult(const TCHAR* Url, uint32 StatusCode, uint64 DurationMs, uint64 Size, uint64 Offset, const char* Memo = "ok")
+static void LogHttpResult(const TCHAR* Host, const TCHAR* Url, uint32 StatusCode, uint64 DurationMs, uint64 Size, uint64 Offset, const char* Memo = "ok")
 {
 	Size >>= 10;
-	UE_LOG(LogIas, VeryVerbose, TEXT("http-%3u: %5" UINT64_FMT "ms %5" UINT64_FMT "KiB[%7" UINT64_FMT "] '%S' %s"), StatusCode, DurationMs, Size, Offset, Memo, Url);
+	UE_LOG(LogIas, VeryVerbose, TEXT("http-%3u: %5" UINT64_FMT "ms %5" UINT64_FMT "KiB[%7" UINT64_FMT "] '%S' %s%s"), StatusCode, DurationMs, Size, Offset, Memo, Host, Url);
 };
 
-FOnDemandHttpClient::FOnDemandHttpClient(const FString& ServiceUrl, int32 MaxConnectionCount)
-	: SvcsUrl(ServiceUrl)
-	, MaxConnections(MaxConnectionCount)
+void FHttpClient::Get(FAnsiStringView Url, const FIoOffsetAndLength& Range, FGetCallback&& Callback)
 {
-	auto ServiceUrlAnsi = StringCast<ANSICHAR>(*ServiceUrl, ServiceUrl.Len());
+	check(PrimaryConnection != INDEX_NONE);
+	IssueRequest(FRequestParams
+	{
+		.Url = FString(Url),
+		.Range = Range,
+		.Callback = MoveTemp(Callback),
+		.Connection = PrimaryConnection
+	});
+}
 
-	HTTP::FConnectionPool::FParams Params;
-	if (Params.SetHostFromUrl(ServiceUrlAnsi) < 0)
+void FHttpClient::Get(FAnsiStringView Url, FGetCallback&& Callback)
+{
+	Get(Url, FIoOffsetAndLength(), MoveTemp(Callback));
+}
+
+bool FHttpClient::Tick(uint32 WaitTimeMs, uint32 MaxKiBPerSecond)
+{
+	if (PrimaryConnection == INDEX_NONE)
 	{
-		UE_LOG(LogIas, Error, TEXT("Failed to set host from '%s'"), *ServiceUrl);
+		return false;
 	}
 
-	if (GIasHttpRecvBufKiB >= 0)
+	EventLoop.Throttle(MaxKiBPerSecond);
+	const bool bTicked = EventLoop.Tick(WaitTimeMs) != 0;
+
+	// Clean up all but the primary connection
+	if (bTicked == false)
 	{
-		UE_LOG(LogIas, VeryVerbose, TEXT("Set HTTP recv buffer size to %dKib"), GIasHttpRecvBufKiB);
-		Params.RecvBufSize = GIasHttpRecvBufKiB << 10;
+		for (int32 Idx = 0, Count = Connections.Num(); Idx < Count; ++Idx)
+		{
+			if (Idx != PrimaryConnection)
+			{
+				Connections[Idx].Reset();
+			}
+		}
 	}
-	Params.ConnectionCount = MaxConnectionCount;
-	ConnectionPool = MakeUnique<HTTP::FConnectionPool>(Params);
+
+	return bTicked;
 }
 
-void FOnDemandHttpClient::Get(FAnsiStringView Url, const FIoOffsetAndLength& Range, FGetCallback&& Callback)
+FHttpClient::FHttpClient(FHttpClientConfig&& ClientConfig)
+	: Config(MoveTemp(ClientConfig))
 {
-	Issue(Url, MoveTemp(Callback), Range);
+	Configure();
 }
 
-void FOnDemandHttpClient::Get(FAnsiStringView Url, FGetCallback&& Callback)
+TUniquePtr<FHttpClient> FHttpClient::Create(FHttpClientConfig&& ClientConfig)
 {
-	Issue(Url, MoveTemp(Callback));
+	if (ClientConfig.Endpoints.IsEmpty())
+	{
+		return TUniquePtr<FHttpClient>();
+	}
+
+	for (const FString& Ep : ClientConfig.Endpoints)
+	{
+		if (Ep.IsEmpty())
+		{
+			return TUniquePtr<FHttpClient>();
+		}
+
+		//TODO: Get rid of all string conversions
+		if (UE::IO::IAS::HTTP::FConnectionPool::IsValidHostUrl(StringCast<ANSICHAR>(*Ep, Ep.Len())) == false)
+		{
+			return TUniquePtr<FHttpClient>();
+		}
+	}
+
+	return TUniquePtr<FHttpClient>(new FHttpClient(MoveTemp(ClientConfig)));
 }
 
-void FOnDemandHttpClient::Issue(FAnsiStringView Url, FGetCallback&& Callback, FIoOffsetAndLength Range)
+TUniquePtr<FHttpClient> FHttpClient::Create(const FString& Endpoint)
+{
+	FHttpClientConfig Config;
+	Config.Endpoints.Add(Endpoint);
+	return Create(MoveTemp(Config));
+}
+
+void FHttpClient::Configure()
+{
+	check(Config.Endpoints.IsEmpty() == false);
+	Connections.SetNum(Config.Endpoints.Num());
+	Connections[Config.PrimaryEndpoint] = CreateConnection(Config.Endpoints[Config.PrimaryEndpoint]);
+	PrimaryConnection = Config.PrimaryEndpoint;
+}
+
+void FHttpClient::IssueRequest(FRequestParams&& Params)
 {
 	using namespace UE::IO::IAS::HTTP;
 
+	//TODO: Remove string conversion
+	const auto Url = StringCast<ANSICHAR>(*Params.Url, Params.Url.Len());
+
 	auto Sink = [
+		this,
+		Params = MoveTemp(Params),
 		Buffer = FIoBuffer(),
-		Callback = MoveTemp(Callback),
-		Url = FString(Url),
-		Offset = Range.GetOffset(),
 		StartTime = FPlatformTime::Cycles64(),
 		StatusCode = uint32(0)]
 		(const FTicketStatus& Status) mutable
-	{
-		if (FTicketStatus::EId::Response == Status.GetId())
-		{
-			FResponse& Response = Status.GetResponse();
-			StatusCode = Response.GetStatusCode();
-			Response.SetDestination(&Buffer);
-		}
-		else if (FTicketStatus::EId::Content == Status.GetId())
 		{
 			const uint64 DurationMs = (uint64)FPlatformTime::ToMilliseconds64(FPlatformTime::Cycles64() - StartTime);
-			const FIoBuffer& Content = Status.GetContent();
-
-			LogHttpResult(*Url, StatusCode, DurationMs, Content.GetSize(), Offset);
-
-			const bool bSuccessful = StatusCode > 199 && StatusCode < 300;
-			if (bSuccessful && Content.GetSize() > 0)
+			switch (Status.GetId())
+			{
+			case FTicketStatus::EId::Response:
 			{
-				Callback(Content, DurationMs);
+				FResponse& Response = Status.GetResponse();
+				StatusCode = Response.GetStatusCode();
+				Response.SetDestination(&Buffer);
+				break;
 			}
-			else
+			case FTicketStatus::EId::Content:
 			{
-				Callback(FIoStatus(EIoErrorCode::NotFound, TEXTVIEW("Invalid Content")), DurationMs);
+				const FIoBuffer& Content = Status.GetContent();
+				LogHttpResult(*Config.Endpoints[Params.Connection], *Params.Url, StatusCode, DurationMs, Content.GetSize(), Params.Range.GetOffset());
+				
+				const bool bSuccessful = StatusCode > 199 && StatusCode < 300;
+				if (bSuccessful && Content.GetSize() > 0)
+				{
+					Params.Callback(Content, DurationMs);
+					if (Params.Attempt > 1 && Params.Connection != PrimaryConnection)
+					{
+						UE_LOG(LogIas, Log, TEXT("HTTP client endpoint changed from '%s' to '%s'"),
+							*Config.Endpoints[PrimaryConnection], *Config.Endpoints[Params.Connection]);
+						PrimaryConnection = Params.Connection;
+					}
+				}
+				else
+				{
+					const bool bServerError = StatusCode > 499 && StatusCode < 600;
+					if (Params.Attempt < Config.MaxRetryCount)
+					{
+						const bool bNextConnection = bServerError == false && Params.Attempt > 0;
+						RetryRequest(MoveTemp(Params), bNextConnection);
+					}
+					else
+					{
+						const EIoErrorCode ErrorCode = bServerError ? EIoErrorCode::ReadError : EIoErrorCode::NotFound;
+						Params.Callback(FIoStatus(ErrorCode), DurationMs);
+					}
+				}
+				break;
 			}
-		}
-		else if (FTicketStatus::EId::Error == Status.GetId())
-		{
-			const uint64 DurationMs = (uint64)FPlatformTime::ToMilliseconds64(FPlatformTime::Cycles64() - StartTime);
-			LogHttpResult(*Url, StatusCode, DurationMs, 0, Offset, Status.GetErrorReason());
-			Callback(FIoStatus(EIoErrorCode::ReadError, FString(Status.GetErrorReason())), DurationMs);
-		}
-	};
+			case FTicketStatus::EId::Error:
+			{
+				if (Params.Attempt < Config.MaxRetryCount)
+				{
+					const bool bNextConnection = Params.Attempt > 0;
+					RetryRequest(MoveTemp(Params), bNextConnection);
+				}
+				else
+				{
+					LogHttpResult(*Config.Endpoints[Params.Connection], *Params.Url, StatusCode, DurationMs, 0, Params.Range.GetOffset(), Status.GetErrorReason());
+					Params.Callback(FIoStatus(EIoErrorCode::ReadError, FString(Status.GetErrorReason())), DurationMs);
+				}
+				break;
+			}
+			default:
+			{
+				LogHttpResult(*Config.Endpoints[Params.Connection], *Params.Url, StatusCode, DurationMs, 0, Params.Range.GetOffset(), "cancelled");
+				Params.Callback(FIoStatus(EIoErrorCode::Cancelled, FString(Status.GetErrorReason())), DurationMs);
+				break;
+			}
+			}
+		};
 
-		UE::IO::IAS::HTTP::FRequest Request = EventLoop.Get(Url, *ConnectionPool);
-		const uint64 RangeStart = Range.GetOffset();
-		const uint64 RangeEnd = Range.GetOffset() + Range.GetLength();
-		if (RangeStart > 0 || RangeEnd > 0)
+	TUniquePtr<FConnectionPool>& Connection = Connections[Params.Connection];
+	check(Connection.IsValid());
+
+	FRequest Request = EventLoop.Get(Url, *Connection);
+
+	if (Params.Range.GetOffset() > 0 || Params.Range.GetLength() > 0)
+	{
+		Request.Header(ANSITEXTVIEW("Range"),
+			WriteToAnsiString<64>(ANSITEXTVIEW("bytes="), Params.Range.GetOffset(), ANSITEXTVIEW("-"), Params.Range.GetOffset() + Params.Range.GetLength()));
+	}
+
+	EventLoop.Send(MoveTemp(Request), MoveTemp(Sink));
+}
+
+void FHttpClient::RetryRequest(FRequestParams&& Params, bool bNextConnection)
+{
+	check(Params.Attempt < Config.MaxRetryCount);
+	if (bNextConnection && Config.Endpoints.Num() > 1)
+	{
+		Params.Connection = (Params.Connection + 1) % Config.Endpoints.Num();
+		if (Connections[Params.Connection].IsValid() == false)
 		{
-			Request.Header(ANSITEXTVIEW("Range"), WriteToAnsiString<64>(ANSITEXTVIEW("bytes="), RangeStart, ANSITEXTVIEW("-"), RangeEnd));
+			Connections[Params.Connection] = CreateConnection(Config.Endpoints[Params.Connection]);
 		}
-
-		EventLoop.Send(MoveTemp(Request), MoveTemp(Sink));
+	}
+	Params.Attempt++;
+	IssueRequest(MoveTemp(Params));
 }
 
-bool FOnDemandHttpClient::Tick(bool Block)
+TUniquePtr<HTTP::FConnectionPool> FHttpClient::CreateConnection(const FStringView& HostAddr)
 {
-	int32 TimeoutMs = Block ? -1 : GIasHttpPollTimeoutMs;
-	EventLoop.Throttle(GIasHttpRateLimitKiBPerSecond);
-	return EventLoop.Tick(TimeoutMs) != 0;
+	HTTP::FConnectionPool::FParams Params;
+	ensure(Params.SetHostFromUrl(StringCast<ANSICHAR>(HostAddr.GetData(), HostAddr.Len())) >= 0);
+	if (Config.ReceiveBufferSize >= 0)
+	{
+		UE_LOG(LogIas, Log, TEXT("HTTP client receive buffer size set to %d"), Config.ReceiveBufferSize);
+		Params.RecvBufSize = Config.ReceiveBufferSize;
+	}
+	Params.ConnectionCount = Config.MaxConnectionCount;
+
+	return MakeUnique<HTTP::FConnectionPool>(Params);
 }
 
 } //namespace UE::IO::IAS
diff --git a/Engine/Source/Runtime/Experimental/IoStoreOnDemand/Private/OnDemandHttpClient.h b/Engine/Source/Runtime/Experimental/IoStoreOnDemand/Private/OnDemandHttpClient.h
--- a/Engine/Source/Runtime/Experimental/IoStoreOnDemand/Private/OnDemandHttpClient.h
+++ b/Engine/Source/Runtime/Experimental/IoStoreOnDemand/Private/OnDemandHttpClient.h
@@ -10,42 +10,54 @@
 #include "Templates/Function.h"
 #include "Templates/UniquePtr.h"
 
-class FConnectionPool;
-
 namespace UE::IO::IAS
 {
 
-class FOnDemandHttpClient
+struct FHttpClientConfig
+{
+	TArray<FString> Endpoints;
+	int32 PrimaryEndpoint = 0;
+	int32 MaxConnectionCount = 8;
+	int32 MaxRetryCount = 1;
+	int32 ReceiveBufferSize = -1;
+};
+
+class FHttpClient
 {
 public:
 	using FGetCallback = TFunction<void(TIoStatusOr<FIoBuffer>, uint64 DurationMs)>;
 
-	FOnDemandHttpClient(const FString& ServiceUrl, int32 MaxConnectionCount = 8);
-	~FOnDemandHttpClient() = default;
+	static TUniquePtr<FHttpClient> Create(FHttpClientConfig&& ClientConfig);
+	static TUniquePtr<FHttpClient> Create(const FString& Endpoint);
 
-	const FString& ServiceUrl() const 
-	{
-		return SvcsUrl;
-	}
-
-	int32 MaxConnectionCount() const
-	{
-		return MaxConnections;
-	}
-
-	void Get(FAnsiStringView Url, FGetCallback&& Callback);
 	void Get(FAnsiStringView Url, const FIoOffsetAndLength& Range, FGetCallback&& Callback);
+	void Get(FAnsiStringView Url, FGetCallback&& Callback);
 
-	/** @return True if the client has pending work otherwise false. */
-	bool Tick(bool Block = false);
+	bool Tick(uint32 WaitTimeMs, uint32 MaxKiBPerSecond);
+	bool Tick() { return Tick(-1, 0); }
 
-private:
-	void Issue(FAnsiStringView Url, FGetCallback&& Callback, FIoOffsetAndLength Range = FIoOffsetAndLength());
+	int32 GetPrimaryConnection() const { return PrimaryConnection; }
 
-	FString SvcsUrl;
-	int32 MaxConnections;
+private:
+	struct FRequestParams
+	{
+		FString Url;
+		FIoOffsetAndLength Range;
+		FGetCallback Callback;
+		int32 Attempt = 0;
+		int32 Connection = INDEX_NONE;
+	};
+
+	FHttpClient(FHttpClientConfig&& ClientConfig);
+	void Configure();
+	void RetryRequest(FRequestParams&& Params, bool bNextConnection);
+	void IssueRequest(FRequestParams&& Params);
+	TUniquePtr<HTTP::FConnectionPool> CreateConnection(const FStringView& HostAddr);
+
+	FHttpClientConfig Config;
 	HTTP::FEventLoop EventLoop;
-	TUniquePtr<HTTP::FConnectionPool> ConnectionPool;
+	TArray<TUniquePtr<HTTP::FConnectionPool>> Connections;
+	int32 PrimaryConnection = INDEX_NONE;
 };
 
 } // namespace UE::IO::IAS
diff --git a/Engine/Source/Runtime/Experimental/IoStoreOnDemand/Private/OnDemandIoDispatcherBackend.cpp b/Engine/Source/Runtime/Experimental/IoStoreOnDemand/Private/OnDemandIoDispatcherBackend.cpp
--- a/Engine/Source/Runtime/Experimental/IoStoreOnDemand/Private/OnDemandIoDispatcherBackend.cpp
+++ b/Engine/Source/Runtime/Experimental/IoStoreOnDemand/Private/OnDemandIoDispatcherBackend.cpp
@@ -63,6 +63,27 @@ namespace UE::IO::IAS
 {
 
 ///////////////////////////////////////////////////////////////////////////////
+int32 GIasHttpPollTimeoutMs = 17;
+static FAutoConsoleVariableRef CVar_GIasHttpPollTimeoutMs(
+	TEXT("ias.HttpPollTimeoutMs"),
+	GIasHttpPollTimeoutMs,
+	TEXT("Http tick poll timeout in milliseconds")
+);
+
+int32 GIasHttpRateLimitKiBPerSecond = 0;
+static FAutoConsoleVariableRef CVar_GIasHttpRateLimitKiBPerSecond(
+	TEXT("ias.HttpRateLimitKiBPerSecond"),
+	GIasHttpRateLimitKiBPerSecond,
+	TEXT("Http throttle limit in KiBPerSecond")
+);
+
+static int32 GIasHttpRecvBufKiB = -1;
+static FAutoConsoleVariableRef CVar_GIasHttpRecvBufKiB(
+	TEXT("ias.HttpRecvBufKiB"),
+	GIasHttpRecvBufKiB,
+	TEXT("Recv buffer size")
+);
+
 int32 GIasMaxHttpConnectionCount = 8;
 static FAutoConsoleVariableRef CVar_IasMaxHttpConnectionCount(
 	TEXT("ias.MaxHttpConnectionCount"),
@@ -754,11 +775,6 @@ struct FChunkRequest
 		return Head;
 	}
 
-	const FIoChunkId& GetChunkId()
-	{
-		return RequestHead->ChunkId;
-	}
-
 	FChunkRequest* NextRequest;
 	FChunkRequestParams Params;
 	FIoRequestImpl* RequestHead;
@@ -1284,6 +1300,12 @@ class FOnDemandIoBackend final
 			return false;
 		}
 
+		FIoChunkId GetChunkId(FChunkRequest* Request)
+		{
+			FScopeLock _(&Mutex);
+			return Request->RequestHead ? Request->RequestHead->ChunkId : FIoChunkId::InvalidChunkId;
+		}
+
 		void Remove(FChunkRequest* Request)
 		{
 			FScopeLock _(&Mutex);
@@ -1365,7 +1387,7 @@ private:
 	FIoStatus DownloadoadOnDemandToc(const FString& CdnUrl, const FString& TocPath);
 
 	void AddDeferredTocs();
-	void ProcessHttpRequests(FOnDemandHttpClient* HttpClient, FBitWindow& HttpErrors, int32 MaxConcurrentRequests);
+	void ProcessHttpRequests(FHttpClient* HttpClient, FBitWindow& HttpErrors, int32 MaxConcurrentRequests);
 
 	TUniquePtr<IIasCache> Cache;
 	TUniquePtr<FOnDemandIoStore> IoStore;
@@ -1901,7 +1923,7 @@ void FOnDemandIoBackend::Mount(const FOnDemandEndpoint& Endpoint)
 	LLM_SCOPE_BYTAG(Ias);
 	TRACE_CPUPROFILER_EVENT_SCOPE(IasBackend::Mount);
 
-	if ((Endpoint.DistributionUrl.IsEmpty() && Endpoint.ServiceUrl.IsEmpty()) || Endpoint.TocPath.IsEmpty())
+	if ((Endpoint.DistributionUrl.IsEmpty() && Endpoint.ServiceUrls.IsEmpty()) || Endpoint.TocPath.IsEmpty())
 	{
 		UE_LOG(LogIas, Error, TEXT("Trying to mount an invalid on demand endpoint"));
 		return;
@@ -1919,7 +1941,7 @@ void FOnDemandIoBackend::Mount(const FOnDemandEndpoint& Endpoint)
 	{
 		{
 			FWriteScopeLock _(Lock);
-			if (Endpoint.ServiceUrl.IsEmpty())
+			if (Endpoint.ServiceUrls.IsEmpty())
 			{
 				if (AvailableEps.HasCurrent() == false)
 				{
@@ -1933,7 +1955,10 @@ void FOnDemandIoBackend::Mount(const FOnDemandEndpoint& Endpoint)
 			}
 			else if (AvailableEps.Urls.IsEmpty())
 			{
-				AvailableEps.Urls.Add(Endpoint.ServiceUrl.Replace(TEXT("https"), TEXT("http")));
+				for (const FString& Url : Endpoint.ServiceUrls)
+				{
+					AvailableEps.Urls.Add(Url.Replace(TEXT("https"), TEXT("http")));
+				}
 				AvailableEps.Current = 0;
 			}
 		}
@@ -1956,13 +1981,12 @@ void FOnDemandIoBackend::Mount(const FOnDemandEndpoint& Endpoint)
 					FWriteScopeLock _(Lock);
 					DeferredTocs.Add(Endpoint.TocPath);
 			}
-	}
-
+		}
 		ConditionallyStartBackendThread();
 	}
 	else
 	{
-		UE_LOG(LogIas, Log, TEXT("Mounting ZEN endpoint, Url='%s'"), *Endpoint.ServiceUrl);
+		UE_LOG(LogIas, Log, TEXT("Mounting ZEN endpoint, Url='%s'"), *Endpoint.ServiceUrls[0]);
 	}
 }
 
@@ -2005,7 +2029,7 @@ void FOnDemandIoBackend::ReportAnalytics(TArray<FAnalyticsEventAttribute>& OutAn
 	}
 }
 
-void FOnDemandIoBackend::ProcessHttpRequests(FOnDemandHttpClient* HttpClient, FBitWindow& HttpErrors, int32 MaxConcurrentRequests)
+void FOnDemandIoBackend::ProcessHttpRequests(FHttpClient* HttpClient, FBitWindow& HttpErrors, int32 MaxConcurrentRequests)
 {
 	int32 NumConcurrentRequests = 0;
 	FChunkRequest* NextChunkRequest = HttpRequests.Dequeue();
@@ -2037,42 +2061,50 @@ void FOnDemandIoBackend::ProcessHttpRequests(FOnDemandHttpClient* HttpClient, FB
 
 					NumConcurrentRequests++;
 					HttpClient->Get(Url.ToView(), ChunkRequest->Params.ChunkRange,
-						[this, &NextChunkRequest, ChunkRequest, &NumConcurrentRequests, &HttpErrors](TIoStatusOr<FIoBuffer> Status, uint64 DurationMs)
+						[this, &NextChunkRequest, ChunkRequest, &NumConcurrentRequests, &HttpErrors]
+						(TIoStatusOr<FIoBuffer> Status, uint64 DurationMs)
 						{
 							NumConcurrentRequests--;
-
-							if (Status.Status().GetErrorCode() == EIoErrorCode::ReadError)
+							switch (Status.Status().GetErrorCode())
 							{
-								if (++ChunkRequest->HttpRetryCount <= GIasMaxHttpRetryCount)
-								{
-									Stats.OnHttpRetry();
-									ChunkRequest->NextRequest = NextChunkRequest;
-									NextChunkRequest = ChunkRequest;
-									return;
-								}
-							}
-
-							if (Status.IsOk())
+							case EIoErrorCode::Ok:
 							{
 								HttpErrors.Add(false);
 								ChunkRequest->Chunk = Status.ConsumeValueOrDie();
 								Stats.OnHttpGet(ChunkRequest->Chunk.DataSize(), DurationMs);
+								break;
 							}
-							else
+							case EIoErrorCode::ReadError:
+							case EIoErrorCode::NotFound:
 							{
 								Stats.OnHttpError();
 								HttpErrors.Add(true);
 
 								const float Average = HttpErrors.AvgSetBits();
 								const bool bAboveHighWaterMark = Average > GIasHttpErrorHighWater;
-								UE_LOG(LogIas, Warning, TEXT("%.2f%% the last %d HTTP requests failed"), Average * 100.0f, GIasHttpErrorSampleCount);
+								UE_LOG(LogIas, Log, TEXT("%.2f%% the last %d HTTP requests failed"), Average * 100.0f, GIasHttpErrorSampleCount);
 
-								if (bAboveHighWaterMark)
+								if (bAboveHighWaterMark && BackendStatus.IsHttpEnabled())
 								{
 									BackendStatus.SetHttpError(true);
 									UE_LOG(LogIas, Warning, TEXT("HTTP streaming disabled due to high water mark of %.2f of the last %d requests reached"),
 										GIasHttpErrorHighWater * 100.0f, GIasHttpErrorSampleCount);
 								}
+								break;
+							}
+							case EIoErrorCode::Cancelled:
+							{
+								const FIoChunkId ChunkId = ChunkRequests.GetChunkId(ChunkRequest);
+								UE_LOG(LogIas, Log, TEXT("HTTP request for chunk '%s' cancelled"), *LexToString(ChunkId));
+								break;
+							}
+							default:
+							{
+								const FIoChunkId ChunkId = ChunkRequests.GetChunkId(ChunkRequest);
+								UE_LOG(LogIas, Warning, TEXT("Unhandled HTTP response '%s' for chunk '%s'"),
+									GetIoErrorText(Status.Status().GetErrorCode()), *LexToString(ChunkId));
+								break;
+							}
 							}
 
 							UE::Tasks::Launch(UE_SOURCE_LOCATION, [this, ChunkRequest]()
@@ -2088,7 +2120,7 @@ void FOnDemandIoBackend::ProcessHttpRequests(FOnDemandHttpClient* HttpClient, FB
 				TRACE_CPUPROFILER_EVENT_SCOPE(IasBackend::TickHttpSaturated);
 				while (NumConcurrentRequests >= MaxConcurrentRequests)
 				{
-					HttpClient->Tick(/*Block*/true);
+					HttpClient->Tick(MAX_uint32, GIasHttpRateLimitKiBPerSecond);
 				}
 			}
 
@@ -2101,7 +2133,7 @@ void FOnDemandIoBackend::ProcessHttpRequests(FOnDemandHttpClient* HttpClient, FB
 		{
 			// Keep processing pending connections until all requests are completed or a new one is issued
 			TRACE_CPUPROFILER_EVENT_SCOPE(IasBackend::TickHttp);
-			while (HttpClient->Tick(/*Block*/false))
+			while (HttpClient->Tick(GIasHttpPollTimeoutMs, GIasHttpRateLimitKiBPerSecond))
 			{
 				if (!NextChunkRequest)
 				{
@@ -2122,11 +2154,19 @@ uint32 FOnDemandIoBackend::Run()
 
 	FBitWindow HttpErrors;
 	HttpErrors.Reset(GIasHttpErrorSampleCount);
-	
-	TUniquePtr<FOnDemandHttpClient> HttpClient;
+
+	const int32 MaxHttpRetryCount = FMath::Max(AvailableEps.Urls.Num() + 1, GIasMaxHttpRetryCount);
+	TUniquePtr<FHttpClient> HttpClient;
 	if (AvailableEps.HasCurrent())
 	{
-		HttpClient = MakeUnique<FOnDemandHttpClient>(AvailableEps.GetCurrent(), GIasMaxHttpConnectionCount);
+		HttpClient = FHttpClient::Create(FHttpClientConfig
+		{
+			.Endpoints = AvailableEps.Urls,
+			.PrimaryEndpoint = AvailableEps.Current,
+			.MaxConnectionCount = GIasMaxHttpConnectionCount,
+			.MaxRetryCount = MaxHttpRetryCount,
+			.ReceiveBufferSize = GIasHttpRecvBufKiB >= 0 ? GIasHttpRecvBufKiB << 10 : -1
+		});
 #if !UE_BUILD_SHIPPING
 		LatencyTest(AvailableEps.GetCurrent(), GetEndpointTestPath());
 #endif 
@@ -2136,6 +2176,7 @@ uint32 FOnDemandIoBackend::Run()
 	{
 		// Process HTTP request(s) even if the client is invalid to ensure enqueued request(s) gets completed.
 		ProcessHttpRequests(HttpClient.Get(), HttpErrors, FMath::Min(2 * GIasMaxHttpConnectionCount, 64));
+		AvailableEps.Current = HttpClient.IsValid() ? HttpClient->GetPrimaryConnection() : INDEX_NONE;
 
 		if (!bStopRequested)
 		{
@@ -2155,12 +2196,40 @@ uint32 FOnDemandIoBackend::Run()
 				if (int32 Idx = LatencyTest(AvailableEps.Urls, TestPath, bStopRequested); Idx != INDEX_NONE)
 				{
 					AvailableEps.Current = Idx;
-					HttpClient = MakeUnique<FOnDemandHttpClient>(AvailableEps.GetCurrent(), GIasMaxHttpConnectionCount);
+					HttpClient = FHttpClient::Create(FHttpClientConfig
+					{
+						.Endpoints = AvailableEps.Urls,
+						.PrimaryEndpoint = Idx,
+						.MaxConnectionCount = GIasMaxHttpConnectionCount,
+						.MaxRetryCount = MaxHttpRetryCount,
+						.ReceiveBufferSize = GIasHttpRecvBufKiB >= 0 ? GIasHttpRecvBufKiB << 10 : -1
+					});
 					BackendStatus.SetHttpError(false);
 					UE_LOG(LogIas, Log, TEXT("Successfully reconnected to '%s'"), *AvailableEps.GetCurrent());
 					AddDeferredTocs();
 				}
 			}
+			else if (HttpClient.IsValid() && HttpClient->GetPrimaryConnection() != 0)
+			{
+				WaitTime = GIasHttpHealthCheckWaitTime;
+				const FString TestPath = GetEndpointTestPath(); 
+				TConstArrayView<FString> Urls = AvailableEps.Urls;
+
+				UE_LOG(LogIas, Log, TEXT("Trying to reconnect to primary endpoint '%s'"), *AvailableEps.Urls[0]);
+				if (int32 Idx = LatencyTest(Urls.Left(1), TestPath, bStopRequested); Idx != INDEX_NONE)
+				{
+					AvailableEps.Current = Idx;
+					HttpClient = FHttpClient::Create(FHttpClientConfig
+					{
+						.Endpoints = AvailableEps.Urls,
+						.PrimaryEndpoint = Idx,
+						.MaxConnectionCount = GIasMaxHttpConnectionCount,
+						.MaxRetryCount = MaxHttpRetryCount,
+						.ReceiveBufferSize = GIasHttpRecvBufKiB >= 0 ? GIasHttpRecvBufKiB << 10 : -1
+					});
+					UE_LOG(LogIas, Log, TEXT("Successfully reconnected to '%s'"), *AvailableEps.GetCurrent());
+				}
+			}
 			if (BackendStatus.ShouldAbandonCache())
 			{
 				BackendStatus.SetAbandonCache(false);
diff --git a/Engine/Source/Runtime/Experimental/IoStoreOnDemand/Private/OnDemandIoDispatcherBackend.h b/Engine/Source/Runtime/Experimental/IoStoreOnDemand/Private/OnDemandIoDispatcherBackend.h
--- a/Engine/Source/Runtime/Experimental/IoStoreOnDemand/Private/OnDemandIoDispatcherBackend.h
+++ b/Engine/Source/Runtime/Experimental/IoStoreOnDemand/Private/OnDemandIoDispatcherBackend.h
@@ -21,12 +21,12 @@ struct FOnDemandEndpoint
 {
 	EOnDemandEndpointType EndpointType;
 	FString DistributionUrl;
-	FString ServiceUrl;
+	TArray<FString> ServiceUrls;
 	FString TocPath;
 
 	bool IsValid() const
 	{
-		return (DistributionUrl.Len() > 0 || ServiceUrl.Len() > 0) && TocPath.Len() > 0;
+		return (DistributionUrl.Len() > 0 || ServiceUrls.Num() > 0) && TocPath.Len() > 0;
 	}
 };
 
