Reduce memory usage in Dynamic Collection :
- Removed Implicits
- Removed Particles
- Compressed Dynamic State

diff --git a/Engine/Source/Programs/HeadlessChaos/Private/GeometryCollection/GeometryCollectionTestClustering.cpp b/Engine/Source/Programs/HeadlessChaos/Private/GeometryCollection/GeometryCollectionTestClustering.cpp
--- a/Engine/Source/Programs/HeadlessChaos/Private/GeometryCollection/GeometryCollectionTestClustering.cpp
+++ b/Engine/Source/Programs/HeadlessChaos/Private/GeometryCollection/GeometryCollectionTestClustering.cpp
@@ -657,8 +657,8 @@ namespace GeometryCollectionTest
 		Params.ClusterGroupIndex = 0;
 		FGeometryCollectionWrapper* Collection1 = TNewSimulationObject<GeometryType::GeometryCollectionWithSuppliedRestCollection>::Init(Params)->template As<FGeometryCollectionWrapper>();
 		TSharedPtr<FGeometryDynamicCollection> DynamicCollection1 = Collection1->DynamicCollection;
-		DynamicCollection1->ModifyAttribute<int32>("DynamicState", FGeometryCollection::TransformGroup)[1] = (uint8)EObjectStateTypeEnum::Chaos_Object_Kinematic;
-		DynamicCollection1->ModifyAttribute<int32>("DynamicState", FGeometryCollection::TransformGroup)[0] = (uint8)EObjectStateTypeEnum::Chaos_Object_Kinematic;
+		DynamicCollection1->ModifyAttribute<uint8>("DynamicState", FGeometryCollection::TransformGroup)[1] = (uint8)EObjectStateTypeEnum::Chaos_Object_Kinematic;
+		DynamicCollection1->ModifyAttribute<uint8>("DynamicState", FGeometryCollection::TransformGroup)[0] = (uint8)EObjectStateTypeEnum::Chaos_Object_Kinematic;
 
 		UnitTest.AddSimulationObject(Collection1);
 
@@ -909,7 +909,7 @@ namespace GeometryCollectionTest
 		Params.MaxClusterLevel = 1;
 		FGeometryCollectionWrapper* Collection = TNewSimulationObject<GeometryType::GeometryCollectionWithSuppliedRestCollection>::Init(Params)->template As<FGeometryCollectionWrapper>();
 
-		Collection->DynamicCollection->template ModifyAttribute<int32>("DynamicState", FGeometryCollection::TransformGroup)[1] = (uint8)EObjectStateTypeEnum::Chaos_Object_Kinematic;
+		Collection->DynamicCollection->template ModifyAttribute<uint8>("DynamicState", FGeometryCollection::TransformGroup)[1] = (uint8)EObjectStateTypeEnum::Chaos_Object_Kinematic;
 
 		UnitTest.AddSimulationObject(Collection);
 		UnitTest.Initialize();
@@ -918,7 +918,7 @@ namespace GeometryCollectionTest
 		FReal CurrentRigidDistance = 0;
 
 		// Staged conditions
-		// Initial state should set up the heirachy correctly, leaving correct disabled flags on frame 1
+		// Initial state should set up the hierarchy correctly, leaving correct disabled flags on frame 1
 		bool bValidInitialState = false;
 		// After releasing particle 8, the states should be updated on frame 2
 		bool bParticle8SucessfulRelease = false;
@@ -1207,7 +1207,7 @@ namespace GeometryCollectionTest
 		Params.MaxClusterLevel = 1;
 		FGeometryCollectionWrapper* Collection = TNewSimulationObject<GeometryType::GeometryCollectionWithSuppliedRestCollection>::Init(Params)->template As<FGeometryCollectionWrapper>();
 
-		Collection->DynamicCollection->template ModifyAttribute<int32>("DynamicState", FGeometryCollection::TransformGroup)[1] = (uint8)EObjectStateTypeEnum::Chaos_Object_Static;
+		Collection->DynamicCollection->template ModifyAttribute<uint8>("DynamicState", FGeometryCollection::TransformGroup)[1] = (uint8)EObjectStateTypeEnum::Chaos_Object_Static;
 
 		UnitTest.AddSimulationObject(Collection);
 		UnitTest.Initialize();
@@ -1940,7 +1940,7 @@ namespace GeometryCollectionTest
 
 		TSharedPtr<FGeometryDynamicCollection> DynamicCollection1 = Collection1->DynamicCollection;
 		TSharedPtr<FGeometryDynamicCollection> DynamicCollection2 = Collection2->DynamicCollection;
-		DynamicCollection1->ModifyAttribute<int32>("DynamicState", FGeometryCollection::TransformGroup)[1] = (uint8)EObjectStateTypeEnum::Chaos_Object_Kinematic;
+		DynamicCollection1->ModifyAttribute<uint8>("DynamicState", FGeometryCollection::TransformGroup)[1] = (uint8)EObjectStateTypeEnum::Chaos_Object_Kinematic;
 
 		UnitTest.Initialize();
 
@@ -2051,7 +2051,7 @@ namespace GeometryCollectionTest
 		Params.ClusterGroupIndex = 0;		
 		FGeometryCollectionWrapper* Collection1 = TNewSimulationObject<GeometryType::GeometryCollectionWithSuppliedRestCollection>::Init(Params)->template As<FGeometryCollectionWrapper>();
 		TSharedPtr<FGeometryDynamicCollection> DynamicCollection1 = Collection1->DynamicCollection;
-		DynamicCollection1->ModifyAttribute<int32>("DynamicState", FGeometryCollection::TransformGroup)[1] = (uint8)EObjectStateTypeEnum::Chaos_Object_Kinematic;
+		DynamicCollection1->ModifyAttribute<uint8>("DynamicState", FGeometryCollection::TransformGroup)[1] = (uint8)EObjectStateTypeEnum::Chaos_Object_Kinematic;
 
 		UnitTest.AddSimulationObject(Collection1);
 
@@ -2119,7 +2119,7 @@ namespace GeometryCollectionTest
 		FGeometryCollectionWrapper* Collection = TNewSimulationObject<GeometryType::GeometryCollectionWithSuppliedRestCollection>::Init(Params)->template As<FGeometryCollectionWrapper>();
 
 		TSharedPtr<FGeometryDynamicCollection> DynamicCollection = Collection->DynamicCollection;
-		DynamicCollection->ModifyAttribute<int32>("DynamicState", FGeometryCollection::TransformGroup)[1] = (uint8)EObjectStateTypeEnum::Chaos_Object_Kinematic;
+		DynamicCollection->ModifyAttribute<uint8>("DynamicState", FGeometryCollection::TransformGroup)[1] = (uint8)EObjectStateTypeEnum::Chaos_Object_Kinematic;
 
 		UnitTest.AddSimulationObject(Collection);
 		UnitTest.Initialize();
diff --git a/Engine/Source/Programs/HeadlessChaos/Private/GeometryCollection/GeometryCollectionTestCollisionResolution.cpp b/Engine/Source/Programs/HeadlessChaos/Private/GeometryCollection/GeometryCollectionTestCollisionResolution.cpp
--- a/Engine/Source/Programs/HeadlessChaos/Private/GeometryCollection/GeometryCollectionTestCollisionResolution.cpp
+++ b/Engine/Source/Programs/HeadlessChaos/Private/GeometryCollection/GeometryCollectionTestCollisionResolution.cpp
@@ -333,11 +333,13 @@ namespace GeometryCollectionTest
 			EXPECT_TRUE(UnitTest.Solver->GetParticles().GetGeometryCollectionParticles().CollisionParticles(0) != nullptr);
 			EXPECT_TRUE(Collection->DynamicCollection->Simplicials[0]->Size() == UnitTest.Solver->GetParticles().GetGeometryCollectionParticles().CollisionParticles(0)->Size());
 			EXPECT_TRUE(Collection->DynamicCollection->Simplicials[0]->Size() != 0);
-			EXPECT_TRUE(Collection->DynamicCollection->Implicits[0]->GetType() == (int32)Chaos::ImplicitObjectType::LevelSet);
+			// The following test has been disabled because from now we remove the Implicits after initialization to free up some memory
+			// EXPECT_TRUE(Collection->DynamicCollection->GetAttribute<Chaos::FImplicitObjectPtr>(FGeometryDynamicCollection::ImplicitsAttribute, FTransformCollection::TransformGroup)[0]->GetType() == (int32)Chaos::ImplicitObjectType::LevelSet);
 
 			EXPECT_EQ(CollectionStaticSphere->DynamicCollection->GetTransforms().Num(), 4);
 			EXPECT_TRUE(CollectionStaticSphere->DynamicCollection->Simplicials[0] == nullptr);
-			EXPECT_TRUE(CollectionStaticSphere->DynamicCollection->Implicits[0]->GetType() == (int32)Chaos::ImplicitObjectType::Sphere);
+			// The following test has been disabled because from now we remove the Implicits after initialization to free up some memory
+			// EXPECT_TRUE(CollectionStaticSphere->DynamicCollection->GetAttribute<Chaos::FImplicitObjectPtr>(FGeometryDynamicCollection::ImplicitsAttribute, FTransformCollection::TransformGroup)[0]->GetType() == (int32)Chaos::ImplicitObjectType::Sphere);
 
 			// validate the ball collides and moved away from the static ball
 			EXPECT_LT(FMath::Abs(Collection->RestCollection->Transform[0].GetTranslation().Z) - 10.f, KINDA_SMALL_NUMBER);
diff --git a/Engine/Source/Programs/HeadlessChaos/Private/GeometryCollection/GeometryCollectionTestSimulationField.cpp b/Engine/Source/Programs/HeadlessChaos/Private/GeometryCollection/GeometryCollectionTestSimulationField.cpp
--- a/Engine/Source/Programs/HeadlessChaos/Private/GeometryCollection/GeometryCollectionTestSimulationField.cpp
+++ b/Engine/Source/Programs/HeadlessChaos/Private/GeometryCollection/GeometryCollectionTestSimulationField.cpp
@@ -86,7 +86,7 @@ namespace GeometryCollectionTest
 			UnitTest.Advance();
 		}
 
-		TManagedArray<int32>& DynamicState = Collection->DynamicCollection->DynamicState;
+		TManagedArray<uint8>& DynamicState = Collection->DynamicCollection->DynamicState;
 		EXPECT_EQ(DynamicState[0], (int32)EObjectStateTypeEnum::Chaos_Object_Kinematic);
 
 		// simulated
@@ -182,7 +182,7 @@ namespace GeometryCollectionTest
 
 		UnitTest.Initialize();
 
-		TManagedArray<int32>& DynamicState = Collection->DynamicCollection->DynamicState;
+		TManagedArray<uint8>& DynamicState = Collection->DynamicCollection->DynamicState;
 
 		FReal PreviousY = 0.f;
 		EXPECT_EQ(Collection->DynamicCollection->GetTransform(0).GetTranslation().X, 0);
diff --git a/Engine/Source/Runtime/Engine/Private/PhysicsEngine/Experimental/PhysScene_Chaos.cpp b/Engine/Source/Runtime/Engine/Private/PhysicsEngine/Experimental/PhysScene_Chaos.cpp
--- a/Engine/Source/Runtime/Engine/Private/PhysicsEngine/Experimental/PhysScene_Chaos.cpp
+++ b/Engine/Source/Runtime/Engine/Private/PhysicsEngine/Experimental/PhysScene_Chaos.cpp
@@ -2175,7 +2175,7 @@ static void UpdateAccelerationStructureFromGeometryCollectionProxy(FGeometryColl
 		return Chaos::FAABB3(Particle.X(), Particle.X());
 	};
 
-	const TManagedArray<TUniquePtr<Chaos::FPBDRigidParticle>>& GTParticles = Proxy.GetExternalParticles();
+	const TArray<TUniquePtr<Chaos::FPBDRigidParticle>>& GTParticles = Proxy.GetExternalParticles();
 	if (bIsParentProxyNull)
 	{
 		for (const TUniquePtr<Chaos::FPBDRigidParticle>& GTParticle : GTParticles)
diff --git a/Engine/Source/Runtime/Experimental/Chaos/Private/GeometryCollectionProxyData.cpp b/Engine/Source/Runtime/Experimental/Chaos/Private/GeometryCollectionProxyData.cpp
--- a/Engine/Source/Runtime/Experimental/Chaos/Private/GeometryCollectionProxyData.cpp
+++ b/Engine/Source/Runtime/Experimental/Chaos/Private/GeometryCollectionProxyData.cpp
@@ -101,8 +101,7 @@ FGeometryDynamicCollection::FGeometryDynamicCollection(const FGeometryCollection
 {
 	// Transform Group
 	AddExternalAttribute<bool>(FGeometryDynamicCollection::ActiveAttribute, FTransformCollection::TransformGroup, Active);
-	AddExternalAttribute<int32>(FGeometryDynamicCollection::DynamicStateAttribute, FTransformCollection::TransformGroup, DynamicState);
-	AddExternalAttribute(ImplicitsAttribute, FTransformCollection::TransformGroup, Implicits);
+	AddExternalAttribute<uint8>(FGeometryDynamicCollection::DynamicStateAttribute, FTransformCollection::TransformGroup, DynamicState);
 	AddExternalAttribute(SimplicialsAttribute, FTransformCollection::TransformGroup, Simplicials);
 	AddExternalAttribute(SimulatableParticlesAttribute, FGeometryCollection::TransformGroup, SimulatableParticles);
 
diff --git a/Engine/Source/Runtime/Experimental/Chaos/Private/PhysicsProxy/GeometryCollectionPhysicsProxy.cpp b/Engine/Source/Runtime/Experimental/Chaos/Private/PhysicsProxy/GeometryCollectionPhysicsProxy.cpp
--- a/Engine/Source/Runtime/Experimental/Chaos/Private/PhysicsProxy/GeometryCollectionPhysicsProxy.cpp
+++ b/Engine/Source/Runtime/Experimental/Chaos/Private/PhysicsProxy/GeometryCollectionPhysicsProxy.cpp
@@ -502,17 +502,12 @@ void FGeometryCollectionPhysicsProxy::Initialize(Chaos::FPBDRigidsEvolutionBase
 
 	InitializeDynamicCollection(DynamicCollection, *Parameters.RestCollection, Parameters);
 
-	// Attach the external particles to the gamethread collection
-	if (DynamicCollection.HasAttribute(FGeometryCollection::ParticlesAttribute, FTransformCollection::TransformGroup))
-		DynamicCollection.RemoveAttribute(FGeometryCollection::ParticlesAttribute, FTransformCollection::TransformGroup);
-	DynamicCollection.AddExternalAttribute<TUniquePtr<FParticle>>(FGeometryCollection::ParticlesAttribute, FTransformCollection::TransformGroup, GTParticles);
-
-
 	NumParticles = DynamicCollection.NumElements(FGeometryCollection::TransformGroup);
 	BaseParticleIndex = 0; // Are we always zero indexed now?
 	SolverClusterID.Init(nullptr, NumParticles);
 	SolverClusterHandles.Init(nullptr, NumParticles);
 	SolverParticleHandles.Init(nullptr, NumParticles);
+	GTParticles.SetNum(NumParticles);
 	
 	TBitArray<> EffectiveParticles;
 	NumEffectiveParticles = CalculateEffectiveParticles(DynamicCollection, NumParticles, Parameters.MaxSimulatedLevel, Parameters.EnableClustering, GetOwner(), EffectiveParticles);
@@ -556,14 +551,6 @@ void FGeometryCollectionPhysicsProxy::Initialize(Chaos::FPBDRigidsEvolutionBase
 	// Initialise GT/External particles
 	const int32 NumTransforms = DynamicCollection.Transform.Num();
 
-	// Attach the external particles to the gamethread collection
-	if (GameThreadCollection.HasAttribute(FGeometryCollection::ParticlesAttribute, FTransformCollection::TransformGroup))
-	{ 
-		GameThreadCollection.RemoveAttribute(FGeometryCollection::ParticlesAttribute, FTransformCollection::TransformGroup);
-	}
-		
-	GameThreadCollection.AddExternalAttribute<TUniquePtr<FParticle>>(FGeometryCollection::ParticlesAttribute, FTransformCollection::TransformGroup, GTParticles);
-
 	// make sure we copy the anchored information over to the physics thread collection
 	const Chaos::Facades::FCollectionAnchoringFacade DynamicCollectionAnchoringFacade(DynamicCollection);
 	Chaos::Facades::FCollectionAnchoringFacade PhysicsThreadCollectionAnchoringFacade(PhysicsThreadCollection);
@@ -574,7 +561,8 @@ void FGeometryCollectionPhysicsProxy::Initialize(Chaos::FPBDRigidsEvolutionBase
 	const TManagedArray<FTransform>& MassToLocal = Parameters.RestCollection->GetAttribute<FTransform>(MassToLocalAttributeName, FTransformCollection::TransformGroup);
 
 	TArray<int32> ChildrenToCheckForParentFix;
-	if(ensure(NumTransforms == GameThreadCollection.Implicits.Num() && NumTransforms == GTParticles.Num())) // Implicits are in the transform group so this invariant should always hold
+	TManagedArray<Chaos::FImplicitObjectPtr>& Implicits = GameThreadCollection.ModifyAttribute<Chaos::FImplicitObjectPtr>(FGeometryDynamicCollection::ImplicitsAttribute, FTransformCollection::TransformGroup);
+	if(ensure(NumTransforms == Implicits.Num() && NumTransforms == GTParticles.Num())) // Implicits are in the transform group so this invariant should always hold
 	{
 		for(int32 Index = 0; Index < NumTransforms; ++Index)
 		{
@@ -600,7 +588,7 @@ void FGeometryCollectionPhysicsProxy::Initialize(Chaos::FPBDRigidsEvolutionBase
 				P->SetUserData(Parameters.UserData);
 				P->SetProxy(this);
 
-				Chaos::FImplicitObjectPtr ImplicitGeometry = GameThreadCollection.Implicits[Index];
+				Chaos::FImplicitObjectPtr ImplicitGeometry = Implicits[Index];
 				if (ImplicitGeometry && !Scale.Equals(FVector::OneVector))
 				{
 					ImplicitGeometry = ImplicitGeometry->CopyGeometryWithScale(Scale);
@@ -657,7 +645,7 @@ void FGeometryCollectionPhysicsProxy::Initialize(Chaos::FPBDRigidsEvolutionBase
 				// step 2: fix the parent if necessary
 				for (const int32 ParentToFixIndex: ParentToPotentiallyFix)
 				{
-					if (GameThreadCollection.Implicits[ParentToFixIndex] == nullptr)
+					if (Implicits[ParentToFixIndex] == nullptr)
 					{
 						const Chaos::FRigidTransform3 ParentShapeTransform =  MassToLocal[ParentToFixIndex];
 				
@@ -667,7 +655,7 @@ void FGeometryCollectionPhysicsProxy::Initialize(Chaos::FPBDRigidsEvolutionBase
 						GameThreadCollection.IterateThroughChildren(ParentToFixIndex, [&](int32 ChildIndex)
 						{
 							// defer if any of the children is a cluster with no collision yet generated 
-							if (GameThreadCollection.Implicits[ChildIndex] == nullptr && GameThreadCollection.HasChildren(ChildIndex))
+							if (Implicits[ChildIndex] == nullptr && GameThreadCollection.HasChildren(ChildIndex))
 							{
 								bAllChildrenHaveCollision = false;
 								return false;
@@ -681,7 +669,7 @@ void FGeometryCollectionPhysicsProxy::Initialize(Chaos::FPBDRigidsEvolutionBase
 							TArray<Chaos::FImplicitObjectPtr> ChildImplicits;
 							GameThreadCollection.IterateThroughChildren(ParentToFixIndex, [&](int32 ChildIndex)
 							{
-								const Chaos::FImplicitObjectPtr& ChildImplicit = GameThreadCollection.Implicits[ChildIndex];
+								const Chaos::FImplicitObjectPtr& ChildImplicit = Implicits[ChildIndex];
 								if (ChildImplicit)
 								{
 									const Chaos::FRigidTransform3 ChildShapeTransform = MassToLocal[ChildIndex] * FTransform(GameThreadCollection.Transform[ChildIndex]);
@@ -706,11 +694,11 @@ void FGeometryCollectionPhysicsProxy::Initialize(Chaos::FPBDRigidsEvolutionBase
 							if (ChildImplicits.Num() > 0)
 							{
 								Chaos::FImplicitObject* UnionImplicit = new Chaos::FImplicitObjectUnion(MoveTemp(ChildImplicits));
-								GameThreadCollection.Implicits[ParentToFixIndex] = Chaos::FImplicitObjectPtr(UnionImplicit);
+								Implicits[ParentToFixIndex] = Chaos::FImplicitObjectPtr(UnionImplicit);
 							}
 							if (GTParticles[ParentToFixIndex] != nullptr)
 							{
-								GTParticles[ParentToFixIndex]->SetGeometry(GameThreadCollection.Implicits[ParentToFixIndex]);
+								GTParticles[ParentToFixIndex]->SetGeometry(Implicits[ParentToFixIndex]);
 							}
 						}
 					}
@@ -749,6 +737,8 @@ void FGeometryCollectionPhysicsProxy::Initialize(Chaos::FPBDRigidsEvolutionBase
 	};
 	static constexpr int32 SkipListSize = sizeof(SkipList) / sizeof(FAttributeAndGroupId);
 
+	// Adding the Implicits to PhysicsThreadCollection, before it copies all matching attributes
+	PhysicsThreadCollection.AddAttribute<Chaos::FImplicitObjectPtr>(FGeometryDynamicCollection::ImplicitsAttribute, FTransformCollection::TransformGroup);
 	PhysicsThreadCollection.CopyMatchingAttributesFrom(DynamicCollection, MakeArrayView(SkipList, SkipListSize));
 	PhysicsThreadCollection.CopyInitialVelocityAttributesFrom(DynamicCollection);
 
@@ -795,6 +785,9 @@ void FGeometryCollectionPhysicsProxy::Initialize(Chaos::FPBDRigidsEvolutionBase
 			}
 		}
 	}
+
+	// The Implicits attributes from the Dynamic Collection are just used for initialization, after they can be removed and so free some memory.
+	GameThreadCollection.RemoveAttribute(FGeometryDynamicCollection::ImplicitsAttribute, FTransformCollection::TransformGroup);
 }
 
 
@@ -817,6 +810,7 @@ void FGeometryCollectionPhysicsProxy::InitializeDynamicCollection(FGeometryDynam
 	};
 	static const int32 SkipListSize = sizeof(SkipList) / sizeof(FAttributeAndGroupId);
 
+	DynamicCollection.AddAttribute<Chaos::FImplicitObjectPtr>(FGeometryDynamicCollection::ImplicitsAttribute, FTransformCollection::TransformGroup);
 	DynamicCollection.CopyMatchingAttributesFrom(RestCollection, MakeArrayView(SkipList, SkipListSize));
 
 	// User defined initial velocities need to be populated. 
@@ -1059,10 +1053,10 @@ void FGeometryCollectionPhysicsProxy::InitializeBodiesPT(Chaos::FPBDRigidsSolver
 		const TManagedArray<FVector3f>& InertiaTensor = RestCollection->GetAttribute<FVector3f>(InertiaTensorAttributeName, FTransformCollection::TransformGroup);
 
 		const int32 NumTransforms = DynamicCollection.NumElements(FTransformCollection::TransformGroup);
-		const TManagedArray<int32>& DynamicState = DynamicCollection.DynamicState;
+		const TManagedArray<uint8>& DynamicState = DynamicCollection.DynamicState;
 		const TManagedArray<bool>& SimulatableParticles = DynamicCollection.SimulatableParticles;
 		const TManagedArray<FTransform>& MassToLocal = RestCollection->GetAttribute<FTransform>(MassToLocalAttributeName, FTransformCollection::TransformGroup);
-		const TManagedArray<Chaos::FImplicitObjectPtr>& Implicits = DynamicCollection.Implicits;
+		const TManagedArray<Chaos::FImplicitObjectPtr>& Implicits = DynamicCollection.GetAttribute<Chaos::FImplicitObjectPtr>(FGeometryDynamicCollection::ImplicitsAttribute, FTransformCollection::TransformGroup);
 		const TManagedArray<TUniquePtr<FCollisionStructureManager::FSimplicial>>& Simplicials = DynamicCollection.Simplicials;
 
 		// In PushToPhysicsState, we're going to compute a relative transform from Parameters.PrevWorldTransform
@@ -1493,6 +1487,8 @@ void FGeometryCollectionPhysicsProxy::InitializeBodiesPT(Chaos::FPBDRigidsSolver
 
 	} // end if simulating...
 
+	// The Implicits attributes from the Dynamic Collection are just used for initialization, after they can be removed and so free some memory.
+	PhysicsThreadCollection.RemoveAttribute(FGeometryDynamicCollection::ImplicitsAttribute, FTransformCollection::TransformGroup);
 }
 
 int32 ReportNoLevelsetCluster = 0;
@@ -1583,11 +1579,11 @@ float FGeometryCollectionPhysicsProxy::ComputeUserDefinedDamageThreshold_Interna
 		{
 			// bounding box volume is used as a fallback to find specific size if the relative size if not available
 			// ( May happen with older GC )
+			const FClusterHandle* Handle = SolverParticleHandles[TransformIndex];
 			FBox LocalBoundingBox;
-			const Chaos::FImplicitObjectPtr& Implicit = PhysicsThreadCollection.Implicits[TransformIndex];
-			if (Implicit && Implicit->HasBoundingBox())
+			if (Handle && Handle->HasBounds())
 			{
-				const Chaos::FAABB3& ImplicitBoundingBox = Implicit->BoundingBox();
+				const Chaos::FAABB3& ImplicitBoundingBox = Handle->LocalBounds();
 				LocalBoundingBox = FBox(ImplicitBoundingBox.Min(), ImplicitBoundingBox.Max());
 			}
 
@@ -1631,10 +1627,10 @@ Chaos::FVec3f FGeometryCollectionPhysicsProxy::AdjustInertiaForScale(const Chaos
 Chaos::TPBDGeometryCollectionParticleHandle<Chaos::FReal, 3>* FGeometryCollectionPhysicsProxy::BuildNonClusters_Internal(const uint32 CollectionClusterIndex, Chaos::FPBDRigidsSolver* RigidsSolver, float Mass, Chaos::FVec3f Inertia)
 {
 	FGeometryDynamicCollection& DynamicCollection = PhysicsThreadCollection;
-	TManagedArray<int32>& DynamicState = DynamicCollection.DynamicState;
+	TManagedArray<uint8>& DynamicState = DynamicCollection.DynamicState;
 	TManagedArray<FTransform3f>& Transform = DynamicCollection.Transform;
 	const TManagedArray<FTransform>& MassToLocal = Parameters.RestCollection->GetAttribute<FTransform>(MassToLocalAttributeName, FTransformCollection::TransformGroup);
-	const TManagedArray<Chaos::FImplicitObjectPtr>& Implicits = DynamicCollection.Implicits;
+	const TManagedArray<Chaos::FImplicitObjectPtr>& Implicits = DynamicCollection.GetAttribute<Chaos::FImplicitObjectPtr>(FGeometryDynamicCollection::ImplicitsAttribute, FTransformCollection::TransformGroup);
 	const TManagedArray<TUniquePtr<FCollisionStructureManager::FSimplicial>>& Simplicials = DynamicCollection.Simplicials;
 	Chaos::Facades::FCollectionAnchoringFacade AnchoringFacade(DynamicCollection);
 
@@ -1749,11 +1745,11 @@ FGeometryCollectionPhysicsProxy::BuildClusters_Internal(
 	check(ChildHandles.Num() != 0);
 
 	FGeometryDynamicCollection& DynamicCollection = PhysicsThreadCollection;
-	TManagedArray<int32>& DynamicState = DynamicCollection.DynamicState;
+	TManagedArray<uint8>& DynamicState = DynamicCollection.DynamicState;
 	TManagedArray<FTransform3f>& Transform = DynamicCollection.Transform;
 	const TManagedArray<FTransform>& MassToLocal = Parameters.RestCollection->GetAttribute<FTransform>(MassToLocalAttributeName, FTransformCollection::TransformGroup);
 	//TManagedArray<TSharedPtr<FCollisionStructureManager::FSimplicial> >& Simplicials = DynamicCollection.Simplicials;
-	TManagedArray<Chaos::FImplicitObjectPtr>& Implicits = DynamicCollection.Implicits;
+	const TManagedArray<Chaos::FImplicitObjectPtr>& Implicits = DynamicCollection.GetAttribute<Chaos::FImplicitObjectPtr>(FGeometryDynamicCollection::ImplicitsAttribute, FTransformCollection::TransformGroup);
 
 	//If we are a root particle use the world transform, otherwise set the relative transform
 	const FTransform CollectionSpaceTransform = GeometryCollectionAlgo::GlobalMatrix(Transform, DynamicCollection, CollectionClusterIndex);
@@ -2661,7 +2657,8 @@ TBitArray<> FGeometryCollectionPhysicsProxy::CalculateClustersToCreateFromChildr
 			// step 2: test the parent for having children with geometry
 			for (const int32 ParentToFixIndex : ParentToPotentiallyFix)
 			{
-				const bool bParentHasCollision = (DynamicCollection.Implicits[ParentToFixIndex] != nullptr) || ClustersToGenerate[ParentToFixIndex];
+				const bool bParentHasCollision = 
+					(DynamicCollection.GetAttribute<Chaos::FImplicitObjectPtr>(FGeometryDynamicCollection::ImplicitsAttribute, FTransformCollection::TransformGroup)[ParentToFixIndex] != nullptr) || ClustersToGenerate[ParentToFixIndex];
 				if (!bParentHasCollision)
 				{
 					// let's make sure all our children have an implicit defined, otherwise, postpone to next iteration 
@@ -2669,7 +2666,8 @@ TBitArray<> FGeometryCollectionPhysicsProxy::CalculateClustersToCreateFromChildr
 					DynamicCollection.IterateThroughChildren(ParentToFixIndex, [&](int32 ChildIndex)
 					{
 						// defer if any of the children is a cluster with no collision yet generated 
-						const bool bChildHasCollision = (DynamicCollection.Implicits[ChildIndex] != nullptr) || ClustersToGenerate[ChildIndex];
+						const bool bChildHasCollision = 
+							(DynamicCollection.GetAttribute<Chaos::FImplicitObjectPtr>(FGeometryDynamicCollection::ImplicitsAttribute, FTransformCollection::TransformGroup)[ChildIndex] != nullptr) || ClustersToGenerate[ChildIndex];
 						if (!bChildHasCollision && DynamicCollection.HasChildren(ChildIndex))
 						{
 							bAllChildrenHaveCollision = false;
@@ -2705,7 +2703,7 @@ int32 FGeometryCollectionPhysicsProxy::CalculateEffectiveParticles(const FGeomet
 		if (Level <= MaxSimulatedLevel || !bEnableClustering)
 		{
 			const bool bIsClusterUsingChildGeometry = (ClustersUsingChildGeometry.Num() > 0) && ClustersUsingChildGeometry[TransformIndex];
-			const bool bHasGeometry = DynamicCollection.Implicits[TransformIndex].IsValid();
+			const bool bHasGeometry = DynamicCollection.GetAttribute<Chaos::FImplicitObjectPtr>(FGeometryDynamicCollection::ImplicitsAttribute, FTransformCollection::TransformGroup)[TransformIndex].IsValid();
 
 			if (bHasGeometry || bIsClusterUsingChildGeometry)
 			{
@@ -3739,7 +3737,7 @@ bool FGeometryCollectionPhysicsProxy::PullNonInterpolatableDataFromSinglePhysics
 			bIsCollectionDirty = true;
 		}
 
-		if (UpdateValue(GameThreadCollection.DynamicState[TransformGroupIndex], StateData.State.DynamicState))
+		if (UpdateValue(GameThreadCollection.DynamicState[TransformGroupIndex], static_cast<uint8>(StateData.State.DynamicState)))
 		{
 			GTParticle.SetObjectState(static_cast<Chaos::EObjectStateType>(StateData.State.DynamicState));
 			bIsCollectionDirty = true;
diff --git a/Engine/Source/Runtime/Experimental/Chaos/Public/GeometryCollectionProxyData.h b/Engine/Source/Runtime/Experimental/Chaos/Public/GeometryCollectionProxyData.h
--- a/Engine/Source/Runtime/Experimental/Chaos/Public/GeometryCollectionProxyData.h
+++ b/Engine/Source/Runtime/Experimental/Chaos/Public/GeometryCollectionProxyData.h
@@ -106,8 +106,9 @@ public:
 	// Transform Group
 	TManagedArray<bool> Active;
 	
-	TManagedArray<int32> DynamicState;
-	TManagedArray<Chaos::FImplicitObjectPtr> Implicits;
+	TManagedArray<uint8> DynamicState; 
+	static_assert(sizeof(EObjectStateTypeEnum) <= sizeof(uint8)); // DynamicState must fit  EObjectStateTypeEnum
+
 	TManagedArray<TUniquePtr<FCollisionStructureManager::FSimplicial>> Simplicials;
 	TManagedArray<bool> SimulatableParticles;
 
diff --git a/Engine/Source/Runtime/Experimental/Chaos/Public/PhysicsProxy/GeometryCollectionPhysicsProxy.h b/Engine/Source/Runtime/Experimental/Chaos/Public/PhysicsProxy/GeometryCollectionPhysicsProxy.h
--- a/Engine/Source/Runtime/Experimental/Chaos/Public/PhysicsProxy/GeometryCollectionPhysicsProxy.h
+++ b/Engine/Source/Runtime/Experimental/Chaos/Public/PhysicsProxy/GeometryCollectionPhysicsProxy.h
@@ -297,7 +297,7 @@ public:
 		return GameThreadCollection;
 	}
 
-	TManagedArray<TUniquePtr<FParticle>>& GetExternalParticles()
+	TArray<TUniquePtr<FParticle>>& GetExternalParticles()
 	{
 		return GTParticles;
 	}
@@ -587,7 +587,7 @@ private:
 
 	EReplicationMode ReplicationMode = EReplicationMode::Unknown;	
 
-	TManagedArray<TUniquePtr<FParticle>> GTParticles;
+	TArray<TUniquePtr<FParticle>> GTParticles;
 	TMap<FParticle*, int32> GTParticlesToTransformGroupIndex;
 	TMap<FParticle*, int32> GTParticlesToInternalClusterUniqueIdx;
 	TMap<int32, TArray<int32>> InternalClusterUniqueIdxToChildrenTransformIndices;
diff --git a/Engine/Source/Runtime/Experimental/GeometryCollectionEngine/Private/GeometryCollection/GeometryCollectionComponent.cpp b/Engine/Source/Runtime/Experimental/GeometryCollectionEngine/Private/GeometryCollection/GeometryCollectionComponent.cpp
--- a/Engine/Source/Runtime/Experimental/GeometryCollectionEngine/Private/GeometryCollection/GeometryCollectionComponent.cpp
+++ b/Engine/Source/Runtime/Experimental/GeometryCollectionEngine/Private/GeometryCollection/GeometryCollectionComponent.cpp
@@ -2918,10 +2918,10 @@ void UGeometryCollectionComponent::SetDynamicState(const Chaos::EObjectStateType
 {
 	if (DynamicCollection)
 	{
-		TManagedArray<int32>& DynamicState = DynamicCollection->DynamicState;
+		TManagedArray<uint8>& DynamicState = DynamicCollection->DynamicState;
 		for (int i = 0; i < DynamicState.Num(); i++)
 		{
-			DynamicState[i] = static_cast<int32>(NewDynamicState);
+			DynamicState[i] = static_cast<uint8>(NewDynamicState);
 		}
 	}
 }
@@ -3589,7 +3589,7 @@ void UGeometryCollectionComponent::OnCreatePhysicsState()
 			FChaosUserData::Set<UPrimitiveComponent>(&PhysicsUserData, this);
 
 			// If the Component is set to Dynamic, we look to the RestCollection for initial dynamic state override per transform.
-			TManagedArray<int32> & DynamicState = DynamicCollection->DynamicState;
+			TManagedArray<uint8> & DynamicState = DynamicCollection->DynamicState;
 
 			// if this code is changed you may need to account for bStartAwake
 			EObjectStateTypeEnum LocalObjectType = (ObjectType != EObjectStateTypeEnum::Chaos_Object_Sleeping) ? ObjectType : EObjectStateTypeEnum::Chaos_Object_Dynamic;
@@ -3600,14 +3600,14 @@ void UGeometryCollectionComponent::OnCreatePhysicsState()
 					TManagedArray<int32>& InitialDynamicState = RestCollection->GetGeometryCollection()->InitialDynamicState;
 					for (int i = 0; i < DynamicState.Num(); i++)
 					{
-						DynamicState[i] = (InitialDynamicState[i] == static_cast<int32>(Chaos::EObjectStateType::Uninitialized)) ? static_cast<int32>(LocalObjectType) : InitialDynamicState[i];
+						DynamicState[i] = (InitialDynamicState[i] == static_cast<int32>(Chaos::EObjectStateType::Uninitialized)) ? static_cast<uint8>(LocalObjectType) : InitialDynamicState[i];
 					}
 				}
 				else
 				{
 					for (int i = 0; i < DynamicState.Num(); i++)
 					{
-						DynamicState[i] = static_cast<int32>(LocalObjectType);
+						DynamicState[i] = static_cast<uint8>(LocalObjectType);
 					}
 				}
 			}
