IAS' cache service thread was waiting for too long.

#rnx

diff --git a/Engine/Source/Runtime/Experimental/IoStoreOnDemand/Private/JournaledCache.cpp b/Engine/Source/Runtime/Experimental/IoStoreOnDemand/Private/JournaledCache.cpp
--- a/Engine/Source/Runtime/Experimental/IoStoreOnDemand/Private/JournaledCache.cpp
+++ b/Engine/Source/Runtime/Experimental/IoStoreOnDemand/Private/JournaledCache.cpp
@@ -1349,16 +1349,24 @@ int32 FGovernor::BeginInternal(uint32 Demand, int64 Cycles)
 	Interval <<= int32(Demand <= DemandThreshold);
 
 	int64 Delta = Cycles - PrevCycles;
-	if (Delta <= Interval)
+	bool bNotYet = (Delta < Interval);
+
+	// Calculate how much time we are into the shortest poll interval
+	int64 Remainder = Delta;
+	Interval = GetMaxWaitCycles();
+	for (; Remainder > Interval; Remainder -= Interval);
+
+	if (bNotYet)
 	{
-		return -int32(Interval - Delta);
+		// We haven't hit the current interval length but might be drawn in if
+		// demand increases. So we return a wait that takes us to that.
+		return int32(Remainder - Interval);
 	}
 
-	do { Delta -= Interval; } while (Delta > Interval);
-	PrevCycles = Cycles - Delta;
+	// PrevCycles is adjusted so we do not lose any left over time
+	PrevCycles = Cycles - Remainder;
 
 	OpCount++;
-
 	return OpAllowance + RunOff;
 }
 
@@ -1695,16 +1703,16 @@ int32 FServiceThread::Update()
 		CycleSlice = FMath::Min(CyclesTillActive, CycleSlice);
 	}
 
-	// Now we've a slice of time to process reads until caches need another tick
-	int64 Cycle = FPlatformTime::Cycles64();
-	int64 StopReadsCycle = Cycle + CycleSlice;
-
 	// Early out
 	if (ActiveReads.Num() == 0)
 	{
 		return CycleSlice;
 	}
 
+	// Now we've a slice of time to process reads until caches need another tick
+	int64 Cycle = FPlatformTime::Cycles64();
+	int64 StopReadsCycle = Cycle + CycleSlice;
+
 	TRACE_CPUPROFILER_EVENT_SCOPE(IasCache::ProcessReads);
 
 	uint32 Index = 0;
@@ -1738,7 +1746,9 @@ int32 FServiceThread::Update()
 	check(Index > 0);
 	ActiveReads.RemoveAt(0, Index);
 
-	return CycleSlice - int32(Cycle - StopReadsCycle);
+	// StopReadsCycle is where the CycleSlice would expire, while Cycle is where
+	// in time we have got to. The difference is how much we need to wait.
+	return int32(StopReadsCycle - Cycle);
 }
 
 ////////////////////////////////////////////////////////////////////////////////
