Fix removal of geometry collection pieces not working ( fallout from CL 28427707)
- made sure FGeometryCollectionDynamicStateFacade use the right type for the dynamic state attribute

diff --git a/Engine/Source/Runtime/Experimental/Chaos/Public/GeometryCollectionProxyData.h b/Engine/Source/Runtime/Experimental/Chaos/Public/GeometryCollectionProxyData.h
--- a/Engine/Source/Runtime/Experimental/Chaos/Public/GeometryCollectionProxyData.h
+++ b/Engine/Source/Runtime/Experimental/Chaos/Public/GeometryCollectionProxyData.h
@@ -187,7 +187,7 @@ private:
 	TManagedArrayAccessor<bool> ActiveAttribute;
 
 	/** physics state of the transform (Dynamic, kinematic, static, sleeping) */
-	TManagedArrayAccessor<int32> DynamicStateAttribute;
+	TManagedArrayAccessor<uint8> DynamicStateAttribute;
 
 	/** type of internal state parent */
 	TManagedArrayAccessor<uint8> InternalClusterParentTypeAttribute;
