Add a cvar 'ias.SuspendSystem' that when enabled will prevent the IoStoreOnDemand module from initializing, preventing it's use. Defaults to false.

#rnx

- NOTE: This cvar is temp, to be used while the system is in testing and is unlikely to be in shipping version of UE 5.4 and should not be relied upon.
-FIoStoreOnDemandModule::Initialize now returns a new enum 'EOnDemandInitResult' describing the result of initialization.
- Add logging when calls to FIoStoreOnDemandModule::SetBulkOptionalEnabled/::SetEnabled/::AbandonCache are deferred.

diff --git a/Engine/Source/Runtime/Experimental/IoStoreOnDemand/Private/IoStoreOnDemand.cpp b/Engine/Source/Runtime/Experimental/IoStoreOnDemand/Private/IoStoreOnDemand.cpp
--- a/Engine/Source/Runtime/Experimental/IoStoreOnDemand/Private/IoStoreOnDemand.cpp
+++ b/Engine/Source/Runtime/Experimental/IoStoreOnDemand/Private/IoStoreOnDemand.cpp
@@ -42,6 +42,14 @@ DEFINE_LOG_CATEGORY(LogIas);
 namespace UE::IO::IAS
 {
 
+bool GIasSuspendSystem = false;
+static FAutoConsoleVariableRef CVar_SuspendSystemEnabled(
+	TEXT("ias.SuspendSystem"),
+	GIasSuspendSystem,
+	TEXT("Suspends the use of the OnDemand system"),
+	ECVF_ReadOnly
+);
+
 ////////////////////////////////////////////////////////////////////////////////
 static int64 ParseSizeParam(FStringView Value)
 {
@@ -1493,6 +1501,7 @@ void FIoStoreOnDemandModule::SetBulkOptionalEnabled(bool bInEnabled)
 	}
 	else
 	{
+		UE_LOG(LogIas, Log, TEXT("Deferring call to FIoStoreOnDemandModule::SetBulkOptionalEnabled(%s)"), bInEnabled ? TEXT("true") : TEXT("false"));
 		DeferredBulkOptionalEnabled = bInEnabled;
 	}
 }
@@ -1505,6 +1514,7 @@ void FIoStoreOnDemandModule::SetEnabled(bool bInEnabled)
 	}
 	else
 	{
+		UE_LOG(LogIas, Log, TEXT("Deferring call to FIoStoreOnDemandModule::SetEnabled(%s)"), bInEnabled ? TEXT("true") : TEXT("false"));
 		DeferredEnabled = bInEnabled;
 	}
 }
@@ -1517,6 +1527,7 @@ void FIoStoreOnDemandModule::AbandonCache()
 	}
 	else
 	{
+		UE_LOG(LogIas, Log, TEXT("Deferring call to FIoStoreOnDemandModule::AbandonCache"));
 		DeferredAbandonCache = true;
 	}
 }
@@ -1634,14 +1645,40 @@ void FIoStoreOnDemandModule::InitializeInternal()
 void FIoStoreOnDemandModule::StartupModule()
 {
 #if !UE_IAS_CUSTOM_INITIALIZATION
-	InitializeInternal();
-#endif
+
+	if (!GIasSuspendSystem)
+	{
+		InitializeInternal();
+	}
+	else
+	{
+		UE_LOG(LogIas, Display, TEXT("The IoStoreOnDemand module has been remotely disabled by the 'ias.SuspendSystemEnabled' cvar"));
+	}
+
+#endif // !UE_IAS_CUSTOM_INITIALIZATION
 }
 
 void FIoStoreOnDemandModule::ShutdownModule()
 {
 }
 
+#if UE_IAS_CUSTOM_INITIALIZATION
+
+EOnDemandInitResult FIoStoreOnDemandModule::Initialize()
+{
+	if (GIasSuspendSystem)
+	{
+		UE_LOG(LogIas, Display, TEXT("The IoStoreOnDemand module has been remotely disabled by the 'ias.SuspendSystemEnabled' cvar"));
+		return EOnDemandInitResult::Suspended;
+	}
+
+	InitializeInternal();
+
+	return Backend.IsValid() ? EOnDemandInitResult::Success : EOnDemandInitResult::Disabled;
+};
+
+#endif // UE_IAS_CUSTOM_INITIALIZATION
+
 } // namespace UE::IO::IAS
 
 ////////////////////////////////////////////////////////////////////////////////
diff --git a/Engine/Source/Runtime/Experimental/IoStoreOnDemand/Public/IO/IoStoreOnDemand.h b/Engine/Source/Runtime/Experimental/IoStoreOnDemand/Public/IO/IoStoreOnDemand.h
--- a/Engine/Source/Runtime/Experimental/IoStoreOnDemand/Public/IO/IoStoreOnDemand.h
+++ b/Engine/Source/Runtime/Experimental/IoStoreOnDemand/Public/IO/IoStoreOnDemand.h
@@ -211,6 +211,27 @@ UE_API FIoStatus PrimeEndPoint(FStringView IoStoreOnDemandIniPath);
 
 class IOnDemandIoDispatcherBackend;
 
+#if UE_IAS_CUSTOM_INITIALIZATION
+
+/** Result of calling FIoStoreOnDemandModule::Initialize */
+enum class EOnDemandInitResult
+{
+	/** The module initialized correctly and can be used */
+	Success = 0,
+	/** The module is disabled as OnDemand data is not required for the current process*/
+	Disabled,
+	/** The module was unable to start up correctly due to an unexpected error */
+	Error,
+
+	/**
+	 * The use of the module has been suspended, if possible calling systems should activate alternative ways
+	 * to access the OnDemand data. This option is temporary and not intended for general use.
+	 */
+	Suspended
+};
+
+#endif // UE_IAS_CUSTOM_INITIALIZATION
+
 class FIoStoreOnDemandModule
 	: public IModuleInterface
 {
@@ -234,8 +255,8 @@ public:
 	virtual void ShutdownModule() override;
 	
 #if UE_IAS_CUSTOM_INITIALIZATION
-	UE_API void Initialize() { InitializeInternal(); };
-#endif
+	UE_API EOnDemandInitResult Initialize();
+#endif //UE_IAS_CUSTOM_INITIALIZATION
 };
 
 } // namespace UE::IO::IAS
