Optimization : Get Managed Array outside of the loop

diff --git a/Engine/Source/Runtime/Experimental/Chaos/Private/PhysicsProxy/GeometryCollectionPhysicsProxy.cpp b/Engine/Source/Runtime/Experimental/Chaos/Private/PhysicsProxy/GeometryCollectionPhysicsProxy.cpp
--- a/Engine/Source/Runtime/Experimental/Chaos/Private/PhysicsProxy/GeometryCollectionPhysicsProxy.cpp
+++ b/Engine/Source/Runtime/Experimental/Chaos/Private/PhysicsProxy/GeometryCollectionPhysicsProxy.cpp
@@ -2632,6 +2632,8 @@ TBitArray<> FGeometryCollectionPhysicsProxy::CalculateClustersToCreateFromChildr
 	{
 		ClustersToGenerate.Init(false, NumTransforms);
 
+		const TManagedArray<Chaos::FImplicitObjectPtr>& Implicits = DynamicCollection.GetAttribute<Chaos::FImplicitObjectPtr>(FGeometryDynamicCollection::ImplicitsAttribute, FTransformCollection::TransformGroup);
+
 		// All the leaf objects
 		TArray<int32> ChildrenToCheckForParentFix;
 		for (int32 Index = 0; Index < NumTransforms; ++Index)
@@ -2657,8 +2659,8 @@ TBitArray<> FGeometryCollectionPhysicsProxy::CalculateClustersToCreateFromChildr
 			// step 2: test the parent for having children with geometry
 			for (const int32 ParentToFixIndex : ParentToPotentiallyFix)
 			{
-				const bool bParentHasCollision = 
-					(DynamicCollection.GetAttribute<Chaos::FImplicitObjectPtr>(FGeometryDynamicCollection::ImplicitsAttribute, FTransformCollection::TransformGroup)[ParentToFixIndex] != nullptr) || ClustersToGenerate[ParentToFixIndex];
+				const bool bParentHasCollision = (Implicits[ParentToFixIndex] != nullptr) || ClustersToGenerate[ParentToFixIndex];
+
 				if (!bParentHasCollision)
 				{
 					// let's make sure all our children have an implicit defined, otherwise, postpone to next iteration 
@@ -2666,8 +2668,7 @@ TBitArray<> FGeometryCollectionPhysicsProxy::CalculateClustersToCreateFromChildr
 					DynamicCollection.IterateThroughChildren(ParentToFixIndex, [&](int32 ChildIndex)
 					{
 						// defer if any of the children is a cluster with no collision yet generated 
-						const bool bChildHasCollision = 
-							(DynamicCollection.GetAttribute<Chaos::FImplicitObjectPtr>(FGeometryDynamicCollection::ImplicitsAttribute, FTransformCollection::TransformGroup)[ChildIndex] != nullptr) || ClustersToGenerate[ChildIndex];
+						const bool bChildHasCollision = (Implicits[ChildIndex] != nullptr) || ClustersToGenerate[ChildIndex];
 						if (!bChildHasCollision && DynamicCollection.HasChildren(ChildIndex))
 						{
 							bAllChildrenHaveCollision = false;
@@ -2697,13 +2698,14 @@ int32 FGeometryCollectionPhysicsProxy::CalculateEffectiveParticles(const FGeomet
 	int32 NumEffectiveParticlesFound = 0;
 	EffectiveParticles.Init(false, NumTransform);
 	const int32 MaxSimulatedLevel = FMath::Min(GlobalMaxSimulatedLevel, InMaxSimulatedLevel);
+	const TManagedArray<Chaos::FImplicitObjectPtr>& Implicits = DynamicCollection.GetAttribute<Chaos::FImplicitObjectPtr>(FGeometryDynamicCollection::ImplicitsAttribute, FTransformCollection::TransformGroup);
 	for (int32 TransformIndex = 0; TransformIndex < NumTransform; ++TransformIndex)
 	{
 		const int32 Level = FMath::Clamp(CalculateHierarchyLevel(DynamicCollection, TransformIndex), 0, INT_MAX);
 		if (Level <= MaxSimulatedLevel || !bEnableClustering)
 		{
 			const bool bIsClusterUsingChildGeometry = (ClustersUsingChildGeometry.Num() > 0) && ClustersUsingChildGeometry[TransformIndex];
-			const bool bHasGeometry = DynamicCollection.GetAttribute<Chaos::FImplicitObjectPtr>(FGeometryDynamicCollection::ImplicitsAttribute, FTransformCollection::TransformGroup)[TransformIndex].IsValid();
+			const bool bHasGeometry = Implicits[TransformIndex].IsValid();
 
 			if (bHasGeometry || bIsClusterUsingChildGeometry)
 			{
