Iris - Mark properties as dirty when pushing out state.

#rnx

diff --git a/Engine/Plugins/Runtime/ReplicationSystemTestPlugin/Source/Private/Tests/Serialization/TestNetRoleNetSerializer.cpp b/Engine/Plugins/Runtime/ReplicationSystemTestPlugin/Source/Private/Tests/Serialization/TestNetRoleNetSerializer.cpp
--- a/Engine/Plugins/Runtime/ReplicationSystemTestPlugin/Source/Private/Tests/Serialization/TestNetRoleNetSerializer.cpp
+++ b/Engine/Plugins/Runtime/ReplicationSystemTestPlugin/Source/Private/Tests/Serialization/TestNetRoleNetSerializer.cpp
@@ -280,7 +280,7 @@ void FTestNetRoleNetSerializer::TestDequantizeSwapsRoles()
 		UE_NET_ASSERT_FALSE(Context.HasError());
 
 		constexpr bool bPushAll = true;
-		TargetReplicationState.PushPropertyReplicationState(Target.Get(), bPushAll);
+		TargetReplicationState.PushPropertyReplicationState(Target.Get(), Target.Get(), bPushAll);
 		
 		UE_NET_ASSERT_EQ_MSG(Target->RemoteRole, Source->Role, "Roles were not swapped");
 		UE_NET_ASSERT_EQ_MSG(Target->Role, Source->RemoteRole, "Roles were not swapped");
@@ -328,7 +328,7 @@ void FTestNetRoleNetSerializer::TestDowngradeRole()
 		FReplicationStateOperations::Dequantize(Context, TargetReplicationState.GetStateBuffer(), SourceBuffer[0], Descriptor);
 
 		constexpr bool bPushAll = true;
-		TargetReplicationState.PushPropertyReplicationState(Target.Get(), bPushAll);
+		TargetReplicationState.PushPropertyReplicationState(Target.Get(), Target.Get(), bPushAll);
 		
 		UE_NET_ASSERT_EQ_MSG(Target->RemoteRole, Source->Role, "Roles were not swapped");
 		UE_NET_ASSERT_EQ_MSG(Target->Role, TEnumAsByte<ENetRole>(ENetRole::ROLE_SimulatedProxy), "Role was not downgraded");
diff --git a/Engine/Source/Runtime/Experimental/Iris/Core/Private/Iris/ReplicationState/PropertyReplicationState.cpp b/Engine/Source/Runtime/Experimental/Iris/Core/Private/Iris/ReplicationState/PropertyReplicationState.cpp
--- a/Engine/Source/Runtime/Experimental/Iris/Core/Private/Iris/ReplicationState/PropertyReplicationState.cpp
+++ b/Engine/Source/Runtime/Experimental/Iris/Core/Private/Iris/ReplicationState/PropertyReplicationState.cpp
@@ -1,19 +1,20 @@
 // Copyright Epic Games, Inc. All Rights Reserved.
 
 #include "Iris/ReplicationState/PropertyReplicationState.h"
+#include "Iris/Core/IrisDebugging.h"
+#include "Iris/Core/IrisLog.h"
+#include "Iris/Core/IrisProfiler.h"
 #include "Iris/ReplicationState/InternalPropertyReplicationState.h"
 #include "Iris/ReplicationState/InternalReplicationStateDescriptorUtils.h"
 #include "Iris/ReplicationState/ReplicationStateUtil.h"
 #include "Iris/ReplicationSystem/ReplicationSystem.h"
 #include "Net/Core/NetBitArray.h"
+#include "Net/Core/PushModel/PushModel.h"
 #include "Net/Core/Trace/NetDebugName.h"
 #include "CoreTypes.h"
 #include "UObject/UnrealType.h"
-#include "Iris/Core/IrisLog.h"
 #include "UObject/PropertyPortFlags.h"
 #include "Containers/StringFwd.h"
-#include "Iris/Core/IrisDebugging.h"
-#include "Iris/Core/IrisProfiler.h"
 
 DEFINE_LOG_CATEGORY_STATIC(LogIrisRepNotify, Warning, All);
 
@@ -319,14 +320,19 @@ bool FPropertyReplicationState::PollPropertyReplicationStateForRepNotifies(const
 	return IsDirty();
 }
 
-void FPropertyReplicationState::PushPropertyReplicationState(void* RESTRICT DstData, bool bInPushAll) const
+void FPropertyReplicationState::PushPropertyReplicationState(const UObject* Owner, void* RESTRICT DstData, bool bInPushAll) const
 {
 	// $IRIS TODO: Rewrite this to iterate over change mask instead of iterating over all members and querying the mask
 	// Note, we need to use a NetBitStreamReader and the changemask descriptor since each member might have different number of bits.
 	if (IsValid())
 	{
+#if WITH_PUSH_MODEL
+		using RepIndexType = decltype(FProperty::RepIndex);
+		TArray<RepIndexType, TInlineAllocator<128>> DirtyRepIndices;
+#endif
+
 		const FReplicationStateDescriptor* Descriptor = ReplicationStateDescriptor;
-		uint8* DstBuffer = reinterpret_cast<uint8*>(DstData);
+		uint8* DstBuffer = static_cast<uint8*>(DstData);
 
 		IRIS_PROFILER_PROTOCOL_NAME(ReplicationStateDescriptor->DebugName->Name);
 
@@ -347,8 +353,25 @@ void FPropertyReplicationState::PushPropertyReplicationState(void* RESTRICT DstD
 				const FProperty* Property = MemberProperties[MemberIt];
 
 				PushPropertyValue(MemberIt, DstBuffer + Property->GetOffset_ForGC() + Property->ElementSize*MemberPropertyDescriptor.ArrayIndex);
+
+#if WITH_PUSH_MODEL
+				if (MemberPropertyDescriptor.ArrayIndex == 0)
+				{
+					DirtyRepIndices.Add(Property->RepIndex);
+				}
+#endif
 			}
 		}
+
+#if WITH_PUSH_MODEL
+		if (Owner != nullptr)
+		{
+			for (RepIndexType RepIndex : DirtyRepIndices)
+			{
+				MARK_PROPERTY_DIRTY_UNSAFE(Owner, RepIndex);
+			}
+		}
+#endif
 	}
 }
 
diff --git a/Engine/Source/Runtime/Experimental/Iris/Core/Private/Iris/ReplicationSystem/PropertyReplicationFragment.cpp b/Engine/Source/Runtime/Experimental/Iris/Core/Private/Iris/ReplicationSystem/PropertyReplicationFragment.cpp
--- a/Engine/Source/Runtime/Experimental/Iris/Core/Private/Iris/ReplicationSystem/PropertyReplicationFragment.cpp
+++ b/Engine/Source/Runtime/Experimental/Iris/Core/Private/Iris/ReplicationSystem/PropertyReplicationFragment.cpp
@@ -132,7 +132,7 @@ void FPropertyReplicationFragment::ApplyReplicatedState(FReplicationStateApplyCo
 	const FPropertyReplicationState ReceivedState(Context.Descriptor, Context.StateBufferData.ExternalStateBuffer);
 
 	// Just push the state data to owner
-	ReceivedState.PushPropertyReplicationState(Owner);
+	ReceivedState.PushPropertyReplicationState(Owner, static_cast<void*>(Owner));
 }
 
 bool FPropertyReplicationFragment::PollReplicatedState(EReplicationFragmentPollFlags PollOption)
diff --git a/Engine/Source/Runtime/Experimental/Iris/Core/Public/Iris/ReplicationState/PropertyReplicationState.h b/Engine/Source/Runtime/Experimental/Iris/Core/Public/Iris/ReplicationState/PropertyReplicationState.h
--- a/Engine/Source/Runtime/Experimental/Iris/Core/Public/Iris/ReplicationState/PropertyReplicationState.h
+++ b/Engine/Source/Runtime/Experimental/Iris/Core/Public/Iris/ReplicationState/PropertyReplicationState.h
@@ -87,15 +87,11 @@ public:
 	*/
 	IRISCORE_API bool PollPropertyReplicationState(const void* RESTRICT SrcData);
 
-	/**
-	 * Poll src data from properties which we want to store in order to determine if we need to call RepNotifies
-	*/
+	/** Poll src data from properties which we want to store in order to determine if we need to call RepNotifies. */
 	IRISCORE_API bool PollPropertyReplicationStateForRepNotifies(const void* RESTRICT SrcData);
 
-	/** Push received state data to properties in DstData buffer, Note: DstData is a UClass/UStruct containing properties
-		Compare and update representation in DstStateBuffer and update ChangeMask 
-	*/
-	IRISCORE_API void PushPropertyReplicationState(void* RESTRICT DstData, bool bPushAll = false) const;
+	/** Push received state data to properties in DstData buffer. Note: DstData is a UClass/UStruct containing properties. Updates representation in DstStateBuffer and mark destination properties as dirty if updated. */
+	IRISCORE_API void PushPropertyReplicationState(const UObject* Owner, void* RESTRICT DstData, bool bPushAll = false) const;
 
 
 	/**
diff --git a/Engine/Source/Runtime/Net/Core/Public/Net/Core/PushModel/PushModel.h b/Engine/Source/Runtime/Net/Core/Public/Net/Core/PushModel/PushModel.h
--- a/Engine/Source/Runtime/Net/Core/Public/Net/Core/PushModel/PushModel.h
+++ b/Engine/Source/Runtime/Net/Core/Public/Net/Core/PushModel/PushModel.h
@@ -469,6 +469,7 @@ namespace UEPushModelPrivate
 
 #else // WITH_PUSH_MODEL
 
+#define MARK_PROPERTY_DIRTY_UNSAFE(Object, RepIndex)
 #define MARK_PROPERTY_DIRTY(Object, Property) 
 #define MARK_PROPERTY_DIRTY_STATIC_ARRAY_INDEX(Object, RepIndex, ArrayIndex) 
 #define MARK_PROPERTY_DIRTY_STATIC_ARRAY(Object, RepIndex, ArrayIndex) 
