IAS - limit number of retries to max event loop activity count

#prefight 651d396a592a81191da58767
#rnx

diff --git a/Engine/Source/Runtime/Experimental/IoStoreOnDemand/Private/CoreHttp/Client.h b/Engine/Source/Runtime/Experimental/IoStoreOnDemand/Private/CoreHttp/Client.h
--- a/Engine/Source/Runtime/Experimental/IoStoreOnDemand/Private/CoreHttp/Client.h
+++ b/Engine/Source/Runtime/Experimental/IoStoreOnDemand/Private/CoreHttp/Client.h
@@ -180,6 +180,8 @@ class UE_API FEventLoop
 	class FImpl;
 
 public:
+	static const uint32		MaxActiveTickets = 64;
+
 	struct FRequestParams
 	{
 		uint32	BufferSize	= 256;
diff --git a/Engine/Source/Runtime/Experimental/IoStoreOnDemand/Private/OnDemandHttpClient.cpp b/Engine/Source/Runtime/Experimental/IoStoreOnDemand/Private/OnDemandHttpClient.cpp
--- a/Engine/Source/Runtime/Experimental/IoStoreOnDemand/Private/OnDemandHttpClient.cpp
+++ b/Engine/Source/Runtime/Experimental/IoStoreOnDemand/Private/OnDemandHttpClient.cpp
@@ -42,10 +42,20 @@ bool FHttpClient::Tick(uint32 WaitTimeMs, uint32 MaxKiBPerSecond)
 	}
 
 	EventLoop.Throttle(MaxKiBPerSecond);
-	const bool bTicked = EventLoop.Tick(WaitTimeMs) != 0;
+	const uint32 TicketCount = EventLoop.Tick(WaitTimeMs);
 
-	// Clean up all but the primary connection
-	if (bTicked == false)
+	if (Retries.IsEmpty() == false)
+	{
+		const int32 RequestCount = FMath::Min(Retries.Num(), int32(HTTP::FEventLoop::MaxActiveTickets - TicketCount));
+		for (int32 Idx = 0; Idx < RequestCount; Idx++)
+		{
+			IssueRequest(MoveTemp(Retries[Idx]));
+		}
+		Retries.RemoveAtSwap(0, RequestCount);
+	}
+
+	const bool bIsIdle = EventLoop.IsIdle();
+	if (bIsIdle)
 	{
 		for (int32 Idx = 0, Count = Connections.Num(); Idx < Count; ++Idx)
 		{
@@ -56,7 +66,7 @@ bool FHttpClient::Tick(uint32 WaitTimeMs, uint32 MaxKiBPerSecond)
 		}
 	}
 
-	return bTicked;
+	return bIsIdle == false;
 }
 
 FHttpClient::FHttpClient(FHttpClientConfig&& ClientConfig)
@@ -210,7 +220,7 @@ void FHttpClient::RetryRequest(FRequestParams&& Params, bool bNextConnection)
 		}
 	}
 	Params.Attempt++;
-	IssueRequest(MoveTemp(Params));
+	Retries.Emplace(MoveTemp(Params));
 }
 
 TUniquePtr<HTTP::FConnectionPool> FHttpClient::CreateConnection(const FStringView& HostAddr)
diff --git a/Engine/Source/Runtime/Experimental/IoStoreOnDemand/Private/OnDemandHttpClient.h b/Engine/Source/Runtime/Experimental/IoStoreOnDemand/Private/OnDemandHttpClient.h
--- a/Engine/Source/Runtime/Experimental/IoStoreOnDemand/Private/OnDemandHttpClient.h
+++ b/Engine/Source/Runtime/Experimental/IoStoreOnDemand/Private/OnDemandHttpClient.h
@@ -57,6 +57,7 @@ private:
 	FHttpClientConfig Config;
 	HTTP::FEventLoop EventLoop;
 	TArray<TUniquePtr<HTTP::FConnectionPool>> Connections;
+	TArray<FRequestParams> Retries;
 	int32 PrimaryConnection = INDEX_NONE;
 };
 
