- FResult() ctor was ambiguous
- Reset() -> Resize(N)
- Socks support was missing connect
- FHost holds connection count instead of pool (strategic)
- Destination FIoBuffer for content move out of "InternalResponse" to "Activity"
- Dest/ErrorReason are used independently so they can be unionised
- Fixed bug tracking how much of a request had been sent
- Waiting differentiates between sends and receives (so one can wait on both if one wishes)

#rnx

diff --git a/Engine/Source/Runtime/Experimental/IoStoreOnDemand/Private/CoreHttp/Http.cpp b/Engine/Source/Runtime/Experimental/IoStoreOnDemand/Private/CoreHttp/Http.cpp
--- a/Engine/Source/Runtime/Experimental/IoStoreOnDemand/Private/CoreHttp/Http.cpp
+++ b/Engine/Source/Runtime/Experimental/IoStoreOnDemand/Private/CoreHttp/Http.cpp
@@ -179,8 +179,8 @@ class FResult
 {
 public:
 				FResult() : FResult(0, "") {}
-				FResult(const char* Msg="") : FResult(-1, Msg) {}
-				FResult(int32 Val, const char* Msg="") : Message(UPTRINT(Msg)), Value(Val) {}
+	explicit	FResult(const char* Msg) : FResult(-1, Msg) {}
+	explicit	FResult(int32 Val, const char* Msg="") : Message(UPTRINT(Msg)), Value(Val) {}
 	const char*	GetMessage() const		{ return (const char*)Message; }
 	int32		GetValue() const		{ return int16(Value); }
 
@@ -502,7 +502,7 @@ public:
 								~FBuffer();
 	FBuffer&					operator = (FBuffer&& Rhs);
 	void						Fix();
-	void						Reset();
+	void						Resize(uint32 Size);
 	const char*					GetData() const;
 	uint32						GetSize() const;
 	uint32						GetCapacity() const;
@@ -567,9 +567,10 @@ void FBuffer::Fix()
 }
 
 ////////////////////////////////////////////////////////////////////////////////
-void FBuffer::Reset()
+void FBuffer::Resize(uint32 Size)
 {
-	Used = 0;
+	check(Size <= Max);
+	Used = Size;
 }
 
 ////////////////////////////////////////////////////////////////////////////////
@@ -705,11 +706,12 @@ public:
 
 	struct FWaiter
 	{
-				FWaiter(FSocket& Socket, bool bInIsRecvWait);
+		enum class EWhat { Send = 0b01, Recv = 0b10, Both = Send|Recv };
+				FWaiter(FSocket& Socket, EWhat InWaitOn);
 		bool	operator == (FSocket& Rhs) const { return UPTRINT(&Rhs) == Candidate; }
-		UPTRINT	Candidate : 62;
-		UPTRINT	bRecvWait : 1;
-		UPTRINT	bReady : 1;
+		UPTRINT	Candidate : 60;
+		UPTRINT	WaitOn : 2;
+		UPTRINT	Ready : 2;
 	};
 
 				FSocket() = default;
@@ -737,11 +739,12 @@ private:
 };
 
 ////////////////////////////////////////////////////////////////////////////////
-FSocket::FWaiter::FWaiter(FSocket& Socket, bool bInIsRecvWait)
+FSocket::FWaiter::FWaiter(FSocket& Socket, EWhat InWaitOn)
 : Candidate(UPTRINT(&Socket))
-, bRecvWait(bInIsRecvWait == true)
-, bReady(0)
+, WaitOn(UPTRINT(InWaitOn))
+, Ready(0)
 {
+	static_assert(sizeof(*this) == sizeof(void*));
 }
 
 ////////////////////////////////////////////////////////////////////////////////
@@ -919,12 +922,14 @@ int32 FSocket::Wait(TArrayView<FWaiter> Waiters, int32 TimeoutMs)
 	// The following looks odd because POLLFD varies subtly from one platform
 	// to the next. To cleanly set members to zero and to not get narrowing
 	// warnings from the compiler, we list-init and don't assume POD types.
+	using PollEventType = decltype(FSelect::events);
+	PollEventType Events[] = { POLLERR, POLLOUT, POLLIN, POLLOUT|POLLOUT };
 	TArray<FSelect, TFixedAllocator<64>> Selects;
 	for (FWaiter& Waiter : Waiters)
 	{
 		Selects.Emplace_GetRef() = {
 			((FSocket*)Waiter.Candidate)->Socket,
-			decltype(FSelect::events)(Waiter.bRecvWait ? POLLIN : POLLOUT),
+			Events[Waiter.WaitOn],
 			/* 0 */
 		};
 	}
@@ -984,12 +989,16 @@ int32 FSocket::Wait(TArrayView<FWaiter> Waiters, int32 TimeoutMs)
 	static const auto TestBits = POLLIN|POLLOUT|POLLERR|POLLHUP|POLLNVAL;
 	for (uint32 i = 0, n = Waiters.Num(); i < n; ++i)
 	{
-		if (int32(Selects[i].revents & TestBits) == 0)
+		auto RetEvents = Selects[i].revents;
+		if (!(RetEvents & TestBits))
 		{
 			continue;
 		}
 
-		Waiters[i].bReady = true;
+		uint32 Value = 0;
+		if (!!(RetEvents & POLLOUT)) Value |= uint32(FSocket::FWaiter::EWhat::Send);
+		if (!!(RetEvents & POLLIN))	 Value |= uint32(FSocket::FWaiter::EWhat::Recv);
+		Waiters[i].Ready = Value ? Value : uint32(FSocket::FWaiter::EWhat::Both);
 	}
 
 	return Result;
@@ -1080,6 +1089,11 @@ static int32 ConnectSocks4(FSocket& Socket, uint32 IpAddress, uint32 Port)
 		uint32	IpAddress;
 	};
 
+	if (!Socket.Connect(GSocksIpAddress, GSocksPort))
+	{
+		return -1;
+	}
+
 	int32 Result;
 
 	FSocks4Request Request = {
@@ -1112,6 +1126,11 @@ static int32 ConnectSocks5(FSocket& Socket, uint32 IpAddress, uint32 Port)
 #pragma warning(disable : 6385)
 #endif
 
+	if (!Socket.Connect(GSocksIpAddress, GSocksPort))
+	{
+		return -1;
+	}
+
 	int32 Result;
 
 	// Greeting
@@ -1194,12 +1213,13 @@ public:
 	enum class EDirection : uint8 { Send, Recv };
 	static const uint32 InvalidIp = 0x00ff'ffff;
 
-					FHost(const ANSICHAR* InHostName, uint32 InPort);
+					FHost(const ANSICHAR* InHostName, uint32 InPort, uint32 InMaxConn);
 	void			SetBufferSize(EDirection Dir, int32 Size);
 	int32			GetBufferSize(EDirection Dir) const;
 	FResult			Connect(FSocket& Socket);
 	int32			IsResolved() const;
 	FResult			ResolveHostName();
+	uint32			GetMaxConnections() const	{ return MaxConnections; }
 	uint32			GetIpAddress() const		{ return IpAddresses[0]; }
 	FAnsiStringView	GetHostName() const			{ return HostName; }
 	uint32			GetPort() const				{ return Port; }
@@ -1210,13 +1230,16 @@ private:
 	int16			SendBufKb = -1;
 	int16			RecvBufKb = -1;
 	uint16			Port;
+	uint8			MaxConnections;
 };
 
 ////////////////////////////////////////////////////////////////////////////////
-FHost::FHost(const ANSICHAR* InHostName, uint32 InPort)
+FHost::FHost(const ANSICHAR* InHostName, uint32 InPort, uint32 InMaxConn)
 : HostName(InHostName)
 , Port(uint16(InPort))
+, MaxConnections(uint8(InMaxConn))
 {
+	check(MaxConnections && MaxConnections == InMaxConn);
 }
 
 ////////////////////////////////////////////////////////////////////////////////
@@ -1284,7 +1307,7 @@ FResult FHost::ResolveHostName()
 
 	if (AddressCount > 0)
 	{
-		return AddressCount;
+		return FResult(AddressCount);
 	}
 
 	return FResult(0, "Unable to resolve host");
@@ -1359,7 +1382,7 @@ FResult FHost::Connect(FSocket& Socket)
 	}
 
 	Socket = MoveTemp(Candidate);
-	return 1;
+	return FResult(1);
 }
 
 
@@ -1377,18 +1400,15 @@ public:
 
 private:
 	uint8			LeaseCount = 0;
-	uint8			MaxLeases;
 	FSocket			Sockets[1/*...N*/]; // this should be the last member
 };
 
 ////////////////////////////////////////////////////////////////////////////////
 FSocketPool::FSocketPool(const ANSICHAR* InHostName, uint32 InPort, uint32 InMaxLeases)
-: FHost(InHostName, InPort)
-, MaxLeases(uint8(InMaxLeases))
+: FHost(InHostName, InPort, InMaxLeases)
 {
-	check(MaxLeases == InMaxLeases); // field overflow
 
-	for (uint32 i = 0; i < InMaxLeases; ++i)
+	for (uint32 i = 0; i < GetMaxConnections(); ++i)
 	{
 		new (Sockets + i) FSocket();
 	}
@@ -1397,7 +1417,7 @@ FSocketPool::FSocketPool(const ANSICHAR* InHostName, uint32 InPort, uint32 InMax
 ////////////////////////////////////////////////////////////////////////////////
 FSocketPool::~FSocketPool()
 {
-	for (uint32 i = 0; i < MaxLeases; ++i)
+	for (uint32 i = 0; i < GetMaxConnections(); ++i)
 	{
 		Sockets[i].Destroy();
 	}
@@ -1412,9 +1432,9 @@ uint32 FSocketPool::GetAllocSize(uint32 MaxLeases)
 ////////////////////////////////////////////////////////////////////////////////
 bool FSocketPool::LeaseSocket(FSocket& Out)
 {
-	check(LeaseCount <= MaxLeases);
+	check(LeaseCount <= GetMaxConnections());
 
-	if (LeaseCount == MaxLeases)
+	if (LeaseCount == GetMaxConnections())
 	{
 		return false;
 	}
@@ -1565,12 +1585,10 @@ int64 FStopwatch::Sample()
 ////////////////////////////////////////////////////////////////////////////////
 struct FResponseInternal
 {
-	FIoBuffer*		Dest;
 	FMessageOffsets Offsets;
 	int32			ContentLength;
 	uint16			MessageLength;
 	mutable int16	Code;
-	uint32			_Unused;
 	const char		Data[];
 };
 
@@ -1606,7 +1624,10 @@ struct alignas(16) FActivity
 	FStopwatch			Stopwatch;
 #endif
 	FSocketPool*		Pool;
-	const char*			ErrorReason;
+	union {
+		FIoBuffer*		Dest;
+		const char*		ErrorReason;
+	};
 	UPTRINT				SinkParam;
 	FTicketSink			Sink;
 	FSocket				Socket;
@@ -1914,7 +1935,8 @@ FAnsiStringView FResponse::GetHeader(FAnsiStringView Name) const
 ////////////////////////////////////////////////////////////////////////////////
 void FResponse::SetDestination(FIoBuffer* Buffer)
 {
-	ToResponseInternal(this).Dest = Buffer;
+	auto* Activity = (FActivity*)this;
+	Activity->Dest = Buffer;
 }
 
 
@@ -1987,8 +2009,7 @@ const FIoBuffer& FTicketStatus::GetContent() const
 {
 	check(GetId() == EId::Content);
 	const auto* Activity = (FActivity*)this;
-	auto& Response = *(FResponseInternal*)(Activity->Buffer.GetData());
-	return *(Response.Dest);
+	return *(Activity->Dest);
 }
 
 ////////////////////////////////////////////////////////////////////////////////
@@ -2071,8 +2092,8 @@ static uint64 ReadyCheck(FActivity** Activities, uint32 Num, uint32 TimeoutMs)
 			}
 			break;
 
-		case EWait::Read:  Waiters.Add({ Activity->Socket, true }); break;
-		case EWait::Write: Waiters.Add({ Activity->Socket, false});	break;
+		case EWait::Read:  Waiters.Add({ Activity->Socket, FSocket::FWaiter::EWhat::Recv }); break;
+		case EWait::Write: Waiters.Add({ Activity->Socket, FSocket::FWaiter::EWhat::Send }); break;
 		}
 	}
 
@@ -2093,7 +2114,7 @@ static uint64 ReadyCheck(FActivity** Activities, uint32 Num, uint32 TimeoutMs)
 	{
 		for (; !(Waiter == Cursor[0]->Socket); ++Cursor);
 
-		if (!Waiter.bReady)
+		if (!Waiter.Ready)
 		{
 			continue;
 		}
@@ -2198,9 +2219,9 @@ static int32 DoSend(FActivity* Activity)
 	const char* SendData = Buffer.GetData();
 	int32 SendSize = Buffer.GetSize();
 
-	uint32 Remaining = Activity->StateParam;
-	SendData += Remaining;
-	SendSize -= Remaining;
+	uint32 AlreadySent = Activity->StateParam;
+	SendData += AlreadySent;
+	SendSize -= AlreadySent;
 	check(SendSize > 0);
 
 	Trace(Activity, ETrace::Send, SendSize);
@@ -2213,7 +2234,6 @@ static int32 DoSend(FActivity* Activity)
 	case FSocket::EResult::Error:		Activity_SetError(Activity, "Error returned from socket send"); return -1;
 	case FSocket::EResult::ConnectError:Activity_SetError(Activity, "Connection error"); return -1;
 	case FSocket::EResult::Wait:
-		Activity->StateParam = Remaining;
 		Activity_BeginWait(Activity, FActivity::EWait::Write);
 		return 1;
 	}
@@ -2223,14 +2243,14 @@ static int32 DoSend(FActivity* Activity)
 #endif
 
 	checkf(Result > 0, TEXT("Result wasn't caught by switch statement so it is expected to be a positive amount of bytes sent"));
-	Remaining = SendSize - Result;
-	if (Remaining != 0)
+	Activity->StateParam += Result;
+	if (Activity->StateParam < Buffer.GetSize())
 	{
 		return DoSend(Activity);
 	}
 
 	// It is expected there will be enough space for a RespInt object
-	Buffer.Reset();
+	Buffer.Resize(0);
 	Buffer.AdvanceUsed(sizeof(FResponseInternal));
 
 	Activity_ChangeState(Activity, FActivity::EState::RecvMessage);
@@ -2323,7 +2343,7 @@ static int32 DoRecvMessage(FActivity* Activity)
 		Headers,
 		[&ContentLength, &IsKeepAlive] (FAnsiStringView Name, FAnsiStringView Value)
 		{
-			// todo; may need smarter value handling; ;/, seperated options & key-value pairs (ex. in rfc2068)
+			// todo; may need smarter value handling; ;/, separated options & key-value pairs (ex. in rfc2068)
 
 			// "Keep-Alive"			- deprecated
 			// "Transfer-Encoding"	- may be required later
@@ -2358,7 +2378,7 @@ static int32 DoRecvMessage(FActivity* Activity)
 	}
 
 	// Call out to the sink to get a content destination
-	Internal.Dest = nullptr;
+	Activity->Dest = nullptr;
 	Internal.Code = -1;
 	Internal.ContentLength = ContentLength;
 	{
@@ -2368,16 +2388,16 @@ static int32 DoRecvMessage(FActivity* Activity)
 
 	if (Activity->NoContent == 0)
 	{
-		if (Internal.Dest == nullptr)
+		if (Activity->Dest == nullptr)
 		{
 			Activity_SetError(Activity, "User did not provide a destination buffer");
 			return -1;
 		}
 
 		// The user seems to have forgotten something. Let's help them along
-		if (Internal.Dest->GetSize() == 0)
+		if (Activity->Dest->GetSize() == 0)
 		{
-			*Internal.Dest = FIoBuffer(ContentLength);
+			*Activity->Dest = FIoBuffer(ContentLength);
 		}
 	}
 
@@ -2401,9 +2421,9 @@ static int32 DoRecvMessage(FActivity* Activity)
 		return 0;
 	}
 
-	check(Internal.Dest != nullptr);
+	check(Activity->Dest != nullptr);
 
-	const bool bStreamed = Internal.Dest->GetSize() < ContentLength;
+	const bool bStreamed = Activity->Dest->GetSize() < ContentLength;
 
 	auto NextState = bStreamed ? FActivity::EState::RecvStream : FActivity::EState::RecvContent;
 	Activity_ChangeState(Activity, NextState, AlreadyReceived);
@@ -2413,7 +2433,7 @@ static int32 DoRecvMessage(FActivity* Activity)
 		return 0;
 	}
 
-	FMutableMemoryView DestView = Internal.Dest->GetMutableView();
+	FMutableMemoryView DestView = Activity->Dest->GetMutableView();
 	const char* Cursor = BufferRight - AlreadyReceived;
 	if (!bStreamed || AlreadyReceived < DestView.GetSize())
 	{
@@ -2446,8 +2466,7 @@ static FHandlerResult DoRecvContent(FActivity* Activity, uint32 MaxRecvSize)
 	TRACE_CPUPROFILER_EVENT_SCOPE(IasHttp::DoRecvContent);
 
 	auto& Response = *(FResponseInternal*)(Activity->Buffer.GetData());
-
-	FMutableMemoryView DestView = Response.Dest->GetMutableView();
+	FMutableMemoryView DestView = Activity->Dest->GetMutableView();
 
 	uint32 RecvSize = 0;
 	while (true)
