Add some code comments explaining why Ias resolves the endpoint in ::Initialize rather than in ::Mount with the rest of the start up code.

#rnx

diff --git a/Engine/Source/Runtime/Experimental/IoStoreOnDemand/Private/OnDemandIoDispatcherBackend.cpp b/Engine/Source/Runtime/Experimental/IoStoreOnDemand/Private/OnDemandIoDispatcherBackend.cpp
--- a/Engine/Source/Runtime/Experimental/IoStoreOnDemand/Private/OnDemandIoDispatcherBackend.cpp
+++ b/Engine/Source/Runtime/Experimental/IoStoreOnDemand/Private/OnDemandIoDispatcherBackend.cpp
@@ -1459,6 +1459,10 @@ void FOnDemandIoBackend::Initialize(TSharedRef<const FIoDispatcherBackendContext
 	UE_LOG(LogIas, Log, TEXT("Initializing on demand I/O dispatcher backend"));
 	BackendContext = Context;
 
+	// We need to resolve the end point in this method which occurs after the config system has initialized
+	// rather than in ::Mount which can occur before that.
+	// Without the config system initialized the http module will not work properly and we will always fail
+	// to resolve and the OnDemand system will not recover.
 	if (DistributionUrl.IsEmpty() == false)
 	{
 		TSharedPtr<FDistributionEndpoints> Resolver = MakeShared<FDistributionEndpoints>();
