Allow requests to be rewound and tried again

Note currently unused here. Request messages are built in
Activity.Buffer. Then the buffer was reset and used for the response
message and internal data. Now the internal data is in the activity and
response message is appended to buffer so the original request message
is can be resent if need be. A note about HangUp from a socket is
included albeit not implemented as yet.

#rnx

diff --git a/Engine/Source/Runtime/Experimental/IoStoreOnDemand/Private/CoreHttp/Http.cpp b/Engine/Source/Runtime/Experimental/IoStoreOnDemand/Private/CoreHttp/Http.cpp
--- a/Engine/Source/Runtime/Experimental/IoStoreOnDemand/Private/CoreHttp/Http.cpp
+++ b/Engine/Source/Runtime/Experimental/IoStoreOnDemand/Private/CoreHttp/Http.cpp
@@ -1589,7 +1589,6 @@ struct FResponseInternal
 	int32			ContentLength;
 	uint16			MessageLength;
 	mutable int16	Code;
-	const char		Data[];
 };
 
 ////////////////////////////////////////////////////////////////////////////////
@@ -1631,6 +1630,7 @@ struct alignas(16) FActivity
 	UPTRINT				SinkParam;
 	FTicketSink			Sink;
 	FSocket				Socket;
+	FResponseInternal	Response;
 	FBuffer				Buffer;
 };
 
@@ -1658,6 +1658,10 @@ static void Activity_ChangeState(FActivity* Activity, FActivity::EState InState,
 	FStopwatch& Stopwatch = Activity->Stopwatch;
 	if (InState == EState::Send)
 	{
+		if (Activity->State == EState::RecvMessage)
+		{
+			Stopwatch = FStopwatch();
+		}
 		Stopwatch.Start();
 	}
 	else if (Activity->State == EState::Send)
@@ -1713,6 +1717,27 @@ static void Activity_EndWait(FActivity* Activity)
 	Activity->SocketWait = FActivity::EWait::None;
 }
 
+////////////////////////////////////////////////////////////////////////////////
+static int32 Activity_Rewind(FActivity* Activity)
+{
+	using EState = FActivity::EState;
+
+	if (Activity->State == EState::Send)
+	{
+		Activity->StateParam = 0;
+		return 0;
+	}
+
+	if (Activity->State == EState::RecvMessage)
+	{
+		Activity->Buffer.Resize(Activity->StateParam);
+		Activity_ChangeState(Activity, EState::Send);
+		return 1;
+	}
+
+	return -1;
+}
+
 ////////////////////////////////////////////////////////////////////////////////
 static FActivity* Activity_Alloc(uint32 BufferSize)
 {
@@ -1822,20 +1847,6 @@ FRequest&& FRequest::Header(FAnsiStringView Key, FAnsiStringView Value)
 
 // {{{1 response ...............................................................
 
-////////////////////////////////////////////////////////////////////////////////
-static FResponseInternal& ToResponseInternal(FResponse* Addr)
-{
-	auto* Activity = (FActivity*)Addr;
-	return *(FResponseInternal*)(Activity->Buffer.GetData());
-}
-
-////////////////////////////////////////////////////////////////////////////////
-static const FResponseInternal& ToResponseInternal(const FResponse* Addr)
-{
-	const auto* Activity = (const FActivity*)Addr;
-	return *(const FResponseInternal*)(Activity->Buffer.GetData());
-}
-
 ////////////////////////////////////////////////////////////////////////////////
 EStatusCodeClass FResponse::GetStatus() const
 {
@@ -1851,11 +1862,13 @@ EStatusCodeClass FResponse::GetStatus() const
 ////////////////////////////////////////////////////////////////////////////////
 uint32 FResponse::GetStatusCode() const
 {
-	const FResponseInternal& Internal = ToResponseInternal(this);
+	const auto* Activity = (const FActivity*)this;
+	const FResponseInternal& Internal = Activity->Response;
 
+	const char* MessageData = Activity->Buffer.GetData() + Activity->StateParam;
 	if (Internal.Code < 0)
 	{
-		const char* CodePtr = Internal.Data + Internal.Offsets.StatusCode;
+		const char* CodePtr = MessageData + Internal.Offsets.StatusCode;
 		Internal.Code = uint16(CrudeToInt(FAnsiStringView(CodePtr, 3)));
 	}
 
@@ -1865,9 +1878,12 @@ uint32 FResponse::GetStatusCode() const
 ////////////////////////////////////////////////////////////////////////////////
 FAnsiStringView FResponse::GetStatusMessage() const
 {
-	const FResponseInternal& Internal = ToResponseInternal(this);
+	const auto* Activity = (const FActivity*)this;
+	const FResponseInternal& Internal = Activity->Response;
+
+	const char* MessageData = Activity->Buffer.GetData() + Activity->StateParam;
 	return FAnsiStringView(
-		Internal.Data + Internal.Offsets.Message,
+		MessageData + Internal.Offsets.Message,
 		Internal.Offsets.Headers - Internal.Offsets.Message
 	);
 }
@@ -1875,7 +1891,9 @@ FAnsiStringView FResponse::GetStatusMessage() const
 ////////////////////////////////////////////////////////////////////////////////
 int64 FResponse::GetContentLength() const
 {
-	return ToResponseInternal(this).ContentLength;
+	const auto* Activity = (const FActivity*)this;
+	const FResponseInternal& Internal = Activity->Response;
+	return Internal.ContentLength;
 }
 
 ////////////////////////////////////////////////////////////////////////////////
@@ -1912,10 +1930,12 @@ void FResponse::GetContentType(FAnsiStringView& Out) const
 ////////////////////////////////////////////////////////////////////////////////
 FAnsiStringView FResponse::GetHeader(FAnsiStringView Name) const
 {
-	const FResponseInternal& Internal = ToResponseInternal(this);
+	const auto* Activity = (const FActivity*)this;
+	const FResponseInternal& Internal = Activity->Response;
 
+	const char* MessageData = Activity->Buffer.GetData() + Activity->StateParam;
 	FAnsiStringView Result, Headers(
-		Internal.Data + Internal.Offsets.Headers,
+		MessageData + Internal.Offsets.Headers,
 		Internal.MessageLength - Internal.Offsets.Headers
 	);
 
@@ -1992,8 +2012,7 @@ uint32 FTicketStatus::GetContentLength() const
 {
 	check(GetId() <= EId::Content);
 	const auto* Activity = (FActivity*)this;
-	auto& Response = *(FResponseInternal*)(Activity->Buffer.GetData());
-	return Response.ContentLength;
+	return Activity->Response.ContentLength;
 }
 
 ////////////////////////////////////////////////////////////////////////////////
@@ -2249,11 +2268,7 @@ static int32 DoSend(FActivity* Activity)
 		return DoSend(Activity);
 	}
 
-	// It is expected there will be enough space for a RespInt object
-	Buffer.Resize(0);
-	Buffer.AdvanceUsed(sizeof(FResponseInternal));
-
-	Activity_ChangeState(Activity, FActivity::EState::RecvMessage);
+	Activity_ChangeState(Activity, FActivity::EState::RecvMessage, Buffer.GetSize());
 	Activity_BeginWait(Activity, FActivity::EWait::Read);
 	return 1;
 }
@@ -2293,7 +2308,7 @@ static int32 DoRecvMessage(FActivity* Activity)
 		// Rewind a little to cover cases where the terminal is fragmented across
 		// recv() calls
 		uint32 DestBias = 0;
-		if (Dest - 3 >= Buffer.GetData() + sizeof(FResponseInternal))
+		if (Dest - 3 >= Buffer.GetData() + Activity->StateParam)
 		{
 			Dest -= (DestBias = 3);
 		}
@@ -2321,10 +2336,11 @@ static int32 DoRecvMessage(FActivity* Activity)
 	}
 
 	// Fill out the internal response object
-	auto& Internal = *(FResponseInternal*)(Buffer.GetData());
-	Internal.MessageLength = uint16(ptrdiff_t(MessageRight - Internal.Data));
+	FResponseInternal& Internal = Activity->Response;
+	const char* MessageData = Buffer.GetData() + Activity->StateParam;
+	Internal.MessageLength = uint16(ptrdiff_t(MessageRight - MessageData));
 
-	FAnsiStringView ResponseView(Internal.Data, Internal.MessageLength);
+	FAnsiStringView ResponseView(MessageData, Internal.MessageLength);
 	if (ParseMessage(ResponseView, Internal.Offsets) < 0)
 	{
 		Activity_SetError(Activity, "Failed to parse message status");
@@ -2465,7 +2481,7 @@ static FHandlerResult DoRecvContent(FActivity* Activity, uint32 MaxRecvSize)
 {
 	TRACE_CPUPROFILER_EVENT_SCOPE(IasHttp::DoRecvContent);
 
-	auto& Response = *(FResponseInternal*)(Activity->Buffer.GetData());
+	FResponseInternal& Response = Activity->Response;
 	FMutableMemoryView DestView = Activity->Dest->GetMutableView();
 
 	uint32 RecvSize = 0;
