- The updating of available throttle bytes is done via reference rather than using a tuple-like return value.
- Change activity handlers to take a socket rather than use the one in the activity directly. Also included is a generalised DoRecv() for each state which matches DoSend().

#rnx

diff --git a/Engine/Source/Runtime/Experimental/IoStoreOnDemand/Private/CoreHttp/Http.cpp b/Engine/Source/Runtime/Experimental/IoStoreOnDemand/Private/CoreHttp/Http.cpp
--- a/Engine/Source/Runtime/Experimental/IoStoreOnDemand/Private/CoreHttp/Http.cpp
+++ b/Engine/Source/Runtime/Experimental/IoStoreOnDemand/Private/CoreHttp/Http.cpp
@@ -2148,13 +2148,6 @@ static uint64 ReadyCheck(FActivity** Activities, uint32 Num, uint32 TimeoutMs)
 
 
 
-////////////////////////////////////////////////////////////////////////////////
-struct FHandlerResult
-{
-	int32	Result;
-	uint32	Size;
-};
-
 ////////////////////////////////////////////////////////////////////////////////
 static int32 DoResolve(FActivity* Activity)
 {
@@ -2230,7 +2223,7 @@ static int32 DoConnect(FActivity* Activity)
 }
 
 ////////////////////////////////////////////////////////////////////////////////
-static int32 DoSend(FActivity* Activity)
+static int32 DoSend(FActivity* Activity, FSocket& Socket)
 {
 	TRACE_CPUPROFILER_EVENT_SCOPE(IasHttp::DoSend);
 
@@ -2244,7 +2237,7 @@ static int32 DoSend(FActivity* Activity)
 	check(SendSize > 0);
 
 	Trace(Activity, ETrace::Send, SendSize);
-	int32 Result = Activity->Socket.Send(SendData, SendSize);
+	int32 Result = Socket.Send(SendData, SendSize);
 	Trace(Activity, ETrace::Send, -1);
 
 	switch (FSocket::EResult(Result))
@@ -2265,7 +2258,7 @@ static int32 DoSend(FActivity* Activity)
 	Activity->StateParam += Result;
 	if (Activity->StateParam < Buffer.GetSize())
 	{
-		return DoSend(Activity);
+		return DoSend(Activity, Socket);
 	}
 
 	Activity_ChangeState(Activity, FActivity::EState::RecvMessage, Buffer.GetSize());
@@ -2274,7 +2267,13 @@ static int32 DoSend(FActivity* Activity)
 }
 
 ////////////////////////////////////////////////////////////////////////////////
-static int32 DoRecvMessage(FActivity* Activity)
+static int32 DoSend(FActivity* Activity)
+{
+	return DoSend(Activity, Activity->Socket);
+}
+
+////////////////////////////////////////////////////////////////////////////////
+static int32 DoRecvMessage(FActivity* Activity, FSocket& Socket)
 {
 	TRACE_CPUPROFILER_EVENT_SCOPE(IasHttp::DoRecvMessage);
 
@@ -2288,7 +2287,7 @@ static int32 DoRecvMessage(FActivity* Activity)
 		auto [Dest, DestSize] = Buffer.GetMutableFree(0, PageSize);
 
 		Trace(Activity, ETrace::Recv, -1);
-		int32 Result = Activity->Socket.Recv(Dest, DestSize);
+		int32 Result = Socket.Recv(Dest, DestSize);
 		Trace(Activity, ETrace::Recv, FMath::Max(Result, 0));
 
 		if (Result == int32(FSocket::EResult::Wait))
@@ -2477,68 +2476,92 @@ static int32 DoRecvMessage(FActivity* Activity)
 }
 
 ////////////////////////////////////////////////////////////////////////////////
-static FHandlerResult DoRecvContent(FActivity* Activity, uint32 MaxRecvSize)
+static int32 DoRecvContent(FActivity* Activity, FSocket& Socket, int32& MaxRecvSize)
 {
 	TRACE_CPUPROFILER_EVENT_SCOPE(IasHttp::DoRecvContent);
 
 	FResponseInternal& Response = Activity->Response;
 	FMutableMemoryView DestView = Activity->Dest->GetMutableView();
 
-	uint32 RecvSize = 0;
 	while (true)
 	{
-		uint32 Size = (Response.ContentLength - Activity->StateParam);
+		int32 Size = (Response.ContentLength - Activity->StateParam);
 		if (Size == 0)
 		{
 			break;
 		}
 
-		Size = FMath::Min(Size, MaxRecvSize - RecvSize);
+		Size = FMath::Min(Size, MaxRecvSize);
+		check(Size >= 0);
 		if (Size == 0)
 		{
-			return { 1, RecvSize };
+			return 1;
 		}
 
 		char* Cursor = (char*)(DestView.GetData()) + Activity->StateParam;
 
 		Trace(Activity, ETrace::Recv, -1);
-		int32 Result = Activity->Socket.Recv(Cursor, Size);
+		int32 Result = Socket.Recv(Cursor, Size);
 		Trace(Activity, ETrace::Recv, FMath::Max(Result, 0));
 
 		if (Result == int32(FSocket::EResult::Wait))
 		{
 			Activity_BeginWait(Activity, FActivity::EWait::Read);
-			return { 1, RecvSize };
+			return 1;
 		}
 
 		if (Result < 0)
 		{
 			Activity_SetError(Activity, "Socket error while receiving content");
-			return { -1 };
+			return -1;
 		}
 
 		if (Result == 0 && (Activity->StateParam != Response.ContentLength))
 		{
 			Activity_SetError(Activity, "ATH0.RecvContent");
-			return { -1 };
+			return -1;
 		}
 
 #if IAS_HTTP_WITH_PERF
 		Activity->Stopwatch.AddCount();
 #endif
 
+		check(Result <= MaxRecvSize);
 		Activity->StateParam += Result;
-		RecvSize += Result;
+		MaxRecvSize -= Result;
 	}
 
 	if (!FLatencyInjector::Begin(FLatencyInjector::EType::Network, Activity->StateParam))
 	{
 		Activity_SetError(Activity, "Forced random failure");
-		return { -1 };
+		return -1;
 	}
 
 	Activity_ChangeState(Activity, FActivity::EState::RecvDone);
-	return { 0, RecvSize };
+	return 0;
+}
+
+////////////////////////////////////////////////////////////////////////////////
+static int32 DoRecvStream(FActivity*, FSocket&, uint32)
+{
+	check(false); // not yet implemented
+	return -1;
+}
+
+////////////////////////////////////////////////////////////////////////////////
+static int32 DoRecv(FActivity* Activity, FSocket& Socket, int32& MaxRecvSize)
+{
+	using EState = FActivity::EState;
+
+	EState State = Activity->State; 
+	check(State >= EState::RecvMessage && State < EState::RecvDone);
+
+	if (State == EState::RecvMessage)	return DoRecvMessage(Activity, Socket);
+	if (State == EState::RecvContent)	return DoRecvContent(Activity, Socket, MaxRecvSize);
+	if (State == EState::RecvStream)	return DoRecvStream(Activity, Socket, MaxRecvSize);
+	
+	check(false); // it is not expected that we'll get here
+	return -1;
 }
 
 ////////////////////////////////////////////////////////////////////////////////
@@ -2558,12 +2581,24 @@ static int32 DoRecvDone(FActivity* Activity)
 }
 
 ////////////////////////////////////////////////////////////////////////////////
-static FHandlerResult DoRecvStream(FActivity* Activity, uint32 RecvSize)
+static int32 DoRecvMessage(FActivity* Activity)
+{
+	return DoRecvMessage(Activity, Activity->Socket);
+}
+
+////////////////////////////////////////////////////////////////////////////////
+static int32 DoRecvContent(FActivity* Activity, int32& RecvSize)
+{
+	return DoRecvContent(Activity, Activity->Socket, RecvSize);
+}
+
+////////////////////////////////////////////////////////////////////////////////
+static int32 DoRecvStream(FActivity* Activity, int32& RecvSize)
 {
 	TRACE_CPUPROFILER_EVENT_SCOPE(IasHttp::DoRecvStream);
 
 	check(false); // not implemented yet
-	return { -1 };
+	return -1;
 }
 
 ////////////////////////////////////////////////////////////////////////////////
@@ -2891,7 +2926,7 @@ uint32 FEventLoop::FImpl::Tick(uint32 PollTimeoutMs)
 			{
 			case FActivity::EState::RecvContent:
 			case FActivity::EState::RecvStream:
-				decltype(DoRecvContent)* Handler;
+				int32 (*Handler)(FActivity*, int32&);
 				if (Activity->State == FActivity::EState::RecvContent)
 				{
 					Handler = DoRecvContent;
@@ -2907,9 +2942,7 @@ uint32 FEventLoop::FImpl::Tick(uint32 PollTimeoutMs)
 					break;
 				}
 
-				auto [ResultInner, RecvSize] = Handler(Activity, RecvAllowance);
-				RecvAllowance -= RecvSize;
-				Result = ResultInner;
+				Result = Handler(Activity, RecvAllowance);
 				if (Result)
 					break;
 			}
