Chaos - debris collision (all cvar gated - currently disabled)
- make the "one way collision" flag on body instance work
- objects set as one-way collision do not apply forces to other objects
- objects set as one-way collision do not collide with each other
- add a per-particle flag to disable the initial-overlap depenetration system (not wired up yet)
- debug draw cleanup

diff --git a/Engine/Source/Runtime/Experimental/Chaos/Private/Chaos/ChaosDebugDraw.cpp b/Engine/Source/Runtime/Experimental/Chaos/Private/Chaos/ChaosDebugDraw.cpp
--- a/Engine/Source/Runtime/Experimental/Chaos/Private/Chaos/ChaosDebugDraw.cpp
+++ b/Engine/Source/Runtime/Experimental/Chaos/Private/Chaos/ChaosDebugDraw.cpp
@@ -170,16 +170,23 @@ namespace Chaos
 			FColor InDynamicColor,
 			FColor InSleepingColor,
 			FColor InKinematicColor,
-			FColor InStaticColor
+			FColor InStaticColor,
+			FColor InDebrisColor
 		)
 			: DynamicColor(InDynamicColor)
 			, SleepingColor(InSleepingColor)
 			, KinematicColor(InKinematicColor)
 			, StaticColor(InStaticColor)
+			, DebrisColor(InDebrisColor)
 		{}
 
-		FColor FChaosDebugDrawColorsByState::GetColorFromState(EObjectStateType State) const
+		FColor FChaosDebugDrawColorsByState::GetColorFromState(EObjectStateType State, bool bIsDebris) const
 		{
+			if (bIsDebris && (State != EObjectStateType::Sleeping))
+			{
+				return DebrisColor;
+			}
+
 			switch (State)
 			{
 			case EObjectStateType::Sleeping:	return SleepingColor;
@@ -193,12 +200,14 @@ namespace Chaos
 		const FChaosDebugDrawColorsByState& GetDefaultShapesColorsByState()
 		{
 			// default colors by state for shapes
-			static FChaosDebugDrawColorsByState ChaosDefaultShapesColorsByState(
+			static FChaosDebugDrawColorsByState ChaosDefaultShapesColorsByState
+			{
 				/* InDynamicColor =	  */ FColor(255, 255, 0),
 				/* InSleepingColor =  */ FColor(128, 128, 128),
 				/* InKinematicColor = */ FColor(0, 128, 255),
-				/* InStaticColor =	  */ FColor(255, 0, 0)
-			);
+				/* InStaticColor =	  */ FColor(255, 0, 0),
+				/* InDebrisColor =	  */ FColor(0, 255, 0),
+			};
 
 			return ChaosDefaultShapesColorsByState;
 		}
@@ -206,12 +215,14 @@ namespace Chaos
 		const FChaosDebugDrawColorsByState& GetDefaultBoundsColorsByState()
 		{
 			// default colors by state for bounds ( darker version of the shapes colors - see above )
-			static FChaosDebugDrawColorsByState ChaosDefaultBoundsColorsByState(
+			static FChaosDebugDrawColorsByState ChaosDefaultBoundsColorsByState
+			{
 				/* InDynamicColor =	  */ FColor(128, 128, 0),
 				/* InSleepingColor =  */ FColor(64, 64, 64),
 				/* InKinematicColor = */ FColor(0, 64, 128),
-				/* InStaticColor =	  */ FColor(128, 0, 0)
-			);
+				/* InStaticColor =	  */ FColor(128, 0, 0),
+				/* InDebrisColor =	  */ FColor(0, 128, 0),
+			};
 
 			return ChaosDefaultBoundsColorsByState;
 		}
@@ -322,6 +333,24 @@ namespace Chaos
 		//
 		//
 
+		inline bool ParticleIsOneWayInteraction(const FGeometryParticleHandle* Particle)
+		{
+			if (const FPBDRigidParticleHandle* Rigid = Particle->CastToRigidParticle())
+			{
+				return Rigid->OneWayInteraction();
+			}
+			return false;
+		}
+
+		inline FColor GetParticleShapeDefaultColor(const FGeometryParticleHandle* Particle, const FChaosDebugDrawSettings& DebugDrawSettings)
+		{
+			return DebugDrawSettings.ShapesColorsPerState.GetColorFromState(Particle->ObjectState(), ParticleIsOneWayInteraction(Particle));
+		}
+
+		//
+		//
+		//
+
 		void DrawShapesImpl(const FGeometryParticleHandle* Particle, const FRigidTransform3& ShapeTransform, const FImplicitObject* Implicit, const FPerShapeData* Shape, const FReal Margin, const FColor& Color, const FRealSingle Duration, const FChaosDebugDrawSettings& Settings);
 
 		void DrawShape(const FRigidTransform3& ShapeTransform, const FImplicitObject* Implicit, const FPerShapeData* Shape, const FColor& Color, const float Duration, const FChaosDebugDrawSettings* Settings)
@@ -614,27 +643,6 @@ namespace Chaos
 						}
 					}
 				}
-				// Draw the manager shapes if the implicit is the root union 
-				if(Particle->GetGeometry() == Implicit)
-				{
-					if(auto ClusteredParticle = Particle->CastToClustered())
-					{
-						const TPimplPtr<Private::FConvexOptimizer>& ConvexOptimizer = ClusteredParticle->ConvexOptimizer();
-						if(ConvexOptimizer && ConvexOptimizer->IsValid())
-						{
-							for (auto& PerShapeData : ConvexOptimizer->GetShapeInstances())
-							{
-								if(FImplicitObjectUnion* Union = PerShapeData->GetGeometry()->AsA<FImplicitObjectUnion>())
-								{
-									for(auto& UnionConvex : Union->GetObjects())
-									{
-										DrawShapesImpl(Particle, ShapeTransform, UnionConvex, PerShapeData.Get(), Margin, Color, Duration, Settings);
-									}
-								}
-							}
-						}
-					}
-				}
 				return;
 			}
 
@@ -849,24 +857,32 @@ namespace Chaos
 
 		void DrawParticleShapesImpl(const FRigidTransform3& SpaceTransform, const FGeometryParticleHandle* Particle, const FColor& InColor, const FChaosDebugDrawSettings& Settings)
 		{
-			FVec3 P = SpaceTransform.TransformPosition(Particle->ObjectState() == EObjectStateType::Dynamic ? Particle->CastToRigidParticle()->P() : Particle->X());
-			FRotation3 Q = SpaceTransform.GetRotation() * (Particle->ObjectState() == EObjectStateType::Dynamic ? Particle->CastToRigidParticle()->Q() : Particle->R());
+			const FVec3 P = SpaceTransform.TransformPosition(Particle->ObjectState() == EObjectStateType::Dynamic ? Particle->CastToRigidParticle()->P() : Particle->X());
+			const FRotation3 Q = SpaceTransform.GetRotation() * (Particle->ObjectState() == EObjectStateType::Dynamic ? Particle->CastToRigidParticle()->Q() : Particle->R());
+			const FRigidTransform3 ParticleSpaceTransform = FRigidTransform3(P, Q);
 
 			for (const FShapeInstancePtr& ShapeInstance : Particle->ShapeInstances())
 			{
-				DrawShapesImpl(Particle, FRigidTransform3(P, Q), ShapeInstance->GetGeometry(), ShapeInstance.Get(), 0.0f, InColor, 0.0f, Settings);
+				DrawShapesImpl(Particle, ParticleSpaceTransform, ShapeInstance->GetGeometry(), ShapeInstance.Get(), 0.0f, InColor, 0.0f, Settings);
+			}
+
+			// Draw the optimized shapes if the particle has any
+			if (const FPBDRigidClusteredParticleHandle* ClusteredParticle = Particle->CastToClustered())
+			{
+				const TPimplPtr<Private::FConvexOptimizer>& ConvexOptimizer = ClusteredParticle->ConvexOptimizer();
+				if (ConvexOptimizer && ConvexOptimizer->IsValid())
+				{
+					for (const FShapeInstancePtr& ShapeInstance : ConvexOptimizer->GetShapeInstances())
+					{
+						DrawShapesImpl(Particle, ParticleSpaceTransform, ShapeInstance->GetGeometry(), ShapeInstance.Get(), 0.0f, FColor::Green, 0.0f, Settings);
+					}
+				}
 			}
 		}
 
 		void DrawParticleShapesImpl(const FRigidTransform3& SpaceTransform, const FGeometryParticle* Particle, const FColor& InColor, const FChaosDebugDrawSettings& Settings)
 		{
-			FVec3 P = SpaceTransform.TransformPosition(Particle->X());
-			FRotation3 Q = SpaceTransform.GetRotation() * (Particle->R());
-
-			for (const FShapeInstanceProxyPtr& ShapeInstance : Particle->ShapeInstances())
-			{
-				DrawShapesImpl(Particle->Handle(), FRigidTransform3(P, Q), ShapeInstance->GetGeometry(), ShapeInstance.Get(), 0.0f, InColor, 0.0f, Settings);
-			}
+			DrawParticleShapesImpl(SpaceTransform, Particle->Handle(), InColor, Settings);
 		}
 
 		void DrawBVHImpl(const FGeometryParticleHandle* Particle, const FRigidTransform3& ShapeTransform, const Private::FImplicitBVH* BVH, const FColor& UnusedColor, const FRealSingle Duration, const FChaosDebugDrawSettings& Settings)
@@ -957,7 +973,7 @@ namespace Chaos
 
 			const FVec3 P = SpaceTransform.TransformPosition(Box.GetCenter());
 			const FRotation3 Q = SpaceTransform.GetRotation();
-			FColor Color = Settings.BoundsColorsPerState.GetColorFromState(InParticle->ObjectState());
+			FColor Color = Settings.BoundsColorsPerState.GetColorFromState(InParticle->ObjectState(), ParticleIsOneWayInteraction(InParticle));
 			if (bChaosDebugDebugDrawColorBoundsByShapeType)
 			{
 				if (const FImplicitObject* Shape = Particle->GetGeometry())
@@ -1807,7 +1823,7 @@ namespace Chaos
 				VisitEveryConnectionToDraw(
 					[&Settings, &Connections, &GetClusterTransform](const FRigidClustering::FClusterHandle& Child)
 					{
-						FColor PointColor = Settings.ShapesColorsPerState.GetColorFromState(Child->ObjectState());
+						FColor PointColor = Settings.ShapesColorsPerState.GetColorFromState(Child->ObjectState(), ParticleIsOneWayInteraction(Child));
 						if (Child->IsAnchored())
 						{
 							PointColor = Settings.ShapesColorsPerState.KinematicColor;
@@ -1892,7 +1908,6 @@ namespace Chaos
 			}
 		}
 
-
 		void DrawParticleShapes(const FRigidTransform3& SpaceTransform, const TParticleView<FGeometryParticles>& ParticlesView, FReal ColorScale, const FChaosDebugDrawSettings* Settings)
 		{
 			if (FDebugDrawQueue::IsDebugDrawingEnabled())
@@ -1900,7 +1915,7 @@ namespace Chaos
 				const FChaosDebugDrawSettings& DebugDrawSettings = GetChaosDebugDrawSettings(Settings);
 				for (auto& Particle : ParticlesView)
 				{				
-					FColor Color = ((float)ColorScale * DebugDrawSettings.ShapesColorsPerState.GetColorFromState(Particle.ObjectState())).ToFColor(false);
+					const FColor Color = GetParticleShapeDefaultColor(Particle.Handle(), DebugDrawSettings);
 					DrawParticleShapesImpl(SpaceTransform, GetHandleHelper(&Particle), Color, DebugDrawSettings);
 				}
 			}
@@ -1913,7 +1928,7 @@ namespace Chaos
 				const FChaosDebugDrawSettings& DebugDrawSettings = GetChaosDebugDrawSettings(Settings);
 				for (auto& Particle : ParticlesView)
 				{
-					FColor Color = ((float)ColorScale * DebugDrawSettings.ShapesColorsPerState.GetColorFromState(Particle.ObjectState())).ToFColor(false);
+					const FColor Color = GetParticleShapeDefaultColor(Particle.Handle(), DebugDrawSettings);
 					DrawParticleShapesImpl(SpaceTransform, GetHandleHelper(&Particle), Color, DebugDrawSettings);
 				}
 			}
@@ -1926,7 +1941,7 @@ namespace Chaos
 				const FChaosDebugDrawSettings& DebugDrawSettings = GetChaosDebugDrawSettings(Settings);
 				for (auto& Particle : ParticlesView)
 				{
-					FColor Color = ((float)ColorScale * DebugDrawSettings.ShapesColorsPerState.GetColorFromState(Particle.ObjectState())).ToFColor(false);
+					const FColor Color = GetParticleShapeDefaultColor(Particle.Handle(), DebugDrawSettings);
 					DrawParticleShapesImpl(SpaceTransform, GetHandleHelper(&Particle), Color, DebugDrawSettings);
 				}
 			}
diff --git a/Engine/Source/Runtime/Experimental/Chaos/Private/Chaos/Collision/PBDCollisionConstraint.cpp b/Engine/Source/Runtime/Experimental/Chaos/Private/Chaos/Collision/PBDCollisionConstraint.cpp
--- a/Engine/Source/Runtime/Experimental/Chaos/Private/Chaos/Collision/PBDCollisionConstraint.cpp
+++ b/Engine/Source/Runtime/Experimental/Chaos/Private/Chaos/Collision/PBDCollisionConstraint.cpp
@@ -50,9 +50,16 @@ namespace Chaos
 	FRealSingle Chaos_Collision_ConvexZeroMargin = 0.0f;
 	FAutoConsoleVariableRef CVarChaos_Collision_ConvexZeroMargin(TEXT("p.Chaos.Collision.ConvexZeroMargin"), Chaos_Collision_ConvexZeroMargin, TEXT(""));
 
+	// A collision solver stiffness override if >= 0. A value < 0 means use ths config setting.
 	FRealSingle Chaos_Collision_Stiffness = -1.0f;
 	FAutoConsoleVariableRef CVarChaos_Collision_Stiffness(TEXT("p.Chaos.Collision.Stiffness"), Chaos_Collision_Stiffness, TEXT("Override the collision solver stiffness (if >= 0)"));
 
+	// The stiffness used for collision between one-way objects and other dynamic objects. The stiffness against kinematics and statics is not affected. A value of less than 1
+	// will treat collisions with the static scene as "harder" than collisions with dynamics which will prevent one-way objects from getting squeezed out of the world when a dynamic
+	// lands on top of them. Instead, they will stay on top of the terrain and sink into the dynamic.
+	FRealSingle Chaos_Collision_OneWayStiffness = 0.5f;
+	FAutoConsoleVariableRef CVarChaos_Collision_OneWayStiffness(TEXT("p.Chaos.Collision.OneWayStiffness"), Chaos_Collision_OneWayStiffness, TEXT("Collision solver stiffnes for one-way interactions"));
+
 	bool bChaos_Collision_EnableBoundsChecks = true;
 	FAutoConsoleVariableRef CVarChaos_Collision_EnableBoundsChecks(TEXT("p.Chaos.Collision.EnableBoundsChecks"), bChaos_Collision_EnableBoundsChecks, TEXT(""));
 
@@ -316,17 +323,21 @@ namespace Chaos
 		Flags.bUseManifold = bInUseManifold;
 		Flags.bUseIncrementalManifold = false;
 
+		// Should we use the initial overlap depenetration mechanism
+		Flags.bOverlapDepenetrationEnabled = FConstGenericParticleHandle(GetParticle0())->InitialOverlapDepentrationEnabled() && FConstGenericParticleHandle(GetParticle1())->InitialOverlapDepentrationEnabled();
+
+		// Is this a one-way interaction?
+		const bool bDynamic0 = FConstGenericParticleHandle(GetParticle0())->IsDynamic();
+		const bool bDynamic1 = FConstGenericParticleHandle(GetParticle1())->IsDynamic();
+		const bool bOneWay0 = FConstGenericParticleHandle(GetParticle0())->OneWayInteraction();
+		const bool bOneWay1 = FConstGenericParticleHandle(GetParticle1())->OneWayInteraction();
+		Flags.bIsOneWayInteraction = bDynamic0 && bDynamic1 && (bOneWay0 || bOneWay1);
+
 		// Only levelsets use incremental collision manifolds
 		if (bInUseManifold && ((ImplicitType0 == ImplicitObjectType::LevelSet) || (ImplicitType1 == ImplicitObjectType::LevelSet)))
 		{
 			Flags.bUseIncrementalManifold = true;
 		}
-
-		// Debug testing for solver stiffness
-		if (Chaos_Collision_Stiffness >= 0)
-		{
-			SetStiffness(Chaos_Collision_Stiffness);
-		}
 	}
 
 	void FPBDCollisionConstraint::InitMarginsAndTolerances(const EImplicitObjectType ImplicitType0, const EImplicitObjectType ImplicitType1, const FRealSingle Margin0, const FRealSingle Margin1)
@@ -451,6 +462,42 @@ namespace Chaos
 	void FPBDCollisionConstraint::UpdateMaterialPropertiesImpl()
 	{
 		ConcreteContainer()->UpdateConstraintMaterialProperties(*this);
+
+		UpdateMassScales();
+	}
+
+	void FPBDCollisionConstraint::UpdateMassScales()
+	{
+		if (Flags.bIsOneWayInteraction)
+		{
+			const bool bOneWay0 = FConstGenericParticleHandle(GetParticle0())->OneWayInteraction();
+			const bool bOneWay1 = FConstGenericParticleHandle(GetParticle1())->OneWayInteraction();
+			if (bOneWay0 && !bOneWay1)
+			{
+				Material.InvMassScale1 = 0;
+				Material.InvInertiaScale1 = 0;
+			}
+			if (bOneWay1 && !bOneWay0)
+			{
+				Material.InvMassScale0 = 0;
+				Material.InvInertiaScale0 = 0;
+			}
+			Stiffness = Chaos_Collision_OneWayStiffness;
+		}
+		else
+		{
+			Material.InvMassScale0 = 1;
+			Material.InvInertiaScale0 = 1;
+			Material.InvMassScale1 = 1;
+			Material.InvInertiaScale1 = 1;
+			Stiffness = 1;
+
+			// Debug testing for solver stiffness
+			if (Chaos_Collision_Stiffness >= 0)
+			{
+				SetStiffness(Chaos_Collision_Stiffness);
+			}
+		}
 	}
 
 	bool FPBDCollisionConstraint::IsSleeping() const
diff --git a/Engine/Source/Runtime/Experimental/Chaos/Private/Chaos/Collision/PBDCollisionContainerSolver.cpp b/Engine/Source/Runtime/Experimental/Chaos/Private/Chaos/Collision/PBDCollisionContainerSolver.cpp
--- a/Engine/Source/Runtime/Experimental/Chaos/Private/Chaos/Collision/PBDCollisionContainerSolver.cpp
+++ b/Engine/Source/Runtime/Experimental/Chaos/Private/Chaos/Collision/PBDCollisionContainerSolver.cpp
@@ -155,7 +155,7 @@ namespace Chaos
 		// Initial Phi for initial-overlap depenetration.
 		// If we have an initial contact, calculate the initial overlap. This will get saved in SetSolverResults
 		FRealSingle WorldContactInitialPhi = 0;
-		if (CVars::bChaos_Collision_EnableInitialDepenetration)
+		if (MaxDepentrationVelocity >= 0)
 		{
 			if (ManifoldPoint.Flags.bInitialContact)
 			{
@@ -219,13 +219,17 @@ namespace Chaos
 		Private::FPBDCollisionSolver& Solver, 
 		FPBDCollisionConstraint* Constraint, 
 		const FSolverReal Dt, 
-		const FRealSingle MaxDepentrationVelocity,
 		const int32 ConstraintPointBeginIndex,
-		const int32 ConstraintPointEndIndex)
+		const int32 ConstraintPointEndIndex,
+		const FPBDCollisionSolverSettings& SolverSettings)
 	{
 		const FConstraintSolverBody& Body0 = Solver.SolverBody0();
 		const FConstraintSolverBody& Body1 = Solver.SolverBody1();
 
+		// Negative DepenetrationVelocity means depenetrate immediately
+		const bool bEnableDepenetrationVelocity = CVars::bChaos_Collision_EnableInitialDepenetration && Constraint->GetInitialOverlapDepentrationEnabled();
+		const FSolverReal MaxDepenetrationVelocity = bEnableDepenetrationVelocity ? SolverSettings.DepenetrationVelocity : FSolverReal(-1);
+
 		// Only calculate state for newly added contacts. Normally this is all of them, but maybe not if incremental collision is used by RBAN.
 		// Also we only add active points to the solver's manifold points list
 		for (int32 ConstraintManifoldPointIndex = ConstraintPointBeginIndex; ConstraintManifoldPointIndex < ConstraintPointEndIndex; ++ConstraintManifoldPointIndex)
@@ -237,7 +241,7 @@ namespace Chaos
 				{
 					// Transform the constraint contact data into world space for use by the solver
 					// We build this data directly into the solver's world-space contact data which looks a bit odd with "Init" called after but there you go
-					UpdateCollisionSolverContactPointFromConstraint(Solver, SolverManifoldPointIndex, Constraint, ConstraintManifoldPointIndex, Dt, MaxDepentrationVelocity, Body0, Body1);
+					UpdateCollisionSolverContactPointFromConstraint(Solver, SolverManifoldPointIndex, Constraint, ConstraintManifoldPointIndex, Dt, MaxDepenetrationVelocity, Body0, Body1);
 				}
 			}
 		}
@@ -293,7 +297,7 @@ namespace Chaos
 
 		bOutPerIterationCollision = (!Constraint->GetUseManifold() || Constraint->GetUseIncrementalCollisionDetection());
 
-		UpdateCollisionSolverManifoldFromConstraint(Solver, Constraint, Dt, SolverSettings.DepenetrationVelocity, 0, Constraint->NumManifoldPoints());
+		UpdateCollisionSolverManifoldFromConstraint(Solver, Constraint, Dt, 0, Constraint->NumManifoldPoints(), SolverSettings);
 	}
 
 	FORCEINLINE_DEBUGGABLE void UpdateCollisionConstraintFromSolver(FPBDCollisionConstraint* Constraint, const Private::FPBDCollisionSolver& Solver, const FSolverReal Dt)
@@ -784,8 +788,9 @@ namespace Chaos
 				// Update the manifold based on the new or updated contacts
 				UpdateCollisionSolverManifoldFromConstraint(
 					CollisionSolver, Constraint, 
-					Dt, ConstraintContainer.GetSolverSettings().DepenetrationVelocity, 
-					BeginPointIndex, Constraint->NumManifoldPoints());
+					Dt,
+					BeginPointIndex, Constraint->NumManifoldPoints(),
+					ConstraintContainer.GetSolverSettings());
 
 				Constraint->SetSolverBodies(nullptr, nullptr);
 			}
diff --git a/Engine/Source/Runtime/Experimental/Chaos/Private/Chaos/PBDRigidClustering.cpp b/Engine/Source/Runtime/Experimental/Chaos/Private/Chaos/PBDRigidClustering.cpp
--- a/Engine/Source/Runtime/Experimental/Chaos/Private/Chaos/PBDRigidClustering.cpp
+++ b/Engine/Source/Runtime/Experimental/Chaos/Private/Chaos/PBDRigidClustering.cpp
@@ -70,6 +70,11 @@ namespace Chaos
 	int32 ClusteringParticleReleaseThrottlingMaxCount = INDEX_NONE;
 	FAutoConsoleVariableRef CVarClusteringParticleReleaseThrottlingMaxCount(TEXT("p.Clustering.ParticleReleaseThrottlingMaxCount"), ClusteringParticleReleaseThrottlingMaxCount, TEXT("Maximum number of active geometry collection to reach before all released clustering disable all released particle instantly"));
 
+	bool bClusteringEnableDebris = false;
+	FAutoConsoleVariableRef CVarClusteringEnableDebris(TEXT("p.Clustering.Debris.Enabled"), bClusteringEnableDebris, TEXT(""));
+
+	float ClusteringDebrisSizeThreshold = 200.0f;
+	FAutoConsoleVariableRef CVarClusteringDebrisSizeThreshold(TEXT("p.Clustering.Debris.SizeThreshold"), ClusteringDebrisSizeThreshold, TEXT(""));
 
 	template <typename TProxy=FGeometryCollectionPhysicsProxy>
 	TProxy* GetConcreteProxy(FPBDRigidClusteredParticleHandle* ClusteredParticle)
@@ -810,6 +815,16 @@ namespace Chaos
 			Child->SetPreV(Child->PreV() + ClusteredParent->PreV());
 			Child->SetPreW(Child->PreW() + ClusteredParent->PreW());
 
+			// Experimental
+			if (bClusteringEnableDebris && Child->HasBounds())
+			{
+				const FReal ChildSize = Child->LocalBounds().Extents().GetMax();
+				if (ChildSize < ClusteringDebrisSizeThreshold)
+				{
+					Child->SetOneWayInteraction(true);
+				}
+			}
+
 			// We also need to do cluster book-keeping on the parent.
 			// If the parent is an internal cluster and has become empty, we need to mark this particle as ready to be destroyed.
 			// The only exception is for cluster unions which are managed separately.
diff --git a/Engine/Source/Runtime/Experimental/Chaos/Private/PBDRigidsSolver.cpp b/Engine/Source/Runtime/Experimental/Chaos/Private/PBDRigidsSolver.cpp
--- a/Engine/Source/Runtime/Experimental/Chaos/Private/PBDRigidsSolver.cpp
+++ b/Engine/Source/Runtime/Experimental/Chaos/Private/PBDRigidsSolver.cpp
@@ -215,6 +215,7 @@ namespace Chaos
 				/* InSleepingColor =  */ FColor(128, 0, 0),
 				/* InKinematicColor = */ FColor(255, 0, 0),
 				/* InStaticColor =	  */ FColor(255, 0, 0),
+				/* InDebrisColor =	  */ FColor(255, 0, 0),
 			};
 			return SolverShapesColorsByState_Server;
 		}
@@ -227,6 +228,7 @@ namespace Chaos
 				/* InSleepingColor =  */ FColor(0, 0, 128),
 				/* InKinematicColor = */ FColor(0, 0, 255),
 				/* InStaticColor =	  */ FColor(0, 0, 255),
+				/* InDebrisColor =	  */ FColor(0, 0, 255),
 			};
 			return SolverShapesColorsByState_Client;
 		}
diff --git a/Engine/Source/Runtime/Experimental/Chaos/Public/Chaos/ChaosDebugDraw.h b/Engine/Source/Runtime/Experimental/Chaos/Public/Chaos/ChaosDebugDraw.h
--- a/Engine/Source/Runtime/Experimental/Chaos/Public/Chaos/ChaosDebugDraw.h
+++ b/Engine/Source/Runtime/Experimental/Chaos/Public/Chaos/ChaosDebugDraw.h
@@ -31,15 +31,17 @@ namespace Chaos
 				FColor InDynamicColor,
 				FColor InSleepingColor,
 				FColor InKinematicColor,
-				FColor InStaticColor
+				FColor InStaticColor,
+				FColor InDebrisColor = FColor::Green
 			);
 
 			FColor DynamicColor;
 			FColor SleepingColor;
 			FColor KinematicColor;
 			FColor StaticColor;
+			FColor DebrisColor;
 
-			CHAOS_API FColor GetColorFromState(EObjectStateType State) const;
+			CHAOS_API FColor GetColorFromState(EObjectStateType State, bool bIsDebris = false) const;
 		};
 
 		struct FChaosDebugDrawColorsByShapeType
diff --git a/Engine/Source/Runtime/Experimental/Chaos/Public/Chaos/Collision/PBDCollisionConstraint.h b/Engine/Source/Runtime/Experimental/Chaos/Public/Chaos/Collision/PBDCollisionConstraint.h
--- a/Engine/Source/Runtime/Experimental/Chaos/Public/Chaos/Collision/PBDCollisionConstraint.h
+++ b/Engine/Source/Runtime/Experimental/Chaos/Public/Chaos/Collision/PBDCollisionConstraint.h
@@ -460,6 +460,9 @@ namespace Chaos
 		// NOTE: This is initially set based on whether are allowing incremental manifolds
 		bool GetUseIncrementalCollisionDetection() const { return !Flags.bUseManifold || Flags.bUseIncrementalManifold; }
 
+		// Does this contact support the initial overlap depenetration. If not, bodies are depenetrated in one frame as opposed to at a fixed speed
+		bool GetInitialOverlapDepentrationEnabled() const { return Flags.bOverlapDepenetrationEnabled; }
+
 		/**
 		* Reset the material properties to those from the shape materials. Called each frame to reset contact modifications to the material.
 		*/
@@ -469,6 +472,11 @@ namespace Chaos
 			{
 				Material.ResetMaterialModifications();
 
+				if (Flags.bIsOneWayInteraction)
+				{
+					UpdateMassScales();
+				}
+
 				// Reset other properties which may have changed in contact modification last frame
 				Flags.bIsProbe = BoundsTestFlags.bIsProbe;
 
@@ -855,6 +863,8 @@ namespace Chaos
 
 		CHAOS_API void UpdateMaterialPropertiesImpl();
 
+		CHAOS_API void UpdateMassScales();
+
 	private:
 		CHAOS_API FRealSingle CalculateSavedManifoldPointDistanceSq(const FSavedManifoldPoint& SavedManifoldPoint, const FManifoldPoint& ManifoldPoint) const;
 
@@ -879,6 +889,8 @@ namespace Chaos
 				uint16 bModifierApplied : 1;			// Was a constraint modifier applied this tick
 				uint16 bMaterialSet : 1;				// Has the material been set (or does it need to be reset)
 				uint16 bInitialContact : 1;				// Is this contact considered an initial contact
+				uint16 bOverlapDepenetrationEnabled : 1;// Does this contact support the initial overlap depenetration system
+				uint16 bIsOneWayInteraction : 1;		// Does one of the bodies have the one-way interaction bit set?
 			};
 			uint16 Bits;
 		};
diff --git a/Engine/Source/Runtime/Experimental/Chaos/Public/Chaos/Collision/SpatialAccelerationBroadPhase.h b/Engine/Source/Runtime/Experimental/Chaos/Public/Chaos/Collision/SpatialAccelerationBroadPhase.h
--- a/Engine/Source/Runtime/Experimental/Chaos/Public/Chaos/Collision/SpatialAccelerationBroadPhase.h
+++ b/Engine/Source/Runtime/Experimental/Chaos/Public/Chaos/Collision/SpatialAccelerationBroadPhase.h
@@ -598,6 +598,13 @@ namespace Chaos
 
 				if (Overlap.bCollisionsEnabled)
 				{
+					const bool bIsOneWay0 = FConstGenericParticleHandle(Overlap.Particles[0])->OneWayInteraction();
+					const bool bIsOneWay1 = FConstGenericParticleHandle(Overlap.Particles[1])->OneWayInteraction();
+					if (bIsOneWay0 && bIsOneWay1)
+					{
+						continue;
+					}
+
 					// Get the midphase for this pair
 					FParticlePairMidPhase* MidPhase = ContextAllocator->GetMidPhase(Overlap.Particles[0], Overlap.Particles[1], Overlap.Particles[Overlap.SearchParticleIndex], BroadphaseContext.CollisionContext);
 					BroadphaseContext.MidPhases[MidPhaseIndex] = MidPhase;
@@ -632,10 +639,16 @@ namespace Chaos
 					BroadphaseContext.MidPhases[Index + PrefetchLookahead]->CachePrefetch();
 				}
 
-				// Run MidPhase + NarrowPhase
-				BroadphaseContext.MidPhases[Index]->GenerateCollisions(BroadphaseContext.CollisionContext.GetSettings().BoundsExpansion, Dt, BroadphaseContext.CollisionContext);
+				FParticlePairMidPhase* MidPhase = BroadphaseContext.MidPhases[Index];
+				const FGeometryParticleHandle* Particle0 = MidPhase->GetParticle0();
+				const FGeometryParticleHandle* Particle1 = MidPhase->GetParticle1();
+				if ((Particle0 != nullptr) && (Particle1 != nullptr))
+				{
+					// Run MidPhase + NarrowPhase
+					MidPhase->GenerateCollisions(BroadphaseContext.CollisionContext.GetSettings().BoundsExpansion, Dt, BroadphaseContext.CollisionContext);
+				}
 
-				CVD_TRACE_MID_PHASE(BroadphaseContext.MidPhases[Index]);
+				CVD_TRACE_MID_PHASE(MidPhase);
 			}
 
 			PHYSICS_CSV_CUSTOM_EXPENSIVE(PhysicsCounters, NumFromBroadphase, NumPotentials, ECsvCustomStatOp::Accumulate);
diff --git a/Engine/Source/Runtime/Experimental/Chaos/Public/Chaos/ParticleDirtyFlags.h b/Engine/Source/Runtime/Experimental/Chaos/Public/Chaos/ParticleDirtyFlags.h
--- a/Engine/Source/Runtime/Experimental/Chaos/Public/Chaos/ParticleDirtyFlags.h
+++ b/Engine/Source/Runtime/Experimental/Chaos/Public/Chaos/ParticleDirtyFlags.h
@@ -459,6 +459,9 @@ public:
 	void AddCollisionConstraintFlag(const ECollisionConstraintFlags Flag) { MCollisionConstraintFlag |= uint32(Flag); }
 	void RemoveCollisionConstraintFlag(const ECollisionConstraintFlags Flag) { MCollisionConstraintFlag &= ~uint32(Flag); }
 	
+	bool InitialOverlapDepenetrationEnabled() const { return MControlFlags.GetInitialOverlapDepenetrationEnabled(); }
+	void SetInitialOverlapDepenetrationEnabled(const bool bEnabled) { MControlFlags.SetInitialOverlapDepenetrationEnabled(bEnabled); }
+
 	bool OneWayInteraction() const { return MControlFlags.GetOneWayInteractionEnabled(); }
 	void SetOneWayInteraction(bool bInOneWayInteraction) { MControlFlags.SetOneWayInteractionEnabled(bInOneWayInteraction); }
 
diff --git a/Engine/Source/Runtime/Experimental/Chaos/Public/Chaos/ParticleHandle.h b/Engine/Source/Runtime/Experimental/Chaos/Public/Chaos/ParticleHandle.h
--- a/Engine/Source/Runtime/Experimental/Chaos/Public/Chaos/ParticleHandle.h
+++ b/Engine/Source/Runtime/Experimental/Chaos/Public/Chaos/ParticleHandle.h
@@ -113,6 +113,7 @@ void PBDRigidParticleDefaultConstruct(FConcrete& Concrete, const FPBDRigidPartic
 	Concrete.SetLinearEtherDrag(0.f);
 	Concrete.SetAngularEtherDrag(0.f);
 	Concrete.SetGravityEnabled(Params.bGravityEnabled);
+	Concrete.SetInitialOverlapDepenetrationEnabled(true);
 	Concrete.SetCCDEnabled(Params.bCCDEnabled);
 	Concrete.SetDisabled(Params.bDisabled);
 	Concrete.SetSleepType(ESleepType::MaterialSleep);
@@ -1230,6 +1231,16 @@ public:
 		PBDRigidParticles->ControlFlags(ParticleIdx).SetCCDEnabled(bEnabled);
 	}
 
+	bool InitialOverlapDepenetrationEnabled() const 
+	{ 
+		return ControlFlags().GetInitialOverlapDepenetrationEnabled();
+	}
+
+	void SetInitialOverlapDepenetrationEnabled(const bool bEnabled)
+	{ 
+		PBDRigidParticles->ControlFlags(ParticleIdx).SetInitialOverlapDepenetrationEnabled(bEnabled);
+	}
+
 	inline bool OneWayInteraction() const
 	{ 
 		return ControlFlags().GetOneWayInteractionEnabled();
@@ -2219,6 +2230,23 @@ public:
 		return FReal(0);
 	}
 
+	bool InitialOverlapDepentrationEnabled() const
+	{
+		if (auto RigidHandle = MHandle->CastToRigidParticle())
+		{
+			return RigidHandle->InitialOverlapDepenetrationEnabled();
+		}
+		return true;
+	}
+
+	bool OneWayInteraction() const
+	{
+		if (auto RigidHandle = MHandle->CastToRigidParticle())
+		{
+			return RigidHandle->OneWayInteraction();
+		}
+		return false;
+	}
 
 #if CHAOS_DEBUG_NAME
 	const TSharedPtr<FString, ESPMode::ThreadSafe>& DebugName() const
@@ -3075,6 +3103,12 @@ public:
 		MMiscData.Modify(true, MDirtyFlags, Proxy, [bUpdateKinematicFromSimulation](auto& Data) { Data.SetUpdateKinematicFromSimulation(bUpdateKinematicFromSimulation); });
 	}
 
+	bool InitialOverlapDepenetrationEnabled() const { return MMiscData.Read().InitialOverlapDepenetrationEnabled(); }
+	void SetInitialOverlapDepenetrationEnabled(const bool bInEnabled)
+	{
+		MMiscData.Modify(true, MDirtyFlags, Proxy, [bInEnabled](auto& Data) { Data.SetInitialOverlapDepenetrationEnabled(bInEnabled); });
+	}
+
 	bool OneWayInteraction() const { return MMiscData.Read().OneWayInteraction(); }
 	void SetOneWayInteraction(const bool bInEnabled)
 	{
diff --git a/Engine/Source/Runtime/Experimental/Chaos/Public/Chaos/RigidParticleControlFlags.h b/Engine/Source/Runtime/Experimental/Chaos/Public/Chaos/RigidParticleControlFlags.h
--- a/Engine/Source/Runtime/Experimental/Chaos/Public/Chaos/RigidParticleControlFlags.h
+++ b/Engine/Source/Runtime/Experimental/Chaos/Public/Chaos/RigidParticleControlFlags.h
@@ -47,9 +47,11 @@ namespace Chaos
 		bool GetOneWayInteractionEnabled() const { return Flags.bOneWayInteractionEnabled; }
 		FRigidParticleControlFlags& SetOneWayInteractionEnabled(const bool bEnabled) { Flags.bOneWayInteractionEnabled = bEnabled; return *this; }
 
-		bool GetMaxDepenetrationVelocityOverrideEnabled() const { return Flags.bMaxDepenetrationVelocityOverrideEnabled; }
-		FRigidParticleControlFlags& SetMaxDepenetrationVelocityOverrideEnabled(const bool bEnabled) { Flags.bMaxDepenetrationVelocityOverrideEnabled = bEnabled; return *this; }
+		// If enabled initial overlaps will be depenetrated at a rate that depends on physics settings
+		bool GetInitialOverlapDepenetrationEnabled() const { return Flags.bEnableInitialOverlapDepenetration; }
+		FRigidParticleControlFlags& SetInitialOverlapDepenetrationEnabled(const bool bEnabled) { Flags.bEnableInitialOverlapDepenetration = bEnabled; return *this; }
 
+		// If enabled, inertia may be increased to improve stability
 		bool GetInertiaConditioningEnabled() const { return Flags.bInertiaConditioningEnabled; }
 		FRigidParticleControlFlags& SetInertiaConditioningEnabled(const bool bEnabled) { Flags.bInertiaConditioningEnabled = bEnabled; return *this; }
 
@@ -63,13 +65,19 @@ namespace Chaos
 		// Serialization
 		friend FChaosArchive& operator<<(FChaosArchive& Ar, FRigidParticleControlFlags& Data);
 
+		UE_DEPRECATED(5.4, "Not used")
+		bool GetMaxDepenetrationVelocityOverrideEnabled() const { return false; }
+
+		UE_DEPRECATED(5.4, "Not used")
+		FRigidParticleControlFlags& SetMaxDepenetrationVelocityOverrideEnabled(const bool bEnabled) { return *this; }
+
 	private:
 		struct FFlags
 		{
 			FStorage bGravityEnabled : 1;
 			FStorage bCCDEnabled : 1;
 			FStorage bOneWayInteractionEnabled : 1;
-			FStorage bMaxDepenetrationVelocityOverrideEnabled : 1;
+			FStorage bEnableInitialOverlapDepenetration : 1;
 			FStorage bInertiaConditioningEnabled : 1;
 			FStorage GravityGroupIndex : 3;
 			FStorage bUpdateKinematicFromSimulation : 1;
diff --git a/Engine/Source/Runtime/Experimental/ChaosVisualDebugger/Private/DataWrappers/ChaosVDParticleDataWrapper.cpp b/Engine/Source/Runtime/Experimental/ChaosVisualDebugger/Private/DataWrappers/ChaosVDParticleDataWrapper.cpp
--- a/Engine/Source/Runtime/Experimental/ChaosVisualDebugger/Private/DataWrappers/ChaosVDParticleDataWrapper.cpp
+++ b/Engine/Source/Runtime/Experimental/ChaosVisualDebugger/Private/DataWrappers/ChaosVDParticleDataWrapper.cpp
@@ -11,7 +11,7 @@ bool FChaosVDFRigidParticleControlFlags::Serialize(FArchive& Ar)
 	Ar << bGravityEnabled;
 	Ar << bCCDEnabled;
 	Ar << bOneWayInteractionEnabled;
-	Ar << bMaxDepenetrationVelocityOverrideEnabled;
+	Ar << bEnableInitialOverlapDepenetration;
 	Ar << bInertiaConditioningEnabled;
 	Ar << GravityGroupIndex;
 
diff --git a/Engine/Source/Runtime/Experimental/ChaosVisualDebugger/Public/DataWrappers/ChaosVDParticleDataWrapper.h b/Engine/Source/Runtime/Experimental/ChaosVisualDebugger/Public/DataWrappers/ChaosVDParticleDataWrapper.h
--- a/Engine/Source/Runtime/Experimental/ChaosVisualDebugger/Public/DataWrappers/ChaosVDParticleDataWrapper.h
+++ b/Engine/Source/Runtime/Experimental/ChaosVisualDebugger/Public/DataWrappers/ChaosVDParticleDataWrapper.h
@@ -63,8 +63,9 @@ struct FChaosVDFRigidParticleControlFlags : public FChaosVDParticleDataBase
 		: bGravityEnabled(false),
 		  bCCDEnabled(false),
 		  bOneWayInteractionEnabled(false),
-		  bMaxDepenetrationVelocityOverrideEnabled(false),
-		  bInertiaConditioningEnabled(false), GravityGroupIndex(0)
+		  bEnableInitialOverlapDepenetration(false),
+		  bInertiaConditioningEnabled(false), 
+		  GravityGroupIndex(0)
 	{
 	}
 
@@ -76,7 +77,7 @@ struct FChaosVDFRigidParticleControlFlags : public FChaosVDParticleDataBase
 		bGravityEnabled = Other.GetGravityEnabled();
 		bCCDEnabled = Other.GetCCDEnabled();
 		bOneWayInteractionEnabled = Other.GetOneWayInteractionEnabled();
-		bMaxDepenetrationVelocityOverrideEnabled = Other.GetMaxDepenetrationVelocityOverrideEnabled();
+		bEnableInitialOverlapDepenetration = Other.GetInitialOverlapDepenetrationEnabled();
 		bInertiaConditioningEnabled = Other.GetInertiaConditioningEnabled();
 		GravityGroupIndex = Other.GetGravityGroupIndex();
 
@@ -89,9 +90,9 @@ struct FChaosVDFRigidParticleControlFlags : public FChaosVDParticleDataBase
 	bool bCCDEnabled;
 	UPROPERTY(EditAnywhere, Category= "Particle Control Flags")
 	bool bOneWayInteractionEnabled;
-	UPROPERTY(EditAnywhere, Category= "Particle Control Flags")
-	bool bMaxDepenetrationVelocityOverrideEnabled;
-	UPROPERTY(EditAnywhere, Category= "Particle Control Flags")
+	UPROPERTY(EditAnywhere, Category = "Particle Control Flags")
+	bool bEnableInitialOverlapDepenetration;
+	UPROPERTY(EditAnywhere, Category = "Particle Control Flags")
 	bool bInertiaConditioningEnabled;
 	UPROPERTY(EditAnywhere, Category= "Particle Control Flags")
 	int32 GravityGroupIndex;
