Add the possibility for a geometry component to stop sending update to the custom renderer if any set
- Also make sure the ISM Pool sub components are properly attached to the root component

diff --git a/Engine/Source/Runtime/Experimental/GeometryCollectionEngine/Private/GeometryCollection/GeometryCollectionComponent.cpp b/Engine/Source/Runtime/Experimental/GeometryCollectionEngine/Private/GeometryCollection/GeometryCollectionComponent.cpp
--- a/Engine/Source/Runtime/Experimental/GeometryCollectionEngine/Private/GeometryCollection/GeometryCollectionComponent.cpp
+++ b/Engine/Source/Runtime/Experimental/GeometryCollectionEngine/Private/GeometryCollection/GeometryCollectionComponent.cpp
@@ -622,6 +622,7 @@ UGeometryCollectionComponent::UGeometryCollectionComponent(const FObjectInitiali
 	, bIsTransformSelectionModeEnabled(false)
 #endif  // #if GEOMETRYCOLLECTION_EDITOR_SELECTION
 	, bIsMoving(false)
+	, bUpdateCustomRenderer(true)
 {
 	// by default tick is registered but disabled, we only need it when we need to update the removal timers
 	// tick will be then enabled only when the root is broken from OnPostPhysicsSync callback
@@ -869,8 +870,10 @@ namespace
 	template <typename TTransformType>
 	inline FBox ComputeBoundsFromTransforms(const UGeometryCollectionComponent& Component, const TTransformType& LocalToWorldWithScale, const TArray<TTransformType>& GlobalMatricesArray)
 	{
+		static FName BoundingBoxAttributeName = "BoundingBox";
+
 		auto GeometryCollectionPtr = Component.GetRestCollection()->GetGeometryCollection();
-		const TManagedArray<FBox>* TransformBoundingBoxes = GeometryCollectionPtr->FindAttribute<FBox>("BoundingBox", FGeometryCollection::TransformGroup);
+		const TManagedArray<FBox>* TransformBoundingBoxes = GeometryCollectionPtr->FindAttribute<FBox>(BoundingBoxAttributeName, FGeometryCollection::TransformGroup);
 		const TManagedArray<FBox>& GeometryBoundingBoxes = Component.GetBoundingBoxArray();
 		const TManagedArray<int32>& TransformToGeometryIndex = Component.GetTransformToGeometryIndexArray();
 
@@ -5727,7 +5730,7 @@ void UGeometryCollectionComponent::RefreshCustomRenderer()
 	if (CanUseCustomRenderer())
 	{
 		// Don't refresh the custom renderer on the server but we still need to do the work of computing component space transforms.
-		const bool bUpdateRenderer = !IsNetMode(NM_DedicatedServer);
+		const bool bUpdateRenderer = !IsNetMode(NM_DedicatedServer) && bUpdateCustomRenderer;
 
 		if (IGeometryCollectionExternalRenderInterface* RendererInterface = CustomRenderer.GetInterface())
 		{
@@ -6193,6 +6196,11 @@ FTransform UGeometryCollectionComponent::GetRootCurrentTransform() const
 	return RootInitialTransform;
 }
 
+FTransform UGeometryCollectionComponent::GetRootCurrentComponentSpaceTransform() const
+{
+	return (RestCollection) ? ComponentSpaceTransforms.RequestRootTransform() : FTransform::Identity;
+}
+
 TArray<FTransform> UGeometryCollectionComponent::GetInitialLocalRestTransforms() const
 {
 	TArray<FTransform> InitialLocalTransforms;
diff --git a/Engine/Source/Runtime/Experimental/GeometryCollectionEngine/Private/GeometryCollection/GeometryCollectionISMPoolComponent.cpp b/Engine/Source/Runtime/Experimental/GeometryCollectionEngine/Private/GeometryCollection/GeometryCollectionISMPoolComponent.cpp
--- a/Engine/Source/Runtime/Experimental/GeometryCollectionEngine/Private/GeometryCollection/GeometryCollectionISMPoolComponent.cpp
+++ b/Engine/Source/Runtime/Experimental/GeometryCollectionEngine/Private/GeometryCollection/GeometryCollectionISMPoolComponent.cpp
@@ -86,6 +86,7 @@ void FGeometryCollectionISM::CreateISM(AActor* InOwningActor, bool bInUseHISM)
 	ISMComponent->SetRemoveSwap();
 	ISMComponent->SetCanEverAffectNavigation(false);
 	ISMComponent->SetCollisionEnabled(ECollisionEnabled::NoCollision);
+	ISMComponent->SetupAttachment(InOwningActor->GetRootComponent());
 
 	InOwningActor->AddInstanceComponent(ISMComponent);
 	ISMComponent->RegisterComponent();
diff --git a/Engine/Source/Runtime/Experimental/GeometryCollectionEngine/Public/GeometryCollection/GeometryCollectionComponent.h b/Engine/Source/Runtime/Experimental/GeometryCollectionEngine/Public/GeometryCollection/GeometryCollectionComponent.h
--- a/Engine/Source/Runtime/Experimental/GeometryCollectionEngine/Public/GeometryCollection/GeometryCollectionComponent.h
+++ b/Engine/Source/Runtime/Experimental/GeometryCollectionEngine/Public/GeometryCollection/GeometryCollectionComponent.h
@@ -783,6 +783,8 @@ public:
 	UFUNCTION(BlueprintCallable, Category = "ChaosPhysics")
 	GEOMETRYCOLLECTIONENGINE_API FTransform GetRootCurrentTransform() const;
 
+	GEOMETRYCOLLECTIONENGINE_API FTransform GetRootCurrentComponentSpaceTransform() const;
+
 	/** return true if the root cluster is not longer active at runtime */
 	UFUNCTION(BlueprintCallable, Category = "ChaosPhysics")
 	GEOMETRYCOLLECTIONENGINE_API bool IsRootBroken() const;
@@ -1280,6 +1282,9 @@ public:
 	/** Force all GC components to reregister their custom renderer objects. */
 	static GEOMETRYCOLLECTIONENGINE_API void ReregisterAllCustomRenderers();
 
+	/** allow update of the custom renderer if enabled */
+	GEOMETRYCOLLECTIONENGINE_API void SetUpdateCustomRenderer(bool bValue) { bUpdateCustomRenderer = bValue; }
+
 public:
 	UPROPERTY(BlueprintAssignable, Category = "Collision")
 	FOnChaosPhysicsCollision OnChaosPhysicsCollision;
@@ -1689,6 +1694,8 @@ private:
 	/** True if GeometryCollection transforms have changed from previous tick. */
 	bool bIsMoving;
 
+	bool bUpdateCustomRenderer;
+
 	//~ Begin IPhysicsComponent Interface.
 public:
 	GEOMETRYCOLLECTIONENGINE_API virtual Chaos::FPhysicsObject* GetPhysicsObjectById(Chaos::FPhysicsObjectId Id) const override;
