Pipeline IAS requests.

Currently the length is hard coded to 1. A subsequent change will expose that fully. Now-unused code has been left in or #if/def'd out for now with the exception of ReceiveWork().

#rnx

diff --git a/Engine/Source/Runtime/Experimental/IoStoreOnDemand/Private/CoreHttp/Http.cpp b/Engine/Source/Runtime/Experimental/IoStoreOnDemand/Private/CoreHttp/Http.cpp
--- a/Engine/Source/Runtime/Experimental/IoStoreOnDemand/Private/CoreHttp/Http.cpp
+++ b/Engine/Source/Runtime/Experimental/IoStoreOnDemand/Private/CoreHttp/Http.cpp
@@ -1612,6 +1612,7 @@ struct alignas(16) FActivity
 		_Num,
 	};
 
+	FActivity*			Next = nullptr;
 	int8				Slot = -1;
 	EState				State = EState::None;
 	EWait				SocketWait = EWait::None;
@@ -1622,8 +1623,8 @@ struct alignas(16) FActivity
 #if IAS_HTTP_WITH_PERF
 	FStopwatch			Stopwatch;
 #endif
-	FSocketPool*		Pool;
 	union {
+		FSocketPool*	Pool;
 		FIoBuffer*		Dest;
 		const char*		ErrorReason;
 	};
@@ -1687,6 +1688,7 @@ static void Activity_ChangeState(FActivity* Activity, FActivity::EState InState,
 ////////////////////////////////////////////////////////////////////////////////
 static void Activity_BeginWait(FActivity* Activity, FActivity::EWait What)
 {
+#if 0
 	check(Activity->SocketWait == FActivity::EWait::None);
 
 #if IAS_HTTP_WITH_PERF
@@ -1698,11 +1700,13 @@ static void Activity_BeginWait(FActivity* Activity, FActivity::EWait What)
 	Trace(Activity, ETrace::Wait);
 
 	Activity->SocketWait = What;
+#endif // 0
 }
 
 ////////////////////////////////////////////////////////////////////////////////
 static void Activity_EndWait(FActivity* Activity)
 {
+#if 0
 	if (Activity->SocketWait != FActivity::EWait::None)
 	{
 #if IAS_HTTP_WITH_PERF
@@ -1715,6 +1719,7 @@ static void Activity_EndWait(FActivity* Activity)
 	}
 
 	Activity->SocketWait = FActivity::EWait::None;
+#endif // 0
 }
 
 ////////////////////////////////////////////////////////////////////////////////
@@ -1760,6 +1765,7 @@ static FActivity* Activity_Alloc(uint32 BufferSize)
 ////////////////////////////////////////////////////////////////////////////////
 static void Activity_Free(FActivity* Activity)
 {
+#if 0
 	if (Activity->State > FActivity::EState::Connect)
 	{
 		FSocket& Socket = Activity->Socket;
@@ -1773,6 +1779,7 @@ static void Activity_Free(FActivity* Activity)
 			Activity->Pool->ReturnLease(MoveTemp(Socket));
 		}
 	}
+#endif // 0
 
 	Activity->~FActivity();
 	FMemory::Free(Activity);
@@ -2242,12 +2249,12 @@ static int32 DoSend(FActivity* Activity, FSocket& Socket)
 
 	switch (FSocket::EResult(Result))
 	{
-	case FSocket::EResult::HangUp:		Activity_SetError(Activity, "ATH0.Send"); return -1;
-	case FSocket::EResult::Error:		Activity_SetError(Activity, "Error returned from socket send"); return -1;
-	case FSocket::EResult::ConnectError:Activity_SetError(Activity, "Connection error"); return -1;
+	case FSocket::EResult::HangUp:		Activity_SetError(Activity, "ATH0.Send"); return Result;
+	case FSocket::EResult::Error:		Activity_SetError(Activity, "Error returned from socket send"); return Result;
+	case FSocket::EResult::ConnectError:Activity_SetError(Activity, "Connection error"); return Result;
 	case FSocket::EResult::Wait:
 		Activity_BeginWait(Activity, FActivity::EWait::Write);
-		return 1;
+		return Result;
 	}
 
 #if IAS_HTTP_WITH_PERF
@@ -2263,7 +2270,7 @@ static int32 DoSend(FActivity* Activity, FSocket& Socket)
 
 	Activity_ChangeState(Activity, FActivity::EState::RecvMessage, Buffer.GetSize());
 	Activity_BeginWait(Activity, FActivity::EWait::Read);
-	return 1;
+	return Result;
 }
 
 ////////////////////////////////////////////////////////////////////////////////
@@ -2615,6 +2622,16 @@ static void DoCancel(FActivity* Activity)
 	Activity->Sink(SinkArg);
 }
 
+////////////////////////////////////////////////////////////////////////////////
+static void DoFail(FActivity* Activity)
+{
+	check(Activity->State == FActivity::EState::Failed);
+
+	// Notify the user we've received everything
+	FTicketStatus& SinkArg = *(FTicketStatus*)Activity;
+	Activity->Sink(SinkArg);
+}
+
 
 
 // {{{1 throttler ..............................................................
@@ -2723,6 +2740,506 @@ void FThrottler::ReturnUnused(uint32 Unused)
 
 
 
+// {{{1 groups .................................................................
+
+/*
+ * - Activities (requests send with a loop) are managed in singly-linked lists
+ * - Each activity has an associated host it is talking to.
+ * - Hosts are emphemeral, or represented externally via a FConnectionPool object
+ * - Loop has a group for each host, and each host-group has a bunch of socket-groups
+ * - Socket-groups manage the list of activities to which they are assigned.
+ *
+ *  Loop
+ *    FHostGroup[HostPtr]
+ *      FSocketGroup[0...HostMaxConnections]
+ *	      Act0 -> Act1 -> Act2 -> Act3 -> ...
+ */
+
+////////////////////////////////////////////////////////////////////////////////
+struct FTickState
+{
+	FActivity*					DoneList;
+	uint64						Cancels;
+	int32&						RecvAllowance;
+	uint32						PollTimeoutMs;
+	TArray<FSocket::FWaiter>*	Waiters;
+};
+
+////////////////////////////////////////////////////////////////////////////////
+class FSocketGroup
+{
+public:
+						FSocketGroup(uint32 PipeLen) : PipelineLength(uint8(PipeLen)) {}
+						~FSocketGroup();
+	bool				operator == (FSocket* Rhs) const { return &Socket == Rhs; }
+	bool				IsBusy() const		{ return IsKeepAlive || List != nullptr; }
+	bool				IsWaiting() const	{ return bWaiting; }
+	void				Unwait()			{ check(bWaiting); bWaiting = false; }
+	void				AddActivity(FActivity* Activity);
+	void				Tick(FTickState& State, FHost& Host);
+	void				Fail(FTickState& State, const char* Reason);
+
+private:
+	void				TickRecv(FTickState& State);
+	void				TickSend(FTickState& State, FActivity* Activity);
+	void				TickCancels(FTickState& State);
+	void				Done(FTickState& State, FActivity* Activity);
+	FSocket				Socket;
+	FActivity*			List = nullptr;
+	FActivity*			ListTail = nullptr;
+	uint64				ActiveSlots = 0;
+	uint8				IsKeepAlive = 1;
+	bool				bWaiting = false;
+	uint8				NumActive = 0;
+	uint8				PipelineLength;
+
+	UE_NONCOPYABLE(FSocketGroup);
+};
+
+////////////////////////////////////////////////////////////////////////////////
+FSocketGroup::~FSocketGroup()
+{
+	check(List == nullptr);
+	check(ListTail == nullptr);
+	check(NumActive == 0);
+	check(ActiveSlots == 0);
+}
+
+////////////////////////////////////////////////////////////////////////////////
+void FSocketGroup::Done(FTickState& State, FActivity* Activity)
+{
+	Activity->Next = State.DoneList;
+	State.DoneList = Activity;
+
+	ActiveSlots ^= 1ull << Activity->Slot;
+}
+
+////////////////////////////////////////////////////////////////////////////////
+void FSocketGroup::Fail(FTickState& State, const char* Reason)
+{
+	// Failure is quite terminal and we need to abort all listed activities.
+
+	check(List != nullptr);
+	FActivity* Activity = List;
+	do
+	{
+		FActivity* Next = Activity->Next;
+		if (Activity->State != FActivity::EState::Failed)
+		{
+			Activity_SetError(Activity, Reason);
+		}
+
+		DoFail(Activity);
+		Done(State, Activity);
+
+		Activity = Next;
+	}
+	while (Activity != nullptr);
+	check(ActiveSlots == 0);
+
+	Socket = FSocket();
+	NumActive = 0;
+	List = ListTail = nullptr;
+}
+
+////////////////////////////////////////////////////////////////////////////////
+void FSocketGroup::TickRecv(FTickState& State)
+{
+	// A helper lambda to use for iterating the activity list
+	auto Next = [this, Iter=List] () mutable -> FActivity*
+	{
+		FActivity* Activity = Iter;
+		if (Activity != nullptr)
+		{
+			Iter = Activity->Next;
+		}
+		return Activity;
+	};
+
+	// Another helper lambda
+	auto IsReceiving = [] (const FActivity* Act)
+	{
+		using EState = FActivity::EState;
+		return (Act->State >= EState::RecvMessage) & (Act->State < EState::RecvDone);
+	};
+
+	FActivity* Activity = Next();
+
+	// All receives first
+	while ((Activity != nullptr) && IsReceiving(Activity))
+	{
+		int32 Result = DoRecv(Activity, Socket, State.RecvAllowance);
+
+		// Any sort of error here is unrecoverable
+		if (Result < 0)
+		{
+			Fail(State, Activity->ErrorReason);
+			return;
+		}
+
+		IsKeepAlive &= Activity->IsKeepAlive;
+
+		// If there was no data available this is far as receiving can go
+		if (Result > 0)
+		{
+			bWaiting = true;
+			State.Waiters->Emplace(Socket, FSocket::FWaiter::EWhat::Recv);
+
+			Activity = Next();
+			for (; (Activity != nullptr) && IsReceiving(Activity); Activity = Next());
+			break;
+		}
+
+		// If we're still in a receiving state we will just try again otherwise it
+		// is finished and we will let DoneList recipient finish it off.
+		if (IsReceiving(Activity))
+		{
+			continue;
+		}
+
+		check(Activity == List);
+		if ((List = Activity->Next) == nullptr)
+		{
+			ListTail = nullptr;
+		}
+		--NumActive;
+
+		DoRecvDone(Activity);
+		Done(State, Activity);
+
+		// If the server wants to close the socket we need to rewind subsequent
+		// activities. Some may have already been sent down the pipeline
+		if (IsKeepAlive == 0)
+		{
+			while ((Activity = Next()) != nullptr)
+			{
+				if (int32 Rewound = Activity_Rewind(Activity); Rewound >= 0)
+				{
+					NumActive -= !!Rewound;
+					continue;
+				}
+
+				Fail(State, "Unable to rewind on keep-alive close");
+				break;
+			}
+
+			Socket = FSocket();
+			return;
+		}
+		Activity = Next();
+	}
+
+	if (IsKeepAlive && Activity != nullptr)
+	{
+		return TickSend(State, Activity);
+	}
+}
+
+////////////////////////////////////////////////////////////////////////////////
+void FSocketGroup::TickSend(FTickState& State, FActivity* Activity)
+{
+	check(IsKeepAlive == 1);
+
+	for (FActivity* Next; Activity != nullptr; Activity = Next)
+	{
+		check(Activity->State == FActivity::EState::Send);
+
+		// If we have maxed out the pipeline there's nothing more we can do
+		if (NumActive == PipelineLength)
+		{
+			break;
+		}
+
+		int32 Result = DoSend(Activity, Socket);
+
+		if (Result == int32(FSocket::EResult::Wait))
+		{
+			// For now we'll not add the socket as a waiter. It is unlikely that
+			// we send enough to need to wait currently.
+			break;
+		}
+
+		Next = Activity->Next;
+
+		if (Result < 0)
+		{
+			Fail(State, Activity->ErrorReason);
+			return;
+		}
+
+		++NumActive;
+	}
+}
+
+////////////////////////////////////////////////////////////////////////////////
+void FSocketGroup::TickCancels(FTickState& State)
+{
+	if (State.Cancels == 0 || (State.Cancels & ActiveSlots) == 0)
+	{
+		return;
+	}
+
+	// We are going to rebuild the list ot activities to maintain order as the
+	// activity list is singular.
+
+	check(List != nullptr);
+	FActivity* Activity = List;
+	List = ListTail = nullptr;
+
+	{
+		FActivity* Next = Activity->Next;
+		if (uint64 Slot = 1ull << Activity->Slot; State.Cancels & Slot)
+		{
+			if (Activity->State > FActivity::EState::Send)
+			{
+				ActiveSlots ^= Slot;
+				--NumActive;
+			}
+			DoCancel(Activity);
+			Done(State, Activity);
+		}
+		else
+		{
+			AddActivity(Activity);
+		}
+
+		Activity = Next;
+	}
+	while (Activity != nullptr);
+}
+
+////////////////////////////////////////////////////////////////////////////////
+void FSocketGroup::Tick(FTickState& State, FHost& Host)
+{
+	// Nothing to tick if we've no work or are still waiting on read/write
+	if (List == nullptr || bWaiting)
+	{
+		return;
+	}
+
+	TickCancels(State);
+
+	if (Socket.IsValid())
+	{
+		return TickRecv(State);
+	}
+
+	// We don't have a connected socket on first use, or if a keep-alive:close
+	// was received from the server. So we connect here.
+	FResult Result;
+	if (!Socket.IsValid())
+	{
+		IsKeepAlive = 1;
+		Result = Host.Connect(Socket);
+
+		// Non-blocking connect
+		if (Result.GetValue() == 0)
+		{
+			bWaiting = true;
+			State.Waiters->Emplace(Socket, FSocket::FWaiter::EWhat::Send);
+			return;
+		}
+
+		// Completed connect
+		if (Result.GetValue() > 0)
+		{
+			return TickSend(State, List);
+		}
+
+		IsKeepAlive = 0;
+	}
+
+	// Fail all this group's work as no connection was established
+	Fail(State, Result.GetMessage());
+}
+
+////////////////////////////////////////////////////////////////////////////////
+void FSocketGroup::AddActivity(FActivity* Activity)
+{
+	// We use a tail pointer here to maintain order that requests were made
+
+	check(Activity->Next == nullptr);
+
+	if (ListTail != nullptr)
+	{
+		ListTail->Next = Activity;
+	}
+	List = (List == nullptr) ? Activity : List;
+	ListTail = Activity;
+
+	ActiveSlots |= (1ull << Activity->Slot);
+}
+
+
+
+////////////////////////////////////////////////////////////////////////////////
+class FHostGroup
+{
+public:
+								FHostGroup(FHost& InHost) : Host(InHost) {}
+	bool						IsBusy() const	{ return !SocketGroups.IsEmpty(); }
+	const FHost&				GetHost() const	{ return Host; }
+	void						Tick(FTickState& State);
+	void						AddActivity(FActivity* Activity);
+
+private:
+	using FSocketGroupPtr = TUniquePtr<FSocketGroup>;
+	using FSocketGroups = TArray<FSocketGroupPtr>;
+	void						ScatterPendings();
+	int32						Wait(uint32 PollTimeoutMs);
+	TArray<FSocket::FWaiter>	Waiters;
+	FSocketGroups				SocketGroups;
+	FHost&						Host;
+	FActivity*					PendingList = nullptr;
+	uint32						ScatterIndex = 0;
+
+	UE_NONCOPYABLE(FHostGroup);
+};
+
+////////////////////////////////////////////////////////////////////////////////
+int32 FHostGroup::Wait(uint32 PollTimeoutMs)
+{
+	if (Waiters.Num() == 0)
+	{
+		return 0;
+	}
+
+	// Actually do the wait
+	int32 Result = FSocket::Wait(Waiters, PollTimeoutMs);
+	if (Result <= 0)
+	{
+		return Result;
+	}
+
+	// For each waiter that's ready, find the associated group "unwait" them.
+	int32 Count = 0;
+	for (int32 i = 0, n = Waiters.Num(); i < n; ++i)
+	{
+		if (Waiters[i].Ready == 0)
+		{
+			continue;
+		}
+
+		auto* Candidate = (FSocket*)(Waiters[i].Candidate);
+		auto Pred = [Candidate] (auto& Lhs) { return *Lhs == Candidate; };
+		FSocketGroupPtr* Group = SocketGroups.FindByPredicate(Pred);
+		check(Group != nullptr);
+		(*Group)->Unwait();
+
+		Waiters.RemoveAtSwap(i);
+		--n, --i, ++Count;
+	}
+	check(Count == Result);
+
+	return Result;
+}
+
+////////////////////////////////////////////////////////////////////////////////
+void FHostGroup::Tick(FTickState& State)
+{
+	ScatterPendings();
+
+	// Do any waits left over from the last pass through
+	if (int32 Result = Wait(State.PollTimeoutMs); Result < 0)
+	{
+		for (FSocketGroupPtr& Group : SocketGroups)
+		{
+			Group->Fail(State, "poll() returned an unexpected error");
+		}
+	}
+
+	// If we are still waiting on all sockets groups there is nothing to do.
+	if (Waiters.Num() == SocketGroups.Num())
+	{
+		return;
+	}
+
+	// Tick everything.
+	State.Waiters = &Waiters;
+	for (FSocketGroupPtr& Group : SocketGroups)
+	{
+		Group->Tick(State, Host);
+	}
+
+	// Clear out any spent socket groups.
+	for (uint32 i = 0, n = SocketGroups.Num(); i < n; ++i)
+	{
+		FSocketGroupPtr& Group = SocketGroups[i];
+		if (Group->IsBusy())
+		{
+			continue;
+		}
+
+		check(!Group->IsWaiting());
+		SocketGroups.RemoveAtSwap(i);
+		--n, --i;
+	}
+}
+
+////////////////////////////////////////////////////////////////////////////////
+void FHostGroup::ScatterPendings()
+{
+	FActivity* Iter = PendingList;
+	if (Iter == nullptr)
+	{
+		return;
+	}
+	PendingList = nullptr;
+
+	// Maximize the number of connections to the host first
+	for (uint32 i = SocketGroups.Num(), n = Host.GetMaxConnections(); i < n; ++i)
+	{
+		FActivity* Activity = Iter;
+		Iter = Iter->Next;
+		Activity->Next = nullptr;
+
+		FSocketGroupPtr& Group = SocketGroups.Emplace_GetRef();
+		Group = MakeUnique<FSocketGroup>(1);
+		Group->AddActivity(Activity);
+
+		if (Iter == nullptr)
+		{
+			return;
+		}
+	}
+
+	// Next we'll favour socket groups that have come back for more.
+	/*
+	 * ...not currently implemented
+	 *
+	do
+	{
+	}
+	while (Iter != PendingList.end());
+	*/
+
+	// Finally we'll simply round-robin distribute items
+	for (uint32 Num = SocketGroups.Num(); Iter != nullptr; ++ScatterIndex)
+	{
+		if (ScatterIndex >= Num)
+		{
+			ScatterIndex = 0;
+		}
+
+		FActivity* Activity = Iter;
+		Iter = Iter->Next;
+		Activity->Next = nullptr;
+
+		FSocketGroupPtr& Group = SocketGroups[ScatterIndex];
+		Group->AddActivity(Activity);
+	}
+}
+
+////////////////////////////////////////////////////////////////////////////////
+void FHostGroup::AddActivity(FActivity* Activity)
+{
+	check(Activity->Next == nullptr);
+	Activity->Next = PendingList;
+	PendingList = Activity;
+}
+
+
+
 // {{{1 event-loop .............................................................
 
 ////////////////////////////////////////////////////////////////////////////////
@@ -2746,18 +3263,16 @@ private:
 	std::atomic<uint64>		FreeSlots		= ~0ull;
 	std::atomic<uint64>		Cancels			= 0;
 	uint64					PrevFreeSlots	= ~0ull;
-	TArray<FActivity*>		Pending;
-	TArray<FActivity*>		Active;
+	FActivity*				Pending			= nullptr;
 	FThrottler				Throttler;
+	TArray<FHostGroup>		Groups;
+	uint32					BusyCount		= 0;
 };
 
 ////////////////////////////////////////////////////////////////////////////////
 FEventLoop::FImpl::~FImpl()
 {
-	for (FActivity* Activity : Active)
-	{
-		Activity_Free(Activity);
-	}
+	check(BusyCount == 0);
 }
 
 ////////////////////////////////////////////////////////////////////////////////
@@ -2797,6 +3312,7 @@ FTicket FEventLoop::FImpl::Send(FActivity* Activity)
 	Trace(Activity, ETrace::RequestBegin);
 
 	FMessageBuilder(Activity->Buffer) << "\r\n";
+	Activity_ChangeState(Activity, FActivity::EState::Send);
 
 	uint64 Slot;
 	{
@@ -2819,7 +3335,10 @@ FTicket FEventLoop::FImpl::Send(FActivity* Activity)
 		}
 		Activity->Slot = int8(63 - FMath::CountLeadingZeros64(Slot));
 
-		Pending.Add(Activity);
+		// This puts pending requests in reverse order of when they were made
+		// but this will be undone when ReceiveWork() traverses the list
+		Activity->Next = Pending;
+		Pending = Activity;
 	}
 
 	return Slot;
@@ -2853,16 +3372,37 @@ void FEventLoop::FImpl::ReceiveWork()
 	}
 	PrevFreeSlots = FreeSlotsLoad;
 
-	TArray<FActivity*> NewActive;
+	// Fetch the pending activities from out in the wild
+	FActivity* Activity = nullptr;
 	{
 		FScopeLock _(&Lock);
-		NewActive = MoveTemp(Pending);
+		Swap(Activity, Pending);
 	}
 
-	Active.Append(NewActive);
+	// Group activities by their host.
+	for (FActivity* Next; Activity != nullptr; Activity = Next)
+	{
+		// Cache the next pointer and detach activity from the pending list
+		Next = Activity->Next;
+		Activity->Next = nullptr;
+
+		// Find or create group by the host pointer. This ultimately matches the
+		// FConnectionPool object used for the activity.
+		FHost& Host = *(Activity->Pool);
+		auto Pred = [&Host] (const FHostGroup& Lhs) { return &Lhs.GetHost() == &Host; };
+		FHostGroup* Group = Groups.FindByPredicate(Pred);
+		if (Group == nullptr)
+		{
+			Group = &(Groups.Emplace_GetRef(Host));
+		}
+
+		Group->AddActivity(Activity);
+		++BusyCount;
+	}
 }
 
 ////////////////////////////////////////////////////////////////////////////////
+#if 0
 uint32 FEventLoop::FImpl::Tick(uint32 PollTimeoutMs)
 {
 	TRACE_CPUPROFILER_EVENT_SCOPE(IasHttp::Tick);
@@ -3007,6 +3547,67 @@ uint32 FEventLoop::FImpl::Tick(uint32 PollTimeoutMs)
 
 	return BusyCount;
 }
+#endif // 0
+
+////////////////////////////////////////////////////////////////////////////////
+uint32 FEventLoop::FImpl::Tick(uint32 PollTimeoutMs)
+{
+	TRACE_CPUPROFILER_EVENT_SCOPE(IasHttp::Tick);
+
+	ReceiveWork();
+
+	int32 RecvAllowance = Throttler.GetAllowance();
+	uint64 CancelsLoad = Cancels.load(std::memory_order_relaxed);
+
+	// Tick groups and then remove ones that are idle
+	FTickState TickState = {
+		.DoneList = nullptr,
+		.Cancels = CancelsLoad,
+		.RecvAllowance = RecvAllowance,
+		.PollTimeoutMs = PollTimeoutMs,
+	};
+	for (FHostGroup& Group : Groups)
+	{
+		Group.Tick(TickState);
+	}
+
+	for (uint32 i = 0, n = Groups.Num(); i < n; ++i)
+	{
+		FHostGroup& Group = Groups[i];
+		if (Group.IsBusy())
+		{
+			continue;
+		}
+
+		Groups.RemoveAtSwap(i);
+		--n, --i;
+	}
+
+	Throttler.ReturnUnused(RecvAllowance);
+
+	uint64 ReturnedSlots = 0;
+	for (FActivity* Activity = TickState.DoneList; Activity != nullptr;)
+	{
+		FActivity* Next = Activity->Next;
+		ReturnedSlots |= (1ull << Activity->Slot);
+		Activity_Free(Activity);
+		--BusyCount;
+		Activity = Next;
+	}
+
+	if (ReturnedSlots)
+	{
+		PrevFreeSlots += ReturnedSlots;
+		FreeSlots.fetch_add(ReturnedSlots, std::memory_order_relaxed);
+	}
+
+	if (CancelsLoad)
+	{
+		Cancels.fetch_and(~CancelsLoad, std::memory_order_relaxed);
+	}
+
+	return BusyCount;
+}
 
 
 
