Surface bespoke pipeline config hyperlocally and move the needle to an agile ballpark of two.

#rnx

diff --git a/Engine/Source/Runtime/Experimental/IoStoreOnDemand/Private/CoreHttp/Client.h b/Engine/Source/Runtime/Experimental/IoStoreOnDemand/Private/CoreHttp/Client.h
--- a/Engine/Source/Runtime/Experimental/IoStoreOnDemand/Private/CoreHttp/Client.h
+++ b/Engine/Source/Runtime/Experimental/IoStoreOnDemand/Private/CoreHttp/Client.h
@@ -56,7 +56,8 @@ public:
 	struct FParams
 	{
 		int32				SetHostFromUrl(FAnsiStringView Url);
-		uint32				ConnectionCount;
+		uint16				ConnectionCount = 1;
+		uint16				PipelineLength = 1;
 		int32				SendBufSize = -1;
 		int32				RecvBufSize = -1;
 		struct {
diff --git a/Engine/Source/Runtime/Experimental/IoStoreOnDemand/Private/CoreHttp/Http.cpp b/Engine/Source/Runtime/Experimental/IoStoreOnDemand/Private/CoreHttp/Http.cpp
--- a/Engine/Source/Runtime/Experimental/IoStoreOnDemand/Private/CoreHttp/Http.cpp
+++ b/Engine/Source/Runtime/Experimental/IoStoreOnDemand/Private/CoreHttp/Http.cpp
@@ -1213,13 +1213,14 @@ public:
 	enum class EDirection : uint8 { Send, Recv };
 	static const uint32 InvalidIp = 0x00ff'ffff;
 
-					FHost(const ANSICHAR* InHostName, uint32 InPort, uint32 InMaxConn);
+					FHost(const ANSICHAR* InHostName, uint32 InPort, uint32 InMaxConn, uint32 PipeLength=1);
 	void			SetBufferSize(EDirection Dir, int32 Size);
 	int32			GetBufferSize(EDirection Dir) const;
 	FResult			Connect(FSocket& Socket);
 	int32			IsResolved() const;
 	FResult			ResolveHostName();
 	uint32			GetMaxConnections() const	{ return MaxConnections; }
+	uint32			GetPipelineLength() const	{ return PipelineLength; }
 	uint32			GetIpAddress() const		{ return IpAddresses[0]; }
 	FAnsiStringView	GetHostName() const			{ return HostName; }
 	uint32			GetPort() const				{ return Port; }
@@ -1231,13 +1232,15 @@ private:
 	int16			RecvBufKb = -1;
 	uint16			Port;
 	uint8			MaxConnections;
+	uint8			PipelineLength;
 };
 
 ////////////////////////////////////////////////////////////////////////////////
-FHost::FHost(const ANSICHAR* InHostName, uint32 InPort, uint32 InMaxConn)
+FHost::FHost(const ANSICHAR* InHostName, uint32 InPort, uint32 InMaxConn, uint32 PipeLength)
 : HostName(InHostName)
 , Port(uint16(InPort))
 , MaxConnections(uint8(InMaxConn))
+, PipelineLength(uint8(PipeLength))
 {
 	check(MaxConnections && MaxConnections == InMaxConn);
 }
@@ -1392,7 +1395,7 @@ class FSocketPool
 	: public FHost
 {
 public:
-					FSocketPool(const ANSICHAR* InHostName, uint32 InPort, uint32 InMaxLeases);
+					FSocketPool(const ANSICHAR* InHostName, uint32 InPort, uint32 InMaxLeases, uint32 PipeLen=1);
 					~FSocketPool();
 	static uint32	GetAllocSize(uint32 MaxLeases);
 	bool			LeaseSocket(FSocket& Out);
@@ -1404,8 +1407,8 @@ private:
 };
 
 ////////////////////////////////////////////////////////////////////////////////
-FSocketPool::FSocketPool(const ANSICHAR* InHostName, uint32 InPort, uint32 InMaxLeases)
-: FHost(InHostName, InPort, InMaxLeases)
+FSocketPool::FSocketPool(const ANSICHAR* InHostName, uint32 InPort, uint32 InMaxLeases, uint32 PipeLen)
+: FHost(InHostName, InPort, InMaxLeases, PipeLen)
 {
 
 	for (uint32 i = 0; i < GetMaxConnections(); ++i)
@@ -1496,7 +1499,8 @@ FConnectionPool::FConnectionPool(const FParams& Params)
 	new (Internal) FSocketPool(
 		HostDest,
 		Params.Host.Port,
-		Params.ConnectionCount
+		Params.ConnectionCount,
+		Params.PipelineLength
 	);
 	Internal->SetBufferSize(FHost::EDirection::Send, Params.SendBufSize);
 	Internal->SetBufferSize(FHost::EDirection::Recv, Params.RecvBufSize);
@@ -3194,7 +3198,7 @@ void FHostGroup::ScatterPendings()
 		Activity->Next = nullptr;
 
 		FSocketGroupPtr& Group = SocketGroups.Emplace_GetRef();
-		Group = MakeUnique<FSocketGroup>(1);
+		Group = MakeUnique<FSocketGroup>(Host.GetPipelineLength());
 		Group->AddActivity(Activity);
 
 		if (Iter == nullptr)
diff --git a/Engine/Source/Runtime/Experimental/IoStoreOnDemand/Private/OnDemandHttpClient.cpp b/Engine/Source/Runtime/Experimental/IoStoreOnDemand/Private/OnDemandHttpClient.cpp
--- a/Engine/Source/Runtime/Experimental/IoStoreOnDemand/Private/OnDemandHttpClient.cpp
+++ b/Engine/Source/Runtime/Experimental/IoStoreOnDemand/Private/OnDemandHttpClient.cpp
@@ -232,7 +232,8 @@ TUniquePtr<HTTP::FConnectionPool> FHttpClient::CreateConnection(const FStringVie
 		UE_LOG(LogIas, Log, TEXT("HTTP client receive buffer size set to %d"), Config.ReceiveBufferSize);
 		Params.RecvBufSize = Config.ReceiveBufferSize;
 	}
-	Params.ConnectionCount = Config.MaxConnectionCount;
+	Params.ConnectionCount = uint16(Config.MaxConnectionCount);
+	Params.PipelineLength = uint16(Config.PipelineLength);
 
 	return MakeUnique<HTTP::FConnectionPool>(Params);
 }
diff --git a/Engine/Source/Runtime/Experimental/IoStoreOnDemand/Private/OnDemandHttpClient.h b/Engine/Source/Runtime/Experimental/IoStoreOnDemand/Private/OnDemandHttpClient.h
--- a/Engine/Source/Runtime/Experimental/IoStoreOnDemand/Private/OnDemandHttpClient.h
+++ b/Engine/Source/Runtime/Experimental/IoStoreOnDemand/Private/OnDemandHttpClient.h
@@ -18,6 +18,7 @@ struct FHttpClientConfig
 	TArray<FString> Endpoints;
 	int32 PrimaryEndpoint = 0;
 	int32 MaxConnectionCount = 8;
+	int32 PipelineLength = 1;
 	int32 MaxRetryCount = 1;
 	int32 ReceiveBufferSize = -1;
 };
diff --git a/Engine/Source/Runtime/Experimental/IoStoreOnDemand/Private/OnDemandIoDispatcherBackend.cpp b/Engine/Source/Runtime/Experimental/IoStoreOnDemand/Private/OnDemandIoDispatcherBackend.cpp
--- a/Engine/Source/Runtime/Experimental/IoStoreOnDemand/Private/OnDemandIoDispatcherBackend.cpp
+++ b/Engine/Source/Runtime/Experimental/IoStoreOnDemand/Private/OnDemandIoDispatcherBackend.cpp
@@ -84,13 +84,20 @@ static FAutoConsoleVariableRef CVar_GIasHttpRecvBufKiB(
 	TEXT("Recv buffer size")
 );
 
-int32 GIasMaxHttpConnectionCount = 8;
+int32 GIasMaxHttpConnectionCount = 4;
 static FAutoConsoleVariableRef CVar_IasMaxHttpConnectionCount(
 	TEXT("ias.MaxHttpConnectionCount"),
 	GIasMaxHttpConnectionCount,
 	TEXT("Max number of open HTTP connections to the on demand endpoint(s).")
 );
 
+static int32 GIasHttpPipelineLength = 2;
+static FAutoConsoleVariableRef CVar_GIasHttpPipelineLength(
+	TEXT("ias.HttpPipelineLength"),
+	GIasHttpPipelineLength,
+	TEXT("Number of concurrent requests on one connection")
+);
+
 int32 GIasMaxHttpRetryCount = 2;
 static FAutoConsoleVariableRef CVar_IasMaxHttpRetryCount(
 	TEXT("ias.MaxHttpRetryCount"),
@@ -2168,6 +2175,7 @@ uint32 FOnDemandIoBackend::Run()
 			.Endpoints = AvailableEps.Urls,
 			.PrimaryEndpoint = AvailableEps.Current,
 			.MaxConnectionCount = GIasMaxHttpConnectionCount,
+			.PipelineLength = GIasHttpPipelineLength,
 			.MaxRetryCount = MaxHttpRetryCount,
 			.ReceiveBufferSize = GIasHttpRecvBufKiB >= 0 ? GIasHttpRecvBufKiB << 10 : -1
 		});
@@ -2179,7 +2187,7 @@ uint32 FOnDemandIoBackend::Run()
 	while (!bStopRequested)
 	{
 		// Process HTTP request(s) even if the client is invalid to ensure enqueued request(s) gets completed.
-		ProcessHttpRequests(HttpClient.Get(), HttpErrors, FMath::Min(2 * GIasMaxHttpConnectionCount, 64));
+		ProcessHttpRequests(HttpClient.Get(), HttpErrors, FMath::Min(GIasHttpPipelineLength * GIasMaxHttpConnectionCount, 64));
 		AvailableEps.Current = HttpClient.IsValid() ? HttpClient->GetPrimaryConnection() : INDEX_NONE;
 
 		if (!bStopRequested)
