IAS - wait for pending requests before shutting down

#rnx

diff --git a/Engine/Source/Runtime/Experimental/IoStoreOnDemand/Private/OnDemandIoDispatcherBackend.cpp b/Engine/Source/Runtime/Experimental/IoStoreOnDemand/Private/OnDemandIoDispatcherBackend.cpp
--- a/Engine/Source/Runtime/Experimental/IoStoreOnDemand/Private/OnDemandIoDispatcherBackend.cpp
+++ b/Engine/Source/Runtime/Experimental/IoStoreOnDemand/Private/OnDemandIoDispatcherBackend.cpp
@@ -1267,6 +1267,7 @@ class FOnDemandIoBackend final
 
 			bOutPending = bOutUpdatePriority = false;
 			FChunkRequest* ChunkRequest = Allocator.Construct(Request, Params);
+			ChunkRequestCount++;
 			Inflight.Add(Params.ChunkKey, ChunkRequest);
 
 			return ChunkRequest;
@@ -1323,17 +1324,29 @@ class FOnDemandIoBackend final
 		{
 			FScopeLock _(&Mutex);
 			check(!IsInFlight(Request));
-			Allocator.Destroy(Request);
+			Destroy(Request);
 		}
 		
 		void RemoveAndRelease(FChunkRequest* Request)
 		{
 			FScopeLock _(&Mutex);
 			Inflight.Remove(Request->Params.ChunkKey);
-			Allocator.Destroy(Request);
+			Destroy(Request);
+		}
+
+		int32 Num()
+		{
+			FScopeLock _(&Mutex);
+			return ChunkRequestCount; 
 		}
 
 	private:
+		void Destroy(FChunkRequest* Request)
+		{
+			Allocator.Destroy(Request);
+			ChunkRequestCount--;
+			check(ChunkRequestCount >= 0);
+		}
 
 		/** Helper intended to be called by methods that have already locked ::Mutex */
 		inline bool IsInFlight(const FChunkRequest* Request) const
@@ -1352,6 +1365,7 @@ class FOnDemandIoBackend final
 		TSingleThreadedSlabAllocator<FChunkRequest, 128> Allocator;
 		TMap<FIoHash, FChunkRequest*> Inflight;
 		FCriticalSection Mutex;
+		int32 ChunkRequestCount = 0;
 	};
 public:
 
@@ -1395,6 +1409,7 @@ private:
 
 	void AddDeferredTocs();
 	void ProcessHttpRequests(FHttpClient* HttpClient, FBitWindow& HttpErrors, int32 MaxConcurrentRequests);
+	int32 WaitForPendingRequests(float WaitTimeSeconds, float PollTimeSeconds);
 
 	TUniquePtr<IIasCache> Cache;
 	TUniquePtr<FOnDemandIoStore> IoStore;
@@ -1519,6 +1534,10 @@ void FOnDemandIoBackend::Shutdown()
 	bStopRequested = true;
 	TickBackendEvent->Trigger();
 	BackendThread.Reset();
+
+	const int32 NumPending = WaitForPendingRequests(5.0f, 0.1f);
+	UE_CLOG(NumPending > 0, LogIas, Warning, TEXT("%d request(s) still pending after shutdown"), NumPending);
+
 	BackendContext.Reset();
 }
 
@@ -1613,12 +1632,12 @@ void FOnDemandIoBackend::CompleteRequest(FChunkRequest* ChunkRequest)
 		Cache->Put(ChunkRequest->Params.ChunkKey, Chunk);
 	}
 
-	ChunkRequests.Release(ChunkRequest);
-
 	if (BackendStatus.ShouldAbandonCache() && InflightCacheRequestCount.load(std::memory_order_relaxed) == 0)
 	{
 		TickBackendEvent->Trigger();
 	}
+
+	ChunkRequests.Release(ChunkRequest);
 }
 
 void FOnDemandIoBackend::CompleteMaterialize(FChunkRequest* ChunkRequest)
@@ -2259,6 +2278,17 @@ uint32 FOnDemandIoBackend::Run()
 	return 0;
 }
 
+int32 FOnDemandIoBackend::WaitForPendingRequests(float WaitTimeSeconds, float PollTimeSeconds)
+{
+	const double StartTime = FPlatformTime::Seconds();
+	while (ChunkRequests.Num() > 0 && float(FPlatformTime::Seconds() - StartTime) < WaitTimeSeconds)
+	{
+		FPlatformProcess::SleepNoStats(PollTimeSeconds);
+	}
+
+	return ChunkRequests.Num();
+}
+
 TSharedPtr<IOnDemandIoDispatcherBackend> MakeOnDemandIoDispatcherBackend(TUniquePtr<IIasCache>&& Cache)
 {
 	return MakeShareable<IOnDemandIoDispatcherBackend>(new FOnDemandIoBackend(MoveTemp(Cache)));
