Kinematic/Dynamic LODs + GeoCL intermediate levels simplification + fix debug draw + faster convex construction

diff --git a/Engine/Source/Runtime/Experimental/Chaos/Private/Chaos/ChaosDebugDraw.cpp b/Engine/Source/Runtime/Experimental/Chaos/Private/Chaos/ChaosDebugDraw.cpp
--- a/Engine/Source/Runtime/Experimental/Chaos/Private/Chaos/ChaosDebugDraw.cpp
+++ b/Engine/Source/Runtime/Experimental/Chaos/Private/Chaos/ChaosDebugDraw.cpp
@@ -61,6 +61,9 @@ namespace Chaos
 		bool bChaosDebugDebugDrawColorShapesBySimQueryType = false;
 		FAutoConsoleVariableRef CVarChaosDebugDebugDrawColorShapesBySimQueryType(TEXT("p.Chaos.DebugDraw.ColorShapesBySimQueryType"), bChaosDebugDebugDrawColorShapesBySimQueryType, TEXT("Whether to show with different colors shapes that are sim enabled and query enabled (sim : blue, query : orange)"));
 
+		bool bChaosDebugDebugDrawColorShapesByConvexType = false;
+		FAutoConsoleVariableRef CVarChaosDebugDebugDrawColorShapesByConvexType(TEXT("p.Chaos.DebugDraw.ColorShapesByConvexType"), bChaosDebugDebugDrawColorShapesByConvexType, TEXT("Whether to show with different colors shapes that are convex and simplified (simplified : green, normal : orange)"));
+
 		bool bChaosDebugDebugDrawColorShapesByClusterUnion = false;
 		FAutoConsoleVariableRef CVarChaosDebugDebugDrawColorShapesByClusterUnion(TEXT("p.Chaos.DebugDraw.ColorShapesByClusterUnion"), bChaosDebugDebugDrawColorShapesByClusterUnion, TEXT("An extension of the ColorShapesByInternalCluster option: instead of using a single color for every internal cluster, will use a unique color per cluster union. Non-cluster unions will be black."));
 
@@ -350,8 +353,8 @@ namespace Chaos
 		//
 		//
 		//
+		void DrawShapesImpl(const FGeometryParticleHandle* Particle, const FRigidTransform3& ShapeTransform, const FImplicitObject* Implicit, const FPerShapeData* Shape, const FReal Margin, const FColor& Color, const FRealSingle Duration, const FChaosDebugDrawSettings& Settings, const bool bHasConvexOptimizer = false);
 
-		void DrawShapesImpl(const FGeometryParticleHandle* Particle, const FRigidTransform3& ShapeTransform, const FImplicitObject* Implicit, const FPerShapeData* Shape, const FReal Margin, const FColor& Color, const FRealSingle Duration, const FChaosDebugDrawSettings& Settings);
 
 		void DrawShape(const FRigidTransform3& ShapeTransform, const FImplicitObject* Implicit, const FPerShapeData* Shape, const FColor& Color, const float Duration, const FChaosDebugDrawSettings* Settings)
 		{
@@ -578,7 +581,7 @@ namespace Chaos
 			}
 		}
 
-		void DrawShapesImpl(const FGeometryParticleHandle* Particle, const FRigidTransform3& ShapeTransform, const FImplicitObject* Implicit, const FPerShapeData* Shape, const FReal Margin, const FColor& Color, const FRealSingle Duration, const FChaosDebugDrawSettings& Settings)
+		void DrawShapesImpl(const FGeometryParticleHandle* Particle, const FRigidTransform3& ShapeTransform, const FImplicitObject* Implicit, const FPerShapeData* Shape, const FReal Margin, const FColor& Color, const FRealSingle Duration, const FChaosDebugDrawSettings& Settings, const bool bHasConvexOptimizer)
 		{
 			if (Implicit == nullptr)
 			{
@@ -703,31 +706,31 @@ namespace Chaos
 					}
 				}
 			}
+			if (bChaosDebugDebugDrawColorShapesByConvexType)
+			{
+				if(InnerType == ImplicitObjectType::Convex)
+				{
+					ShapeColor = bHasConvexOptimizer ? FColor::Green : Particle->GetGeometry()->IsUnderlyingUnion() ? FColor::Blue : FColor::Orange;
+				}
+				else
+				{
+					ShapeColor = FColor::Yellow;
+				}
+			}
 			if (Shape && bChaosDebugDebugDrawColorShapesBySimQueryType)
 			{
 				if (Shape->GetSimEnabled())
 				{
-					const bool bIsUnion = Particle->GetGeometry()->IsUnderlyingUnion();
-					if (bIsUnion)
-					{
-						if(Implicit->GetDoCollide())
-						{
-							ShapeColor = FColor::Green;
-						}
-						else
-						{
-							ShapeColor = FColor::Blue;
-						}
-					}
-					else
-					{
-						ShapeColor = FColor::Orange;
-					}
+					ShapeColor = FColor::Green;
 				}
 				else if (Shape->GetQueryEnabled())
 				{
 					ShapeColor = FColor::Red;
 				}
+				else if (Shape->GetIsProbe())
+				{
+					ShapeColor = FColor::Orange;
+				}
 			}
 			if (Shape && !bChaosDebugDebugDrawShowQueryOnlyShapes)
 			{
@@ -854,6 +857,31 @@ namespace Chaos
 				FDebugDrawQueue::GetInstance().DrawDebugBox(ShapeBoundsPos, 0.5f * ShapeBounds.Extents(), ShapeTransform.GetRotation(), ShapeBoundsColor, false, Duration, uint8(Settings.DrawPriority), Settings.LineThickness);
 			}
 		}
+		
+		bool DrawConvexOptimizerShapes(const FVec3& P, const FRotation3& Q, const FGeometryParticleHandle* Particle, const FColor& InColor, const FChaosDebugDrawSettings& Settings)
+		{
+			if(auto ClusteredParticle = Particle->CastToClustered())
+			{
+				const TPimplPtr<Private::FConvexOptimizer>& ConvexOptimizer = ClusteredParticle->ConvexOptimizer();
+				if(ConvexOptimizer && ConvexOptimizer->IsValid())
+				{
+					ConvexOptimizer->VisitCollisionObjects([&Particle, &P, &Q, &ConvexOptimizer, &InColor, &Settings](
+						const FImplicitObject* ImplicitObject, const FRigidTransform3& RelativeTransform, const int32 RootObjectIndex, const int32 ObjectIndex, const int32 LeafObjectIndex)-> void
+					{
+						if(RootObjectIndex == INDEX_NONE && !ConvexOptimizer->GetShapeInstances().IsEmpty())
+						{
+							DrawShapesImpl(Particle, RelativeTransform * FRigidTransform3(P, Q), ImplicitObject, ConvexOptimizer->GetShapeInstances()[0].Get(), 0.0f, InColor, 0.0f, Settings, true);
+						}
+						else if(Particle->ShapeInstances().IsValidIndex(RootObjectIndex))
+						{
+							DrawShapesImpl(Particle, RelativeTransform * FRigidTransform3(P, Q), ImplicitObject, Particle->ShapeInstances()[RootObjectIndex].Get(), 0.0f, InColor, 0.0f, Settings, false);
+						}
+					});
+					return true;
+				}
+			}
+			return false;
+		}
 
 		void DrawParticleShapesImpl(const FRigidTransform3& SpaceTransform, const FGeometryParticleHandle* Particle, const FColor& InColor, const FChaosDebugDrawSettings& Settings)
 		{
@@ -861,21 +889,11 @@ namespace Chaos
 			const FRotation3 Q = SpaceTransform.GetRotation() * (Particle->ObjectState() == EObjectStateType::Dynamic ? Particle->CastToRigidParticle()->Q() : Particle->R());
 			const FRigidTransform3 ParticleSpaceTransform = FRigidTransform3(P, Q);
 
-			for (const FShapeInstancePtr& ShapeInstance : Particle->ShapeInstances())
-			{
-				DrawShapesImpl(Particle, ParticleSpaceTransform, ShapeInstance->GetGeometry(), ShapeInstance.Get(), 0.0f, InColor, 0.0f, Settings);
-			}
-
-			// Draw the optimized shapes if the particle has any
-			if (const FPBDRigidClusteredParticleHandle* ClusteredParticle = Particle->CastToClustered())
+			if(!DrawConvexOptimizerShapes(P, Q, Particle, InColor, Settings))
 			{
-				const TPimplPtr<Private::FConvexOptimizer>& ConvexOptimizer = ClusteredParticle->ConvexOptimizer();
-				if (ConvexOptimizer && ConvexOptimizer->IsValid())
+				for (const FShapeInstancePtr& ShapeInstance : Particle->ShapeInstances())
 				{
-					for (const FShapeInstancePtr& ShapeInstance : ConvexOptimizer->GetShapeInstances())
-					{
-						DrawShapesImpl(Particle, ParticleSpaceTransform, ShapeInstance->GetGeometry(), ShapeInstance.Get(), 0.0f, FColor::Green, 0.0f, Settings);
-					}
+					DrawShapesImpl(Particle, ParticleSpaceTransform, ShapeInstance->GetGeometry(), ShapeInstance.Get(), 0.0f, InColor, 0.0f, Settings);
 				}
 			}
 		}
diff --git a/Engine/Source/Runtime/Experimental/Chaos/Private/Chaos/Convex.cpp b/Engine/Source/Runtime/Experimental/Chaos/Private/Chaos/Convex.cpp
--- a/Engine/Source/Runtime/Experimental/Chaos/Private/Chaos/Convex.cpp
+++ b/Engine/Source/Runtime/Experimental/Chaos/Private/Chaos/Convex.cpp
@@ -303,9 +303,11 @@ namespace Chaos
 	}
 
 	// Store the structure data with the convex. This is used by manifold generation, for example
-	void FConvex::CreateStructureData(TArray<TArray<int32>>&& PlaneVertexIndices)
+	DECLARE_CYCLE_STAT(TEXT("FConvex::CreateStructureData"), STAT_CreateConvexStructureData, STATGROUP_ChaosCollision);
+	void FConvex::CreateStructureData(TArray<TArray<int32>>&& PlaneVertexIndices, const bool bRegularDatas)
 	{
-		const bool bSuccess = StructureData.SetPlaneVertices(MoveTemp(PlaneVertexIndices), Vertices.Num());
+		SCOPE_CYCLE_COUNTER(STAT_CreateConvexStructureData);
+		const bool bSuccess = StructureData.SetPlaneVertices(MoveTemp(PlaneVertexIndices), Vertices.Num(), bRegularDatas);
 		if (!bSuccess || !StructureData.IsValid())
 		{
 			UE_LOG(LogChaos, Error, TEXT("Unable to create structure data for %s"), *ToStringFull());
@@ -379,8 +381,10 @@ namespace Chaos
 		*this = FConvex(NewPoints, 0.0f);
 	}
 
+	DECLARE_CYCLE_STAT(TEXT("FConvex::ComputeUnitMassInertiaTensorAndRotationOfMass"), STAT_ComputeConvexMassInertia, STATGROUP_ChaosCollision);
 	void FConvex::ComputeUnitMassInertiaTensorAndRotationOfMass(const FReal InVolume)
 	{
+		SCOPE_CYCLE_COUNTER(STAT_ComputeConvexMassInertia);
 		if (InVolume < UE_SMALL_NUMBER || !StructureData.IsValid())
 		{
 			UnitMassInertiaTensor = FVec3{1., 1., 1.};
diff --git a/Engine/Source/Runtime/Experimental/Chaos/Private/Chaos/ConvexOptimizer.cpp b/Engine/Source/Runtime/Experimental/Chaos/Private/Chaos/ConvexOptimizer.cpp
--- a/Engine/Source/Runtime/Experimental/Chaos/Private/Chaos/ConvexOptimizer.cpp
+++ b/Engine/Source/Runtime/Experimental/Chaos/Private/Chaos/ConvexOptimizer.cpp
@@ -12,12 +12,16 @@ namespace Chaos
 namespace CVars
 {
 	// Replace all the convexes within an implicit hierarchy with a simplified one (kdop18 tribox for now) for collision
-	bool bChaosConvexSimplifyUnion = false;
+	bool bChaosConvexSimplifyUnion = true;
 	FAutoConsoleVariableRef CVarChaosSimplifyUnion(TEXT("p.Chaos.Convex.SimplifyUnion"), bChaosConvexSimplifyUnion, TEXT("If true replace all the convexes within an implcit hierarchy with a simplified one (kdop18 tribox for now) for collision"));
 
-	// Max number of convex LODs used during simplication
-	int32 ChaosConvexMaxLODs = 4;
-	FAutoConsoleVariableRef CVarChaosConvexMaxLODs(TEXT("p.Chaos.Convex.MaxLODs"), ChaosConvexMaxLODs, TEXT("Max number of convex LODs used during simplication"));
+	// Max number of convex LODs used during simplification  for dynamic particles
+	int32 ChaosConvexMaxDynamicLODs = 0;
+	FAutoConsoleVariableRef CVarChaosConvexMaxDynamicLODs(TEXT("p.Chaos.Convex.MaxDynamicLODs"), ChaosConvexMaxDynamicLODs, TEXT("Max number of convex LODs used during simplication for dynamic particles"));
+
+	// Max number of convex LODs used during simplification  for kinematic particles
+	int32 ChaosConvexMaxKinematicLODs = -1;
+	FAutoConsoleVariableRef CVarChaosConvexMaxKinematicLODs(TEXT("p.Chaos.Convex.MaxKinematicLODs"), ChaosConvexMaxKinematicLODs, TEXT("Max number of convex LODs used during simplication for kinematic particles"));
 	
 	// Tribox volume / convex hull threshold to trigger a volume splitting during tree construction
 	float ChaosConvexSplittingThreshold = 1.2f;
@@ -137,6 +141,9 @@ struct FTriboxTree
 
 	// Tree buffer index 
 	int32 BufferIndex = 0;
+	
+	// Tree max LODs
+	int32 MaxLODs = 0;
 };
 	
 FORCEINLINE FTriboxTree::FTriboxTree(const int32 NumLODs, const FConvexOptimizer::FCachedTriboxes& RootTriboxes) : TreeNodes(), BufferIndex(0)
@@ -156,12 +163,13 @@ FORCEINLINE FTriboxTree::FTriboxTree(const int32 NumLODs, const FConvexOptimizer
 		}
 	}
 	TreeNodes[CurrentBuffer()].Emplace(TriboxNode);
+	MaxLODs = NumLODs;
 }
 
 FORCEINLINE void FTriboxTree::BuildTreeLODs(TArray<FImplicitObjectPtr>& SimplifiedConvexes, const TUniquePtr<Private::FCollisionObjects>& CollisionObjects)
 {
 	//UE_LOG(LogTemp, Log, TEXT("	Building tribox binary tree"))
-	for(int32 LODIndex = 1; LODIndex < CVars::ChaosConvexMaxLODs; ++LODIndex)
+	for(int32 LODIndex = 1; LODIndex < MaxLODs; ++LODIndex)
 	{
 		// Swap the node buffers (working/result)
 		SwapNodeBuffers();
@@ -204,7 +212,7 @@ FORCEINLINE void FTriboxTree::AddTreeNode(const int32 LODIndex, const FTriboxNod
 			//UE_LOG(LogTemp, Log, TEXT("				Node volume = [%f %f] -> %f"), NodeVolume, ObjectsVolume, NodeVolume/ObjectsVolume)
 			
 			if((NodeVolume/ObjectsVolume > CVars::ChaosConvexSplittingThreshold) &&
-				(LODIndex < (CVars::ChaosConvexMaxLODs-1)))
+				(LODIndex < (MaxLODs-1)))
 			{
 				// If volume ratio (concavity) is big enough add this node to the current list to be processed by the next LOD
 				TreeNodes[CurrentBuffer()].Add(TriboxNode);
@@ -292,23 +300,24 @@ void FConvexOptimizer::BuildConvexShapes(const FShapesArray& UnionShapes)
 	ShapeData->SetQueryEnabled(false);
 }
 
-void FConvexOptimizer::SimplifyRootConvexes(const Chaos::FImplicitObjectUnionPtr& UnionGeometry, const FShapesArray& UnionShapes)
+void FConvexOptimizer::SimplifyRootConvexes(const Chaos::FImplicitObjectUnionPtr& UnionGeometry, const FShapesArray& UnionShapes, const EObjectStateType ObjectState)
 {
 	SCOPE_CYCLE_COUNTER(STAT_Collisions_SimplifyConvexes);
 	if(CVars::bChaosConvexSimplifyUnion && UnionGeometry && (UnionShapes.Num() > 0) &&
-		(UnionShapes.Num() == UnionGeometry->GetObjects().Num()))
+		(UnionShapes.Num() <= UnionGeometry->GetObjects().Num()))
 	{
 		SimplifiedConvexes.Reset();
 		CollisionObjects->ImplicitObjects.Reset();
 		ShapesArray.Reset();
+		const int32 MaxLODs = (ObjectState == EObjectStateType::Dynamic) ? CVars::ChaosConvexMaxDynamicLODs : CVars::ChaosConvexMaxKinematicLODs;
 
-		if(CVars::ChaosConvexMaxLODs == 0)
+		if(MaxLODs == 0)
 		{
 			BuildSingleConvex(UnionGeometry, UnionShapes);
 		}
 		else
 		{
-			BuildMultipleConvex(UnionGeometry, UnionShapes);
+			BuildMultipleConvex(UnionGeometry, UnionShapes, MaxLODs);
 		}
 		if(!SimplifiedConvexes.IsEmpty())
 		{
@@ -362,7 +371,8 @@ FORCEINLINE void BuildConvexTriboxes(const Chaos::FImplicitObjectUnionPtr& Union
 	
 	for(int32 RootObjectIndex = 0, NumRootObjects = UnionGeometry->GetObjects().Num(); RootObjectIndex < NumRootObjects; ++RootObjectIndex)
 	{
-		if(UnionShapes[RootObjectIndex]->GetSimEnabled())
+		const int32 ShapeIndex = UnionShapes.IsValidIndex(RootObjectIndex) ? RootObjectIndex : 0;
+		if(UnionShapes[ShapeIndex]->GetSimEnabled())
 		{
 			if(FImplicitObject* RootObject = UnionGeometry->GetObjects()[RootObjectIndex].GetReference())
 			{
@@ -388,6 +398,30 @@ FORCEINLINE void BuildConvexTriboxes(const Chaos::FImplicitObjectUnionPtr& Union
 							LocalTribox.AddConvex(Convex, ConvexTransform);
 						}
 					}
+					else if(const TImplicitObjectScaled<FConvex>* ScaledObject = TImplicitObjectScaled<FConvex>::AsScaled(*ImplicitObject))
+					{
+						if(!bHasRootTribox)
+						{
+							const FTribox::FRigidTransform3Type ConvexTransform(RelativeTransform.GetTranslation(),
+								RelativeTransform.GetRotation(), RelativeTransform.GetScale3D() * ScaledObject->GetScale());
+							// Disable collision for all the convexes that are going to be used to build the tribox
+							// For now only used for debug draw
+							const_cast<FImplicitObject*>(ImplicitObject)->SetDoCollide(false);
+							LocalTribox.AddConvex(ScaledObject->GetUnscaledObject(), ConvexTransform);
+						}
+					}
+
+					else if(const TImplicitObjectInstanced<FConvex>* InstancedObject = TImplicitObjectInstanced<FConvex>::AsInstanced(*ImplicitObject))
+					{
+						if(!bHasRootTribox)
+						{
+							const FTribox::FRigidTransform3Type ConvexTransform(RelativeTransform);
+							// Disable collision for all the convexes that are going to be used to build the tribox
+							// For now only used for debug draw
+							const_cast<FImplicitObject*>(ImplicitObject)->SetDoCollide(false);
+							LocalTribox.AddConvex(ScaledObject->Object(), ConvexTransform);
+						}
+					}
 					else
 					{
 						// Add non convex implicits to the list of collision objects
@@ -450,11 +484,11 @@ FORCEINLINE void FindClosestPlane(const FTribox::FVec3Type& PointPosition, const
 	}
 }
 	
-void FConvexOptimizer::BuildMultipleConvex(const Chaos::FImplicitObjectUnionPtr& UnionGeometry, const FShapesArray& UnionShapes)
+void FConvexOptimizer::BuildMultipleConvex(const Chaos::FImplicitObjectUnionPtr& UnionGeometry, const FShapesArray& UnionShapes, const int32 MaxLODs)
 {
 	BuildConvexTriboxes(UnionGeometry, UnionShapes, CollisionObjects, RootTriboxes);
 
-	if(CVars::ChaosConvexMaxLODs < 0)
+	if(MaxLODs < 0)
 	{
 		for(auto& RootTribox : RootTriboxes)
         {
@@ -473,11 +507,16 @@ void FConvexOptimizer::BuildMultipleConvex(const Chaos::FImplicitObjectUnionPtr&
 	else
 	{
 		// Build the binary tree and add the leaves to the cimplified convexes
-		FTriboxTree TriboxTree(CVars::ChaosConvexMaxLODs, RootTriboxes);
+		FTriboxTree TriboxTree(MaxLODs, RootTriboxes);
 		TriboxTree.BuildTreeLODs(SimplifiedConvexes, CollisionObjects);
 	}
 }
 
+int32 FConvexOptimizer::NumCollisionObjects() const
+{
+	return BVH.IsValid() ? BVH->GetObjects().Num() : CollisionObjects->ImplicitObjects.Num();
+}
+
 void VisitCollisionObjects(const FConvexOptimizer* ConvexOptimizer, const FImplicitObject* ImplicitObject, const FImplicitHierarchyVisitor& VisitorFunc)
 {
 	if(ConvexOptimizer && ConvexOptimizer->IsValid())
diff --git a/Engine/Source/Runtime/Experimental/Chaos/Private/Chaos/PBDRigidClustering.cpp b/Engine/Source/Runtime/Experimental/Chaos/Private/Chaos/PBDRigidClustering.cpp
--- a/Engine/Source/Runtime/Experimental/Chaos/Private/Chaos/PBDRigidClustering.cpp
+++ b/Engine/Source/Runtime/Experimental/Chaos/Private/Chaos/PBDRigidClustering.cpp
@@ -265,6 +265,9 @@ namespace Chaos
 		}
 		UpdateKinematicProperties(NewParticle, MChildren, MEvolution);
 		UpdateGeometry(NewParticle, ChildrenSet, MChildren, ProxyGeometry, Parameters);
+
+		// Build the convex optimizer if required
+		FRigidClustering::BuildConvexOptimizer(NewParticle);
 		
 		GenerateConnectionGraph(NewParticle, Parameters);
 
@@ -351,12 +354,15 @@ namespace Chaos
 
 	void FRigidClustering::BuildConvexOptimizer(FPBDRigidClusteredParticleHandle* Particle)
 	{
-		if(!Particle->ConvexOptimizer())
-		{
-			Particle->ConvexOptimizer() = MakePimpl<Private::FConvexOptimizer>();
+		if(Particle && Particle->GetGeometry())
+		{ 
+			if (!Particle->ConvexOptimizer())
+			{
+				Particle->ConvexOptimizer() = MakePimpl<Private::FConvexOptimizer>();
+			}
+			Particle->ConvexOptimizer()->SimplifyRootConvexes(Particle->GetGeometry()->template AsA<FImplicitObjectUnion>(),
+				Particle->ShapesArray(), Particle->ObjectState());
 		}
-		Particle->ConvexOptimizer()->SimplifyRootConvexes( Particle->GetGeometry()->template AsA<FImplicitObjectUnion>(),
-					Particle->ShapesArray());
 	}
 
 	DECLARE_CYCLE_STAT(TEXT("TPBDRigidClustering<>::RemoveParticlesFromCluster"), STAT_RemoveParticlesFromCluster, STATGROUP_Chaos);
@@ -547,6 +553,9 @@ namespace Chaos
 
 		UpdateGeometry(NewParticle, ChildrenSet, MChildren, FImplicitObjectPtr(nullptr), NoCleanParams);
 
+		// Build the convex optimizer if required
+		FRigidClustering::BuildConvexOptimizer(NewParticle);
+
 		return NewParticle;
 	}
 
diff --git a/Engine/Source/Runtime/Experimental/Chaos/Private/Chaos/Tribox.cpp b/Engine/Source/Runtime/Experimental/Chaos/Private/Chaos/Tribox.cpp
--- a/Engine/Source/Runtime/Experimental/Chaos/Private/Chaos/Tribox.cpp
+++ b/Engine/Source/Runtime/Experimental/Chaos/Private/Chaos/Tribox.cpp
@@ -89,81 +89,81 @@ static const FTribox::FMatrix33Type PrincipalMatrices[8][9]=
 		{-0.000000, Sqrt2, -1.000000, Sqrt2, -0.000000, -1.000000, -0.000000, 0.000000, 1.000000}
 	},
 	{
-	{1.000000, -0.000000, 0.000000, -1.000000, 0.000000, Sqrt2, 1.000000, -Sqrt2, 0.000000},
-	{-0.000000, 1.000000, -0.000000, 0.000000, -1.000000, Sqrt2, -Sqrt2, -1.000000, Sqrt2},
-	{-0.000000, 0.000000, 1.000000, Sqrt2, -Sqrt2, 1.000000, -0.000000, -Sqrt2, 1.000000},
-	{-1.000000, 0.000000, Sqrt2, 1.000000, -0.000000, 0.000000, -1.000000, -Sqrt2, Sqrt2},
-	{0.000000, -1.000000, Sqrt2, -0.000000, 1.000000, -0.000000, -Sqrt2, 1.000000, 0.000000},
-	{-Sqrt2, Sqrt2, 1.000000, 0.000000, -0.000000, 1.000000, -Sqrt2, 0.000000, 1.000000},
-	{-1.000000, Sqrt2, -0.000000, 1.000000, -Sqrt2, Sqrt2, -1.000000, 0.000000, -0.000000},
-	{-Sqrt2, 1.000000, Sqrt2, Sqrt2, -1.000000, 0.000000, -0.000000, -1.000000, -0.000000},
-	{0.000000, Sqrt2, -1.000000, Sqrt2, 0.000000, -1.000000, 0.000000, -0.000000, -1.000000}
+		{1.000000, -0.000000, 0.000000, -1.000000, 0.000000, Sqrt2, 1.000000, -Sqrt2, 0.000000},
+		{-0.000000, 1.000000, -0.000000, 0.000000, -1.000000, Sqrt2, -Sqrt2, -1.000000, Sqrt2},
+		{-0.000000, 0.000000, 1.000000, Sqrt2, -Sqrt2, 1.000000, -0.000000, -Sqrt2, 1.000000},
+		{-1.000000, 0.000000, Sqrt2, 1.000000, -0.000000, 0.000000, -1.000000, -Sqrt2, Sqrt2},
+		{0.000000, -1.000000, Sqrt2, -0.000000, 1.000000, -0.000000, -Sqrt2, 1.000000, 0.000000},
+		{-Sqrt2, Sqrt2, 1.000000, 0.000000, -0.000000, 1.000000, -Sqrt2, 0.000000, 1.000000},
+		{-1.000000, Sqrt2, -0.000000, 1.000000, -Sqrt2, Sqrt2, -1.000000, 0.000000, -0.000000},
+		{-Sqrt2, 1.000000, Sqrt2, Sqrt2, -1.000000, 0.000000, -0.000000, -1.000000, -0.000000},
+		{0.000000, Sqrt2, -1.000000, Sqrt2, 0.000000, -1.000000, 0.000000, -0.000000, -1.000000}
 	},
 	{
-	{1.000000, -0.000000, 0.000000, 1.000000, 0.000000, -Sqrt2, -1.000000, Sqrt2, 0.000000},
-	{-0.000000, 1.000000, -0.000000, 0.000000, 1.000000, -Sqrt2, Sqrt2, 1.000000, -Sqrt2},
-	{-0.000000, 0.000000, 1.000000, -Sqrt2, Sqrt2, -1.000000, -0.000000, Sqrt2, -1.000000},
-	{-1.000000, 0.000000, Sqrt2, -1.000000, -0.000000, 0.000000, 1.000000, Sqrt2, -Sqrt2},
-	{0.000000, -1.000000, Sqrt2, -0.000000, -1.000000, -0.000000, Sqrt2, -1.000000, 0.000000},
-	{-Sqrt2, Sqrt2, 1.000000, 0.000000, -0.000000, -1.000000, Sqrt2, 0.000000, -1.000000},
-	{-1.000000, Sqrt2, -0.000000, -1.000000, Sqrt2, -Sqrt2, 1.000000, 0.000000, -0.000000},
-	{-Sqrt2, 1.000000, Sqrt2, -Sqrt2, 1.000000, 0.000000, -0.000000, 1.000000, -0.000000},
-	{0.000000, Sqrt2, -1.000000, -Sqrt2, 0.000000, 1.000000, 0.000000, -0.000000, 1.000000}
+		{1.000000, -0.000000, 0.000000, 1.000000, 0.000000, -Sqrt2, -1.000000, Sqrt2, 0.000000},
+		{-0.000000, 1.000000, -0.000000, 0.000000, 1.000000, -Sqrt2, Sqrt2, 1.000000, -Sqrt2},
+		{-0.000000, 0.000000, 1.000000, -Sqrt2, Sqrt2, -1.000000, -0.000000, Sqrt2, -1.000000},
+		{-1.000000, 0.000000, Sqrt2, -1.000000, -0.000000, 0.000000, 1.000000, Sqrt2, -Sqrt2},
+		{0.000000, -1.000000, Sqrt2, -0.000000, -1.000000, -0.000000, Sqrt2, -1.000000, 0.000000},
+		{-Sqrt2, Sqrt2, 1.000000, 0.000000, -0.000000, -1.000000, Sqrt2, 0.000000, -1.000000},
+		{-1.000000, Sqrt2, -0.000000, -1.000000, Sqrt2, -Sqrt2, 1.000000, 0.000000, -0.000000},
+		{-Sqrt2, 1.000000, Sqrt2, -Sqrt2, 1.000000, 0.000000, -0.000000, 1.000000, -0.000000},
+		{0.000000, Sqrt2, -1.000000, -Sqrt2, 0.000000, 1.000000, 0.000000, -0.000000, 1.000000}
 	},
 	{
-	{-1.000000, -0.000000, 0.000000, -1.000000, 0.000000, Sqrt2, -1.000000, Sqrt2, 0.000000},
-	{-0.000000, -1.000000, -0.000000, 0.000000, -1.000000, Sqrt2, Sqrt2, 1.000000, -Sqrt2},
-	{-0.000000, 0.000000, -1.000000, Sqrt2, -Sqrt2, 1.000000, -0.000000, Sqrt2, -1.000000},
-	{1.000000, 0.000000, -Sqrt2, 1.000000, -0.000000, 0.000000, 1.000000, Sqrt2, -Sqrt2},
-	{0.000000, 1.000000, -Sqrt2, -0.000000, 1.000000, -0.000000, Sqrt2, -1.000000, 0.000000},
-	{Sqrt2, -Sqrt2, -1.000000, 0.000000, -0.000000, 1.000000, Sqrt2, 0.000000, -1.000000},
-	{1.000000, -Sqrt2, -0.000000, 1.000000, -Sqrt2, Sqrt2, 1.000000, 0.000000, -0.000000},
-	{Sqrt2, -1.000000, -Sqrt2, Sqrt2, -1.000000, 0.000000, -0.000000, 1.000000, -0.000000},
-	{0.000000, -Sqrt2, 1.000000, Sqrt2, 0.000000, -1.000000, 0.000000, -0.000000, 1.000000}
+		{-1.000000, -0.000000, 0.000000, -1.000000, 0.000000, Sqrt2, -1.000000, Sqrt2, 0.000000},
+		{-0.000000, -1.000000, -0.000000, 0.000000, -1.000000, Sqrt2, Sqrt2, 1.000000, -Sqrt2},
+		{-0.000000, 0.000000, -1.000000, Sqrt2, -Sqrt2, 1.000000, -0.000000, Sqrt2, -1.000000},
+		{1.000000, 0.000000, -Sqrt2, 1.000000, -0.000000, 0.000000, 1.000000, Sqrt2, -Sqrt2},
+		{0.000000, 1.000000, -Sqrt2, -0.000000, 1.000000, -0.000000, Sqrt2, -1.000000, 0.000000},
+		{Sqrt2, -Sqrt2, -1.000000, 0.000000, -0.000000, 1.000000, Sqrt2, 0.000000, -1.000000},
+		{1.000000, -Sqrt2, -0.000000, 1.000000, -Sqrt2, Sqrt2, 1.000000, 0.000000, -0.000000},
+		{Sqrt2, -1.000000, -Sqrt2, Sqrt2, -1.000000, 0.000000, -0.000000, 1.000000, -0.000000},
+		{0.000000, -Sqrt2, 1.000000, Sqrt2, 0.000000, -1.000000, 0.000000, -0.000000, 1.000000}
 	},
 	{
-	{-1.000000, 0.000000, -0.000000, 1.000000, -0.000000, -Sqrt2, -1.000000, Sqrt2, -0.000000},
-	{0.000000, -1.000000, 0.000000, -0.000000, 1.000000, -Sqrt2, Sqrt2, 1.000000, -Sqrt2},
-	{0.000000, -0.000000, -1.000000, -Sqrt2, Sqrt2, -1.000000, 0.000000, Sqrt2, -1.000000},
-	{1.000000, -0.000000, -Sqrt2, -1.000000, 0.000000, -0.000000, 1.000000, Sqrt2, -Sqrt2},
-	{-0.000000, 1.000000, -Sqrt2, 0.000000, -1.000000, 0.000000, Sqrt2, -1.000000, -0.000000},
-	{Sqrt2, -Sqrt2, -1.000000, -0.000000, 0.000000, -1.000000, Sqrt2, -0.000000, -1.000000},
-	{1.000000, -Sqrt2, 0.000000, -1.000000, Sqrt2, -Sqrt2, 1.000000, -0.000000, 0.000000},
-	{Sqrt2, -1.000000, -Sqrt2, -Sqrt2, 1.000000, -0.000000, 0.000000, 1.000000, 0.000000},
-	{-0.000000, -Sqrt2, 1.000000, -Sqrt2, -0.000000, 1.000000, -0.000000, 0.000000, 1.000000}
+		{-1.000000, 0.000000, -0.000000, 1.000000, -0.000000, -Sqrt2, -1.000000, Sqrt2, -0.000000},
+		{0.000000, -1.000000, 0.000000, -0.000000, 1.000000, -Sqrt2, Sqrt2, 1.000000, -Sqrt2},
+		{0.000000, -0.000000, -1.000000, -Sqrt2, Sqrt2, -1.000000, 0.000000, Sqrt2, -1.000000},
+		{1.000000, -0.000000, -Sqrt2, -1.000000, 0.000000, -0.000000, 1.000000, Sqrt2, -Sqrt2},
+		{-0.000000, 1.000000, -Sqrt2, 0.000000, -1.000000, 0.000000, Sqrt2, -1.000000, -0.000000},
+		{Sqrt2, -Sqrt2, -1.000000, -0.000000, 0.000000, -1.000000, Sqrt2, -0.000000, -1.000000},
+		{1.000000, -Sqrt2, 0.000000, -1.000000, Sqrt2, -Sqrt2, 1.000000, -0.000000, 0.000000},
+		{Sqrt2, -1.000000, -Sqrt2, -Sqrt2, 1.000000, -0.000000, 0.000000, 1.000000, 0.000000},
+		{-0.000000, -Sqrt2, 1.000000, -Sqrt2, -0.000000, 1.000000, -0.000000, 0.000000, 1.000000}
 	},
 	{
-	{-1.000000, 0.000000, -0.000000, -1.000000, -0.000000, Sqrt2, 1.000000, -Sqrt2, -0.000000},
-	{0.000000, -1.000000, 0.000000, -0.000000, -1.000000, Sqrt2, -Sqrt2, -1.000000, Sqrt2},
-	{0.000000, -0.000000, -1.000000, Sqrt2, -Sqrt2, 1.000000, 0.000000, -Sqrt2, 1.000000},
-	{1.000000, -0.000000, -Sqrt2, 1.000000, 0.000000, -0.000000, -1.000000, -Sqrt2, Sqrt2},
-	{-0.000000, 1.000000, -Sqrt2, 0.000000, 1.000000, 0.000000, -Sqrt2, 1.000000, -0.000000},
-	{Sqrt2, -Sqrt2, -1.000000, -0.000000, 0.000000, 1.000000, -Sqrt2, -0.000000, 1.000000},
-	{1.000000, -Sqrt2, 0.000000, 1.000000, -Sqrt2, Sqrt2, -1.000000, -0.000000, 0.000000},
-	{Sqrt2, -1.000000, -Sqrt2, Sqrt2, -1.000000, -0.000000, 0.000000, -1.000000, 0.000000},
-	{-0.000000, -Sqrt2, 1.000000, Sqrt2, -0.000000, -1.000000, -0.000000, 0.000000, -1.000000}
+		{-1.000000, 0.000000, -0.000000, -1.000000, -0.000000, Sqrt2, 1.000000, -Sqrt2, -0.000000},
+		{0.000000, -1.000000, 0.000000, -0.000000, -1.000000, Sqrt2, -Sqrt2, -1.000000, Sqrt2},
+		{0.000000, -0.000000, -1.000000, Sqrt2, -Sqrt2, 1.000000, 0.000000, -Sqrt2, 1.000000},
+		{1.000000, -0.000000, -Sqrt2, 1.000000, 0.000000, -0.000000, -1.000000, -Sqrt2, Sqrt2},
+		{-0.000000, 1.000000, -Sqrt2, 0.000000, 1.000000, 0.000000, -Sqrt2, 1.000000, -0.000000},
+		{Sqrt2, -Sqrt2, -1.000000, -0.000000, 0.000000, 1.000000, -Sqrt2, -0.000000, 1.000000},
+		{1.000000, -Sqrt2, 0.000000, 1.000000, -Sqrt2, Sqrt2, -1.000000, -0.000000, 0.000000},
+		{Sqrt2, -1.000000, -Sqrt2, Sqrt2, -1.000000, -0.000000, 0.000000, -1.000000, 0.000000},
+		{-0.000000, -Sqrt2, 1.000000, Sqrt2, -0.000000, -1.000000, -0.000000, 0.000000, -1.000000}
 	},
 	{
-	{1.000000, 0.000000, -0.000000, 1.000000, -0.000000, -Sqrt2, 1.000000, -Sqrt2, -0.000000},
-	{0.000000, 1.000000, 0.000000, -0.000000, 1.000000, -Sqrt2, -Sqrt2, -1.000000, Sqrt2},
-	{0.000000, -0.000000, 1.000000, -Sqrt2, Sqrt2, -1.000000, 0.000000, -Sqrt2, 1.000000},
-	{-1.000000, -0.000000, Sqrt2, -1.000000, 0.000000, -0.000000, -1.000000, -Sqrt2, Sqrt2},
-	{-0.000000, -1.000000, Sqrt2, 0.000000, -1.000000, 0.000000, -Sqrt2, 1.000000, -0.000000},
-	{-Sqrt2, Sqrt2, 1.000000, -0.000000, 0.000000, -1.000000, -Sqrt2, -0.000000, 1.000000},
-	{-1.000000, Sqrt2, 0.000000, -1.000000, Sqrt2, -Sqrt2, -1.000000, -0.000000, 0.000000},
-	{-Sqrt2, 1.000000, Sqrt2, -Sqrt2, 1.000000, -0.000000, 0.000000, -1.000000, 0.000000},
-	{-0.000000, Sqrt2, -1.000000, -Sqrt2, -0.000000, 1.000000, -0.000000, 0.000000, -1.000000}
+		{1.000000, 0.000000, -0.000000, 1.000000, -0.000000, -Sqrt2, 1.000000, -Sqrt2, -0.000000},
+		{0.000000, 1.000000, 0.000000, -0.000000, 1.000000, -Sqrt2, -Sqrt2, -1.000000, Sqrt2},
+		{0.000000, -0.000000, 1.000000, -Sqrt2, Sqrt2, -1.000000, 0.000000, -Sqrt2, 1.000000},
+		{-1.000000, -0.000000, Sqrt2, -1.000000, 0.000000, -0.000000, -1.000000, -Sqrt2, Sqrt2},
+		{-0.000000, -1.000000, Sqrt2, 0.000000, -1.000000, 0.000000, -Sqrt2, 1.000000, -0.000000},
+		{-Sqrt2, Sqrt2, 1.000000, -0.000000, 0.000000, -1.000000, -Sqrt2, -0.000000, 1.000000},
+		{-1.000000, Sqrt2, 0.000000, -1.000000, Sqrt2, -Sqrt2, -1.000000, -0.000000, 0.000000},
+		{-Sqrt2, 1.000000, Sqrt2, -Sqrt2, 1.000000, -0.000000, 0.000000, -1.000000, 0.000000},
+		{-0.000000, Sqrt2, -1.000000, -Sqrt2, -0.000000, 1.000000, -0.000000, 0.000000, -1.000000}
 	},
 	{
-	{-1.000000, -0.000000, 0.000000, 1.000000, 0.000000, -Sqrt2, 1.000000, -Sqrt2, 0.000000},
-	{-0.000000, -1.000000, -0.000000, 0.000000, 1.000000, -Sqrt2, -Sqrt2, -1.000000, Sqrt2},
-	{-0.000000, 0.000000, -1.000000, -Sqrt2, Sqrt2, -1.000000, -0.000000, -Sqrt2, 1.000000},
-	{1.000000, 0.000000, -Sqrt2, -1.000000, -0.000000, 0.000000, -1.000000, -Sqrt2, Sqrt2},
-	{0.000000, 1.000000, -Sqrt2, -0.000000, -1.000000, -0.000000, -Sqrt2, 1.000000, 0.000000},
-	{Sqrt2, -Sqrt2, -1.000000, 0.000000, -0.000000, -1.000000, -Sqrt2, 0.000000, 1.000000},
-	{1.000000, -Sqrt2, -0.000000, -1.000000, Sqrt2, -Sqrt2, -1.000000, 0.000000, -0.000000},
-	{Sqrt2, -1.000000, -Sqrt2, -Sqrt2, 1.000000, 0.000000, -0.000000, -1.000000, -0.000000},
-	{0.000000, -Sqrt2, 1.000000, -Sqrt2, 0.000000, 1.000000, 0.000000, -0.000000, -1.000000}
+		{-1.000000, -0.000000, 0.000000, 1.000000, 0.000000, -Sqrt2, 1.000000, -Sqrt2, 0.000000},
+		{-0.000000, -1.000000, -0.000000, 0.000000, 1.000000, -Sqrt2, -Sqrt2, -1.000000, Sqrt2},
+		{-0.000000, 0.000000, -1.000000, -Sqrt2, Sqrt2, -1.000000, -0.000000, -Sqrt2, 1.000000},
+		{1.000000, 0.000000, -Sqrt2, -1.000000, -0.000000, 0.000000, -1.000000, -Sqrt2, Sqrt2},
+		{0.000000, 1.000000, -Sqrt2, -0.000000, -1.000000, -0.000000, -Sqrt2, 1.000000, 0.000000},
+		{Sqrt2, -Sqrt2, -1.000000, 0.000000, -0.000000, -1.000000, -Sqrt2, 0.000000, 1.000000},
+		{1.000000, -Sqrt2, -0.000000, -1.000000, Sqrt2, -Sqrt2, -1.000000, 0.000000, -0.000000},
+		{Sqrt2, -1.000000, -Sqrt2, -Sqrt2, 1.000000, 0.000000, -0.000000, -1.000000, -0.000000},
+		{0.000000, -Sqrt2, 1.000000, -Sqrt2, 0.000000, 1.000000, 0.000000, -0.000000, -1.000000}
 	}
 };
 
@@ -395,6 +395,71 @@ FORCEINLINE void CompressFaces(TArray<TArray<int32>>& FaceIndices)
 	}
 }
 
+FImplicitObjectPtr FTribox::CreateConvexFromTopology(
+	TArray<FConvex::FPlaneType>&& ConvexPlanes, TArray<TArray<int32>>&& FaceIndices, TArray<FConvex::FVec3Type>&& ConvexVertices) const
+{
+	FRotation3 RotationMatrix = FRotation3::Identity;
+	FVec3Type InertiaTensor(1.0f, 1.0f, 1.0f);
+	FRealType MinVolume = FLT_MAX;
+	{
+		const FRealType HX = (MaxDists[0]+MaxDists[1]);
+		const FRealType PX = (MaxDists[6]+MaxDists[7]);
+		const FRealType MX = (MaxDists[8]+MaxDists[9]);
+
+		const FRealType HY = (MaxDists[2]+MaxDists[3]);
+		const FRealType PY = (MaxDists[10]+MaxDists[11]);
+		const FRealType MY = (MaxDists[12]+MaxDists[13]);
+
+		const FRealType HZ = (MaxDists[4]+MaxDists[5]);
+		const FRealType PZ = (MaxDists[14]+MaxDists[15]);
+		const FRealType MZ = (MaxDists[16]+MaxDists[17]);
+		
+		FRealType Volumes[4];
+		Volumes[0] = HX * HY * HZ;
+		Volumes[1] = HX * PX * MX;
+		Volumes[2] = HY * PY * MY;
+		Volumes[3] = HZ * PZ * MZ;
+		
+		int8 MinIndex = INDEX_NONE;
+		for(int8 VolumeIndex = 0; VolumeIndex < 4; ++VolumeIndex)
+		{
+			if(Volumes[VolumeIndex] < MinVolume)
+			{
+				MinVolume = Volumes[VolumeIndex];
+				MinIndex = VolumeIndex;
+			}
+		}
+		
+		if(MinIndex == 0)
+		{
+			static const FRotation3 ConstRotationMatrix = FRotation3::Identity;
+			RotationMatrix = ConstRotationMatrix;
+			InertiaTensor = FVec3Type(HY * HY + HZ * HZ, HX * HX + HZ * HZ, HY * HY + HX * HX) / (12.0f * Volumes[0]);
+		}
+		else if(MinIndex == 1)
+		{
+			static const FRotation3 ConstRotationMatrix = FRotation3::FromAxisAngle(FVec3(1.0f, 0.0f, 0.0f), PI / 4.0f);
+			RotationMatrix = ConstRotationMatrix;
+			InertiaTensor = FVec3Type(PX * PX + MX * MX, HX * HX + MX * MX, PX * PX + HX * HX) / (12.0f * Volumes[1]);
+		}
+		else if(MinIndex == 2)
+		{
+			static const FRotation3 ConstRotationMatrix = FRotation3::FromAxisAngle(FVec3(0.0f, 1.0f, 0.0f), PI / 4.0f);
+			RotationMatrix = ConstRotationMatrix;
+			InertiaTensor = FVec3Type(MY * MY + HY * HY, MY * MY + PY * PY, HY * HY + PY * PY) / (12.0f * Volumes[2]);
+		}
+		else if(MinIndex == 3)
+		{
+			static const FRotation3 ConstRotationMatrix = FRotation3::FromAxisAngle(FVec3(0.0f, 0.0f, 1.0f), PI / 4.0f);
+			RotationMatrix = ConstRotationMatrix;
+			InertiaTensor = FVec3Type(MZ * MZ + HZ * HZ, PZ * PZ + HZ * HZ, PZ * PZ + MZ * MZ) / (12.0f * Volumes[3]);
+		}
+	}
+	return MakeImplicitObjectPtr<Chaos::FConvex>(MoveTemp(ConvexPlanes), MoveTemp(FaceIndices), MoveTemp(ConvexVertices), 
+		FConvex::FVec3Type(-MaxDists[1], -MaxDists[3], -MaxDists[5]), FConvex::FVec3Type(MaxDists[0], MaxDists[2], MaxDists[4]), 
+		MinVolume, InertiaTensor, RotationMatrix, true);
+}
+
 DECLARE_CYCLE_STAT(TEXT("Collisions::MakeConvex"), STAT_MakeTriboxConvex, STATGROUP_ChaosCollision);
 FImplicitObjectPtr FTribox::MakeConvex() const
 {
@@ -451,21 +516,26 @@ FImplicitObjectPtr FTribox::MakeConvex() const
 		}
 		CompressFaces(FaceIndices);
 	}
-	return MakeImplicitObjectPtr<Chaos::FConvex>(MoveTemp(ConvexPlanes), MoveTemp(FaceIndices), MoveTemp(ConvexVertices));
+
+	return CreateConvexFromTopology(MoveTemp(ConvexPlanes), MoveTemp(FaceIndices), MoveTemp(ConvexVertices));
 }
-	
-bool FTribox::OverlapTribox(const FTribox& OtherTribox, FTribox& OverlapTribox) const
+
+bool FTribox::IsTriboxOverlapping(const FTribox& OtherTribox) const
 {
-	bool bOverlapTriboxes = true;
 	for(int32 AxisIndex = 0; AxisIndex < NumPlanes; AxisIndex += 2)
 	{
 		if((MaxDists[AxisIndex] < -OtherTribox.MaxDists[AxisIndex+1]) ||
 			(OtherTribox.MaxDists[AxisIndex] < -MaxDists[AxisIndex+1]))
 		{
-			bOverlapTriboxes = false;
-			break;
+			return false;
 		}
 	}
+	return true;
+}
+	
+bool FTribox::OverlapTribox(const FTribox& OtherTribox, FTribox& OverlapTribox) const
+{
+	bool bOverlapTriboxes = IsTriboxOverlapping(OtherTribox);
 	if(bOverlapTriboxes)
 	{
 		for(int32 AxisIndex = 0; AxisIndex < NumPlanes; AxisIndex += 2)
diff --git a/Engine/Source/Runtime/Experimental/Chaos/Private/PBDRigidsSolver.cpp b/Engine/Source/Runtime/Experimental/Chaos/Private/PBDRigidsSolver.cpp
--- a/Engine/Source/Runtime/Experimental/Chaos/Private/PBDRigidsSolver.cpp
+++ b/Engine/Source/Runtime/Experimental/Chaos/Private/PBDRigidsSolver.cpp
@@ -27,6 +27,7 @@
 #include "ChaosSolverConfiguration.h"
 #include "Chaos/PullPhysicsDataImp.h"
 #include "Chaos/PhysicsSolverBaseImpl.h"
+#include "Chaos/ConvexOptimizer.h"
 
 #include "ProfilingDebugging/CountersTrace.h"
 #include "ProfilingDebugging/CsvProfiler.h"
@@ -2265,6 +2266,13 @@ TRACE_COUNTER_SET(ChaosTraceCounter_##Name, Value)
 				{
 					NumShapes = Union->GetNumLeafObjects();
 				}
+				if(auto* ClusteredParticle = Particle.CastToClustered())
+				{
+					if((ClusteredParticle->ConvexOptimizer().Get() != nullptr) && ClusteredParticle->ConvexOptimizer()->IsValid())
+					{
+						NumShapes = ClusteredParticle->ConvexOptimizer()->NumCollisionObjects();
+					}
+				}
 				if (P->IsDynamic())
 				{
 					NumDynamicShapes += NumShapes;
diff --git a/Engine/Source/Runtime/Experimental/Chaos/Public/Chaos/Convex.h b/Engine/Source/Runtime/Experimental/Chaos/Public/Chaos/Convex.h
--- a/Engine/Source/Runtime/Experimental/Chaos/Public/Chaos/Convex.h
+++ b/Engine/Source/Runtime/Experimental/Chaos/Public/Chaos/Convex.h
@@ -100,6 +100,22 @@ namespace Chaos
 			ComputeUnitMassInertiaTensorAndRotationOfMass(Volume);
 		}
 
+		FConvex(TArray<FPlaneType>&& InPlanes, TArray<TArray<int32>>&& InFaceIndices, TArray<FVec3Type>&& InVertices, 
+			const FVec3Type& InMin, const FVec3Type& InMax, const FRealType InVolume, const FVec3Type InInertiaTensor, const FRotation3& InRotationMatrix, const bool bRegularDatas)
+			: FImplicitObject(EImplicitObject::IsConvex | EImplicitObject::HasBoundingBox, ImplicitObjectType::Convex)
+			, Planes(MoveTemp(InPlanes))
+			, Vertices(MoveTemp(InVertices))
+		{
+			LocalBoundingBox = FAABB3Type(InMin, InMax);
+			CenterOfMass = LocalBoundingBox.GetCenterOfMass();
+			
+			Volume = InVolume;
+			RotationOfMass = InRotationMatrix;
+			UnitMassInertiaTensor = InInertiaTensor;
+			
+			CreateStructureData(MoveTemp(InFaceIndices), bRegularDatas);
+		}
+
 		FConvex(const TArray<FVec3Type>& InVertices, const FReal InMargin, FConvexBuilder::EBuildMethod BuildMethod = FConvexBuilder::EBuildMethod::Default)
 		    : FImplicitObject(EImplicitObject::IsConvex | EImplicitObject::HasBoundingBox, ImplicitObjectType::Convex)
 		{
@@ -184,7 +200,7 @@ namespace Chaos
 		CHAOS_API void MovePlanesAndRebuild(FRealType InDelta);
 
 	private:
-		CHAOS_API void CreateStructureData(TArray<TArray<int32>>&& FaceIndices);
+		CHAOS_API void CreateStructureData(TArray<TArray<int32>>&& FaceIndices, const bool bRegularDatas = false);
 
 	public:
 		static constexpr EImplicitObjectType StaticType()
@@ -224,7 +240,6 @@ namespace Chaos
 			return PhiWithNormalScaledInternal(X, Scale, Normal);
 		}
 
-
 	private:
 		// Distance to the surface
 		FReal PhiWithNormalInternal(const FVec3& X, FVec3& Normal) const
diff --git a/Engine/Source/Runtime/Experimental/Chaos/Public/Chaos/ConvexHalfEdgeStructureData.h b/Engine/Source/Runtime/Experimental/Chaos/Public/Chaos/ConvexHalfEdgeStructureData.h
--- a/Engine/Source/Runtime/Experimental/Chaos/Public/Chaos/ConvexHalfEdgeStructureData.h
+++ b/Engine/Source/Runtime/Experimental/Chaos/Public/Chaos/ConvexHalfEdgeStructureData.h
@@ -93,7 +93,14 @@ namespace Chaos
 		struct FVertexPlanes
 		{
 			FIndex PlaneIndices[3];
-			FIndex NumPlaneIndices;
+			FIndex NumPlaneIndices = 0;
+		};
+
+		// List of 3 halfedges per vertex for regular convexes  
+		struct FVertexHalfEdges
+		{
+			FIndex HalfEdgeIndices[3];
+			FIndex NumHalfEdgeIndices = 0;
 		};
 
 		// Initialize the structure data from the array of vertex indices per plane (in CW or CCW order - it is retained in structure)
@@ -324,6 +331,111 @@ namespace Chaos
 			return (int32)VertexPlane.NumPlaneIndices;
 		}
 
+		// Build half edge structure datas for regular convexes with exactly 3 faces per vertex
+		bool BuildRegularDatas(const TArray<TArray<int32>>& InPlaneVertices, int32 InNumVertices)
+		{
+			// Count the edges
+			int32 HalfEdgeCount = 0;
+			for (int32 PlaneIndex = 0; PlaneIndex < InPlaneVertices.Num(); ++PlaneIndex)
+			{
+				HalfEdgeCount += InPlaneVertices[PlaneIndex].Num();
+			}
+
+			if ((InPlaneVertices.Num() > MaxIndex) || (HalfEdgeCount > MaxIndex) || (InNumVertices > MaxIndex))
+			{
+				// We should never get here. See GetRequiredIndexType::GetRequiredIndexType which calls CanMake() on eachn index type until it fits
+				UE_LOG(LogChaos, Error, TEXT("Unable to create structure data for convex. MaxIndex too small (%d bytes) for Planes: %d HalfEdges: %d Verts: %d"), sizeof(FIndex), InPlaneVertices.Num(), HalfEdgeCount, InNumVertices);
+				return false;
+			}
+
+			Planes.SetNum(InPlaneVertices.Num());
+			HalfEdges.SetNum(HalfEdgeCount);
+			Vertices.SetNum(InNumVertices);
+			VertexPlanes.SetNum(InNumVertices);
+
+			TArray<FVertexHalfEdges> VertexHalfEdges;
+			TArray<FIndex> HalfEdgeVertices;
+			
+			VertexHalfEdges.SetNum(InNumVertices);
+			HalfEdgeVertices.SetNum(HalfEdgeCount);
+
+			// Initialize the vertex list - it will be filled in as we build the edge list
+			for (int32 VertexIndex = 0; VertexIndex < Vertices.Num(); ++VertexIndex)
+			{
+				GetVertex(VertexIndex).FirstHalfEdgeIndex = InvalidIndex;
+			}
+
+			// Build the planes and edges. The edges for a plane are stored sequentially in the half-edge array.
+			// On the first pass, the edges contain 2 vertex indices, rather than a vertex index and a twin edge index.
+			// We fix this up on a second pass.
+			int32 NextHalfEdgeIndex = 0;
+			for (int32 PlaneIndex = 0; PlaneIndex < InPlaneVertices.Num(); ++PlaneIndex)
+			{
+				const TArray<int32>& PlaneVertices = InPlaneVertices[PlaneIndex];
+
+				GetPlane(PlaneIndex) =
+				{
+					(FIndex)NextHalfEdgeIndex,
+					(FIndex)PlaneVertices.Num()
+				};
+
+				for (int32 PlaneVertexIndex = 0; PlaneVertexIndex < PlaneVertices.Num(); ++PlaneVertexIndex)
+				{
+					// Add a new edge
+					const int32 VertexIndex0 = PlaneVertices[PlaneVertexIndex];
+					const int32 VertexIndex1 = PlaneVertices[(PlaneVertexIndex + 1) % PlaneVertices.Num()];
+					GetHalfEdge(NextHalfEdgeIndex) =
+					{
+						(FIndex)PlaneIndex,
+						(FIndex)VertexIndex0,
+						(FIndex)VertexIndex1,	// Will get converted to a half-edge index later
+					};
+					HalfEdgeVertices[NextHalfEdgeIndex] = (FIndex)VertexIndex1;
+
+					// If this is the first time Vertex0 has showed up, set its edge index
+					// For valid and regular convexes each vertices have 2 halfedges starting from itself
+					if (Vertices[VertexIndex0].FirstHalfEdgeIndex == InvalidIndex)
+					{
+						Vertices[VertexIndex0].FirstHalfEdgeIndex = (FIndex)NextHalfEdgeIndex;
+					}
+					
+					// For regular convexes each vertices have exactly 3 halfedges
+                    VertexHalfEdges[VertexIndex0].HalfEdgeIndices[VertexHalfEdges[VertexIndex0].NumHalfEdgeIndices++] = (FIndex)NextHalfEdgeIndex;
+
+					// For regular convexes each vertices have exactly 3 planes
+					VertexPlanes[VertexIndex0].PlaneIndices[VertexPlanes[VertexIndex0].NumPlaneIndices++] = (FIndex)PlaneIndex;
+
+					++NextHalfEdgeIndex;
+				}
+			}
+			Edges.Empty();
+			Edges.Reserve(NumHalfEdges() / 2);
+			
+			for (int32 HalfEdgeIndex = 0; HalfEdgeIndex < HalfEdges.Num(); ++HalfEdgeIndex)
+			{
+				const int32 VertexIndex0 = HalfEdges[HalfEdgeIndex].VertexIndex;
+				const int32 VertexIndex1 = HalfEdges[HalfEdgeIndex].TwinHalfEdgeIndex;
+
+				for(int32 VertexHalfEdgeIndex = 0, NumHalfEdgeIndices = VertexHalfEdges[VertexIndex1].NumHalfEdgeIndices; 
+				                     VertexHalfEdgeIndex < NumHalfEdgeIndices; ++VertexHalfEdgeIndex)
+				{
+				    const FIndex TwinHalfEdgeIndex = VertexHalfEdges[VertexIndex1].HalfEdgeIndices[VertexHalfEdgeIndex];
+					if(HalfEdgeVertices[TwinHalfEdgeIndex] == VertexIndex0)
+					{
+						HalfEdges[HalfEdgeIndex].TwinHalfEdgeIndex = TwinHalfEdgeIndex;
+						break;
+					}
+				}
+
+				if(VertexIndex0 < HalfEdges[HalfEdges[HalfEdgeIndex].TwinHalfEdgeIndex].VertexIndex)
+				{
+					Edges.Add((FIndex)HalfEdgeIndex);
+				}
+			}
+			return true;
+		}
+		
+
 		// Initialize the structure data from the set of vertices associated with each plane.
 		// The vertex indices are assumed to be in CCW order (or CW order - doesn't matter here
 		// as long as it is sequential).
diff --git a/Engine/Source/Runtime/Experimental/Chaos/Public/Chaos/ConvexOptimizer.h b/Engine/Source/Runtime/Experimental/Chaos/Public/Chaos/ConvexOptimizer.h
--- a/Engine/Source/Runtime/Experimental/Chaos/Public/Chaos/ConvexOptimizer.h
+++ b/Engine/Source/Runtime/Experimental/Chaos/Public/Chaos/ConvexOptimizer.h
@@ -42,7 +42,7 @@ namespace Private
 		CHAOS_API ~FConvexOptimizer();
 
 		// Simplify all the convexes in the hierarchy 
-		CHAOS_API void SimplifyRootConvexes(const Chaos::FImplicitObjectUnionPtr& UnionGeometry, const FShapesArray& UnionShapes);
+		CHAOS_API void SimplifyRootConvexes(const Chaos::FImplicitObjectUnionPtr& UnionGeometry, const FShapesArray& UnionShapes, const EObjectStateType ObjectState);
 
 		// Check if the manager is valid or not
 		CHAOS_API bool IsValid() const {return !SimplifiedConvexes.IsEmpty();}
@@ -55,6 +55,9 @@ namespace Private
 
 		// Get the shapes array 
 		CHAOS_API const FShapeInstanceArray& GetShapeInstances() const { return ShapesArray;}
+		
+		// Get the number of collision objects
+		CHAOS_API int32 NumCollisionObjects() const; 
 
 	private:
 
@@ -62,7 +65,7 @@ namespace Private
 		void BuildSingleConvex(const Chaos::FImplicitObjectUnionPtr& UnionGeometry, const FShapesArray& UnionShapes);
 
 		// Build several convexes 
-		void BuildMultipleConvex(const Chaos::FImplicitObjectUnionPtr& UnionGeometry, const FShapesArray& UnionShapes);
+		void BuildMultipleConvex(const Chaos::FImplicitObjectUnionPtr& UnionGeometry, const FShapesArray& UnionShapes, const int32 MaxLODs);
 
 		// Build the simplified shapes
 		void BuildConvexShapes(const FShapesArray& UnionShapes);
diff --git a/Engine/Source/Runtime/Experimental/Chaos/Public/Chaos/ConvexStructureData.h b/Engine/Source/Runtime/Experimental/Chaos/Public/Chaos/ConvexStructureData.h
--- a/Engine/Source/Runtime/Experimental/Chaos/Public/Chaos/ConvexStructureData.h
+++ b/Engine/Source/Runtime/Experimental/Chaos/Public/Chaos/ConvexStructureData.h
@@ -238,13 +238,15 @@ namespace Chaos
 		}
 
 		// Initialize the structure data from the set of vertices for each face of the convex
-		bool SetPlaneVertices(const TArray<TArray<int32>>& InPlaneVertices, int32 NumVerts)
+		bool SetPlaneVertices(const TArray<TArray<int32>>& InPlaneVertices, int32 NumVerts, const bool bRegularDatas = false)
 		{
 			const EIndexType NewIndexType = GetRequiredIndexType(InPlaneVertices, NumVerts);
 			CreateDataContainer(NewIndexType);
 
 			return NonConstDataOp(
-				[InPlaneVertices, NumVerts](auto& ConcreteData) { return ConcreteData.SetPlaneVertices(InPlaneVertices, NumVerts); },
+				[InPlaneVertices, NumVerts, bRegularDatas](auto& ConcreteData) {
+				return bRegularDatas ? ConcreteData.BuildRegularDatas(InPlaneVertices, NumVerts) : ConcreteData.SetPlaneVertices(InPlaneVertices, NumVerts); 
+					},
 				[]() { return true; });
 		}
 
diff --git a/Engine/Source/Runtime/Experimental/Chaos/Public/Chaos/ImplicitObjectScaled.h b/Engine/Source/Runtime/Experimental/Chaos/Public/Chaos/ImplicitObjectScaled.h
--- a/Engine/Source/Runtime/Experimental/Chaos/Public/Chaos/ImplicitObjectScaled.h
+++ b/Engine/Source/Runtime/Experimental/Chaos/Public/Chaos/ImplicitObjectScaled.h
@@ -221,6 +221,32 @@ public:
 		}
 		return static_cast<const TImplicitObjectInstanced<TConcrete>&>(Obj);
 	}
+	
+	static const TImplicitObjectInstanced<TConcrete>* AsInstanced(const FImplicitObject& Obj)
+	{
+		if constexpr (std::is_same_v<TConcrete, FImplicitObject>)
+		{
+			//can cast any scaled to ImplicitObject base
+			return IsInstanced(Obj.GetType()) ? static_cast<const TImplicitObjectInstanced<TConcrete>*>(&Obj) : nullptr;
+		}
+		else
+		{
+			return StaticType() == Obj.GetType() ? static_cast<const TImplicitObjectInstanced<TConcrete>*>(&Obj) : nullptr;
+		}
+	}
+
+	static TImplicitObjectInstanced<TConcrete>* AsInstanced(FImplicitObject& Obj)
+	{
+		if constexpr (std::is_same_v<TConcrete, FImplicitObject>)
+		{
+			//can cast any scaled to ImplicitObject base
+			return IsInstanced(Obj.GetType()) ? static_cast<TImplicitObjectInstanced<TConcrete>*>(&Obj) : nullptr;
+		}
+		else
+		{
+			return StaticType() == Obj.GetType() ? static_cast<TImplicitObjectInstanced<TConcrete>*>(&Obj) : nullptr;
+		}
+	}
 
 	/** This is a low level function and assumes the internal object has a SweepGeom function. Should not be called directly. See GeometryQueries.h : SweepQuery */
 	template <typename QueryGeomType>
diff --git a/Engine/Source/Runtime/Experimental/Chaos/Public/Chaos/Tribox.h b/Engine/Source/Runtime/Experimental/Chaos/Public/Chaos/Tribox.h
--- a/Engine/Source/Runtime/Experimental/Chaos/Public/Chaos/Tribox.h
+++ b/Engine/Source/Runtime/Experimental/Chaos/Public/Chaos/Tribox.h
@@ -78,6 +78,9 @@ public :
 	// Find the overlapping tribox 
 	bool OverlapTribox(const FTribox& OtherTribox, FTribox& OverlapTribox) const;
 
+	// Check if the tribox is overlapping or not
+	bool IsTriboxOverlapping(const FTribox& OtherTribox) const;
+
 	// Split the tribox in 2 along a defined cuttng plane
 	bool SplitTriboxSlab(const int32 PlaneAxis, const FRealType& PlaneDistance,
 					FTribox& LeftTribox, FTribox& RightTribox) const;
@@ -135,6 +138,10 @@ private :
 	void BuildChamferDist(const FVec3Type& P, const int32 CoordIndexA, const int32 CoordIndexB,
 			const int32 DistsIndexA, const int32 DistsIndexB, const int32 DistsIndexC, const int32 DistsIndexD);
 
+	// Create a FConvex from a list of convex planes, face indices and convex vertices
+	FImplicitObjectPtr CreateConvexFromTopology(TArray<FConvex::FPlaneType>&& ConvexPlanes,
+		TArray<TArray<int32>>&& FaceIndices, TArray<FConvex::FVec3Type>&& ConvexVertices) const;
+
 	// Max distance along eaxh tribox axis (principal + chamfer)
 	FRealType MaxDists[NumPlanes] = {TNumericLimits<FRealType>::Lowest(), TNumericLimits<FRealType>::Lowest(),
 									 TNumericLimits<FRealType>::Lowest(), TNumericLimits<FRealType>::Lowest(),
