Ias: Disable VeryVerbose log entries and add periodic Statistics log entries.

#rnx

diff --git a/Engine/Source/Runtime/Experimental/IoStoreOnDemand/Private/Statistics.cpp b/Engine/Source/Runtime/Experimental/IoStoreOnDemand/Private/Statistics.cpp
--- a/Engine/Source/Runtime/Experimental/IoStoreOnDemand/Private/Statistics.cpp
+++ b/Engine/Source/Runtime/Experimental/IoStoreOnDemand/Private/Statistics.cpp
@@ -2,9 +2,11 @@
 
 #include "Statistics.h"
 
+#include "Algo/MaxElement.h"
 #include "AnalyticsEventAttribute.h"
 #include "HAL/IConsoleManager.h"
 #include "Misc/CoreDelegates.h"
+#include "IO/IoStoreOnDemand.h"
 
 LLM_DEFINE_TAG(Ias);
 
@@ -13,6 +15,12 @@ LLM_DEFINE_TAG(Ias);
 namespace UE::IO::IAS
 {
 
+float GIasStatisticsLogInterval = 60.f;
+static FAutoConsoleVariableRef CVar_StatisticsLogInterval(
+	TEXT("ias.StatisticsLogInterval"),
+	GIasStatisticsLogInterval,
+	TEXT("Enables and sets interval for periodic logging of statistics"));
+
 bool GIasReportHttpAnalyticsEnabled = true;
 static FAutoConsoleVariableRef CVar_ReportHttpAnalytics(
 	TEXT("ias.ReportHttpAnalytics"),
@@ -276,24 +284,53 @@ FOnDemandIoBackendStats::FOnDemandIoBackendStats()
 	GStatisticsEndFrameDelegateHandle = FCoreDelegates::OnEndFrame.AddLambda([this]()
 	{
 		// cache stat totals
-		CSV_CUSTOM_STAT_DEFINED(CacheGetCount, int32(GCacheGetCount.Get()), ECsvCustomStatOp::Set);
-		CSV_CUSTOM_STAT_DEFINED(CacheErrorCount, int32(GCacheErrorCount.Get()), ECsvCustomStatOp::Set);
-		CSV_CUSTOM_STAT_DEFINED(CachePutCount, int32(GCachePutCount.Get()), ECsvCustomStatOp::Set);
-		CSV_CUSTOM_STAT_DEFINED(CachePutExistingCount, int32(GCachePutExistingCount.Get()), ECsvCustomStatOp::Set);
-		CSV_CUSTOM_STAT_DEFINED(CachePutRejectCount, int32(GCachePutRejectCount.Get()), ECsvCustomStatOp::Set);
-		CSV_CUSTOM_STAT_DEFINED(CacheCachedMB, BytesToApproxMB(GCacheCachedBytes.Get()), ECsvCustomStatOp::Set);
-		CSV_CUSTOM_STAT_DEFINED(CacheReadMB, BytesToApproxMB(GCacheReadBytes.Get()), ECsvCustomStatOp::Set);
-		CSV_CUSTOM_STAT_DEFINED(CacheRejectedMB, BytesToApproxMB(GCacheRejectBytes.Get()), ECsvCustomStatOp::Set);
+		int32 CGetCount = int32(GCacheGetCount.Get());
+		int32 CErrorCount = int32(GCacheErrorCount.Get());
+		int32 CPutCount = int32(GCachePutCount.Get());
+		int32 CPutExistingCount = int32(GCachePutExistingCount.Get());
+		int32 CPutRejectCount = int32(GCachePutRejectCount.Get());
+		int32 CCachedMB = BytesToApproxMB(GCacheCachedBytes.Get());
+		int32 CReadMB = BytesToApproxMB(GCacheReadBytes.Get());
+		int32 CRejectedMB = BytesToApproxMB(GCacheRejectBytes.Get());
+		CSV_CUSTOM_STAT_DEFINED(CacheGetCount, CGetCount, ECsvCustomStatOp::Set);
+		CSV_CUSTOM_STAT_DEFINED(CacheErrorCount, CErrorCount, ECsvCustomStatOp::Set);
+		CSV_CUSTOM_STAT_DEFINED(CachePutCount, CPutCount, ECsvCustomStatOp::Set);
+		CSV_CUSTOM_STAT_DEFINED(CachePutExistingCount, CPutExistingCount, ECsvCustomStatOp::Set);
+		CSV_CUSTOM_STAT_DEFINED(CachePutRejectCount, CPutRejectCount, ECsvCustomStatOp::Set);
+		CSV_CUSTOM_STAT_DEFINED(CacheCachedMB, CCachedMB, ECsvCustomStatOp::Set);
+		CSV_CUSTOM_STAT_DEFINED(CacheReadMB, CReadMB, ECsvCustomStatOp::Set);
+		CSV_CUSTOM_STAT_DEFINED(CacheRejectedMB, CRejectedMB, ECsvCustomStatOp::Set);
 
 		// http stat totals
-		CSV_CUSTOM_STAT_DEFINED(HttpGetCount, int32(GHttpGetCount.Get()), ECsvCustomStatOp::Set);
-		CSV_CUSTOM_STAT_DEFINED(HttpRetryCount, int32(GHttpRetryCount.Get()), ECsvCustomStatOp::Set);
-		CSV_CUSTOM_STAT_DEFINED(HttpErrorCount, int32(GHttpErrorCount.Get()), ECsvCustomStatOp::Set);
-		CSV_CUSTOM_STAT_DEFINED(HttpPendingCount, int32(GHttpPendingCount.Get()), ECsvCustomStatOp::Set);
-		CSV_CUSTOM_STAT_DEFINED(HttpDownloadedMB, BytesToApproxMB(GHttpDownloadedBytes.Get()), ECsvCustomStatOp::Set);
-		CSV_CUSTOM_STAT_DEFINED(HttpBandwidthMpbs, int32(GHttpBandwidthMpbs.Get()), ECsvCustomStatOp::Set);
-		CSV_CUSTOM_STAT_DEFINED(HttpDurationMsAvg, int32(GHttpDurationMsAvg.Get()), ECsvCustomStatOp::Set);
-		CSV_CUSTOM_STAT_DEFINED(HttpDurationMsMax, int32(GHttpDurationMsMax.Get()), ECsvCustomStatOp::Set);
+		int32 HGetCount = int32(GHttpGetCount.Get());
+		int32 HRetryCount = int32(GHttpRetryCount.Get());
+		int32 HErrorCount = int32(GHttpErrorCount.Get());
+		int32 HPendingCount = int32(GHttpPendingCount.Get());
+		int32 HDownloadedMB = BytesToApproxMB(GHttpDownloadedBytes.Get());
+		int32 HBandwidthMpbs = int32(GHttpBandwidthMpbs.Get());
+		int32 HDurationMsAvg = int32(GHttpDurationMsAvg.Get());
+		int32 HDurationMsMax = int32(GHttpDurationMsMax.Get());
+		CSV_CUSTOM_STAT_DEFINED(HttpGetCount, HGetCount, ECsvCustomStatOp::Set);
+		CSV_CUSTOM_STAT_DEFINED(HttpRetryCount, HRetryCount, ECsvCustomStatOp::Set);
+		CSV_CUSTOM_STAT_DEFINED(HttpErrorCount, HErrorCount, ECsvCustomStatOp::Set);
+		CSV_CUSTOM_STAT_DEFINED(HttpPendingCount, HPendingCount, ECsvCustomStatOp::Set);
+		CSV_CUSTOM_STAT_DEFINED(HttpDownloadedMB, HDownloadedMB, ECsvCustomStatOp::Set);
+		CSV_CUSTOM_STAT_DEFINED(HttpBandwidthMpbs, HBandwidthMpbs, ECsvCustomStatOp::Set);
+		CSV_CUSTOM_STAT_DEFINED(HttpDurationMsAvg, HDurationMsAvg, ECsvCustomStatOp::Set);
+		CSV_CUSTOM_STAT_DEFINED(HttpDurationMsMax, HDurationMsMax, ECsvCustomStatOp::Set);
+
+		if (GIasStatisticsLogInterval > 0.f)
+		{
+			static double LastLogTime = 0.0;
+			if (double Time = FPlatformTime::Seconds(); Time - LastLogTime > (double)GIasStatisticsLogInterval)
+			{
+				UE_LOG(LogIas, Log, TEXT("CacheStats: CachedMB=%d, RejectedMB=%d, ReadMB=%d, Get=%d, Error=%d, Put=%d, PutReject=%d, PutExisting=%d"),
+					CCachedMB, CRejectedMB, CReadMB, CGetCount, CErrorCount, CPutCount, CPutRejectCount, CPutExistingCount);
+				UE_LOG(LogIas, Log, TEXT("HttpStats: DownloadedMB=%d, Get=%d, Retry=%d, Error=%d, CurPending=%d, CurDurationMsAvg=%d, CurDurationMsMax=%d"),
+					HDownloadedMB, HGetCount, HRetryCount, HErrorCount, HPendingCount, HDurationMsAvg, HDurationMsMax);
+				LastLogTime = Time;
+			}
+		}
 	});
 }
 
@@ -445,9 +482,13 @@ void FOnDemandIoBackendStats::OnHttpGet(uint64 SizeBytes, uint64 DurationMs)
 	GHttpBandwidthMpbs.Set((GHttpHistoryTotalBytes*8)/(GHttpHistoryTotalDuration+1)/1000);
 	GHttpDurationMsAvg.Set(GHttpHistoryTotalDuration/GHttpHistoryCount);
 
-	if (GHttpDurationMsMax.Get() < (int64)DurationMs)
+	if (NewDuration > GHttpDurationMsMax.Get())
+	{
+		GHttpDurationMsMax.Set(NewDuration);
+	}
+	else if (OldDuration == GHttpDurationMsMax.Get())
 	{
-		GHttpDurationMsMax.Set((int64)DurationMs);
+		GHttpDurationMsMax.Set(*Algo::MaxElement(GHttpHistoryDuration));
 	}
 
 	GHttpHistoryIndex = (GHttpHistoryIndex + 1) % GHttpHistoryCount;
diff --git a/Engine/Source/Runtime/Experimental/IoStoreOnDemand/Public/IO/IoStoreOnDemand.h b/Engine/Source/Runtime/Experimental/IoStoreOnDemand/Public/IO/IoStoreOnDemand.h
--- a/Engine/Source/Runtime/Experimental/IoStoreOnDemand/Public/IO/IoStoreOnDemand.h
+++ b/Engine/Source/Runtime/Experimental/IoStoreOnDemand/Public/IO/IoStoreOnDemand.h
@@ -34,7 +34,7 @@ using FIoBlockHash = uint32;
 	#define UE_IAS_CUSTOM_INITIALIZATION 0
 #endif
 
-UE_API DECLARE_LOG_CATEGORY_EXTERN(LogIas, VeryVerbose, All);
+UE_API DECLARE_LOG_CATEGORY_EXTERN(LogIas, Log, All);
 
 namespace UE::IO::IAS
 {
