Ias: Add missing include for Statistics.cpp compile error

#rnx

diff --git a/Engine/Source/Runtime/Experimental/IoStoreOnDemand/Private/Statistics.cpp b/Engine/Source/Runtime/Experimental/IoStoreOnDemand/Private/Statistics.cpp
--- a/Engine/Source/Runtime/Experimental/IoStoreOnDemand/Private/Statistics.cpp
+++ b/Engine/Source/Runtime/Experimental/IoStoreOnDemand/Private/Statistics.cpp
@@ -5,6 +5,7 @@
 #include "Algo/MaxElement.h"
 #include "AnalyticsEventAttribute.h"
 #include "HAL/IConsoleManager.h"
+#include "HAL/PlatformTime.h"
 #include "Misc/CoreDelegates.h"
 #include "IO/IoStoreOnDemand.h"
 
