Fix CIS for the convex optimizer

diff --git a/Engine/Source/Runtime/Experimental/Chaos/Private/Chaos/ChaosDebugDraw.cpp b/Engine/Source/Runtime/Experimental/Chaos/Private/Chaos/ChaosDebugDraw.cpp
--- a/Engine/Source/Runtime/Experimental/Chaos/Private/Chaos/ChaosDebugDraw.cpp
+++ b/Engine/Source/Runtime/Experimental/Chaos/Private/Chaos/ChaosDebugDraw.cpp
@@ -710,7 +710,7 @@ namespace Chaos
 			{
 				if(InnerType == ImplicitObjectType::Convex)
 				{
-					ShapeColor = bHasConvexOptimizer ? FColor::Green : Particle->GetGeometry()->IsUnderlyingUnion() ? FColor::Blue : FColor::Orange;
+					ShapeColor = bHasConvexOptimizer ? FColor::Green : FColor::Orange;
 				}
 				else
 				{
diff --git a/Engine/Source/Runtime/Experimental/Chaos/Private/Chaos/ConvexOptimizer.cpp b/Engine/Source/Runtime/Experimental/Chaos/Private/Chaos/ConvexOptimizer.cpp
--- a/Engine/Source/Runtime/Experimental/Chaos/Private/Chaos/ConvexOptimizer.cpp
+++ b/Engine/Source/Runtime/Experimental/Chaos/Private/Chaos/ConvexOptimizer.cpp
@@ -410,7 +410,6 @@ FORCEINLINE void BuildConvexTriboxes(const Chaos::FImplicitObjectUnionPtr& Union
 							LocalTribox.AddConvex(ScaledObject->GetUnscaledObject(), ConvexTransform);
 						}
 					}
-
 					else if(const TImplicitObjectInstanced<FConvex>* InstancedObject = TImplicitObjectInstanced<FConvex>::AsInstanced(*ImplicitObject))
 					{
 						if(!bHasRootTribox)
@@ -419,7 +418,7 @@ FORCEINLINE void BuildConvexTriboxes(const Chaos::FImplicitObjectUnionPtr& Union
 							// Disable collision for all the convexes that are going to be used to build the tribox
 							// For now only used for debug draw
 							const_cast<FImplicitObject*>(ImplicitObject)->SetDoCollide(false);
-							LocalTribox.AddConvex(ScaledObject->Object(), ConvexTransform);
+							LocalTribox.AddConvex(InstancedObject->Object(), ConvexTransform);
 						}
 					}
 					else
@@ -506,7 +505,7 @@ void FConvexOptimizer::BuildMultipleConvex(const Chaos::FImplicitObjectUnionPtr&
 	}
 	else
 	{
-		// Build the binary tree and add the leaves to the cimplified convexes
+		// Build the binary tree and add the leaves to the simplified convexes
 		FTriboxTree TriboxTree(MaxLODs, RootTriboxes);
 		TriboxTree.BuildTreeLODs(SimplifiedConvexes, CollisionObjects);
 	}
