IAS - make endpoint changes configurable and some minor cleanup

#rnx

diff --git a/Engine/Source/Runtime/Experimental/IoStoreOnDemand/Private/OnDemandHttpClient.cpp b/Engine/Source/Runtime/Experimental/IoStoreOnDemand/Private/OnDemandHttpClient.cpp
--- a/Engine/Source/Runtime/Experimental/IoStoreOnDemand/Private/OnDemandHttpClient.cpp
+++ b/Engine/Source/Runtime/Experimental/IoStoreOnDemand/Private/OnDemandHttpClient.cpp
@@ -19,13 +19,13 @@ static void LogHttpResult(const TCHAR* Host, const TCHAR* Url, uint32 StatusCode
 
 void FHttpClient::Get(FAnsiStringView Url, const FIoOffsetAndLength& Range, FGetCallback&& Callback)
 {
-	check(PrimaryConnection != INDEX_NONE);
+	check(CurrentEndpoint != INDEX_NONE);
 	IssueRequest(FRequestParams
 	{
 		.Url = FString(Url),
 		.Range = Range,
 		.Callback = MoveTemp(Callback),
-		.Connection = PrimaryConnection
+		.Endpoint = CurrentEndpoint
 	});
 }
 
@@ -36,11 +36,6 @@ void FHttpClient::Get(FAnsiStringView Url, FGetCallback&& Callback)
 
 bool FHttpClient::Tick(uint32 WaitTimeMs, uint32 MaxKiBPerSecond)
 {
-	if (PrimaryConnection == INDEX_NONE)
-	{
-		return false;
-	}
-
 	EventLoop.Throttle(MaxKiBPerSecond);
 	const uint32 TicketCount = EventLoop.Tick(WaitTimeMs);
 
@@ -59,7 +54,7 @@ bool FHttpClient::Tick(uint32 WaitTimeMs, uint32 MaxKiBPerSecond)
 	{
 		for (int32 Idx = 0, Count = Connections.Num(); Idx < Count; ++Idx)
 		{
-			if (Idx != PrimaryConnection)
+			if (Idx != CurrentEndpoint)
 			{
 				Connections[Idx].Reset();
 			}
@@ -69,10 +64,23 @@ bool FHttpClient::Tick(uint32 WaitTimeMs, uint32 MaxKiBPerSecond)
 	return bIsIdle == false;
 }
 
+void FHttpClient::SetEndpoint(int32 Endpoint)
+{
+	if (CurrentEndpoint == Endpoint)
+	{
+		return;
+	}
+
+	EnsureConnection(Endpoint);
+	CurrentEndpoint = Endpoint;
+}
+
 FHttpClient::FHttpClient(FHttpClientConfig&& ClientConfig)
 	: Config(MoveTemp(ClientConfig))
 {
-	Configure();
+	check(Config.Endpoints.IsEmpty() == false);
+	Connections.SetNum(Config.Endpoints.Num());
+	SetEndpoint(Config.PrimaryEndpoint);
 }
 
 TUniquePtr<FHttpClient> FHttpClient::Create(FHttpClientConfig&& ClientConfig)
@@ -106,14 +114,6 @@ TUniquePtr<FHttpClient> FHttpClient::Create(const FString& Endpoint)
 	return Create(MoveTemp(Config));
 }
 
-void FHttpClient::Configure()
-{
-	check(Config.Endpoints.IsEmpty() == false);
-	Connections.SetNum(Config.Endpoints.Num());
-	Connections[Config.PrimaryEndpoint] = CreateConnection(Config.Endpoints[Config.PrimaryEndpoint]);
-	PrimaryConnection = Config.PrimaryEndpoint;
-}
-
 void FHttpClient::IssueRequest(FRequestParams&& Params)
 {
 	using namespace UE::IO::IAS::HTTP;
@@ -142,17 +142,15 @@ void FHttpClient::IssueRequest(FRequestParams&& Params)
 			case FTicketStatus::EId::Content:
 			{
 				const FIoBuffer& Content = Status.GetContent();
-				LogHttpResult(*Config.Endpoints[Params.Connection], *Params.Url, StatusCode, DurationMs, Content.GetSize(), Params.Range.GetOffset());
+				LogHttpResult(*GetEndpointUrl(Params.Endpoint), *Params.Url, StatusCode, DurationMs, Content.GetSize(), Params.Range.GetOffset());
 				
 				const bool bSuccessful = StatusCode > 199 && StatusCode < 300;
 				if (bSuccessful && Content.GetSize() > 0)
 				{
 					Params.Callback(Content, DurationMs);
-					if (Params.Attempt > 1 && Params.Connection != PrimaryConnection)
+					if (Params.Attempt > 0 && Config.bChangeEndpointAfterSuccessfulRetry)
 					{
-						UE_LOG(LogIas, Log, TEXT("HTTP client endpoint changed from '%s' to '%s'"),
-							*Config.Endpoints[PrimaryConnection], *Config.Endpoints[Params.Connection]);
-						PrimaryConnection = Params.Connection;
+						SetEndpoint(Params.Endpoint);
 					}
 				}
 				else
@@ -160,8 +158,8 @@ void FHttpClient::IssueRequest(FRequestParams&& Params)
 					const bool bServerError = StatusCode > 499 && StatusCode < 600;
 					if (Params.Attempt < Config.MaxRetryCount)
 					{
-						const bool bNextConnection = bServerError == false && Params.Attempt > 0;
-						RetryRequest(MoveTemp(Params), bNextConnection);
+						const bool bNextEndpoint = bServerError == false && Params.Attempt > 0;
+						RetryRequest(MoveTemp(Params), bNextEndpoint);
 					}
 					else
 					{
@@ -175,26 +173,26 @@ void FHttpClient::IssueRequest(FRequestParams&& Params)
 			{
 				if (Params.Attempt < Config.MaxRetryCount)
 				{
-					const bool bNextConnection = Params.Attempt > 0;
-					RetryRequest(MoveTemp(Params), bNextConnection);
+					const bool bNextEndpoint = Params.Attempt > 0;
+					RetryRequest(MoveTemp(Params), bNextEndpoint);
 				}
 				else
 				{
-					LogHttpResult(*Config.Endpoints[Params.Connection], *Params.Url, StatusCode, DurationMs, 0, Params.Range.GetOffset(), Status.GetErrorReason());
+					LogHttpResult(*GetEndpointUrl(Params.Endpoint), *Params.Url, StatusCode, DurationMs, 0, Params.Range.GetOffset(), Status.GetErrorReason());
 					Params.Callback(FIoStatus(EIoErrorCode::ReadError, FString(Status.GetErrorReason())), DurationMs);
 				}
 				break;
 			}
 			default:
 			{
-				LogHttpResult(*Config.Endpoints[Params.Connection], *Params.Url, StatusCode, DurationMs, 0, Params.Range.GetOffset(), "cancelled");
+				LogHttpResult(*GetEndpointUrl(Params.Endpoint), *Params.Url, StatusCode, DurationMs, 0, Params.Range.GetOffset(), "cancelled");
 				Params.Callback(FIoStatus(EIoErrorCode::Cancelled, FString(Status.GetErrorReason())), DurationMs);
 				break;
 			}
 			}
 		};
 
-	TUniquePtr<FConnectionPool>& Connection = Connections[Params.Connection];
+	TUniquePtr<FConnectionPool>& Connection = Connections[Params.Endpoint];
 	check(Connection.IsValid());
 
 	FRequest Request = EventLoop.Get(Url, *Connection);
@@ -208,16 +206,13 @@ void FHttpClient::IssueRequest(FRequestParams&& Params)
 	EventLoop.Send(MoveTemp(Request), MoveTemp(Sink));
 }
 
-void FHttpClient::RetryRequest(FRequestParams&& Params, bool bNextConnection)
+void FHttpClient::RetryRequest(FRequestParams&& Params, bool bNextEndpoint)
 {
 	check(Params.Attempt < Config.MaxRetryCount);
-	if (bNextConnection && Config.Endpoints.Num() > 1)
+	if (bNextEndpoint && Config.Endpoints.Num() > 1)
 	{
-		Params.Connection = (Params.Connection + 1) % Config.Endpoints.Num();
-		if (Connections[Params.Connection].IsValid() == false)
-		{
-			Connections[Params.Connection] = CreateConnection(Config.Endpoints[Params.Connection]);
-		}
+		Params.Endpoint = (Params.Endpoint + 1) % Config.Endpoints.Num();
+		EnsureConnection(Params.Endpoint);
 	}
 	Params.Attempt++;
 	Retries.Emplace(MoveTemp(Params));
@@ -238,4 +233,20 @@ TUniquePtr<HTTP::FConnectionPool> FHttpClient::CreateConnection(const FStringVie
 	return MakeUnique<HTTP::FConnectionPool>(Params);
 }
 
+void FHttpClient::EnsureConnection(int32 Connection)
+{
+	check(Connection >= INDEX_NONE && Connection < Connections.Num());
+	if (Connection != INDEX_NONE && Connections[Connection].IsValid() == false)
+	{
+		Connections[Connection] = CreateConnection(Config.Endpoints[Connection]);
+	}
+}
+
+const FString& FHttpClient::GetEndpointUrl(int32 Endpoint)
+{
+	check(Endpoint >= INDEX_NONE && Endpoint < Config.Endpoints.Num());
+	static const FString None(TEXT("None"));
+	return Endpoint == INDEX_NONE ? None : Config.Endpoints[Endpoint];
+}
+
 } //namespace UE::IO::IAS
diff --git a/Engine/Source/Runtime/Experimental/IoStoreOnDemand/Private/OnDemandHttpClient.h b/Engine/Source/Runtime/Experimental/IoStoreOnDemand/Private/OnDemandHttpClient.h
--- a/Engine/Source/Runtime/Experimental/IoStoreOnDemand/Private/OnDemandHttpClient.h
+++ b/Engine/Source/Runtime/Experimental/IoStoreOnDemand/Private/OnDemandHttpClient.h
@@ -21,6 +21,7 @@ struct FHttpClientConfig
 	int32 PipelineLength = 1;
 	int32 MaxRetryCount = 1;
 	int32 ReceiveBufferSize = -1;
+	bool bChangeEndpointAfterSuccessfulRetry = true;
 };
 
 class FHttpClient
@@ -37,7 +38,10 @@ public:
 	bool Tick(uint32 WaitTimeMs, uint32 MaxKiBPerSecond);
 	bool Tick() { return Tick(-1, 0); }
 
-	int32 GetPrimaryConnection() const { return PrimaryConnection; }
+	int32 GetEndpoint() const { return CurrentEndpoint; }
+	int32 GetPrimaryEndpoint() const { return Config.PrimaryEndpoint; }
+	void SetEndpoint(int32 Endpoint);
+	bool IsUsingPrimaryEndpoint() const { return CurrentEndpoint == Config.PrimaryEndpoint; }
 
 private:
 	struct FRequestParams
@@ -46,20 +50,21 @@ private:
 		FIoOffsetAndLength Range;
 		FGetCallback Callback;
 		int32 Attempt = 0;
-		int32 Connection = INDEX_NONE;
+		int32 Endpoint = INDEX_NONE;
 	};
 
 	FHttpClient(FHttpClientConfig&& ClientConfig);
-	void Configure();
-	void RetryRequest(FRequestParams&& Params, bool bNextConnection);
+	void RetryRequest(FRequestParams&& Params, bool bNextEndpoint);
 	void IssueRequest(FRequestParams&& Params);
 	TUniquePtr<HTTP::FConnectionPool> CreateConnection(const FStringView& HostAddr);
+	void EnsureConnection(int32 Connection);
+	const FString& GetEndpointUrl(int32 Endpoint);
 
 	FHttpClientConfig Config;
 	HTTP::FEventLoop EventLoop;
 	TArray<TUniquePtr<HTTP::FConnectionPool>> Connections;
 	TArray<FRequestParams> Retries;
-	int32 PrimaryConnection = INDEX_NONE;
+	int32 CurrentEndpoint = INDEX_NONE;
 };
 
 } // namespace UE::IO::IAS
diff --git a/Engine/Source/Runtime/Experimental/IoStoreOnDemand/Private/OnDemandIoDispatcherBackend.cpp b/Engine/Source/Runtime/Experimental/IoStoreOnDemand/Private/OnDemandIoDispatcherBackend.cpp
--- a/Engine/Source/Runtime/Experimental/IoStoreOnDemand/Private/OnDemandIoDispatcherBackend.cpp
+++ b/Engine/Source/Runtime/Experimental/IoStoreOnDemand/Private/OnDemandIoDispatcherBackend.cpp
@@ -63,6 +63,20 @@ namespace UE::IO::IAS
 {
 
 ///////////////////////////////////////////////////////////////////////////////
+int32 GIasHttpPrimaryEndpoint = 0;
+static FAutoConsoleVariableRef CVar_IasHttpPrimaryEndpoint(
+	TEXT("ias.HttpPrimaryEndpoint"),
+	GIasHttpPrimaryEndpoint,
+	TEXT("Primary endpoint to use returned from the distribution endpoint")
+);
+
+bool GIasHttpChangeEndpointAfterSuccessfulRetry = true;
+static FAutoConsoleVariableRef CVar_IasHttpChangeEndpointAfterSuccessfulRetry(
+	TEXT("ias.HttpChangeEndpointAfterSuccessfulRetry"),
+	GIasHttpChangeEndpointAfterSuccessfulRetry,
+	TEXT("Whether to change the current endpoint after a sucessful retry")
+);
+
 int32 GIasHttpPollTimeoutMs = 17;
 static FAutoConsoleVariableRef CVar_GIasHttpPollTimeoutMs(
 	TEXT("ias.HttpPollTimeoutMs"),
@@ -1408,7 +1422,7 @@ private:
 	FIoStatus DownloadoadOnDemandToc(const FString& CdnUrl, const FString& TocPath);
 
 	void AddDeferredTocs();
-	void ProcessHttpRequests(FHttpClient* HttpClient, FBitWindow& HttpErrors, int32 MaxConcurrentRequests);
+	void ProcessHttpRequests(FHttpClient& HttpClient, FBitWindow& HttpErrors, int32 MaxConcurrentRequests);
 	int32 WaitForPendingRequests(float WaitTimeSeconds, float PollTimeSeconds);
 
 	TUniquePtr<IIasCache> Cache;
@@ -2059,7 +2073,7 @@ void FOnDemandIoBackend::ReportAnalytics(TArray<FAnalyticsEventAttribute>& OutAn
 	}
 }
 
-void FOnDemandIoBackend::ProcessHttpRequests(FHttpClient* HttpClient, FBitWindow& HttpErrors, int32 MaxConcurrentRequests)
+void FOnDemandIoBackend::ProcessHttpRequests(FHttpClient& HttpClient, FBitWindow& HttpErrors, int32 MaxConcurrentRequests)
 {
 	int32 NumConcurrentRequests = 0;
 	FChunkRequest* NextChunkRequest = HttpRequests.Dequeue();
@@ -2085,12 +2099,12 @@ void FOnDemandIoBackend::ProcessHttpRequests(FHttpClient* HttpClient, FBitWindow
 				}
 				else
 				{
-					check(HttpClient);
+					check(HttpClient.GetEndpoint() != INDEX_NONE);
 					TAnsiStringBuilder<256> Url;
 					ChunkRequest->Params.GetUrl(Url);
 
 					NumConcurrentRequests++;
-					HttpClient->Get(Url.ToView(), ChunkRequest->Params.ChunkRange,
+					HttpClient.Get(Url.ToView(), ChunkRequest->Params.ChunkRange,
 						[this, &NextChunkRequest, ChunkRequest, &NumConcurrentRequests, &HttpErrors]
 						(TIoStatusOr<FIoBuffer> Status, uint64 DurationMs)
 						{
@@ -2150,7 +2164,7 @@ void FOnDemandIoBackend::ProcessHttpRequests(FHttpClient* HttpClient, FBitWindow
 				TRACE_CPUPROFILER_EVENT_SCOPE(IasBackend::TickHttpSaturated);
 				while (NumConcurrentRequests >= MaxConcurrentRequests)
 				{
-					HttpClient->Tick(MAX_uint32, GIasHttpRateLimitKiBPerSecond);
+					HttpClient.Tick(MAX_uint32, GIasHttpRateLimitKiBPerSecond);
 				}
 			}
 
@@ -2163,7 +2177,7 @@ void FOnDemandIoBackend::ProcessHttpRequests(FHttpClient* HttpClient, FBitWindow
 		{
 			// Keep processing pending connections until all requests are completed or a new one is issued
 			TRACE_CPUPROFILER_EVENT_SCOPE(IasBackend::TickHttp);
-			while (HttpClient->Tick(GIasHttpPollTimeoutMs, GIasHttpRateLimitKiBPerSecond))
+			while (HttpClient.Tick(GIasHttpPollTimeoutMs, GIasHttpRateLimitKiBPerSecond))
 			{
 				if (!NextChunkRequest)
 				{
@@ -2185,29 +2199,30 @@ uint32 FOnDemandIoBackend::Run()
 	FBitWindow HttpErrors;
 	HttpErrors.Reset(GIasHttpErrorSampleCount);
 
-	const int32 MaxHttpRetryCount = FMath::Max(AvailableEps.Urls.Num() + 1, GIasMaxHttpRetryCount);
-	TUniquePtr<FHttpClient> HttpClient;
+	TUniquePtr<FHttpClient> HttpClient = FHttpClient::Create(FHttpClientConfig
+	{
+		.Endpoints = AvailableEps.Urls,
+		.PrimaryEndpoint = FMath::Min(GIasHttpPrimaryEndpoint, AvailableEps.Urls.Num() -1),
+		.MaxConnectionCount = GIasMaxHttpConnectionCount,
+		.PipelineLength = GIasHttpPipelineLength,
+		.MaxRetryCount = FMath::Max(AvailableEps.Urls.Num() + 1, GIasMaxHttpRetryCount),
+		.ReceiveBufferSize = GIasHttpRecvBufKiB >= 0 ? GIasHttpRecvBufKiB << 10 : -1,
+		.bChangeEndpointAfterSuccessfulRetry = GIasHttpChangeEndpointAfterSuccessfulRetry
+	});
+	check(HttpClient.IsValid());
+	HttpClient->SetEndpoint(AvailableEps.Current);
+#if !UE_BUILD_SHIPPING
 	if (AvailableEps.HasCurrent())
 	{
-		HttpClient = FHttpClient::Create(FHttpClientConfig
-		{
-			.Endpoints = AvailableEps.Urls,
-			.PrimaryEndpoint = AvailableEps.Current,
-			.MaxConnectionCount = GIasMaxHttpConnectionCount,
-			.PipelineLength = GIasHttpPipelineLength,
-			.MaxRetryCount = MaxHttpRetryCount,
-			.ReceiveBufferSize = GIasHttpRecvBufKiB >= 0 ? GIasHttpRecvBufKiB << 10 : -1
-		});
-#if !UE_BUILD_SHIPPING
 		LatencyTest(AvailableEps.GetCurrent(), GetEndpointTestPath());
-#endif 
 	}
+#endif 
 
 	while (!bStopRequested)
 	{
 		// Process HTTP request(s) even if the client is invalid to ensure enqueued request(s) gets completed.
-		ProcessHttpRequests(HttpClient.Get(), HttpErrors, FMath::Min(GIasHttpPipelineLength * GIasMaxHttpConnectionCount, 64));
-		AvailableEps.Current = HttpClient.IsValid() ? HttpClient->GetPrimaryConnection() : INDEX_NONE;
+		ProcessHttpRequests(*HttpClient, HttpErrors, FMath::Min(GIasHttpPipelineLength * GIasMaxHttpConnectionCount, 64));
+		AvailableEps.Current = HttpClient->GetEndpoint();
 
 		if (!bStopRequested)
 		{
@@ -2215,11 +2230,11 @@ uint32 FOnDemandIoBackend::Run()
 			if (BackendStatus.IsHttpError())
 			{
 				WaitTime = GIasHttpHealthCheckWaitTime;
-				if (HttpClient.IsValid())
+				if (HttpClient->GetEndpoint() != INDEX_NONE)
 				{
-					HttpClient.Reset();
-					HttpErrors.Reset(GIasHttpErrorSampleCount);
 					AvailableEps.Current = INDEX_NONE;
+					HttpClient->SetEndpoint(INDEX_NONE);
+					HttpErrors.Reset(GIasHttpErrorSampleCount);
 				}
 
 				UE_LOG(LogIas, Log, TEXT("Trying to reconnect to any available endpoint"));
@@ -2227,38 +2242,23 @@ uint32 FOnDemandIoBackend::Run()
 				if (int32 Idx = LatencyTest(AvailableEps.Urls, TestPath, bStopRequested); Idx != INDEX_NONE)
 				{
 					AvailableEps.Current = Idx;
-					HttpClient = FHttpClient::Create(FHttpClientConfig
-					{
-						.Endpoints = AvailableEps.Urls,
-						.PrimaryEndpoint = Idx,
-						.MaxConnectionCount = GIasMaxHttpConnectionCount,
-						.MaxRetryCount = MaxHttpRetryCount,
-						.ReceiveBufferSize = GIasHttpRecvBufKiB >= 0 ? GIasHttpRecvBufKiB << 10 : -1
-					});
+					HttpClient->SetEndpoint(Idx);
 					BackendStatus.SetHttpError(false);
 					UE_LOG(LogIas, Log, TEXT("Successfully reconnected to '%s'"), *AvailableEps.GetCurrent());
 					AddDeferredTocs();
 				}
 			}
-			else if (HttpClient.IsValid() && HttpClient->GetPrimaryConnection() != 0)
+			else if (HttpClient->IsUsingPrimaryEndpoint() == false)
 			{
 				WaitTime = GIasHttpHealthCheckWaitTime;
 				const FString TestPath = GetEndpointTestPath(); 
 				TConstArrayView<FString> Urls = AvailableEps.Urls;
 
-				UE_LOG(LogIas, Log, TEXT("Trying to reconnect to primary endpoint '%s'"), *AvailableEps.Urls[0]);
 				if (int32 Idx = LatencyTest(Urls.Left(1), TestPath, bStopRequested); Idx != INDEX_NONE)
 				{
 					AvailableEps.Current = Idx;
-					HttpClient = FHttpClient::Create(FHttpClientConfig
-					{
-						.Endpoints = AvailableEps.Urls,
-						.PrimaryEndpoint = Idx,
-						.MaxConnectionCount = GIasMaxHttpConnectionCount,
-						.MaxRetryCount = MaxHttpRetryCount,
-						.ReceiveBufferSize = GIasHttpRecvBufKiB >= 0 ? GIasHttpRecvBufKiB << 10 : -1
-					});
-					UE_LOG(LogIas, Log, TEXT("Successfully reconnected to '%s'"), *AvailableEps.GetCurrent());
+					HttpClient->SetEndpoint(Idx);
+					UE_LOG(LogIas, Log, TEXT("Reconnected to primary endpoint '%s'"), *AvailableEps.GetCurrent());
 				}
 			}
 			if (BackendStatus.ShouldAbandonCache())
