Add the possibility to reset the overriden rest transform on a geometry collection

diff --git a/Engine/Source/Runtime/Experimental/GeometryCollectionEngine/Private/GeometryCollection/GeometryCollectionComponent.cpp b/Engine/Source/Runtime/Experimental/GeometryCollectionEngine/Private/GeometryCollection/GeometryCollectionComponent.cpp
--- a/Engine/Source/Runtime/Experimental/GeometryCollectionEngine/Private/GeometryCollection/GeometryCollectionComponent.cpp
+++ b/Engine/Source/Runtime/Experimental/GeometryCollectionEngine/Private/GeometryCollection/GeometryCollectionComponent.cpp
@@ -1608,13 +1608,31 @@ int32 UGeometryCollectionComponent::EmbeddedIndexToTransformIndex(const UInstanc
 }
 #endif
 
+void UGeometryCollectionComponent::ResetRestTransforms()
+{
+	const bool bWasOverriden = (RestTransforms.Num() > 0);
+	if (bWasOverriden)
+	{
+		RestTransforms.Empty();
+
+		if (RestCollection && RestCollection->GetGeometryCollection())
+		{
+			RestTransformsChanged(RestCollection->GetGeometryCollection()->Transform.GetConstArray());
+		}
+	}
+}
+
 void UGeometryCollectionComponent::SetRestState(TArray<FTransform>&& InRestTransforms)
-{	
+{
 	RestTransforms = InRestTransforms;
-	
+	RestTransformsChanged(RestTransforms);
+}
+
+void UGeometryCollectionComponent::RestTransformsChanged(const TArray<FTransform>& NewRestTransform)
+{
 	if (DynamicCollection)
 	{
-		SetInitialTransforms(RestTransforms);
+		SetInitialTransforms(NewRestTransform);
 	}
 
 	if (SceneProxy)
diff --git a/Engine/Source/Runtime/Experimental/GeometryCollectionEngine/Public/GeometryCollection/GeometryCollectionComponent.h b/Engine/Source/Runtime/Experimental/GeometryCollectionEngine/Public/GeometryCollection/GeometryCollectionComponent.h
--- a/Engine/Source/Runtime/Experimental/GeometryCollectionEngine/Public/GeometryCollection/GeometryCollectionComponent.h
+++ b/Engine/Source/Runtime/Experimental/GeometryCollectionEngine/Public/GeometryCollection/GeometryCollectionComponent.h
@@ -1247,6 +1247,9 @@ public:
 	// #todo should this only be available in editor?
 	GEOMETRYCOLLECTIONENGINE_API void SetRestState(TArray<FTransform>&& InRestTransforms);
 
+	// this reset the rest transform to use the rest collection asset ones
+	GEOMETRYCOLLECTIONENGINE_API void ResetRestTransforms();
+
 	/** Set the dynamic state for all bodies in the DynamicCollection. */
 	GEOMETRYCOLLECTIONENGINE_API void SetDynamicState(const Chaos::EObjectStateType& NewDynamicState);
 
@@ -1527,6 +1530,10 @@ private:
 	void ProcessRepStateDataOnPT();
 	void ProcessRepDynamicDataOnPT();
 
+	// called when the rest transform are updated from SetRestState / ResetRestTransforms
+	// this upadtes the renderer as well as the dynamic collection initial transforms
+	void RestTransformsChanged(const TArray<FTransform>& NewRestTransform);
+
 	// return the most actual transforms
 	// this can be the rest collection ones, the overriden RestTransforms or the dynamic collection ones
 	const TArray<FTransform>& GetCurrentRestTransforms() const;
