Chaos - improved MaxDepentrationVelocity handling

diff --git a/Engine/Source/Runtime/Experimental/Chaos/Private/Chaos/Collision/PBDCollisionContainerSolver.cpp b/Engine/Source/Runtime/Experimental/Chaos/Private/Chaos/Collision/PBDCollisionContainerSolver.cpp
--- a/Engine/Source/Runtime/Experimental/Chaos/Private/Chaos/Collision/PBDCollisionContainerSolver.cpp
+++ b/Engine/Source/Runtime/Experimental/Chaos/Private/Chaos/Collision/PBDCollisionContainerSolver.cpp
@@ -75,6 +75,7 @@ namespace Chaos
 		const int32 ConstraintPointIndex, 
 		const FRealSingle Dt, 
 		const FRealSingle MaxDepentrationVelocity,
+		const FRealSingle MaxPushOut,
 		const FConstraintSolverBody& Body0,
 		const FConstraintSolverBody& Body1)
 	{
@@ -197,6 +198,15 @@ namespace Chaos
 			ManifoldPoint.InitialPhi = WorldContactInitialPhi;
 		}
 
+		// Limit the depenetration for this tick if desired
+		if (MaxPushOut > 0)
+		{
+			if (WorldContactDeltaNormal < -MaxPushOut)
+			{
+				WorldContactDeltaNormal = -MaxPushOut;
+			}
+		}
+
 		// Adjust depth to account for target penetration from user
 		const FRealSingle TargetPhi = ManifoldPoint.TargetPhi;
 		WorldContactDeltaNormal -= TargetPhi;
@@ -230,6 +240,8 @@ namespace Chaos
 		const bool bEnableDepenetrationVelocity = CVars::bChaos_Collision_EnableInitialDepenetration && Constraint->GetInitialOverlapDepentrationEnabled();
 		const FSolverReal MaxDepenetrationVelocity = bEnableDepenetrationVelocity ? SolverSettings.DepenetrationVelocity : FSolverReal(-1);
 
+		const FSolverReal MaxPushOut = (SolverSettings.MaxPushOutVelocity > 0) ? (FSolverReal(SolverSettings.MaxPushOutVelocity) * Dt) : FSolverReal(0);
+
 		// Only calculate state for newly added contacts. Normally this is all of them, but maybe not if incremental collision is used by RBAN.
 		// Also we only add active points to the solver's manifold points list
 		for (int32 ConstraintManifoldPointIndex = ConstraintPointBeginIndex; ConstraintManifoldPointIndex < ConstraintPointEndIndex; ++ConstraintManifoldPointIndex)
@@ -241,7 +253,7 @@ namespace Chaos
 				{
 					// Transform the constraint contact data into world space for use by the solver
 					// We build this data directly into the solver's world-space contact data which looks a bit odd with "Init" called after but there you go
-					UpdateCollisionSolverContactPointFromConstraint(Solver, SolverManifoldPointIndex, Constraint, ConstraintManifoldPointIndex, Dt, MaxDepenetrationVelocity, Body0, Body1);
+					UpdateCollisionSolverContactPointFromConstraint(Solver, SolverManifoldPointIndex, Constraint, ConstraintManifoldPointIndex, Dt, MaxDepenetrationVelocity, MaxPushOut, Body0, Body1);
 				}
 			}
 		}
@@ -701,7 +713,7 @@ namespace Chaos
 		// Adjust max pushout to attempt to make it iteration count independent
 		const FSolverReal Dt = FSolverReal(InDt);
 		const bool bApplyStaticFriction = (It >= (NumIts - SolverSettings.NumPositionFrictionIterations));
-		const FSolverReal MaxPushOut = (SolverSettings.MaxPushOutVelocity > 0) ? (FSolverReal(SolverSettings.MaxPushOutVelocity) * Dt) / FSolverReal(NumIts) : FSolverReal(0);
+		const FSolverReal MaxPushOut = FSolverReal(0);	// Now handled in UpdateCollisionSolverContactPointFromConstraint
 
 		// Apply the position correction
 		if (bApplyStaticFriction)
