Add VelChange option and strain to AddRadial Impulse function

diff --git a/Engine/Source/Runtime/Engine/Classes/PhysicsEngine/PhysicsObjectBlueprintLibrary.h b/Engine/Source/Runtime/Engine/Classes/PhysicsEngine/PhysicsObjectBlueprintLibrary.h
--- a/Engine/Source/Runtime/Engine/Classes/PhysicsEngine/PhysicsObjectBlueprintLibrary.h
+++ b/Engine/Source/Runtime/Engine/Classes/PhysicsEngine/PhysicsObjectBlueprintLibrary.h
@@ -23,5 +23,5 @@ public:
 	static FTransform GetPhysicsObjectWorldTransform(UPrimitiveComponent* Component, FName BoneName);
 
 	UFUNCTION(BlueprintCallable, Category = "Physics Object")
-	static void ApplyRadialImpulse(UPrimitiveComponent* Component, FVector Origin, float Radius, float Strength, enum ERadialImpulseFalloff Falloff, bool bApplyStrain);
+	static void ApplyRadialImpulse(UPrimitiveComponent* Component, FVector Origin, float Radius, float Strength, enum ERadialImpulseFalloff Falloff, bool bApplyStrain, float Strain, bool bVelChange = false);
 };
\ No newline at end of file
diff --git a/Engine/Source/Runtime/Engine/Private/PhysicsEngine/PhysicsObjectBlueprintLibrary.cpp b/Engine/Source/Runtime/Engine/Private/PhysicsEngine/PhysicsObjectBlueprintLibrary.cpp
--- a/Engine/Source/Runtime/Engine/Private/PhysicsEngine/PhysicsObjectBlueprintLibrary.cpp
+++ b/Engine/Source/Runtime/Engine/Private/PhysicsEngine/PhysicsObjectBlueprintLibrary.cpp
@@ -54,7 +54,7 @@ FTransform UPhysicsObjectBlueprintLibrary::GetPhysicsObjectWorldTransform(UPrimi
 	return FTransform::Identity;
 }
 
-void UPhysicsObjectBlueprintLibrary::ApplyRadialImpulse(UPrimitiveComponent* Component, FVector Origin, float Radius, float Strength, enum ERadialImpulseFalloff Falloff, bool bApplyStrain)
+void UPhysicsObjectBlueprintLibrary::ApplyRadialImpulse(UPrimitiveComponent* Component, FVector Origin, float Radius, float Strength, enum ERadialImpulseFalloff Falloff, bool bApplyStrain, float Strain, bool bVelChange)
 {
 	if (Component == nullptr)
 	{
@@ -66,6 +66,6 @@ void UPhysicsObjectBlueprintLibrary::ApplyRadialImpulse(UPrimitiveComponent* Com
 		TArray<Chaos::FPhysicsObjectHandle>  PhysicsObjects;
 		PhysicsObjects.Add(PhysicsObject);
 		FLockedWritePhysicsObjectExternalInterface Interface = FPhysicsObjectExternalInterface::LockWrite(PhysicsObjects);
-		return Interface->AddRadialImpulse(PhysicsObjects, Origin, Radius, Strength, Falloff, bApplyStrain, true);
+		return Interface->AddRadialImpulse(PhysicsObjects, Origin, Radius, Strength, Falloff, bApplyStrain, Strain, /* bInvalidate */ true, bVelChange);
 	}
 }
diff --git a/Engine/Source/Runtime/Experimental/Chaos/Private/Chaos/PhysicsObjectInterface.cpp b/Engine/Source/Runtime/Experimental/Chaos/Private/Chaos/PhysicsObjectInterface.cpp
--- a/Engine/Source/Runtime/Experimental/Chaos/Private/Chaos/PhysicsObjectInterface.cpp
+++ b/Engine/Source/Runtime/Experimental/Chaos/Private/Chaos/PhysicsObjectInterface.cpp
@@ -16,7 +16,7 @@ namespace PhysicsObjectInterfaceCVars
 	static FAutoConsoleVariableRef CVarStrainModifier(
 		TEXT("Chaos.Debug.StrainModifier"),
 		StrainModifier,
-		TEXT("Modify the strain by this factor"),
+		TEXT("(Deprecated) When using radial impulse, compute the strain by multiplier the impulse by this factor"),
 		ECVF_Default
 	);
 
@@ -88,7 +88,7 @@ namespace
 	}
 
 	template<typename T>
-	float AddRadialImpulseHelper(T ParticleHandle, FVector Origin, float Radius, float Strength, enum ERadialImpulseFalloff Falloff, float VelocityRatio = 1.0f, bool bInvalidate = true)
+	float AddRadialImpulseHelper(T ParticleHandle, FVector Origin, float Radius, float Strength, enum ERadialImpulseFalloff Falloff, float VelocityRatio = 1.0f, bool bInvalidate = true, bool bVelChange = false)
 	{
 		using namespace Chaos;
 
@@ -130,7 +130,8 @@ namespace
 				{
 					const FVec3 Normal = ParticleToOrigin.GetSafeNormal();
 					const FVec3 Impulse = Normal * FalloffStrength;
-					const FVec3 Velocity = Impulse * ParticleHandle->InvM() * VelocityRatio;
+					const FReal InvMass = bVelChange ? 1.0 : ParticleHandle->InvM();
+					const FVec3 Velocity = Impulse * InvMass * VelocityRatio;
 
 					const FVec3 CurrentImpulseVelocity = ParticleHandle->LinearImpulseVelocity();
 
@@ -894,7 +895,14 @@ namespace Chaos
 	}
 
 	template<EThreadContext Id>
-	void FWritePhysicsObjectInterface<Id>::AddRadialImpulse(TArrayView<const FPhysicsObjectHandle> InObjects, FVector Origin, float Radius, float Strength, enum ERadialImpulseFalloff Falloff, bool bApplyStrain, bool bInvalidate)
+	void FWritePhysicsObjectInterface<Id>::AddRadialImpulse(TArrayView<const FPhysicsObjectHandle> InObjects, FVector Origin, float Radius, float Strength, enum ERadialImpulseFalloff Falloff, bool bApplyStrain, bool bInvalidate, bool bVelChange)
+	{
+		// passing -1.0f as a strain will make use of the strain modifier instead ( legacy system )
+		AddRadialImpulse(InObjects, Origin, Radius, Strength, Falloff, bApplyStrain, -1.0f, bInvalidate, bVelChange);
+	}
+
+	template<EThreadContext Id>
+	void FWritePhysicsObjectInterface<Id>::AddRadialImpulse(TArrayView<const FPhysicsObjectHandle> InObjects, FVector Origin, float Radius, float Strength, enum ERadialImpulseFalloff Falloff, bool bApplyStrain, float Strain, bool bInvalidate, bool bVelChange)
 	{
 		//TODO: create a PT version of this, plus the damping functions
 		if (Chaos::FPBDRigidsSolver* RigidSolver = Chaos::FPhysicsObjectInterface::GetSolver(InObjects))
@@ -907,7 +915,10 @@ namespace Chaos
 				Strength,
 				Falloff,
 				bApplyStrain,
-				bInvalidate]()
+				Strain,
+				bInvalidate,
+				bVelChange
+				]()
 				{
 					using namespace Chaos;
 					for (FPhysicsObjectHandle Object : InObjects)
@@ -939,10 +950,13 @@ namespace Chaos
 
 								for (FPBDRigidParticleHandle* ChildHandle : *ChildrenHandles)
 								{
-									const float FalloffStrength = AddRadialImpulseHelper(ChildHandle, Origin, Radius, Strength, Falloff, VelocityRatio, bInvalidate);
+									const float FalloffStrength = AddRadialImpulseHelper(ChildHandle, Origin, Radius, Strength, Falloff, VelocityRatio, bInvalidate, bVelChange);
 
 									//to do: remove cvar when material system is in place and densities are updated
-									const float StrainToApply = PhysicsObjectInterfaceCVars::StrainModifier * FalloffStrength;
+									const float StrainToApply = 
+										(Strain < 0)
+											? (PhysicsObjectInterfaceCVars::StrainModifier * FalloffStrength)
+											: Strain;
 									if (StrainToApply > 0)
 									{
 										Clustering.SetExternalStrain(ChildHandle->CastToClustered(), StrainToApply);
@@ -959,7 +973,7 @@ namespace Chaos
 									RigidSolver->GetEvolution()->GetParticles().MarkTransientDirtyParticle(ParticleHandle);
 								}
 
-								AddRadialImpulseHelper(ParticleHandle, Origin, Radius, Strength, Falloff, bInvalidate);
+								AddRadialImpulseHelper(ParticleHandle, Origin, Radius, Strength, Falloff, bInvalidate, bVelChange);
 							}
 						}
 						else
@@ -972,7 +986,7 @@ namespace Chaos
 								RigidSolver->GetEvolution()->GetParticles().MarkTransientDirtyParticle(ParticleHandle);
 							}
 
-							AddRadialImpulseHelper(ParticleHandle, Origin, Radius, Strength, Falloff, bInvalidate);
+							AddRadialImpulseHelper(ParticleHandle, Origin, Radius, Strength, Falloff, bInvalidate, bVelChange);
 						}
 					}
 				});
diff --git a/Engine/Source/Runtime/Experimental/Chaos/Public/Chaos/PhysicsObjectInterface.h b/Engine/Source/Runtime/Experimental/Chaos/Public/Chaos/PhysicsObjectInterface.h
--- a/Engine/Source/Runtime/Experimental/Chaos/Public/Chaos/PhysicsObjectInterface.h
+++ b/Engine/Source/Runtime/Experimental/Chaos/Public/Chaos/PhysicsObjectInterface.h
@@ -125,7 +125,11 @@ namespace Chaos
 		CHAOS_API void ForceKinematic(TArrayView<const FPhysicsObjectHandle> InObjects);
 		CHAOS_API void AddForce(TArrayView<const FPhysicsObjectHandle> InObjects, const FVector& Force, bool bInvalidate);
 		CHAOS_API void AddTorque(TArrayView<const FPhysicsObjectHandle> InObjects, const FVector& Torque, bool bInvalidate);
-		CHAOS_API void AddRadialImpulse(TArrayView<const FPhysicsObjectHandle> InObjects, FVector Origin, float Radius, float Strength, enum ERadialImpulseFalloff Falloff, bool bApplyStrain, bool bInvalidate);
+
+		UE_DEPRECATED(5.4, "This version AddRadialImpulse has been deprecated. Please use the version where the strain value is passed explicitly")
+		CHAOS_API void AddRadialImpulse(TArrayView<const FPhysicsObjectHandle> InObjects, FVector Origin, float Radius, float Strength, enum ERadialImpulseFalloff Falloff, bool bApplyStrain, bool bInvalidate, bool bVelChange = false);
+
+		CHAOS_API void AddRadialImpulse(TArrayView<const FPhysicsObjectHandle> InObjects, FVector Origin, float Radius, float Strength, enum ERadialImpulseFalloff Falloff, bool bApplyStrain, float Strain, bool bInvalidate, bool bVelChange = false);
 		CHAOS_API void SetLinearEtherDrag(TArrayView<const FPhysicsObjectHandle> InObjects, float InLinearDrag);
 		CHAOS_API void SetAngularEtherDrag(TArrayView<const FPhysicsObjectHandle> InObjects, float InAngularDrag);
 
