Option to guard against infinite timeouts occurring

#rnx

diff --git a/Engine/Source/Runtime/Experimental/IoStoreOnDemand/Private/CoreHttp/Client.h b/Engine/Source/Runtime/Experimental/IoStoreOnDemand/Private/CoreHttp/Client.h
--- a/Engine/Source/Runtime/Experimental/IoStoreOnDemand/Private/CoreHttp/Client.h
+++ b/Engine/Source/Runtime/Experimental/IoStoreOnDemand/Private/CoreHttp/Client.h
@@ -194,8 +194,9 @@ public:
 
 							FEventLoop();
 							~FEventLoop();
-	uint32					Tick(uint32 PollTimeoutMs=0);
+	uint32					Tick(int32 PollTimeoutMs=0);
 	void					Throttle(uint32 KiBPerSec);
+	void					SetFailTimeout(int32 TimeoutMs);
 	bool					IsIdle() const;
 	void					Cancel(FTicket Ticket);
 	[[nodiscard]] FRequest	Request(FAnsiStringView Method, FAnsiStringView Url, const FRequestParams* Params=nullptr);
diff --git a/Engine/Source/Runtime/Experimental/IoStoreOnDemand/Private/CoreHttp/Http.cpp b/Engine/Source/Runtime/Experimental/IoStoreOnDemand/Private/CoreHttp/Http.cpp
--- a/Engine/Source/Runtime/Experimental/IoStoreOnDemand/Private/CoreHttp/Http.cpp
+++ b/Engine/Source/Runtime/Experimental/IoStoreOnDemand/Private/CoreHttp/Http.cpp
@@ -2765,7 +2765,7 @@ struct FTickState
 	FActivity*					DoneList;
 	uint64						Cancels;
 	int32&						RecvAllowance;
-	uint32						PollTimeoutMs;
+	int32						PollTimeoutMs;
 	TArray<FSocket::FWaiter>*	Waiters;
 };
 
@@ -3090,7 +3090,7 @@ private:
 	using FSocketGroupPtr = TUniquePtr<FSocketGroup>;
 	using FSocketGroups = TArray<FSocketGroupPtr>;
 	void						ScatterPendings();
-	int32						Wait(uint32 PollTimeoutMs);
+	int32						Wait(int32 PollTimeoutMs);
 	TArray<FSocket::FWaiter>	Waiters;
 	FSocketGroups				SocketGroups;
 	FHost&						Host;
@@ -3101,17 +3101,29 @@ private:
 };
 
 ////////////////////////////////////////////////////////////////////////////////
-int32 FHostGroup::Wait(uint32 PollTimeoutMs)
+int32 FHostGroup::Wait(int32 PollTimeoutMs)
 {
 	if (Waiters.Num() == 0)
 	{
 		return 0;
 	}
 
+	// If the poll timeout is negative then treat that as a fatal timeout
+	bool bFailOnTimeout = false;
+	if (PollTimeoutMs < -1)
+	{
+		PollTimeoutMs = -PollTimeoutMs;
+		bFailOnTimeout = true;
+	}
+
 	// Actually do the wait
 	int32 Result = FSocket::Wait(Waiters, PollTimeoutMs);
 	if (Result <= 0)
 	{
+		if (bFailOnTimeout && Result == 0)
+		{
+			return MIN_int32;
+		}
 		return Result;
 	}
 
@@ -3146,9 +3158,13 @@ void FHostGroup::Tick(FTickState& State)
 	// Do any waits left over from the last pass through
 	if (int32 Result = Wait(State.PollTimeoutMs); Result < 0)
 	{
+		const char* Reason = (Result == MIN_int32)
+			? "FailTimeout hit"
+			: "poll() returned an unexpected error";
+
 		for (FSocketGroupPtr& Group : SocketGroups)
 		{
-			Group->Fail(State, "poll() returned an unexpected error");
+			Group->Fail(State, Reason);
 		}
 	}
 
@@ -3254,9 +3270,10 @@ class FEventLoop::FImpl
 {
 public:
 							~FImpl();
-	uint32					Tick(uint32 PollTimeoutMs=0);
+	uint32					Tick(int32 PollTimeoutMs=0);
 	bool					IsIdle() const;
 	void					Throttle(uint32 KiBPerSec);
+	void					SetFailTimeout(int32 TimeoutMs);
 	void					Cancel(FTicket Ticket);
 	FRequest				Request(FAnsiStringView Method, FAnsiStringView Path, FActivity* Activity);
 	FTicket					Send(FActivity* Activity);
@@ -3270,6 +3287,7 @@ private:
 	FActivity*				Pending			= nullptr;
 	FThrottler				Throttler;
 	TArray<FHostGroup>		Groups;
+	int32					FailTimeoutMs	= 0;
 	uint32					BusyCount		= 0;
 };
 
@@ -3360,6 +3378,14 @@ void FEventLoop::FImpl::Throttle(uint32 KiBPerSec)
 	Throttler.SetLimit(KiBPerSec);
 }
 
+////////////////////////////////////////////////////////////////////////////////
+void FEventLoop::FImpl::SetFailTimeout(int32 TimeoutMs)
+{
+	// While TimeoutMs must be >=0, it is signed so that MAX_uint32 can't be used
+	check(TimeoutMs >= 0);
+	FailTimeoutMs = TimeoutMs;
+}
+
 ////////////////////////////////////////////////////////////////////////////////
 void FEventLoop::FImpl::Cancel(FTicket Ticket)
 {
@@ -3554,7 +3580,7 @@ uint32 FEventLoop::FImpl::Tick(uint32 PollTimeoutMs)
 #endif // 0
 
 ////////////////////////////////////////////////////////////////////////////////
-uint32 FEventLoop::FImpl::Tick(uint32 PollTimeoutMs)
+uint32 FEventLoop::FImpl::Tick(int32 PollTimeoutMs)
 {
 	TRACE_CPUPROFILER_EVENT_SCOPE(IasHttp::Tick);
 
@@ -3563,6 +3589,14 @@ uint32 FEventLoop::FImpl::Tick(uint32 PollTimeoutMs)
 	int32 RecvAllowance = Throttler.GetAllowance();
 	uint64 CancelsLoad = Cancels.load(std::memory_order_relaxed);
 
+	// There is a special internal meaning to the poll time out. If a infinite
+	// timeout is given and a fail one is set on the loop, we use a <-1 value
+	// that Wait() knows what to do with.
+	if (PollTimeoutMs < 0 && FailTimeoutMs)
+	{
+		PollTimeoutMs = -(FailTimeoutMs + 1);
+	}
+
 	// Tick groups and then remove ones that are idle
 	FTickState TickState = {
 		.DoneList = nullptr,
@@ -3618,10 +3652,11 @@ uint32 FEventLoop::FImpl::Tick(uint32 PollTimeoutMs)
 ////////////////////////////////////////////////////////////////////////////////
 FEventLoop::FEventLoop()						{ Impl = new FEventLoop::FImpl(); Trace(Impl, ETrace::LoopCreate); }
 FEventLoop::~FEventLoop()						{ delete Impl; Trace(Impl, ETrace::LoopDestroy); }
-uint32 FEventLoop::Tick(uint32 PollTimeoutMs)	{ return Impl->Tick(PollTimeoutMs); }
+uint32 FEventLoop::Tick(int32 PollTimeoutMs)	{ return Impl->Tick(PollTimeoutMs); }
 bool FEventLoop::IsIdle() const					{ return Impl->IsIdle(); }
 void FEventLoop::Cancel(FTicket Ticket)			{ return Impl->Cancel(Ticket); }
 void FEventLoop::Throttle(uint32 KiBPerSec)		{ return Impl->Throttle(KiBPerSec); }
+void FEventLoop::SetFailTimeout(int32 Ms)		{ return Impl->SetFailTimeout(Ms); }
 
 ////////////////////////////////////////////////////////////////////////////////
 FRequest FEventLoop::Request(
@@ -4102,6 +4137,37 @@ IOSTOREONDEMAND_API void IasHttpTest(const ANSICHAR* TestHost="localhost")
 		WaitForLoopIdle();
 	}
 
+	// fatal timeout
+	for (int32 i = 0; i < 14; ++i)
+	{
+		bool bExpectFailTimeout = !!(i & 1);
+		FEventLoop Loop2;
+		Loop2.Send(
+			Loop2.Get(BuildUrl("/data?stall", 9494)),
+			[bExpectFailTimeout, Dest=FIoBuffer()] (const FTicketStatus& Status) mutable
+			{
+				if (Status.GetId() == FTicketStatus::EId::Response)
+				{
+					FResponse& Response = Status.GetResponse();
+					Response.SetDestination(&Dest);
+					return;
+				}
+
+				check(Status.GetId() == FTicketStatus::EId::Error);
+
+				const char* Reason = Status.GetErrorReason();
+				bool IsFailTimeout = (FCStringAnsi::Strstr(Reason, "FailTimeout") != nullptr);
+				check(IsFailTimeout == bExpectFailTimeout);
+			}
+		);
+
+		if (bExpectFailTimeout)
+		{
+			Loop2.SetFailTimeout(1000);
+		}
+		while (Loop2.Tick(-1));
+	}
+
 	// no connect
 	{
 		FRequest Request[] = {
diff --git a/Engine/Source/Runtime/Experimental/IoStoreOnDemand/Private/OnDemandHttpClient.cpp b/Engine/Source/Runtime/Experimental/IoStoreOnDemand/Private/OnDemandHttpClient.cpp
--- a/Engine/Source/Runtime/Experimental/IoStoreOnDemand/Private/OnDemandHttpClient.cpp
+++ b/Engine/Source/Runtime/Experimental/IoStoreOnDemand/Private/OnDemandHttpClient.cpp
@@ -34,7 +34,7 @@ void FHttpClient::Get(FAnsiStringView Url, FGetCallback&& Callback)
 	Get(Url, FIoOffsetAndLength(), MoveTemp(Callback));
 }
 
-bool FHttpClient::Tick(uint32 WaitTimeMs, uint32 MaxKiBPerSecond)
+bool FHttpClient::Tick(int32 WaitTimeMs, uint32 MaxKiBPerSecond)
 {
 	EventLoop.Throttle(MaxKiBPerSecond);
 	const uint32 TicketCount = EventLoop.Tick(WaitTimeMs);
@@ -81,6 +81,8 @@ FHttpClient::FHttpClient(FHttpClientConfig&& ClientConfig)
 	check(Config.Endpoints.IsEmpty() == false);
 	Connections.SetNum(Config.Endpoints.Num());
 	SetEndpoint(Config.PrimaryEndpoint);
+
+	EventLoop.SetFailTimeout(Config.FailTimeoutMs);
 }
 
 TUniquePtr<FHttpClient> FHttpClient::Create(FHttpClientConfig&& ClientConfig)
diff --git a/Engine/Source/Runtime/Experimental/IoStoreOnDemand/Private/OnDemandHttpClient.h b/Engine/Source/Runtime/Experimental/IoStoreOnDemand/Private/OnDemandHttpClient.h
--- a/Engine/Source/Runtime/Experimental/IoStoreOnDemand/Private/OnDemandHttpClient.h
+++ b/Engine/Source/Runtime/Experimental/IoStoreOnDemand/Private/OnDemandHttpClient.h
@@ -21,6 +21,7 @@ struct FHttpClientConfig
 	int32 PipelineLength = 1;
 	int32 MaxRetryCount = 1;
 	int32 ReceiveBufferSize = -1;
+	int32 FailTimeoutMs = 0;
 	bool bChangeEndpointAfterSuccessfulRetry = true;
 };
 
@@ -35,7 +36,7 @@ public:
 	void Get(FAnsiStringView Url, const FIoOffsetAndLength& Range, FGetCallback&& Callback);
 	void Get(FAnsiStringView Url, FGetCallback&& Callback);
 
-	bool Tick(uint32 WaitTimeMs, uint32 MaxKiBPerSecond);
+	bool Tick(int32 WaitTimeMs, uint32 MaxKiBPerSecond);
 	bool Tick() { return Tick(-1, 0); }
 
 	int32 GetEndpoint() const { return CurrentEndpoint; }
diff --git a/Engine/Source/Runtime/Experimental/IoStoreOnDemand/Private/OnDemandIoDispatcherBackend.cpp b/Engine/Source/Runtime/Experimental/IoStoreOnDemand/Private/OnDemandIoDispatcherBackend.cpp
--- a/Engine/Source/Runtime/Experimental/IoStoreOnDemand/Private/OnDemandIoDispatcherBackend.cpp
+++ b/Engine/Source/Runtime/Experimental/IoStoreOnDemand/Private/OnDemandIoDispatcherBackend.cpp
@@ -126,6 +126,13 @@ static FAutoConsoleVariableRef CVar_IasHttpTimeOutMs(
 	TEXT("Time out value for HTTP requests in milliseconds")
 );
 
+static int32 GIasHttpFailTimeOutMs = 4 * 1000;
+static FAutoConsoleVariableRef CVar_IasHttpFailTimeOutMs(
+	TEXT("ias.HttpFailTimeOutMs"),
+	GIasHttpFailTimeOutMs,
+	TEXT("Fail infinite network waits that take longer than this (in ms, 0=disabled)")
+);
+
 int32 GIasHttpHealthCheckWaitTime = 3000;
 static FAutoConsoleVariableRef CVar_IasHttpHealthCheckWaitTime(
 	TEXT("ias.HttpHealthCheckWaitTime"),
@@ -2207,7 +2214,8 @@ uint32 FOnDemandIoBackend::Run()
 		.PipelineLength = GIasHttpPipelineLength,
 		.MaxRetryCount = FMath::Max(AvailableEps.Urls.Num() + 1, GIasMaxHttpRetryCount),
 		.ReceiveBufferSize = GIasHttpRecvBufKiB >= 0 ? GIasHttpRecvBufKiB << 10 : -1,
-		.bChangeEndpointAfterSuccessfulRetry = GIasHttpChangeEndpointAfterSuccessfulRetry
+		.FailTimeoutMs = GIasHttpFailTimeOutMs,
+		.bChangeEndpointAfterSuccessfulRetry = GIasHttpChangeEndpointAfterSuccessfulRetry,
 	});
 	check(HttpClient.IsValid());
 	HttpClient->SetEndpoint(AvailableEps.Current);
