Updated IAS test script

#rnx

diff --git a/Engine/Source/Runtime/Experimental/IoStoreOnDemand/Private/CoreHttp/TestServer.py b/Engine/Source/Runtime/Experimental/IoStoreOnDemand/Private/CoreHttp/TestServer.py
--- a/Engine/Source/Runtime/Experimental/IoStoreOnDemand/Private/CoreHttp/TestServer.py
+++ b/Engine/Source/Runtime/Experimental/IoStoreOnDemand/Private/CoreHttp/TestServer.py
@@ -15,11 +15,19 @@ async def proxy_impl(loop, client):
     while not msg.endswith(b"\r\n\r\n"):
         msg += await loop.sock_recv(client, 2048)
 
+    min_trunc = 0
     disconnect = False
+    stall = False
+
     line = next(iter(io.BytesIO(msg)))
     if b"?disconnect HTTP" in line:
         disconnect = True
 
+    if b"?stall HTTP" in line:
+        disconnect = True
+        stall = True
+        min_trunc = len(msg) + 16
+
     send_time = 1.0 + (random.random() * 0.5)
     with socket.socket(socket.AF_INET, socket.SOCK_STREAM) as httpd:
         httpd.connect(("127.0.0.1", 9493))
@@ -31,20 +39,24 @@ async def proxy_impl(loop, client):
         data_size = len(data)
 
         if disconnect:
-            trunc = int(data_size * random.random())
+            trunc = max(int(data_size * random.random()), min_trunc)
             if trunc >= data_size:
-                breakpoint()
                 trunc = data_size - 1
-            print("trun:", data_size, ">", trunc)
+            print("trun:", data_size, "->", trunc)
             data = data[:trunc]
 
         while data:
             percent = 0.02 + (random.random() * 0.08)
             send_size = max(int(data_size * percent), 1)
-            print("sent", send_size, "total", data_size, "sleep", send_time * percent, "percent", percent)
+            print("snd:%5d tot:%5d slp:%5f %%:%5f" % (send_size, data_size, send_time * percent, percent), end="\r")
             await loop.sock_sendall(client, data[:send_size])
             data = data[send_size:]
             await asyncio.sleep(send_time * percent)
+        print("")
+
+        if stall:
+            print("...stalling", end="")
+            await asyncio.sleep(2)
 
     client.close()
 
