Add a feature to allow the .iochunktoc to be written to disk when it is being uploaded then read that file from disk at runtime rather than downloading or generating it.

#rnx

- Note that this is considered a temp feature as we really want to be generating the .iochunktoc file at runtime when ondemand .utoc files are mounted. However for the short term this is a more reliable work around. It is likely we will remove this before UE 5.4 is released and should not be relied upon.

### IoStoreOnDemand
- Add a new cvar 'ias.LoadOnDemandToc' which when enabled (default false) will cause us to try and serialize the .iochunktoc file from disk rather than trying to generate it from mounted .utoc files. Note that this will only matter if 's.IasGenerateOnDemandToc' is set to true.
- Add a new utility function ::CreateReaderFromPlatformPackage which creates a FArchive capable of reading data from disk using the exact same file setup rules as FPlatformMisc::LoadTextFileFromPlatformPackage. Previously we have had difficulty trying to use the file manager directly to read files from the 'Cloud' directory and as time is a factor we should just stick with what we know works. If this feature is to be kept long term we should rethink this code.
- When attempting to load the .iochunktoc' file we do not consider the file missing as an error, instead we log a warning and then fallback trying to generate what we need from the mounted '.utoc' files as before. This is an attempt to make the system more fault tolerant for release. If we want to keep this feature we probably should consider the file being missing an error if we are trying to find and use it.

### UnrealPak
- The command '-Upload=' now supports an additional '-WriteToDisk' parameter which indicates that we should write the .iochunktoc to disk as well as uploading it. NOTE: That in order for this new parameter to work the caller must also provide a valid '-ConfigFilePath=' parameter as we want to write out the .iochunktoc next to the provided config file.
-- If we decide to keep this feature long term we should consider changing this logic so that the caller can output the file where ever they want.
- Implemented our own version of FFileHelper::SaveArrayToFile as none of the current versions support serializing data with length int64 without copying the data first.
-- It is not expected that we'd breach int32 but I am tired of casting data lengths back and forth between different types. The FFileHelper api will be fixed in a different work item later.

diff --git a/Engine/Source/Runtime/Experimental/IoStoreOnDemand/Private/IoStoreOnDemand.cpp b/Engine/Source/Runtime/Experimental/IoStoreOnDemand/Private/IoStoreOnDemand.cpp
--- a/Engine/Source/Runtime/Experimental/IoStoreOnDemand/Private/IoStoreOnDemand.cpp
+++ b/Engine/Source/Runtime/Experimental/IoStoreOnDemand/Private/IoStoreOnDemand.cpp
@@ -17,6 +17,7 @@
 #include "Misc/DateTime.h"
 #include "Misc/FileHelper.h"
 #include "Misc/Parse.h"
+#include "Misc/PathViews.h"
 #include "Modules/ModuleManager.h"
 #include "OnDemandHttpClient.h"
 #include "OnDemandIoDispatcherBackend.h"
@@ -248,6 +249,23 @@ static FIasCacheConfig GetIasCacheConfig(const TCHAR* CommandLine)
 	return Ret;
 }
 
+////////////////////////////////////////////////////////////////////////////////
+
+/** Utility for saving data to disk, similar to FFileHelper::SaveArrayToFile but supporting larger sizes */
+static bool SaveArrayToFile(TArrayView64<uint8> Data, const TCHAR* Filename, uint32 WriteFlags = FILEWRITE_None)
+{
+	TUniquePtr<FArchive> Ar = TUniquePtr<FArchive>(IFileManager::Get().CreateFileWriter(Filename, WriteFlags));
+	if (!Ar)
+	{
+		return false;
+	}
+
+	Ar->Serialize(static_cast<void*>(Data.GetData()), Data.Num());
+	Ar->Close();
+
+	return !Ar->IsError() && !Ar->IsCriticalError();
+}
+
 ////////////////////////////////////////////////////////////////////////////////
 FArchive& operator<<(FArchive& Ar, FTocMeta& Meta)
 {
@@ -766,6 +784,21 @@ TIoStatusOr<FIoStoreUploadParams> FIoStoreUploadParams::Parse(const TCHAR* Comma
 	FParse::Value(CommandLine, TEXT("TargetPlatform="), Params.TargetPlatform);
 	Params.bDeleteContainerFiles = !FParse::Param(CommandLine, TEXT("KeepContainerFiles"));
 	Params.bDeletePakFiles = !FParse::Param(CommandLine, TEXT("KeepPakFiles"));
+	Params.bWriteTocToDisk = FParse::Param(CommandLine, TEXT("WriteTocToDisk"));
+
+	if (Params.bWriteTocToDisk)
+	{
+		// If we keep this feature we should allow the caller to set this path themselves rather than rely on the config file location
+		FString ConfigFilePath;
+		FParse::Value(FCommandLine::Get(), TEXT("ConfigFilePath="), ConfigFilePath);
+
+		Params.TocOutputDir = FPaths::GetPath(ConfigFilePath);
+
+		if (Params.TocOutputDir.IsEmpty())
+		{
+			return FIoStatus(EIoErrorCode::InvalidParameter, TEXT("Cmdline param 'WriteToDisk' requires a valid 'ConfigFilePath' param aswell"));
+		}
+	}
 
 	if (Params.AccessKey.IsEmpty() &&
 		Params.SecretKey.IsEmpty() &&
@@ -1115,6 +1148,18 @@ TIoStatusOr<FIoStoreUploadResult> UploadContainerFiles(
 			UE_LOG(LogIas, Warning, TEXT("Failed to upload TOC '%s/%s/%s', StatusCode: %u"), *Client.GetConfig().ServiceUrl, *UploadParams.Bucket, Key.ToString(), Response.StatusCode);
 			return FIoStatus(EIoErrorCode::WriteError, TEXT("Failed to upload TOC"));
 		}
+
+		if (UploadParams.bWriteTocToDisk)
+		{
+			TStringBuilder<512> OnDemandTocFilePath;
+			FPathViews::Append(OnDemandTocFilePath, UploadParams.TocOutputDir, UploadResult.TocHash);
+			OnDemandTocFilePath << TEXT(".iochunktoc");
+
+			if (SaveArrayToFile(TArrayView64<uint8>(Ar.GetData(), Ar.TotalSize()), OnDemandTocFilePath.ToString()) == false)
+			{
+				UE_LOG(LogIoStore, Error, TEXT("Failed to save on demand toc file '%s'"), OnDemandTocFilePath.ToString());
+			}
+		}
 	}
 
 	for (const FString& Path : FilesToDelete)
diff --git a/Engine/Source/Runtime/Experimental/IoStoreOnDemand/Private/OnDemandIoDispatcherBackend.cpp b/Engine/Source/Runtime/Experimental/IoStoreOnDemand/Private/OnDemandIoDispatcherBackend.cpp
--- a/Engine/Source/Runtime/Experimental/IoStoreOnDemand/Private/OnDemandIoDispatcherBackend.cpp
+++ b/Engine/Source/Runtime/Experimental/IoStoreOnDemand/Private/OnDemandIoDispatcherBackend.cpp
@@ -10,13 +10,14 @@
 #include "EncryptionKeyManager.h"
 #include "GenericPlatform/GenericPlatformCrashContext.h"
 #include "HAL/Event.h"
+#include "HAL/FileManagerGeneric.h"
 #include "HAL/LowLevelMemTracker.h"
 #include "HAL/Platform.h"
+#include "HAL/PlatformFileManager.h"
 #include "HAL/PlatformTime.h"
 #include "HAL/PreprocessorHelpers.h"
 #include "HAL/Runnable.h"
 #include "HAL/RunnableThread.h"
-#include "HAL/PlatformFileManager.h"
 #include "HttpManager.h"
 #include "IO/IoAllocators.h"
 #include "IO/IoChunkEncoding.h"
@@ -29,8 +30,8 @@
 #include "Math/NumericLimits.h"
 #include "Misc/CommandLine.h"
 #include "Misc/EnumClassFlags.h"
-#include "Misc/Paths.h"
 #include "Misc/PathViews.h"
+#include "Misc/Paths.h"
 #include "Misc/ScopeLock.h"
 #include "Misc/ScopeRWLock.h"
 #include "OnDemandHttpClient.h"
@@ -189,6 +190,16 @@ static FAutoConsoleVariableRef CVar_IasGenerateOnDemandToc(
 	ECVF_ReadOnly
 );
 
+// A temp fallback path allowing us to attempt to load the OnDemand toc from disk rather than
+// trying to generate it.
+bool GIasLoadOnDemandToc = false;
+static FAutoConsoleVariableRef CVar_IasLoadOnDemandToc(
+	TEXT("ias.LoadOnDemandToc"),
+	GIasLoadOnDemandToc,
+	TEXT("Attempts to load the OnDemand toc from disk rather than generating it (TEMP)"),
+	ECVF_ReadOnly
+);
+
 bool GIasAsyncTocGenerationEnabled = true;
 static FAutoConsoleVariableRef CVar_IasAsyncTocGeneration(
 	TEXT("s.IasEnableThreadedTocGeneration"),
@@ -901,15 +912,67 @@ static TArray<FString> FindOnDemandUtocFilesOnDisk()
 	return FoundFiles;
 }
 
+/**
+ * Utility to create a FArchive capable of reading from disk using the exact same pathing
+ * rules as FPlatformMisc::LoadTextFileFromPlatformPackage but without forcing the entire
+ * file to be loaded at once.
+ */
+static TUniquePtr<FArchive> CreateReaderFromPlatformPackage(const FString& RelPath)
+{
+	const FString AbsPath = FPaths::Combine(FGenericPlatformMisc::RootDir(), RelPath);
+
+	IFileHandle* File = IPlatformFile::GetPlatformPhysical().OpenRead(*AbsPath);
+	if (File)
+	{
+		return MakeUnique<FArchiveFileReaderGeneric>(File, *AbsPath, File->Size());
+	}
+	else
+	{
+		return TUniquePtr<FArchive>();
+	}	
+}
+
 /** Generate a FOnDemandToc based on utoc files on disk which support the OnDemand feature */
-TIoStatusOr<FOnDemandToc> GenerateOnDemandTocFromDisk()
+static TIoStatusOr<FOnDemandToc> GenerateOnDemandTocFromDisk(FStringView TocHash)
 {
 	TRACE_CPUPROFILER_EVENT_SCOPE(IasBackend::GenerateOnDemandTocFromDisk);
 
-	UE_LOG(LogIas, Log, TEXT("Generating OnDemandToc via utoc files from disk"));
-
 	FOnDemandToc OutToc;
 
+	if (GIasLoadOnDemandToc)
+	{
+		UE_LOG(LogIas, Log, TEXT("Serializing .iochunktoc from disk"));
+
+		FString TocFileName = FString(TocHash);
+		TocFileName .Append(TEXT(".iochunktoc"));
+
+		const FString TocPath = FPaths::Combine(TEXT("Cloud"), TocFileName);
+
+		if (FPlatformMisc::FileExistsInPlatformPackage(TocPath))
+		{	
+			TUniquePtr<FArchive> Ar = CreateReaderFromPlatformPackage(TocPath);
+			*Ar << OutToc;
+
+			Ar->Close();
+
+			if (!Ar->IsError() && !Ar->IsCriticalError())
+			{
+				UE_LOG(LogIas, Log, TEXT("Loaded '%s' from disk"), *TocPath);
+				return OutToc;
+			}
+			else
+			{
+				return FIoStatus(EIoErrorCode::ReadError, WriteToString<256>(TEXT("Failed to open '"), TocPath, TEXT("' from disk")));
+			}
+		}
+		else
+		{
+			UE_LOG(LogIas, Warning, TEXT("Unable to find '%s' on disk, will attempt to generate it as a fallback"), *TocPath);
+		}
+	}
+
+	UE_LOG(LogIas, Log, TEXT("Generating .iochunktoc via utoc files from disk"));
+
 	TArray<FString> UtocFilePaths = FindOnDemandUtocFilesOnDisk();
 	const TMap<FGuid, FAES::FAESKey> EncryptionKeys = FEncryptionKeyManager::Get().GetAllKeys();
 
@@ -1886,7 +1949,7 @@ FIoStatus FOnDemandIoBackend::ApplyGeneratedOnDemandToc(const FString& CdnUrl, c
 	}
 	else
 	{
-		GeneratedTocResult = GenerateOnDemandTocFromDisk();
+		GeneratedTocResult = GenerateOnDemandTocFromDisk(FPathViews::GetBaseFilename(TocPath));
 	}
 	
 	if (!GeneratedTocResult.IsOk())
@@ -1982,9 +2045,10 @@ void FOnDemandIoBackend::Mount(const FOnDemandEndpoint& Endpoint)
 
 	if (GIasGenerateOnDemandToc && GIasAsyncTocGenerationEnabled)
 	{
-		OnDemandTocTask = UE::Tasks::Launch(UE_SOURCE_LOCATION, [this]() -> TIoStatusOr<FOnDemandToc>
+		FString TocHash = FPaths::GetBaseFilename(Endpoint.TocPath);
+		OnDemandTocTask = UE::Tasks::Launch(UE_SOURCE_LOCATION, [TocHash = MoveTemp(TocHash)]() -> TIoStatusOr<FOnDemandToc>
 			{
-				return GenerateOnDemandTocFromDisk();
+				return GenerateOnDemandTocFromDisk(TocHash);
 			});
 	}
 
diff --git a/Engine/Source/Runtime/Experimental/IoStoreOnDemand/Public/IO/IoStoreOnDemand.h b/Engine/Source/Runtime/Experimental/IoStoreOnDemand/Public/IO/IoStoreOnDemand.h
--- a/Engine/Source/Runtime/Experimental/IoStoreOnDemand/Public/IO/IoStoreOnDemand.h
+++ b/Engine/Source/Runtime/Experimental/IoStoreOnDemand/Public/IO/IoStoreOnDemand.h
@@ -171,6 +171,11 @@ struct FIoStoreUploadParams
 	int32 MaxConcurrentUploads = 16;
 	bool bDeleteContainerFiles = true;
 	bool bDeletePakFiles = true;
+
+	/** If we should write out the .iochunktoc to disk as well as uploading it. */
+	bool bWriteTocToDisk = false;
+	/** Where the .iochunktoc file should be written out. */
+	FString TocOutputDir;
 	
 	UE_API static TIoStatusOr<FIoStoreUploadParams> Parse(const TCHAR* CommandLine);
 };
