Ias: Increase frequency of statistics logging to every 30 seconds (from 60 seconds) for easier diagnosis of any ias related issues.

#rnx

diff --git a/Engine/Source/Runtime/Experimental/IoStoreOnDemand/Private/Statistics.cpp b/Engine/Source/Runtime/Experimental/IoStoreOnDemand/Private/Statistics.cpp
--- a/Engine/Source/Runtime/Experimental/IoStoreOnDemand/Private/Statistics.cpp
+++ b/Engine/Source/Runtime/Experimental/IoStoreOnDemand/Private/Statistics.cpp
@@ -16,7 +16,7 @@ LLM_DEFINE_TAG(Ias);
 namespace UE::IO::IAS
 {
 
-float GIasStatisticsLogInterval = 60.f;
+float GIasStatisticsLogInterval = 30.f;
 static FAutoConsoleVariableRef CVar_StatisticsLogInterval(
 	TEXT("ias.StatisticsLogInterval"),
 	GIasStatisticsLogInterval,
