Defer unregister operations in certain edge cases

In certain edge cases (only seen with automated testing) the service thread might get both register and unregister operations for the same cache object in the same update pass. In this case defer the unregister operation to the next update pass.

diff --git a/Engine/Source/Runtime/Experimental/IoStoreOnDemand/Private/JournaledCache.cpp b/Engine/Source/Runtime/Experimental/IoStoreOnDemand/Private/JournaledCache.cpp
--- a/Engine/Source/Runtime/Experimental/IoStoreOnDemand/Private/JournaledCache.cpp
+++ b/Engine/Source/Runtime/Experimental/IoStoreOnDemand/Private/JournaledCache.cpp
@@ -1497,7 +1497,7 @@ private:
 	uint32						UpdateCache(FCache* Cache);
 	virtual uint32				Run() override;
 	virtual void				Stop() override;
-	void						SubmitWork(FWork* Work, uint32 Num);
+	void						SubmitWork(const FWork* Work, uint32 Num);
 	void						ReceiveWork();
 	TUniquePtr<FRunnableThread> Thread;
 	FEventRef					WakeEvent;
@@ -1533,7 +1533,7 @@ FServiceThread::~FServiceThread()
 }
 
 ////////////////////////////////////////////////////////////////////////////////
-void FServiceThread::SubmitWork(FWork* Work, uint32 Num)
+void FServiceThread::SubmitWork(const FWork* Work, uint32 Num)
 {
 	FScopeLock _(&Lock);
 	PendingWork.Append(Work, Num);
@@ -1774,6 +1774,7 @@ void FServiceThread::ReceiveWork()
 		}
 
 		auto* CachePtr = (FCache*)(Work.GetPtr());
+		bool bFound = false;
 		for (int32 i = 0, n = Caches.Num(); i < n; ++i)
 		{
 			if (Caches[i].Get() != CachePtr)
@@ -1782,8 +1783,22 @@ void FServiceThread::ReceiveWork()
 			}
 
 			Caches.RemoveAtSwap(i);
+			bFound = true;
 			break;
 		}
+
+		// It is possible that the cache's register operation is part of the inbound work.
+		// In that case re-submit operation to pending work and process next pass
+		if (!bFound)
+		{
+			for(const FWork& RegWork : InboundWork)
+			{
+				if (RegWork.What == FWork::Work_Register && RegWork.GetPtr() == CachePtr)
+				{
+					SubmitWork(&Work, 1);
+				}
+			}
+		}
 	}
 
 	// ...then the rest
