Iris
* Refactor NetBlobManager rpc queuing to improve error messages and make the code more explicit on what its doing.

#rnx

diff --git a/Engine/Source/Runtime/Experimental/Iris/Core/Private/Iris/ReplicationSystem/NetBlob/NetBlobManager.cpp b/Engine/Source/Runtime/Experimental/Iris/Core/Private/Iris/ReplicationSystem/NetBlob/NetBlobManager.cpp
--- a/Engine/Source/Runtime/Experimental/Iris/Core/Private/Iris/ReplicationSystem/NetBlob/NetBlobManager.cpp
+++ b/Engine/Source/Runtime/Experimental/Iris/Core/Private/Iris/ReplicationSystem/NetBlob/NetBlobManager.cpp
@@ -57,33 +57,40 @@ bool FNetBlobManager::QueueNetObjectAttachment(uint32 ConnectionId, const FNetOb
 		return true;
 	}
 
-	FInternalNetRefIndex OwnerIndex = 0;
-	FInternalNetRefIndex SubObjectIndex = 0;
-	const FNetObjectReference& NetObjectReference = Attachment->GetNetObjectReference();
+	FRPCOwner OwnerInfo;
 
-	const FNetObjectReference OwnerOrSubObjectReference = TargetRef;	
-	FNetObjectReference OwnerReference = OwnerOrSubObjectReference;
+	OwnerInfo.CallerRef = TargetRef;
+	OwnerInfo.TargetRef = TargetRef;
+
+	bool bCanSendRpc = GetRootObjectAndSubObjectIndicesFromAnyHandle(TargetRef.GetRefHandle(), OwnerInfo.RootObjectIndex, OwnerInfo.SubObjectIndex);
 
-	bool bCanSendRpc = GetOwnerAndSubObjectIndicesFromHandle(OwnerOrSubObjectReference.GetRefHandle(), OwnerIndex, SubObjectIndex);
 	if (!bCanSendRpc)
 	{
-		// If the TargetRef is a valid reference but not replicated, then the owner must be used instead
-		if (OwnerOrSubObjectReference.IsValid())
+		if (!TargetRef.IsValid())
 		{
-			// See if we can find a replicated handle
-			OwnerReference = ObjectReferenceCache->GetReplicatedOuter(TargetRef);
-			bCanSendRpc = GetOwnerAndSubObjectIndicesFromHandle(OwnerReference.GetRefHandle(), OwnerIndex, SubObjectIndex);
+			UE_LOG(LogIris, Warning, TEXT("QueueNetObjectAttachment %s Failed due to invalid Target. Unable to resolve target reference %s."), *Attachment->GetNetObjectReference().ToString(), *TargetRef.ToString());
+			return false;
 		}
 
+		//$IRIS TODO: It's possible the Target and Caller will be differrent from the RootObjectIndex and SubObjectIndex since we use the ReplicatedOuter instead of the true Root.
+		//This probably happens when the outer list looks like this: Actor->ActorComponent->SubObject1->SubObject2. 
+
+		// If the TargetRef is a valid reference but not replicated, then the outer must be used instead
+		OwnerInfo.CallerRef = ObjectReferenceCache->GetReplicatedOuter(TargetRef);
+
+		// Can that outer send RPCs
+		bCanSendRpc = GetRootObjectAndSubObjectIndicesFromAnyHandle(OwnerInfo.CallerRef.GetRefHandle(), OwnerInfo.RootObjectIndex, OwnerInfo.SubObjectIndex);
 		if (!bCanSendRpc)
 		{
-			UE_LOG(LogIris, Warning, TEXT("QueueNetObjectAttachment %s Failed. Unable to resolve object reference %s."), *NetObjectReference.ToString(), *OwnerReference.ToString());
+			UE_LOG(LogIris, Warning, TEXT("QueueNetObjectAttachment %s Failed due to invalid Outer. Unable to resolve outer reference %s (index: %u) for target reference %s (index: %u)."), 
+				*Attachment->GetNetObjectReference().ToString(), ToCStr(OwnerInfo.CallerRef.ToString()), OwnerInfo.RootObjectIndex, ToCStr(TargetRef.ToString()), OwnerInfo.SubObjectIndex);
+
 			return false;
 		}
 	}
 
-	Attachment->SetNetObjectReference(OwnerReference, OwnerOrSubObjectReference);
-	AttachmentSendQueue.Enqueue(ConnectionId, OwnerIndex, SubObjectIndex, Attachment, SendFlags);
+	Attachment->SetNetObjectReference(OwnerInfo.CallerRef, OwnerInfo.TargetRef);
+	AttachmentSendQueue.Enqueue(ConnectionId, OwnerInfo.RootObjectIndex, OwnerInfo.SubObjectIndex, Attachment, SendFlags);
 	return true;
 }
 
@@ -114,37 +121,20 @@ bool FNetBlobManager::SendRPC(const UObject* Object, const UObject* SubObject, c
 		return true;
 	}
 
-	FInternalNetRefIndex OwnerIndex = 0;
-	FInternalNetRefIndex SubObjectIndex = 0;
-
-	const FNetObjectReference OwnerOrSubObjectReference = ObjectReferenceCache->GetOrCreateObjectReference(SubObject ? SubObject : Object);	
-	FNetObjectReference OwnerReference = OwnerOrSubObjectReference;
-
-	bool bCanSendRpc = GetOwnerAndSubObjectIndicesFromHandle(OwnerOrSubObjectReference.GetRefHandle(), OwnerIndex, SubObjectIndex);
-	if (!bCanSendRpc)
+	FRPCOwner OwnerInfo;
+	if (!GetRPCOwner(OwnerInfo, Object, SubObject, Function))
 	{
-		// If the subobject is a valid reference but not replicated, then the owner must be used instead
-		if (OwnerOrSubObjectReference.IsValid() && SubObject)
-		{
-			OwnerReference = ObjectReferenceCache->GetOrCreateObjectReference(Object);
-			bCanSendRpc = GetOwnerAndSubObjectIndicesFromHandle(OwnerReference.GetRefHandle(), OwnerIndex, SubObjectIndex);
-		}
-
-		if (!bCanSendRpc)
-		{
-			UE_LOG(LogIris, Warning, TEXT("SendRPC %s for %s Failed. Unable to resolve object reference %s."), ToCStr(Function->GetName()), ToCStr(Object->GetName()), ToCStr(OwnerReference.ToString()));
-			return false;
-		}
+		return false;
 	}
 
-	const TRefCountPtr<FNetRPC>& RPC = Handler->CreateRPC(OwnerReference, Function, Parameters);
+	const TRefCountPtr<FNetRPC>& RPC = Handler->CreateRPC(OwnerInfo.CallerRef, Function, Parameters);
 	if (!RPC.IsValid())
 	{
 		return true;
 	}
 
-	RPC->SetNetObjectReference(OwnerReference, OwnerOrSubObjectReference);
-	AttachmentSendQueue.Enqueue(OwnerIndex, SubObjectIndex, reinterpret_cast<const TRefCountPtr<FNetObjectAttachment>&>(RPC), SendFlags);
+	RPC->SetNetObjectReference(OwnerInfo.CallerRef, OwnerInfo.TargetRef);
+	AttachmentSendQueue.Enqueue(OwnerInfo.RootObjectIndex, OwnerInfo.SubObjectIndex, reinterpret_cast<const TRefCountPtr<FNetObjectAttachment>&>(RPC), SendFlags);
 	return true;
 }
 
@@ -173,63 +163,146 @@ bool FNetBlobManager::SendRPC(uint32 ConnectionId, const UObject* Object, const
 		return true;
 	}
 
-	FInternalNetRefIndex OwnerIndex = 0;
-	FInternalNetRefIndex SubObjectIndex = 0;
+	FRPCOwner OwnerInfo;
+	if (!GetRPCOwner(OwnerInfo, Object, SubObject, Function))
+	{
+		return false;
+	}
+
+	const TRefCountPtr<FNetRPC>& RPC = Handler->CreateRPC(OwnerInfo.CallerRef, Function, Parameters);
+	if (!RPC.IsValid())
+	{
+		UE_LOG(LogIris, Warning, TEXT("Unable to create RPC for function %s."), ToCStr(Function->GetName()));
+		return true;
+	}
+
+	RPC->SetNetObjectReference(OwnerInfo.CallerRef, OwnerInfo.TargetRef);
+	AttachmentSendQueue.Enqueue(ConnectionId, OwnerInfo.RootObjectIndex, OwnerInfo.SubObjectIndex, reinterpret_cast<const TRefCountPtr<FNetObjectAttachment>&>(RPC), SendFlags);
+	return true;
+}
 
-	const FNetObjectReference OwnerOrSubObjectReference = ObjectReferenceCache->GetOrCreateObjectReference(SubObject ? SubObject : Object);	
-	FNetObjectReference OwnerReference = OwnerOrSubObjectReference;
+bool FNetBlobManager::GetRPCOwner(FRPCOwner& OutOwnerInfo, const UObject* RootObject, const UObject* SubObject, const UFunction* Function) const
+{
+	bool bCanSendRpc = false;
 
-	bool bCanSendRpc = GetOwnerAndSubObjectIndicesFromHandle(OwnerOrSubObjectReference.GetRefHandle(), OwnerIndex, SubObjectIndex);
-	if (!bCanSendRpc)
+	// If a root object is sending an RPC
+	if (SubObject == nullptr)
 	{
-		// If the subobject is a valid reference but  not replicated, then the owner must be used instead
-		if (OwnerOrSubObjectReference.IsValid() && SubObject)
+		OutOwnerInfo.TargetRef = ObjectReferenceCache->GetOrCreateObjectReference(RootObject);
+		OutOwnerInfo.CallerRef = OutOwnerInfo.TargetRef;
+
+		bCanSendRpc = GetRootObjectIndicesFromHandle(OutOwnerInfo.TargetRef.GetRefHandle(), OutOwnerInfo.RootObjectIndex);
+
+		if (!bCanSendRpc)
 		{
-			OwnerReference = ObjectReferenceCache->GetOrCreateObjectReference(Object);
-			bCanSendRpc = GetOwnerAndSubObjectIndicesFromHandle(OwnerReference.GetRefHandle(), OwnerIndex, SubObjectIndex);
+			UE_LOG(LogIris, Warning, TEXT("SendRPC %s for %s Failed. This rootobject is not yet replicated (RefHandle: %s Index: %u)."),
+				ToCStr(Function->GetName()), *GetNameSafe(RootObject), ToCStr(OutOwnerInfo.CallerRef.GetRefHandle().ToString()), OutOwnerInfo.RootObjectIndex);
 		}
+	}
+	// If a subobject is sending an RPC
+	else
+	{
+		const FNetRefHandle SubObjectNetRef = ObjectReferenceCache->GetObjectReferenceHandleFromObject(SubObject);
+
+		// If the subobject can be referenced
+		if (SubObjectNetRef.IsValid())
+		{
+			OutOwnerInfo.TargetRef = ObjectReferenceCache->GetOrCreateObjectReference(SubObject);
+			OutOwnerInfo.CallerRef = OutOwnerInfo.TargetRef;
 
+			check(OutOwnerInfo.TargetRef.GetRefHandle() == SubObjectNetRef);
+
+			bCanSendRpc = GetRootObjectAndSubObjectIndicesFromSubObjectHandle(OutOwnerInfo.TargetRef.GetRefHandle(), OutOwnerInfo.RootObjectIndex, OutOwnerInfo.SubObjectIndex);
+		}
+		
+		// Send the RPC via the Root object if the subobject is not capable
 		if (!bCanSendRpc)
 		{
-			UE_LOG(LogIris, Warning, TEXT("SendRPC %s for %s Failed. Unable to resolve object reference %s."), ToCStr(Function->GetName()), ToCStr(Object->GetName()), ToCStr(OwnerReference.ToString()));
-			return false;
+			OutOwnerInfo.TargetRef = ObjectReferenceCache->GetOrCreateObjectReference(SubObject);
+			OutOwnerInfo.CallerRef = ObjectReferenceCache->GetOrCreateObjectReference(RootObject);
+
+			bCanSendRpc = GetRootObjectIndicesFromHandle(OutOwnerInfo.CallerRef.GetRefHandle(), OutOwnerInfo.RootObjectIndex);
+		}
+		
+		if (!bCanSendRpc)
+		{
+			UE_LOG(LogIris, Warning, TEXT("SendRPC %s for %s::%s Failed. The root object (RefHandle: %s Index: %u) and subobject (RefHandle: %s Index: %u) is not yet replicated."),
+				ToCStr(Function->GetName()), *GetNameSafe(RootObject), *GetNameSafe(SubObject), ToCStr(OutOwnerInfo.CallerRef.ToString()), OutOwnerInfo.RootObjectIndex, ToCStr(OutOwnerInfo.TargetRef.ToString()), OutOwnerInfo.SubObjectIndex);
 		}
 	}
 
-	const TRefCountPtr<FNetRPC>& RPC = Handler->CreateRPC(OwnerReference, Function, Parameters);
-	if (!RPC.IsValid())
+	return bCanSendRpc;
+}
+
+bool FNetBlobManager::GetRootObjectIndicesFromHandle(FNetRefHandle RootObjectRefHandle, FInternalNetRefIndex& OutRootObjectIndex) const
+{
+	if (!RootObjectRefHandle.IsValid())
 	{
-		UE_LOG(LogIris, Warning, TEXT("Unable to create RPC for function %s."), ToCStr(Function->GetName()));
-		return true;
+		return false;
+	}
+
+	const FInternalNetRefIndex ObjectIndex = NetRefHandleManager->GetInternalIndex(RootObjectRefHandle);
+	if (ObjectIndex == FNetRefHandleManager::InvalidInternalIndex)
+	{
+		return false;
 	}
 
-	RPC->SetNetObjectReference(OwnerReference, OwnerOrSubObjectReference);
-	AttachmentSendQueue.Enqueue(ConnectionId, OwnerIndex, SubObjectIndex, reinterpret_cast<const TRefCountPtr<FNetObjectAttachment>&>(RPC), SendFlags);
+	checkf(NetRefHandleManager->GetReplicatedObjectDataNoCheck(ObjectIndex).SubObjectRootIndex == FNetRefHandleManager::InvalidInternalIndex, TEXT("Object %s (index:%u) (netref:%s) is not a rootobject"),
+		*GetNameSafe(NetRefHandleManager->GetReplicatedObjectInstance(ObjectIndex)), ObjectIndex, ToCStr(RootObjectRefHandle.ToString()));
+
+	OutRootObjectIndex = ObjectIndex;
+
 	return true;
 }
 
-bool FNetBlobManager::GetOwnerAndSubObjectIndicesFromHandle(FNetRefHandle RefHandle, FInternalNetRefIndex& OutOwnerIndex, FInternalNetRefIndex& OutSubObjectIndex)
+bool FNetBlobManager::GetRootObjectAndSubObjectIndicesFromSubObjectHandle(FNetRefHandle SubObjectRefHandle, FInternalNetRefIndex& OutRootObjectIndex, FInternalNetRefIndex& OutSubObjectIndex) const
+{
+	if (!SubObjectRefHandle.IsValid())
+	{
+		return false;
+	}
+
+	const FInternalNetRefIndex ObjectIndex = NetRefHandleManager->GetInternalIndex(SubObjectRefHandle);
+	if (ObjectIndex == FNetRefHandleManager::InvalidInternalIndex)
+	{
+		return false;
+	}
+
+	const FNetRefHandleManager::FReplicatedObjectData& ObjectData = NetRefHandleManager->GetReplicatedObjectDataNoCheck(ObjectIndex);
+
+	checkf(ObjectData.SubObjectRootIndex != FNetRefHandleManager::InvalidInternalIndex, TEXT("SubObject %s (index:%u) (netref:%s) does not have a rootobject"),
+		*GetNameSafe(NetRefHandleManager->GetReplicatedObjectInstance(ObjectIndex)), ObjectIndex, ToCStr(SubObjectRefHandle.ToString()));
+
+	OutRootObjectIndex = ObjectData.SubObjectRootIndex;
+	OutSubObjectIndex = ObjectIndex;
+
+	return OutRootObjectIndex != FNetRefHandleManager::InvalidInternalIndex;
+}
+
+bool FNetBlobManager::GetRootObjectAndSubObjectIndicesFromAnyHandle(FNetRefHandle AnyRefHandle, FInternalNetRefIndex& OutRootObjectIndex, FInternalNetRefIndex& OutSubObjectIndex) const
 {
-	if (!RefHandle.IsValid())
+	if (!AnyRefHandle.IsValid())
 	{
 		return false;
 	}
 
-	const FInternalNetRefIndex ObjectIndex = NetRefHandleManager->GetInternalIndex(RefHandle);
+	const FInternalNetRefIndex ObjectIndex = NetRefHandleManager->GetInternalIndex(AnyRefHandle);
 	if (ObjectIndex == FNetRefHandleManager::InvalidInternalIndex)
 	{
 		return false;
 	}
 
 	const FNetRefHandleManager::FReplicatedObjectData& ObjectData = NetRefHandleManager->GetReplicatedObjectDataNoCheck(ObjectIndex);
+	// If it's a sub object
 	if (ObjectData.SubObjectRootIndex != FNetRefHandleManager::InvalidInternalIndex)
 	{
-		OutOwnerIndex = ObjectData.SubObjectRootIndex;
+		OutRootObjectIndex = ObjectData.SubObjectRootIndex;
 		OutSubObjectIndex = ObjectIndex;
 	}
+	// If it's a root object
 	else
 	{
-		OutOwnerIndex = ObjectIndex;
+		OutRootObjectIndex = ObjectIndex;
 		OutSubObjectIndex = FNetRefHandleManager::InvalidInternalIndex;
 	}
 
diff --git a/Engine/Source/Runtime/Experimental/Iris/Core/Private/Iris/ReplicationSystem/NetBlob/NetBlobManager.h b/Engine/Source/Runtime/Experimental/Iris/Core/Private/Iris/ReplicationSystem/NetBlob/NetBlobManager.h
--- a/Engine/Source/Runtime/Experimental/Iris/Core/Private/Iris/ReplicationSystem/NetBlob/NetBlobManager.h
+++ b/Engine/Source/Runtime/Experimental/Iris/Core/Private/Iris/ReplicationSystem/NetBlob/NetBlobManager.h
@@ -74,7 +74,36 @@ public:
 private:
 	void RegisterDefaultHandlers();
 
-	bool GetOwnerAndSubObjectIndicesFromHandle(FNetRefHandle RefHandle, FInternalNetRefIndex& OutOwnerIndex, FInternalNetRefIndex& OutSubObjectIndex);
+	struct FRPCOwner
+	{
+		// The replicated object responsible for carrying (sending) the RPC.
+		FNetObjectReference CallerRef;
+
+		// The object the RPC will be applied to.
+		FNetObjectReference TargetRef;
+
+		FInternalNetRefIndex RootObjectIndex = FNetRefHandleManager::InvalidInternalIndex;
+		FInternalNetRefIndex SubObjectIndex = FNetRefHandleManager::InvalidInternalIndex;
+	};
+	bool GetRPCOwner(FRPCOwner& OutOwnerInfo, const UObject* RootObject, const UObject* SubObject, const UFunction* Function) const;
+
+	/** 
+	* Validates that RootObjectRefHandle is a true root object and return it's index
+	* @return True if the root object handle is valid and has been replicated.
+	*/
+	bool GetRootObjectIndicesFromHandle(FNetRefHandle RootObjectRefHandle, FInternalNetRefIndex& OutRootObjectIndex) const;
+
+	/**
+	* Validates that the SubObjectRefHandle is a true subobject and return it's index and the index of it's root object.
+	* @return True if the subobject handle is valid and has been replicated.
+	*/
+	bool GetRootObjectAndSubObjectIndicesFromSubObjectHandle(FNetRefHandle SubObjectRefHandle, FInternalNetRefIndex& OutRootObjectIndex, FInternalNetRefIndex& OutSubObjectIndex) const;
+
+	/**
+	 * Receives the handle of a root object or a sub object and returns the index of the root object and subobject if there is one.
+	 * @return True if the object handle is valid has been replicated.
+	 */
+	bool GetRootObjectAndSubObjectIndicesFromAnyHandle(FNetRefHandle AnyRefHandle, FInternalNetRefIndex& OutRootObjectIndex, FInternalNetRefIndex& OutSubObjectIndex) const;
 
 	class FNetObjectAttachmentSendQueue
 	{
