Server Ensure in Chaos::TImplicitObjectScaled<Chaos::FTriangleMeshImplicitObject, true>::FindGeometryOpposingNormal
is quite safe, so changing that to a chaos ensure

diff --git a/Engine/Source/Runtime/Experimental/Chaos/Public/Chaos/ImplicitObjectScaled.h b/Engine/Source/Runtime/Experimental/Chaos/Public/Chaos/ImplicitObjectScaled.h
--- a/Engine/Source/Runtime/Experimental/Chaos/Public/Chaos/ImplicitObjectScaled.h
+++ b/Engine/Source/Runtime/Experimental/Chaos/Public/Chaos/ImplicitObjectScaled.h
@@ -850,7 +850,7 @@ public:
 		const FVec3 LocalOriginalNormalDenorm = OriginalNormal * MScale;
 		const FReal NormalLengthScale = LocalOriginalNormalDenorm.Size();
 		const FVec3 LocalOriginalNormal
-			= ensure(NormalLengthScale > UE_SMALL_NUMBER)
+			= CHAOS_ENSURE(NormalLengthScale > UE_SMALL_NUMBER)
 			? LocalOriginalNormalDenorm / NormalLengthScale
 			: FVec3(0, 0, 1);
 
