Geometry collection replication : when client release particles that are no longer active on the server , make sure they are disabled right away

diff --git a/Engine/Source/Runtime/Experimental/GeometryCollectionEngine/Private/GeometryCollection/GeometryCollectionComponent.cpp b/Engine/Source/Runtime/Experimental/GeometryCollectionEngine/Private/GeometryCollection/GeometryCollectionComponent.cpp
--- a/Engine/Source/Runtime/Experimental/GeometryCollectionEngine/Private/GeometryCollection/GeometryCollectionComponent.cpp
+++ b/Engine/Source/Runtime/Experimental/GeometryCollectionEngine/Private/GeometryCollection/GeometryCollectionComponent.cpp
@@ -2687,7 +2687,14 @@ namespace
 							}
 							else
 							{
-								// todo(chaos) : no release data : we can disable the particle right away ( it was released in the past )
+								// if we are a leaf we can disable this particle right away if there's no release data
+								// becuase this means the server has already disabled the particle from a removal 
+								if (ParticleHandle->ClusterIds().NumChildren == 0)
+								{
+									// todo we could probably disable all at once later to have only one call 
+									Solver->GetEvolution()->DisableParticle(ParticleHandle);
+									Solver->GetParticles().MarkTransientDirtyParticle(ParticleHandle);
+								}
 							}
 						}
 					}
