Fix the radial impulse not properly accounting for falloff in when strain value was specified

diff --git a/Engine/Source/Runtime/Experimental/Chaos/Private/Chaos/PhysicsObjectInterface.cpp b/Engine/Source/Runtime/Experimental/Chaos/Private/Chaos/PhysicsObjectInterface.cpp
--- a/Engine/Source/Runtime/Experimental/Chaos/Private/Chaos/PhysicsObjectInterface.cpp
+++ b/Engine/Source/Runtime/Experimental/Chaos/Private/Chaos/PhysicsObjectInterface.cpp
@@ -87,12 +87,14 @@ namespace
 		}
 	}
 
+	// Apply an impulse based on an origin,radius and falloff parameters
+	// return the fall off value ( between 0 and 1 )
 	template<typename T>
-	float AddRadialImpulseHelper(T ParticleHandle, FVector Origin, float Radius, float Strength, enum ERadialImpulseFalloff Falloff, float VelocityRatio = 1.0f, bool bInvalidate = true, bool bVelChange = false)
+	float AddRadialImpulseHelper(T ParticleHandle, const FVector& Origin, const float Radius, const float Strength, const enum ERadialImpulseFalloff Falloff, const float VelocityRatio = 1.0f, bool bInvalidate = true, bool bVelChange = false)
 	{
 		using namespace Chaos;
 
-		float FalloffStrength = 0.f;
+		float FalloffAlpha = 0.f;
 
 		if (ParticleHandle)
 		{
@@ -117,19 +119,19 @@ namespace
 			{
 				// by default we are within the radius so strength is maximum
 				// equivalent to ERadialImpulseFalloff::RIF_Constant
-				FalloffStrength = Strength;
+				FalloffAlpha = 1.0f;
 
 				if (Falloff == ERadialImpulseFalloff::RIF_Linear)
 				{
 					const float Distance = FMath::Sqrt(DistanceSquared);
-					FalloffStrength = static_cast<float>(FalloffStrength * (1.0f - Distance / Radius));
+					FalloffAlpha = static_cast<float>(1.0f - Distance / Radius);
 				}
 
 				// if the strength was still strong enough to consider
-				if (FalloffStrength > 0)
+				if (FalloffAlpha > 0)
 				{
 					const FVec3 Normal = ParticleToOrigin.GetSafeNormal();
-					const FVec3 Impulse = Normal * FalloffStrength;
+					const FVec3 Impulse = Normal * FalloffAlpha * Strength;
 					const FReal InvMass = bVelChange ? 1.0 : ParticleHandle->InvM();
 					const FVec3 Velocity = Impulse * InvMass * VelocityRatio;
 
@@ -139,7 +141,7 @@ namespace
 				}
 			}
 		}
-		return FalloffStrength;
+		return FalloffAlpha;
 	}
 }
 
@@ -950,13 +952,13 @@ namespace Chaos
 
 								for (FPBDRigidParticleHandle* ChildHandle : *ChildrenHandles)
 								{
-									const float FalloffStrength = AddRadialImpulseHelper(ChildHandle, Origin, Radius, Strength, Falloff, VelocityRatio, bInvalidate, bVelChange);
+									const float FalloffAlpha = AddRadialImpulseHelper(ChildHandle, Origin, Radius, Strength, Falloff, VelocityRatio, bInvalidate, bVelChange);
 
 									//to do: remove cvar when material system is in place and densities are updated
 									const float StrainToApply = 
 										(Strain < 0)
-											? (PhysicsObjectInterfaceCVars::StrainModifier * FalloffStrength)
-											: Strain;
+											? (PhysicsObjectInterfaceCVars::StrainModifier * FalloffAlpha * Strength)
+											: Strain * FalloffAlpha;
 									if (StrainToApply > 0)
 									{
 										Clustering.SetExternalStrain(ChildHandle->CastToClustered(), StrainToApply);
