Chaos - Distribute MidPhases equally among workers for better load balancing
- Chaos ParallelFor support for MaxWorkers cvar (helps with testing low-core physics on PC)

diff --git a/Engine/Source/Runtime/Experimental/Chaos/Private/Chaos/Framework/Parallel.cpp b/Engine/Source/Runtime/Experimental/Chaos/Private/Chaos/Framework/Parallel.cpp
--- a/Engine/Source/Runtime/Experimental/Chaos/Private/Chaos/Framework/Parallel.cpp
+++ b/Engine/Source/Runtime/Experimental/Chaos/Private/Chaos/Framework/Parallel.cpp
@@ -73,7 +73,13 @@ void Chaos::PhysicsParallelFor(int32 InNum, TFunctionRef<void(int32)> InCallable
 #endif
 		InCallable(Idx);
 	};
-	::ParallelFor(InNum, PassThrough, !!GSingleThreadedPhysics || bDisablePhysicsParallelFor || bForceSingleThreaded);
+
+	const bool bSingleThreaded = !!GSingleThreadedPhysics || bDisablePhysicsParallelFor || bForceSingleThreaded;
+	const EParallelForFlags Flags = (bSingleThreaded ? EParallelForFlags::ForceSingleThread : EParallelForFlags::None);
+	const int32 MinBatchSize = ((MaxNumWorkers > 0) && (InNum > MaxNumWorkers)) ? FMath::DivideAndRoundUp(InNum, MaxNumWorkers) : 1;
+
+	ParallelFor(TEXT("PhysicsParallelFor"), InNum, MinBatchSize, PassThrough, Flags);
+	//::ParallelFor(InNum, PassThrough, !!GSingleThreadedPhysics || bDisablePhysicsParallelFor || bForceSingleThreaded);
 }
 
 void Chaos::PhysicsParallelForRange(int32 InNum, TFunctionRef<void(int32, int32)> InCallable, const int32 InMinBatchSize, bool bForceSingleThreaded)
@@ -167,6 +173,7 @@ void Chaos::PhysicsParallelForWithContext(int32 InNum, TFunctionRef<int32 (int32
 
 	const bool bSingleThreaded = !!GSingleThreadedPhysics || bDisablePhysicsParallelFor || bForceSingleThreaded;
 	const EParallelForFlags Flags = bSingleThreaded ? (EParallelForFlags::ForceSingleThread) : (EParallelForFlags::None);
+	const int32 MinBatchSize = ((MaxNumWorkers > 0) && (InNum > MaxNumWorkers)) ? FMath::DivideAndRoundUp(InNum, MaxNumWorkers) : 1;
 
 	// Unfortunately ParallelForWithTaskContext takes an array of context objects - we don't use it and in our case
 	// it ends up being an array where array[index] = index.
@@ -175,7 +182,7 @@ void Chaos::PhysicsParallelForWithContext(int32 InNum, TFunctionRef<int32 (int32
 	// contexts and use the context indeex to get its context from it.
 	TArray<int32, TInlineAllocator<16>> Contexts;
 
-	::ParallelForWithTaskContext(Contexts, InNum, InContextCreator, PassThrough, Flags);
+	::ParallelForWithTaskContext(TEXT("PhysicsParallelForWithContext"), Contexts, InNum, MinBatchSize, InContextCreator, PassThrough, Flags);
 }
 
 
diff --git a/Engine/Source/Runtime/Experimental/Chaos/Private/Chaos/PBDRigidsEvolutionGBF.cpp b/Engine/Source/Runtime/Experimental/Chaos/Private/Chaos/PBDRigidsEvolutionGBF.cpp
--- a/Engine/Source/Runtime/Experimental/Chaos/Private/Chaos/PBDRigidsEvolutionGBF.cpp
+++ b/Engine/Source/Runtime/Experimental/Chaos/Private/Chaos/PBDRigidsEvolutionGBF.cpp
@@ -168,6 +168,10 @@ namespace Chaos
 		bool bChaosSolverCheckParticleViews = false;
 		FAutoConsoleVariableRef CVarChaosSolverCheckParticleViews(TEXT("p.Chaos.Solver.CheckParticleViews"), bChaosSolverCheckParticleViews, TEXT(""));
 
+		// Enable improved midphase distribution among worker threads. Without this the midphase and narrowphase run on whatever thread the
+		// broadphas eoverlap was detected on which tends to give poor distribution.
+		bool bChaosMidPhaseRedistributionEnabled = true;
+		FAutoConsoleVariableRef CVarChaosSolverMidPhaseRedistributionEnabled(TEXT("p.Chaos.BroadPhase.MidPhaseRedistributionEnabled"), bChaosMidPhaseRedistributionEnabled, TEXT(""));
 	}
 	using namespace CVars;
 
diff --git a/Engine/Source/Runtime/Experimental/Chaos/Public/Chaos/Collision/SpatialAccelerationBroadPhase.h b/Engine/Source/Runtime/Experimental/Chaos/Public/Chaos/Collision/SpatialAccelerationBroadPhase.h
--- a/Engine/Source/Runtime/Experimental/Chaos/Public/Chaos/Collision/SpatialAccelerationBroadPhase.h
+++ b/Engine/Source/Runtime/Experimental/Chaos/Public/Chaos/Collision/SpatialAccelerationBroadPhase.h
@@ -20,6 +20,11 @@
 #define CHAOS_CHECK_BROADPHASE 0
 #endif
 
+namespace Chaos::CVars
+{
+	extern bool bChaosMidPhaseRedistributionEnabled;
+}
+
 namespace Chaos
 {
 	template <typename TPayloadType, typename T, int d>
@@ -379,7 +384,14 @@ namespace Chaos
 
 			FChaosVDContextWrapper CVDContext;
 			CVD_GET_WRAPPED_CURRENT_CONTEXT(CVDContext);
-	
+
+			// The broadphase overlaps probably resulted in a very different number of MidPhases 
+			// (overlaps) on each thread. Redistribute the MidPhases so that we get more even 
+			// processing. NOTE: This does not take into account that some MidPhases may be expensive. 
+			// For that we would need to queue and redistribute the NarrowPhases, but currently each
+			// MidPhase runs the NarrowPhase in ProcessMidPhase
+			RedistributeMidPhasesInContexts();
+
 			const auto& ProcessMidPhasesWorker = [this, Dt, &CVDContext](const int32 ContextIndex)
 			{
 				CVD_SCOPE_CONTEXT(CVDContext.Context);
@@ -387,8 +399,7 @@ namespace Chaos
 			};
 			PhysicsParallelFor(NumActiveBroadphaseContexts, ProcessMidPhasesWorker, bDisableCollisionParallelFor);
 		}
-	
-		
+
 		/** @brief This function is the outer loop of collision detection. It loops over the
 		 * particles view and do the broadphase + narrowphase collision detection
 		 * @param OverlapView View to consider for the outer loop
@@ -577,6 +588,56 @@ namespace Chaos
 			}
 		}
 
+		// Redistribute the midphases among the contects to even out the per-core work a bit
+		void RedistributeMidPhasesInContexts()
+		{
+			QUICK_SCOPE_CYCLE_COUNTER(STAT_Collisions_RedistributeMidPhases);
+
+			if ((NumActiveBroadphaseContexts > 1) && CVars::bChaosMidPhaseRedistributionEnabled)
+			{
+				// We want this many midphases per worker
+				const int32 NumMidPhasesPerContext = FMath::DivideAndRoundUp(NumMidPhases, NumActiveBroadphaseContexts);
+
+				// Reserve array space
+				for (int32 ContextIndex = 0; ContextIndex < NumActiveBroadphaseContexts; ++ContextIndex)
+				{
+					BroadphaseContexts[ContextIndex].MidPhases.Reserve(NumMidPhasesPerContext);
+				}
+
+				// Find the over-filled arrays and move elements to the under-filled arrays
+				for (int32 SrcContextIndex = 0; SrcContextIndex < NumActiveBroadphaseContexts; ++SrcContextIndex)
+				{
+					TArray<FParticlePairMidPhase*>& SrcMidPhases = BroadphaseContexts[SrcContextIndex].MidPhases;
+					int32 SrcNum = SrcMidPhases.Num();
+					if (SrcNum > NumMidPhasesPerContext)
+					{
+						// SrcMidPhases is over-filled
+						for (int32 DstContextIndex = 0; DstContextIndex < NumActiveBroadphaseContexts; ++DstContextIndex)
+						{
+							TArray<FParticlePairMidPhase*>& DstMidPhases = BroadphaseContexts[DstContextIndex].MidPhases;
+							if ((DstContextIndex != SrcContextIndex) && (DstMidPhases.Num() < NumMidPhasesPerContext))
+							{
+								// DstMidPhases is under-filled, see how many elements we can move into it
+								const int32 NumToMove = FMath::Min(SrcNum - NumMidPhasesPerContext, NumMidPhasesPerContext - DstMidPhases.Num());
+								
+								// Copy the elements from the end or Src to end of Dst (we will resize SrcMidPhases at the end)
+								SrcNum -= NumToMove;
+								DstMidPhases.Append(&SrcMidPhases[SrcNum], NumToMove);
+
+								if (SrcNum == NumMidPhasesPerContext)
+								{
+									break;
+								}
+							}
+						}
+
+						check(SrcNum == NumMidPhasesPerContext);
+						SrcMidPhases.SetNum(SrcNum);
+					}
+				}
+			}
+		}
+
 		// Find or assign midphases to each of the overlapping particle pairs
 		// @todo(chaos): optimize
 		void AssignMidPhases(Private::FBroadPhaseContext& BroadphaseContext)
