- Fixed sending requests too often.
- Exclude proxy support from shipping builds.

#rnx

diff --git a/Engine/Source/Runtime/Experimental/IoStoreOnDemand/Private/CoreHttp/Http.cpp b/Engine/Source/Runtime/Experimental/IoStoreOnDemand/Private/CoreHttp/Http.cpp
--- a/Engine/Source/Runtime/Experimental/IoStoreOnDemand/Private/CoreHttp/Http.cpp
+++ b/Engine/Source/Runtime/Experimental/IoStoreOnDemand/Private/CoreHttp/Http.cpp
@@ -1008,6 +1008,8 @@ int32 FSocket::Wait(TArrayView<FWaiter> Waiters, int32 TimeoutMs)
 
 // {{{1 socks ..................................................................
 
+#if !UE_BUILD_SHIPPING
+
 ///////////////////////////////////////////////////////////////////////////////
 static uint32	GSocksIpAddress	= 0;
 static int32	GSocksPort		= 1080; // default SOCKS5 port
@@ -1185,9 +1187,14 @@ static int32 ConnectSocks5(FSocket& Socket, uint32 IpAddress, uint32 Port)
 #endif
 }
 
+#endif // UE_BUILD_SHIPPING
+
 ////////////////////////////////////////////////////////////////////////////////
 static int32 MaybeConnectSocks(FSocket& Socket, uint32 IpAddress, uint32 Port)
 {
+#if UE_BUILD_SHIPPING
+	return 0;
+#else
 	if (GSocksIpAddress == 0)
 	{
 		return 0;
@@ -1200,6 +1207,7 @@ static int32 MaybeConnectSocks(FSocket& Socket, uint32 IpAddress, uint32 Port)
 	}
 
 	return -1;
+#endif // UE_BUILD_SHIPPING
 }
 
 
@@ -1378,14 +1386,20 @@ FResult FHost::Connect(FSocket& Socket)
 		Candidate.SetRecvBufSize(OptValue);
 	}
 
-	// Adjust socket send and recv buffer sizes
-	if (!bSocksConnected && !Candidate.Connect(IpAddress, Port))
+	// Socks connect in a blocking fashion so we're all set (ret=1)
+	if (bSocksConnected)
+	{
+		return FResult(1);
+	}
+
+	// Issue the connect - this is done non-blocking so we need to wait (ret=0)
+	if (!Candidate.Connect(IpAddress, Port))
 	{
 		return FResult("Socket connect failed");
 	}
 
 	Socket = MoveTemp(Candidate);
-	return FResult(1);
+	return FResult(0);
 }
 
 
