Removed old unused code

#rnx

diff --git a/Engine/Source/Runtime/Experimental/IoStoreOnDemand/Private/CoreHttp/Client.h b/Engine/Source/Runtime/Experimental/IoStoreOnDemand/Private/CoreHttp/Client.h
--- a/Engine/Source/Runtime/Experimental/IoStoreOnDemand/Private/CoreHttp/Client.h
+++ b/Engine/Source/Runtime/Experimental/IoStoreOnDemand/Private/CoreHttp/Client.h
@@ -82,7 +82,7 @@ public:
 
 private:
 	friend					class FEventLoop;
-	class FSocketPool*		Ptr = nullptr;
+	class FHost*			Ptr = nullptr;
 
 private:
 							FConnectionPool(const FConnectionPool&) = delete;
diff --git a/Engine/Source/Runtime/Experimental/IoStoreOnDemand/Private/CoreHttp/Http.cpp b/Engine/Source/Runtime/Experimental/IoStoreOnDemand/Private/CoreHttp/Http.cpp
--- a/Engine/Source/Runtime/Experimental/IoStoreOnDemand/Private/CoreHttp/Http.cpp
+++ b/Engine/Source/Runtime/Experimental/IoStoreOnDemand/Private/CoreHttp/Http.cpp
@@ -1404,74 +1404,6 @@ FResult FHost::Connect(FSocket& Socket)
 
 
 
-////////////////////////////////////////////////////////////////////////////////
-class FSocketPool
-	: public FHost
-{
-public:
-					FSocketPool(const ANSICHAR* InHostName, uint32 InPort, uint32 InMaxLeases, uint32 PipeLen=1);
-					~FSocketPool();
-	static uint32	GetAllocSize(uint32 MaxLeases);
-	bool			LeaseSocket(FSocket& Out);
-	void			ReturnLease(FSocket&& Socket);
-
-private:
-	uint8			LeaseCount = 0;
-	FSocket			Sockets[1/*...N*/]; // this should be the last member
-};
-
-////////////////////////////////////////////////////////////////////////////////
-FSocketPool::FSocketPool(const ANSICHAR* InHostName, uint32 InPort, uint32 InMaxLeases, uint32 PipeLen)
-: FHost(InHostName, InPort, InMaxLeases, PipeLen)
-{
-
-	for (uint32 i = 0; i < GetMaxConnections(); ++i)
-	{
-		new (Sockets + i) FSocket();
-	}
-}
-
-////////////////////////////////////////////////////////////////////////////////
-FSocketPool::~FSocketPool()
-{
-	for (uint32 i = 0; i < GetMaxConnections(); ++i)
-	{
-		Sockets[i].Destroy();
-	}
-}
-
-////////////////////////////////////////////////////////////////////////////////
-uint32 FSocketPool::GetAllocSize(uint32 MaxLeases)
-{
-	return sizeof(FSocketPool) + (sizeof(Sockets[0]) * (MaxLeases - 1));
-}
-
-////////////////////////////////////////////////////////////////////////////////
-bool FSocketPool::LeaseSocket(FSocket& Out)
-{
-	check(LeaseCount <= GetMaxConnections());
-
-	if (LeaseCount == GetMaxConnections())
-	{
-		return false;
-	}
-
-	Out = MoveTemp(Sockets[LeaseCount]);
-	++LeaseCount;
-
-	return true;
-}
-
-////////////////////////////////////////////////////////////////////////////////
-void FSocketPool::ReturnLease(FSocket&& Socket)
-{
-	check(LeaseCount > 0);
-	--LeaseCount;
-	Sockets[LeaseCount] = MoveTemp(Socket);
-}
-
-
-
 ////////////////////////////////////////////////////////////////////////////////
 int32 FConnectionPool::FParams::SetHostFromUrl(FAnsiStringView Url)
 {
@@ -1499,18 +1431,17 @@ FConnectionPool::FConnectionPool(const FParams& Params)
 	check(Params.Host.Port - 1 <= 0xfffeu);
 
 	// Alloc a new internal object
-	uint32 PoolAllocSize = FSocketPool::GetAllocSize(Params.ConnectionCount);
 	uint32 HostNameLen = Params.Host.Name.Len();
-	uint32 AllocSize = PoolAllocSize + (HostNameLen + 1);
-	auto* Internal = (FSocketPool*)FMemory::Malloc(AllocSize, alignof(FSocketPool));
+	uint32 AllocSize = sizeof(FHost) + (HostNameLen + 1);
+	auto* Internal = (FHost*)FMemory::Malloc(AllocSize, alignof(FHost));
 
 	// Copy host
-	char* HostDest = (char*)Internal + PoolAllocSize;
+	char* HostDest = (char*)(Internal + 1);
 	memcpy(HostDest, Params.Host.Name.GetData(), HostNameLen);
 	HostDest[HostNameLen] = '\0';
 
 	// Init internal object
-	new (Internal) FSocketPool(
+	new (Internal) FHost(
 		HostDest,
 		Params.Host.Port,
 		Params.ConnectionCount,
@@ -1612,13 +1543,10 @@ struct FResponseInternal
 ////////////////////////////////////////////////////////////////////////////////
 struct alignas(16) FActivity
 {
-	enum class EWait : uint8 { None, Read, Write, Pool };
 	enum class EState : uint8
 	{
 		None,
 		Build,
-		Resolve,
-		Connect,
 		Send,
 		RecvMessage,
 		RecvContent,
@@ -1633,7 +1561,6 @@ struct alignas(16) FActivity
 	FActivity*			Next = nullptr;
 	int8				Slot = -1;
 	EState				State = EState::None;
-	EWait				SocketWait = EWait::None;
 	uint8				IsKeepAlive : 1;
 	uint8				NoContent : 1;
 	uint8				_Unused0 : 6;
@@ -1642,7 +1569,7 @@ struct alignas(16) FActivity
 	FStopwatch			Stopwatch;
 #endif
 	union {
-		FSocketPool*	Pool;
+		FHost*			Host;
 		FIoBuffer*		Dest;
 		const char*		ErrorReason;
 	};
@@ -1657,9 +1584,8 @@ struct alignas(16) FActivity
 static void Activity_TraceStateNames()
 {
 	FAnsiStringView StateNames[] = {
-		"None", "Build", "Resolve", "Connect",	"Send", "RecvMessage",
-		"RecvContent", "RecvStream", "RecvDone", "Completed", "Cancelled",
-		"Failed", "$",
+		"None", "Build", "Send", "RecvMessage", "RecvContent", "RecvStream",
+		"RecvDone", "Completed", "Cancelled", "Failed", "$",
 	};
 	static_assert(UE_ARRAY_COUNT(StateNames) == int32(FActivity::EState::_Num) + 1);
 
@@ -1703,43 +1629,6 @@ static void Activity_ChangeState(FActivity* Activity, FActivity::EState InState,
 	Activity->StateParam = Param;
 }
 
-////////////////////////////////////////////////////////////////////////////////
-static void Activity_BeginWait(FActivity* Activity, FActivity::EWait What)
-{
-#if 0
-	check(Activity->SocketWait == FActivity::EWait::None);
-
-#if IAS_HTTP_WITH_PERF
-	if (Activity->State >= FActivity::EState::Send)
-	{
-		Activity->Stopwatch.Wait();
-	}
-#endif
-	Trace(Activity, ETrace::Wait);
-
-	Activity->SocketWait = What;
-#endif // 0
-}
-
-////////////////////////////////////////////////////////////////////////////////
-static void Activity_EndWait(FActivity* Activity)
-{
-#if 0
-	if (Activity->SocketWait != FActivity::EWait::None)
-	{
-#if IAS_HTTP_WITH_PERF
-		if (Activity->State >= FActivity::EState::Send)
-		{
-			Activity->Stopwatch.Unwait();
-		}
-#endif
-		Trace(Activity, ETrace::Unwait);
-	}
-
-	Activity->SocketWait = FActivity::EWait::None;
-#endif // 0
-}
-
 ////////////////////////////////////////////////////////////////////////////////
 static int32 Activity_Rewind(FActivity* Activity)
 {
@@ -1783,22 +1672,6 @@ static FActivity* Activity_Alloc(uint32 BufferSize)
 ////////////////////////////////////////////////////////////////////////////////
 static void Activity_Free(FActivity* Activity)
 {
-#if 0
-	if (Activity->State > FActivity::EState::Connect)
-	{
-		FSocket& Socket = Activity->Socket;
-		if (!Activity->IsKeepAlive && Socket.IsValid())
-		{
-			Socket = FSocket();
-		}
-
-		if (Activity->Pool->GetIpAddress() > FHost::InvalidIp)
-		{
-			Activity->Pool->ReturnLease(MoveTemp(Socket));
-		}
-	}
-#endif // 0
-
 	Activity->~FActivity();
 	FMemory::Free(Activity);
 
@@ -1811,7 +1684,6 @@ static void Activity_SetError(FActivity* Activity, const char* Reason)
 	Activity->IsKeepAlive = 0;
 	Activity->ErrorReason = Reason;
 
-	Activity_EndWait(Activity);
 	Activity_ChangeState(Activity, FActivity::EState::Failed, LastSocketResult());
 }
 
@@ -2108,145 +1980,6 @@ FTicketPerf::FSample FTicketPerf::GetRecvSample() const
 
 // {{{1 event-loop-int .........................................................
 
-////////////////////////////////////////////////////////////////////////////////
-static uint64 ReadyCheck(FActivity** Activities, uint32 Num, uint32 TimeoutMs)
-{
-	TRACE_CPUPROFILER_EVENT_SCOPE(IasHttp::ReadyCheck);
-
-	using EWait = FActivity::EWait;
-
-	uint64 Ret = 0;
-
-	TArray<FSocket::FWaiter, TFixedAllocator<64>> Waiters;
-
-	for (uint32 i = 0; i < Num; ++i)
-	{
-		FActivity* Activity = Activities[i];
-		switch (Activity->SocketWait)
-		{
-		case EWait::None:
-			Ret |= (1ull << Activity->Slot);
-			break;
-
-		case EWait::Pool:
-			if (Activity->Pool->GetIpAddress() > FHost::InvalidIp)
-			{
-				Activity_EndWait(Activities[i]);
-				Ret |= (1ull << Activity->Slot);
-			}
-			break;
-
-		case EWait::Read:  Waiters.Add({ Activity->Socket, FSocket::FWaiter::EWhat::Recv }); break;
-		case EWait::Write: Waiters.Add({ Activity->Socket, FSocket::FWaiter::EWhat::Send }); break;
-		}
-	}
-
-	// Early out if there's work to do or nothing to wait on.
-	if (Waiters.IsEmpty())
-	{
-		return Ret;
-	}
-
-	// Wait on sockets and transfer the results
-	if (!FSocket::Wait(Waiters, TimeoutMs))
-	{
-		return Ret;
-	}
-
-	FActivity* const* Cursor = Activities;
-	for (const FSocket::FWaiter& Waiter : Waiters)
-	{
-		for (; !(Waiter == Cursor[0]->Socket); ++Cursor);
-
-		if (!Waiter.Ready)
-		{
-			continue;
-		}
-
-		FActivity* Activity = *Cursor;
-		Activity_EndWait(Activity);
-		Ret |= (1ull << Activity->Slot);
-	}
-
-	return Ret;
-}
-
-
-
-////////////////////////////////////////////////////////////////////////////////
-static int32 DoResolve(FActivity* Activity)
-{
-	// There could be many activities using the same socket pool. Resolving only
-	// needs to happen once, the first activity in can do the honours. Everyone
-	// else can wait.
-	FSocketPool* Pool = Activity->Pool;
-	int32 Result = Pool->IsResolved();
-	if (Result > 0)
-	{
-		Activity_ChangeState(Activity, FActivity::EState::Connect);
-		return 0;
-	}
-
-	if (Result < 0)
-	{
-		Activity_ChangeState(Activity, FActivity::EState::Connect);
-		Activity_BeginWait(Activity, FActivity::EWait::Pool);
-		return 1;
-	}
-
-	// We won! We WON!
-	Result = Pool->ResolveHostName().GetValue();
-	switch (Result)
-	{
-	case 0:  Activity_SetError(Activity, "Unable to resolve host"); return -1;
-	case -1: Activity_SetError(Activity, "Error encountered resolving"); return -1;
-	case -2: Activity_SetError(Activity, "Unexpected address family during resolve"); return -1;
-	}
-
-	Activity_ChangeState(Activity, FActivity::EState::Connect);
-	return 0;
-}
-
-////////////////////////////////////////////////////////////////////////////////
-static int32 DoConnect(FActivity* Activity)
-{
-	TRACE_CPUPROFILER_EVENT_SCOPE(IasHttp::DoConnect);
-
-	FSocketPool* Pool = Activity->Pool;
-
-	// Claim an existing socket from the pool.
-	FSocket Candidate;
-	if (!Pool->LeaseSocket(Candidate))
-	{
-		// none available at this time
-		return 1;
-	}
-
-	if (Candidate.IsValid())
-	{
-		Activity->Socket = MoveTemp(Candidate);
-		Activity_EndWait(Activity);
-		Activity_ChangeState(Activity, FActivity::EState::Send);
-		return 0;
-	}
-
-	Trace(Activity, ETrace::Connect);
-
-	FResult Result = Pool->Connect(Candidate);
-	if (Result.GetValue() < 0)
-	{
-		Activity_SetError(Activity, Result.GetMessage());
-	}
-	Activity->Socket = MoveTemp(Candidate);
-
-	Activity_ChangeState(Activity, FActivity::EState::Send);
-	if (Result.GetValue() == 0)
-	{
-		Activity_BeginWait(Activity, FActivity::EWait::Write);
-	}
-	return 1;
-}
-
 ////////////////////////////////////////////////////////////////////////////////
 static int32 DoSend(FActivity* Activity, FSocket& Socket)
 {
@@ -2270,9 +2003,7 @@ static int32 DoSend(FActivity* Activity, FSocket& Socket)
 	case FSocket::EResult::HangUp:		Activity_SetError(Activity, "ATH0.Send"); return Result;
 	case FSocket::EResult::Error:		Activity_SetError(Activity, "Error returned from socket send"); return Result;
 	case FSocket::EResult::ConnectError:Activity_SetError(Activity, "Connection error"); return Result;
-	case FSocket::EResult::Wait:
-		Activity_BeginWait(Activity, FActivity::EWait::Write);
-		return Result;
+	case FSocket::EResult::Wait:		return Result;
 	}
 
 #if IAS_HTTP_WITH_PERF
@@ -2287,16 +2018,9 @@ static int32 DoSend(FActivity* Activity, FSocket& Socket)
 	}
 
 	Activity_ChangeState(Activity, FActivity::EState::RecvMessage, Buffer.GetSize());
-	Activity_BeginWait(Activity, FActivity::EWait::Read);
 	return Result;
 }
 
-////////////////////////////////////////////////////////////////////////////////
-static int32 DoSend(FActivity* Activity)
-{
-	return DoSend(Activity, Activity->Socket);
-}
-
 ////////////////////////////////////////////////////////////////////////////////
 static int32 DoRecvMessage(FActivity* Activity, FSocket& Socket)
 {
@@ -2317,7 +2041,6 @@ static int32 DoRecvMessage(FActivity* Activity, FSocket& Socket)
 
 		if (Result == int32(FSocket::EResult::Wait))
 		{
-			Activity_BeginWait(Activity, FActivity::EWait::Read);
 			return 1;
 		}
 
@@ -2531,7 +2254,6 @@ static int32 DoRecvContent(FActivity* Activity, FSocket& Socket, int32& MaxRecvS
 
 		if (Result == int32(FSocket::EResult::Wait))
 		{
-			Activity_BeginWait(Activity, FActivity::EWait::Read);
 			return 1;
 		}
 
@@ -2605,27 +2327,6 @@ static int32 DoRecvDone(FActivity* Activity)
 	return 0;
 }
 
-////////////////////////////////////////////////////////////////////////////////
-static int32 DoRecvMessage(FActivity* Activity)
-{
-	return DoRecvMessage(Activity, Activity->Socket);
-}
-
-////////////////////////////////////////////////////////////////////////////////
-static int32 DoRecvContent(FActivity* Activity, int32& RecvSize)
-{
-	return DoRecvContent(Activity, Activity->Socket, RecvSize);
-}
-
-////////////////////////////////////////////////////////////////////////////////
-static int32 DoRecvStream(FActivity* Activity, int32& RecvSize)
-{
-	TRACE_CPUPROFILER_EVENT_SCOPE(IasHttp::DoRecvStream);
-
-	check(false); // not implemented yet
-	return -1;
-}
-
 ////////////////////////////////////////////////////////////////////////////////
 static void DoCancel(FActivity* Activity)
 {
@@ -3329,7 +3030,7 @@ FRequest FEventLoop::FImpl::Request(
 	FMessageBuilder Builder(Activity->Buffer);
 
 	Builder << Method << " " << Path << " HTTP/1.1" "\r\n"
-		"Host: " << Activity->Pool->GetHostName() << "\r\n";
+		"Host: " << Activity->Host->GetHostName() << "\r\n";
 
 	// HTTP/1.1 is persistent by default thus "Connection" header isn't required
 	if (!Activity->IsKeepAlive)
@@ -3426,13 +3127,10 @@ void FEventLoop::FImpl::ReceiveWork()
 	// Group activities by their host.
 	for (FActivity* Next; Activity != nullptr; Activity = Next)
 	{
-		// Cache the next pointer and detach activity from the pending list
 		Next = Activity->Next;
 		Activity->Next = nullptr;
 
-		// Find or create group by the host pointer. This ultimately matches the
-		// FConnectionPool object used for the activity.
-		FHost& Host = *(Activity->Pool);
+		FHost& Host = *(Activity->Host);
 		auto Pred = [&Host] (const FHostGroup& Lhs) { return &Lhs.GetHost() == &Host; };
 		FHostGroup* Group = Groups.FindByPredicate(Pred);
 		if (Group == nullptr)
@@ -3445,154 +3143,6 @@ void FEventLoop::FImpl::ReceiveWork()
 	}
 }
 
-////////////////////////////////////////////////////////////////////////////////
-#if 0
-uint32 FEventLoop::FImpl::Tick(uint32 PollTimeoutMs)
-{
-	TRACE_CPUPROFILER_EVENT_SCOPE(IasHttp::Tick);
-
-	ReceiveWork();
-
-	uint32 BusyCount = Active.Num();
-
-	// Work out which activities are ready
-	uint64 Readies = ReadyCheck(Active.GetData(), Active.Num(), PollTimeoutMs);
-	if (Readies == 0)
-	{
-		return BusyCount;
-	}
-
-	// Tick activities
-	int32 RecvAllowance = Throttler.GetAllowance();
-	uint64 CancelsLoad = Cancels.load(std::memory_order_relaxed);
-	for (FActivity* Activity : Active)
-	{
-		Trace(this, ETrace::LoopTick, Activity);
-
-		uint64 SlotBit = 1ull << Activity->Slot;
-
-		if (SlotBit & CancelsLoad)
-		{
-			--BusyCount;
-			DoCancel(Activity);
-			continue;
-		}
-
-		if ((Readies & SlotBit) == 0)
-		{
-			continue;
-		}
-
-		int32 Result = 1;
-		switch (Activity->State)
-		{
-		case FActivity::EState::Resolve:
-			Result = DoResolve(Activity);
-			if (Result)
-				break;
-
-		case FActivity::EState::Connect:
-			Result = DoConnect(Activity);
-			if (Result)
-				break;
-
-		case FActivity::EState::Send:
-			Result = DoSend(Activity);
-			if (Result)
-				break;
-
-		case FActivity::EState::RecvMessage:
-			Result = DoRecvMessage(Activity);
-			if (Result)
-				break;
-
-			if (Activity->State != FActivity::EState::RecvDone)
-			{
-			case FActivity::EState::RecvContent:
-			case FActivity::EState::RecvStream:
-				int32 (*Handler)(FActivity*, int32&);
-				if (Activity->State == FActivity::EState::RecvContent)
-				{
-					Handler = DoRecvContent;
-				}
-				else
-				{
-					Handler = DoRecvStream;
-				}
-
-				if (RecvAllowance <= 0)
-				{
-					Result = 0;
-					break;
-				}
-
-				Result = Handler(Activity, RecvAllowance);
-				if (Result)
-					break;
-			}
-
-		case FActivity::EState::RecvDone:
-			Result = DoRecvDone(Activity);
-			if (Result)
-				break;
-
-		case FActivity::EState::Completed:
-			--BusyCount;
-			Result = 0;
-			break;
-		}
-
-		if (Result == -1)
-		{
-			check(Activity->State == FActivity::EState::Failed);
-
-			FTicketStatus& SinkArg = *(FTicketStatus*)Activity;
-			Activity->Sink(SinkArg);
-
-			--BusyCount;
-		}
-	}
-	Throttler.ReturnUnused(RecvAllowance);
-
-	// Reap done activities
-	uint64 ReturnedSlots = 0;
-	if (int32 n = Active.Num(); BusyCount != n)
-	{
-		int32 i = 0;
-		for (--n; i <= n;)
-		{
-			FActivity* Activity = Active[i];
-			if (Activity->State < FActivity::EState::Completed)
-			{
-				++i;
-				continue;
-			}
-
-			ReturnedSlots += (1ull << Activity->Slot);
-
-			Activity_Free(Activity);
-
-			Active[i] = Active[n];
-			--n;
-		}
-		Active.SetNum(i, false);
-	}
-
-	if (ReturnedSlots)
-	{
-		PrevFreeSlots += ReturnedSlots;
-		FreeSlots.fetch_add(ReturnedSlots, std::memory_order_relaxed);
-	}
-
-	if (CancelsLoad)
-	{
-		Cancels.fetch_and(~CancelsLoad, std::memory_order_relaxed);
-	}
-
-	return BusyCount;
-}
-#endif // 0
-
 ////////////////////////////////////////////////////////////////////////////////
 uint32 FEventLoop::FImpl::Tick(int32 PollTimeoutMs)
 {
@@ -3707,17 +3257,17 @@ FRequest FEventLoop::Request(
 		Path = Url.Mid(UrlOffsets.Path);
 	}
 
-	// Create an activity and an emphemeral socket pool
+	// Create an activity and an emphemeral host
 	Params = (Params != nullptr) ? Params : &GDefaultParams;
 
 	uint32 BufferSize = Params->BufferSize;
 	BufferSize = (BufferSize >= 128) ? BufferSize : 128;
-	BufferSize += sizeof(FSocketPool) + HostName.Len();
+	BufferSize += sizeof(FHost) + HostName.Len();
 	FActivity* Activity = Activity_Alloc(BufferSize);
 
 	FBuffer& Buffer = Activity->Buffer;
 
-	FSocketPool* Pool = Activity->Pool = Buffer.Alloc<FSocketPool>();
+	FHost* Host = Activity->Host = Buffer.Alloc<FHost>();
 	Activity->IsKeepAlive = 0;
 
 	uint32 HostNameLength = HostName.Len();
@@ -3727,7 +3277,7 @@ FRequest FEventLoop::Request(
 
 	memcpy(HostNamePtr, HostName.GetData(), HostNameLength);
 	HostNamePtr[HostNameLength] = '\0';
-	new (Pool) FSocketPool(HostNamePtr, Port, 1);
+	new (Host) FHost(HostNamePtr, Port, 1);
 
 	return Impl->Request(Method, Path, Activity);
 }
@@ -3747,7 +3297,7 @@ FRequest FEventLoop::Request(
 	BufferSize = (BufferSize >= 128) ? BufferSize : 128;
 	FActivity* Activity = Activity_Alloc(BufferSize);
 
-	Activity->Pool = Pool.Ptr;
+	Activity->Host = Pool.Ptr;
 	Activity->IsKeepAlive = 1;
 
 	return Impl->Request(Method, Path, Activity);
@@ -3760,9 +3310,6 @@ FTicket FEventLoop::Send(FRequest&& Request, FTicketSink Sink, UPTRINT SinkParam
 	Swap(Activity, Request.Ptr);
 	Activity->SinkParam = SinkParam;
 	Activity->Sink = Sink;
-
-	Activity_ChangeState(Activity, FActivity::EState::Resolve);
-
 	return Impl->Send(Activity);
 }
 
