Remove cache markers

Removes the journaled markers from the disk cache. The markers were used for finding the most recent entry and to find where overwrites occurred. Instead the journal now contains a hash of the first 64 bytes of each data entry and the marker. This enables writing data entries from a package continously into the cache, even when split into several journal operations.

diff --git a/Engine/Source/Runtime/Experimental/IoStoreOnDemand/Private/JournaledCache.cpp b/Engine/Source/Runtime/Experimental/IoStoreOnDemand/Private/JournaledCache.cpp
--- a/Engine/Source/Runtime/Experimental/IoStoreOnDemand/Private/JournaledCache.cpp
+++ b/Engine/Source/Runtime/Experimental/IoStoreOnDemand/Private/JournaledCache.cpp
@@ -13,8 +13,8 @@
 #include "HAL/PlatformProcess.h"
 #include "HAL/Runnable.h"
 #include "HAL/RunnableThread.h"
+#include "Hash/CityHash.h"
 #include "IO/IoBuffer.h"
-#include "IO/IoDispatcher.h"
 #include "IO/IoHash.h"
 #include "IO/IoStatus.h"
 #include "Math/UnrealMath.h"
@@ -291,10 +291,11 @@ int32 FMemCache::Drop(uint32 Size)
 // {{{1 phrase .................................................................
 
 ////////////////////////////////////////////////////////////////////////////////
-static const uint32 MAGIC = 0x04930001;
+static const uint32 MAGIC = 0x04930002;
 static const uint32 SIZE_BITS = 25;
+static const uint32 MARKER_MAX = 0x3fffffff;
 
-using MarkerType = uint32;
+using EntryHash = uint32;
 
 struct FDataEntry
 {
@@ -307,9 +308,9 @@ struct FDataEntry
 struct FPhraseDesc
 {
 	uint32		Magic;
-	MarkerType	Marker;
-	uint64		EntryCount : 16;
-	uint64		DataCursor : 48;
+	EntryHash	Hash;
+	uint64		Marker : 30; // Safely store ~1G ops (1<<30)/2 operations (with wrap)
+	uint64		DataCursor : 34; // Allows for cache size up to 31 Gb
 	FDataEntry	Entries[];
 };
 
@@ -326,13 +327,12 @@ struct FDiskPhrase
 	const FDataEntry*	GetEntries() const		{ return Entries.GetData() + Index; }
 	int32				GetEntryCount() const	{ return Entries.Num() - Index; }
 	uint8*				GetPhraseData() const	{ return Buffer.Get(); }
-	uint32				GetPhraseSize() const	{ return Cursor; }
-	uint32				GetDataSize() const		{ return Cursor - sizeof(MarkerType); }
+	uint32				GetDataSize() const		{ return Cursor; } 
 
 private:
 	TUniquePtr<uint8[]>	Buffer;
 	TArray<FDataEntry>&	Entries;
-	uint32				Cursor = sizeof(MarkerType);
+	uint32				Cursor = 0;
 	uint32				Index;
 	int32				MaxEntries;
 
@@ -348,7 +348,6 @@ FDiskPhrase::FDiskPhrase(TArray<FDataEntry>& InEntries, int32 InMaxEntries, uint
 , Index(Entries.Num())
 , MaxEntries(InMaxEntries)
 {
-	DataSize += Cursor;
 	Buffer = TUniquePtr<uint8[]>(new uint8[DataSize]);
 }
 
@@ -386,12 +385,14 @@ private:
 	friend bool				LoadCache(FDiskCache&);
 	void					GetPath(TStringBuilder<64>& Out);
 	void					OpenJrnFile();
+	static uint32			HashBytes(const uint8* Data, uint32 Size, uint32 Seed);
 	TArray<FDataEntry>		Entries;
 	FStringView				RootPath;
-	MarkerType				Marker = 0;
+	uint32					Marker = 0;
 	TUniquePtr<IFileHandle>	JrnHandle;
 	uint32					Cursor = 0;
 	uint32					MaxSize;
+	static constexpr uint32	BytesToHash = 64;
 };
 
 ////////////////////////////////////////////////////////////////////////////////
@@ -446,6 +447,12 @@ void FDiskJournal::OpenJrnFile()
 	JrnHandle.Reset(Handle);
 }
 
+////////////////////////////////////////////////////////////////////////////////
+uint32 FDiskJournal::HashBytes(const uint8* Data, uint32 Size, uint32 Seed)
+{
+	return Seed + CityHash32((const char*)Data, Size);
+}
+
 ////////////////////////////////////////////////////////////////////////////////
 void FDiskJournal::GetPath(TStringBuilder<64>& Out)
 {
@@ -459,14 +466,11 @@ FDiskPhrase FDiskJournal::OpenPhrase(uint32 DataSize)
 	check((Cursor & (sizeof(FDataEntry) - 1)) == 0);
 
 	Entries.Add(FDataEntry{});
-	int32 MaxEntries = int32((MaxSize - Cursor) / sizeof(FDataEntry)) - 1;
-
-	FDiskPhrase Ret(Entries, MaxEntries, DataSize);
+	const int32 MaxEntries = int32((MaxSize - Cursor) / sizeof(FDataEntry)) - 1;
 
-	uint8* PhraseData = Ret.GetPhraseData();
-	std::memcpy(PhraseData, &Marker, sizeof(MarkerType));
+	FDiskPhrase Ret(Entries, FMath::Min(MaxEntries, int32(UINT16_MAX)), DataSize);
 
-	return MoveTemp(Ret);
+	return Ret;
 }
 
 ////////////////////////////////////////////////////////////////////////////////
@@ -479,15 +483,26 @@ void FDiskJournal::ClosePhrase(FDiskPhrase&& Phrase, uint64 DataCursor)
 		return;
 	}
 
-	Entries.Last().EntryCount = uint16(EntryCount);
+	// Write the number of data entries in the first and last entry
+	check(EntryCount <= UINT16_MAX);
+	FDataEntry& FirstEntry = (FDataEntry&)(Phrase.GetEntries()[0]);
+	FDataEntry& LastEntry = (FDataEntry&) Entries.Last();
+	FirstEntry.EntryCount = uint16(EntryCount);
+	LastEntry.EntryCount = uint16(EntryCount);
 
 	auto& Desc = (FPhraseDesc&)(Phrase.GetEntries()[-1]);
 	Desc.Magic = MAGIC;
 	Desc.Marker = Marker;
-	Desc.EntryCount = uint16(EntryCount);
 	Desc.DataCursor = DataCursor;
+	
+	const uint32 HashSize = FMath::Min(BytesToHash, Phrase.GetDataSize());
+	Desc.Hash = HashBytes(Phrase.GetPhraseData(), HashSize, Desc.Marker);
 
-	++Marker;
+	// Increment and wrap marker
+	if (++Marker > MARKER_MAX)
+	{
+		Marker = 0;
+	}
 }
 
 ////////////////////////////////////////////////////////////////////////////////
@@ -628,7 +643,7 @@ void FDiskCache::ClosePhrase(FDiskPhrase&& Phrase)
 		return;
 	}
 
-	uint32 WriteSize = Phrase.GetPhraseSize();
+	uint32 WriteSize = Phrase.GetDataSize();
 	if (DataCursor + WriteSize > MaxDataSize)
 	{
 		OverRemoval = 0;
@@ -704,7 +719,7 @@ uint64 FDiskCache::Insert(uint64 DataBase, const FDataEntry& Entry)
 	FMapEntry Value;
 	Value.DataCursor = DataBase + Entry.Offset;
 	Value.Size = Entry.Size;
-	Value.First = (Entry.Offset == sizeof(MarkerType));
+	Value.First = (Entry.Offset == 0);
 	DataMap.Add(Entry.Key, Value);
 	return Entry.Size;
 }
@@ -866,7 +881,7 @@ static bool LoadCache(FDiskCache& DiskCache)
 	struct FParagraph
 	{
 		const FPhraseDesc*	Phrase;
-		MarkerType			Marker;
+		uint32				EntryCount;
 		uint32				DataSize;
 	};
 
@@ -889,23 +904,24 @@ static bool LoadCache(FDiskCache& DiskCache)
 			return nullptr;
 		}
 
-		uint32 ToConsume = sizeof(FDataEntry) * (Header->EntryCount + 1);
-		if (IsOob(Cursor + ToConsume))
+		const auto* FirstEntry = Header->Entries;
+		const uint32 BytesToConsume = sizeof(FDataEntry) * (FirstEntry->EntryCount + 1); //Entries + PhraseDesc
+		if (IsOob(Cursor + BytesToConsume))
 		{
 			return nullptr;
 		}
 
-		Cursor += ToConsume;
+		Cursor += BytesToConsume;
 
 		const auto* LastEntry = (FDataEntry*)Cursor - 1;
-		if (LastEntry->EntryCount != Header->EntryCount)
+		if (LastEntry->EntryCount != FirstEntry->EntryCount)
 		{
 			return nullptr;
 		}
 
 		Out.Phrase = Header;
-		Out.Marker = Header->Marker;
-		Out.DataSize = uint32(LastEntry->Offset + LastEntry->Size) + sizeof(MarkerType);
+		Out.EntryCount = uint32(LastEntry->EntryCount);
+		Out.DataSize = uint32(LastEntry->Offset + LastEntry->Size);
 
 		return Cursor;
 	};
@@ -959,8 +975,8 @@ static bool LoadCache(FDiskCache& DiskCache)
 		const FParagraph& Lhs,
 		const FParagraph& Rhs)
 	{
-        MarkerType L = Lhs.Marker;
-        MarkerType R = Rhs.Marker;
+        const auto L = Lhs.Phrase->Marker;
+        const auto R = Rhs.Phrase->Marker;
         enum : uint32 { LowQuarter = 1u << 30, HighQuarter = 3u << 30 };
         int32 Wrap = (L < LowQuarter) & (R >= HighQuarter);
         Wrap |= (R < LowQuarter) & (L >= HighQuarter);
@@ -983,14 +999,12 @@ static bool LoadCache(FDiskCache& DiskCache)
 		}
 
 		const FParagraph& Older = Paragraphs[i];
-		if (Newer.Marker != Older.Marker + 1)
+		if (Newer.Phrase->Marker != Older.Phrase->Marker + 1)
 		{
 			break;
 		}
 	}
 
-	IPlatformFile& Ipf = IPlatformFile::GetPlatformPhysical();
-
 	if (!DiskCache.DataHandle.IsValid())
 	{
 		UE_LOG(LogIas, VeryVerbose, TEXT("JournaledCache: unable to open '%s'"), *DiskCache.BinPath);
@@ -1009,15 +1023,24 @@ static bool LoadCache(FDiskCache& DiskCache)
 	}
 
 	// Detect data writes that are newer than any journal flushes.
-	auto ReadBinMarker = [File] (uint64 Cursor, uint32& Out)
+	auto ReadHash = [File] (uint64 Cursor, uint32& OutHash, uint32 MaxHashSize, uint32 Seed)
 	{
-		if (Cursor + sizeof(Out) > uint64(File->Size()))
+		if (Cursor + sizeof(OutHash) > uint64(File->Size()))
 		{
 			return false;
 		}
 
 		File->Seek(Cursor);
-		return File->Read((uint8*)(&Out), sizeof(Out));
+		
+		const uint32 HashSize = FMath::Min(FDiskJournal::BytesToHash, MaxHashSize);
+		TArray<uint8, TInlineAllocator<FDiskJournal::BytesToHash>> BytesToHash;
+		BytesToHash.AddZeroed(HashSize);
+		if (const bool Result = File->Read(BytesToHash.GetData(), HashSize))
+		{
+			OutHash = FDiskJournal::HashBytes(BytesToHash.GetData(), BytesToHash.Num(), Seed);
+			return true;
+		}
+		return false;
 	};
 
 	for (; BasisIndex < Paragraphs.Num(); ++BasisIndex)
@@ -1030,8 +1053,9 @@ static bool LoadCache(FDiskCache& DiskCache)
 			DataBase = 0;
 		}
 
-		uint32 DataMark;
-		if (ReadBinMarker(DataBase, DataMark) && (DataMark == Stock.Marker))
+		uint32 Hash;
+		// There could be a phrase with only one very short data entry. Make sure we don't hash too much.
+		if (ReadHash(DataBase, Hash, Stock.DataSize, uint32(Stock.Phrase->Marker)) && (Hash == Stock.Phrase->Hash))
 		{
 			break;
 		}
@@ -1041,17 +1065,17 @@ static bool LoadCache(FDiskCache& DiskCache)
 	for (uint32 i = BasisIndex, n = Paragraphs.Num(); i < n; ++i)
 	{
 		const FPhraseDesc& Holm = Paragraphs[i].Phrase[0];
-		DiskCache.Insert(Holm.DataCursor, Holm.Entries, Holm.EntryCount);
+		DiskCache.Insert(Holm.DataCursor, Holm.Entries, Paragraphs[i].EntryCount);
 	}
 	
 	// Prime the journal's state
 	const FParagraph& LastPara = Paragraphs.Last();
 	const FPhraseDesc* LastPhrase = LastPara.Phrase;
-	Journal.Marker = LastPara.Marker + 1;
+	Journal.Marker = uint32(LastPara.Phrase->Marker) + 1;
 
 	if (DataSize <= Journal.MaxSize)
 	{
-		Journal.Cursor = uint32(UPTRINT(LastPhrase + LastPhrase->EntryCount + 1) - UPTRINT(Data.Get()));
+		Journal.Cursor = uint32(UPTRINT(LastPhrase + LastPara.EntryCount + 1) - UPTRINT(Data.Get()));
 	}
 	else
 	{
@@ -1508,7 +1532,7 @@ private:
 	TArray<FWork>				ActiveReads;
 	TArray<FWork>				PendingWork;
 	std::atomic_int				PendingCount = 0;
-	int32						PendingPrev;
+	int32						PendingPrev = -1;
 	uint32						ReadIdCounter = 0;
 };
 
@@ -1927,7 +1951,6 @@ private:
 		EIoErrorCode*			Status;
 	};
 
-	void						Update();
 	virtual void				OnRead(const FReadResult* Results, int32 Num) override;
 	static uint64				ReduceKey(const FIoHash& Key);
 	JournaledCache::FCache*		Cache = nullptr;
