remove the old method of using a single physics material for all shapes in the cluster union and instead rely on the transfer that exists in TransferClusterUnionShapeData

diff --git a/Engine/Source/Runtime/Experimental/Chaos/Private/Chaos/ClusterUnionManager.cpp b/Engine/Source/Runtime/Experimental/Chaos/Private/Chaos/ClusterUnionManager.cpp
--- a/Engine/Source/Runtime/Experimental/Chaos/Private/Chaos/ClusterUnionManager.cpp
+++ b/Engine/Source/Runtime/Experimental/Chaos/Private/Chaos/ClusterUnionManager.cpp
@@ -618,7 +618,6 @@ namespace Chaos
 		// For all the particles that have been added to the cluster we need to set their parent proxy to the
 		// cluster's proxy if it exists. We need the proxy type check because for non-cluster union proxy backed unions,
 		// the cluster union's proxy will be the proxy of the most recently added particle.
-		FMaterialHandle MatHandle;
 		if (OldProxy && OldProxyType == EPhysicsProxyType::ClusterUnionProxy)
 		{
 			for (FPBDRigidParticleHandle* Particle : FinalParticlesToAdd)
@@ -626,20 +625,6 @@ namespace Chaos
 				if (Particle && Particle->PhysicsProxy())
 				{
 					Particle->PhysicsProxy()->SetParentProxy(OldProxy);
-
-					if(!MatHandle.InnerHandle.IsValid())
-					{
-						const FShapesArray& Shapes = Particle->ShapesArray();
-						for(const TUniquePtr<FPerShapeData>& ShapeData : Shapes)
-						{
-							const int32 NumShapeMaterials = ShapeData->NumMaterials();
-							if(NumShapeMaterials > 0)
-							{
-								MatHandle = ShapeData->GetMaterial(0);
-								break;
-							}
-						}
-					}
 				}
 			}
 		}
@@ -666,14 +651,6 @@ namespace Chaos
 				MEvolution.SetParticleObjectState(Cluster->InternalCluster, Chaos::EObjectStateType::Sleeping);
 			}
 
-			if(MatHandle.InnerHandle.IsValid())
-			{
-				for(const TUniquePtr<FPerShapeData>& ShapeData : Cluster->InternalCluster->ShapesArray())
-				{
-					ShapeData->SetMaterial(MatHandle);
-				}
-			}
-
 			if (!Cluster->IsGravityOverrideSet() && !FinalParticlesToAdd.IsEmpty())
 			{
 				Cluster->InternalCluster->SetGravityGroupIndex(FinalParticlesToAdd[0]->GravityGroupIndex());
