Chaos Clustering : make sure that active particles being pro-actively disabled are properly dirtied so that the GC proxy gets a update

diff --git a/Engine/Source/Runtime/Experimental/Chaos/Private/Chaos/PBDRigidClustering.cpp b/Engine/Source/Runtime/Experimental/Chaos/Private/Chaos/PBDRigidClustering.cpp
--- a/Engine/Source/Runtime/Experimental/Chaos/Private/Chaos/PBDRigidClustering.cpp
+++ b/Engine/Source/Runtime/Experimental/Chaos/Private/Chaos/PBDRigidClustering.cpp
@@ -1229,6 +1229,7 @@ namespace Chaos
 					{
 						DisabledParticleCount++;
 						MEvolution.DisableParticle(Child);
+						MEvolution.GetParticles().MarkTransientDirtyParticle(Child);
 						ChildIt.RemoveCurrent();
 					}
 					if (DisabledParticleCount >= NumberOfParticlesToDisable)
