Updated some IAS instrumentation

#rnx

diff --git a/Engine/Source/Runtime/Experimental/IoStoreOnDemand/Private/CoreHttp/Http.cpp b/Engine/Source/Runtime/Experimental/IoStoreOnDemand/Private/CoreHttp/Http.cpp
--- a/Engine/Source/Runtime/Experimental/IoStoreOnDemand/Private/CoreHttp/Http.cpp
+++ b/Engine/Source/Runtime/Experimental/IoStoreOnDemand/Private/CoreHttp/Http.cpp
@@ -92,6 +92,8 @@ namespace UE::IO::IAS::HTTP
 	x(LoopDestroy) \
 	x(ActivityCreate) \
 	x(ActivityDestroy) \
+	x(SocketCreate) \
+	x(SocketDestroy) \
 	x(RequestBegin) \
 	x(StateChange) \
 	x(Wait) \
@@ -147,7 +149,7 @@ static void TraceEnum(const FAnsiStringView* Names)
 }
 
 ////////////////////////////////////////////////////////////////////////////////
-static void TraceInner(const void* Id, ETrace Action, UPTRINT Param)
+static void TraceInner(const UPTRINT Id, ETrace Action, UPTRINT Param)
 {
 #if UE_TRACE_ENABLED
 	static bool Once = [] {
@@ -158,16 +160,16 @@ static void TraceInner(const void* Id, ETrace Action, UPTRINT Param)
 
 	UE_TRACE_LOG(IasHttp, Event, IasHttpChannel)
 		<< Event.Cycle(FPlatformTime::Cycles64())
-		<< Event.Id(uint32(UPTRINT(Id)))
-		<< Event.Param(uint32(UPTRINT(Param)))
+		<< Event.Id(uint32(Id))
+		<< Event.Param(uint32(Param))
 		<< Event.Action(uint8(Action));
 #endif // UE_TRACE_ENABLED
 }
 
-template <typename T=UPTRINT>
-static void Trace(const void* Id, ETrace Action, T Param=0)
+template <typename T, typename U=UPTRINT>
+static void Trace(T Id, ETrace Action, U Param=0)
 {
-	TraceInner(Id, Action, UPTRINT(Param));
+	TraceInner(UPTRINT(Id), Action, UPTRINT(Param));
 }
 
 
@@ -715,7 +717,7 @@ public:
 	};
 
 				FSocket() = default;
-				~FSocket()					{ if (IsValid()) Destroy(); }
+				~FSocket()					{ Destroy(); }
 				FSocket(FSocket&& Rhs)		{ Move(MoveTemp(Rhs)); }
 	FSocket&	operator = (FSocket&& Rhs)	{ Move(MoveTemp(Rhs)); return *this; }
 	bool		IsValid() const				{ return Socket != InvalidSocket; }
@@ -759,7 +761,14 @@ bool FSocket::Create()
 {
 	check(!IsValid());
 	Socket = socket(AF_INET, SOCK_STREAM, IPPROTO_TCP);
-	return IsValid();
+
+	if (!IsValid())
+	{
+		return false;
+	}
+
+	Trace(Socket, ETrace::SocketCreate);
+	return true;
 }
 
 ////////////////////////////////////////////////////////////////////////////////
@@ -770,6 +779,8 @@ void FSocket::Destroy()
 		return;
 	}
 
+	Trace(Socket, ETrace::SocketDestroy);
+
 	closesocket(Socket);
 	Socket = InvalidSocket;
 }
@@ -779,6 +790,8 @@ bool FSocket::Connect(uint32 IpAddress, uint32 Port)
 {
 	check(IsValid());
 
+	Trace(Socket, ETrace::Connect, IpAddress);
+
 	IpAddress = htonl(IpAddress);
 
 	sockaddr_in AddrInet = { sizeof(sockaddr_in) };
@@ -806,7 +819,9 @@ void FSocket::Disconnect()
 ////////////////////////////////////////////////////////////////////////////////
 int32 FSocket::Send(const char* Data, uint32 Size)
 {
+	Trace(Socket, ETrace::Send, -1);
 	int32 Result = send(Socket, Data, Size, MsgFlagType(0));
+	Trace(Socket, ETrace::Send, FMath::Max(Result, 0));
 
 	if (Result > 0)
 	{
@@ -842,7 +857,9 @@ int32 FSocket::Send(const char* Data, uint32 Size)
 ////////////////////////////////////////////////////////////////////////////////
 int32 FSocket::Recv(char* Dest, uint32 Size)
 {
+	Trace(Socket, ETrace::Recv, -1);
 	int32 Result = recv(Socket, Dest, Size, MsgFlagType(0));
+	Trace(Socket, ETrace::Recv, FMath::Max(0, Result));
 
 	if (Result > 0)
 	{
@@ -1983,6 +2000,8 @@ FTicketPerf::FSample FTicketPerf::GetRecvSample() const
 ////////////////////////////////////////////////////////////////////////////////
 static int32 DoSend(FActivity* Activity, FSocket& Socket)
 {
+	Trace(Activity, ETrace::StateChange, Activity->State);
+
 	TRACE_CPUPROFILER_EVENT_SCOPE(IasHttp::DoSend);
 
 	FBuffer& Buffer = Activity->Buffer;
@@ -1994,9 +2013,7 @@ static int32 DoSend(FActivity* Activity, FSocket& Socket)
 	SendSize -= AlreadySent;
 	check(SendSize > 0);
 
-	Trace(Activity, ETrace::Send, SendSize);
 	int32 Result = Socket.Send(SendData, SendSize);
-	Trace(Activity, ETrace::Send, -1);
 
 	switch (FSocket::EResult(Result))
 	{
@@ -2033,11 +2050,11 @@ static int32 DoRecvMessage(FActivity* Activity, FSocket& Socket)
 	const char* MessageRight;
 	while (true)
 	{
+		Trace(Activity, ETrace::StateChange, Activity->State);
+
 		auto [Dest, DestSize] = Buffer.GetMutableFree(0, PageSize);
 
-		Trace(Activity, ETrace::Recv, -1);
 		int32 Result = Socket.Recv(Dest, DestSize);
-		Trace(Activity, ETrace::Recv, FMath::Max(Result, 0));
 
 		if (Result == int32(FSocket::EResult::Wait))
 		{
@@ -2246,11 +2263,11 @@ static int32 DoRecvContent(FActivity* Activity, FSocket& Socket, int32& MaxRecvS
 			return 1;
 		}
 
+		Trace(Activity, ETrace::StateChange, Activity->State);
+
 		char* Cursor = (char*)(DestView.GetData()) + Activity->StateParam;
 
-		Trace(Activity, ETrace::Recv, -1);
 		int32 Result = Socket.Recv(Cursor, Size);
-		Trace(Activity, ETrace::Recv, FMath::Max(Result, 0));
 
 		if (Result == int32(FSocket::EResult::Wait))
 		{
@@ -2823,6 +2840,9 @@ int32 FHostGroup::Wait(int32 PollTimeoutMs)
 		return 0;
 	}
 
+	Trace(0, ETrace::Wait);
+	ON_SCOPE_EXIT { Trace(0, ETrace::Unwait); };
+
 	// If the poll timeout is negative then treat that as a fatal timeout
 	bool bFailOnTimeout = false;
 	if (PollTimeoutMs < -1)
@@ -2857,7 +2877,7 @@ int32 FHostGroup::Wait(int32 PollTimeoutMs)
 		check(Group != nullptr);
 		(*Group)->Unwait();
 
-		Waiters.RemoveAtSwap(i);
+		Waiters.RemoveAtSwap(i, 1, false);
 		--n, --i, ++Count;
 	}
 	check(Count == Result);
@@ -2906,7 +2926,7 @@ void FHostGroup::Tick(FTickState& State)
 		}
 
 		check(!Group->IsWaiting());
-		SocketGroups.RemoveAtSwap(i);
+		SocketGroups.RemoveAtSwap(i, 1, false);
 		--n, --i;
 	}
 }
