Don't cancel activities that have already started their send

#rnx

diff --git a/Engine/Source/Runtime/Experimental/IoStoreOnDemand/Private/CoreHttp/Http.cpp b/Engine/Source/Runtime/Experimental/IoStoreOnDemand/Private/CoreHttp/Http.cpp
--- a/Engine/Source/Runtime/Experimental/IoStoreOnDemand/Private/CoreHttp/Http.cpp
+++ b/Engine/Source/Runtime/Experimental/IoStoreOnDemand/Private/CoreHttp/Http.cpp
@@ -2547,7 +2547,7 @@ void FSocketGroup::Done(FTickState& State, FActivity* Activity)
 	Activity->Next = State.DoneList;
 	State.DoneList = Activity;
 
-	ActiveSlots ^= 1ull << Activity->Slot;
+	ActiveSlots &= ~(1ull << Activity->Slot);
 }
 
 ////////////////////////////////////////////////////////////////////////////////
@@ -2721,27 +2721,29 @@ void FSocketGroup::TickCancels(FTickState& State)
 	check(List != nullptr);
 	FActivity* Activity = List;
 	List = ListTail = nullptr;
+	ActiveSlots = 0;
 
+	for (FActivity* Next; Activity != nullptr; Activity = Next)
 	{
-		FActivity* Next = Activity->Next;
+		Next = Activity->Next;
+
+		// Is this one of the activities we'd like to cancel?
+		bool bCancel = false;
 		if (uint64 Slot = 1ull << Activity->Slot; State.Cancels & Slot)
 		{
-			if (Activity->State > FActivity::EState::Send)
-			{
-				ActiveSlots ^= Slot;
-				--NumActive;
-			}
-			DoCancel(Activity);
-			Done(State, Activity);
+			// We can't cancel things that have already been (or are being) requested
+			bCancel = (Activity->State == FActivity::EState::Send) & (Activity->StateParam == 0);
 		}
-		else
+
+		if (!bCancel)
 		{
 			AddActivity(Activity);
+			continue;
 		}
 
-		Activity = Next;
+		DoCancel(Activity);
+		Done(State, Activity);
 	}
-	while (Activity != nullptr);
 }
 
 ////////////////////////////////////////////////////////////////////////////////
