Correctly clear everything out when a timeout fails

#rnx

diff --git a/Engine/Source/Runtime/Experimental/IoStoreOnDemand/Private/CoreHttp/Http.cpp b/Engine/Source/Runtime/Experimental/IoStoreOnDemand/Private/CoreHttp/Http.cpp
--- a/Engine/Source/Runtime/Experimental/IoStoreOnDemand/Private/CoreHttp/Http.cpp
+++ b/Engine/Source/Runtime/Experimental/IoStoreOnDemand/Private/CoreHttp/Http.cpp
@@ -2555,9 +2555,8 @@ void FSocketGroup::Fail(FTickState& State, const char* Reason)
 {
 	// Failure is quite terminal and we need to abort all listed activities.
 
-	check(List != nullptr);
 	FActivity* Activity = List;
-	do
+	while (Activity != nullptr)
 	{
 		FActivity* Next = Activity->Next;
 		if (Activity->State != FActivity::EState::Failed)
@@ -2570,12 +2569,13 @@ void FSocketGroup::Fail(FTickState& State, const char* Reason)
 
 		Activity = Next;
 	}
-	while (Activity != nullptr);
 	check(ActiveSlots == 0);
 
 	Socket = FSocket();
 	NumActive = 0;
 	List = ListTail = nullptr;
+	bWaiting = false;
+	IsKeepAlive = 0;
 }
 
 ////////////////////////////////////////////////////////////////////////////////
@@ -2903,6 +2903,10 @@ void FHostGroup::Tick(FTickState& State)
 		{
 			Group->Fail(State, Reason);
 		}
+
+		Waiters.Empty();
+		SocketGroups.SetNum(0, false);
+		return;
 	}
 
 	// If we are still waiting on all sockets groups there is nothing to do.
