Remove check

diff --git a/Engine/Source/Runtime/Experimental/IoStoreOnDemand/Private/JournaledCache.cpp b/Engine/Source/Runtime/Experimental/IoStoreOnDemand/Private/JournaledCache.cpp
--- a/Engine/Source/Runtime/Experimental/IoStoreOnDemand/Private/JournaledCache.cpp
+++ b/Engine/Source/Runtime/Experimental/IoStoreOnDemand/Private/JournaledCache.cpp
@@ -776,7 +776,6 @@ void FDiskCache::Prune(uint64 DataBase, uint32 Size)
 		BytesRemoved += int32(B[1] - B[0]);
 		if (BytesRemoved - Overage >= int32(Size))
 		{
-			check(BytesRemoved - Overage == int32(Size));
 			OverRemoval = int32(Overage);
 			break;
 		}
