Log resolved address for backends

diff --git a/Engine/Source/Runtime/Experimental/IoStoreOnDemand/Private/CoreHttp/Client.h b/Engine/Source/Runtime/Experimental/IoStoreOnDemand/Private/CoreHttp/Client.h
--- a/Engine/Source/Runtime/Experimental/IoStoreOnDemand/Private/CoreHttp/Client.h
+++ b/Engine/Source/Runtime/Experimental/IoStoreOnDemand/Private/CoreHttp/Client.h
@@ -77,6 +77,7 @@ public:
 							FConnectionPool(FConnectionPool&& Rhs)	{ *this = MoveTemp(Rhs); }
 	FConnectionPool&		operator = (FConnectionPool&& Rhs)		{ Swap(Ptr, Rhs.Ptr); return *this; }
 	bool					Resolve();
+	void					Describe(FAnsiStringBuilderBase&) const;
 
 	static bool				IsValidHostUrl(FAnsiStringView Url);
 
diff --git a/Engine/Source/Runtime/Experimental/IoStoreOnDemand/Private/CoreHttp/Http.cpp b/Engine/Source/Runtime/Experimental/IoStoreOnDemand/Private/CoreHttp/Http.cpp
--- a/Engine/Source/Runtime/Experimental/IoStoreOnDemand/Private/CoreHttp/Http.cpp
+++ b/Engine/Source/Runtime/Experimental/IoStoreOnDemand/Private/CoreHttp/Http.cpp
@@ -1485,6 +1485,27 @@ bool FConnectionPool::Resolve()
 	return (Ptr->ResolveHostName().GetValue() > 0);
 }
 
+////////////////////////////////////////////////////////////////////////////////
+void FConnectionPool::Describe(FAnsiStringBuilderBase& OutString) const
+{
+	const FAnsiStringView HostName = Ptr->GetHostName();
+	OutString.Appendf("%.*s", HostName.Len(), HostName.GetData());
+	if (!!Ptr->IsResolved())
+	{
+		const auto IpAddress = Ptr->GetIpAddress();
+		OutString.Appendf(" (%u.%u.%u.%u)",
+						(IpAddress >> 24) & 0xff,
+						(IpAddress >> 16) & 0xff,
+						(IpAddress >> 8) & 0xff,
+						IpAddress & 0xff
+		);
+	}
+	else
+	{
+		OutString.Append(" (unresolved)");
+	}
+}
+
 ////////////////////////////////////////////////////////////////////////////////
 bool FConnectionPool::IsValidHostUrl(FAnsiStringView Url)
 {
diff --git a/Engine/Source/Runtime/Experimental/IoStoreOnDemand/Private/CoreHttp/LatencyTesting.cpp b/Engine/Source/Runtime/Experimental/IoStoreOnDemand/Private/CoreHttp/LatencyTesting.cpp
--- a/Engine/Source/Runtime/Experimental/IoStoreOnDemand/Private/CoreHttp/LatencyTesting.cpp
+++ b/Engine/Source/Runtime/Experimental/IoStoreOnDemand/Private/CoreHttp/LatencyTesting.cpp
@@ -6,6 +6,7 @@
 #include "Containers/StringConv.h"
 #include "Containers/StringView.h"
 #include "HAL/PlatformTime.h"
+#include "IO/IoStoreOnDemand.h"
 #include "Misc/StringBuilder.h"
 #include "Templates/Function.h"
 
@@ -59,6 +60,10 @@ void LatencyTest(FStringView InUrl, FStringView InPath, uint32 InTimeOutMs, TArr
 
 		Result = int32((int64(Result) * 1000) / Freq);
 	}
+	
+	TAnsiStringBuilder<512> ConnectionDesc;
+	Pool.Describe(ConnectionDesc);
+	UE_LOG(LogIas, Log, TEXT("Testing endpoint %s"), ANSI_TO_TCHAR(ConnectionDesc.ToString()));
 }
 
 } // namespace UE::IO::IAS::HTTP
