Ias: Restore missing GHttpRetryCount counter

#rnx

diff --git a/Engine/Source/Runtime/Experimental/IoStoreOnDemand/Private/OnDemandHttpClient.cpp b/Engine/Source/Runtime/Experimental/IoStoreOnDemand/Private/OnDemandHttpClient.cpp
--- a/Engine/Source/Runtime/Experimental/IoStoreOnDemand/Private/OnDemandHttpClient.cpp
+++ b/Engine/Source/Runtime/Experimental/IoStoreOnDemand/Private/OnDemandHttpClient.cpp
@@ -1,6 +1,7 @@
 // Copyright Epic Games, Inc. All Rights Reserved.
 
 #include "OnDemandHttpClient.h"
+#include "Statistics.h"
 
 #include "Containers/StringConv.h"
 #include "HAL/IConsoleManager.h"
@@ -211,6 +212,10 @@ void FHttpClient::IssueRequest(FRequestParams&& Params)
 void FHttpClient::RetryRequest(FRequestParams&& Params, bool bNextEndpoint)
 {
 	check(Params.Attempt < Config.MaxRetryCount);
+	if (Params.Attempt == 0)
+	{
+		FOnDemandIoBackendStats::Get()->OnHttpRetry();
+	}
 	if (bNextEndpoint && Config.Endpoints.Num() > 1)
 	{
 		Params.Endpoint = (Params.Endpoint + 1) % Config.Endpoints.Num();
diff --git a/Engine/Source/Runtime/Experimental/IoStoreOnDemand/Private/Statistics.cpp b/Engine/Source/Runtime/Experimental/IoStoreOnDemand/Private/Statistics.cpp
--- a/Engine/Source/Runtime/Experimental/IoStoreOnDemand/Private/Statistics.cpp
+++ b/Engine/Source/Runtime/Experimental/IoStoreOnDemand/Private/Statistics.cpp
@@ -501,7 +501,6 @@ void FOnDemandIoBackendStats::OnHttpGet(uint64 SizeBytes, uint64 DurationMs)
 
 void FOnDemandIoBackendStats::OnHttpRetry()
 {
-	GHttpInflightCount.Add(-1);
 	GHttpRetryCount.Add(1);
 }
 
