Ias: Split out GIasHttpConcurrentRequests and GIasHttpDistributionRetryCount as separate cvars.
Rename MaxHttp* cvars to Http* for consistency.

#rnx

diff --git a/Engine/Source/Runtime/Experimental/IoStoreOnDemand/Private/DistributionEndpoints.cpp b/Engine/Source/Runtime/Experimental/IoStoreOnDemand/Private/DistributionEndpoints.cpp
--- a/Engine/Source/Runtime/Experimental/IoStoreOnDemand/Private/DistributionEndpoints.cpp
+++ b/Engine/Source/Runtime/Experimental/IoStoreOnDemand/Private/DistributionEndpoints.cpp
@@ -3,6 +3,7 @@
 #include "DistributionEndpoints.h"
 
 #include "Dom/JsonValue.h"
+#include "HAL/IConsoleManager.h"
 #include "HAL/PlatformTime.h"
 #include "HttpManager.h"
 #include "HttpModule.h"
@@ -16,7 +17,13 @@
 namespace UE::IO::IAS
 {
 
-extern int32 GIasMaxHttpRetryCount;
+static int32 GIasHttpDistributionRetryCount = 2;
+static FAutoConsoleVariableRef CVar_IasHttpDistributionRetryCount(
+	TEXT("ias.HttpDistributionRetryCount"),
+	GIasHttpDistributionRetryCount,
+	TEXT("Number of HTTP distribution request retries.")
+);
+
 
 FDistributionEndpoints::~FDistributionEndpoints()
 {
@@ -106,7 +113,7 @@ void FDistributionEndpoints::IssueEndpointRequests()
 	TRACE_CPUPROFILER_EVENT_SCOPE(FOnDemandIoBackend::IssueEndpointRequests);
 	// Currently we need to use the HTTP module in order to resolve service endpoints due to HTTPS
 	FHttpModule& HttpModule = FHttpModule::Get();
-	const int32 MaxAttempts = GIasMaxHttpRetryCount;
+	const int32 MaxAttempts = GIasHttpDistributionRetryCount;
 
 	TArray<FHttpRequestPtr, TInlineAllocator<2>> HttpRequests;
 	{
diff --git a/Engine/Source/Runtime/Experimental/IoStoreOnDemand/Private/OnDemandIoDispatcherBackend.cpp b/Engine/Source/Runtime/Experimental/IoStoreOnDemand/Private/OnDemandIoDispatcherBackend.cpp
--- a/Engine/Source/Runtime/Experimental/IoStoreOnDemand/Private/OnDemandIoDispatcherBackend.cpp
+++ b/Engine/Source/Runtime/Experimental/IoStoreOnDemand/Private/OnDemandIoDispatcherBackend.cpp
@@ -99,11 +99,18 @@ static FAutoConsoleVariableRef CVar_GIasHttpRecvBufKiB(
 	TEXT("Recv buffer size")
 );
 
-int32 GIasMaxHttpConnectionCount = 4;
-static FAutoConsoleVariableRef CVar_IasMaxHttpConnectionCount(
-	TEXT("ias.MaxHttpConnectionCount"),
-	GIasMaxHttpConnectionCount,
-	TEXT("Max number of open HTTP connections to the on demand endpoint(s).")
+static int32 GIasHttpConcurrentRequests = 8;
+static FAutoConsoleVariableRef CVar_IasHttpConcurrentRequests(
+	TEXT("ias.HttpConcurrentRequests"),
+	GIasHttpConcurrentRequests,
+	TEXT("Number of concurrent requests in the http client.")
+);
+
+static int32 GIasHttpConnectionCount = 4;
+static FAutoConsoleVariableRef CVar_IasHttpConnectionCount(
+	TEXT("ias.HttpConnectionCount"),
+	GIasHttpConnectionCount,
+	TEXT("Number of open HTTP connections to the on demand endpoint(s).")
 );
 
 static int32 GIasHttpPipelineLength = 2;
@@ -113,11 +120,11 @@ static FAutoConsoleVariableRef CVar_GIasHttpPipelineLength(
 	TEXT("Number of concurrent requests on one connection")
 );
 
-int32 GIasMaxHttpRetryCount = 2;
-static FAutoConsoleVariableRef CVar_IasMaxHttpRetryCount(
-	TEXT("ias.MaxHttpRetryCount"),
-	GIasMaxHttpRetryCount,
-	TEXT("Max number of HTTP request retries before failing the I/O request.")
+static int32 GIasHttpRetryCount = 2;
+static FAutoConsoleVariableRef CVar_IasHttpRetryCount(
+	TEXT("ias.HttpRetryCount"),
+	GIasHttpRetryCount,
+	TEXT("Number of HTTP request retries before failing the I/O request.")
 );
 
 int32 GIasHttpTimeOutMs = 10 * 1000;
@@ -2274,9 +2281,9 @@ uint32 FOnDemandIoBackend::Run()
 	{
 		.Endpoints = AvailableEps.Urls,
 		.PrimaryEndpoint = FMath::Min(GIasHttpPrimaryEndpoint, AvailableEps.Urls.Num() -1),
-		.MaxConnectionCount = GIasMaxHttpConnectionCount,
+		.MaxConnectionCount = GIasHttpConnectionCount,
 		.PipelineLength = GIasHttpPipelineLength,
-		.MaxRetryCount = FMath::Max(AvailableEps.Urls.Num() + 1, GIasMaxHttpRetryCount),
+		.MaxRetryCount = FMath::Max(AvailableEps.Urls.Num() + 1, GIasHttpRetryCount),
 		.ReceiveBufferSize = GIasHttpRecvBufKiB >= 0 ? GIasHttpRecvBufKiB << 10 : -1,
 		.FailTimeoutMs = GIasHttpFailTimeOutMs,
 		.bChangeEndpointAfterSuccessfulRetry = GIasHttpChangeEndpointAfterSuccessfulRetry,
@@ -2293,7 +2300,7 @@ uint32 FOnDemandIoBackend::Run()
 	while (!bStopRequested)
 	{
 		// Process HTTP request(s) even if the client is invalid to ensure enqueued request(s) gets completed.
-		ProcessHttpRequests(*HttpClient, HttpErrors, FMath::Min(GIasHttpPipelineLength * GIasMaxHttpConnectionCount, 64));
+		ProcessHttpRequests(*HttpClient, HttpErrors, FMath::Min(GIasHttpConcurrentRequests, 32));
 		AvailableEps.Current = HttpClient->GetEndpoint();
 
 		if (!bStopRequested)
