Fix a bug preventing the .iochunktoc file from being loaded properly on android

#rnx

- In our current setup the .iochunktoc file is an android asset meaning that the handle is relative to the start of the entire package, not the file itself. We seem to have a bug in FFileHandleAndroid where it has the wrong initial offset when created with an asset path rather than a local path, which is only fixed on the first seek. For now as a quick and safe fix we will just do a dummy Seek(0) when loading the file on android with a proper fix to be done later.
- The code was not handling the possibility of CreateReaderFromPlatformPackage returning a nullptr so added that while I was on.

diff --git a/Engine/Source/Runtime/Experimental/IoStoreOnDemand/Private/OnDemandIoDispatcherBackend.cpp b/Engine/Source/Runtime/Experimental/IoStoreOnDemand/Private/OnDemandIoDispatcherBackend.cpp
--- a/Engine/Source/Runtime/Experimental/IoStoreOnDemand/Private/OnDemandIoDispatcherBackend.cpp
+++ b/Engine/Source/Runtime/Experimental/IoStoreOnDemand/Private/OnDemandIoDispatcherBackend.cpp
@@ -931,6 +931,12 @@ static TUniquePtr<FArchive> CreateReaderFromPlatformPackage(const FString& RelPa
 	IFileHandle* File = IPlatformFile::GetPlatformPhysical().OpenRead(*AbsPath);
 	if (File)
 	{
+#if PLATFORM_ANDROID
+		// This is a handle to an asset so we need to call Seek(0) to move the internal
+		// offset to the start of the asset file.
+		File->Seek(0);
+#endif //PLATFORM_ANDROID
+
 		return MakeUnique<FArchiveFileReaderGeneric>(File, *AbsPath, File->Size());
 	}
 	else
@@ -958,6 +964,11 @@ static TIoStatusOr<FOnDemandToc> GenerateOnDemandTocFromDisk(FStringView TocHash
 		if (FPlatformMisc::FileExistsInPlatformPackage(TocPath))
 		{	
 			TUniquePtr<FArchive> Ar = CreateReaderFromPlatformPackage(TocPath);
+			if (!Ar.IsValid())
+			{
+				return FIoStatus(EIoErrorCode::ReadError, WriteToString<256>(TEXT("Failed to open '"), TocPath, TEXT("' from disk")));
+			}
+
 			*Ar << OutToc;
 
 			Ar->Close();
@@ -969,7 +980,7 @@ static TIoStatusOr<FOnDemandToc> GenerateOnDemandTocFromDisk(FStringView TocHash
 			}
 			else
 			{
-				return FIoStatus(EIoErrorCode::ReadError, WriteToString<256>(TEXT("Failed to open '"), TocPath, TEXT("' from disk")));
+				return FIoStatus(EIoErrorCode::ReadError, WriteToString<256>(TEXT("Failed to serialize '"), TocPath, TEXT("' from disk")));
 			}
 		}
 		else
