Fix crash due to incorrect transform array length being passed to ISMPool transform update.
The reason it can happen is due to a change in geometry collection build behavior that will be fixed.
This fix is to prevent crashes while that is worked on.

diff --git a/Engine/Source/Runtime/Experimental/GeometryCollectionEngine/Private/GeometryCollection/GeometryCollectionISMPoolComponent.cpp b/Engine/Source/Runtime/Experimental/GeometryCollectionEngine/Private/GeometryCollection/GeometryCollectionISMPoolComponent.cpp
--- a/Engine/Source/Runtime/Experimental/GeometryCollectionEngine/Private/GeometryCollection/GeometryCollectionISMPoolComponent.cpp
+++ b/Engine/Source/Runtime/Experimental/GeometryCollectionEngine/Private/GeometryCollection/GeometryCollectionISMPoolComponent.cpp
@@ -242,24 +242,28 @@ bool FGeometryCollectionISMPool::BatchUpdateInstancesTransforms(FGeometryCollect
 	{
 		FGeometryCollectionISM& ISM = ISMs[MeshInfo.ISMIndex];
 		const FInstanceGroups::FInstanceGroupRange& InstanceGroup = ISM.InstanceGroups.GroupRanges[MeshInfo.InstanceGroupIndex];
-		ensure((StartInstanceIndex + NewInstancesTransforms.Num()) <= InstanceGroup.Count);
+
+		// The transform count should fit within the instance group.
+		// Clamp it if it doesn't, but if we hit this ensure we need to investigate why.
+		ensure(StartInstanceIndex + NewInstancesTransforms.Num() <= InstanceGroup.Count);
+		const int32 NumTransforms = FMath::Min(NewInstancesTransforms.Num(), InstanceGroup.Count - StartInstanceIndex);
 
 		// If ISM component has identity transform (the common case) then we can skip world space to component space maths inside BatchUpdateInstancesTransforms()
 		bWorldSpace &= !ISM.ISMComponent->GetComponentTransform().Equals(FTransform::Identity, 0.f);
 
-		int32 StartIndex = ISM.InstanceIndexToRenderIndex[InstanceGroup.Start];
+		int32 StartRenderIndex = ISM.InstanceIndexToRenderIndex[InstanceGroup.Start + StartInstanceIndex];
 		int32 TransformIndex = 0;
 		int32 BatchCount = 1;
 
-		for (int InstanceIndex = StartInstanceIndex + 1; InstanceIndex < NewInstancesTransforms.Num(); ++InstanceIndex)
+		for (int InstanceIndex = StartInstanceIndex + 1; InstanceIndex < StartInstanceIndex + NumTransforms; ++InstanceIndex)
 		{
 			// Flush batch for non-sequential instances.
 			int32 RenderIndex = ISM.InstanceIndexToRenderIndex[InstanceGroup.Start + InstanceIndex];
-			if (RenderIndex != (StartIndex + BatchCount))
+			if (RenderIndex != (StartRenderIndex + BatchCount))
 			{
 				TArrayView<const FTransform> BatchedTransformsView = MakeArrayView(NewInstancesTransforms.GetData() + TransformIndex, BatchCount);
-				ISM.ISMComponent->BatchUpdateInstancesTransforms(StartIndex, BatchedTransformsView, bWorldSpace, bMarkRenderStateDirty, bTeleport);
-				StartIndex = RenderIndex;
+				ISM.ISMComponent->BatchUpdateInstancesTransforms(StartRenderIndex, BatchedTransformsView, bWorldSpace, bMarkRenderStateDirty, bTeleport);
+				StartRenderIndex = RenderIndex;
 				TransformIndex += BatchCount;
 				BatchCount = 0;
 			}
@@ -268,7 +272,7 @@ bool FGeometryCollectionISMPool::BatchUpdateInstancesTransforms(FGeometryCollect
 
 		// last one
 		TArrayView<const FTransform> BatchedTransformsView = MakeArrayView(NewInstancesTransforms.GetData() + TransformIndex, BatchCount);
-		return ISM.ISMComponent->BatchUpdateInstancesTransforms(StartIndex, BatchedTransformsView, bWorldSpace, bMarkRenderStateDirty, bTeleport);
+		return ISM.ISMComponent->BatchUpdateInstancesTransforms(StartRenderIndex, BatchedTransformsView, bWorldSpace, bMarkRenderStateDirty, bTeleport);
 	}
 	UE_LOG(LogChaos, Warning, TEXT("UGeometryCollectionISMPoolComponent : Invalid ISM Id (%d) when updating the transform "), MeshInfo.ISMIndex);
 	return false;
