Set a fail timeout for all IAS

#rnx

diff --git a/Engine/Source/Runtime/Experimental/IoStoreOnDemand/Private/OnDemandHttpClient.cpp b/Engine/Source/Runtime/Experimental/IoStoreOnDemand/Private/OnDemandHttpClient.cpp
--- a/Engine/Source/Runtime/Experimental/IoStoreOnDemand/Private/OnDemandHttpClient.cpp
+++ b/Engine/Source/Runtime/Experimental/IoStoreOnDemand/Private/OnDemandHttpClient.cpp
@@ -12,6 +12,13 @@
 namespace UE::IO::IAS
 {
 
+static int32 GIasHttpFailTimeOutMs = 4 * 1000;
+static FAutoConsoleVariableRef CVar_IasHttpFailTimeOutMs(
+	TEXT("ias.HttpFailTimeOutMs"),
+	GIasHttpFailTimeOutMs,
+	TEXT("Fail infinite network waits that take longer than this (in ms, 0=disabled)")
+);
+
 static void LogHttpResult(const TCHAR* Host, const TCHAR* Url, uint32 StatusCode, uint64 DurationMs, uint64 Size, uint64 Offset, const char* Memo = "ok")
 {
 	Size >>= 10;
@@ -83,7 +90,7 @@ FHttpClient::FHttpClient(FHttpClientConfig&& ClientConfig)
 	Connections.SetNum(Config.Endpoints.Num());
 	SetEndpoint(Config.PrimaryEndpoint);
 
-	EventLoop.SetFailTimeout(Config.FailTimeoutMs);
+	EventLoop.SetFailTimeout(GIasHttpFailTimeOutMs);
 }
 
 TUniquePtr<FHttpClient> FHttpClient::Create(FHttpClientConfig&& ClientConfig)
diff --git a/Engine/Source/Runtime/Experimental/IoStoreOnDemand/Private/OnDemandHttpClient.h b/Engine/Source/Runtime/Experimental/IoStoreOnDemand/Private/OnDemandHttpClient.h
--- a/Engine/Source/Runtime/Experimental/IoStoreOnDemand/Private/OnDemandHttpClient.h
+++ b/Engine/Source/Runtime/Experimental/IoStoreOnDemand/Private/OnDemandHttpClient.h
@@ -21,7 +21,6 @@ struct FHttpClientConfig
 	int32 PipelineLength = 1;
 	int32 MaxRetryCount = 1;
 	int32 ReceiveBufferSize = -1;
-	int32 FailTimeoutMs = 0;
 	bool bChangeEndpointAfterSuccessfulRetry = true;
 };
 
diff --git a/Engine/Source/Runtime/Experimental/IoStoreOnDemand/Private/OnDemandIoDispatcherBackend.cpp b/Engine/Source/Runtime/Experimental/IoStoreOnDemand/Private/OnDemandIoDispatcherBackend.cpp
--- a/Engine/Source/Runtime/Experimental/IoStoreOnDemand/Private/OnDemandIoDispatcherBackend.cpp
+++ b/Engine/Source/Runtime/Experimental/IoStoreOnDemand/Private/OnDemandIoDispatcherBackend.cpp
@@ -134,13 +134,6 @@ static FAutoConsoleVariableRef CVar_IasHttpTimeOutMs(
 	TEXT("Time out value for HTTP requests in milliseconds")
 );
 
-static int32 GIasHttpFailTimeOutMs = 4 * 1000;
-static FAutoConsoleVariableRef CVar_IasHttpFailTimeOutMs(
-	TEXT("ias.HttpFailTimeOutMs"),
-	GIasHttpFailTimeOutMs,
-	TEXT("Fail infinite network waits that take longer than this (in ms, 0=disabled)")
-);
-
 int32 GIasHttpHealthCheckWaitTime = 3000;
 static FAutoConsoleVariableRef CVar_IasHttpHealthCheckWaitTime(
 	TEXT("ias.HttpHealthCheckWaitTime"),
@@ -2296,7 +2289,6 @@ uint32 FOnDemandIoBackend::Run()
 		.PipelineLength = GIasHttpPipelineLength,
 		.MaxRetryCount = FMath::Max(AvailableEps.Urls.Num() + 1, GIasHttpRetryCount),
 		.ReceiveBufferSize = GIasHttpRecvBufKiB >= 0 ? GIasHttpRecvBufKiB << 10 : -1,
-		.FailTimeoutMs = GIasHttpFailTimeOutMs,
 		.bChangeEndpointAfterSuccessfulRetry = GIasHttpChangeEndpointAfterSuccessfulRetry,
 	});
 	check(HttpClient.IsValid());
