Iris - Make sure COND_None is zero instead of assuming it.

#rnx

diff --git a/Engine/Source/Runtime/Experimental/Iris/Core/Private/Iris/ReplicationSystem/NetDependencyData.cpp b/Engine/Source/Runtime/Experimental/Iris/Core/Private/Iris/ReplicationSystem/NetDependencyData.cpp
--- a/Engine/Source/Runtime/Experimental/Iris/Core/Private/Iris/ReplicationSystem/NetDependencyData.cpp
+++ b/Engine/Source/Runtime/Experimental/Iris/Core/Private/Iris/ReplicationSystem/NetDependencyData.cpp
@@ -84,6 +84,7 @@ FNetDependencyData::FSubObjectConditionalsArray& FNetDependencyData::GetOrCreate
 		
 		// Make sure that we initialize the conditionals to match the number of SubObjects
 		const int32 NumChildSubObjects = DependentObjectsStorage[Entry.ArrayIndices[EArrayType::ChildSubObjects]].Num();
+		static_assert(COND_None == 0, "Can't use SetNumZeroed() to initialize COND_None");
 		SubObjectConditionalsArray->SetNumZeroed(NumChildSubObjects);
 
 		return *SubObjectConditionalsArray;
