Iris
- Fixed bug with changemask not being set in context from CopyAndQuantizeIfDirty
- This caused arrays to always be fully quantized if one member was dirty
- Implemented DequantizeWithMask to support partial dequantization (not used)

diff --git a/Engine/Source/Runtime/Experimental/Iris/Core/Private/Iris/ReplicationSystem/ReplicationOperations.cpp b/Engine/Source/Runtime/Experimental/Iris/Core/Private/Iris/ReplicationSystem/ReplicationOperations.cpp
--- a/Engine/Source/Runtime/Experimental/Iris/Core/Private/Iris/ReplicationSystem/ReplicationOperations.cpp
+++ b/Engine/Source/Runtime/Experimental/Iris/Core/Private/Iris/ReplicationSystem/ReplicationOperations.cpp
@@ -55,6 +55,7 @@ static_assert(COND_Max == 17, "s_LifetimeConditionDebugNames may need updating."
 // Append changemask bits to ChangeMaskWriter and conditionally to the ConditionalChangeMaskWriter. If the state is dirty, the function will return true.
 static bool AppendMemberChangeMasks(FNetBitStreamWriter* ChangeMaskWriter, FNetBitStreamWriter* ConditionalChangeMaskWriter, uint8* ExternalStateBuffer, const FReplicationStateDescriptor* Descriptor);
 static void ResetMemberChangeMasks(uint8* ExternalStateBuffer, const FReplicationStateDescriptor* Descriptor);
+static void CheckChangeMask(const FNetSerializationContext& Context, const FNetBitArrayView& ChangeMask) { check(Context.GetChangeMask() && (Context.GetChangeMask()->GetData() == ChangeMask.GetData()));}
 
 }
 
@@ -102,6 +103,7 @@ void FReplicationStateOperations::QuantizeWithMask(FNetSerializationContext& Con
 {
 #if !UE_BUILD_SHIPPING && !UE_BUILD_TEST
 	check(!Descriptor->IsInitState());
+	Private::CheckChangeMask(Context, ChangeMask);
 #endif
 
 	check(IsAligned(DstInternalBuffer, Descriptor->InternalAlignment) && IsAligned(SrcExternalBuffer, Descriptor->ExternalAlignment));
@@ -133,6 +135,43 @@ void FReplicationStateOperations::QuantizeWithMask(FNetSerializationContext& Con
 	}
 }
 
+void FReplicationStateOperations::DequantizeWithMask(FNetSerializationContext& Context, const FNetBitArrayView& ChangeMask, const uint32 ChangeMaskOffset, uint8* RESTRICT DstExternalBuffer, const uint8* RESTRICT SrcInternalBuffer, const FReplicationStateDescriptor* Descriptor)
+{
+#if !UE_BUILD_SHIPPING && !UE_BUILD_TEST
+	check(!Descriptor->IsInitState());
+	Private::CheckChangeMask(Context, ChangeMask);
+#endif
+
+	check(IsAligned(SrcInternalBuffer, Descriptor->InternalAlignment) && IsAligned(DstExternalBuffer, Descriptor->ExternalAlignment));
+
+	const FReplicationStateMemberDescriptor* MemberDescriptors = Descriptor->MemberDescriptors;
+	const FReplicationStateMemberSerializerDescriptor* MemberSerializerDescriptors = Descriptor->MemberSerializerDescriptors;
+	const uint32 MemberCount = Descriptor->MemberCount;
+
+	// If adding code to only dequantize dirty states then special consideration is needed for init states.
+	for (uint32 MemberIt = 0; MemberIt < MemberCount; ++MemberIt)
+	{
+		const FReplicationStateMemberDescriptor& MemberDescriptor = MemberDescriptors[MemberIt];
+		const FReplicationStateMemberSerializerDescriptor& MemberSerializerDescriptor = MemberSerializerDescriptors[MemberIt];
+		const FReplicationStateMemberChangeMaskDescriptor& MemberChangeMaskDescriptor = Descriptor->MemberChangeMaskDescriptors[MemberIt];
+		const uint32 MemberChangeMaskOffset = ChangeMaskOffset + MemberChangeMaskDescriptor.BitOffset;
+
+		if (ChangeMask.IsAnyBitSet(MemberChangeMaskOffset, MemberChangeMaskDescriptor.BitCount))
+		{
+			FNetDequantizeArgs Args;
+			Args.Version = 0;
+			Args.NetSerializerConfig = MemberSerializerDescriptor.SerializerConfig;
+			Args.Target = reinterpret_cast<NetSerializerValuePointer>(DstExternalBuffer + MemberDescriptor.ExternalMemberOffset);
+			Args.Source = reinterpret_cast<NetSerializerValuePointer>(SrcInternalBuffer + MemberDescriptor.InternalMemberOffset);
+
+			Args.ChangeMaskInfo.BitCount = MemberChangeMaskDescriptor.BitCount;
+			Args.ChangeMaskInfo.BitOffset = MemberChangeMaskOffset;
+
+			MemberSerializerDescriptor.Serializer->Dequantize(Context, Args);
+		}
+	}
+}
+
 void FReplicationStateOperations::Dequantize(FNetSerializationContext& Context, uint8* RESTRICT DstExternalBuffer, const uint8* RESTRICT SrcInternalBuffer, const FReplicationStateDescriptor* Descriptor)
 {
 	check(IsAligned(SrcInternalBuffer, Descriptor->InternalAlignment) && IsAligned(DstExternalBuffer, Descriptor->ExternalAlignment));
@@ -328,6 +367,7 @@ void FReplicationStateOperations::SerializeWithMask(FNetSerializationContext& Co
 {
 #if !UE_BUILD_SHIPPING && !UE_BUILD_TEST
 	check(!Descriptor->IsInitState());
+	Private::CheckChangeMask(Context, ChangeMask);
 #endif
 
 	const FReplicationStateMemberDescriptor* MemberDescriptors = Descriptor->MemberDescriptors;
@@ -374,6 +414,7 @@ void FReplicationStateOperations::DeserializeWithMask(FNetSerializationContext&
 {
 #if !UE_BUILD_SHIPPING && !UE_BUILD_TEST
 	check(!Descriptor->IsInitState());
+	Private::CheckChangeMask(Context, ChangeMask);
 #endif
 
 	const FReplicationStateMemberDescriptor* MemberDescriptors = Descriptor->MemberDescriptors;
@@ -419,6 +460,7 @@ void FReplicationStateOperations::SerializeDeltaWithMask(FNetSerializationContex
 {
 #if !UE_BUILD_SHIPPING && !UE_BUILD_TEST
 	check(!Descriptor->IsInitState());
+	Private::CheckChangeMask(Context, ChangeMask);
 #endif
 
 	const FReplicationStateMemberDescriptor* MemberDescriptors = Descriptor->MemberDescriptors;
@@ -466,6 +508,7 @@ void FReplicationStateOperations::DeserializeDeltaWithMask(FNetSerializationCont
 {
 #if !UE_BUILD_SHIPPING && !UE_BUILD_TEST
 	check(!Descriptor->IsInitState());
+	Private::CheckChangeMask(Context, ChangeMask);
 #endif
 
 	const FReplicationStateMemberDescriptor* MemberDescriptors = Descriptor->MemberDescriptors;
@@ -627,7 +670,9 @@ void FReplicationInstanceOperations::CopyAndQuantizeIfDirty(FNetSerializationCon
 			}
 			else
 			{
-				FReplicationStateOperations::QuantizeWithMask(Context, UE::Net::Private::GetMemberChangeMask(FragmentData[StateIt].ExternalSrcBuffer, CurrentDescriptor), 0, CurrentInternalStateBuffer, FragmentData[StateIt].ExternalSrcBuffer, CurrentDescriptor);
+				FNetBitArrayView ChangeMask = UE::Net::Private::GetMemberChangeMask(FragmentData[StateIt].ExternalSrcBuffer, CurrentDescriptor);
+				Context.SetChangeMask(&ChangeMask);
+				FReplicationStateOperations::QuantizeWithMask(Context, ChangeMask, 0, CurrentInternalStateBuffer, FragmentData[StateIt].ExternalSrcBuffer, CurrentDescriptor);
 			}
 		}
 		
diff --git a/Engine/Source/Runtime/Experimental/Iris/Core/Public/Iris/ReplicationSystem/ReplicationOperations.h b/Engine/Source/Runtime/Experimental/Iris/Core/Public/Iris/ReplicationSystem/ReplicationOperations.h
--- a/Engine/Source/Runtime/Experimental/Iris/Core/Public/Iris/ReplicationSystem/ReplicationOperations.h
+++ b/Engine/Source/Runtime/Experimental/Iris/Core/Public/Iris/ReplicationSystem/ReplicationOperations.h
@@ -32,6 +32,9 @@ struct FReplicationStateOperations
 	/** Dequantize a Replication state from internal buffer to already constructed ExternalBuffer */
 	static IRISCORE_API void Dequantize(FNetSerializationContext& Context, uint8* RESTRICT DstExternalBuffer, const uint8* RESTRICT SrcInternalBuffer, const FReplicationStateDescriptor* Descriptor);
 
+	/** Dequantize a Replication state from internal buffer to already constructed ExternalBuffer */
+	static IRISCORE_API void DequantizeWithMask(FNetSerializationContext& Context, const FNetBitArrayView& ChangeMask, const uint32 ChangeMaskOffset, uint8* RESTRICT DstExternalBuffer, const uint8* RESTRICT SrcInternalBuffer, const FReplicationStateDescriptor* Descriptor);
+
 	/** Quantize a Replication state from ExternalBuffer to internal buffer, DstInternalBuffer does not need to be initialized */
 	static IRISCORE_API void QuantizeWithMask(FNetSerializationContext& Context, const FNetBitArrayView& ChangeMask, const uint32 ChangeMaskOffset, uint8* RESTRICT DstInternalBuffer, const uint8* RESTRICT SrcExternalBuffer, const FReplicationStateDescriptor* Descriptor);
 
