Iris - Added missing include.

#rnx

diff --git a/Engine/Source/Runtime/Experimental/Iris/Core/Private/Iris/ReplicationSystem/NetDependencyData.cpp b/Engine/Source/Runtime/Experimental/Iris/Core/Private/Iris/ReplicationSystem/NetDependencyData.cpp
--- a/Engine/Source/Runtime/Experimental/Iris/Core/Private/Iris/ReplicationSystem/NetDependencyData.cpp
+++ b/Engine/Source/Runtime/Experimental/Iris/Core/Private/Iris/ReplicationSystem/NetDependencyData.cpp
@@ -2,7 +2,8 @@
 
 #include "NetDependencyData.h"
 #include "Iris/Core/IrisLog.h"
-#include "ReplicationOperationsInternal.h"
+#include "Iris/ReplicationSystem/ReplicationOperationsInternal.h"
+#include "UObject/CoreNetTypes.h"
 
 namespace UE::Net::Private
 {
