Fixed debug proxy not working.

#rnx

diff --git a/Engine/Source/Runtime/Experimental/IoStoreOnDemand/Private/CoreHttp/Http.cpp b/Engine/Source/Runtime/Experimental/IoStoreOnDemand/Private/CoreHttp/Http.cpp
--- a/Engine/Source/Runtime/Experimental/IoStoreOnDemand/Private/CoreHttp/Http.cpp
+++ b/Engine/Source/Runtime/Experimental/IoStoreOnDemand/Private/CoreHttp/Http.cpp
@@ -1063,7 +1063,7 @@ static void SetSocksIpAddress(const TCHAR* Value)
 		GSocksIpAddress = 0;
 		return;
 	}
-	GSocksIpAddress = htonl(IpAddress);
+	GSocksIpAddress = IpAddress;
 }
 
 static FString GSocksIpAddressStr;
@@ -1117,7 +1117,7 @@ static int32 ConnectSocks4(FSocket& Socket, uint32 IpAddress, uint32 Port)
 
 	FSocks4Request Request = {
 		.Port		= htons(uint16(Port)),
-		.IpAddress	= IpAddress,
+		.IpAddress	= htonl(IpAddress),
 	};
 	Result = Socket.Send((const char*)&Request, sizeof(Request));
 	if (Result <= 0)
@@ -1132,7 +1132,7 @@ static int32 ConnectSocks4(FSocket& Socket, uint32 IpAddress, uint32 Port)
 		return -1;
 	}
 
-	return 0;
+	return 1;
 }
 
 ////////////////////////////////////////////////////////////////////////////////
@@ -1174,6 +1174,7 @@ static int32 ConnectSocks5(FSocket& Socket, uint32 IpAddress, uint32 Port)
 	}
 
 	// Connection request
+	IpAddress = htonl(IpAddress);
 	uint16 NsPort = htons(uint16(Port));
 	char Request[] = { 5, 1, 0, 1, 0x11,0x11,0x11,0x11, 0x22,0x22 };
 	std::memcpy(Request + 4, &IpAddress, sizeof(IpAddress));
@@ -1192,12 +1193,12 @@ static int32 ConnectSocks5(FSocket& Socket, uint32 IpAddress, uint32 Port)
 		return -1;
 	}
 
-	if (Reply[0] != 0x05 && Reply[1] != 0x00)
+	if (Reply[0] != 0x05 || Reply[1] != 0x00)
 	{
 		return -1;
 	}
 
-	return 0;
+	return 1;
 
 #ifdef _MSC_VER
 #pragma warning(pop)
@@ -1406,6 +1407,7 @@ FResult FHost::Connect(FSocket& Socket)
 	// Socks connect in a blocking fashion so we're all set (ret=1)
 	if (bSocksConnected)
 	{
+		Socket = MoveTemp(Candidate);
 		return FResult(1);
 	}
 
