Allow multiple meshids in a meshgroup to have the same staticmesh description

diff --git a/Engine/Source/Runtime/Experimental/GeometryCollectionEngine/Private/GeometryCollection/GeometryCollectionISMPoolComponent.cpp b/Engine/Source/Runtime/Experimental/GeometryCollectionEngine/Private/GeometryCollection/GeometryCollectionISMPoolComponent.cpp
--- a/Engine/Source/Runtime/Experimental/GeometryCollectionEngine/Private/GeometryCollection/GeometryCollectionISMPoolComponent.cpp
+++ b/Engine/Source/Runtime/Experimental/GeometryCollectionEngine/Private/GeometryCollection/GeometryCollectionISMPoolComponent.cpp
@@ -26,14 +26,7 @@ FAutoConsoleVariableRef CVarISMPoolComponentFreeListTargetSize(
 
 FGeometryCollectionMeshGroup::FMeshId FGeometryCollectionMeshGroup::AddMesh(const FGeometryCollectionStaticMeshInstance& MeshInstance, int32 InstanceCount, const FGeometryCollectionMeshInfo& ISMInstanceInfo)
 {
-	FMeshId* MeshIndex = Meshes.Find(MeshInstance);
-	if (MeshIndex)
-	{
-		return *MeshIndex;
-	}
-
 	const FMeshId MeshInfoIndex = MeshInfos.Emplace(ISMInstanceInfo);
-	Meshes.Add(MeshInstance, MeshInfoIndex);
 	return MeshInfoIndex;
 }
 
@@ -67,7 +60,6 @@ void FGeometryCollectionMeshGroup::RemoveAllMeshes(FGeometryCollectionISMPool& I
 		ISMPool.RemoveISM(MeshInfo);
 	}
 	MeshInfos.Empty();
-	Meshes.Empty();
 }
 
 void FGeometryCollectionISM::CreateISM(AActor* InOwningActor, bool bInUseHISM)
diff --git a/Engine/Source/Runtime/Experimental/GeometryCollectionEngine/Public/GeometryCollection/GeometryCollectionISMPoolComponent.h b/Engine/Source/Runtime/Experimental/GeometryCollectionEngine/Public/GeometryCollection/GeometryCollectionISMPoolComponent.h
--- a/Engine/Source/Runtime/Experimental/GeometryCollectionEngine/Public/GeometryCollection/GeometryCollectionISMPoolComponent.h
+++ b/Engine/Source/Runtime/Experimental/GeometryCollectionEngine/Public/GeometryCollection/GeometryCollectionISMPoolComponent.h
@@ -238,8 +238,6 @@ struct FGeometryCollectionMeshGroup
 
 	/** Array of allocated mesh infos. */
 	TArray<FGeometryCollectionMeshInfo> MeshInfos;
-	/** Map from mesh instance description to the its index in the MeshInfo array. */
-	TMap<FGeometryCollectionStaticMeshInstance, FMeshId> Meshes;
 };
 
 /** Structure containting all info for a single ISM. */
