Iris - Fixed objects always being enabled in NetObjectConnectionFilter.

diff --git a/Engine/Source/Runtime/Experimental/Iris/Core/Private/Iris/ReplicationSystem/Prioritization/NetObjectConnectionFilter.cpp b/Engine/Source/Runtime/Experimental/Iris/Core/Private/Iris/ReplicationSystem/Prioritization/NetObjectConnectionFilter.cpp
--- a/Engine/Source/Runtime/Experimental/Iris/Core/Private/Iris/ReplicationSystem/Prioritization/NetObjectConnectionFilter.cpp
+++ b/Engine/Source/Runtime/Experimental/Iris/Core/Private/Iris/ReplicationSystem/Prioritization/NetObjectConnectionFilter.cpp
@@ -23,7 +23,7 @@ void UNetObjectConnectionFilter::SetReplicateToConnection(UE::Net::FNetRefHandle
 	{
 		const uint32 LocalIndex = FilteringInfo->GetLocalObjectIndex();
 		FPerConnectionInfo& PerConnectionInfo = PerConnectionInfos[ConnectionId];
-		PerConnectionInfo.ReplicationEnabledObjects.SetBit(LocalIndex);
+		PerConnectionInfo.ReplicationEnabledObjects.SetBitValue(LocalIndex, FilterStatus == UE::Net::ENetFilterStatus::Allow);
 	}
 }
 
