Re-enable IAS loading the .iochunktoc from disk only instead of generating the final toc from ondemand .utoc files we will load from a file generated during the data build.

#rnx

- All known issues on various platforms have been fixed so it should be safe to enable this once more.

diff --git a/Engine/Source/Runtime/Experimental/IoStoreOnDemand/Private/OnDemandIoDispatcherBackend.cpp b/Engine/Source/Runtime/Experimental/IoStoreOnDemand/Private/OnDemandIoDispatcherBackend.cpp
--- a/Engine/Source/Runtime/Experimental/IoStoreOnDemand/Private/OnDemandIoDispatcherBackend.cpp
+++ b/Engine/Source/Runtime/Experimental/IoStoreOnDemand/Private/OnDemandIoDispatcherBackend.cpp
@@ -182,7 +182,7 @@ static FAutoConsoleVariableRef CVar_IoReportAnalytics(
 	GIasReportAnalyticsEnabled,
 	TEXT("Enables reporting statics to the analytics system"));
 
-bool GIasGenerateOnDemandToc = false;
+bool GIasGenerateOnDemandToc = true;
 static FAutoConsoleVariableRef CVar_IasGenerateOnDemandToc(
 	TEXT("s.IasGenerateOnDemandToc"),
 	GIasGenerateOnDemandToc,
@@ -192,7 +192,7 @@ static FAutoConsoleVariableRef CVar_IasGenerateOnDemandToc(
 
 // A temp fallback path allowing us to attempt to load the OnDemand toc from disk rather than
 // trying to generate it.
-bool GIasLoadOnDemandToc = false;
+bool GIasLoadOnDemandToc = true;
 static FAutoConsoleVariableRef CVar_IasLoadOnDemandToc(
 	TEXT("ias.LoadOnDemandToc"),
 	GIasLoadOnDemandToc,
