additions to the geometry collection component API

- ShouldUpdateComponentTransformToRootBone to expose a getter for bUpdateComponentTransformToRootBone
- New FOnGeometryCollectionRootMovedEvent event to know when the root bone moved

diff --git a/Engine/Source/Runtime/Experimental/GeometryCollectionEngine/Private/GeometryCollection/GeometryCollectionComponent.cpp b/Engine/Source/Runtime/Experimental/GeometryCollectionEngine/Private/GeometryCollection/GeometryCollectionComponent.cpp
--- a/Engine/Source/Runtime/Experimental/GeometryCollectionEngine/Private/GeometryCollection/GeometryCollectionComponent.cpp
+++ b/Engine/Source/Runtime/Experimental/GeometryCollectionEngine/Private/GeometryCollection/GeometryCollectionComponent.cpp
@@ -3844,7 +3844,16 @@ void UGeometryCollectionComponent::OnPostPhysicsSync()
 	// dirty the transform if the collection is
 	if (DynamicCollection && DynamicCollection->IsDirty())
 	{
+		const FTransform& PreviousRootTransform = OnRootMovedEvent.IsBound() ? ComponentSpaceTransforms.RequestRootTransform() : FTransform::Identity;
 		ComponentSpaceTransforms.MarkDirty();
+		if (OnRootMovedEvent.IsBound())
+		{
+			const FTransform& NewRootTransform = ComponentSpaceTransforms.RequestRootTransform();
+			if (!PreviousRootTransform.EqualsNoScale(NewRootTransform))
+			{
+				OnRootMovedEvent.Broadcast();
+			}
+		}
 	}
 
 	UpdateAttachedChildrenTransform();
diff --git a/Engine/Source/Runtime/Experimental/GeometryCollectionEngine/Public/GeometryCollection/GeometryCollectionComponent.h b/Engine/Source/Runtime/Experimental/GeometryCollectionEngine/Public/GeometryCollection/GeometryCollectionComponent.h
--- a/Engine/Source/Runtime/Experimental/GeometryCollectionEngine/Public/GeometryCollection/GeometryCollectionComponent.h
+++ b/Engine/Source/Runtime/Experimental/GeometryCollectionEngine/Public/GeometryCollection/GeometryCollectionComponent.h
@@ -65,6 +65,8 @@ DECLARE_DYNAMIC_MULTICAST_DELEGATE_OneParam(FOnChaosCrumblingEvent, const FChaos
 
 DECLARE_DYNAMIC_MULTICAST_DELEGATE(FOnGeometryCollectionFullyDecayedEvent);
 
+DECLARE_DYNAMIC_MULTICAST_DELEGATE(FOnGeometryCollectionRootMovedEvent);
+
 namespace GeometryCollection
 {
 	enum class ESelectionMode : uint8
@@ -1207,6 +1209,7 @@ public:
 	FOnChaosBreakEvent OnRootBreakEvent;
 
 	FOnGeometryCollectionFullyDecayedEvent OnFullyDecayedEvent;
+	FOnGeometryCollectionRootMovedEvent OnRootMovedEvent;
 
 	GEOMETRYCOLLECTIONENGINE_API void DispatchBreakEvent(const FChaosBreakEvent& Event);
 
@@ -1283,6 +1286,8 @@ public:
 	/** allow update of the custom renderer if enabled */
 	GEOMETRYCOLLECTIONENGINE_API void SetUpdateCustomRenderer(bool bValue) { bUpdateCustomRenderer = bValue; }
 
+	GEOMETRYCOLLECTIONENGINE_API bool ShouldUpdateComponentTransformToRootBone() const { return bUpdateComponentTransformToRootBone; }
+
 public:
 	UPROPERTY(BlueprintAssignable, Category = "Collision")
 	FOnChaosPhysicsCollision OnChaosPhysicsCollision;
