geometry collection : fix issue where "transformed" geometry were not ignored when importing collision from source

diff --git a/Engine/Source/Runtime/Experimental/GeometryCollectionEngine/Private/GeometryCollection/GeometryCollectionEngineConversion.cpp b/Engine/Source/Runtime/Experimental/GeometryCollectionEngine/Private/GeometryCollection/GeometryCollectionEngineConversion.cpp
--- a/Engine/Source/Runtime/Experimental/GeometryCollectionEngine/Private/GeometryCollection/GeometryCollectionEngineConversion.cpp
+++ b/Engine/Source/Runtime/Experimental/GeometryCollectionEngine/Private/GeometryCollection/GeometryCollectionEngineConversion.cpp
@@ -76,9 +76,17 @@ FORCEINLINE uint32 GetTypeHash(const FUniqueVertex& UniqueVertex)
 	return VertexHash;
 }
 
-static bool IsImportableImplicitObjectType(Chaos::EImplicitObjectType Type)
+static bool IsImportableImplicitObjectType(const Chaos::FImplicitObject& ImplicitObject)
 {
-	const Chaos::EImplicitObjectType InnerType = Type & (~(Chaos::ImplicitObjectType::IsScaled | Chaos::ImplicitObjectType::IsInstanced));
+	const Chaos::EImplicitObjectType InnerType = ImplicitObject.GetType() & (~(Chaos::ImplicitObjectType::IsScaled | Chaos::ImplicitObjectType::IsInstanced));
+	if (InnerType == Chaos::ImplicitObjectType::Transformed)
+	{
+		const Chaos::FImplicitObjectTransformed& TransformedImplicitObject = static_cast<const Chaos::FImplicitObjectTransformed&>(ImplicitObject);
+		if (const Chaos::FImplicitObject* SubObject = TransformedImplicitObject.GetTransformedObject())
+		{
+			return IsImportableImplicitObjectType(*SubObject);
+		}
+	}
 	return (InnerType == Chaos::ImplicitObjectType::Box || InnerType == Chaos::ImplicitObjectType::Sphere || InnerType == Chaos::ImplicitObjectType::Capsule || InnerType == Chaos::ImplicitObjectType::Convex);
 }
 
@@ -307,7 +315,7 @@ void FGeometryCollectionEngineConversion::AppendMeshDescription(
 		for (int32 GeomIndex = 0; GeomIndex < Geoms.Num();)
 		{
 			// make sure we only import box, sphere, capsule or convex
-			if (IsImportableImplicitObjectType(Geoms[GeomIndex]->GetType()))
+			if (Geoms[GeomIndex] && IsImportableImplicitObjectType(*Geoms[GeomIndex]))
 			{
 				GeomIndex++;
 			}
