Make sure auto instance meshes array does not contain duplicates when generating geometry collection from Dataflow
- Add a SetAutoInstanceMeshes method to geometry collection asset that rem,ove duplicates and re-assign the auto instance indices as well as consolidating per instance primitive data
- call this method when setting the array when creating the asset in the geometry collection terminal node

diff --git a/Engine/Plugins/Experimental/GeometryCollectionPlugin/Source/GeometryCollectionNodes/Private/Dataflow/Nodes/GeometryCollectionAssetNodes.cpp b/Engine/Plugins/Experimental/GeometryCollectionPlugin/Source/GeometryCollectionNodes/Private/Dataflow/Nodes/GeometryCollectionAssetNodes.cpp
--- a/Engine/Plugins/Experimental/GeometryCollectionPlugin/Source/GeometryCollectionNodes/Private/Dataflow/Nodes/GeometryCollectionAssetNodes.cpp
+++ b/Engine/Plugins/Experimental/GeometryCollectionPlugin/Source/GeometryCollectionNodes/Private/Dataflow/Nodes/GeometryCollectionAssetNodes.cpp
@@ -65,7 +65,7 @@ void FGeometryCollectionTerminalDataflowNode::SetAssetValue(TObjectPtr<UObject>
 
 			const bool bHasInternalMaterial = false; // with data flow there's no assumption of internal materials
 			CollectionAsset->ResetFrom(InCollection, InMaterials, false);
-			CollectionAsset->AutoInstanceMeshes = InInstancedMeshes;
+			CollectionAsset->SetAutoInstanceMeshes(InInstancedMeshes);
 
 #if WITH_EDITOR
 			// make sure we rebuild the render data when we are done setting everything 
diff --git a/Engine/Source/Runtime/Experimental/GeometryCollectionEngine/Private/GeometryCollection/GeometryCollectionObject.cpp b/Engine/Source/Runtime/Experimental/GeometryCollectionEngine/Private/GeometryCollection/GeometryCollectionObject.cpp
--- a/Engine/Source/Runtime/Experimental/GeometryCollectionEngine/Private/GeometryCollection/GeometryCollectionObject.cpp
+++ b/Engine/Source/Runtime/Experimental/GeometryCollectionEngine/Private/GeometryCollection/GeometryCollectionObject.cpp
@@ -1591,6 +1591,77 @@ int32 UGeometryCollection::FindOrAddAutoInstanceMesh(const UStaticMesh* StaticMe
 	return FindOrAddAutoInstanceMesh(NewMesh);
 }
 
+void UGeometryCollection::SetAutoInstanceMeshes(const TArray<FGeometryCollectionAutoInstanceMesh>& InAutoInstanceMeshes)
+{
+	AutoInstanceMeshes = InAutoInstanceMeshes;
+
+	// dedup array and reassign indices
+	if (AutoInstanceMeshes.Num() > 0)
+	{
+		TArray<FGeometryCollectionAutoInstanceMesh> UniqueAutoInstanceMeshes;
+		TArray<int32> InstanceMeshIndexRemap;
+
+		UniqueAutoInstanceMeshes.Reserve(AutoInstanceMeshes.Num());
+		InstanceMeshIndexRemap.Reserve(AutoInstanceMeshes.Num());
+
+		// now we may have two similar entries  we need to consolidate them 
+		for (int32 InstanceMeshIndex = 0; InstanceMeshIndex < AutoInstanceMeshes.Num(); InstanceMeshIndex++)
+		{
+			const FGeometryCollectionAutoInstanceMesh& InstanceMesh = AutoInstanceMeshes[InstanceMeshIndex];
+			int32 UniqueInstanceMeshIndex = UniqueAutoInstanceMeshes.Find(InstanceMesh);
+			if (UniqueInstanceMeshIndex == INDEX_NONE)
+			{
+				FGeometryCollectionAutoInstanceMesh UniqueInstanceMesh = InstanceMesh;
+				UniqueInstanceMesh.NumInstances = 0;
+				UniqueInstanceMesh.CustomData.Reset();
+				UniqueInstanceMeshIndex = UniqueAutoInstanceMeshes.Add(UniqueInstanceMesh);
+			}
+			// make sure num instance are custom data are aggregated 
+			UniqueAutoInstanceMeshes[UniqueInstanceMeshIndex].NumInstances += InstanceMesh.NumInstances;
+			InstanceMeshIndexRemap.Add(UniqueInstanceMeshIndex);
+		}
+
+		GeometryCollection::Facades::FCollectionInstancedMeshFacade InstancedMeshFacade(*GetGeometryCollection());
+
+		const TManagedArray<TSet<int32>>& Children = GetGeometryCollection()->Children;
+
+		// relocate custom data : we cannot just aggregate them because we may have interleaved transform indices with alternating colors
+		// also adjust the transform index to instance mesh index via the facade 
+		TArray<int32> DataReadOffsets;
+		DataReadOffsets.SetNumZeroed(AutoInstanceMeshes.Num());
+		for (int32 TransformIndex = 0; TransformIndex < InstancedMeshFacade.GetNumIndices(); TransformIndex++)
+		{
+			// only for leaves
+			if (Children[TransformIndex].Num() == 0)
+			{
+				const int32 OldIndex = InstancedMeshFacade.GetIndex(TransformIndex);
+				if (InstanceMeshIndexRemap.IsValidIndex(OldIndex))
+				{
+					const FGeometryCollectionAutoInstanceMesh& OldInstanceMesh = AutoInstanceMeshes[OldIndex];
+
+					const int32 NewIndex = InstanceMeshIndexRemap[OldIndex];
+					FGeometryCollectionAutoInstanceMesh& NewInstanceMesh = UniqueAutoInstanceMeshes[NewIndex];
+
+					InstancedMeshFacade.SetIndex(TransformIndex, NewIndex);
+
+					const int32 NumDataPerInstance = OldInstanceMesh.GetNumDataPerInstance();
+					if (NumDataPerInstance > 0)
+					{
+						const int32 DataReadOffset = DataReadOffsets[OldIndex];
+						for (int32 DataIndex = 0; DataIndex < NumDataPerInstance; DataIndex++)
+						{
+							const float OldData = OldInstanceMesh.CustomData[DataReadOffset + DataIndex];
+							NewInstanceMesh.CustomData.Add(OldData);
+						}
+						DataReadOffsets[OldIndex] += NumDataPerInstance;
+					}
+				}
+			}
+		}
+		AutoInstanceMeshes = MoveTemp(UniqueAutoInstanceMeshes);
+	}
+}
+
 FGuid UGeometryCollection::GetIdGuid() const
 {
 	return PersistentGuid;
diff --git a/Engine/Source/Runtime/Experimental/GeometryCollectionEngine/Public/GeometryCollection/GeometryCollectionObject.h b/Engine/Source/Runtime/Experimental/GeometryCollectionEngine/Public/GeometryCollection/GeometryCollectionObject.h
--- a/Engine/Source/Runtime/Experimental/GeometryCollectionEngine/Public/GeometryCollection/GeometryCollectionObject.h
+++ b/Engine/Source/Runtime/Experimental/GeometryCollectionEngine/Public/GeometryCollection/GeometryCollectionObject.h
@@ -438,6 +438,12 @@ public:
 	/** find or add a auto instance mesh and return its index */
 	GEOMETRYCOLLECTIONENGINE_API const FGeometryCollectionAutoInstanceMesh& GetAutoInstanceMesh(int32 AutoInstanceMeshIndex) const;
 
+	/** 
+	* Assign an auto instanced meshes array
+	* if there's duplicate entry in the array , they will be collapsed as one and the index attribute will be adjusted acoordingly
+	*/
+	GEOMETRYCOLLECTIONENGINE_API void SetAutoInstanceMeshes(const TArray<FGeometryCollectionAutoInstanceMesh>& InAutoInstanceMeshes);
+	
 	/**  find or add a auto instance mesh from another one and return its index */
 	GEOMETRYCOLLECTIONENGINE_API int32 FindOrAddAutoInstanceMesh(const FGeometryCollectionAutoInstanceMesh& AutoInstanceMesh);
 
