fix comment as per review

diff --git a/Engine/Source/Runtime/Experimental/GeometryCollectionEngine/Private/GeometryCollection/GeometryCollectionObject.cpp b/Engine/Source/Runtime/Experimental/GeometryCollectionEngine/Private/GeometryCollection/GeometryCollectionObject.cpp
--- a/Engine/Source/Runtime/Experimental/GeometryCollectionEngine/Private/GeometryCollection/GeometryCollectionObject.cpp
+++ b/Engine/Source/Runtime/Experimental/GeometryCollectionEngine/Private/GeometryCollection/GeometryCollectionObject.cpp
@@ -1616,7 +1616,7 @@ void UGeometryCollection::SetAutoInstanceMeshes(const TArray<FGeometryCollection
 				UniqueInstanceMesh.CustomData.Reset();
 				UniqueInstanceMeshIndex = UniqueAutoInstanceMeshes.Add(UniqueInstanceMesh);
 			}
-			// make sure num instance are custom data are aggregated 
+			// make sure num instances are aggregated
 			UniqueAutoInstanceMeshes[UniqueInstanceMeshIndex].NumInstances += InstanceMesh.NumInstances;
 			InstanceMeshIndexRemap.Add(UniqueInstanceMeshIndex);
 		}
