Chaos - wire up OneWayIntertaction flag to geometry collection and cluster union debris

diff --git a/Engine/Source/Runtime/Experimental/Chaos/Private/Chaos/PBDRigidClustering.cpp b/Engine/Source/Runtime/Experimental/Chaos/Private/Chaos/PBDRigidClustering.cpp
--- a/Engine/Source/Runtime/Experimental/Chaos/Private/Chaos/PBDRigidClustering.cpp
+++ b/Engine/Source/Runtime/Experimental/Chaos/Private/Chaos/PBDRigidClustering.cpp
@@ -70,12 +70,6 @@ namespace Chaos
 	int32 ClusteringParticleReleaseThrottlingMaxCount = INDEX_NONE;
 	FAutoConsoleVariableRef CVarClusteringParticleReleaseThrottlingMaxCount(TEXT("p.Clustering.ParticleReleaseThrottlingMaxCount"), ClusteringParticleReleaseThrottlingMaxCount, TEXT("Maximum number of active geometry collection to reach before all released clustering disable all released particle instantly"));
 
-	bool bClusteringEnableDebris = false;
-	FAutoConsoleVariableRef CVarClusteringEnableDebris(TEXT("p.Clustering.Debris.Enabled"), bClusteringEnableDebris, TEXT(""));
-
-	float ClusteringDebrisSizeThreshold = 200.0f;
-	FAutoConsoleVariableRef CVarClusteringDebrisSizeThreshold(TEXT("p.Clustering.Debris.SizeThreshold"), ClusteringDebrisSizeThreshold, TEXT(""));
-
 	template <typename TProxy=FGeometryCollectionPhysicsProxy>
 	TProxy* GetConcreteProxy(FPBDRigidClusteredParticleHandle* ClusteredParticle)
 	{
@@ -824,16 +818,6 @@ namespace Chaos
 			Child->SetPreV(Child->PreV() + ClusteredParent->PreV());
 			Child->SetPreW(Child->PreW() + ClusteredParent->PreW());
 
-			// Experimental
-			if (bClusteringEnableDebris && Child->HasBounds())
-			{
-				const FReal ChildSize = Child->LocalBounds().Extents().GetMax();
-				if (ChildSize < ClusteringDebrisSizeThreshold)
-				{
-					Child->SetOneWayInteraction(true);
-				}
-			}
-
 			// We also need to do cluster book-keeping on the parent.
 			// If the parent is an internal cluster and has become empty, we need to mark this particle as ready to be destroyed.
 			// The only exception is for cluster unions which are managed separately.
diff --git a/Engine/Source/Runtime/Experimental/Chaos/Private/PhysicsProxy/GeometryCollectionPhysicsProxy.cpp b/Engine/Source/Runtime/Experimental/Chaos/Private/PhysicsProxy/GeometryCollectionPhysicsProxy.cpp
--- a/Engine/Source/Runtime/Experimental/Chaos/Private/PhysicsProxy/GeometryCollectionPhysicsProxy.cpp
+++ b/Engine/Source/Runtime/Experimental/Chaos/Private/PhysicsProxy/GeometryCollectionPhysicsProxy.cpp
@@ -475,6 +475,7 @@ FGeometryCollectionPhysicsProxy::FGeometryCollectionPhysicsProxy(
 	, bIsCollisionFilterDataDirty(false)
 	, bIsDamageThresholdDataDirty(false)
 	, bIsGravityGroupIndexDirty(false)
+	, bIsOneWayInteractionDirty(false)
 	, CollectorGuid(InCollectorGuid)
 {
 	// We rely on a guarded buffer.
@@ -1465,6 +1466,7 @@ void FGeometryCollectionPhysicsProxy::InitializeBodiesPT(Chaos::FPBDRigidsSolver
 				Handle->SetGravityEnabled(bEnableGravity);
 				Handle->SetGravityGroupIndex(Parameters.GravityGroupIndex);
 				Handle->SetCCDEnabled(Parameters.UseCCD);
+				Handle->SetOneWayInteraction(Parameters.bIsOneWayInteraction);
 				Handle->SetInertiaConditioningEnabled(Parameters.UseInertiaConditioning);
 				Handle->SetLinearEtherDrag(Parameters.LinearDamping);
 				Handle->SetAngularEtherDrag(Parameters.AngularDamping);
@@ -2973,6 +2975,13 @@ void FGeometryCollectionPhysicsProxy::SetGravityGroupIndex_External(int32 Gravit
 	SetProxyDirty_External();
 }
 
+void FGeometryCollectionPhysicsProxy::SetIsOneWayInteraction_External(bool bIsOneWayInteraction)
+{
+	check(IsInGameThread());
+	GameThreadPerFrameData.bIsOneWayInteraction = bIsOneWayInteraction;
+	SetProxyDirty_External();
+}
+
 void FGeometryCollectionPhysicsProxy::PushStateOnGameThread(Chaos::FPBDRigidsSolver* InSolver)
 {
 	// rest all the dirty flags
@@ -2980,6 +2989,7 @@ void FGeometryCollectionPhysicsProxy::PushStateOnGameThread(Chaos::FPBDRigidsSol
 	bIsCollisionFilterDataDirty = false;
 	bIsDamageThresholdDataDirty = false;
 	bIsGravityGroupIndexDirty = false;
+	bIsOneWayInteractionDirty = false;
 	MaterialOverrideMassScaleMultiplierChange = 0;
 
 	// CONTEXT: GAMETHREAD
@@ -3062,6 +3072,13 @@ void FGeometryCollectionPhysicsProxy::PushStateOnGameThread(Chaos::FPBDRigidsSol
 		bIsGravityGroupIndexDirty = true;
 	}
 
+	if (GameThreadPerFrameData.bIsOneWayInteraction.IsSet())
+	{
+		Parameters.bIsOneWayInteraction = GameThreadPerFrameData.bIsOneWayInteraction.GetValue();
+		GameThreadPerFrameData.bIsOneWayInteraction.Reset();
+		bIsOneWayInteractionDirty = true;
+	}
+
 	if (GameThreadPerFrameData.DamagePropagationData.IsSet())
 	{
 		const FGeometryCollectioPerFrameData::FDamagePropagationData& Data = GameThreadPerFrameData.DamagePropagationData.GetValue();
@@ -3213,6 +3230,17 @@ void FGeometryCollectionPhysicsProxy::PushToPhysicsState()
 		}
 	}
 
+	if (bIsOneWayInteractionDirty)
+	{
+		for (int32 TransformIndex = 0; TransformIndex < SolverParticleHandles.Num(); ++TransformIndex)
+		{
+			if (Chaos::FPBDRigidClusteredParticleHandle* Handle = SolverParticleHandles[TransformIndex])
+			{
+				Handle->SetOneWayInteraction(Parameters.bIsOneWayInteraction);
+			}
+		}
+	}
+
 	if (MaterialOverrideMassScaleMultiplierChange > 0)
 	{
 		for (int32 TransformIndex = 0; TransformIndex < SolverParticleHandles.Num(); ++TransformIndex)
diff --git a/Engine/Source/Runtime/Experimental/Chaos/Public/GeometryCollection/GeometryCollectionSimulationCoreTypes.h b/Engine/Source/Runtime/Experimental/Chaos/Public/GeometryCollection/GeometryCollectionSimulationCoreTypes.h
--- a/Engine/Source/Runtime/Experimental/Chaos/Public/GeometryCollection/GeometryCollectionSimulationCoreTypes.h
+++ b/Engine/Source/Runtime/Experimental/Chaos/Public/GeometryCollection/GeometryCollectionSimulationCoreTypes.h
@@ -213,6 +213,7 @@ struct FSimulationParameters
 		, GravityGroupIndex(0)
 		, UseInertiaConditioning(true)
 		, UseCCD(false)
+		, bIsOneWayInteraction(false)
 		, LinearDamping(0.01f)
 		, AngularDamping(0)
 		, bUseDamagePropagation(false)
@@ -272,6 +273,7 @@ struct FSimulationParameters
 		, GravityGroupIndex(Other.GravityGroupIndex)
 		, UseInertiaConditioning(Other.UseInertiaConditioning)
 		, UseCCD(Other.UseCCD)
+		, bIsOneWayInteraction(Other.bIsOneWayInteraction)
 		, LinearDamping(Other.LinearDamping)
 		, AngularDamping(Other.AngularDamping)
 		, bUseDamagePropagation(Other.bUseDamagePropagation)
@@ -361,6 +363,7 @@ struct FSimulationParameters
 	int32 GravityGroupIndex;
 	bool UseInertiaConditioning;
 	bool UseCCD;
+	bool bIsOneWayInteraction;
 	float LinearDamping;
 	float AngularDamping;
 
diff --git a/Engine/Source/Runtime/Experimental/Chaos/Public/GeometryCollectionProxyData.h b/Engine/Source/Runtime/Experimental/Chaos/Public/GeometryCollectionProxyData.h
--- a/Engine/Source/Runtime/Experimental/Chaos/Public/GeometryCollectionProxyData.h
+++ b/Engine/Source/Runtime/Experimental/Chaos/Public/GeometryCollectionProxyData.h
@@ -382,6 +382,7 @@ public:
 	TOptional<bool> bUseMaterialDamageModifiers;
 	TOptional<float> MaterialOverrideMassScaleMultiplier;
 	TOptional<uint8> GravityGroupIndex;
+	TOptional<bool> bIsOneWayInteraction;
 };
 
 /**
diff --git a/Engine/Source/Runtime/Experimental/Chaos/Public/PhysicsProxy/GeometryCollectionPhysicsProxy.h b/Engine/Source/Runtime/Experimental/Chaos/Public/PhysicsProxy/GeometryCollectionPhysicsProxy.h
--- a/Engine/Source/Runtime/Experimental/Chaos/Public/PhysicsProxy/GeometryCollectionPhysicsProxy.h
+++ b/Engine/Source/Runtime/Experimental/Chaos/Public/PhysicsProxy/GeometryCollectionPhysicsProxy.h
@@ -463,6 +463,7 @@ public:
 	CHAOS_API void SetUseMaterialDamageModifiers_External(bool bUseMaterialDamageModifiers);
 	CHAOS_API void SetMaterialOverrideMassScaleMultiplier_External(float InMultiplier);
 	CHAOS_API void SetGravityGroupIndex_External(int32 GravityGroupIndex);
+	CHAOS_API void SetIsOneWayInteraction_External(bool bInIsOneWayInteraction);
 	/** 
 	 * Traverses the parents of TransformGroupIdx counting number of levels,
 	 * and sets levels array value for TransformGroupIdx and its parents if not yet initialized.
@@ -642,6 +643,7 @@ private:
 	uint8 bIsCollisionFilterDataDirty: 1;
 	uint8 bIsDamageThresholdDataDirty: 1;
 	uint8 bIsGravityGroupIndexDirty: 1;
+	uint8 bIsOneWayInteractionDirty : 1;
 
 	// Currently this is using triple buffers for game-physics and 
 	// physics-game thread communication, but not for any reason other than this 
diff --git a/Engine/Source/Runtime/Experimental/GeometryCollectionEngine/Private/GeometryCollection/GeometryCollectionComponent.cpp b/Engine/Source/Runtime/Experimental/GeometryCollectionEngine/Private/GeometryCollection/GeometryCollectionComponent.cpp
--- a/Engine/Source/Runtime/Experimental/GeometryCollectionEngine/Private/GeometryCollection/GeometryCollectionComponent.cpp
+++ b/Engine/Source/Runtime/Experimental/GeometryCollectionEngine/Private/GeometryCollection/GeometryCollectionComponent.cpp
@@ -540,6 +540,7 @@ UGeometryCollectionComponent::UGeometryCollectionComponent(const FObjectInitiali
 	, InitializationState(ESimulationInitializationState::Unintialized)
 	, ObjectType(EObjectStateTypeEnum::Chaos_Object_Dynamic)
 	, GravityGroupIndex(0)
+	, bIsOneWayInteraction(false)
 	, bDensityFromPhysicsMaterial(false)
 	, bForceMotionBlur()
 	, EnableClustering(true)
@@ -946,6 +947,15 @@ void UGeometryCollectionComponent::SetGravityGroupIndex(int32 InGravityGroupInde
 	}
 }
 
+void UGeometryCollectionComponent::SetIsOneWayInteraction(bool bInIsOneWayInteraction)
+{
+	bIsOneWayInteraction = bInIsOneWayInteraction;
+	if (PhysicsProxy)
+	{
+		PhysicsProxy->SetIsOneWayInteraction_External(bInIsOneWayInteraction);
+	}
+}
+
 void UGeometryCollectionComponent::SetDamageThreshold(const TArray<float>& InDamageThreshold)
 {
 	// NOTE: Should only call this during construction, not during runtime
@@ -3740,6 +3750,7 @@ void UGeometryCollectionComponent::RegisterAndInitializePhysicsProxy()
 		SimulationParameters.bGenerateGlobalCrumblingChildrenData = bGlobalCrumblingEventIncludesChildren;
 		SimulationParameters.EnableGravity = BodyInstance.bEnableGravity;
 		SimulationParameters.GravityGroupIndex = GravityGroupIndex;
+		SimulationParameters.bIsOneWayInteraction = bIsOneWayInteraction;
 		SimulationParameters.UseInertiaConditioning = BodyInstance.IsInertiaConditioningEnabled();
 		SimulationParameters.UseCCD = BodyInstance.bUseCCD;
 		SimulationParameters.LinearDamping = BodyInstance.LinearDamping;
diff --git a/Engine/Source/Runtime/Experimental/GeometryCollectionEngine/Public/GeometryCollection/GeometryCollectionComponent.h b/Engine/Source/Runtime/Experimental/GeometryCollectionEngine/Public/GeometryCollection/GeometryCollectionComponent.h
--- a/Engine/Source/Runtime/Experimental/GeometryCollectionEngine/Public/GeometryCollection/GeometryCollectionComponent.h
+++ b/Engine/Source/Runtime/Experimental/GeometryCollectionEngine/Public/GeometryCollection/GeometryCollectionComponent.h
@@ -931,6 +931,12 @@ public:
 	UFUNCTION(BlueprintCallable, BlueprintInternalUseOnly)
 	GEOMETRYCOLLECTIONENGINE_API void SetGravityGroupIndex(int32 InGravityGroupIndex);
 
+	UPROPERTY(EditAnywhere, BlueprintReadWrite, BlueprintSetter = SetIsOneWayInteraction, Category = "ChaosPhysics|General")
+	bool bIsOneWayInteraction;
+
+	UFUNCTION(BlueprintCallable, BlueprintInternalUseOnly)
+	GEOMETRYCOLLECTIONENGINE_API void SetIsOneWayInteraction(bool bInOneWayInteraction);
+
 	/** when true, density will be used to compute mass using the assigned physics material */
 	UPROPERTY(EditAnywhere, BlueprintReadWrite, BlueprintSetter = SetDensityFromPhysicsMaterial, Category = "ChaosPhysics|General")
 	bool bDensityFromPhysicsMaterial;
