Store transform remap array for AutoInstanceMeshes in the geometry collection object.
This allows it to be calculated and stored once, for reference by multiple component instances.

diff --git a/Engine/Source/Runtime/Experimental/GeometryCollectionEngine/Private/GeometryCollection/GeometryCollectionObject.cpp b/Engine/Source/Runtime/Experimental/GeometryCollectionEngine/Private/GeometryCollection/GeometryCollectionObject.cpp
--- a/Engine/Source/Runtime/Experimental/GeometryCollectionEngine/Private/GeometryCollection/GeometryCollectionObject.cpp
+++ b/Engine/Source/Runtime/Experimental/GeometryCollectionEngine/Private/GeometryCollection/GeometryCollectionObject.cpp
@@ -329,6 +329,51 @@ void UGeometryCollection::CacheBreadthFirstTransformIndices()
 	}
 }
 
+void UGeometryCollection::CacheAutoInstanceTransformRemapIndices()
+{
+	AutoInstanceTransformRemapIndices.Reset();
+	if (GeometryCollection == nullptr)
+	{
+		return;
+	}
+	const GeometryCollection::Facades::FCollectionInstancedMeshFacade InstancedMeshFacade(*GeometryCollection);
+	if (!InstancedMeshFacade.IsValid())
+	{
+		return;
+	}
+
+	const int32 NumMeshes = AutoInstanceMeshes.Num();
+
+	TArray<int32> TransformGroups;
+	TransformGroups.AddZeroed(NumMeshes);
+	TArray<int32> TransformStarts;
+	TransformStarts.AddUninitialized(NumMeshes);
+	TArray<int32> InstanceCounts;
+	InstanceCounts.AddUninitialized(NumMeshes);
+	TArray<int32> WrittenTransformCounts;
+	WrittenTransformCounts.AddZeroed(NumMeshes);
+
+	for (int32 MeshIndex = 0; MeshIndex < NumMeshes; MeshIndex++)
+	{
+		const int32 NumInstances = AutoInstanceMeshes[MeshIndex].NumInstances;
+		TransformStarts[MeshIndex] = MeshIndex == 0 ? 0 : TransformStarts[MeshIndex - 1] + InstanceCounts[MeshIndex - 1];
+		InstanceCounts[MeshIndex] = NumInstances;
+	}
+	
+	AutoInstanceTransformRemapIndices.AddUninitialized(TransformStarts.Last() + InstanceCounts.Last());
+
+	const int32 NumTransforms = InstancedMeshFacade.GetNumIndices();
+	for (int32 TransformIndex = 0; TransformIndex < NumTransforms; TransformIndex++)
+	{
+		if (GeometryCollection->Children[TransformIndex].Num() == 0)
+		{
+			const int32 AutoInstanceMeshIndex = InstancedMeshFacade.GetIndex(TransformIndex);
+			const int32 TransformArrayIndex = TransformStarts[AutoInstanceMeshIndex] + WrittenTransformCounts[AutoInstanceMeshIndex]++;
+			AutoInstanceTransformRemapIndices[TransformArrayIndex] = TransformIndex;
+		}
+	}
+}
+
 void UGeometryCollection::UpdateGeometryDependentProperties()
 {
 #if WITH_EDITOR
@@ -1153,6 +1198,8 @@ void UGeometryCollection::Serialize(FArchive& Ar)
 
 	// Generate root to leave order lookup
 	CacheBreadthFirstTransformIndices();
+	// Generate transform remap for AutoInstanceMeshes instances
+	CacheAutoInstanceTransformRemapIndices();
 
 	if (Ar.IsLoading())
 	{
diff --git a/Engine/Source/Runtime/Experimental/GeometryCollectionEngine/Public/GeometryCollection/GeometryCollectionObject.h b/Engine/Source/Runtime/Experimental/GeometryCollectionEngine/Public/GeometryCollection/GeometryCollectionObject.h
--- a/Engine/Source/Runtime/Experimental/GeometryCollectionEngine/Public/GeometryCollection/GeometryCollectionObject.h
+++ b/Engine/Source/Runtime/Experimental/GeometryCollectionEngine/Public/GeometryCollection/GeometryCollectionObject.h
@@ -820,6 +820,8 @@ public:
 
 	GEOMETRYCOLLECTIONENGINE_API const TArray<int32>& GetBreadthFirstTransformIndices() const { return BreadthFirstTransformIndices; }
 
+	GEOMETRYCOLLECTIONENGINE_API const TArray<int32>& GetAutoInstanceTransformRemapIndices() const { return AutoInstanceTransformRemapIndices; }
+
 private:
 #if WITH_EDITOR
 	GEOMETRYCOLLECTIONENGINE_API void CreateSimulationDataImp(bool bCopyFromDDC);
@@ -839,6 +841,7 @@ private:
 	GEOMETRYCOLLECTIONENGINE_API void FillAutoInstanceMeshesInstancesIfNeeded();
 
 	void CacheBreadthFirstTransformIndices();
+	void CacheAutoInstanceTransformRemapIndices();
 
 private:
 	/** Guid created on construction of this collection. It should be used to uniquely identify this collection */
@@ -873,6 +876,10 @@ private:
 	UPROPERTY(VisibleAnywhere, Transient, Category = "Clustering")
 	TArray<int32> BreadthFirstTransformIndices;
 
+	// cache transform remap indices
+	UPROPERTY(VisibleAnywhere, Transient, Category = "Clustering")
+	TArray<int32> AutoInstanceTransformRemapIndices;
+
 	// #todo(dmp): rename to be consistent BoneSelectedMaterialID?
 	// Legacy index of the bone selected material in the object's Materials array, or INDEX_NONE if it is not stored there.
 	// Note for new objects the bone selected material should not be stored in the Materials array, so this should be INDEX_NONE
