ISMPool changed to use the new ISM InstanceId permanant handles.
ISMPool can remove/re-add instances from the ISM depending on whether they are being updated with a zero transform.
This can help keep ISM GPUScene cost down in the case where we are hiding (now removing) many instances.

diff --git a/Engine/Source/Runtime/Engine/Classes/Components/InstancedStaticMeshComponent.h b/Engine/Source/Runtime/Engine/Classes/Components/InstancedStaticMeshComponent.h
--- a/Engine/Source/Runtime/Engine/Classes/Components/InstancedStaticMeshComponent.h
+++ b/Engine/Source/Runtime/Engine/Classes/Components/InstancedStaticMeshComponent.h
@@ -190,6 +190,8 @@ class UInstancedStaticMeshComponent : public UStaticMeshComponent, public ISMIns
 	 */
 	ENGINE_API void SetCustomDataById(const TArrayView<const FPrimitiveInstanceId> &InstanceIds, TArrayView<const float> CustomDataFloats); 
 	inline void SetCustomDataById(FPrimitiveInstanceId InstanceId, TArrayView<const float> CustomDataFloats) { SetCustomDataById(MakeArrayView(&InstanceId, 1), CustomDataFloats); }
+	ENGINE_API void SetCustomDataValueById(FPrimitiveInstanceId InstanceId, int32 CustomDataIndex, float CustomDataValue);
+
 	/**
 	 */
 	ENGINE_API virtual void RemoveInstancesById(const TArrayView<const FPrimitiveInstanceId> &InstanceIds, bool bUpdateNavigation = true);
diff --git a/Engine/Source/Runtime/Engine/Private/InstancedStaticMesh.cpp b/Engine/Source/Runtime/Engine/Private/InstancedStaticMesh.cpp
--- a/Engine/Source/Runtime/Engine/Private/InstancedStaticMesh.cpp
+++ b/Engine/Source/Runtime/Engine/Private/InstancedStaticMesh.cpp
@@ -3192,6 +3192,12 @@ void UInstancedStaticMeshComponent::SetCustomDataById(const TArrayView<const FPr
 	}
 }
 
+void UInstancedStaticMeshComponent::SetCustomDataValueById(FPrimitiveInstanceId InstanceId, int32 CustomDataIndex, float CustomDataValue)
+{
+	check(PrimitiveInstanceDataManager.GetMode() != FPrimitiveInstanceDataManager::EMode::ExternalLegacyData);
+	SetCustomDataValue(PrimitiveInstanceDataManager.IdToIndex(InstanceId), CustomDataIndex, CustomDataValue);
+}
+
 void UInstancedStaticMeshComponent::RemoveInstancesById(const TArrayView<const FPrimitiveInstanceId>& InstanceIds, const bool bUpdateNavigation)
 {
 	check(PrimitiveInstanceDataManager.GetMode() != FPrimitiveInstanceDataManager::EMode::ExternalLegacyData);
diff --git a/Engine/Source/Runtime/Experimental/GeometryCollectionEngine/Private/GeometryCollection/GeometryCollectionISMPoolComponent.cpp b/Engine/Source/Runtime/Experimental/GeometryCollectionEngine/Private/GeometryCollection/GeometryCollectionISMPoolComponent.cpp
--- a/Engine/Source/Runtime/Experimental/GeometryCollectionEngine/Private/GeometryCollection/GeometryCollectionISMPoolComponent.cpp
+++ b/Engine/Source/Runtime/Experimental/GeometryCollectionEngine/Private/GeometryCollection/GeometryCollectionISMPoolComponent.cpp
@@ -140,23 +140,19 @@ FInstanceGroups::FInstanceGroupId FGeometryCollectionISM::AddInstanceGroup(int32
 	const FInstanceGroups::FInstanceGroupId InstanceGroupIndex = InstanceGroups.AddGroup(InstanceCount);
 	const FInstanceGroups::FInstanceGroupRange& NewInstanceGroup = InstanceGroups.GroupRanges[InstanceGroupIndex];
 
+	// Ensure that remapping arrays are big enough to hold any new items.
+	InstanceIds.SetNum(InstanceGroups.GetMaxInstanceIndex(), false);
+
 	FTransform ZeroScaleTransform;
 	ZeroScaleTransform.SetIdentityZeroScale();
 	TArray<FTransform> ZeroScaleTransforms;
 	ZeroScaleTransforms.Init(ZeroScaleTransform, InstanceCount);
 
 	ISMComponent->PreAllocateInstancesMemory(InstanceCount);
-	TArray<int32> RenderInstances = ISMComponent->AddInstances(ZeroScaleTransforms, true, true);
-
-	// Ensure that remapping arrays are big enough to hold any new items.
-	InstanceIndexToRenderIndex.SetNum(InstanceGroups.GetMaxInstanceIndex(), false);
-	RenderIndexToInstanceIndex.SetNum(ISMComponent->PerInstanceSMData.Num(), false);
-	// Store mapping between our fixed instance index and the mutable ISM render index.
-	// todo: Improve ISM API so that we don't need to pay the memory overhead here to manage this.
+	TArray<FPrimitiveInstanceId> AddedInstanceIds = ISMComponent->AddInstancesById(ZeroScaleTransforms, true, true);
 	for (int32 InstanceIndex = 0; InstanceIndex < InstanceCount; ++InstanceIndex)
 	{
-		InstanceIndexToRenderIndex[NewInstanceGroup.Start + InstanceIndex] = RenderInstances[InstanceIndex];
-		RenderIndexToInstanceIndex[RenderInstances[InstanceIndex]] = NewInstanceGroup.Start + InstanceIndex;
+		InstanceIds[NewInstanceGroup.Start + InstanceIndex] = AddedInstanceIds[InstanceIndex];
 	}
 
 	// Set any custom data.
@@ -167,7 +163,7 @@ FInstanceGroups::FInstanceGroupId FGeometryCollectionISM::AddInstanceGroup(int32
 		{
 			for (int32 InstanceIndex = 0; InstanceIndex < InstanceCount; ++InstanceIndex)
 			{
-				ISMComponent->SetCustomData(RenderInstances[InstanceIndex], CustomDataFloats.Slice(InstanceIndex * NumCustomDataFloats, NumCustomDataFloats));
+				ISMComponent->SetCustomDataById(AddedInstanceIds[InstanceIndex], CustomDataFloats.Slice(InstanceIndex * NumCustomDataFloats, NumCustomDataFloats));
 			}
 		}
 	}
@@ -230,44 +226,48 @@ bool FGeometryCollectionISMPool::BatchUpdateInstancesTransforms(FGeometryCollect
 
 bool FGeometryCollectionISMPool::BatchUpdateInstancesTransforms(FGeometryCollectionMeshInfo& MeshInfo, int32 StartInstanceIndex, TArrayView<const FTransform> NewInstancesTransforms, bool bWorldSpace, bool bMarkRenderStateDirty, bool bTeleport)
 {
-	if (ISMs.IsValidIndex(MeshInfo.ISMIndex))
+	if (!ISMs.IsValidIndex(MeshInfo.ISMIndex))
 	{
-		FGeometryCollectionISM& ISM = ISMs[MeshInfo.ISMIndex];
-		const FInstanceGroups::FInstanceGroupRange& InstanceGroup = ISM.InstanceGroups.GroupRanges[MeshInfo.InstanceGroupIndex];
-
-		// The transform count should fit within the instance group.
-		// Clamp it if it doesn't, but if we hit this ensure we need to investigate why.
-		ensure(StartInstanceIndex + NewInstancesTransforms.Num() <= InstanceGroup.Count);
-		const int32 NumTransforms = FMath::Min(NewInstancesTransforms.Num(), InstanceGroup.Count - StartInstanceIndex);
-
-		// If ISM component has identity transform (the common case) then we can skip world space to component space maths inside BatchUpdateInstancesTransforms()
-		bWorldSpace &= !ISM.ISMComponent->GetComponentTransform().Equals(FTransform::Identity, 0.f);
-
-		int32 StartRenderIndex = ISM.InstanceIndexToRenderIndex[InstanceGroup.Start + StartInstanceIndex];
-		int32 TransformIndex = 0;
-		int32 BatchCount = 1;
+		UE_LOG(LogChaos, Warning, TEXT("UGeometryCollectionISMPoolComponent : Invalid ISM Id (%d) when updating the transform "), MeshInfo.ISMIndex);
+		return false;
+	}
+		
+	FGeometryCollectionISM& ISM = ISMs[MeshInfo.ISMIndex];
+	const FInstanceGroups::FInstanceGroupRange& InstanceGroup = ISM.InstanceGroups.GroupRanges[MeshInfo.InstanceGroupIndex];
+	// If ISM component has identity transform (the common case) then we can skip world space to component space maths inside BatchUpdateInstancesTransforms()
+	bWorldSpace &= !ISM.ISMComponent->GetComponentTransform().Equals(FTransform::Identity, 0.f);
+
+	// The transform count should fit within the instance group.
+	// Clamp it if it doesn't, but if we hit this ensure we need to investigate why.
+	ensure(StartInstanceIndex + NewInstancesTransforms.Num() <= InstanceGroup.Count);
+	const int32 NumTransforms = FMath::Min(NewInstancesTransforms.Num(), InstanceGroup.Count - StartInstanceIndex);
+
+	// Loop over transforms.
+	// todo: There may be some value in batching InstanceIds and caling one function for each of Add/Remove/Update.
+	// However the ISM batched calls themselves seem to be just simple loops over the single instance calls, so probably no benefit.
+	for (int InstanceIndex = StartInstanceIndex; InstanceIndex < StartInstanceIndex + NumTransforms; ++InstanceIndex)
+	{
+		FPrimitiveInstanceId InstanceId = ISM.InstanceIds[InstanceGroup.Start + InstanceIndex];
+		FTransform const& Transform = NewInstancesTransforms[InstanceIndex];
 
-		for (int InstanceIndex = StartInstanceIndex + 1; InstanceIndex < StartInstanceIndex + NumTransforms; ++InstanceIndex)
+		if (Transform.GetScale3D().IsZero() && InstanceId.IsValid())
 		{
-			// Flush batch for non-sequential instances.
-			int32 RenderIndex = ISM.InstanceIndexToRenderIndex[InstanceGroup.Start + InstanceIndex];
-			if (RenderIndex != (StartRenderIndex + BatchCount))
-			{
-				TArrayView<const FTransform> BatchedTransformsView = MakeArrayView(NewInstancesTransforms.GetData() + TransformIndex, BatchCount);
-				ISM.ISMComponent->BatchUpdateInstancesTransforms(StartRenderIndex, BatchedTransformsView, bWorldSpace, bMarkRenderStateDirty, bTeleport);
-				StartRenderIndex = RenderIndex;
-				TransformIndex += BatchCount;
-				BatchCount = 0;
-			}
-			BatchCount++;
+			// Zero scale is used to indicate that we should remove the instance from the ISM.
+			ISM.ISMComponent->RemoveInstanceById(InstanceId);
+			ISM.InstanceIds[InstanceGroup.Start + InstanceIndex] = FPrimitiveInstanceId();
+		}
+		else if (!Transform.GetScale3D().IsZero() && !InstanceId.IsValid())
+		{
+			// Re-add the instance to the ISM if the scale becomes non-zero.
+			ISM.InstanceIds[InstanceGroup.Start + InstanceIndex] = ISM.ISMComponent->AddInstanceById(Transform, bWorldSpace);
+		}
+		else if (InstanceId.IsValid())
+		{
+			ISM.ISMComponent->UpdateInstanceTransformById(InstanceId, Transform, bWorldSpace, bTeleport);
 		}
-
-		// last one
-		TArrayView<const FTransform> BatchedTransformsView = MakeArrayView(NewInstancesTransforms.GetData() + TransformIndex, BatchCount);
-		return ISM.ISMComponent->BatchUpdateInstancesTransforms(StartRenderIndex, BatchedTransformsView, bWorldSpace, bMarkRenderStateDirty, bTeleport);
 	}
-	UE_LOG(LogChaos, Warning, TEXT("UGeometryCollectionISMPoolComponent : Invalid ISM Id (%d) when updating the transform "), MeshInfo.ISMIndex);
-	return false;
+
+	return true;
 }
 
 void FGeometryCollectionISMPool::BatchUpdateInstanceCustomData(FGeometryCollectionMeshInfo const& MeshInfo, int32 CustomFloatIndex, float CustomFloatValue)
@@ -286,9 +286,11 @@ void FGeometryCollectionISMPool::BatchUpdateInstanceCustomData(FGeometryCollecti
 	const FInstanceGroups::FInstanceGroupRange& InstanceGroup = ISM.InstanceGroups.GroupRanges[MeshInfo.InstanceGroupIndex];
 	for (int32 InstanceIndex = 0; InstanceIndex < InstanceGroup.Count; ++InstanceIndex)
 	{
-		const int32 RenderIndex = ISM.InstanceIndexToRenderIndex[InstanceGroup.Start + InstanceIndex];
-		const bool bLastUpdate = InstanceIndex == InstanceGroup.Count - 1;
-		ISM.ISMComponent->SetCustomDataValue(RenderIndex, CustomFloatIndex, CustomFloatValue, bLastUpdate);
+		const FPrimitiveInstanceId InstanceId = ISM.InstanceIds[InstanceGroup.Start + InstanceIndex];
+		if (InstanceId.IsValid())
+		{
+			ISM.ISMComponent->SetCustomDataValueById(InstanceId, CustomFloatIndex, CustomFloatValue);
+		}
 	}
 }
 
@@ -299,41 +301,30 @@ void FGeometryCollectionISMPool::RemoveISM(const FGeometryCollectionMeshInfo& Me
 		FGeometryCollectionISM& ISM = ISMs[MeshInfo.ISMIndex];
 		const FInstanceGroups::FInstanceGroupRange& InstanceGroup = ISM.InstanceGroups.GroupRanges[MeshInfo.InstanceGroupIndex];
 		
-		TArray<int32> InstancesToRemove;
-		InstancesToRemove.SetNum(InstanceGroup.Count);
-		for (int32 InstanceIndex = 0; InstanceIndex < InstanceGroup.Count; ++InstanceIndex)
+		for (int32 Index = 0; Index < InstanceGroup.Count; ++Index)
 		{
-			// We need render index to pass to the ISMComponent.
-			InstancesToRemove[InstanceIndex] = ISM.InstanceIndexToRenderIndex[InstanceGroup.Start + InstanceIndex];
-			// Clear the stored render index since we're about to remove it.
-			ISM.InstanceIndexToRenderIndex[InstanceGroup.Start + InstanceIndex] = -1;
-		}
-
-		// we sort the array on the spot because we use it after calling RemoveInstances to fix up our own indices
-		InstancesToRemove.Sort(TGreater<int32>());
-		constexpr bool bArrayAlreadySorted = true;
-		ISM.ISMComponent->RemoveInstances(InstancesToRemove, bArrayAlreadySorted);
-
-		// Fix up instance index remapping to match what will have happened in our ISM component in RemoveInstances()
-		check(ISM.ISMComponent->SupportsRemoveSwap());
-		for (int32 RenderIndex : InstancesToRemove)
-		{
-			ISM.RenderIndexToInstanceIndex.RemoveAtSwap(RenderIndex, 1, false);
-			if (RenderIndex < ISM.RenderIndexToInstanceIndex.Num())
+			FPrimitiveInstanceId InstanceId = ISM.InstanceIds[InstanceGroup.Start + Index];
+			if (InstanceId.IsValid())
 			{
-				const int32 MovedInstanceIndex = ISM.RenderIndexToInstanceIndex[RenderIndex];
-				ISM.InstanceIndexToRenderIndex[MovedInstanceIndex] = RenderIndex;
+				// todo: Could RemoveInstanceByIds() instead as long as that function can handle skipping invalid InstanceIds.
+				ISM.ISMComponent->RemoveInstanceById(InstanceId);
 			}
 		}
-
+		
+#if DO_CHECK
+		// clear the IDs
+		for (int32 Index = 0; Index < InstanceGroup.Count; ++Index)
+		{
+			ISM.InstanceIds[InstanceGroup.Start + Index] = FPrimitiveInstanceId();
+		}
+#endif
 		ISM.InstanceGroups.RemoveGroup(MeshInfo.InstanceGroupIndex);
 	
 		if (ISM.InstanceGroups.IsEmpty())
 		{
 			// No live instances, so take opportunity to reset indexing.
 			ISM.InstanceGroups.Reset();
-			ISM.InstanceIndexToRenderIndex.Reset();
-			ISM.RenderIndexToInstanceIndex.Reset();
+			ISM.InstanceIds.Reset();
 		}
 
 		if (GUseComponentFreeList && ISM.ISMComponent->PerInstanceSMData.Num() == 0)
@@ -501,8 +492,7 @@ void UGeometryCollectionISMPoolComponent::GetResourceSizeEx(FResourceSizeEx& Cum
 	
 	for (FGeometryCollectionISM ISM : Pool.ISMs)
 	{
-		SizeBytes += ISM.InstanceIndexToRenderIndex.GetAllocatedSize()
-			+ ISM.RenderIndexToInstanceIndex.GetAllocatedSize()
+		SizeBytes += ISM.InstanceIds.GetAllocatedSize()
 			+ ISM.InstanceGroups.GroupRanges.GetAllocatedSize()
 			+ ISM.InstanceGroups.FreeList.GetAllocatedSize();
 	}
diff --git a/Engine/Source/Runtime/Experimental/GeometryCollectionEngine/Public/GeometryCollection/GeometryCollectionISMPoolComponent.h b/Engine/Source/Runtime/Experimental/GeometryCollectionEngine/Public/GeometryCollection/GeometryCollectionISMPoolComponent.h
--- a/Engine/Source/Runtime/Experimental/GeometryCollectionEngine/Public/GeometryCollection/GeometryCollectionISMPoolComponent.h
+++ b/Engine/Source/Runtime/Experimental/GeometryCollectionEngine/Public/GeometryCollection/GeometryCollectionISMPoolComponent.h
@@ -6,6 +6,7 @@
 #include "Containers/Map.h"
 #include "InstancedStaticMeshDelegates.h"
 #include "Materials/MaterialInterface.h"
+#include "InstanceDataTypes.h"
 
 #include "GeometryCollectionISMPoolComponent.generated.h"
 
@@ -256,10 +257,8 @@ struct FGeometryCollectionISM
 	TObjectPtr<UInstancedStaticMeshComponent> ISMComponent;
 	/** Groups of instances allocated in the ISM. */
 	FInstanceGroups InstanceGroups;
-	/** Mapping from our instance index to the ISM Component index. */
-	TArray<int32> InstanceIndexToRenderIndex;
-	/** Mapping from the ISM Component index to our instance index . */
-	TArray<int32> RenderIndexToInstanceIndex;
+	/** Id of Instance in ISMC */
+	TArray<FPrimitiveInstanceId> InstanceIds;
 };
 
 /** A pool of ISMs. */
