Apply radial impulse : reduce code complexity and fix dangerous cast to clustered without any check

diff --git a/Engine/Source/Runtime/Experimental/Chaos/Private/Chaos/PhysicsObjectInterface.cpp b/Engine/Source/Runtime/Experimental/Chaos/Private/Chaos/PhysicsObjectInterface.cpp
--- a/Engine/Source/Runtime/Experimental/Chaos/Private/Chaos/PhysicsObjectInterface.cpp
+++ b/Engine/Source/Runtime/Experimental/Chaos/Private/Chaos/PhysicsObjectInterface.cpp
@@ -930,12 +930,13 @@ namespace Chaos
 							continue;
 						}
 
-						//can apply to clusters or rigid particles
-						TThreadParticle<EThreadContext::Internal>* Particle = Object->GetParticle<EThreadContext::Internal>();
-						if (FPBDRigidClusteredParticleHandle* Cluster = reinterpret_cast<FPBDRigidClusteredParticleHandle*>(Particle))
+						// can apply to clusters or rigid particles
+						if (TThreadParticle<EThreadContext::Internal>* Particle = Object->GetParticle<EThreadContext::Internal>())
 						{
-							if (bApplyStrain)
+							FPBDRigidClusteredParticleHandle* Cluster = Particle->CastToClustered();
+							if (Cluster && bApplyStrain)
 							{
+								// it is a cluster and we are applying strain so we need to get to the children particles
 								FRigidClustering& Clustering = RigidSolver->GetEvolution()->GetRigidClustering();
 								TArray<FPBDRigidParticleHandle*>* ChildrenHandles = Clustering.GetChildrenMap().Find(Cluster);
 
@@ -954,11 +955,11 @@ namespace Chaos
 								{
 									const float FalloffAlpha = AddRadialImpulseHelper(ChildHandle, Origin, Radius, Strength, Falloff, VelocityRatio, bInvalidate, bVelChange);
 
-									//to do: remove cvar when material system is in place and densities are updated
-									const float StrainToApply = 
+									// todo(chaos) : Remove StrainModifier cvar when material system is in place and densities are updated
+									const float StrainToApply =
 										(Strain < 0)
-											? (PhysicsObjectInterfaceCVars::StrainModifier * FalloffAlpha * Strength)
-											: Strain * FalloffAlpha;
+										? (PhysicsObjectInterfaceCVars::StrainModifier * FalloffAlpha * Strength)
+										: Strain * FalloffAlpha;
 									if (StrainToApply > 0)
 									{
 										Clustering.SetExternalStrain(ChildHandle->CastToClustered(), StrainToApply);
@@ -967,8 +968,8 @@ namespace Chaos
 							}
 							else
 							{
-								FPBDRigidParticleHandle* ParticleHandle = Cluster->CastToRigidParticle();
-
+								// apply the impulse to the particle itself 
+								FPBDRigidParticleHandle* ParticleHandle = Particle->CastToRigidParticle();
 								if (ParticleHandle != nullptr && ParticleHandle->IsSleeping())
 								{
 									RigidSolver->GetEvolution()->SetParticleObjectState(ParticleHandle, Chaos::EObjectStateType::Dynamic);
@@ -978,18 +979,6 @@ namespace Chaos
 								AddRadialImpulseHelper(ParticleHandle, Origin, Radius, Strength, Falloff, bInvalidate, bVelChange);
 							}
 						}
-						else
-						{
-							FPBDRigidParticleHandle* ParticleHandle = reinterpret_cast<FPBDRigidParticleHandle*>(Particle);
-
-							if (ParticleHandle != nullptr && ParticleHandle->IsSleeping())
-							{
-								RigidSolver->GetEvolution()->SetParticleObjectState(ParticleHandle, Chaos::EObjectStateType::Dynamic);
-								RigidSolver->GetEvolution()->GetParticles().MarkTransientDirtyParticle(ParticleHandle);
-							}
-
-							AddRadialImpulseHelper(ParticleHandle, Origin, Radius, Strength, Falloff, bInvalidate, bVelChange);
-						}
 					}
 				});
 		}
