Fixed warning as error problem.

#rnx

diff --git a/Engine/Source/Runtime/Experimental/Iris/Core/Private/Iris/ReplicationSystem/ReplicationWriter.cpp b/Engine/Source/Runtime/Experimental/Iris/Core/Private/Iris/ReplicationSystem/ReplicationWriter.cpp
--- a/Engine/Source/Runtime/Experimental/Iris/Core/Private/Iris/ReplicationSystem/ReplicationWriter.cpp
+++ b/Engine/Source/Runtime/Experimental/Iris/Core/Private/Iris/ReplicationSystem/ReplicationWriter.cpp
@@ -1936,7 +1936,7 @@ FReplicationWriter::EWriteObjectStatus FReplicationWriter::WriteObjectAndSubObje
 	BatchEntry.bHasUnsentAttachments = bHasAttachments;
 
 	// Check if we must defer tearoff until after flush
-	const bool bSentTearOff = Info.TearOff && (GetFlushStatus(InternalIndex, Info, Info.FlushFlags | EFlushFlags::FlushFlags_FlushTornOffSubObjects) == EFlushFlags::FlushFlags_None);
+	const bool bSentTearOff = Info.TearOff && (GetFlushStatus(InternalIndex, Info, uint32(Info.FlushFlags | EFlushFlags::FlushFlags_FlushTornOffSubObjects)) == EFlushFlags::FlushFlags_None);
 
 	Context.SetIsInitState(bIsInitialState);
 
