- Keep IAS download work queued rather than scattered about immediately
- socket groups get work from the queue as they come to the end of existing work (or have nowt to do)
- no decisions are made based on the wait state to mitigate quirky poll() scenarios.

#rnx

diff --git a/Engine/Source/Runtime/Experimental/IoStoreOnDemand/Private/CoreHttp/Http.cpp b/Engine/Source/Runtime/Experimental/IoStoreOnDemand/Private/CoreHttp/Http.cpp
--- a/Engine/Source/Runtime/Experimental/IoStoreOnDemand/Private/CoreHttp/Http.cpp
+++ b/Engine/Source/Runtime/Experimental/IoStoreOnDemand/Private/CoreHttp/Http.cpp
@@ -5,6 +5,7 @@
 #if !defined(NO_UE_INCLUDES)
 #include "Containers/Array.h"
 #include "HAL/CriticalSection.h"
+#include "HAL/IConsoleManager.h"
 #include "IO/IoBuffer.h"
 #include "LatencyInjector.h"
 #include "Math/UnrealMathUtility.h"
@@ -101,6 +102,7 @@ namespace UE::IO::IAS::HTTP
 	x(Connect) \
 	x(Send) \
 	x(Recv) \
+	x(StartWork) \
 	x($)
 
 enum class ETrace
@@ -176,6 +178,14 @@ static void Trace(T Id, ETrace Action, U Param=0)
 
 // {{{1 misc ...................................................................
 
+////////////////////////////////////////////////////////////////////////////////
+static int32 GIasHttpRecvWorkThresholdKiB = 80;
+static FAutoConsoleVariableRef CVar_IasHttpRecvWorkThresholdKiB(
+	TEXT("ias.HttpWorkThreshold"),
+	GIasHttpRecvWorkThresholdKiB,
+	TEXT("Threshold of data remaining at which next request is sent (in KiB)")
+);
+
 ////////////////////////////////////////////////////////////////////////////////
 class FResult
 {
@@ -709,7 +719,9 @@ public:
 	struct FWaiter
 	{
 		enum class EWhat { Send = 0b01, Recv = 0b10, Both = Send|Recv };
-				FWaiter(FSocket& Socket, EWhat InWaitOn);
+				FWaiter() { std::memset(this, 0, sizeof(*this)); }
+				FWaiter(const FSocket& Socket, EWhat InWaitOn);
+		bool	IsValid() const { UPTRINT x{0}; return std::memcmp(this, &x, sizeof(*this)) != 0; }
 		bool	operator == (FSocket& Rhs) const { return UPTRINT(&Rhs) == Candidate; }
 		UPTRINT	Candidate : 60;
 		UPTRINT	WaitOn : 2;
@@ -741,7 +753,7 @@ private:
 };
 
 ////////////////////////////////////////////////////////////////////////////////
-FSocket::FWaiter::FWaiter(FSocket& Socket, EWhat InWaitOn)
+FSocket::FWaiter::FWaiter(const FSocket& Socket, EWhat InWaitOn)
 : Candidate(UPTRINT(&Socket))
 , WaitOn(UPTRINT(InWaitOn))
 , Ready(0)
@@ -1575,7 +1587,7 @@ int64 FStopwatch::Sample()
 struct FResponseInternal
 {
 	FMessageOffsets Offsets;
-	int32			ContentLength;
+	int32			ContentLength = 0;
 	uint16			MessageLength;
 	mutable int16	Code;
 };
@@ -1690,6 +1702,17 @@ static int32 Activity_Rewind(FActivity* Activity)
 	return -1;
 }
 
+////////////////////////////////////////////////////////////////////////////////
+static uint32 Activity_RemainingKiB(FActivity* Activity)
+{
+	if (Activity->State < FActivity::EState::RecvContent) return MAX_uint32;
+	if (Activity->State > FActivity::EState::RecvContent) return 0;
+
+	uint32 ContentLength = uint32(Activity->Response.ContentLength);
+	check(Activity->StateParam <= ContentLength);
+	return (ContentLength - Activity->StateParam) >> 10;
+}
+
 ////////////////////////////////////////////////////////////////////////////////
 static FActivity* Activity_Alloc(uint32 BufferSize)
 {
@@ -2504,14 +2527,18 @@ void FThrottler::ReturnUnused(uint32 Unused)
 /*
  * - Activities (requests send with a loop) are managed in singly-linked lists
  * - Each activity has an associated host it is talking to.
- * - Hosts are emphemeral, or represented externally via a FConnectionPool object
+ * - Hosts are ephemeral, or represented externally via a FConnectionPool object
  * - Loop has a group for each host, and each host-group has a bunch of socket-groups
- * - Socket-groups manage the list of activities to which they are assigned.
+ * - Host-group has a list of work; pending activities waiting to start
+ * - Socket-groups own up to two activities; one sending, one receiving
+ * - As it recvs, a socket-group will, if possible, fetch more work from the host
  *
- *  Loop
- *    FHostGroup[HostPtr]
- *      FSocketGroup[0...HostMaxConnections]
- *	      Act0 -> Act1 -> Act2 -> Act3 -> ...
+ *  Loop:
+ *    FHostGroup[HostPtr]:
+ *	    Work: Act0 -> Act1 -> Act2 -> Act3 -> ...
+ *      FSocketGroup[0...HostMaxConnections]:
+ *			Act.Send
+ *			Act.Recv
  */
 
 ////////////////////////////////////////////////////////////////////////////////
@@ -2521,36 +2548,141 @@ struct FTickState
 	uint64						Cancels;
 	int32&						RecvAllowance;
 	int32						PollTimeoutMs;
-	TArray<FSocket::FWaiter>*	Waiters;
+	class FWorkQueue*			Work;
+};
+
+
+
+////////////////////////////////////////////////////////////////////////////////
+class FWorkQueue
+{
+public:
+						FWorkQueue() = default;
+						~FWorkQueue();
+	bool				HasWork() const { return List != nullptr; }
+	void				AddActivity(FActivity* Activity);
+	void				PushActivity(FActivity* Activity);
+	FActivity*			PopActivity();
+	void				TickCancels(FTickState& State);
+
+private:
+	FActivity*			List = nullptr;
+	FActivity*			ListTail = nullptr;
+	uint64				ActiveSlots = 0;
+
+	UE_NONCOPYABLE(FWorkQueue);
 };
 
+////////////////////////////////////////////////////////////////////////////////
+FWorkQueue::~FWorkQueue()
+{
+	check(List == nullptr);
+	check(ListTail == nullptr);
+}
+
+////////////////////////////////////////////////////////////////////////////////
+void FWorkQueue::AddActivity(FActivity* Activity)
+{
+	// We use a tail pointer here to maintain order that requests were made
+
+	check(Activity->Next == nullptr);
+
+	if (ListTail != nullptr)
+	{
+		ListTail->Next = Activity;
+	}
+	List = (List == nullptr) ? Activity : List;
+	ListTail = Activity;
+
+	ActiveSlots |= (1ull << Activity->Slot);
+}
+
+////////////////////////////////////////////////////////////////////////////////
+void FWorkQueue::PushActivity(FActivity* Activity)
+{
+	Activity->Next = List;
+	List = Activity;
+	ListTail = (ListTail != nullptr) ? ListTail : Activity;
+}
+
+////////////////////////////////////////////////////////////////////////////////
+FActivity* FWorkQueue::PopActivity()
+{
+	if (List == nullptr)
+	{
+		return nullptr;
+	}
+
+	FActivity* Activity = List;
+	if ((List = List->Next) == nullptr)
+	{
+		ListTail = nullptr;
+	}
+
+	check(ActiveSlots & (1ull << Activity->Slot));
+	ActiveSlots ^= (1ull << Activity->Slot);
+
+	Activity->Next = nullptr;
+	return Activity;
+}
+
+////////////////////////////////////////////////////////////////////////////////
+void FWorkQueue::TickCancels(FTickState& State)
+{
+	if (State.Cancels == 0 || (State.Cancels & ActiveSlots) == 0)
+	{
+		return;
+	}
+
+	// We are going to rebuild the list of activities to maintain order as the
+	// activity list is singular.
+
+	check(List != nullptr);
+	FActivity* Activity = List;
+	List = ListTail = nullptr;
+	ActiveSlots = 0;
+
+	for (FActivity* Next; Activity != nullptr; Activity = Next)
+	{
+		Next = Activity->Next;
+
+		if (uint64 Slot = (1ull << Activity->Slot); (State.Cancels & Slot) == 0)
+		{
+			Activity->Next = nullptr;
+			AddActivity(Activity);
+			continue;
+		}
+
+		DoCancel(Activity);
+
+		Activity->Next = State.DoneList;
+		State.DoneList = Activity;
+	}
+}
+
+
+
 ////////////////////////////////////////////////////////////////////////////////
 class FSocketGroup
 {
 public:
-						FSocketGroup(uint32 PipeLen) : PipelineLength(uint8(PipeLen)) {}
+						FSocketGroup() = default;
 						~FSocketGroup();
 	bool				operator == (FSocket* Rhs) const { return &Socket == Rhs; }
-	bool				IsBusy() const		{ return IsKeepAlive || List != nullptr; }
-	bool				IsWaiting() const	{ return bWaiting; }
 	void				Unwait()			{ check(bWaiting); bWaiting = false; }
-	void				AddActivity(FActivity* Activity);
-	void				Tick(FTickState& State, FHost& Host);
+	FSocket::FWaiter	GetWaiter() const;
+	bool 				Tick(FTickState& State);
+	void				TickSend(FTickState& State, FHost& Host);
 	void				Fail(FTickState& State, const char* Reason);
 
 private:
-	void				TickRecv(FTickState& State);
-	void				TickSend(FTickState& State, FActivity* Activity);
-	void				TickCancels(FTickState& State);
-	void				Done(FTickState& State, FActivity* Activity);
+	void				RecvInternal(FTickState& State);
+	void				SendInternal(FTickState& State);
+	FActivity*			Send = nullptr;
+	FActivity*			Recv = nullptr;
 	FSocket				Socket;
-	FActivity*			List = nullptr;
-	FActivity*			ListTail = nullptr;
-	uint64				ActiveSlots = 0;
-	uint8				IsKeepAlive = 1;
+	uint8				IsKeepAlive = 0;
 	bool				bWaiting = false;
-	uint8				NumActive = 0;
-	uint8				PipelineLength;
 
 	UE_NONCOPYABLE(FSocketGroup);
 };
@@ -2558,62 +2690,59 @@ private:
 ////////////////////////////////////////////////////////////////////////////////
 FSocketGroup::~FSocketGroup()
 {
-	check(List == nullptr);
-	check(ListTail == nullptr);
-	check(NumActive == 0);
-	check(ActiveSlots == 0);
+	check(Send == nullptr);
+	check(Recv == nullptr);
 }
 
 ////////////////////////////////////////////////////////////////////////////////
-void FSocketGroup::Done(FTickState& State, FActivity* Activity)
+FSocket::FWaiter FSocketGroup::GetWaiter() const
 {
-	Activity->Next = State.DoneList;
-	State.DoneList = Activity;
+	if (!bWaiting)
+	{
+		return FSocket::FWaiter();
+	}
 
-	ActiveSlots &= ~(1ull << Activity->Slot);
+	using EWhat = FSocket::FWaiter::EWhat;
+	EWhat What = (Recv != nullptr) ? EWhat::Recv : EWhat::Send;
+	return FSocket::FWaiter(Socket, What);
 }
 
 ////////////////////////////////////////////////////////////////////////////////
 void FSocketGroup::Fail(FTickState& State, const char* Reason)
 {
-	// Failure is quite terminal and we need to abort all listed activities.
+	// Any send left at this point is unrecoverable
+	if (Send != nullptr)
+	{
+		Send->Next = Recv;
+		Recv = Send;
+	}
 
-	FActivity* Activity = List;
-	while (Activity != nullptr)
+	// Failure is quite terminal and we need to abort everything
+	for (FActivity* Activity = Recv; Activity != nullptr;)
 	{
-		FActivity* Next = Activity->Next;
 		if (Activity->State != FActivity::EState::Failed)
 		{
 			Activity_SetError(Activity, Reason);
 		}
 
 		DoFail(Activity);
-		Done(State, Activity);
 
+		FActivity* Next = Activity->Next;
+		Activity->Next = State.DoneList;
+		State.DoneList = Activity;
 		Activity = Next;
 	}
-	check(ActiveSlots == 0);
 
 	Socket = FSocket();
-	NumActive = 0;
-	List = ListTail = nullptr;
+	Send = Recv = nullptr;
 	bWaiting = false;
 	IsKeepAlive = 0;
 }
 
 ////////////////////////////////////////////////////////////////////////////////
-void FSocketGroup::TickRecv(FTickState& State)
+void FSocketGroup::RecvInternal(FTickState& State)
 {
-	// A helper lambda to use for iterating the activity list
-	auto Next = [this, Iter=List] () mutable -> FActivity*
-	{
-		FActivity* Activity = Iter;
-		if (Activity != nullptr)
-		{
-			Iter = Activity->Next;
-		}
-		return Activity;
-	};
+	check(Recv != nullptr);
 
 	// Another helper lambda
 	auto IsReceiving = [] (const FActivity* Act)
@@ -2622,213 +2751,165 @@ void FSocketGroup::TickRecv(FTickState& State)
 		return (Act->State >= EState::RecvMessage) & (Act->State < EState::RecvDone);
 	};
 
-	FActivity* Activity = Next();
+	FActivity* Activity = Recv;
+	check(IsReceiving(Activity));
 
-	// All receives first
-	while ((Activity != nullptr) && IsReceiving(Activity))
-	{
-		int32 Result = DoRecv(Activity, Socket, State.RecvAllowance);
+	int32 Result = DoRecv(Activity, Socket, State.RecvAllowance);
 
-		// Any sort of error here is unrecoverable
-		if (Result < 0)
-		{
-			Fail(State, Activity->ErrorReason);
-			return;
-		}
+	// Any sort of error here is unrecoverable
+	if (Result < 0)
+	{
+		Fail(State, Activity->ErrorReason);
+		return;
+	}
 
-		IsKeepAlive &= Activity->IsKeepAlive;
+	IsKeepAlive &= Activity->IsKeepAlive;
 
-		// If there was no data available this is far as receiving can go
-		if (Result > 0)
+	// If we've only a small amount left to receive we can start more work
+	if (IsKeepAlive & (Recv->Next == nullptr))
+	{
+		uint32 Remaining = Activity_RemainingKiB(Activity);
+		if (Remaining < uint32(GIasHttpRecvWorkThresholdKiB))
 		{
-			bWaiting = true;
-			State.Waiters->Emplace(Socket, FSocket::FWaiter::EWhat::Recv);
-
-			Activity = Next();
-			for (; (Activity != nullptr) && IsReceiving(Activity); Activity = Next());
-			break;
+			if (FActivity* Next = State.Work->PopActivity(); Next != nullptr)
+			{
+				Trace(Activity, ETrace::StartWork);
+	
+				check(Send == nullptr);
+				Send = Next;
+				SendInternal(State);
+			}
 		}
+	}
 
-		// If we're still in a receiving state we will just try again otherwise it
-		// is finished and we will let DoneList recipient finish it off.
-		if (IsReceiving(Activity))
-		{
-			continue;
-		}
+	// If there was no data available this is far as receiving can go
+	if (bWaiting = (Result > 0); bWaiting)
+	{
+		return;
+	}
 
-		check(Activity == List);
-		if ((List = Activity->Next) == nullptr)
-		{
-			ListTail = nullptr;
-		}
-		--NumActive;
+	// If we're still in a receiving state we will just try again otherwise it
+	// is finished and we will let DoneList recipient finish it off.
+	if (IsReceiving(Activity))
+	{
+		return;
+	}
 
-		DoRecvDone(Activity);
-		Done(State, Activity);
+	DoRecvDone(Activity);
 
-		// If the server wants to close the socket we need to rewind subsequent
-		// activities. Some may have already been sent down the pipeline
-		if (IsKeepAlive == 0)
-		{
-			while ((Activity = Next()) != nullptr)
-			{
-				if (int32 Rewound = Activity_Rewind(Activity); Rewound >= 0)
-				{
-					NumActive -= !!Rewound;
-					continue;
-				}
-
-				Fail(State, "Unable to rewind on keep-alive close");
-				break;
-			}
+	Recv = Activity->Next;
+	Activity->Next = State.DoneList;
+	State.DoneList = Activity;
 
-			Socket = FSocket();
-			return;
-		}
-		Activity = Next();
+	// If the server wants to close the socket we need to rewind the send
+	if (IsKeepAlive != 0)
+	{
+		return;
 	}
 
-	if (IsKeepAlive && Activity != nullptr)
+	if (Send != nullptr && Activity_Rewind(Send) < 0)
 	{
-		return TickSend(State, Activity);
+		Fail(State, "Unable to rewind on keep-alive close");
+		return;
 	}
+
+	Socket = FSocket();
 }
 
 ////////////////////////////////////////////////////////////////////////////////
-void FSocketGroup::TickSend(FTickState& State, FActivity* Activity)
+void FSocketGroup::SendInternal(FTickState& State)
 {
 	check(IsKeepAlive == 1);
+	check(Send != nullptr);
 
-	for (FActivity* Next; Activity != nullptr; Activity = Next)
-	{
-		check(Activity->State == FActivity::EState::Send);
-
-		// If we have maxed out the pipeline there's nothing more we can do
-		if (NumActive == PipelineLength)
-		{
-			break;
-		}
+	FActivity* Activity = Send;
 
-		int32 Result = DoSend(Activity, Socket);
+	int32 Result = DoSend(Activity, Socket);
 
-		if (Result == int32(FSocket::EResult::Wait))
-		{
-			// For now we'll not add the socket as a waiter. It is unlikely that
-			// we send enough to need to wait currently.
-			break;
-		}
+	if (Result == int32(FSocket::EResult::Wait))
+	{
+		// For now we'll not add the socket as a waiter. It is unlikely that we
+		// send enough to need to wait currently.
+		return;
+	}
 
-		Next = Activity->Next;
+	if (Result < 0)
+	{
+		Fail(State, Activity->ErrorReason);
+		return;
+	}
 
-		if (Result < 0)
-		{
-			Fail(State, Activity->ErrorReason);
-			return;
-		}
+	Send = nullptr;
 
-		++NumActive;
+	// Pass along this send to be received
+	if (Recv == nullptr)
+	{
+		Recv = Activity;
+		return;
 	}
+
+	check(Recv->Next == nullptr);
+	Recv->Next = Activity;
 }
 
 ////////////////////////////////////////////////////////////////////////////////
-void FSocketGroup::TickCancels(FTickState& State)
+bool FSocketGroup::Tick(FTickState& State)
 {
-	if (State.Cancels == 0 || (State.Cancels & ActiveSlots) == 0)
+	if (Send != nullptr)
 	{
-		return;
+		SendInternal(State);
 	}
 
-	// We are going to rebuild the list ot activities to maintain order as the
-	// activity list is singular.
-
-	check(List != nullptr);
-	FActivity* Activity = List;
-	List = ListTail = nullptr;
-	ActiveSlots = 0;
-
-	for (FActivity* Next; Activity != nullptr; Activity = Next)
+	if (Recv != nullptr)
 	{
-		Next = Activity->Next;
-
-		// Is this one of the activities we'd like to cancel?
-		bool bCancel = false;
-		if (uint64 Slot = 1ull << Activity->Slot; State.Cancels & Slot)
-		{
-			// We can't cancel things that have already been (or are being) requested
-			bCancel = (Activity->State == FActivity::EState::Send) & (Activity->StateParam == 0);
-		}
-
-		if (!bCancel)
-		{
-			AddActivity(Activity);
-			continue;
-		}
-
-		DoCancel(Activity);
-		Done(State, Activity);
+		RecvInternal(State);
 	}
+
+	return IsKeepAlive | !!(UPTRINT(Send) | UPTRINT(Recv));
 }
 
 ////////////////////////////////////////////////////////////////////////////////
-void FSocketGroup::Tick(FTickState& State, FHost& Host)
+void FSocketGroup::TickSend(FTickState& State, FHost& Host)
 {
-	// Nothing to tick if we've no work or are still waiting on read/write
-	if (List == nullptr || bWaiting)
+	// This path is only for those that are idle and have nothing to do
+	if (Send != nullptr || Recv != nullptr)
 	{
 		return;
 	}
 
-	TickCancels(State);
-
-	if (Socket.IsValid())
-	{
-		return TickRecv(State);
-	}
+	// Failing will try and recover work which we don't want to happen yet
+	FActivity* Pending = State.Work->PopActivity();
+	check(Pending != nullptr);
 
 	// We don't have a connected socket on first use, or if a keep-alive:close
 	// was received from the server. So we connect here.
-	FResult Result;
+	bool bWillBlock = false;
 	if (!Socket.IsValid())
 	{
 		IsKeepAlive = 1;
-		Result = Host.Connect(Socket);
+		FResult Result = Host.Connect(Socket);
 
-		// Non-blocking connect
-		if (Result.GetValue() == 0)
+		// We failed to connect, let's bail.
+		if (Result.GetValue() < 0)
 		{
-			bWaiting = true;
-			State.Waiters->Emplace(Socket, FSocket::FWaiter::EWhat::Send);
+			Pending->Next = Recv;
+			Recv = Pending;
+			Fail(State, Result.GetMessage());
 			return;
 		}
 
-		// Completed connect
-		if (Result.GetValue() > 0)
-		{
-			return TickSend(State, List);
-		}
-
-		IsKeepAlive = 0;
+		bWillBlock = (Result.GetValue() == 0);
 	}
 
-	// Fail all this group's work as no connection was established
-	Fail(State, Result.GetMessage());
-}
+	Send = Pending;
 
-////////////////////////////////////////////////////////////////////////////////
-void FSocketGroup::AddActivity(FActivity* Activity)
-{
-	// We use a tail pointer here to maintain order that requests were made
-
-	check(Activity->Next == nullptr);
-
-	if (ListTail != nullptr)
+	if (!bWillBlock)
 	{
-		ListTail->Next = Activity;
+		return SendInternal(State);
 	}
-	List = (List == nullptr) ? Activity : List;
-	ListTail = Activity;
 
-	ActiveSlots |= (1ull << Activity->Slot);
+	// Non-blocking connect
+	bWaiting = true;
 }
 
 
@@ -2837,30 +2918,45 @@ void FSocketGroup::AddActivity(FActivity* Activity)
 class FHostGroup
 {
 public:
-								FHostGroup(FHost& InHost) : Host(InHost) {}
-	bool						IsBusy() const	{ return !SocketGroups.IsEmpty(); }
+								FHostGroup(FHost& InHost);
+	bool						IsBusy() const	{ return BusyCount != 0; }
 	const FHost&				GetHost() const	{ return Host; }
 	void						Tick(FTickState& State);
 	void						AddActivity(FActivity* Activity);
 
 private:
-	using FSocketGroupPtr = TUniquePtr<FSocketGroup>;
-	using FSocketGroups = TArray<FSocketGroupPtr>;
-	void						ScatterPendings();
 	int32						Wait(int32 PollTimeoutMs);
-	TArray<FSocket::FWaiter>	Waiters;
-	FSocketGroups				SocketGroups;
+	TArray<FSocketGroup>		SocketGroups;
+	FWorkQueue					Work;
 	FHost&						Host;
-	FActivity*					PendingList = nullptr;
-	uint32						ScatterIndex = 0;
+	uint32						BusyCount = 0;
 
 	UE_NONCOPYABLE(FHostGroup);
 };
 
+////////////////////////////////////////////////////////////////////////////////
+FHostGroup::FHostGroup(FHost& InHost)
+: Host(InHost)
+{
+	uint32 Num = InHost.GetMaxConnections();
+	SocketGroups.SetNum(Num);
+}
+
 ////////////////////////////////////////////////////////////////////////////////
 int32 FHostGroup::Wait(int32 PollTimeoutMs)
 {
-	if (Waiters.Num() == 0)
+	// Collect groups that are waiting on something
+	TArray<FSocket::FWaiter, TFixedAllocator<64>> Waiters;
+	for (FSocketGroup& Group : SocketGroups)
+	{
+		FSocket::FWaiter Waiter = Group.GetWaiter();
+		if (Waiter.IsValid())
+		{
+			Waiters.Add(Waiter);
+		}
+	}
+
+	if (Waiters.IsEmpty())
 	{
 		return 0;
 	}
@@ -2897,10 +2993,10 @@ int32 FHostGroup::Wait(int32 PollTimeoutMs)
 		}
 
 		auto* Candidate = (FSocket*)(Waiters[i].Candidate);
-		auto Pred = [Candidate] (auto& Lhs) { return *Lhs == Candidate; };
-		FSocketGroupPtr* Group = SocketGroups.FindByPredicate(Pred);
+		auto Pred = [Candidate] (auto& Lhs) { return Lhs == Candidate; };
+		FSocketGroup* Group = SocketGroups.FindByPredicate(Pred);
 		check(Group != nullptr);
-		(*Group)->Unwait();
+		Group->Unwait();
 
 		Waiters.RemoveAtSwap(i, 1, false);
 		--n, --i, ++Count;
@@ -2913,113 +3009,51 @@ int32 FHostGroup::Wait(int32 PollTimeoutMs)
 ////////////////////////////////////////////////////////////////////////////////
 void FHostGroup::Tick(FTickState& State)
 {
-	ScatterPendings();
+	State.Work = &Work;
 
-	// Do any waits left over from the last pass through
-	if (int32 Result = Wait(State.PollTimeoutMs); Result < 0)
+	if (BusyCount = Work.HasWork(); BusyCount)
 	{
-		const char* Reason = (Result == MIN_int32)
-			? "FailTimeout hit"
-			: "poll() returned an unexpected error";
+		Work.TickCancels(State);
 
-		for (FSocketGroupPtr& Group : SocketGroups)
+		// Get available work out on idle sockets as soon as possible
+		for (FSocketGroup& Group : SocketGroups)
 		{
-			Group->Fail(State, Reason);
-		}
-
-		Waiters.Empty();
-		SocketGroups.SetNum(0, false);
-		return;
-	}
-
-	// If we are still waiting on all sockets groups there is nothing to do.
-	if (Waiters.Num() == SocketGroups.Num())
-	{
-		return;
-	}
-
-	// Tick everything.
-	State.Waiters = &Waiters;
-	for (FSocketGroupPtr& Group : SocketGroups)
-	{
-		Group->Tick(State, Host);
-	}
+			if (!Work.HasWork())
+			{
+				break;
+			}
 
-	// Clear out any spent socket groups.
-	for (uint32 i = 0, n = SocketGroups.Num(); i < n; ++i)
-	{
-		FSocketGroupPtr& Group = SocketGroups[i];
-		if (Group->IsBusy())
-		{
-			continue;
+			Group.TickSend(State, Host);
 		}
-
-		check(!Group->IsWaiting());
-		SocketGroups.RemoveAtSwap(i, 1, false);
-		--n, --i;
-	}
-}
-
-////////////////////////////////////////////////////////////////////////////////
-void FHostGroup::ScatterPendings()
-{
-	FActivity* Iter = PendingList;
-	if (Iter == nullptr)
-	{
-		return;
 	}
-	PendingList = nullptr;
 
-	// Maximize the number of connections to the host first
-	for (uint32 i = SocketGroups.Num(), n = Host.GetMaxConnections(); i < n; ++i)
+	// Wait on the groups that are
+	if (int32 Result = Wait(State.PollTimeoutMs); Result < 0)
 	{
-		FActivity* Activity = Iter;
-		Iter = Iter->Next;
-		Activity->Next = nullptr;
-
-		FSocketGroupPtr& Group = SocketGroups.Emplace_GetRef();
-		Group = MakeUnique<FSocketGroup>(Host.GetPipelineLength());
-		Group->AddActivity(Activity);
+		const char* Reason = (Result == MIN_int32)
+			? "FailTimeout hit"
+			: "poll() returned an unexpected error";
 
-		if (Iter == nullptr)
+		for (FSocketGroup& Group : SocketGroups)
 		{
-			return;
+			Group.Fail(State, Reason);
 		}
-	}
 
-	// Next we'll favour socket groups that have come back for more.
-	/*
-	 * ...not currently implemented
-	 *
-	do
-	{
+		SocketGroups.SetNum(0, false);
+		return;
 	}
-	while (Iter != PendingList.end());
-	*/
 
-	// Finally we'll simply round-robin distribute items
-	for (uint32 Num = SocketGroups.Num(); Iter != nullptr; ++ScatterIndex)
+	// Tick everything, starting with groups that are maybe closest to finishing
+	for (FSocketGroup& Group : SocketGroups)
 	{
-		if (ScatterIndex >= Num)
-		{
-			ScatterIndex = 0;
-		}
-
-		FActivity* Activity = Iter;
-		Iter = Iter->Next;
-		Activity->Next = nullptr;
-
-		FSocketGroupPtr& Group = SocketGroups[ScatterIndex];
-		Group->AddActivity(Activity);
+		BusyCount += (Group.Tick(State) == true);
 	}
 }
 
 ////////////////////////////////////////////////////////////////////////////////
 void FHostGroup::AddActivity(FActivity* Activity)
 {
-	check(Activity->Next == nullptr);
-	Activity->Next = PendingList;
-	PendingList = Activity;
+	Work.AddActivity(Activity);
 }
 
 
@@ -3230,7 +3264,7 @@ uint32 FEventLoop::FImpl::Tick(int32 PollTimeoutMs)
 			continue;
 		}
 
-		Groups.RemoveAtSwap(i);
+		Groups.RemoveAtSwap(i, 1, false);
 		--n, --i;
 	}
 
@@ -3908,6 +3942,4 @@ IOSTOREONDEMAND_API void IasHttpTest(const ANSICHAR* TestHost="localhost")
 
 // }}}
 
-} // namespace UE::IO::IAS
-
-/* vim: set noet : */
+} // namespace UE::IO::IAS::HTTP
