Ias: Add on screen debug message to show status of Ias.
Can be toggled with ias.DisplayOnScreenStatistics=true/false.

#rnx

diff --git a/Engine/Source/Runtime/Experimental/IoStoreOnDemand/Private/Statistics.cpp b/Engine/Source/Runtime/Experimental/IoStoreOnDemand/Private/Statistics.cpp
--- a/Engine/Source/Runtime/Experimental/IoStoreOnDemand/Private/Statistics.cpp
+++ b/Engine/Source/Runtime/Experimental/IoStoreOnDemand/Private/Statistics.cpp
@@ -6,8 +6,9 @@
 #include "AnalyticsEventAttribute.h"
 #include "HAL/IConsoleManager.h"
 #include "HAL/PlatformTime.h"
-#include "Misc/CoreDelegates.h"
 #include "IO/IoStoreOnDemand.h"
+#include "Internationalization/Internationalization.h"
+#include "Misc/CoreDelegates.h"
 
 LLM_DEFINE_TAG(Ias);
 
@@ -22,6 +23,12 @@ static FAutoConsoleVariableRef CVar_StatisticsLogInterval(
 	GIasStatisticsLogInterval,
 	TEXT("Enables and sets interval for periodic logging of statistics"));
 
+bool GIasDisplayOnScreenStatistics = false;
+static FAutoConsoleVariableRef CVar_DisplayOnScreenStatistics(
+	TEXT("ias.DisplayOnScreenStatistics"),
+	GIasDisplayOnScreenStatistics,
+	TEXT("Enables display of Ias on screen statistics"));
+
 bool GIasReportHttpAnalyticsEnabled = true;
 static FAutoConsoleVariableRef CVar_ReportHttpAnalytics(
 	TEXT("ias.ReportHttpAnalytics"),
@@ -277,6 +284,7 @@ CSV_DEFINE_STAT(Ias, HttpDurationMsMax);
 
 static FOnDemandIoBackendStats* GStatistics = nullptr;
 static FDelegateHandle GStatisticsEndFrameDelegateHandle;
+static FDelegateHandle GStatisticsOnScreenDelegateHandle;
 
 FOnDemandIoBackendStats::FOnDemandIoBackendStats()
 {
@@ -333,11 +341,35 @@ FOnDemandIoBackendStats::FOnDemandIoBackendStats()
 			}
 		}
 	});
+
+#define LOCTEXT_NAMESPACE "IAS"
+	GStatisticsOnScreenDelegateHandle = FCoreDelegates::OnGetOnScreenMessages.AddLambda(
+		[this] (FCoreDelegates::FSeverityMessageMap& Out)
+		{
+			if (GIasDisplayOnScreenStatistics)
+			{
+				FText Message = FText::Format(
+					LOCTEXT("IAS", "IAS: Cached:{0} KiB | Read:{1} KiB ({2}) | Downloaded:{3} KiB ({4}) {5} ms | Retries:{6} | Pending:{7}"),
+					GCacheCachedBytes.Get() >> 10,
+					GCacheReadBytes.Get() >> 10,
+					GCacheGetCount.Get(),
+					GHttpDownloadedBytes.Get() >> 10,
+					GHttpGetCount.Get(),
+					GHttpDurationMsAvg.Get(),
+					GHttpRetryCount.Get(),
+					GHttpPendingCount.Get()
+				);
+				Out.Add(FCoreDelegates::EOnScreenMessageSeverity::Info, Message);
+			}
+		}
+	);
+#undef LOCTEXT_NAMESPACE
 }
 
 FOnDemandIoBackendStats::~FOnDemandIoBackendStats()
 {
 	FCoreDelegates::OnEndFrame.Remove(GStatisticsEndFrameDelegateHandle);
+	FCoreDelegates::OnGetOnScreenMessages.Remove(GStatisticsOnScreenDelegateHandle);
 	GStatistics = nullptr;
 }
 
