Clean up IAS retry code

#rnx

- Removed FChunkRequest::HttpRetryCount as it is no longer used (the retry logic is inside of the request itself now)
- Added some code comments to the 'ias.HttpRetryCount' cvar to make its new use more apparent.

diff --git a/Engine/Source/Runtime/Experimental/IoStoreOnDemand/Private/OnDemandIoDispatcherBackend.cpp b/Engine/Source/Runtime/Experimental/IoStoreOnDemand/Private/OnDemandIoDispatcherBackend.cpp
--- a/Engine/Source/Runtime/Experimental/IoStoreOnDemand/Private/OnDemandIoDispatcherBackend.cpp
+++ b/Engine/Source/Runtime/Experimental/IoStoreOnDemand/Private/OnDemandIoDispatcherBackend.cpp
@@ -120,11 +120,16 @@ static FAutoConsoleVariableRef CVar_GIasHttpPipelineLength(
 	TEXT("Number of concurrent requests on one connection")
 );
 
+/**
+ *This is only applied when the connection was made to a single ServiceUrl rather than a DistributedUrl.
+ * In the latter case we will make two attempts on the primary CDN followed by a single attempt for the
+ * remaining CDN's to be tried in the order provided by the distributed endpoint.
+ */
 static int32 GIasHttpRetryCount = 2;
 static FAutoConsoleVariableRef CVar_IasHttpRetryCount(
 	TEXT("ias.HttpRetryCount"),
 	GIasHttpRetryCount,
-	TEXT("Number of HTTP request retries before failing the I/O request.")
+	TEXT("Number of HTTP request retries before failing the request (if connected to a service url rather than distributed endpoints).")
 );
 
 int32 GIasHttpTimeOutMs = 10 * 1000;
@@ -734,7 +739,6 @@ struct FChunkRequest
 		, StartTime(FPlatformTime::Cycles64())
 		, Priority(Request->Priority)
 		, RequestCount(1)
-		, HttpRetryCount(0)
 		, bCached(false)
 	{
 		check(Request && NextRequest == nullptr);
@@ -823,7 +827,6 @@ struct FChunkRequest
 	uint64 StartTime;
 	int32 Priority;
 	uint16 RequestCount;
-	uint16 HttpRetryCount;
 	bool bCached;
 	bool bCancelled = false;
 	EIoErrorCode CacheGetStatus;
