Iris
- PropertyReplicationFragment now only stores previous value for properties with RepNotify if they are about to be overwritten

diff --git a/Engine/Source/Runtime/Experimental/Iris/Core/Private/Iris/ReplicationState/PropertyReplicationState.cpp b/Engine/Source/Runtime/Experimental/Iris/Core/Private/Iris/ReplicationState/PropertyReplicationState.cpp
--- a/Engine/Source/Runtime/Experimental/Iris/Core/Private/Iris/ReplicationState/PropertyReplicationState.cpp
+++ b/Engine/Source/Runtime/Experimental/Iris/Core/Private/Iris/ReplicationState/PropertyReplicationState.cpp
@@ -290,7 +290,7 @@ bool FPropertyReplicationState::PollPropertyReplicationState(const void* RESTRIC
 	return IsDirty();
 }
 
-bool FPropertyReplicationState::PollPropertyReplicationStateForRepNotifies(const void* RESTRICT SrcStateData)
+bool FPropertyReplicationState::StoreCurrentPropertyReplicationStateForRepNotifies(const void* RESTRICT SrcStateData, const FPropertyReplicationState* NewStateToBeApplied)
 {
 	if (IsValid())
 	{
@@ -304,15 +304,21 @@ bool FPropertyReplicationState::PollPropertyReplicationStateForRepNotifies(const
 		const FReplicationStateMemberPropertyDescriptor* MemberPropertyDescriptors = Descriptor->MemberPropertyDescriptors;
 		const uint32 MemberCount = Descriptor->MemberCount;
 
+		// Copy all if this is a state with no changemask or if NewStateToBeApplied is not set
+		const bool bCopyAll = IsInitState() || NewStateToBeApplied == nullptr;
 		for (uint32 MemberIt = 0; MemberIt < MemberCount; ++MemberIt)
 		{
 			const FReplicationStateMemberDescriptor& MemberDescriptor = MemberDescriptors[MemberIt];
 			const FReplicationStateMemberPropertyDescriptor& MemberPropertyDescriptor = MemberPropertyDescriptors[MemberIt];
 
-			if (MemberPropertyDescriptor.RepNotifyFunction)
+			if (MemberPropertyDescriptor.RepNotifyFunction && (bCopyAll || NewStateToBeApplied->IsDirty(MemberIt)))
 			{
 				const FProperty* Property = MemberProperties[MemberIt];
-				PollPropertyValue(MemberIt, SrcBuffer + Property->GetOffset_ForGC() + Property->ElementSize*MemberPropertyDescriptor.ArrayIndex);
+
+				void* DstValue = StateBuffer + Descriptor->MemberDescriptors[MemberIt].ExternalMemberOffset;
+				const void* SrcValue = SrcBuffer + Property->GetOffset_ForGC() + Property->ElementSize*MemberPropertyDescriptor.ArrayIndex;
+
+				Private::InternalCopyPropertyValue(Descriptor, MemberIt, DstValue, SrcBuffer + Property->GetOffset_ForGC() + Property->ElementSize*MemberPropertyDescriptor.ArrayIndex);
 			}
 		}
 	}
diff --git a/Engine/Source/Runtime/Experimental/Iris/Core/Private/Iris/ReplicationSystem/PropertyReplicationFragment.cpp b/Engine/Source/Runtime/Experimental/Iris/Core/Private/Iris/ReplicationSystem/PropertyReplicationFragment.cpp
--- a/Engine/Source/Runtime/Experimental/Iris/Core/Private/Iris/ReplicationSystem/PropertyReplicationFragment.cpp
+++ b/Engine/Source/Runtime/Experimental/Iris/Core/Private/Iris/ReplicationSystem/PropertyReplicationFragment.cpp
@@ -37,8 +37,8 @@ FPropertyReplicationFragment::FPropertyReplicationFragment(EReplicationFragmentT
 			{
 				PrevReplicationState = MakeUnique<FPropertyReplicationState>(InDescriptor);
 
-				// Poll to get instance default for our previous state
-				PrevReplicationState->PollPropertyReplicationStateForRepNotifies(InOwner);
+				// Full store of initial value for repnotifies
+				PrevReplicationState->StoreCurrentPropertyReplicationStateForRepNotifies(InOwner, nullptr);
 
 				Traits |= EReplicationFragmentTraits::KeepPreviousState;
 			}
@@ -133,15 +133,15 @@ void FPropertyReplicationFragment::ApplyReplicatedState(FReplicationStateApplyCo
 {
 	IRIS_PROFILER_SCOPE(PropertyReplicationFragment_ApplyReplicatedState);
 
+	// Create a wrapping property replication state, cheap as we are simply injecting the already constructed state
+	const FPropertyReplicationState ReceivedState(Context.Descriptor, Context.StateBufferData.ExternalStateBuffer);
+
 	// If we do not rely on received data to issue rep notifies we need to store a copy of the local state before we apply the new received state.
 	if (!bUsePrevReceivedStateForOnReps && PrevReplicationState)
 	{
-		PrevReplicationState->PollPropertyReplicationStateForRepNotifies(Owner);
+		PrevReplicationState->StoreCurrentPropertyReplicationStateForRepNotifies(Owner, &ReceivedState);
 	}
 
-	// Create a wrapping property replication state, cheap as we are simply injecting the already constructed state
-	const FPropertyReplicationState ReceivedState(Context.Descriptor, Context.StateBufferData.ExternalStateBuffer);
-
 	// Just push the state data to owner
 	ReceivedState.PushPropertyReplicationState(Owner, static_cast<void*>(Owner));
 }
diff --git a/Engine/Source/Runtime/Experimental/Iris/Core/Public/Iris/ReplicationState/PropertyReplicationState.h b/Engine/Source/Runtime/Experimental/Iris/Core/Public/Iris/ReplicationState/PropertyReplicationState.h
--- a/Engine/Source/Runtime/Experimental/Iris/Core/Public/Iris/ReplicationState/PropertyReplicationState.h
+++ b/Engine/Source/Runtime/Experimental/Iris/Core/Public/Iris/ReplicationState/PropertyReplicationState.h
@@ -87,8 +87,10 @@ public:
 	*/
 	IRISCORE_API bool PollPropertyReplicationState(const void* RESTRICT SrcData);
 
-	/** Poll src data from properties which we want to store in order to determine if we need to call RepNotifies. */
-	IRISCORE_API bool PollPropertyReplicationStateForRepNotifies(const void* RESTRICT SrcData);
+	/** Certain rep notifies requires the current state to be stored before we overwrite it, this method will copy property values from src data, , where SrcData is a UClass/UStruct containing properties
+	 * if the property is marked as dirty in NewStateToBeApplied
+	 */
+	IRISCORE_API bool StoreCurrentPropertyReplicationStateForRepNotifies(const void* RESTRICT SrcData, const FPropertyReplicationState* NewStateToBeApplied);
 
 	/** Push received state data to properties in DstData buffer. Note: DstData is a UClass/UStruct containing properties. Updates representation in DstStateBuffer and mark destination properties as dirty if updated. */
 	IRISCORE_API void PushPropertyReplicationState(const UObject* Owner, void* RESTRICT DstData, bool bPushAll = false) const;
