Chaos: Fixed event buffer growing unbounded when running in single-threaded modes. Previously the event consumer buffers were only cleared if we were in modes other than single-buffered. But if we don't clear it in single buffer mode we just keep accumulating events and re-dispatching them every frame.

diff --git a/Engine/Source/Runtime/Experimental/Chaos/Private/EventManager.cpp b/Engine/Source/Runtime/Experimental/Chaos/Private/EventManager.cpp
--- a/Engine/Source/Runtime/Experimental/Chaos/Private/EventManager.cpp
+++ b/Engine/Source/Runtime/Experimental/Chaos/Private/EventManager.cpp
@@ -70,10 +70,7 @@ namespace Chaos
 		{
 			if (EventContainer)
 			{
-				if (BufferMode != EMultiBufferMode::Single)
-				{
-					EventContainer->ResetConsumerBuffer();
-				}
+				EventContainer->ResetConsumerBuffer();
 				EventContainer->FlipBufferIfRequired();
 			}
 		}
