Add "IsProbe" and "ConvertToNonProbe" options to contact modification

diff --git a/Engine/Source/Runtime/Experimental/Chaos/Private/Chaos/ContactModification.cpp b/Engine/Source/Runtime/Experimental/Chaos/Private/Chaos/ContactModification.cpp
--- a/Engine/Source/Runtime/Experimental/Chaos/Private/Chaos/ContactModification.cpp
+++ b/Engine/Source/Runtime/Experimental/Chaos/Private/Chaos/ContactModification.cpp
@@ -21,6 +21,16 @@ namespace Chaos
 		Modifier->ConvertToProbeConstraint(*Constraint);
 	}
 
+	void FContactPairModifier::ConvertToNonProbe()
+	{
+		Modifier->ConvertToNonProbeConstraint(*Constraint);
+	}
+
+	bool FContactPairModifier::GetIsProbe() const
+	{
+		return Constraint->GetIsProbe();
+	}
+
 	int32 FContactPairModifier::GetNumContacts() const
 	{
 		return Constraint->GetManifoldPoints().Num();
@@ -456,6 +466,19 @@ namespace Chaos
 		}
 	}
 
+	FGeometryParticleHandle* FContactPairModifier::GetOtherParticle(FGeometryParticleHandle* Particle) const
+	{
+		if (Particle == Constraint->GetParticle0())
+		{
+			return Constraint->GetParticle1();
+		}
+		else if (Particle == Constraint->GetParticle1())
+		{
+			return Constraint->GetParticle0();
+		}
+		return nullptr;
+	}
+
 	TVec2<FGeometryParticleHandle*> FContactPairModifier::GetParticlePair() const
 	{
 		return { Constraint->GetParticle0(), Constraint->GetParticle1() };
@@ -611,6 +634,13 @@ namespace Chaos
 		Constraint.SetIsProbe(true);
 	}
 
+	void FCollisionContactModifier::ConvertToNonProbeConstraint(FPBDCollisionConstraint& Constraint)
+	{
+		Constraint.SetModifierApplied();
+
+		Constraint.SetIsProbe(false);
+	}
+
 	void FCollisionContactModifier::MarkConstraintForManifoldUpdate(FPBDCollisionConstraint& Constraint)
 	{
 		NeedsManifoldUpdate.Add(&Constraint);
diff --git a/Engine/Source/Runtime/Experimental/Chaos/Public/Chaos/ContactModification.h b/Engine/Source/Runtime/Experimental/Chaos/Public/Chaos/ContactModification.h
--- a/Engine/Source/Runtime/Experimental/Chaos/Public/Chaos/ContactModification.h
+++ b/Engine/Source/Runtime/Experimental/Chaos/Public/Chaos/ContactModification.h
@@ -54,6 +54,16 @@ namespace Chaos
 		*/
 		CHAOS_API void ConvertToProbe();
 
+		/**
+		 * Convert the constraint from probe to a regular physical collision
+		 */
+		CHAOS_API void ConvertToNonProbe();
+
+		/**
+		 * Return true if this constraint is a probe type
+		 */
+		CHAOS_API bool GetIsProbe() const;
+
 		/**
 		* @return Number of contact points in constraint pair. ContactPointIdx must be below number of contacts.
 		*/
@@ -221,6 +231,12 @@ namespace Chaos
 		*/
 		CHAOS_API void ModifyInvMassScale(FReal InvMassScale, int32 ParticleIdx);
 
+		/*
+		* Get the other particle in the constraint - if the input particle isn't in the constraint,
+		* get null
+		*/
+		CHAOS_API FGeometryParticleHandle* GetOtherParticle(FGeometryParticleHandle* Particle) const;
+
 		/*
 		* Retrieve physics handles for particles in contact pair.
 		*/
@@ -494,6 +510,9 @@ namespace Chaos
 		// but will not produce impulses
 		CHAOS_API void ConvertToProbeConstraint(FPBDCollisionConstraint& Constraint);
 
+		// Turn this constraint from a probe into a regular constraint
+		CHAOS_API void ConvertToNonProbeConstraint(FPBDCollisionConstraint& Constraint);
+
 		CHAOS_API void MarkConstraintForManifoldUpdate(FPBDCollisionConstraint& Constraint);
 
 		// Update manifolds of modified constraints.
