Cache gets from the MemCache could accidentally miss.

#rnx

diff --git a/Engine/Source/Runtime/Experimental/IoStoreOnDemand/Private/JournaledCache.cpp b/Engine/Source/Runtime/Experimental/IoStoreOnDemand/Private/JournaledCache.cpp
--- a/Engine/Source/Runtime/Experimental/IoStoreOnDemand/Private/JournaledCache.cpp
+++ b/Engine/Source/Runtime/Experimental/IoStoreOnDemand/Private/JournaledCache.cpp
@@ -2040,13 +2040,14 @@ EIoErrorCode FJournaledCache::Get(const FIoHash& Key, FIoBuffer& OutData)
 
 	uint64 InnerKey = ReduceKey(Key);
 
-	auto Entry = Cache->Get(InnerKey, OutData);
-	if (Entry == 0)
+	uint64 GetKey = Cache->Get(InnerKey, OutData);
+	if (OutData.GetData() != nullptr)
 	{
-		return EIoErrorCode::NotFound;
+		return EIoErrorCode::Ok;
 	}
-	
-	return (OutData.GetData() != nullptr) ? EIoErrorCode::Ok : EIoErrorCode::FileNotOpen;
+
+	// "File not open" to indicate that we have Key, just not to hand.
+	return (GetKey == 0) ? EIoErrorCode::NotFound : EIoErrorCode::FileNotOpen;
 }
 
 ////////////////////////////////////////////////////////////////////////////////
