Convert RestCollection in single precision
- This change won't have much impact on memory when building in game, but the impact will be when loading. Also it harmonize the Dynamic Collection and the Rest Collection. And so it optimizes the access, and conversion to transforms.
- DynamicCollection has a function set back transforms to its RestCollection, and free memory
- Remove function from ManageArrayCollection makes sure to free memory now.

diff --git a/Engine/Plugins/Experimental/ChaosCaching/Source/ChaosCaching/Private/Chaos/Adapters/GeometryCollectionComponentCacheAdapter.cpp b/Engine/Plugins/Experimental/ChaosCaching/Source/ChaosCaching/Private/Chaos/Adapters/GeometryCollectionComponentCacheAdapter.cpp
--- a/Engine/Plugins/Experimental/ChaosCaching/Source/ChaosCaching/Private/Chaos/Adapters/GeometryCollectionComponentCacheAdapter.cpp
+++ b/Engine/Plugins/Experimental/ChaosCaching/Source/ChaosCaching/Private/Chaos/Adapters/GeometryCollectionComponentCacheAdapter.cpp
@@ -661,7 +661,7 @@ namespace Chaos
 				RestTransforms.SetNum(NumTransforms);
 				for (int32 Idx = 0; Idx < NumTransforms; ++Idx)
 				{
-					RestTransforms[Idx] = GeometryCollection->Transform[Idx];
+					RestTransforms[Idx] = FTransform(GeometryCollection->Transform[Idx]);
 				}
 
 				for (int32 CacheIdx = 0; CacheIdx < NumCacheTransforms; ++CacheIdx)
@@ -786,7 +786,7 @@ namespace Chaos
 					InitialTransforms.SetNum(NumTransforms);
 					for (int32 Idx = 0; Idx < NumTransforms; ++Idx)
 					{
-						InitialTransforms[Idx] = GeometryCollection->Transform[Idx];
+						InitialTransforms[Idx] = FTransform(GeometryCollection->Transform[Idx]);
 					}
 
 					for (int32 CacheIdx = 0; CacheIdx < NumCacheTransforms; ++CacheIdx)
@@ -814,7 +814,7 @@ namespace Chaos
 
 	TArray<int32> FGeometryCollectionCacheAdapter::GatherAllBreaksUpToTime(UChaosCache* InCache, float InTime) const
 	{
-		// Evaluate all breaking event that have occured from the beginning of the cache up to the specified time.
+		// Evaluate all breaking event that have occurred from the beginning of the cache up to the specified time.
 		TArray<int32> ReleaseIndices;
 
 		FPlaybackTickRecord TickRecord;
diff --git a/Engine/Plugins/Experimental/ChaosEditor/Source/FractureEditor/Private/FractureEditorMode.cpp b/Engine/Plugins/Experimental/ChaosEditor/Source/FractureEditor/Private/FractureEditorMode.cpp
--- a/Engine/Plugins/Experimental/ChaosEditor/Source/FractureEditor/Private/FractureEditorMode.cpp
+++ b/Engine/Plugins/Experimental/ChaosEditor/Source/FractureEditor/Private/FractureEditorMode.cpp
@@ -613,7 +613,7 @@ void UFractureEditorMode::GetComponentGlobalBounds(UGeometryCollectionComponent*
 	TSharedPtr<FGeometryCollection, ESPMode::ThreadSafe> GeometryCollectionPtr = GeometryCollection->GetGeometryCollection();
 	FGeometryCollection* OutGeometryCollection = GeometryCollectionPtr.Get();
 
-	const TManagedArray<FTransform>& Transform = OutGeometryCollection->Transform;
+	const TManagedArray<FTransform3f>& Transform = OutGeometryCollection->Transform;
 	const TManagedArray<FBox>& BoundingBox = OutGeometryCollection->BoundingBox;
 	const TManagedArray<int32>& TransformToGeometryIndex = OutGeometryCollection->TransformToGeometryIndex;
 	const TManagedArray<int32>& Parent = OutGeometryCollection->Parent;
diff --git a/Engine/Plugins/Experimental/ChaosEditor/Source/FractureEditor/Private/FractureEditorModeToolkit.cpp b/Engine/Plugins/Experimental/ChaosEditor/Source/FractureEditor/Private/FractureEditorModeToolkit.cpp
--- a/Engine/Plugins/Experimental/ChaosEditor/Source/FractureEditor/Private/FractureEditorModeToolkit.cpp
+++ b/Engine/Plugins/Experimental/ChaosEditor/Source/FractureEditor/Private/FractureEditorModeToolkit.cpp
@@ -1941,7 +1941,7 @@ void FFractureEditorModeToolkit::UpdateExplodedVectors(UGeometryCollectionCompon
 		check(OutGeometryCollection->HasAttribute("ExplodedVector", FGeometryCollection::TransformGroup));
 
 		TManagedArray<FVector3f>& ExplodedVectors = OutGeometryCollection->ModifyAttribute<FVector3f>("ExplodedVector", FGeometryCollection::TransformGroup);
-		const TManagedArray<FTransform>& Transform = OutGeometryCollection->GetAttribute<FTransform>("Transform", FGeometryCollection::TransformGroup);
+		const TManagedArray<FTransform3f>& Transform = OutGeometryCollection->GetAttribute<FTransform3f>("Transform", FGeometryCollection::TransformGroup);
 		const TManagedArray<int32>& TransformToGeometryIndex = OutGeometryCollection->GetAttribute<int32>("TransformToGeometryIndex", FGeometryCollection::TransformGroup);
 		const TManagedArray<FBox>& BoundingBox = OutGeometryCollection->GetAttribute<FBox>("BoundingBox", FGeometryCollection::GeometryGroup);
 
diff --git a/Engine/Plugins/Experimental/ChaosEditor/Source/FractureEditor/Private/FractureToolConvert.cpp b/Engine/Plugins/Experimental/ChaosEditor/Source/FractureEditor/Private/FractureToolConvert.cpp
--- a/Engine/Plugins/Experimental/ChaosEditor/Source/FractureEditor/Private/FractureToolConvert.cpp
+++ b/Engine/Plugins/Experimental/ChaosEditor/Source/FractureEditor/Private/FractureToolConvert.cpp
@@ -110,7 +110,7 @@ TArray<FFractureToolContext> UFractureToolConvert::GetFractureToolContexts() con
 		FullSelection.ConvertSelectionToRigidNodes();
 		
 		// Update global transforms and bounds -- TODO: pull this bounds update out to a shared function?
-		const TManagedArray<FTransform>& Transform = FullSelection.GetGeometryCollection()->Transform;
+		const TManagedArray<FTransform3f>& Transform = FullSelection.GetGeometryCollection()->Transform;
 		const TManagedArray<int32>& TransformToGeometryIndex = FullSelection.GetGeometryCollection()->TransformToGeometryIndex;
 		const TManagedArray<FBox>& BoundingBoxes = FullSelection.GetGeometryCollection()->BoundingBox;
 
diff --git a/Engine/Plugins/Experimental/ChaosEditor/Source/FractureEditor/Private/FractureToolCutter.cpp b/Engine/Plugins/Experimental/ChaosEditor/Source/FractureEditor/Private/FractureToolCutter.cpp
--- a/Engine/Plugins/Experimental/ChaosEditor/Source/FractureEditor/Private/FractureToolCutter.cpp
+++ b/Engine/Plugins/Experimental/ChaosEditor/Source/FractureEditor/Private/FractureToolCutter.cpp
@@ -173,7 +173,7 @@ TArray<FFractureToolContext> UFractureToolCutterBase::GetFractureToolContexts()
 			FullSelection.RandomReduceSelection(CutterSettings->ChanceToFracture);
 
 			// Update global transforms and bounds		
-			const TManagedArray<FTransform>& Transform = FullSelection.GetGeometryCollection()->GetAttribute<FTransform>("Transform", FGeometryCollection::TransformGroup);
+			const TManagedArray<FTransform3f>& Transform = FullSelection.GetGeometryCollection()->GetAttribute<FTransform3f>("Transform", FGeometryCollection::TransformGroup);
 			const TManagedArray<int32>& TransformToGeometryIndex = FullSelection.GetGeometryCollection()->GetAttribute<int32>("TransformToGeometryIndex", FGeometryCollection::TransformGroup);
 			const TManagedArray<FBox>& BoundingBoxes = FullSelection.GetGeometryCollection()->GetAttribute<FBox>("BoundingBox", FGeometryCollection::GeometryGroup);
 
diff --git a/Engine/Plugins/Experimental/ChaosEditor/Source/FractureEditor/Private/FractureToolEmbed.cpp b/Engine/Plugins/Experimental/ChaosEditor/Source/FractureEditor/Private/FractureToolEmbed.cpp
--- a/Engine/Plugins/Experimental/ChaosEditor/Source/FractureEditor/Private/FractureToolEmbed.cpp
+++ b/Engine/Plugins/Experimental/ChaosEditor/Source/FractureEditor/Private/FractureToolEmbed.cpp
@@ -67,7 +67,7 @@ void UFractureToolAddEmbeddedGeometry::Execute(TWeakPtr<FFractureEditorModeToolk
 			UGeometryCollection* FracturedGeometryCollection = GeometryCollectionEdit.GetRestCollection();
 
 			FFractureToolContext::FGeometryCollectionPtr GeometryCollection = Context.GetGeometryCollection();
-			const TManagedArray<FTransform>& Transform = GeometryCollection->Transform;
+			const TManagedArray<FTransform3f>& Transform = GeometryCollection->Transform;
 			const TManagedArray<int32>& Parent = GeometryCollection->Parent;
 
 			int32 StartTransformCount = Transform.Num();
@@ -195,7 +195,7 @@ void UFractureToolAutoEmbedGeometry::Execute(TWeakPtr<FFractureEditorModeToolkit
 
 				FGeometryCollection* GeometryCollection = Context.GetGeometryCollection().Get();
 				FGeometryCollectionConvexUtility::FGeometryCollectionConvexData ConvexData = FGeometryCollectionConvexUtility::GetValidConvexHullData(GeometryCollection);
-				const TManagedArray<FTransform>& Transform = GeometryCollection->Transform;
+				const TManagedArray<FTransform3f>& Transform = GeometryCollection->Transform;
 				const TManagedArray<int32>& Parent = GeometryCollection->Parent;
 				const TManagedArray<int32>& SimulationType = GeometryCollection->SimulationType;
 				TArray<FTransform> BoneGlobalTransforms;
@@ -233,7 +233,7 @@ void UFractureToolAutoEmbedGeometry::Execute(TWeakPtr<FFractureEditorModeToolkit
 			if (ClosestGeometryCollection.IsValid() && ClosestComponent && ClosestContext)
 			{
 				// Which bone points to the closest convex?
-				const TManagedArray<FTransform>& Transform = ClosestGeometryCollection->Transform;
+				const TManagedArray<FTransform3f>& Transform = ClosestGeometryCollection->Transform;
 				const TManagedArray<int32>& Parent = ClosestGeometryCollection->Parent;
 				const TManagedArray<TSet<int32>>& TransformToConvexIndices =
 					ClosestGeometryCollection->GetAttribute<TSet<int32>>("TransformToConvexIndices", FTransformCollection::TransformGroup);
@@ -378,7 +378,7 @@ void UFractureToolFlushEmbeddedGeometry::Execute(TWeakPtr<FFractureEditorModeToo
 			}
 			
 			const TManagedArray<int32>& ExemplarIndex = Context.GetGeometryCollection()->GetAttribute<int32>("ExemplarIndex", FGeometryCollection::TransformGroup);
-			const TManagedArray<FTransform>& Transforms = Context.GetGeometryCollection()->GetAttribute<FTransform>("Transform", FGeometryCollection::TransformGroup);
+			const TManagedArray<FTransform3f>& Transforms = Context.GetGeometryCollection()->GetAttribute<FTransform3f>("Transform", FGeometryCollection::TransformGroup);
 			const TManagedArray<int32>& ParentIndices = Context.GetGeometryCollection()->GetAttribute<int32>("Parent", FGeometryCollection::TransformGroup);
 
 			TArray<FMatrix> GlobalMatrices;
diff --git a/Engine/Plugins/Experimental/ChaosEditor/Source/FractureEditor/Private/FractureToolUV.cpp b/Engine/Plugins/Experimental/ChaosEditor/Source/FractureEditor/Private/FractureToolUV.cpp
--- a/Engine/Plugins/Experimental/ChaosEditor/Source/FractureEditor/Private/FractureToolUV.cpp
+++ b/Engine/Plugins/Experimental/ChaosEditor/Source/FractureEditor/Private/FractureToolUV.cpp
@@ -269,7 +269,7 @@ TArray<FFractureToolContext> UFractureToolAutoUV::GetFractureToolContexts() cons
 		FullSelection.ConvertSelectionToRigidNodes();
 		
 		// Update global transforms and bounds -- TODO: pull this bounds update out to a shared function?
-		const TManagedArray<FTransform>& Transform = FullSelection.GetGeometryCollection()->Transform;
+		const TManagedArray<FTransform3f>& Transform = FullSelection.GetGeometryCollection()->Transform;
 		const TManagedArray<int32>& TransformToGeometryIndex = FullSelection.GetGeometryCollection()->TransformToGeometryIndex;
 		const TManagedArray<FBox>& BoundingBoxes = FullSelection.GetGeometryCollection()->BoundingBox;
 
diff --git a/Engine/Plugins/Experimental/Dataflow/Source/DataflowEnginePlugin/Private/Dataflow/DataflowEngineRendering.cpp b/Engine/Plugins/Experimental/Dataflow/Source/DataflowEnginePlugin/Private/Dataflow/DataflowEngineRendering.cpp
--- a/Engine/Plugins/Experimental/Dataflow/Source/DataflowEnginePlugin/Private/Dataflow/DataflowEngineRendering.cpp
+++ b/Engine/Plugins/Experimental/Dataflow/Source/DataflowEnginePlugin/Private/Dataflow/DataflowEngineRendering.cpp
@@ -46,7 +46,7 @@ namespace Dataflow
 						{
 							const TManagedArray<int32>& BoneIndex = Collection.GetAttribute<int32>("BoneMap", FGeometryCollection::VerticesGroup);
 							const TManagedArray<int32>& Parents = Collection.GetAttribute<int32>(FTransformCollection::ParentAttribute, FTransformCollection::TransformGroup);
-							const TManagedArray<FTransform>& Transforms = Collection.GetAttribute<FTransform>(FTransformCollection::TransformAttribute, FTransformCollection::TransformGroup);
+							const TManagedArray<FTransform3f>& Transforms = Collection.GetAttribute<FTransform3f>(FTransformCollection::TransformAttribute, FTransformCollection::TransformGroup);
 
 							TArray<FTransform> M;
 							GeometryCollectionAlgo::GlobalMatrices(Transforms, Parents, M);
diff --git a/Engine/Plugins/Experimental/Fracture/Source/FractureEngine/Private/FractureEngineClustering.cpp b/Engine/Plugins/Experimental/Fracture/Source/FractureEngine/Private/FractureEngineClustering.cpp
--- a/Engine/Plugins/Experimental/Fracture/Source/FractureEngine/Private/FractureEngineClustering.cpp
+++ b/Engine/Plugins/Experimental/Fracture/Source/FractureEngine/Private/FractureEngineClustering.cpp
@@ -297,9 +297,9 @@ FBox FVoronoiPartitioner::GenerateBounds(const FGeometryCollection* GeometryColl
 
 	if (GeometryCollection->IsRigid(TransformIndex))
 	{
-		const TManagedArray<FTransform>& Transforms = GeometryCollection->Transform;
+		const TManagedArray<FTransform3f>& Transforms = GeometryCollection->Transform;
 		const TManagedArray<int32>& Parents = GeometryCollection->Parent;
-		FTransform GlobalTransform = GeometryCollectionAlgo::GlobalMatrix(Transforms, Parents, TransformIndex);
+		FTransform3f GlobalTransform = GeometryCollectionAlgo::GlobalMatrix3f(Transforms, Parents, TransformIndex);
 		
 		int32 GeometryIndex = GeometryCollection->TransformToGeometryIndex[TransformIndex];
 		int32 VertexStart = GeometryCollection->VertexStart[GeometryIndex];
@@ -310,7 +310,7 @@ FBox FVoronoiPartitioner::GenerateBounds(const FGeometryCollection* GeometryColl
 		Vertices.SetNum(VertexCount);
 		for (int32 VertexOffset = 0; VertexOffset < VertexCount; ++VertexOffset)
 		{
-			Vertices[VertexOffset] = GlobalTransform.TransformPosition((FVector)GCVertices[VertexStart + VertexOffset]);
+			Vertices[VertexOffset] = FVector(GlobalTransform.TransformPosition(GCVertices[VertexStart + VertexOffset]));
 		}
 
 		return FBox(Vertices);
@@ -677,7 +677,7 @@ void FFractureEngineClustering::AutoCluster(FGeometryCollection& GeometryCollect
 		GeometryCollection.BoneName[NewClusterIndex] = "ClusterBone";
 		GeometryCollection.Children[NewClusterIndex] = TSet<int32>(NewCluster);
 		GeometryCollection.SimulationType[NewClusterIndex] = FGeometryCollection::ESimulationTypes::FST_Clustered;
-		GeometryCollection.Transform[NewClusterIndex] = FTransform::Identity;
+		GeometryCollection.Transform[NewClusterIndex] = FTransform3f::Identity;
 		GeometryCollectionAlgo::ParentTransforms(&GeometryCollection, NewClusterIndex, NewCluster);
 	}
 	if (bHasEmptyClusters)
diff --git a/Engine/Plugins/Experimental/Fracture/Source/FractureEngine/Private/FractureEngineFracturing.cpp b/Engine/Plugins/Experimental/Fracture/Source/FractureEngine/Private/FractureEngineFracturing.cpp
--- a/Engine/Plugins/Experimental/Fracture/Source/FractureEngine/Private/FractureEngineFracturing.cpp
+++ b/Engine/Plugins/Experimental/Fracture/Source/FractureEngine/Private/FractureEngineFracturing.cpp
@@ -91,7 +91,7 @@ void FFractureEngineFracturing::GenerateExplodedViewAttribute(FManagedArrayColle
 		check(InOutCollection.HasAttribute("ExplodedVector", FGeometryCollection::TransformGroup));
 
 		TManagedArray<FVector3f>& ExplodedVectors = InOutCollection.ModifyAttribute<FVector3f>("ExplodedVector", FGeometryCollection::TransformGroup);
-		const TManagedArray<FTransform>& Transforms = InOutCollection.GetAttribute<FTransform>("Transform", FGeometryCollection::TransformGroup);
+		const TManagedArray<FTransform3f>& Transforms = InOutCollection.GetAttribute<FTransform3f>("Transform", FGeometryCollection::TransformGroup);
 		const TManagedArray<int32>& TransformToGeometryIndices = InOutCollection.GetAttribute<int32>("TransformToGeometryIndex", FGeometryCollection::TransformGroup);
 		const TManagedArray<FBox>& BoundingBoxes = InOutCollection.GetAttribute<FBox>("BoundingBox", FGeometryCollection::GeometryGroup);
 
@@ -204,7 +204,7 @@ void FFractureEngineFracturing::VoronoiFracture(FManagedArrayCollection& InOutCo
 				InOutCollection.HasAttribute("TransformIndex", FGeometryCollection::GeometryGroup) &&
 				InOutCollection.HasAttribute("BoundingBox", FGeometryCollection::GeometryGroup))
 			{
-				const TManagedArray<FTransform>& Transforms = InOutCollection.GetAttribute<FTransform>("Transform", FGeometryCollection::TransformGroup);
+				const TManagedArray<FTransform3f>& Transforms = InOutCollection.GetAttribute<FTransform3f>("Transform", FGeometryCollection::TransformGroup);
 				const TManagedArray<int32>& ParentIndices = InOutCollection.GetAttribute<int32>("Parent", FGeometryCollection::TransformGroup);
 				const TManagedArray<int32>& TransformIndices = InOutCollection.GetAttribute<int32>("TransformIndex", FGeometryCollection::GeometryGroup);
 				const TManagedArray<FBox>& BoundingBoxes = InOutCollection.GetAttribute<FBox>("BoundingBox", FGeometryCollection::GeometryGroup);
diff --git a/Engine/Plugins/Experimental/GeometryCollectionPlugin/Source/GeometryCollectionEditor/Private/GeometryCollection/GeometryCollectionCommands.cpp b/Engine/Plugins/Experimental/GeometryCollectionPlugin/Source/GeometryCollectionEditor/Private/GeometryCollection/GeometryCollectionCommands.cpp
--- a/Engine/Plugins/Experimental/GeometryCollectionPlugin/Source/GeometryCollectionEditor/Private/GeometryCollection/GeometryCollectionCommands.cpp
+++ b/Engine/Plugins/Experimental/GeometryCollectionPlugin/Source/GeometryCollectionEditor/Private/GeometryCollection/GeometryCollectionCommands.cpp
@@ -427,7 +427,7 @@ int32 FGeometryCollectionCommands::EnsureSingleRoot(UGeometryCollection* RestCol
 {
 	if (RestCollection)
 	{
-		TManagedArray<FTransform>& Transform = RestCollection->GetGeometryCollection()->Transform;
+		TManagedArray<FTransform3f>& Transform = RestCollection->GetGeometryCollection()->Transform;
 		TManagedArray<int32>& Parent = RestCollection->GetGeometryCollection()->Parent;
 		int32 NumElements = RestCollection->GetGeometryCollection()->NumElements(FGeometryCollection::TransformGroup);
 		if (GeometryCollectionAlgo::HasMultipleRoots(RestCollection->GetGeometryCollection().Get()))
@@ -441,7 +441,7 @@ int32 FGeometryCollectionCommands::EnsureSingleRoot(UGeometryCollection* RestCol
 				}
 			}
 			int32 RootIndex = RestCollection->GetGeometryCollection()->AddElements(1, FGeometryCollection::TransformGroup);
-			Transform[RootIndex].SetTranslation(GeometryCollectionAlgo::AveragePosition(RestCollection->GetGeometryCollection().Get(), RootIndices));
+			Transform[RootIndex].SetTranslation(FVector3f(GeometryCollectionAlgo::AveragePosition(RestCollection->GetGeometryCollection().Get(), RootIndices)));
 			GeometryCollectionAlgo::ParentTransforms(RestCollection->GetGeometryCollection().Get(), RootIndex, RootIndices);
 			return RootIndex;
 		}
@@ -466,7 +466,7 @@ void SplitAcrossYZPlaneRecursive(uint32 RootIndex, const FTransform & ParentTran
 {
 	TSet<uint32> RootIndices;
 	TManagedArray<TSet<int32>>& Children = Collection->GetGeometryCollection()->Children;
-	TManagedArray<FTransform>& Transform = Collection->GetGeometryCollection()->Transform;
+	TManagedArray<FTransform3f>& Transform = Collection->GetGeometryCollection()->Transform;
 
 	TArray<int32> SelectedBonesA, SelectedBonesB;
 	for (auto& ChildIndex : Children[RootIndex])
@@ -476,7 +476,7 @@ void SplitAcrossYZPlaneRecursive(uint32 RootIndex, const FTransform & ParentTran
 			SplitAcrossYZPlaneRecursive(ChildIndex, ParentTransform, Collection);
 		}
 
-		FVector Translation = (Transform[ChildIndex]*ParentTransform).GetTranslation();
+		FVector Translation = (FTransform(Transform[ChildIndex])*ParentTransform).GetTranslation();
 		UE_LOG(UGeometryCollectionCommandsLogging, Log, TEXT("... [%d] global:(%3.5f,%3.5f,%3.5f) local:(%3.5f,%3.5f,%3.5f)"),
 			ChildIndex, Translation.X, Translation.Y, Translation.Z, Transform[ChildIndex].GetTranslation().X,
 			Transform[ChildIndex].GetTranslation().Y, Transform[ChildIndex].GetTranslation().Z );
@@ -495,12 +495,12 @@ void SplitAcrossYZPlaneRecursive(uint32 RootIndex, const FTransform & ParentTran
 	{
 		int32 BoneAIndex = Collection->GetGeometryCollection()->AddElements(1, FGeometryCollection::TransformGroup);
 		GeometryCollectionAlgo::ParentTransform(Collection->GetGeometryCollection().Get(), RootIndex, BoneAIndex);
-		Transform[BoneAIndex].SetTranslation(GeometryCollectionAlgo::AveragePosition(Collection->GetGeometryCollection().Get(), SelectedBonesA));
+		Transform[BoneAIndex].SetTranslation(FVector3f(GeometryCollectionAlgo::AveragePosition(Collection->GetGeometryCollection().Get(), SelectedBonesA)));
 		GeometryCollectionAlgo::ParentTransforms(Collection->GetGeometryCollection().Get(), BoneAIndex, SelectedBonesA);
 
 		int32 BoneBIndex = Collection->GetGeometryCollection()->AddElements(1, FGeometryCollection::TransformGroup);
 		GeometryCollectionAlgo::ParentTransform(Collection->GetGeometryCollection().Get(), RootIndex, BoneBIndex);
-		Transform[BoneBIndex].SetTranslation(GeometryCollectionAlgo::AveragePosition(Collection->GetGeometryCollection().Get(), SelectedBonesB));
+		Transform[BoneBIndex].SetTranslation(FVector3f(GeometryCollectionAlgo::AveragePosition(Collection->GetGeometryCollection().Get(), SelectedBonesB)));
 		GeometryCollectionAlgo::ParentTransforms(Collection->GetGeometryCollection().Get(), BoneBIndex, SelectedBonesB);
 	}
 }
diff --git a/Engine/Plugins/Experimental/GeometryCollectionPlugin/Source/GeometryCollectionEditor/Private/GeometryCollection/GeometryCollectionFactory.cpp b/Engine/Plugins/Experimental/GeometryCollectionPlugin/Source/GeometryCollectionEditor/Private/GeometryCollection/GeometryCollectionFactory.cpp
--- a/Engine/Plugins/Experimental/GeometryCollectionPlugin/Source/GeometryCollectionEditor/Private/GeometryCollection/GeometryCollectionFactory.cpp
+++ b/Engine/Plugins/Experimental/GeometryCollectionPlugin/Source/GeometryCollectionEditor/Private/GeometryCollection/GeometryCollectionFactory.cpp
@@ -179,13 +179,13 @@ UObject* UGeometryCollectionFactory::FactoryCreateNew(UClass* Class, UObject* In
 	if (TSharedPtr<FGeometryCollection, ESPMode::ThreadSafe> Collection = NewGeometryCollection->GetGeometryCollection())
 	{
 		TManagedArray<int32> & Parent = Collection->Parent;
-		TManagedArray<FTransform>& Transform = Collection->Transform;
+		TManagedArray<FTransform3f>& Transform = Collection->Transform;
 
 		for(int TransformGroupIndex =0; TransformGroupIndex<Collection->Transform.Num(); TransformGroupIndex++)
 		{
 			if (Parent[TransformGroupIndex] == FGeometryCollection::Invalid)
 			{
-				Transform[TransformGroupIndex] = Transform[TransformGroupIndex].GetRelativeTransform(LastTransform);
+				Transform[TransformGroupIndex] = Transform[TransformGroupIndex].GetRelativeTransform(FTransform3f(LastTransform));
 			}
 		}
 	}
diff --git a/Engine/Plugins/Experimental/PlanarCutPlugin/Source/PlanarCut/Private/GeometryMeshConversion.cpp b/Engine/Plugins/Experimental/PlanarCutPlugin/Source/PlanarCut/Private/GeometryMeshConversion.cpp
--- a/Engine/Plugins/Experimental/PlanarCutPlugin/Source/PlanarCut/Private/GeometryMeshConversion.cpp
+++ b/Engine/Plugins/Experimental/PlanarCutPlugin/Source/PlanarCut/Private/GeometryMeshConversion.cpp
@@ -2528,8 +2528,8 @@ void FCellMeshes::CreateMeshesForSinglePlane(const FPlanarCells& Cells, const FA
 	}
 }
 
-
-void FDynamicMeshCollection::Init(const FGeometryCollection* Collection, const TManagedArray<FTransform>& Transforms, const TArrayView<const int32>& TransformIndices, FTransform TransformCollection, bool bSaveIsolatedVertices)
+template<typename TransformType>
+void FDynamicMeshCollection::InitTemplate(const FGeometryCollection* Collection, const TManagedArray<TransformType>& Transforms, const TArrayView<const int32>& TransformIndices, FTransform TransformCollection, bool bSaveIsolatedVertices)
 {
 	GeometryCollection::UV::FConstUVLayers UVLayers = GeometryCollection::UV::FindActiveUVLayers(*Collection);
 	int32 NumUVLayers = UVLayers.Num();
@@ -2631,6 +2631,16 @@ void FDynamicMeshCollection::Init(const FGeometryCollection* Collection, const T
 }
 
 
+void FDynamicMeshCollection::Init(const FGeometryCollection* Collection, const TManagedArray<FTransform>& Transforms, const TArrayView<const int32>& TransformIndices, FTransform TransformCollection, bool bSaveIsolatedVertices)
+{
+	InitTemplate(Collection, Transforms, TransformIndices, TransformCollection, bSaveIsolatedVertices);
+}
+
+void FDynamicMeshCollection::Init(const FGeometryCollection* Collection, const TManagedArray<FTransform3f>& Transforms, const TArrayView<const int32>& TransformIndices, FTransform TransformCollection, bool bSaveIsolatedVertices)
+{
+	InitTemplate(Collection, Transforms, TransformIndices, TransformCollection, bSaveIsolatedVertices);
+}
+
 void FDynamicMeshCollection::SetGeometryVisibility(FGeometryCollection* Collection, const TArray<int32>& GeometryIndices, bool bVisible)
 {
 	TManagedArray<bool>& Visible = Collection->Visible;
@@ -3386,7 +3396,7 @@ int32 FDynamicMeshCollection::AppendToCollection(const FTransform& FromCollectio
 		Output.Children[TransformParent].Add(TransformIdx);
 		Output.SimulationType[TransformParent] = FGeometryCollection::ESimulationTypes::FST_Clustered;
 	}
-	Output.Transform[TransformIdx] = FTransform::Identity;
+	Output.Transform[TransformIdx] = FTransform3f::Identity;
 	Output.SimulationType[TransformIdx] = FGeometryCollection::ESimulationTypes::FST_Rigid;
 
 	int32 FacesStart = Output.AddElements(NumTriangles, FGeometryCollection::FacesGroup);
diff --git a/Engine/Plugins/Experimental/PlanarCutPlugin/Source/PlanarCut/Private/GeometryMeshConversion.h b/Engine/Plugins/Experimental/PlanarCutPlugin/Source/PlanarCut/Private/GeometryMeshConversion.h
--- a/Engine/Plugins/Experimental/PlanarCutPlugin/Source/PlanarCut/Private/GeometryMeshConversion.h
+++ b/Engine/Plugins/Experimental/PlanarCutPlugin/Source/PlanarCut/Private/GeometryMeshConversion.h
@@ -198,6 +198,8 @@ struct PLANARCUT_API FDynamicMeshCollection
 	TIndirectArray<FMeshData> Meshes;
 	UE::Geometry::FAxisAlignedBox3d Bounds;
 	
+	FDynamicMeshCollection() {}
+
 	FDynamicMeshCollection(const FGeometryCollection* Collection, const TArrayView<const int32>& TransformIndices, FTransform TransformCollection, bool bSaveIsolatedVertices = false)
 	{
 		Init(Collection, TransformIndices, TransformCollection, bSaveIsolatedVertices);
@@ -208,12 +210,18 @@ struct PLANARCUT_API FDynamicMeshCollection
 		Init(Collection, Transforms, TransformIndices, TransformCollection, bSaveIsolatedVertices);
 	}
 
+	FDynamicMeshCollection(const FGeometryCollection* Collection, const TManagedArray<FTransform3f>& Transforms, const TArrayView<const int32>& TransformIndices, FTransform TransformCollection, bool bSaveIsolatedVertices = false)
+	{
+		Init(Collection, Transforms, TransformIndices, TransformCollection, bSaveIsolatedVertices);
+	}
+
 	void Init(const FGeometryCollection* Collection, const TArrayView<const int32>& TransformIndices, FTransform TransformCollection, bool bSaveIsolatedVertices = false)
 	{
 		Init(Collection, Collection->Transform, TransformIndices, TransformCollection, bSaveIsolatedVertices);
 	}
 
 	void Init(const FGeometryCollection* Collection, const TManagedArray<FTransform>& Transforms, const TArrayView<const int32>& TransformIndices, FTransform TransformCollection, bool bSaveIsolatedVertices = false);
+	void Init(const FGeometryCollection* Collection, const TManagedArray<FTransform3f>& Transforms, const TArrayView<const int32>& TransformIndices, FTransform TransformCollection, bool bSaveIsolatedVertices = false);
 
 	int32 CutWithMultiplePlanes(
 		const TArrayView<const FPlane>& Planes,
@@ -276,6 +284,10 @@ struct PLANARCUT_API FDynamicMeshCollection
 
 private:
 
+	template<typename TransformType>
+	void InitTemplate(const FGeometryCollection* Collection, const TManagedArray<TransformType>& Transforms, const TArrayView<const int32>& TransformIndices, FTransform TransformCollection, bool bSaveIsolatedVertices);
+
+
 	void SetGeometryVisibility(FGeometryCollection* Collection, const TArray<int32>& GeometryIndices, bool bVisible);
 
 	// Update an existing geometry in a collection w/ a new mesh (w/ the same number of faces and vertices!)
diff --git a/Engine/Plugins/Experimental/PlanarCutPlugin/Source/PlanarCut/Private/PlanarCut.cpp b/Engine/Plugins/Experimental/PlanarCutPlugin/Source/PlanarCut/Private/PlanarCut.cpp
--- a/Engine/Plugins/Experimental/PlanarCutPlugin/Source/PlanarCut/Private/PlanarCut.cpp
+++ b/Engine/Plugins/Experimental/PlanarCutPlugin/Source/PlanarCut/Private/PlanarCut.cpp
@@ -2085,13 +2085,13 @@ int32 AddCollisionSampleVertices(double CollisionSampleSpacing, FGeometryCollect
 	return INDEX_NONE;
 }
 
-
-void ConvertToMeshDescription(
+template <typename TransformType>
+void ConvertToMeshDescriptionTemplate(
 	FMeshDescription& MeshOut,
 	FTransform& TransformOut,
 	bool bCenterPivot,
 	FGeometryCollection& Collection,
-	const TManagedArray<FTransform>& BoneTransforms,
+	const TManagedArray<TransformType>& BoneTransforms,
 	const TArrayView<const int32>& TransformIndices,
 	TFunction<int32(int32, bool)> RemapMaterialIDs
 )
@@ -2099,7 +2099,15 @@ void ConvertToMeshDescription(
 	FTransform CellsToWorld = FTransform::Identity;
 	TransformOut = FTransform::Identity;
 
-	FDynamicMeshCollection MeshCollection(&Collection, BoneTransforms.Num() ? BoneTransforms : Collection.Transform, TransformIndices, CellsToWorld);
+	FDynamicMeshCollection MeshCollection;
+	if (BoneTransforms.Num())
+	{
+		MeshCollection.Init(&Collection, BoneTransforms, TransformIndices, CellsToWorld);
+	}
+	else
+	{
+		MeshCollection.Init(&Collection, Collection.Transform, TransformIndices, CellsToWorld);
+	}
 	
 	FDynamicMesh3 CombinedMesh;
 	SetGeometryCollectionAttributes(CombinedMesh, Collection.NumUVLayers());
@@ -2154,4 +2162,31 @@ void ConvertToMeshDescription(
 	Converter.Convert(&CombinedMesh, MeshOut, true);
 }
 
+void ConvertToMeshDescription(
+	FMeshDescription& MeshOut,
+	FTransform& TransformOut,
+	bool bCenterPivot,
+	FGeometryCollection& Collection,
+	const TManagedArray<FTransform>& BoneTransforms,
+	const TArrayView<const int32>& TransformIndices,
+	TFunction<int32(int32, bool)> RemapMaterialIDs
+)
+{
+	ConvertToMeshDescriptionTemplate(MeshOut, TransformOut, bCenterPivot, Collection, BoneTransforms, TransformIndices, RemapMaterialIDs);
+}
+
+void ConvertToMeshDescription(
+	FMeshDescription& MeshOut,
+	FTransform& TransformOut,
+	bool bCenterPivot,
+	FGeometryCollection& Collection,
+	const TManagedArray<FTransform3f>& BoneTransforms,
+	const TArrayView<const int32>& TransformIndices,
+	TFunction<int32(int32, bool)> RemapMaterialIDs
+)
+{
+	ConvertToMeshDescriptionTemplate(MeshOut, TransformOut, bCenterPivot, Collection, BoneTransforms, TransformIndices, RemapMaterialIDs);
+}
+
+
 #undef LOCTEXT_NAMESPACE
diff --git a/Engine/Plugins/Experimental/PlanarCutPlugin/Source/PlanarCut/Public/PlanarCut.h b/Engine/Plugins/Experimental/PlanarCutPlugin/Source/PlanarCut/Public/PlanarCut.h
--- a/Engine/Plugins/Experimental/PlanarCutPlugin/Source/PlanarCut/Public/PlanarCut.h
+++ b/Engine/Plugins/Experimental/PlanarCutPlugin/Source/PlanarCut/Public/PlanarCut.h
@@ -537,6 +537,26 @@ void PLANARCUT_API ConvertToMeshDescription(
 	TFunction<int32(int32, bool)> RemapMaterialIDs = nullptr
 );
 
+/**
+ * Convert chosen Geometry groups inside a GeometryCollection to a single Mesh Description.
+ *
+ * @param OutputMesh				Mesh to be filled with the geometry collection geometry
+ * @param TransformOut				Transform taking output mesh geometry to local space of geometry collection
+ * @param bCenterPivot				Whether to center the geometry at the origin
+ * @param Collection				The collection to be converted
+ * @param TransformIndices			Which transform groups inside the collection to convert
+ * @param RemapMaterialIDs			Optional function to remap (MaterialID, bIsInternal) -> NewMaterialID
+ */
+void PLANARCUT_API ConvertToMeshDescription(
+	FMeshDescription& OutputMesh,
+	FTransform& TransformOut,
+	bool bCenterPivot,
+	FGeometryCollection& Collection,
+	const TManagedArray<FTransform3f>& BoneTransforms,
+	const TArrayView<const int32>& TransformIndices,
+	TFunction<int32(int32, bool)> RemapMaterialIDs = nullptr
+);
+
 
 #if UE_ENABLE_INCLUDE_ORDER_DEPRECATED_IN_5_2
 #include "CoreMinimal.h"
diff --git a/Engine/Source/Programs/HeadlessChaos/Private/GeometryCollection/GeometryCollectionTestClustering.cpp b/Engine/Source/Programs/HeadlessChaos/Private/GeometryCollection/GeometryCollectionTestClustering.cpp
--- a/Engine/Source/Programs/HeadlessChaos/Private/GeometryCollection/GeometryCollectionTestClustering.cpp
+++ b/Engine/Source/Programs/HeadlessChaos/Private/GeometryCollection/GeometryCollectionTestClustering.cpp
@@ -76,7 +76,7 @@ namespace GeometryCollectionTest
 
 		FGeometryCollectionClusteringUtility::ClusterAllBonesUnderNewRoot(RestCollection.Get());
 		EXPECT_EQ(RestCollection->Transform.Num(), 3);
-		RestCollection->Transform[2] = FTransform(FQuat::MakeFromEuler(FVector(90.0, 0, 0.)), FVector(0, 0, 40));
+		RestCollection->Transform[2] = FTransform3f(FQuat4f::MakeFromEuler(FVector3f(90.0, 0, 0.)), FVector3f(0, 0, 40));
 
 		//GeometryCollectionAlgo::PrintParentHierarchy(RestCollection.Get());
 
@@ -141,7 +141,7 @@ namespace GeometryCollectionTest
 
 		FGeometryCollectionClusteringUtility::ClusterAllBonesUnderNewRoot(RestCollection.Get());
 		EXPECT_EQ(RestCollection->Transform.Num(), 3);
-		RestCollection->Transform[2] = FTransform(FQuat::MakeFromEuler(FVector(90.0, 0, 0.)), FVector(0, 0, 40));
+		RestCollection->Transform[2] = FTransform3f(FQuat4f::MakeFromEuler(FVector3f(90.0, 0, 0.)), FVector3f(0, 0, 40));
 
 		//GeometryCollectionAlgo::PrintParentHierarchy(RestCollection.Get());
 
@@ -440,7 +440,7 @@ namespace GeometryCollectionTest
 
 
 		TSharedPtr<FGeometryCollection> RestCollection = CreateClusteredBody(FVector::ZeroVector);
-		RestCollection->Transform[2] = FTransform(FQuat::MakeFromEuler(FVector(0., 90.f, 0.)), FVector(0, 0, 17));
+		RestCollection->Transform[2] = FTransform3f(FQuat4f::MakeFromEuler(FVector3f(0., 90.f, 0.)), FVector3f(0, 0, 17));
 		CreationParameters Params;
 		Params.RestCollection = RestCollection;
 		Params.DynamicState = EObjectStateTypeEnum::Chaos_Object_Dynamic;
@@ -537,11 +537,11 @@ namespace GeometryCollectionTest
 
 		FGeometryCollectionClusteringUtility::ClusterAllBonesUnderNewRoot(RestCollection.Get());
 		EXPECT_EQ(RestCollection->Transform.Num(), 3);
-		RestCollection->Transform[2] = FTransform(FQuat::MakeFromEuler(FVector(90.f, 0, 0.)), FVector(0, 0, 40));
+		RestCollection->Transform[2] = FTransform3f(FQuat4f::MakeFromEuler(FVector3f(90.f, 0, 0.)), FVector3f(0, 0, 40));
 
 		FGeometryCollectionClusteringUtility::ClusterBonesUnderNewNode(RestCollection.Get(), 3, { 2 }, true);
 		EXPECT_EQ(RestCollection->Transform.Num(), 4);
-		RestCollection->Transform[3] = FTransform(FQuat::MakeFromEuler(FVector(0.f, 0, 0.)), FVector(0, 0, 10));
+		RestCollection->Transform[3] = FTransform3f(FQuat4f::MakeFromEuler(FVector3f(0.f, 0, 0.)), FVector4f(0, 0, 10));
 
 		//GeometryCollectionAlgo::PrintParentHierarchy(RestCollection.Get());
 
@@ -2222,7 +2222,7 @@ namespace GeometryCollectionTest
 		RestCollection->AppendGeometry(*GeometryCollection::MakeCubeElement(FTransform(FQuat::MakeFromEuler(FVector(0.f)), FVector(500, 0, 0)), FVector(1.0)));
 
 		RestCollection->AddElements(7, FGeometryCollection::TransformGroup);
-		RestCollection->Transform[10].SetTranslation(FVector(0,0,0));
+		RestCollection->Transform[10].SetTranslation(FVector3f(0,0,0));
 
 		RestCollection->SimulationType[0] = FGeometryCollection::ESimulationTypes::FST_Rigid;
 		RestCollection->SimulationType[1] = FGeometryCollection::ESimulationTypes::FST_Rigid;
diff --git a/Engine/Source/Programs/HeadlessChaos/Private/GeometryCollection/GeometryCollectionTestCreation.cpp b/Engine/Source/Programs/HeadlessChaos/Private/GeometryCollection/GeometryCollectionTestCreation.cpp
--- a/Engine/Source/Programs/HeadlessChaos/Private/GeometryCollection/GeometryCollectionTestCreation.cpp
+++ b/Engine/Source/Programs/HeadlessChaos/Private/GeometryCollection/GeometryCollectionTestCreation.cpp
@@ -1092,12 +1092,12 @@ namespace GeometryCollectionTest
 		FGeometryCollection* Collection = new FGeometryCollection();
 
 		int32 TransformIndex = Collection->AddElements(1, FGeometryCollection::TransformGroup);
-		(Collection->Transform)[TransformIndex].SetTranslation(FVector(13));
+		(Collection->Transform)[TransformIndex].SetTranslation(FVector3f(13));
 		(Collection->Parent)[TransformIndex]= -1;
 		EXPECT_EQ(TransformIndex, 0);
 
 		TransformIndex = Collection->AddElements(1, FGeometryCollection::TransformGroup);
-		(Collection->Transform)[TransformIndex].SetTranslation(FVector(7));
+		(Collection->Transform)[TransformIndex].SetTranslation(FVector3f(7));
 		(Collection->Parent)[TransformIndex] = -1;
 		EXPECT_EQ(TransformIndex, 1);
 
@@ -1113,7 +1113,7 @@ namespace GeometryCollectionTest
 
 		TArray<FTransform> GlobalTransform;
 		GeometryCollectionAlgo::GlobalMatrices(Collection->Transform, Collection->Parent, GlobalTransform);
-		EXPECT_LT(((Collection->Transform)[0].GetTranslation() - FVector(6)).Size(), KINDA_SMALL_NUMBER);
+		EXPECT_LT(((Collection->Transform)[0].GetTranslation() - FVector3f(6)).Size(), KINDA_SMALL_NUMBER);
 		EXPECT_LT((GlobalTransform[0].GetTranslation()-FVector(13)).Size(), KINDA_SMALL_NUMBER);
 
 		EXPECT_EQ(Collection->TransformToGeometryIndex.Num(), 2);
@@ -1124,7 +1124,7 @@ namespace GeometryCollectionTest
 		// Add some geometry
 		//
 		TransformIndex = Collection->AppendGeometry(*GeometryCollection::MakeCubeElement(FTransform(FVector(3)), FVector(1.0)));
-		EXPECT_LT(((Collection->Transform)[TransformIndex].GetTranslation() - FVector(3)).Size(), KINDA_SMALL_NUMBER);
+		EXPECT_LT(((Collection->Transform)[TransformIndex].GetTranslation() - FVector3f(3)).Size(), KINDA_SMALL_NUMBER);
 		EXPECT_EQ((Collection->TransformIndex).Num(), 1);
 		EXPECT_EQ((Collection->TransformIndex)[0], TransformIndex);
 		EXPECT_EQ((Collection->VertexStart)[0], 0);
@@ -1148,7 +1148,7 @@ namespace GeometryCollectionTest
 		EXPECT_EQ((Collection->Children)[1].Num(), 1);
 		EXPECT_TRUE((Collection->Children)[1].Contains(0));
 		EXPECT_EQ((Collection->Parent)[1], -1);
-		EXPECT_LT(((Collection->Transform)[TransformIndex].GetTranslation() - FVector(-10)).Size(), KINDA_SMALL_NUMBER);
+		EXPECT_LT(((Collection->Transform)[TransformIndex].GetTranslation() - FVector3f(-10)).Size(), KINDA_SMALL_NUMBER);
 		EXPECT_EQ((Collection->TransformIndex).Num(), 1);
 		EXPECT_EQ((Collection->TransformIndex)[0], TransformIndex);
 		EXPECT_EQ((Collection->VertexStart)[0], 0);
diff --git a/Engine/Source/Programs/HeadlessChaos/Private/GeometryCollection/GeometryCollectionTestFramework.cpp b/Engine/Source/Programs/HeadlessChaos/Private/GeometryCollection/GeometryCollectionTestFramework.cpp
--- a/Engine/Source/Programs/HeadlessChaos/Private/GeometryCollection/GeometryCollectionTestFramework.cpp
+++ b/Engine/Source/Programs/HeadlessChaos/Private/GeometryCollection/GeometryCollectionTestFramework.cpp
@@ -228,7 +228,7 @@ namespace GeometryCollectionTest
 			RestCollection->InnerRadius[0] = 1.f;	// Assume sphere w/radius 1
 			RestCollection->OuterRadius[0] = 1.f;	// Assume sphere w/radius 1
 			RestCollection->AddElements(1, FGeometryCollection::TransformGroup);
-			RestCollection->Transform[0] = Params.RootTransform;
+			RestCollection->Transform[0] = FTransform3f(Params.RootTransform);
 			RestCollection->Transform[0].NormalizeRotation();
 			break;
 		default:
diff --git a/Engine/Source/Programs/HeadlessChaos/Private/GeometryCollection/GeometryCollectionTestMatrices.cpp b/Engine/Source/Programs/HeadlessChaos/Private/GeometryCollection/GeometryCollectionTestMatrices.cpp
--- a/Engine/Source/Programs/HeadlessChaos/Private/GeometryCollection/GeometryCollectionTestMatrices.cpp
+++ b/Engine/Source/Programs/HeadlessChaos/Private/GeometryCollection/GeometryCollectionTestMatrices.cpp
@@ -53,14 +53,14 @@ namespace GeometryCollectionTest
 		FVector Pos2 = GlobalTransform[2].GetTranslation();
 		EXPECT_FVECTOR_NEAR(Pos2, FVector(-10, 0, 0), 0.0001);
 
-		FTransform Frame = GeometryCollectionAlgo::GlobalMatrix(Collection->Transform, Collection->Parent, 2);
-		EXPECT_FVECTOR_NEAR(Frame.GetRotation().Euler(), FVector(0, 0, -90), 0.0001);
-		EXPECT_FVECTOR_NEAR(Frame.GetTranslation(), FVector(-10, 0, 0), 0.0001);
-
-		Frame = GeometryCollectionAlgo::GlobalMatrix(Collection->Transform, Collection->Parent, 1);
-		FVector FrameRot = Frame.GetRotation().Euler();
-		EXPECT_TRUE(FrameRot.Equals(FVector(0, 0, -180), 0.0001) || FrameRot.Equals(FVector(0, 0, 180), 0.0001));
-		EXPECT_FVECTOR_NEAR(Frame.GetTranslation(), FVector(-10, 10, 0), 0.0001);
+		FTransform3f Frame = GeometryCollectionAlgo::GlobalMatrix3f(Collection->Transform, Collection->Parent, 2);
+		EXPECT_FVECTOR_NEAR(Frame.GetRotation().Euler(), FVector3f(0, 0, -90), 0.0001);
+		EXPECT_FVECTOR_NEAR(Frame.GetTranslation(), FVector3f(-10, 0, 0), 0.0001);
+
+		Frame = GeometryCollectionAlgo::GlobalMatrix3f(Collection->Transform, Collection->Parent, 1);
+		FVector3f FrameRot = Frame.GetRotation().Euler();
+		EXPECT_TRUE(FrameRot.Equals(FVector3f(0, 0, -180), 0.0001) || FrameRot.Equals(FVector3f(0, 0, 180), 0.0001));
+		EXPECT_FVECTOR_NEAR(Frame.GetTranslation(), FVector3f(-10, 10, 0), 0.0001);
 	}
 
 	void ReparentingMatrices()
@@ -82,8 +82,8 @@ namespace GeometryCollectionTest
 		//  ......2
 		TArray<int32> Bones = { 2 };
 		GeometryCollectionAlgo::ParentTransforms(Collection.Get(), 1, Bones);
-		EXPECT_FVECTOR_NEAR((Collection->Transform)[2].GetTranslation(), FVector(0, 10, 0), 0.0001);
-		EXPECT_FVECTOR_NEAR((Collection->Transform)[2].GetRotation().Euler(), FVector(0, 0, 90.), 0.0001);
+		EXPECT_FVECTOR_NEAR((Collection->Transform)[2].GetTranslation(), FVector3f(0, 10, 0), 0.0001);
+		EXPECT_FVECTOR_NEAR((Collection->Transform)[2].GetRotation().Euler(), FVector3f(0, 0, 90.), 0.0001);
 
 		TArray<FTransform> GlobalTransform;
 		GeometryCollectionAlgo::GlobalMatrices(Collection->Transform, Collection->Parent, GlobalTransform);
@@ -120,10 +120,10 @@ namespace GeometryCollectionTest
 		}
 		(Collection.Children)[index].Empty();
 
-		(Collection.Transform)[0] = FTransform(FQuat::MakeFromEuler(FVector(0, 0, 0)), FVector(0, 0, 0));
-		(Collection.Transform)[1] = FTransform(FQuat::MakeFromEuler(FVector(0, 0, 90.)), FVector(1, 0, 0));
-		(Collection.Transform)[2] = FTransform(FQuat::MakeFromEuler(FVector(0, 90.,0)), FVector(1, 0, 0));
-		(Collection.Transform)[3] = FTransform(FQuat::MakeFromEuler(FVector(0, 0, 0)), FVector(1, 0, 0));
+		(Collection.Transform)[0] = FTransform3f(FQuat4f::MakeFromEuler(FVector3f(0, 0, 0)), FVector3f(0, 0, 0));
+		(Collection.Transform)[1] = FTransform3f(FQuat4f::MakeFromEuler(FVector3f(0, 0, 90.)), FVector3f(1, 0, 0));
+		(Collection.Transform)[2] = FTransform3f(FQuat4f::MakeFromEuler(FVector3f(0, 90.,0)), FVector3f(1, 0, 0));
+		(Collection.Transform)[3] = FTransform3f(FQuat4f::MakeFromEuler(FVector3f(0, 0, 0)), FVector3f(1, 0, 0));
 
 		TArray<FTransform> GlobalMatrices0;
 		GeometryCollectionAlgo::GlobalMatrices(Collection.Transform, Collection.Parent, GlobalMatrices0);
diff --git a/Engine/Source/Programs/HeadlessChaos/Private/GeometryCollection/GeometryCollectionTestSerialization.cpp b/Engine/Source/Programs/HeadlessChaos/Private/GeometryCollection/GeometryCollectionTestSerialization.cpp
--- a/Engine/Source/Programs/HeadlessChaos/Private/GeometryCollection/GeometryCollectionTestSerialization.cpp
+++ b/Engine/Source/Programs/HeadlessChaos/Private/GeometryCollection/GeometryCollectionTestSerialization.cpp
@@ -52,10 +52,10 @@ namespace GeometryCollectionTests
 			EXPECT_EQ(TestCollection.Children[1].Num(), 2);
 			EXPECT_EQ(TestCollection.Children[2].Num(), 0);
 			EXPECT_EQ(TestCollection.Children[3].Num(), 0);
-			EXPECT_EQ(TestCollection.Transform[0].GetTranslation(), FVector(0.0f, 0.0f, 0.0f));
-			EXPECT_EQ(TestCollection.Transform[1].GetTranslation(), FVector(1.0f, 0.0f, 5.0f));
-			EXPECT_EQ(TestCollection.Transform[2].GetTranslation(), FVector(3.0f, 0.0f, 6.0f));
-			EXPECT_EQ(TestCollection.Transform[3].GetTranslation(), FVector(-1.0f, 0.0f, 6.0f));
+			EXPECT_EQ(TestCollection.Transform[0].GetTranslation(), FVector3f(0.0f, 0.0f, 0.0f));
+			EXPECT_EQ(TestCollection.Transform[1].GetTranslation(), FVector3f(1.0f, 0.0f, 5.0f));
+			EXPECT_EQ(TestCollection.Transform[2].GetTranslation(), FVector3f(3.0f, 0.0f, 6.0f));
+			EXPECT_EQ(TestCollection.Transform[3].GetTranslation(), FVector3f(-1.0f, 0.0f, 6.0f));
 		}
 	}
 }
diff --git a/Engine/Source/Programs/HeadlessChaos/Private/GeometryCollection/GeometryCollectionTestSimulationField.cpp b/Engine/Source/Programs/HeadlessChaos/Private/GeometryCollection/GeometryCollectionTestSimulationField.cpp
--- a/Engine/Source/Programs/HeadlessChaos/Private/GeometryCollection/GeometryCollectionTestSimulationField.cpp
+++ b/Engine/Source/Programs/HeadlessChaos/Private/GeometryCollection/GeometryCollectionTestSimulationField.cpp
@@ -1092,7 +1092,7 @@ namespace GeometryCollectionTest
 		RestCollection->AppendGeometry(*GeometryCollection::MakeCubeElement(FTransform(FQuat::MakeFromEuler(FVector(0.f)), FVector(200, 0, 0)), FVector(1.0)));
 		RestCollection->AppendGeometry(*GeometryCollection::MakeCubeElement(FTransform(FQuat::MakeFromEuler(FVector(0.f)), FVector(300, 0, 0)), FVector(1.0)));
 		RestCollection->AppendGeometry(*GeometryCollection::MakeCubeElement(FTransform(FQuat::MakeFromEuler(FVector(0.f)), FVector(400, 0, 0)), FVector(1.0)));
-		RestCollection->Transform[0].SetTranslation(FVector(0.0f));
+		RestCollection->Transform[0].SetTranslation(FVector3f(0.0f));
 		RestCollection->SimulationType[0] = FGeometryCollection::ESimulationTypes::FST_Clustered;
 		RestCollection->SimulationType[1] = FGeometryCollection::ESimulationTypes::FST_Clustered;
 		RestCollection->SimulationType[2] = FGeometryCollection::ESimulationTypes::FST_Clustered;
@@ -1161,7 +1161,7 @@ namespace GeometryCollectionTest
 		RestCollection->AppendGeometry(*GeometryCollection::MakeCubeElement(FTransform(FQuat::MakeFromEuler(FVector(0.f)), FVector(200, 0, 0)), FVector(1.0)));
 		RestCollection->AppendGeometry(*GeometryCollection::MakeCubeElement(FTransform(FQuat::MakeFromEuler(FVector(0.f)), FVector(300, 0, 0)), FVector(1.0)));
 		RestCollection->AppendGeometry(*GeometryCollection::MakeCubeElement(FTransform(FQuat::MakeFromEuler(FVector(0.f)), FVector(400, 0, 0)), FVector(1.0)));
-		RestCollection->Transform[0].SetTranslation(FVector(0.0f));
+		RestCollection->Transform[0].SetTranslation(FVector3f(0.0f));
 		RestCollection->SimulationType[0] = FGeometryCollection::ESimulationTypes::FST_Clustered;
 		RestCollection->SimulationType[1] = FGeometryCollection::ESimulationTypes::FST_Clustered;
 		RestCollection->SimulationType[2] = FGeometryCollection::ESimulationTypes::FST_Clustered;
diff --git a/Engine/Source/Programs/HeadlessChaos/Private/GeometryCollection/GeometryCollectionTestSimulationSolver.cpp b/Engine/Source/Programs/HeadlessChaos/Private/GeometryCollection/GeometryCollectionTestSimulationSolver.cpp
--- a/Engine/Source/Programs/HeadlessChaos/Private/GeometryCollection/GeometryCollectionTestSimulationSolver.cpp
+++ b/Engine/Source/Programs/HeadlessChaos/Private/GeometryCollection/GeometryCollectionTestSimulationSolver.cpp
@@ -58,7 +58,7 @@ using namespace ChaosTest;
 
 		{
 			// never touched
-			TManagedArray<FTransform>& RestTransform = Collection->RestCollection->Transform;
+			TManagedArray<FTransform3f>& RestTransform = Collection->RestCollection->Transform;
 			EXPECT_LT(FMath::Abs(RestTransform[0].GetTranslation().Z), SMALL_THRESHOLD);
 
 			// simulated
@@ -81,7 +81,7 @@ using namespace ChaosTest;
 
 		FGeometryCollectionClusteringUtility::ClusterAllBonesUnderNewRoot(RestCollection.Get());
 		EXPECT_EQ(RestCollection->Transform.Num(), 3);
-		RestCollection->Transform[2] = FTransform(FQuat::MakeFromEuler(FVector(90.f, 0, 0.)), FVector(0, 0, 40));
+		RestCollection->Transform[2] = FTransform3f(FQuat4f::MakeFromEuler(FVector3f(90.f, 0, 0.)), FVector3f(0, 0, 40));
 
 		//GeometryCollectionAlgo::PrintParentHierarchy(RestCollection.Get());
 
@@ -218,7 +218,7 @@ using namespace ChaosTest;
 				RestCollection->AppendGeometry(*GeometryCollection::MakeCubeElement(FTransform(FQuat::MakeFromEuler(FVector(0, 0, 0.)), FVector(0, 0, 0)), FVector(1.0)));
 				RestCollection->AppendGeometry(*GeometryCollection::MakeCubeElement(FTransform(FQuat::MakeFromEuler(FVector(0, 0, 0.)), FVector(0, 0, 10)), FVector(1.0)));
 				FGeometryCollectionClusteringUtility::ClusterAllBonesUnderNewRoot(RestCollection.Get());
-				RestCollection->Transform[4] = FTransform(FQuat::MakeFromEuler(FVector(90.f, 0, 0.)), FVector(0, 0, 40));
+				RestCollection->Transform[4] = FTransform3f(FQuat4f::MakeFromEuler(FVector3f(90.f, 0, 0.)), FVector3f(0, 0, 40));
 			}
 			else
 			{
diff --git a/Engine/Source/Programs/HeadlessChaos/Private/GeometryCollection/GeometryCollectionTestUtility.cpp b/Engine/Source/Programs/HeadlessChaos/Private/GeometryCollection/GeometryCollectionTestUtility.cpp
--- a/Engine/Source/Programs/HeadlessChaos/Private/GeometryCollection/GeometryCollectionTestUtility.cpp
+++ b/Engine/Source/Programs/HeadlessChaos/Private/GeometryCollection/GeometryCollectionTestUtility.cpp
@@ -50,7 +50,7 @@ namespace GeometryCollectionTest {
 		RestCollection->SimulationType[2] = FGeometryCollection::ESimulationTypes::FST_Clustered;
 
 		GeometryCollectionAlgo::ParentTransforms(RestCollection.Get(), 2, { 0,1 });
-		RestCollection->Transform[2].SetTranslation(Position);
+		RestCollection->Transform[2].SetTranslation(FVector3f(Position));
 
 		return RestCollection;
 	}
@@ -68,7 +68,7 @@ namespace GeometryCollectionTest {
 		RestCollection->SimulationType[4] = FGeometryCollection::ESimulationTypes::FST_Clustered;
 
 		GeometryCollectionAlgo::ParentTransforms(RestCollection.Get(), 4, { 0,1,2,3 });
-		RestCollection->Transform[4].SetTranslation(Position);
+		RestCollection->Transform[4].SetTranslation(FVector3f(Position));
 
 		return RestCollection;
 	}
@@ -105,8 +105,8 @@ namespace GeometryCollectionTest {
 
 		GeometryCollectionAlgo::ParentTransforms(RestCollection.Get(), 2, { 0,1 });
 		GeometryCollectionAlgo::ParentTransforms(RestCollection.Get(), 3, { 2 });
-		RestCollection->Transform[2].SetTranslation(FVector(0.f, 10.f, 0.f));
-		RestCollection->Transform[3].SetTranslation(Position);
+		RestCollection->Transform[2].SetTranslation(FVector3f(0.f, 10.f, 0.f));
+		RestCollection->Transform[3].SetTranslation(FVector3f(Position));
 
 		return RestCollection;
 	}
@@ -143,7 +143,7 @@ namespace GeometryCollectionTest {
 		RestCollection->AppendGeometry(*GeometryCollection::MakeCubeElement(FTransform(FQuat::MakeFromEuler(FVector(0.f)), FVector(300,0,0)), FVector(1.0)));
 
 		RestCollection->AddElements(3, FGeometryCollection::TransformGroup);
-		RestCollection->Transform[6].SetTranslation(Position);
+		RestCollection->Transform[6].SetTranslation(FVector3f(Position));
 
 		// @todo(ClusteringUtils) This is a bad assumption, the state flags should be initialized to zero.
 		RestCollection->SimulationType[0] = FGeometryCollection::ESimulationTypes::FST_Rigid;
@@ -174,7 +174,7 @@ namespace GeometryCollectionTest {
 		RestCollection->AppendGeometry(*GeometryCollection::MakeCubeElement(FTransform(FQuat::MakeFromEuler(FVector(0.f)), FVector(500, 0, 0)), FVector(1.0)));
 
 		RestCollection->AddElements(3, FGeometryCollection::TransformGroup);
-		RestCollection->Transform[8].SetTranslation(Position);
+		RestCollection->Transform[8].SetTranslation(FVector3f(Position));
 
 		// @todo(ClusteringUtils) This is a bad assumption, the state flags should be initialized to zero.
 		RestCollection->SimulationType[0] = FGeometryCollection::ESimulationTypes::FST_Rigid;
diff --git a/Engine/Source/Runtime/Experimental/Chaos/Private/GeometryCollection/GeometryCollection.cpp b/Engine/Source/Runtime/Experimental/Chaos/Private/GeometryCollection/GeometryCollection.cpp
--- a/Engine/Source/Runtime/Experimental/Chaos/Private/GeometryCollection/GeometryCollection.cpp
+++ b/Engine/Source/Runtime/Experimental/Chaos/Private/GeometryCollection/GeometryCollection.cpp
@@ -216,7 +216,6 @@ int32 FGeometryCollection::AppendGeometry(const FGeometryCollection & Element, i
 	const TManagedArray<int32>& ElementFaceStart = Element.FaceStart;
 	const TManagedArray<int32>& ElementFaceCount = Element.FaceCount;
 
-	const TManagedArray<FTransform>& ElementTransform = Element.Transform;
 	const TManagedArray<FString>& ElementBoneName = Element.BoneName;
 	const TManagedArray<FGeometryCollectionSection>& ElementSections = Element.Sections;
 
@@ -407,7 +406,7 @@ bool FGeometryCollection::AppendEmbeddedInstance(int32 InExemplarIndex, int32 In
 
 	// add a new embedded instance
 	int32 Element = AddElements(1, FGeometryCollection::TransformGroup);
-	Transform[Element] = InTransform;
+	Transform[Element] = FTransform3f(InTransform);
 	Parent[Element] = InParentIndex;
 	Children[InParentIndex].Add(Element);
 	SimulationType[Element] = FST_None;
@@ -1603,18 +1602,18 @@ void FGeometryCollection::Init(FGeometryCollection* Collection, const TArray<flo
 		TManagedArray<int32>& MaterialID = Collection->MaterialID;
 		TManagedArray<int32>& MaterialIndex = Collection->MaterialIndex;
 		TManagedArray<bool>& Internal = Collection->Internal;
-		TManagedArray<FTransform>& Transform = Collection->Transform;
+		TManagedArray<FTransform3f>& Transform = Collection->Transform;
 		TManagedArray<int32>& BoneMap = Collection->BoneMap;
 		
 		Collection->SetNumUVLayers(1);
 
 		// set the vertex information
 		TManagedArray<FVector2f>* UV0 = Collection->FindUVLayer(0);
-		FVector3d TempVertices(0.f, 0.f, 0.f);
+		FVector3f TempVertices(0.f, 0.f, 0.f);
 		for (int32 Idx = 0; Idx < NumNewVertices; ++Idx)
 		{
 			Vertices[Idx] = FVector3f(RawVertexArray[3 * Idx], RawVertexArray[3 * Idx + 1], RawVertexArray[3 * Idx + 2]);
-			TempVertices += FVector3d(Vertices[Idx]);
+			TempVertices += Vertices[Idx];
 			(*UV0)[Idx] = FVector2f::ZeroVector;
 
 			Colors[Idx] = FLinearColor::White;
@@ -1625,7 +1624,7 @@ void FGeometryCollection::Init(FGeometryCollection* Collection, const TArray<flo
 
 		// set the particle information
 		TempVertices /= (float)NumNewVertices;
-		Transform[0] = FTransform(TempVertices);
+		Transform[0] = FTransform3f(TempVertices);
 		Transform[0].NormalizeRotation();
 
 		// set the index information
@@ -1767,12 +1766,12 @@ void FGeometryCollection::WriteDataToHeaderFile(const FString &Name, const FStri
 	DataFile << "const TArray<FTransform> " << TCHAR_TO_UTF8(*Name) << "::RawTransformArray = {" << endl;
 
 	int32 NumTransforms = NumElements(FGeometryCollection::TransformGroup);
-	const TManagedArray<FTransform>& TransformArray = Transform;
+	const TManagedArray<FTransform3f>& TransformArray = Transform;
 	for (int32 IdxTransform = 0; IdxTransform < NumTransforms; ++IdxTransform)
 	{
-		FQuat Rotation = TransformArray[IdxTransform].GetRotation();
-		FVector Translation = TransformArray[IdxTransform].GetTranslation();
-		FVector Scale3D = TransformArray[IdxTransform].GetScale3D();
+		FQuat4f Rotation = TransformArray[IdxTransform].GetRotation();
+		FVector3f Translation = TransformArray[IdxTransform].GetTranslation();
+		FVector3f Scale3D = TransformArray[IdxTransform].GetScale3D();
 
 		DataFile << "   FTransform(FQuat(" <<
 			Rotation.X << ", " <<
@@ -1925,7 +1924,7 @@ FGeometryCollection* FGeometryCollection::NewGeometryCollection(const TArray<flo
 	TManagedArray<int32>&  MaterialID = RestCollection->MaterialID;
 	TManagedArray<int32>&  MaterialIndex = RestCollection->MaterialIndex;
 	TManagedArray<bool>& Internal = RestCollection->Internal;
-	TManagedArray<FTransform>&  Transform = RestCollection->Transform;
+	TManagedArray<FTransform3f>&  Transform = RestCollection->Transform;
 	TManagedArray<int32>& Parent = RestCollection->Parent;
 	TManagedArray<TSet<int32>>& Children = RestCollection->Children;
 	TManagedArray<int32>& SimulationType = RestCollection->SimulationType;
@@ -1951,7 +1950,7 @@ FGeometryCollection* FGeometryCollection::NewGeometryCollection(const TArray<flo
 
 	for (int32 Idx = 0; Idx < NumNewTransforms; ++Idx)
 	{
-		Transform[Idx] = RawTransformArray[Idx];
+		Transform[Idx] = FTransform3f(RawTransformArray[Idx]);
 		Transform[Idx].NormalizeRotation();
 
 		Parent[Idx] = RawParentArray[Idx];
@@ -2051,7 +2050,7 @@ TArray<TArray<int32>> FGeometryCollection::ConnectionGraph()
 	TArray<TArray<int32>> Connectivity;
 	Connectivity.Init(TArray<int32>(), NumTransforms);
 
-	TArray<FTransform> GlobalMatrices;
+	TArray<FTransform3f> GlobalMatrices;
 	GeometryCollectionAlgo::GlobalMatrices(Transform, Parent, GlobalMatrices);
 
 	TArray<FVector> Pts;
@@ -2061,7 +2060,7 @@ TArray<TArray<int32>> FGeometryCollection::ConnectionGraph()
 		if (IsGeometry(TransformGroupIndex))
 		{
 			Remap.Add(Pts.Num(), TransformGroupIndex);
-			Pts.Add(GlobalMatrices[TransformGroupIndex].GetTranslation());
+			Pts.Add(FVector(GlobalMatrices[TransformGroupIndex].GetTranslation()));
 		}
 	}
 
diff --git a/Engine/Source/Runtime/Experimental/Chaos/Private/GeometryCollection/GeometryCollectionAlgo.cpp b/Engine/Source/Runtime/Experimental/Chaos/Private/GeometryCollection/GeometryCollectionAlgo.cpp
--- a/Engine/Source/Runtime/Experimental/Chaos/Private/GeometryCollection/GeometryCollectionAlgo.cpp
+++ b/Engine/Source/Runtime/Experimental/Chaos/Private/GeometryCollection/GeometryCollectionAlgo.cpp
@@ -15,8 +15,9 @@ DEFINE_LOG_CATEGORY_STATIC(GeometryCollectionAlgoLog, Log, All);
 namespace GeometryCollectionAlgo
 {
 
+	template<typename TransformType>
 	void PrintParentHierarchyRecursive(int32 Index
-		, const TManagedArray<FTransform>& Transform
+		, const TManagedArray<TransformType>& Transform
 		, const TManagedArray<int32>& Parent
 		, const TManagedArray<TSet<int32>>& Children
 		, const TManagedArray<int32>& SimulationType
@@ -47,7 +48,7 @@ namespace GeometryCollectionAlgo
 	{
 		check(Collection);
 
-		const TManagedArray<FTransform>& Transform = Collection->Transform;
+		const TManagedArray<FTransform3f>& Transform = Collection->Transform;
 		const TManagedArray<FString>& BoneNames = Collection->BoneName;
 		const TManagedArray<int32>& Parent = Collection->Parent;
 		const TManagedArray<TSet<int32>>& Children = Collection->Children;
@@ -149,10 +150,10 @@ namespace GeometryCollectionAlgo
 
 	FVector AveragePosition(FGeometryCollection* Collection, const TArray<int32>& Indices)
 	{
-		TManagedArray<FTransform>& Transform = Collection->Transform;
+		TManagedArray<FTransform3f>& Transform = Collection->Transform;
 		int32 NumIndices = Indices.Num();
 
-		FVector Translation(0);
+		FVector3f Translation(0);
 		for (int32 Index = 0; Index < NumIndices; Index++)
 		{
 			Translation += Transform[Indices[Index]].GetTranslation();
@@ -161,7 +162,7 @@ namespace GeometryCollectionAlgo
 		{
 			Translation /= static_cast<float>(NumIndices);
 		}
-		return Translation;
+		return FVector(Translation);
 	}
 
 	bool HasMultipleRoots(FGeometryCollection * Collection)
@@ -210,14 +211,14 @@ namespace GeometryCollectionAlgo
 	{
 		check(GeometryCollection != nullptr);
 
-		TManagedArray<FTransform>& Transform = GeometryCollection->Transform;
+		TManagedArray<FTransform3f>& Transform = GeometryCollection->Transform;
 		TManagedArray<int32>& Parents = GeometryCollection->Parent;
 		TManagedArray<TSet<int32>>& Children = GeometryCollection->Children;
 
 		if (ensure(-1 <= TransformIndex && TransformIndex < GeometryCollection->NumElements(FGeometryCollection::TransformGroup)))
 		{
 			// pre calculate global positions
-			TArray<FTransform> GlobalTransform;
+			TArray<FTransform3f> GlobalTransform;
 			GeometryCollectionAlgo::GlobalMatrices(GeometryCollection->Transform, Parents, GlobalTransform);
 
 			// append children 
@@ -241,7 +242,7 @@ namespace GeometryCollectionAlgo
 				}
 			}
 
-			FTransform ParentInverse = FTransform::Identity;
+			FTransform3f ParentInverse = FTransform3f::Identity;
 			if (TransformIndex != FGeometryCollection::Invalid)
 			{
 				Children[TransformIndex].Append(SelectedBones);
@@ -383,7 +384,9 @@ namespace GeometryCollectionAlgo
 		}
 	}
 
-	void GlobalMatricesHelper(const int32 Index, const TManagedArray<int32>& Parents, const TManagedArray<FTransform>& Transform, TArray<bool>& IsTransformComputed, const TManagedArray<FTransform>* UniformScale, TArray<FMatrix>& OutGlobalTransforms)
+	// #note: this version outputs FTransforms to support functionality for getting global matrices for an array of indices.
+	template<typename TransformType, typename TransformTypeOut>
+	void GlobalMatricesHelper(const int32 Index, const TManagedArray<int32>& Parents, const TManagedArray<TransformType>& Transform, TArray<bool>& IsTransformComputed, const TManagedArray<FTransform>* UniformScale, TArray<TransformTypeOut>& OutGlobalTransforms)
 	{
 		if (IsTransformComputed[Index])
 		{
@@ -400,27 +403,27 @@ namespace GeometryCollectionAlgo
 		{
 			const int32 ProcessIndex = ToProcess.Pop(false);
 			const int32 ParentIndex = Parents[ProcessIndex];
-			FMatrix Result = FTransform(Transform[ProcessIndex]).ToMatrixWithScale();
+			FTransform Result = FTransform(Transform[ProcessIndex]);
 			if (ParentIndex != FGeometryCollection::Invalid)
 			{
-				Result *= OutGlobalTransforms[ParentIndex];
+				Result *= FTransform(OutGlobalTransforms[ParentIndex]);
 			}
 
 			if (UniformScale)
 			{
-				OutGlobalTransforms[ProcessIndex] = (*UniformScale)[ProcessIndex].ToMatrixWithScale() * Result;
+				OutGlobalTransforms[ProcessIndex] = TransformTypeOut((*UniformScale)[ProcessIndex] * Result);
 			}
 			else
 			{
-				OutGlobalTransforms[ProcessIndex] = Result;
+				OutGlobalTransforms[ProcessIndex] = TransformTypeOut(Result);
 			}
 
 			IsTransformComputed[ProcessIndex] = true;
 		}
 	}
 
-	// #note: this version outputs FTransforms to support functionality for getting global matrices for an array of indices.
-	void GlobalMatricesHelper(const int32 Index, const TManagedArray<int32>& Parents, const TManagedArray<FTransform>& Transform, TArray<bool>& IsTransformComputed, const TManagedArray<FTransform>* UniformScale, TArray<FTransform>& OutGlobalTransforms)
+	template<typename TransformType>
+	void GlobalMatricesHelper(const int32 Index, const TManagedArray<int32>& Parents, const TManagedArray<TransformType>& Transform, TArray<bool>& IsTransformComputed, const TManagedArray<FTransform>* UniformScale, TArray<FMatrix>& OutGlobalTransforms)
 	{
 		if (IsTransformComputed[Index])
 		{
@@ -437,7 +440,7 @@ namespace GeometryCollectionAlgo
 		{
 			const int32 ProcessIndex = ToProcess.Pop(false);
 			const int32 ParentIndex = Parents[ProcessIndex];
-			FTransform Result = FTransform(Transform[ProcessIndex]);
+			FMatrix Result = FTransform(Transform[ProcessIndex]).ToMatrixWithScale();
 			if (ParentIndex != FGeometryCollection::Invalid)
 			{
 				Result *= OutGlobalTransforms[ParentIndex];
@@ -445,7 +448,7 @@ namespace GeometryCollectionAlgo
 
 			if (UniformScale)
 			{
-				OutGlobalTransforms[ProcessIndex] = (*UniformScale)[ProcessIndex] * Result;
+				OutGlobalTransforms[ProcessIndex] = (*UniformScale)[ProcessIndex].ToMatrixWithScale() * Result;
 			}
 			else
 			{
@@ -456,7 +459,7 @@ namespace GeometryCollectionAlgo
 		}
 	}
 
-	FTransform GlobalMatricesHelperForIndices(const int32 Index, const FGeometryDynamicCollection& DynamicCollection, TArray<bool>& IsTransformComputed, const TManagedArray<FTransform>* UniformScale, TArray<FTransform>& TransformCache)
+	FTransform GlobalMatricesHelperForIndicesDynCol(const int32 Index, const FGeometryDynamicCollection& DynamicCollection, TArray<bool>& IsTransformComputed, const TManagedArray<FTransform>* UniformScale, TArray<FTransform>& TransformCache)
 	{
 		if (IsTransformComputed[Index])
 		{
@@ -494,42 +497,43 @@ namespace GeometryCollectionAlgo
 		return TransformCache[Index];
 	}
 
-	FTransform GlobalMatricesHelperForIndices(const int32 Index, const TManagedArray<int32>& Parents, const TManagedArray<FTransform>& Transform, TArray<bool>& IsTransformComputed, const TManagedArray<FTransform>* UniformScale, TArray<FTransform>& TransformCache)
+	template<typename TransnformType>
+	TransnformType GlobalMatricesHelperForIndices(const int32 Index, const TManagedArray<int32>& Parents, const TManagedArray<TransnformType>& Transform, TArray<bool>& IsTransformComputed, const TManagedArray<TransnformType>* UniformScale, TArray<FTransform>& TransformCache)
 	{
 		if (IsTransformComputed[Index])
 		{
-			return TransformCache[Index];
+			return TransnformType(TransformCache[Index]);
 		}
 
 		FIndicesNeedMatricesArray ToProcess;
 		if (!GlobalMatricesGetIndicesToProcessHelper(Index, Parents, IsTransformComputed, ToProcess))
 		{
-			return TransformCache[Index];
+			return TransnformType(TransformCache[Index]);
 		}
 
 		while (!ToProcess.IsEmpty())
 		{
 			const int32 ProcessIndex = ToProcess.Pop(false);
 			const int32 ParentIndex = Parents[ProcessIndex];
-			FTransform Result = FTransform(Transform[ProcessIndex]);
+			TransnformType Result = Transform[ProcessIndex];
 			if (ParentIndex != FGeometryCollection::Invalid)
 			{
-				Result *= TransformCache[ParentIndex];
+				Result *= TransnformType(TransformCache[ParentIndex]);
 			}
 
 			if (UniformScale)
 			{
-				TransformCache[ProcessIndex] = (*UniformScale)[ProcessIndex] * Result;
+				TransformCache[ProcessIndex] = FTransform((*UniformScale)[ProcessIndex] * Result);
 			}
 			else
 			{
-				TransformCache[ProcessIndex] = Result;
+				TransformCache[ProcessIndex] = FTransform(Result);
 			}
 
 			IsTransformComputed[ProcessIndex] = true;
 		}
 
-		return TransformCache[Index];
+		return TransnformType(TransformCache[Index]);
 	}
 
 	namespace Private 
@@ -548,21 +552,38 @@ namespace GeometryCollectionAlgo
 		}
 	}
 
-	FTransform GlobalMatrix(const TManagedArray<FTransform>& RelativeTransforms, const TManagedArray<int32>& Parents, int32 Index)
+	template<typename TransformType>
+	TransformType GlobalMatrixTemplate(const TManagedArray<TransformType>& RelativeTransforms, const TManagedArray<int32>& Parents, int32 Index)
 	{
-		FTransform Transform = FTransform::Identity;
+		TransformType Transform = TransformType::Identity;
 
 		if (RelativeTransforms.IsValidIndex(Index))
 		{
 			while (Index != FGeometryCollection::Invalid)
 			{
-				Transform = Transform * FTransform(RelativeTransforms[Index]);
+				Transform = Transform * RelativeTransforms[Index];
 				Index = Parents[Index];
 			}
 		}
 		return Transform;
 	}
 
+
+	FTransform GlobalMatrix(const TManagedArray<FTransform>& RelativeTransforms, const TManagedArray<int32>& Parents, int32 Index)
+	{
+		return GlobalMatrixTemplate<FTransform>(RelativeTransforms, Parents, Index);
+	}
+
+	FTransform3f GlobalMatrix3f(const TManagedArray<FTransform3f>& RelativeTransforms, const TManagedArray<int32>& Parents, int32 Index)
+	{
+		return GlobalMatrixTemplate<FTransform3f>(RelativeTransforms, Parents, Index);
+	}
+
+	FTransform GlobalMatrix(const TManagedArray<FTransform3f>& RelativeTransforms, const TManagedArray<int32>& Parents, int32 Index)
+	{
+		return FTransform(GlobalMatrixTemplate<FTransform3f>(RelativeTransforms, Parents, Index));
+	}
+
 	namespace Private
 	{
 		void GlobalMatrices(const FGeometryDynamicCollection& DynamicCollection, const TArray<int32>& Indices, TArray<FTransform>& OutGlobalTransforms)
@@ -577,13 +598,13 @@ namespace GeometryCollectionAlgo
 			OutGlobalTransforms.SetNumUninitialized(Indices.Num(), false);
 			for (int Idx = 0; Idx < Indices.Num(); Idx++)
 			{
-				OutGlobalTransforms[Idx] = GlobalMatricesHelperForIndices(Indices[Idx], DynamicCollection, IsTransformComputed, nullptr, TransformCache);
+				OutGlobalTransforms[Idx] = GlobalMatricesHelperForIndicesDynCol(Indices[Idx], DynamicCollection, IsTransformComputed, nullptr, TransformCache);
 			}
 		}
 	}
 
-	template<class TransformType>
-	void GlobalMatricesTemplate(const TManagedArray<TransformType>& RelativeTransforms, const TManagedArray<int32>& Parents, const TArray<int32>& Indices, TArray<FTransform>& OutGlobalTransforms)
+	template<class TransformType, class TransformTypeOut>
+	void GlobalMatricesTemplate(const TManagedArray<TransformType>& RelativeTransforms, const TManagedArray<int32>& Parents, const TArray<int32>& Indices, TArray<TransformTypeOut>& OutGlobalTransforms)
 	{
 		TArray<bool> IsTransformComputed;
 		IsTransformComputed.AddDefaulted(RelativeTransforms.Num());
@@ -594,13 +615,18 @@ namespace GeometryCollectionAlgo
 		OutGlobalTransforms.SetNumUninitialized(Indices.Num(), false);
 		for (int Idx = 0; Idx < Indices.Num(); Idx++)
 		{
-			OutGlobalTransforms[Idx] = GlobalMatricesHelperForIndices(Indices[Idx], Parents, RelativeTransforms, IsTransformComputed, nullptr, TransformCache);
+			OutGlobalTransforms[Idx] = GlobalMatricesHelperForIndices<TransformType>(Indices[Idx], Parents, RelativeTransforms, IsTransformComputed, nullptr, TransformCache);
 		}
 	}
 
 	void GlobalMatrices(const TManagedArray<FTransform>& RelativeTransforms, const TManagedArray<int32>& Parents, const TArray<int32>& Indices, TArray<FTransform>& Transforms)
 	{
-		GlobalMatricesTemplate<FTransform>(RelativeTransforms, Parents, Indices, Transforms);
+		GlobalMatricesTemplate<FTransform, FTransform>(RelativeTransforms, Parents, Indices, Transforms);
+	}
+
+	void GlobalMatrices(const TManagedArray<FTransform3f>& RelativeTransforms, const TManagedArray<int32>& Parents, const TArray<int32>& Indices, TArray<FTransform3f>& Transforms)
+	{
+		GlobalMatricesTemplate<FTransform3f, FTransform3f>(RelativeTransforms, Parents, Indices, Transforms);
 	}
 
 	void GlobalMatricesFromRoot(const int32 ParentTransformIndex, const TManagedArray<FTransform>& RelativeTransforms, const TManagedArray<TSet<int32>>& Children, TArray<FMatrix>& Transforms)
@@ -650,8 +676,8 @@ namespace GeometryCollectionAlgo
 		}
 	}
 
-	template<typename MatrixType>
-	void GlobalMatrices(const TManagedArray<FTransform>& RelativeTransforms, const TManagedArray<int32>& Parents, TArray<MatrixType>& OutGlobalTransforms)
+	template<typename MatrixType, typename TransformType>
+	void GlobalMatrices(const TManagedArray<TransformType>& RelativeTransforms, const TManagedArray<int32>& Parents, TArray<MatrixType>& OutGlobalTransforms)
 	{
 		int32 NumTransforms = RelativeTransforms.Num();
 
@@ -667,9 +693,13 @@ namespace GeometryCollectionAlgo
 	}
 
 	template void CHAOS_API GlobalMatrices<FTransform>(const TManagedArray<FTransform>&, const TManagedArray<int32>&, const TManagedArray<FTransform>&, TArray<FTransform>&);
-	template void CHAOS_API GlobalMatrices<FTransform>(const TManagedArray<FTransform>&, const TManagedArray<int32>&, TArray<FTransform>&);
 	template void CHAOS_API GlobalMatrices<FMatrix>(const TManagedArray<FTransform>&, const TManagedArray<int32>&, const TManagedArray<FTransform>&, TArray<FMatrix>&);
-	template void CHAOS_API GlobalMatrices<FMatrix>(const TManagedArray<FTransform>&, const TManagedArray<int32>&, TArray<FMatrix>&);
+
+	template void CHAOS_API GlobalMatrices<FTransform, FTransform>(const TManagedArray<FTransform>&, const TManagedArray<int32>&, TArray<FTransform>&);
+	template void CHAOS_API GlobalMatrices<FMatrix, FTransform>(const TManagedArray<FTransform>&, const TManagedArray<int32>&, TArray<FMatrix>&);
+	template void CHAOS_API GlobalMatrices<FTransform, FTransform3f>(const TManagedArray<FTransform3f>&, const TManagedArray<int32>&, TArray<FTransform>&);
+	template void CHAOS_API GlobalMatrices<FMatrix, FTransform3f>(const TManagedArray<FTransform3f>&, const TManagedArray<int32>&, TArray<FMatrix>&);
+
 
 	void FloodForOverlappedPairs(int Level, int32 BoneIndex, TMap<int32, int32> &BoneToGroup, const TManagedArray<int32>& Levels, const TMap<int32, FBox>& BoundingBoxes, TSet<TTuple<int32, int32>>& OutOverlappedPairs)
 	{
@@ -758,7 +788,7 @@ namespace GeometryCollectionAlgo
 	{
 		check(GeometryCollection);
 
-		TManagedArray<FTransform>& Transform = GeometryCollection->Transform;
+		TManagedArray<FTransform3f>& Transform = GeometryCollection->Transform;
 		if (Transform.Num())
 		{
 			const TManagedArray<int32>& BoneMap = GeometryCollection->BoneMap;
@@ -777,14 +807,14 @@ namespace GeometryCollectionAlgo
 				CenterOfMass[ParticleIndex] += Vertex[i];
 			}
 
-			FVector CombinedCenterOfMassWorld(ForceInitToZero);
+			FVector3f CombinedCenterOfMassWorld(ForceInitToZero);
 			for (int i = 0; i < Transform.Num(); i++)
 			{
 				if (SurfaceParticlesCount[i])
 				{
 					CenterOfMass[i] /= static_cast<float>(SurfaceParticlesCount[i]);
 
-					FTransform Tmp((FVector)CenterOfMass[i]);
+					FTransform3f Tmp((FVector3f)CenterOfMass[i]);
 
 					// Translate back to original object space position (because vertex position will be centered at the origin), 
 					// then apply the original parent transform.  This ensures the pivot remains the same
@@ -805,7 +835,7 @@ namespace GeometryCollectionAlgo
 			{
 				for (int i = 0; i < Transform.Num(); i++)
 				{
-					FTransform Tmp(-CombinedCenterOfMassWorld);
+					FTransform3f Tmp(-CombinedCenterOfMassWorld);
 
 					// Apply the parent transform, then center at the origin
 					Transform[i] = Transform[i] * Tmp;
diff --git a/Engine/Source/Runtime/Experimental/Chaos/Private/GeometryCollection/GeometryCollectionClusteringUtility.cpp b/Engine/Source/Runtime/Experimental/Chaos/Private/GeometryCollection/GeometryCollectionClusteringUtility.cpp
--- a/Engine/Source/Runtime/Experimental/Chaos/Private/GeometryCollection/GeometryCollectionClusteringUtility.cpp
+++ b/Engine/Source/Runtime/Experimental/Chaos/Private/GeometryCollection/GeometryCollectionClusteringUtility.cpp
@@ -24,7 +24,7 @@ int32 FGeometryCollectionClusteringUtility::ClusterBonesUnderNewNodeWithParent(F
 	check(GeometryCollection);
 
 
-	TManagedArray<FTransform>& Transforms = GeometryCollection->Transform;
+	TManagedArray<FTransform3f>& Transforms = GeometryCollection->Transform;
 	TManagedArray<FString>& BoneNames = GeometryCollection->BoneName;
 	TManagedArray<int32>& Parents = GeometryCollection->Parent;
 	TManagedArray<TSet<int32>>& Children = GeometryCollection->Children;
@@ -38,7 +38,7 @@ int32 FGeometryCollectionClusteringUtility::ClusterBonesUnderNewNodeWithParent(F
 	Children[NewBoneIndex] = TSet<int32>(SelectedBones);
 	SimType[NewBoneIndex] = FGeometryCollection::ESimulationTypes::FST_Clustered;
 
-	Transforms[NewBoneIndex] = FTransform::Identity;
+	Transforms[NewBoneIndex] = FTransform3f::Identity;
 
 	// re-parent all the geometry nodes under the new shared bone
 	GeometryCollectionAlgo::ParentTransforms(GeometryCollection, NewBoneIndex, SelectedBones);
@@ -68,7 +68,7 @@ void FGeometryCollectionClusteringUtility::ClusterAllBonesUnderNewRoot(FGeometry
 	check(GeometryCollection);
 	bool CalcNewLocalTransform = true;
 
-	TManagedArray<FTransform>& Transforms = GeometryCollection->Transform;
+	TManagedArray<FTransform3f>& Transforms = GeometryCollection->Transform;
 	TManagedArray<FString>& BoneNames = GeometryCollection->BoneName;
 	TManagedArray<int32>& Parents = GeometryCollection->Parent;
 	TManagedArray<TSet<int32>>& Children = GeometryCollection->Children;
@@ -114,11 +114,11 @@ void FGeometryCollectionClusteringUtility::ClusterAllBonesUnderNewRoot(FGeometry
 		FVector3f SumOfOffsets(0, 0, 0);
 		for (int32 ChildBoneIndex : ChildBones)
 		{
-			ExplodedVectors[ChildBoneIndex] = FVector3f(Transforms[ChildBoneIndex].GetLocation());	//LWC_TODO: Precision loss
-			ExplodedTransforms[ChildBoneIndex] = Transforms[ChildBoneIndex];
+			ExplodedVectors[ChildBoneIndex] = FVector3f(Transforms[ChildBoneIndex].GetLocation());
+			ExplodedTransforms[ChildBoneIndex] = FTransform(Transforms[ChildBoneIndex]);
 			SumOfOffsets += ExplodedVectors[ChildBoneIndex];
 		}
-		ExplodedTransforms[RootNoneIndex] = Transforms[RootNoneIndex];
+		ExplodedTransforms[RootNoneIndex] = FTransform(Transforms[RootNoneIndex]);
 		// This bones offset is the average of all the selected bones
 		ExplodedVectors[RootNoneIndex] = SumOfOffsets / static_cast<float>(ChildBones.Num());
 	}
@@ -129,7 +129,7 @@ void FGeometryCollectionClusteringUtility::ClusterAllBonesUnderNewRoot(FGeometry
 		Parents[ChildBoneIndex] = RootNoneIndex;
 	}
 
-	Transforms[RootNoneIndex] = FTransform::Identity;
+	Transforms[RootNoneIndex] = FTransform3f::Identity;
 
 
 	RecursivelyUpdateChildBoneNames(RootNoneIndex, Children, BoneNames);
@@ -145,7 +145,6 @@ void FGeometryCollectionClusteringUtility::ClusterBonesUnderExistingRoot(FGeomet
 
 	TManagedArray<int32>& Levels = GeometryCollection->ModifyAttribute<int32>("Level", FGeometryCollection::TransformGroup);
 
-	TManagedArray<FTransform>& Transforms = GeometryCollection->Transform;
 	TManagedArray<FString>& BoneNames = GeometryCollection->BoneName;
 	TManagedArray<int32>& Parents = GeometryCollection->Parent;
 	TManagedArray<TSet<int32>>& Children = GeometryCollection->Children;
@@ -212,7 +211,6 @@ void FGeometryCollectionClusteringUtility::ClusterBonesUnderExistingNode(FGeomet
 
 	TManagedArray<int32>& Parents = GeometryCollection->Parent;
 	TManagedArray<TSet<int32>>& Children = GeometryCollection->Children;
-	TManagedArray<FTransform>& Transforms = GeometryCollection->Transform;
 	TManagedArray<FString>& BoneNames = GeometryCollection->BoneName;
 
 	// These attributes are apparently deprecated?
@@ -244,7 +242,7 @@ void FGeometryCollectionClusteringUtility::ClusterBonesUnderExistingNode(FGeomet
 		if (!IllegalOperation)
 		{
 			TArray<int32> ParentsToUpdateNames;
-			// determine original parents of moved nodes so we can update their childrens names
+			// determine original parents of moved nodes so we can update their children's names
 			for (int32 SourceElement : SourceElementsIn)
 			{
 				int32 Parent = Parents[SourceElement];
@@ -310,7 +308,6 @@ void FGeometryCollectionClusteringUtility::CollapseHierarchyOneLevel(FGeometryCo
 
 	TManagedArray<int32>& Parents = GeometryCollection->Parent;
 	TManagedArray<TSet<int32>>& Children = GeometryCollection->Children;
-	TManagedArray<FTransform>& Transforms = GeometryCollection->Transform;
 	TManagedArray<FString>& BoneNames = GeometryCollection->BoneName;
 	TManagedArray<int32>& Levels = GeometryCollection->ModifyAttribute<int32>("Level", FGeometryCollection::TransformGroup);
 
@@ -851,7 +848,6 @@ void FGeometryCollectionClusteringUtility::MoveUpOneHierarchyLevel(FGeometryColl
 
 	TManagedArray<int32>& Parents = GeometryCollection->Parent;
 	TManagedArray<TSet<int32>>& Children = GeometryCollection->Children;
-	TManagedArray<FTransform>& Transforms = GeometryCollection->Transform;
 	TManagedArray<FString>& BoneNames = GeometryCollection->BoneName;
 
 	for (int32 BoneIndex : SelectedBones)
@@ -971,10 +967,9 @@ bool FGeometryCollectionClusteringUtility::RemoveDanglingClusters(FGeometryColle
 {
 	check(GeometryCollection);
 
-	const TManagedArray<FTransform>& Transforms = GeometryCollection->Transform;
 	const TManagedArray<int32>& SimulationType = GeometryCollection->SimulationType;
 
-	const int32 TransformCount = Transforms.Num();
+	const int32 TransformCount = GeometryCollection->Transform.Num();
 	TArray<int32> DeletionList;
 	for (int32 Idx = 0; Idx < TransformCount; ++Idx)
 	{
diff --git a/Engine/Source/Runtime/Experimental/Chaos/Private/GeometryCollection/GeometryCollectionConnectionGraphUtility.cpp b/Engine/Source/Runtime/Experimental/Chaos/Private/GeometryCollection/GeometryCollectionConnectionGraphUtility.cpp
--- a/Engine/Source/Runtime/Experimental/Chaos/Private/GeometryCollection/GeometryCollectionConnectionGraphUtility.cpp
+++ b/Engine/Source/Runtime/Experimental/Chaos/Private/GeometryCollection/GeometryCollectionConnectionGraphUtility.cpp
@@ -72,7 +72,7 @@ FGeometryCollectionConnectionGraphGenerator::ComputeVoronoiNeighbors(const FGeom
 {
 	FVoronoiNeighbors Neighbors;
 
-	const TManagedArray<FTransform>& Transforms = Collection.Transform;
+	const TManagedArray<FTransform3f>& Transforms = Collection.Transform;
 	const TManagedArray<FTransform>& MassToLocal = Collection.GetAttribute<FTransform>("MassToLocal", FGeometryCollection::TransformGroup);
 
 	const TSet<int32>& Children = Collection.Children[ClusterTransformIndex];
@@ -81,7 +81,7 @@ FGeometryCollectionConnectionGraphGenerator::ComputeVoronoiNeighbors(const FGeom
 	
 	for (const int32& ChildTransformIndex: Children)
 	{
-		const FTransform ChildTransform = MassToLocal[ChildTransformIndex] * Transforms[ChildTransformIndex];
+		const FTransform ChildTransform = MassToLocal[ChildTransformIndex] * FTransform(Transforms[ChildTransformIndex]);
 		Neighbors.Points.Add(ChildTransform.GetLocation());
 	}
 	
diff --git a/Engine/Source/Runtime/Experimental/Chaos/Private/GeometryCollection/GeometryCollectionUtility.cpp b/Engine/Source/Runtime/Experimental/Chaos/Private/GeometryCollection/GeometryCollectionUtility.cpp
--- a/Engine/Source/Runtime/Experimental/Chaos/Private/GeometryCollection/GeometryCollectionUtility.cpp
+++ b/Engine/Source/Runtime/Experimental/Chaos/Private/GeometryCollection/GeometryCollectionUtility.cpp
@@ -39,12 +39,12 @@ namespace GeometryCollection
 		TManagedArray<int32>&  MaterialIndex = RestCollection->MaterialIndex;
 		TManagedArray<int32>&  MaterialID = RestCollection->MaterialID;
 		TManagedArray<bool>& Internal = RestCollection->Internal;
-		TManagedArray<FTransform>&  Transform = RestCollection->Transform;
+		TManagedArray<FTransform3f>&  Transform = RestCollection->Transform;
 		TManagedArray<int32>& SimType = RestCollection->SimulationType;
 		TManagedArray<int32>& BoneMap = RestCollection->BoneMap;
 
 		// set the particle information
-		Transform[0] = center;
+		Transform[0] = FTransform3f(center);
 		Transform[0].NormalizeRotation();
 		SimType[0] = FGeometryCollection::ESimulationTypes::FST_Rigid;
 
diff --git a/Engine/Source/Runtime/Experimental/Chaos/Private/GeometryCollection/ManagedArrayCollection.cpp b/Engine/Source/Runtime/Experimental/Chaos/Private/GeometryCollection/ManagedArrayCollection.cpp
--- a/Engine/Source/Runtime/Experimental/Chaos/Private/GeometryCollection/ManagedArrayCollection.cpp
+++ b/Engine/Source/Runtime/Experimental/Chaos/Private/GeometryCollection/ManagedArrayCollection.cpp
@@ -302,6 +302,11 @@ void FManagedArrayCollection::Append(const FManagedArrayCollection& InCollection
 void FManagedArrayCollection::RemoveAttribute(FName Name, FName Group)
 {
 	FKeyType Key = FManagedArrayCollection::MakeMapKey(Name, Group);
+	FValueType* FoundValue = Map.Find(Key);
+	if (FoundValue != nullptr)
+	{
+		FoundValue->Value->Empty();
+	}
 	Map.Remove(Key);
 }
 
@@ -712,11 +717,15 @@ void FManagedArrayCollection::Serialize(Chaos::FChaosArchive& Ar)
 		{
 			if (FValueType* Existing = Map.Find(Pair.Key))
 			{
-				if (ensureMsgf(Existing->ArrayType == Pair.Value.ArrayType, TEXT("Type change not supported. Ignoring serialized data")))
+				if (Existing->ArrayType == Pair.Value.ArrayType)
 				{
 					Existing->Value->ExchangeArrays(*Pair.Value.Value);	//if there is already an entry do an exchange. This way external arrays get correct serialization
 					//question: should we validate if group dependency has changed in some invalid way?
 				}
+				else
+				{
+					Existing->Value->Convert(*Pair.Value.Value);
+				}
 			}
 			else
 			{
diff --git a/Engine/Source/Runtime/Experimental/Chaos/Private/GeometryCollection/TransformCollection.cpp b/Engine/Source/Runtime/Experimental/Chaos/Private/GeometryCollection/TransformCollection.cpp
--- a/Engine/Source/Runtime/Experimental/Chaos/Private/GeometryCollection/TransformCollection.cpp
+++ b/Engine/Source/Runtime/Experimental/Chaos/Private/GeometryCollection/TransformCollection.cpp
@@ -25,7 +25,7 @@ FTransformCollection::FTransformCollection()
 void FTransformCollection::DefineTransformSchema(FManagedArrayCollection& InCollection)
 {
 	// Hierarchy Group
-	InCollection.AddAttribute<FTransform>(FTransformCollection::TransformAttribute, FTransformCollection::TransformGroup);
+	InCollection.AddAttribute<FTransform3f>(FTransformCollection::TransformAttribute, FTransformCollection::TransformGroup);
 	InCollection.AddAttribute<FString>("BoneName", FTransformCollection::TransformGroup);
 	InCollection.AddAttribute<FLinearColor>("BoneColor", FTransformCollection::TransformGroup);
 	InCollection.AddAttribute<int32>(FTransformCollection::ParentAttribute, FTransformCollection::TransformGroup);
@@ -35,7 +35,7 @@ void FTransformCollection::DefineTransformSchema(FManagedArrayCollection& InColl
 void FTransformCollection::Construct()
 {
 	// Hierarchy Group
-	AddExternalAttribute<FTransform>(FTransformCollection::TransformAttribute, FTransformCollection::TransformGroup, Transform);
+	AddExternalAttribute<FTransform3f>(FTransformCollection::TransformAttribute, FTransformCollection::TransformGroup, Transform);
 	AddExternalAttribute<FString>("BoneName", FTransformCollection::TransformGroup, BoneName);
 	AddExternalAttribute<FLinearColor>("BoneColor", FTransformCollection::TransformGroup, BoneColor);
 	AddExternalAttribute<int32>(FTransformCollection::ParentAttribute, FTransformCollection::TransformGroup, Parent);
@@ -74,7 +74,7 @@ FTransformCollection FTransformCollection::SingleTransform(const FTransform& Tra
 {
 	FTransformCollection TransformCollection;
 	TransformCollection.AddElements(1, FTransformCollection::TransformGroup);
-	TransformCollection.Transform[0] = TransformRoot;
+	TransformCollection.Transform[0] = FTransform3f(TransformRoot);
 	TransformCollection.Parent[0] = Invalid;
 	return TransformCollection;
 }
@@ -105,7 +105,7 @@ void FTransformCollection::Append(const FManagedArrayCollection& InCollection)
 int32 FTransformCollection::AppendTransform(const FTransformCollection & Element, const FTransform& TransformRoot)
 {
 	check(Element.NumElements(FTransformCollection::TransformGroup) > 0);
-	const TManagedArray<FTransform>& ElementTransform = Element.Transform;
+	const TManagedArray<FTransform3f>& ElementTransform = Element.Transform;
 	const TManagedArray<FString>& ElementBoneName = Element.BoneName;
 	const TManagedArray<FLinearColor>& ElementBoneColor = Element.BoneColor;
 	const TManagedArray<int32>& ElementParent = Element.Parent;
@@ -114,14 +114,15 @@ int32 FTransformCollection::AppendTransform(const FTransformCollection & Element
 	int OriginalNumTransforms = NumElements(FTransformCollection::TransformGroup);
 	int NumElements = Element.NumElements(FTransformCollection::TransformGroup);
 	int FirstNewElement = AddElements(NumElements, FTransformCollection::TransformGroup);
+	const FTransform3f TransformRoot3f(TransformRoot);
 	for (int Index = 0; Index < NumElements; Index++)
 	{
 		int ParticleIndex = FirstNewElement + Index;
-		TManagedArray<FTransform>& Transforms = Transform;
+		TManagedArray<FTransform3f>& Transforms = Transform;
 		if (ElementParent[Index] == FTransformCollection::Invalid)
 		{
 			// is root with additional transform
-			Transforms[ParticleIndex] = ElementTransform[Index] * TransformRoot;
+			Transforms[ParticleIndex] = ElementTransform[Index] * TransformRoot3f;
 		}
 		else
 		{
@@ -169,17 +170,19 @@ void FTransformCollection::RelativeTransformation(const int32& Index, const FTra
 	if (ensureMsgf(Index < NumElements(FTransformCollection::TransformGroup), TEXT("Index out of range.")))
 	{
 		TManagedArray<TSet<int32>>& ChildrenArray = Children;
-		TManagedArray<FTransform>& TransformArray = Transform;
+		TManagedArray<FTransform3f>& TransformArray = Transform;
 
+		const FTransform3f LocalOffset3f = FTransform3f(LocalOffset);
 		if (ChildrenArray[Index].Num())
 		{
-			FTransform LocalOffsetInverse = LocalOffset.Inverse();
+			
+			const FTransform3f LocalOffsetInverse = LocalOffset3f.Inverse();
 			for (int32 Child : ChildrenArray[Index])
 			{
-				TransformArray[Child] = TransformArray[Child] * LocalOffset.Inverse();
+				TransformArray[Child] = TransformArray[Child] * LocalOffsetInverse;
 			}
 		}
-		TransformArray[Index] = LocalOffset * TransformArray[Index];
+		TransformArray[Index] = LocalOffset3f * TransformArray[Index];
 	}
 }
 
@@ -193,7 +196,7 @@ void FTransformCollection::RemoveElements(const FName & Group, const TArray<int3
 
 			TManagedArray<int32>& ParentArray = Parent;
 			TManagedArray<TSet<int32>>& ChildrenArray = Children;
-			TManagedArray<FTransform>&  LocalTransform = Transform;
+			TManagedArray<FTransform3f>&  LocalTransform = Transform;
 			for (int32 sdx = 0; sdx < SortedDeletionList.Num(); sdx++)
 			{
 				TArray<FTransform> GlobalTransform;
@@ -217,7 +220,7 @@ void FTransformCollection::RemoveElements(const FName & Group, const TArray<int3
 						ParentTransform = GlobalTransform[ParentID].Inverse();
 					}
 
-					LocalTransform[ChildID] = GlobalTransform[ChildID] * ParentTransform;
+					LocalTransform[ChildID] = FTransform3f(GlobalTransform[ChildID] * ParentTransform);
 				}
 
 				if (0 <= ParentID)
diff --git a/Engine/Source/Runtime/Experimental/Chaos/Private/GeometryCollectionProxyData.cpp b/Engine/Source/Runtime/Experimental/Chaos/Private/GeometryCollectionProxyData.cpp
--- a/Engine/Source/Runtime/Experimental/Chaos/Private/GeometryCollectionProxyData.cpp
+++ b/Engine/Source/Runtime/Experimental/Chaos/Private/GeometryCollectionProxyData.cpp
@@ -40,11 +40,11 @@ void FTransformDynamicCollection::InitializeTransforms()
 	}
 }
 
-FTransform3f FTransformDynamicCollection::GetTransform(int32 Index) const
+const FTransform3f& FTransformDynamicCollection::GetTransform(int32 Index) const
 {
 	if (bTransformHasChanged == false)
 	{
-		return FTransform3f(RestCollection->Transform[Index]);
+		return RestCollection->Transform[Index];
 	}
 	return Transform[Index];
 }
@@ -61,6 +61,15 @@ int32 FTransformDynamicCollection::GetNumTransforms() const
 	return RestCollection->Transform.Num();
 }
 
+void FTransformDynamicCollection::ResetInitialTransforms()
+{
+	if (bTransformHasChanged)
+	{
+		RemoveAttribute(FTransformCollection::TransformAttribute, FTransformCollection::TransformGroup);
+		bTransformHasChanged = false;
+	}
+}
+
 const TManagedArray<bool>& FTransformDynamicCollection::GetHasParent() const
 {
 	return HasParent;
diff --git a/Engine/Source/Runtime/Experimental/Chaos/Private/PhysicsProxy/GeometryCollectionPhysicsProxy.cpp b/Engine/Source/Runtime/Experimental/Chaos/Private/PhysicsProxy/GeometryCollectionPhysicsProxy.cpp
--- a/Engine/Source/Runtime/Experimental/Chaos/Private/PhysicsProxy/GeometryCollectionPhysicsProxy.cpp
+++ b/Engine/Source/Runtime/Experimental/Chaos/Private/PhysicsProxy/GeometryCollectionPhysicsProxy.cpp
@@ -551,7 +551,6 @@ void FGeometryCollectionPhysicsProxy::Initialize(Chaos::FPBDRigidsEvolutionBase
 
 	// Initialise GT/External particles
 	const int32 NumTransforms = DynamicCollection.GetNumTransforms();
-
 	// make sure we copy the anchored information over to the physics thread collection
 	const Chaos::Facades::FCollectionAnchoringFacade DynamicCollectionAnchoringFacade(DynamicCollection);
 	Chaos::Facades::FCollectionAnchoringFacade PhysicsThreadCollectionAnchoringFacade(PhysicsThreadCollection);
@@ -4519,7 +4518,7 @@ void FGeometryCollectionPhysicsProxy::InitializeSharedCollisionStructures(
 	
 	TArray<FTransform> CollectionSpaceTransforms;
 	{ // tmp scope
-		const TManagedArray<FTransform>& HierarchyTransform = RestCollection.Transform;
+		const TManagedArray<FTransform3f>& HierarchyTransform = RestCollection.Transform;
 		GeometryCollectionAlgo::GlobalMatrices(HierarchyTransform, Parent, CollectionSpaceTransforms);
 	} // tmp scope
 
diff --git a/Engine/Source/Runtime/Experimental/Chaos/Public/GeometryCollection/GeometryCollectionAlgo.h b/Engine/Source/Runtime/Experimental/Chaos/Public/GeometryCollection/GeometryCollectionAlgo.h
--- a/Engine/Source/Runtime/Experimental/Chaos/Public/GeometryCollection/GeometryCollectionAlgo.h
+++ b/Engine/Source/Runtime/Experimental/Chaos/Public/GeometryCollection/GeometryCollectionAlgo.h
@@ -134,12 +134,15 @@ namespace GeometryCollectionAlgo
 	*  Global Matrices of the specified index.
 	*/
 	FTransform CHAOS_API GlobalMatrix(const TManagedArray<FTransform>& RelativeTransforms, const TManagedArray<int32>& Parents, int32 Index);
+	FTransform CHAOS_API GlobalMatrix(const TManagedArray<FTransform3f>& RelativeTransforms, const TManagedArray<int32>& Parents, int32 Index);
+	FTransform3f CHAOS_API GlobalMatrix3f(const TManagedArray<FTransform3f>& RelativeTransforms, const TManagedArray<int32>& Parents, int32 Index);
 
 
 	/*
 	*  Global Matrices of the collection based on list of indices
 	*/
 	void CHAOS_API GlobalMatrices(const TManagedArray<FTransform>& RelativeTransforms, const TManagedArray<int32>& Parents, const TArray<int32>& Indices, TArray<FTransform>& Transforms);
+	void CHAOS_API GlobalMatrices(const TManagedArray<FTransform3f>& RelativeTransforms, const TManagedArray<int32>& Parents, const TArray<int32>& Indices, TArray<FTransform3f>& Transforms);
 
 	/*
 	 *  Recursively traverse from a root node down
@@ -152,8 +155,8 @@ namespace GeometryCollectionAlgo
 	template<typename MatrixType>
 	void CHAOS_API GlobalMatrices(const TManagedArray<FTransform>& RelativeTransforms, const TManagedArray<int32>& Parents, const TManagedArray<FTransform>& UniformScale, TArray<MatrixType>& Transforms);
 
-	template<typename MatrixType>
-	void CHAOS_API GlobalMatrices(const TManagedArray<FTransform>& RelativeTransforms, const TManagedArray<int32>& Parents, TArray<MatrixType>& Transforms);
+	template<typename MatrixType, typename TransformType>
+	void CHAOS_API GlobalMatrices(const TManagedArray<TransformType>& RelativeTransforms, const TManagedArray<int32>& Parents, TArray<MatrixType>& Transforms);
 
 	/*
 	*  Gets pairs of elements whose bounding boxes overlap.
diff --git a/Engine/Source/Runtime/Experimental/Chaos/Public/GeometryCollection/GeometryCollectionUtility.h b/Engine/Source/Runtime/Experimental/Chaos/Public/GeometryCollection/GeometryCollectionUtility.h
--- a/Engine/Source/Runtime/Experimental/Chaos/Public/GeometryCollection/GeometryCollectionUtility.h
+++ b/Engine/Source/Runtime/Experimental/Chaos/Public/GeometryCollection/GeometryCollectionUtility.h
@@ -46,12 +46,12 @@ namespace GeometryCollection
 		TManagedArray<int32>& MaterialIndex = RestCollection->MaterialIndex;
 		TManagedArray<int32>& MaterialID = RestCollection->MaterialID;
 		TManagedArray<bool>& Internal = RestCollection->Internal;
-		TManagedArray<FTransform>& Transform = RestCollection->Transform;
+		TManagedArray<FTransform3f>& Transform = RestCollection->Transform;
 		TManagedArray<int32>& SimulationType = RestCollection->SimulationType;
 		TManagedArray<int32>& BoneMap = RestCollection->BoneMap;
 
 		// Set particle info
-		Transform[0] = Xf;
+		Transform[0] = FTransform3f(Xf);
 		Transform[0].NormalizeRotation();
 		SimulationType[0] = FGeometryCollection::ESimulationTypes::FST_Rigid;
 
diff --git a/Engine/Source/Runtime/Experimental/Chaos/Public/GeometryCollection/ManagedArray.h b/Engine/Source/Runtime/Experimental/Chaos/Public/GeometryCollection/ManagedArray.h
--- a/Engine/Source/Runtime/Experimental/Chaos/Public/GeometryCollection/ManagedArray.h
+++ b/Engine/Source/Runtime/Experimental/Chaos/Public/GeometryCollection/ManagedArray.h
@@ -182,7 +182,7 @@ protected:
 	/**
 	* Convert from a predefined Array, the managed array itself should have defined its conversion procedure
 	*/
-	virtual void Convert(const FManagedArrayBase&) { check(false); /* This type has no conversion process defined*/ };
+	virtual void Convert(const FManagedArrayBase&) { ensureMsgf(false, TEXT("Type change not supported")); /* This type has no conversion process defined*/ };
 
 	/**
 	* Copy a range of values from the ConstArray into this
diff --git a/Engine/Source/Runtime/Experimental/Chaos/Public/GeometryCollection/ManagedArrayCollection.h b/Engine/Source/Runtime/Experimental/Chaos/Public/GeometryCollection/ManagedArrayCollection.h
--- a/Engine/Source/Runtime/Experimental/Chaos/Public/GeometryCollection/ManagedArrayCollection.h
+++ b/Engine/Source/Runtime/Experimental/Chaos/Public/GeometryCollection/ManagedArrayCollection.h
@@ -404,7 +404,7 @@ public:
 
 
 	/**
-	* Remove the attribute from the collection.
+	* Remove the attribute from the collection, and clear the memory.
 	* @param Name - The name of the attribute
 	* @param Group - The group that manages the attribute
 	*/
diff --git a/Engine/Source/Runtime/Experimental/Chaos/Public/GeometryCollection/TransformCollection.h b/Engine/Source/Runtime/Experimental/Chaos/Public/GeometryCollection/TransformCollection.h
--- a/Engine/Source/Runtime/Experimental/Chaos/Public/GeometryCollection/TransformCollection.h
+++ b/Engine/Source/Runtime/Experimental/Chaos/Public/GeometryCollection/TransformCollection.h
@@ -33,7 +33,7 @@ public:
 	*   TransformGroup ("Transform")
 	*		Default Attributes :
 	*
-	*          FTransformArray Transform =  GetAttribute<FTransform>("Transform", TransformGroup)
+	*          FTransform3fArray Transform =  GetAttribute<FTransform3f>("Transform", TransformGroup)
 	*		   FInt32Array Level = GetAttribute<int32>("Level", TransformGroup) FIX
 	*		   FInt32Array Parent = GetAttribute<int32>("Parent", TransformGroup) FIX
 	*		   FInt32Array Children = GetAttribute<TSet<int32>>("Children", TransformGroup) FIX
@@ -109,7 +109,7 @@ public:
 	CHAOS_API virtual void Reset() override;
 
 	// Transform Group
-	TManagedArray<FTransform>   Transform;
+	TManagedArray<FTransform3f> Transform;
 	TManagedArray<FString>      BoneName;
 	TManagedArray<FLinearColor> BoneColor;
 	TManagedArray<int32>        Parent;
diff --git a/Engine/Source/Runtime/Experimental/Chaos/Public/GeometryCollectionProxyData.h b/Engine/Source/Runtime/Experimental/Chaos/Public/GeometryCollectionProxyData.h
--- a/Engine/Source/Runtime/Experimental/Chaos/Public/GeometryCollectionProxyData.h
+++ b/Engine/Source/Runtime/Experimental/Chaos/Public/GeometryCollectionProxyData.h
@@ -28,9 +28,10 @@ public:
 	FTransformDynamicCollection(FTransformDynamicCollection&&) = delete;
 	FTransformDynamicCollection& operator=(FTransformDynamicCollection&&) = delete;
 
-	CHAOS_API FTransform3f GetTransform(int32 Index) const;
+	CHAOS_API const FTransform3f& GetTransform(int32 Index) const;
 	CHAOS_API void SetTransform(int32 Index, const FTransform3f& Transform);
 	CHAOS_API int32 GetNumTransforms() const;
+	CHAOS_API void ResetInitialTransforms();
 
 	CHAOS_API const TManagedArray<bool>& GetHasParent() const;
 	CHAOS_API bool GetHasParent(int32 Index) const;
diff --git a/Engine/Source/Runtime/Experimental/GeometryCollectionEngine/Private/GeometryCollection/GeometryCollectionComponent.cpp b/Engine/Source/Runtime/Experimental/GeometryCollectionEngine/Private/GeometryCollection/GeometryCollectionComponent.cpp
--- a/Engine/Source/Runtime/Experimental/GeometryCollectionEngine/Private/GeometryCollection/GeometryCollectionComponent.cpp
+++ b/Engine/Source/Runtime/Experimental/GeometryCollectionEngine/Private/GeometryCollection/GeometryCollectionComponent.cpp
@@ -502,12 +502,6 @@ private:
 // Define the methods
 COPY_ON_WRITE_ATTRIBUTES
 
-
-const TManagedArray<FTransform>& UGeometryCollectionComponent::GetTransformArrayRest() const
-{
-	return RestCollection->GetGeometryCollection()->Transform;
-}
-
 TManagedArray<int32>& UGeometryCollectionComponent::GetParentArrayCopyOnWrite()
 {
 	if (!IndirectParentArray)
@@ -1391,14 +1385,7 @@ bool UGeometryCollectionComponent::DoCustomNavigableGeometryExport(FNavigableGeo
 
 	// Get all the geometry transforms in component space (they are stored natively in parent-bone space)
 	TArray<FTransform> GeomToComponent;
-	if (DynamicCollection)
-	{
-		GeometryCollectionAlgo::Private::GlobalMatrices(*DynamicCollection, TransformIndexBuffer, GeomToComponent);
-	}
-	else
-	{
-		GeometryCollectionAlgo::GlobalMatrices(TManagedArray<FTransform>(GetCurrentRestTransforms()), GetParentArrayRest(), TransformIndexBuffer, GeomToComponent);
-	}
+	ComputeCurrentGlobalsMatrices(GeomToComponent);
 
 	OutVertexBuffer.AddUninitialized(VertexCount);
 
@@ -1627,7 +1614,11 @@ void UGeometryCollectionComponent::ResetRestTransforms()
 
 		if (RestCollection && RestCollection->GetGeometryCollection())
 		{
-			RestTransformsChanged(RestCollection->GetGeometryCollection()->Transform.GetConstArray());
+			if (DynamicCollection)
+			{
+				DynamicCollection->ResetInitialTransforms();
+			}
+			RestTransformsChanged();
 		}
 	}
 }
@@ -1635,16 +1626,12 @@ void UGeometryCollectionComponent::ResetRestTransforms()
 void UGeometryCollectionComponent::SetRestState(TArray<FTransform>&& InRestTransforms)
 {
 	RestTransforms = InRestTransforms;
-	RestTransformsChanged(RestTransforms);
+	SetInitialTransforms(RestTransforms);
+	RestTransformsChanged();
 }
 
-void UGeometryCollectionComponent::RestTransformsChanged(const TArray<FTransform>& NewRestTransform)
+void UGeometryCollectionComponent::RestTransformsChanged()
 {
-	if (DynamicCollection)
-	{
-		SetInitialTransforms(NewRestTransform);
-	}
-
 	if (SceneProxy)
 	{
 		FGeometryCollectionDynamicData* DynamicData = GDynamicDataPool.Allocate();
@@ -2960,7 +2947,6 @@ void UGeometryCollectionComponent::SetInitialClusterBreaks(const TArray<int32>&
 	if (DynamicCollection)
 	{
 		const int32 NumTransforms = DynamicCollection->GetNumTransforms();
-
 		for (int32 ReleaseIndex : ReleaseIndices)
 		{
 			if (ReleaseIndex < NumTransforms)
@@ -3033,9 +3019,16 @@ void UGeometryCollectionComponent::GetHiddenTransforms(TArray<bool>& OutHiddenTr
 void UGeometryCollectionComponent::GetRestTransforms(TArray<FMatrix44f>& OutRestTransforms) const
 {
 	TArray<FMatrix> RestMatrices;
-	const TArray<FTransform>& RestCollectionTransforms = RestCollection->GetGeometryCollection()->Transform.GetConstArray();
-	const TArray<FTransform>& RestTransformToUse = (RestTransforms.Num() == RestCollectionTransforms.Num()) ? RestTransforms : RestCollectionTransforms;
-	GeometryCollectionAlgo::GlobalMatrices(TManagedArray<FTransform>(RestTransformToUse), RestCollection->GetGeometryCollection()->Parent, RestMatrices);
+	
+	const int32 RestCollectionNum = RestCollection->GetGeometryCollection()->Transform.Num();
+	if (RestTransforms.Num() == RestCollectionNum)
+	{
+		GeometryCollectionAlgo::GlobalMatrices(TManagedArray<FTransform>(RestTransforms), RestCollection->GetGeometryCollection()->Parent, RestMatrices);
+	}
+	else
+	{
+		GeometryCollectionAlgo::GlobalMatrices(RestCollection->GetGeometryCollection()->Transform, RestCollection->GetGeometryCollection()->Parent, RestMatrices);
+	}
 #if WITH_EDITOR
 	UpdateGlobalMatricesWithExplodedVectors(RestMatrices, *(RestCollection->GetGeometryCollection()));
 #endif
@@ -3564,7 +3557,7 @@ void UGeometryCollectionComponent::ResetDynamicCollection()
 		RestTransforms.Reset();
 	}
 
-	// if Reset transform have been overriden uses them to initialize the dynamic collection transforms
+	// if Reset transform have been overridden uses them to initialize the dynamic collection transforms
 	if (RestTransforms.Num() > 0)
 	{
 		SetInitialTransforms(RestTransforms);
@@ -3929,8 +3922,8 @@ void UGeometryCollectionComponent::MoveComponentToRootTransform()
 			const bool bHasDynamicOPrClusterUnionParent = DynamicStateFacade.HasDynamicInternalClusterParent(RootIndex) || DynamicStateFacade.HasClusterUnionParent(RootIndex);
 			if (bIsRootActive || bHasDynamicOPrClusterUnionParent)
 			{
-				const FTransform& OriginalComponentSpaceRootTransformOffset = AssetCollection->Transform[RootIndex];
-				DynamicCollection->SetTransform(RootIndex, FTransform3f(OriginalComponentSpaceRootTransformOffset));
+				const FTransform3f& OriginalComponentSpaceRootTransformOffset = AssetCollection->Transform[RootIndex];
+				DynamicCollection->SetTransform(RootIndex, OriginalComponentSpaceRootTransformOffset);
 				const Chaos::FPBDRigidParticle* RootParticle = PhysicsProxy->GetExternalParticles()[RootIndex].Get();
 				const FTransform ParticleWorldPosition(RootParticle->R(), RootParticle->X());
 				const FTransform MassToLocal = MassToLocalAttribute[RootIndex];
@@ -5407,15 +5400,39 @@ static void CheckForNaNs(const TArray<FTransform>& Transforms)
 }
 #endif
 
-const TArray<FTransform>& UGeometryCollectionComponent::GetCurrentRestTransforms() const
+FTransform UGeometryCollectionComponent::GetCurrentTransform(int32 Index) const
 {
+	if (DynamicCollection)
+	{
+		return FTransform(DynamicCollection->GetTransform(Index));
+	}
+
 	const bool bRestTransformsOverriden = (RestTransforms.Num() > 0);
 	if (bRestTransformsOverriden)
 	{
-		return RestTransforms;
+		return RestTransforms[Index];
 	}
+	return FTransform(RestCollection->GetGeometryCollection()->Transform[Index]);
+}
 
-	return GetTransformArrayRest().GetConstArray();
+void UGeometryCollectionComponent::ComputeCurrentGlobalsMatrices(TArray<FTransform>& OutTransforms) const
+{
+	if (DynamicCollection)
+	{
+		GeometryCollectionAlgo::Private::GlobalMatrices(*DynamicCollection, OutTransforms);
+	}
+	else
+	{
+		const bool bRestTransformsOverriden = (RestTransforms.Num() > 0);
+		if (bRestTransformsOverriden)
+		{
+			GeometryCollectionAlgo::GlobalMatrices(TManagedArray<FTransform>(RestTransforms), GetParentArrayRest(), OutTransforms);
+		}
+		else
+		{
+			GeometryCollectionAlgo::GlobalMatrices(RestCollection->GetGeometryCollection()->Transform, GetParentArrayRest(), OutTransforms);
+		}
+	}
 }
 
 const FTransform& UGeometryCollectionComponent::FComponentSpaceTransforms::RequestRootTransform() const
@@ -5430,15 +5447,7 @@ const FTransform& UGeometryCollectionComponent::FComponentSpaceTransforms::Reque
 	{
 		if (bIsRootDirty)
 		{
-			if (Component->DynamicCollection)
-			{
-				Transforms[RootIndex] = FTransform(Component->DynamicCollection->GetTransform(RootIndex));
-			}
-			else
-			{
-				Transforms[RootIndex] = Component->GetCurrentRestTransforms()[RootIndex];
-			}
-			
+			Transforms[RootIndex] = Component->GetCurrentTransform(RootIndex);
 			bIsRootDirty = false;
 		}
 
@@ -5470,7 +5479,7 @@ const TArray<FTransform>& UGeometryCollectionComponent::FComponentSpaceTransform
 	}
 	else
 	{
-		CurrentTransformNum = Component->GetCurrentRestTransforms().Num();
+		CurrentTransformNum = Component->RestCollection->GetGeometryCollection()->Transform.Num();
 	}
 
 	bool bFastPath = false;
@@ -5490,15 +5499,7 @@ const TArray<FTransform>& UGeometryCollectionComponent::FComponentSpaceTransform
 			const int32 ParentTransformIndex = Component->GetParent(TransformIndex);
 
 
-			FTransform CurrentTransform;
-			if (Component->DynamicCollection)
-			{
-				CurrentTransform = FTransform(Component->DynamicCollection->GetTransform(TransformIndex));
-			}
-			else
-			{
-				CurrentTransform = Component->GetCurrentRestTransforms()[TransformIndex];
-			}
+			FTransform CurrentTransform = Component->GetCurrentTransform(TransformIndex);
 			if (ParentTransformIndex == INDEX_NONE)
 			{
 				
@@ -5513,15 +5514,7 @@ const TArray<FTransform>& UGeometryCollectionComponent::FComponentSpaceTransform
 	}
 	else
 	{
-		if (Component->DynamicCollection)
-		{
-			GeometryCollectionAlgo::Private::GlobalMatrices(*Component->DynamicCollection, Transforms);
-		}
-		else
-		{
-			GeometryCollectionAlgo::GlobalMatrices(TManagedArray<FTransform>(Component->GetCurrentRestTransforms()), Component->GetParentArrayRest(), Transforms);
-		}
-		
+		Component->ComputeCurrentGlobalsMatrices(Transforms);
 	}
 
 #if WITH_EDITOR
@@ -6493,10 +6486,11 @@ void UGeometryCollectionComponent::Serialize(FArchive& Ar)
 			}
 			else
 			{
-				const TArray<FTransform>& RestCollectionTransforms = RestCollection->GetGeometryCollection()->Transform.GetConstArray();
+				const TArray<FTransform3f>& RestCollectionTransforms = RestCollection->GetGeometryCollection()->Transform.GetConstArray();
 				for (int32 Index = 0; Index < RestTransforms.Num(); Index++)
 				{
-					if (!RestTransforms[Index].Equals(RestCollectionTransforms[Index]))
+					// Do a float comparison otherwise, changing from float to double coulld cause transforms to be different
+					if (!FTransform3f(RestTransforms[Index]).Equals(RestCollectionTransforms[Index]))
 					{
 						bNeedToRestRestTransform = true;
 						break;
diff --git a/Engine/Source/Runtime/Experimental/GeometryCollectionEngine/Private/GeometryCollection/GeometryCollectionEngineConversion.cpp b/Engine/Source/Runtime/Experimental/GeometryCollectionEngine/Private/GeometryCollection/GeometryCollectionEngineConversion.cpp
--- a/Engine/Source/Runtime/Experimental/GeometryCollectionEngine/Private/GeometryCollection/GeometryCollectionEngineConversion.cpp
+++ b/Engine/Source/Runtime/Experimental/GeometryCollectionEngine/Private/GeometryCollection/GeometryCollectionEngineConversion.cpp
@@ -282,11 +282,11 @@ void FGeometryCollectionEngineConversion::AppendMeshDescription(
 	AppendMeshDescriptionTask.EnterProgressFrame(1);
 
 	// Geometry transform
-	TManagedArray<FTransform>& Transform = GeometryCollection->Transform;
+	TManagedArray<FTransform3f>& Transform = GeometryCollection->Transform;
 
 	int32 TransformIndex1 = GeometryCollection->AddElements(1, FGeometryCollection::TransformGroup);
-	Transform[TransformIndex1] = StaticMeshTransform;
-	Transform[TransformIndex1].SetScale3D(FVector(1.f, 1.f, 1.f));
+	Transform[TransformIndex1] = FTransform3f(StaticMeshTransform);
+	Transform[TransformIndex1].SetScale3D(FVector3f(1.f, 1.f, 1.f));
 
 	// collisions
 	if (BodySetup)
@@ -770,8 +770,8 @@ bool FGeometryCollectionEngineConversion::AppendGeometryCollection(const FGeomet
 	const int32 GeometryCount = SourceGeometryCollectionPtr->TransformIndex.Num();
 	const int32 SectionCount = SourceGeometryCollectionPtr->Sections.Num();
 
-	FVector Scale = GeometryCollectionTransform.GetScale3D();
-	FTransform AppliedTransform = GeometryCollectionTransform;
+	FVector3f Scale = FVector3f(GeometryCollectionTransform.GetScale3D());
+	FTransform3f AppliedTransform = FTransform3f(GeometryCollectionTransform);
 	AppliedTransform.RemoveScaling();
 
 	const int32 VertexStart = TargetGeometryCollection->AddElements(VertexCount, FGeometryCollection::VerticesGroup);
@@ -850,7 +850,7 @@ bool FGeometryCollectionEngineConversion::AppendGeometryCollection(const FGeomet
 	}
 
 	// source transform information
-	const TManagedArray<FTransform>& SourceTransform = SourceGeometryCollectionPtr->Transform;
+	const TManagedArray<FTransform3f>& SourceTransform = SourceGeometryCollectionPtr->Transform;
 	const TManagedArray<FString>& SourceBoneName = SourceGeometryCollectionPtr->BoneName;
 	const TManagedArray<FLinearColor>& SourceBoneColor = SourceGeometryCollectionPtr->BoneColor;
 	const TManagedArray<int32>& SourceParent = SourceGeometryCollectionPtr->Parent;
@@ -862,7 +862,7 @@ bool FGeometryCollectionEngineConversion::AppendGeometryCollection(const FGeomet
 	const TManagedArray<Chaos::FImplicitObjectPtr>* SourceExternalCollisions = SourceGeometryCollectionPtr->FindAttribute<Chaos::FImplicitObjectPtr>(FGeometryCollection::ExternalCollisionsAttribute, FGeometryCollection::TransformGroup);
 
 	// target transform information
-	TManagedArray<FTransform>& TargetTransform = TargetGeometryCollection->Transform;
+	TManagedArray<FTransform3f>& TargetTransform = TargetGeometryCollection->Transform;
 	TManagedArray<FString>& TargetBoneName = TargetGeometryCollection->BoneName;
 	TManagedArray<FLinearColor>& TargetBoneColor = TargetGeometryCollection->BoneColor;
 	TManagedArray<int32>& TargetParent = TargetGeometryCollection->Parent;
@@ -885,7 +885,7 @@ bool FGeometryCollectionEngineConversion::AppendGeometryCollection(const FGeomet
 		}
 		else
 		{
-			FTransform ScaledTranslation = SourceTransform[TransformIndex];
+			FTransform3f ScaledTranslation = SourceTransform[TransformIndex];
 			ScaledTranslation.ScaleTranslation(Scale);
 			TargetTransform[TransformOffset] = ScaledTranslation;
 		}
diff --git a/Engine/Source/Runtime/Experimental/GeometryCollectionEngine/Private/GeometryCollection/GeometryCollectionRenderData.cpp b/Engine/Source/Runtime/Experimental/GeometryCollectionEngine/Private/GeometryCollection/GeometryCollectionRenderData.cpp
--- a/Engine/Source/Runtime/Experimental/GeometryCollectionEngine/Private/GeometryCollection/GeometryCollectionRenderData.cpp
+++ b/Engine/Source/Runtime/Experimental/GeometryCollectionEngine/Private/GeometryCollection/GeometryCollectionRenderData.cpp
@@ -569,7 +569,7 @@ FBoxSphereBounds BuildPreSkinnedBounds(FGeometryCollection& InCollection)
 	TRACE_CPUPROFILER_EVENT_SCOPE(FGeometryCollectionRenderData::BuildBounds);
 
 	const TManagedArray<FBox>& BoundingBoxes = InCollection.BoundingBox;
-	const TManagedArray<FTransform>& Transform = InCollection.Transform;
+	const TManagedArray<FTransform3f>& Transform = InCollection.Transform;
 	const TManagedArray<int32>& Parent = InCollection.Parent;
 	const TManagedArray<int32>& TransformIndices = InCollection.TransformIndex;
 	const int32 NumBoxes = BoundingBoxes.Num();
diff --git a/Engine/Source/Runtime/Experimental/GeometryCollectionEngine/Public/GeometryCollection/GeometryCollectionComponent.h b/Engine/Source/Runtime/Experimental/GeometryCollectionEngine/Public/GeometryCollection/GeometryCollectionComponent.h
--- a/Engine/Source/Runtime/Experimental/GeometryCollectionEngine/Public/GeometryCollection/GeometryCollectionComponent.h
+++ b/Engine/Source/Runtime/Experimental/GeometryCollectionEngine/Public/GeometryCollection/GeometryCollectionComponent.h
@@ -896,8 +896,6 @@ public:
 	// Declare all the methods
 	COPY_ON_WRITE_ATTRIBUTES
 
-	GEOMETRYCOLLECTIONENGINE_API const TManagedArray<FTransform>& GetTransformArrayRest() const;
-
 	GEOMETRYCOLLECTIONENGINE_API TManagedArray<int32>& GetParentArrayCopyOnWrite();
 	GEOMETRYCOLLECTIONENGINE_API int32 GetParent(int32 Index) const;
 	GEOMETRYCOLLECTIONENGINE_API const TManagedArray<int32>& GetParentArrayRest() const;
@@ -1542,12 +1540,13 @@ private:
 	void ProcessRepDynamicDataOnPT();
 
 	// called when the rest transform are updated from SetRestState / ResetRestTransforms
-	// this upadtes the renderer as well as the dynamic collection initial transforms
-	void RestTransformsChanged(const TArray<FTransform>& NewRestTransform);
+	// this updates only the renderer, the dynamic collection should be initialized when calling this function
+	void RestTransformsChanged();
 
 	// return the most actual transforms
-	// this can be the rest collection ones, the overriden RestTransforms or the dynamic collection ones
-	const TArray<FTransform>& GetCurrentRestTransforms() const;
+	// this can be the rest collection ones, the overridden RestTransforms or the dynamic collection ones
+	FTransform GetCurrentTransform(int32 Index) const;
+	void ComputeCurrentGlobalsMatrices(TArray<FTransform>& OutTransforms) const;
 
 	bool bRenderStateDirty;
 	bool bEnableBoneSelection;
