( Microoptimisation ) Swap FName("") default params for NAME_None to reduce FName constructions.

diff --git a/Engine/Source/Runtime/Experimental/Chaos/Public/GeometryCollection/ManagedArrayAccessor.h b/Engine/Source/Runtime/Experimental/Chaos/Public/GeometryCollection/ManagedArrayAccessor.h
--- a/Engine/Source/Runtime/Experimental/Chaos/Public/GeometryCollection/ManagedArrayAccessor.h
+++ b/Engine/Source/Runtime/Experimental/Chaos/Public/GeometryCollection/ManagedArrayAccessor.h
@@ -19,7 +19,7 @@ template <typename T>
 struct TManagedArrayAccessor
 {
 public:
-	TManagedArrayAccessor(FManagedArrayCollection& InCollection, const FName& InAttributeName, const FName& InAttributeGroup,  const FName& InGroupDependency = FName(""))
+	TManagedArrayAccessor(FManagedArrayCollection& InCollection, const FName& InAttributeName, const FName& InAttributeGroup,  const FName& InGroupDependency = NAME_None)
 		: ConstCollection(InCollection)
 		, Collection(&InCollection)
 		, Name(InAttributeName)
@@ -29,7 +29,7 @@ public:
 		, ConstAttributeArray(AttributeArray)
 	{}
 
-	TManagedArrayAccessor(const FManagedArrayCollection& InCollection, const FName& InAttributeName, const FName& InAttributeGroup, const FName& InGroupDependency = FName(""))
+	TManagedArrayAccessor(const FManagedArrayCollection& InCollection, const FName& InAttributeName, const FName& InAttributeGroup, const FName& InGroupDependency = NAME_None)
 		: ConstCollection(InCollection)
 		, Collection(nullptr)
 		, Name(InAttributeName)
