Fix early release of referenced ISM when all instances are hidden.
Add cvar to disable add/remove hidden instances from ISM Pool.

diff --git a/Engine/Source/Runtime/Experimental/GeometryCollectionEngine/Private/GeometryCollection/GeometryCollectionISMPoolComponent.cpp b/Engine/Source/Runtime/Experimental/GeometryCollectionEngine/Private/GeometryCollection/GeometryCollectionISMPoolComponent.cpp
--- a/Engine/Source/Runtime/Experimental/GeometryCollectionEngine/Private/GeometryCollection/GeometryCollectionISMPoolComponent.cpp
+++ b/Engine/Source/Runtime/Experimental/GeometryCollectionEngine/Private/GeometryCollection/GeometryCollectionISMPoolComponent.cpp
@@ -24,6 +24,12 @@ FAutoConsoleVariableRef CVarISMPoolComponentFreeListTargetSize(
 	GComponentFreeListTargetSize,
 	TEXT("Target size for number of ISM components in the ISMPool."));
 
+static bool GAutoRemoveInstances = false;
+FAutoConsoleVariableRef CVarISMPoolAutoRemoveInstances(
+	TEXT("r.ISMPool.AutoRemoveInstances"),
+	GAutoRemoveInstances,
+	TEXT("Remove instances with zero scale from ISM components in the ISMPool."));
+
 FGeometryCollectionMeshGroup::FMeshId FGeometryCollectionMeshGroup::AddMesh(const FGeometryCollectionStaticMeshInstance& MeshInstance, int32 InstanceCount, const FGeometryCollectionMeshInfo& ISMInstanceInfo)
 {
 	const FMeshId MeshInfoIndex = MeshInfos.Emplace(ISMInstanceInfo);
@@ -250,18 +256,24 @@ bool FGeometryCollectionISMPool::BatchUpdateInstancesTransforms(FGeometryCollect
 		FPrimitiveInstanceId InstanceId = ISM.InstanceIds[InstanceGroup.Start + InstanceIndex];
 		FTransform const& Transform = NewInstancesTransforms[InstanceIndex];
 
-		if (Transform.GetScale3D().IsZero() && InstanceId.IsValid())
-		{
-			// Zero scale is used to indicate that we should remove the instance from the ISM.
-			ISM.ISMComponent->RemoveInstanceById(InstanceId);
-			ISM.InstanceIds[InstanceGroup.Start + InstanceIndex] = FPrimitiveInstanceId();
-		}
-		else if (!Transform.GetScale3D().IsZero() && !InstanceId.IsValid())
+		if (GAutoRemoveInstances)
 		{
-			// Re-add the instance to the ISM if the scale becomes non-zero.
-			ISM.InstanceIds[InstanceGroup.Start + InstanceIndex] = ISM.ISMComponent->AddInstanceById(Transform, bWorldSpace);
+			if (Transform.GetScale3D().IsZero() && InstanceId.IsValid())
+			{
+				// Zero scale is used to indicate that we should remove the instance from the ISM.
+				ISM.ISMComponent->RemoveInstanceById(InstanceId);
+				ISM.InstanceIds[InstanceGroup.Start + InstanceIndex] = FPrimitiveInstanceId();
+				continue;
+			}
+			else if (!Transform.GetScale3D().IsZero() && !InstanceId.IsValid())
+			{
+				// Re-add the instance to the ISM if the scale becomes non-zero.
+				ISM.InstanceIds[InstanceGroup.Start + InstanceIndex] = ISM.ISMComponent->AddInstanceById(Transform, bWorldSpace);
+				continue;
+			}
 		}
-		else if (InstanceId.IsValid())
+
+		if (InstanceId.IsValid())
 		{
 			ISM.ISMComponent->UpdateInstanceTransformById(InstanceId, Transform, bWorldSpace, bTeleport);
 		}
@@ -327,9 +339,10 @@ void FGeometryCollectionISMPool::RemoveISM(const FGeometryCollectionMeshInfo& Me
 			ISM.InstanceIds.Reset();
 		}
 
-		if (GUseComponentFreeList && ISM.ISMComponent->PerInstanceSMData.Num() == 0)
+		if (GUseComponentFreeList && ISM.InstanceGroups.IsEmpty())
 		{
 			// Remove component and push this ISM slot to the free list.
+			ensure(ISM.ISMComponent->PerInstanceSMData.Num() == 0);
 			MeshToISMIndex.Remove(ISM.MeshInstance);
 
 			const bool bIsHISM = (ISM.MeshInstance.Desc.Flags & FISMComponentDescription::UseHISM) != 0;
