Comment update from review feedback.

diff --git a/Engine/Source/Runtime/Experimental/GeometryCollectionEngine/Public/GeometryCollection/GeometryCollectionObject.h b/Engine/Source/Runtime/Experimental/GeometryCollectionEngine/Public/GeometryCollection/GeometryCollectionObject.h
--- a/Engine/Source/Runtime/Experimental/GeometryCollectionEngine/Public/GeometryCollection/GeometryCollectionObject.h
+++ b/Engine/Source/Runtime/Experimental/GeometryCollectionEngine/Public/GeometryCollection/GeometryCollectionObject.h
@@ -876,7 +876,7 @@ private:
 	UPROPERTY(VisibleAnywhere, Transient, Category = "Clustering")
 	TArray<int32> BreadthFirstTransformIndices;
 
-	// cache transform remap indices
+	// cache transform remapping for instanced meshes indices
 	UPROPERTY(VisibleAnywhere, Transient, Category = "Clustering")
 	TArray<int32> AutoInstanceTransformRemapIndices;
 
