Fix Modular Vehicle compilation error
Fallout from CL 28720845

diff --git a/Engine/Source/Runtime/Experimental/ChaosVehicles/ChaosVehiclesCore/Private/SimModule/DeferredForcesModular.cpp b/Engine/Source/Runtime/Experimental/ChaosVehicles/ChaosVehiclesCore/Private/SimModule/DeferredForcesModular.cpp
--- a/Engine/Source/Runtime/Experimental/ChaosVehicles/ChaosVehiclesCore/Private/SimModule/DeferredForcesModular.cpp
+++ b/Engine/Source/Runtime/Experimental/ChaosVehicles/ChaosVehiclesCore/Private/SimModule/DeferredForcesModular.cpp
@@ -50,6 +50,37 @@ Chaos::FPBDRigidParticleHandle* FDeferredForcesModular::GetParticle(TArray<Chaos
 	return nullptr;
 }
 
+// #TODO: passing all these parameters in is horrible, tidy this up now we know what needs to be done
+Chaos::FPBDRigidParticleHandle* FDeferredForcesModular::GetParticle(TArray<Chaos::FPBDRigidClusteredParticleHandle*>& Particles
+	, TArray<Chaos::FPBDRigidClusteredParticleHandle*>& ClusterParticles
+	, int TransformIndex
+	, const FVector& PositionalOffset
+	, const FTransform& Transform
+	, const TManagedArray<FTransform>& CollectionMassToLocal
+	, const TManagedArray<int32>& Parent
+	, FTransform& TransformOut)
+{
+
+	Chaos::FPBDRigidParticleHandle* SimParticle = Particles[TransformIndex];
+	Chaos::FPBDRigidClusteredParticleHandle* SimClusterParticle = ClusterParticles[TransformIndex];
+
+	if (SimParticle && !SimParticle->Disabled())
+	{
+		TransformOut = FTransform::Identity;
+		TransformOut.AddToTranslation(PositionalOffset);
+		return SimParticle;		// apply to broken off child
+	}
+	else if (SimClusterParticle && !SimClusterParticle->Disabled())
+	{
+		TransformOut.SetLocation(CollectionMassToLocal[Parent[TransformIndex]].InverseTransformPosition(Transform.GetLocation() + PositionalOffset));
+		TransformOut.SetRotation(Transform.GetRotation());
+
+		return SimClusterParticle; // apply to intact cluster
+	}
+
+	return nullptr;
+}
+
 Chaos::FPBDRigidParticleHandle* FDeferredForcesModular::GetParticle(const FTransform& OffsetTransform
 	, TArray<Chaos::FPBDRigidParticleHandle*>& Particles
 	, TArray<Chaos::FPBDRigidClusteredParticleHandle*>& ClusterParticles
@@ -92,9 +123,10 @@ Chaos::FPBDRigidParticleHandle* FDeferredForcesModular::GetParticle(const FTrans
 	return ClusterParticles[SingleChassisIndex];
 }
 
-void FDeferredForcesModular::Apply(TArray<Chaos::FPBDRigidClusteredParticleHandle*>& Particles
+template<typename TransformType>
+void FDeferredForcesModular::ApplyTemplate(TArray<Chaos::FPBDRigidClusteredParticleHandle*>& Particles
 		, TArray<Chaos::FPBDRigidClusteredParticleHandle*>& ClusterParticles
-		, const TManagedArray<FTransform>& Transforms
+		, const TManagedArray<TransformType>& Transforms
 		, const TManagedArray<FTransform>& CollectionMassToLocal
 		, const TManagedArray<int32>& Parent)
 {
@@ -103,7 +135,7 @@ void FDeferredForcesModular::Apply(TArray<Chaos::FPBDRigidClusteredParticleHandl
 		FTransform RelativeTransform;
 		for (const FApplyForceData& Data : ApplyForceDatas)
 		{
-			Chaos::FPBDRigidParticleHandle* RigidHandle = GetParticle(Particles, ClusterParticles, Data.TransformIndex, FVector::ZeroVector, Transforms, CollectionMassToLocal, Parent, RelativeTransform);
+			Chaos::FPBDRigidParticleHandle* RigidHandle = GetParticle(Particles, ClusterParticles, Data.TransformIndex, FVector::ZeroVector, FTransform(Transforms[Data.TransformIndex]), CollectionMassToLocal, Parent, RelativeTransform);
 			if (RigidHandle)
 			{
 				AddForce(RigidHandle, Data, RelativeTransform);
@@ -112,7 +144,7 @@ void FDeferredForcesModular::Apply(TArray<Chaos::FPBDRigidClusteredParticleHandl
 
 		for (const FApplyForceAtPositionData& Data : ApplyForceAtPositionDatas)
 		{
-			Chaos::FPBDRigidParticleHandle* RigidHandle = GetParticle(Particles, ClusterParticles, Data.TransformIndex, Data.Position, Transforms, CollectionMassToLocal, Parent, RelativeTransform);
+			Chaos::FPBDRigidParticleHandle* RigidHandle = GetParticle(Particles, ClusterParticles, Data.TransformIndex, Data.Position, FTransform(Transforms[Data.TransformIndex]), CollectionMassToLocal, Parent, RelativeTransform);
 			if (RigidHandle)
 			{
 				AddForceAtPosition(RigidHandle, Data, RelativeTransform);
@@ -121,7 +153,7 @@ void FDeferredForcesModular::Apply(TArray<Chaos::FPBDRigidClusteredParticleHandl
 
 		for (const FAddTorqueInRadiansData& Data : ApplyTorqueDatas)
 		{
-			Chaos::FPBDRigidParticleHandle* RigidHandle = GetParticle(Particles, ClusterParticles, Data.TransformIndex, FVector::ZeroVector, Transforms, CollectionMassToLocal, Parent, RelativeTransform);
+			Chaos::FPBDRigidParticleHandle* RigidHandle = GetParticle(Particles, ClusterParticles, Data.TransformIndex, FVector::ZeroVector, FTransform(Transforms[Data.TransformIndex]), CollectionMassToLocal, Parent, RelativeTransform);
 			if (RigidHandle)
 			{
 				AddTorque(RigidHandle, Data, RelativeTransform);
@@ -134,6 +166,23 @@ void FDeferredForcesModular::Apply(TArray<Chaos::FPBDRigidClusteredParticleHandl
 	ApplyTorqueDatas.Empty();
 }
 
+void FDeferredForcesModular::Apply(TArray<Chaos::FPBDRigidClusteredParticleHandle*>& Particles
+	, TArray<Chaos::FPBDRigidClusteredParticleHandle*>& ClusterParticles
+	, const TManagedArray<FTransform>& Transforms
+	, const TManagedArray<FTransform>& CollectionMassToLocal
+	, const TManagedArray<int32>& Parent)
+{
+	ApplyTemplate(Particles, ClusterParticles, Transforms, CollectionMassToLocal, Parent);
+}
+
+void FDeferredForcesModular::Apply(TArray<Chaos::FPBDRigidClusteredParticleHandle*>& Particles
+	, TArray<Chaos::FPBDRigidClusteredParticleHandle*>& ClusterParticles
+	, const TManagedArray<FTransform3f>& Transforms
+	, const TManagedArray<FTransform>& CollectionMassToLocal
+	, const TManagedArray<int32>& Parent)
+{
+	ApplyTemplate(Particles, ClusterParticles, Transforms, CollectionMassToLocal, Parent);
+}
 
 
 void FDeferredForcesModular::Apply(TArray<Chaos::FPBDRigidParticleHandle*>& Particles
diff --git a/Engine/Source/Runtime/Experimental/ChaosVehicles/ChaosVehiclesCore/Public/SimModule/DeferredForcesModular.h b/Engine/Source/Runtime/Experimental/ChaosVehicles/ChaosVehiclesCore/Public/SimModule/DeferredForcesModular.h
--- a/Engine/Source/Runtime/Experimental/ChaosVehicles/ChaosVehiclesCore/Public/SimModule/DeferredForcesModular.h
+++ b/Engine/Source/Runtime/Experimental/ChaosVehicles/ChaosVehiclesCore/Public/SimModule/DeferredForcesModular.h
@@ -156,6 +156,15 @@ public:
 			, const TManagedArray<int32>& Parent
 			, FTransform& TransformOut);
 
+	Chaos::FPBDRigidParticleHandle* GetParticle(TArray<Chaos::FPBDRigidClusteredParticleHandle*>& Particles
+		, TArray<Chaos::FPBDRigidClusteredParticleHandle*>& ClusterParticles
+		, int TransformIndex
+		, const FVector& PositionalOffset
+		, const FTransform& Transform
+		, const TManagedArray<FTransform>& CollectionMassToLocal
+		, const TManagedArray<int32>& Parent
+		, FTransform& TransformOut);
+
 	Chaos::FPBDRigidParticleHandle* GetParticle(const FTransform& OffsetTransform
 		, TArray<Chaos::FPBDRigidParticleHandle*>& Particles
 		, TArray<Chaos::FPBDRigidClusteredParticleHandle*>& ClusterParticles
@@ -168,12 +177,25 @@ public:
 		, const TManagedArray<FTransform>& Transforms
 		, const TManagedArray<FTransform>& CollectionMassToLocal
 		, const TManagedArray<int32>& Parent);
+
+	void Apply(TArray<Chaos::FPBDRigidClusteredParticleHandle*>& Particles
+		, TArray<Chaos::FPBDRigidClusteredParticleHandle*>& ClusterParticles
+		, const TManagedArray<FTransform3f>& Transforms
+		, const TManagedArray<FTransform>& CollectionMassToLocal
+		, const TManagedArray<int32>& Parent);
 	
 	void Apply(TArray<Chaos::FPBDRigidParticleHandle*>& Particles
 		, TArray<Chaos::FPBDRigidClusteredParticleHandle*>& ClusterParticles);
 
 private:
 
+	template<typename TransformType>
+	void ApplyTemplate(TArray<Chaos::FPBDRigidClusteredParticleHandle*>& Particles
+		, TArray<Chaos::FPBDRigidClusteredParticleHandle*>& ClusterParticles
+		, const TManagedArray<TransformType>& Transforms
+		, const TManagedArray<FTransform>& CollectionMassToLocal
+		, const TManagedArray<int32>& Parent);
+
 	void AddForceAtPosition(Chaos::FPBDRigidParticleHandle* RigidHandle, const FApplyForceAtPositionData& DataIn, const FTransform& OffsetTransform);
 	void AddTorque(Chaos::FPBDRigidParticleHandle* RigidHandle, const FAddTorqueInRadiansData& DataIn, const FTransform& OffsetTransform);
 	void AddForce(Chaos::FPBDRigidParticleHandle* RigidHandle, const FApplyForceData& DataIn, const FTransform& OffsetTransform);
