Chaos - fix mesh collision sometimes not setting the faceindex, leading to incorrect material properties in collision

diff --git a/Engine/Source/Runtime/Experimental/Chaos/Private/Chaos/Collision/MeshContactGenerator.cpp b/Engine/Source/Runtime/Experimental/Chaos/Private/Chaos/Collision/MeshContactGenerator.cpp
--- a/Engine/Source/Runtime/Experimental/Chaos/Private/Chaos/Collision/MeshContactGenerator.cpp
+++ b/Engine/Source/Runtime/Experimental/Chaos/Private/Chaos/Collision/MeshContactGenerator.cpp
@@ -112,6 +112,7 @@ namespace Chaos::Private
 						if (ContactNormalDotTriangleNormal > ExistingContactData.GetContactNormalDotTriangleNormal())
 						{
 							ExistingContact = ContactPoint;
+							ExistingContact.FaceIndex = Triangle.GetTriangleIndex();
 
 							ExistingContactData.SetVertexID(VertexID0);
 							ExistingContactData.SetContactNormalDotTriangleNormal(ContactNormalDotTriangleNormal);
