the check for OnRootMovedEvent in UGeometryCollectionComponent::OnPostPhysicsSync needs to store PreviousRootTransform as a value so the next call to RequestRootTransform doesn't make the prev value the same as the new value

diff --git a/Engine/Source/Runtime/Experimental/GeometryCollectionEngine/Private/GeometryCollection/GeometryCollectionComponent.cpp b/Engine/Source/Runtime/Experimental/GeometryCollectionEngine/Private/GeometryCollection/GeometryCollectionComponent.cpp
--- a/Engine/Source/Runtime/Experimental/GeometryCollectionEngine/Private/GeometryCollection/GeometryCollectionComponent.cpp
+++ b/Engine/Source/Runtime/Experimental/GeometryCollectionEngine/Private/GeometryCollection/GeometryCollectionComponent.cpp
@@ -3848,7 +3848,8 @@ void UGeometryCollectionComponent::OnPostPhysicsSync()
 	// dirty the transform if the collection is
 	if (DynamicCollection && DynamicCollection->IsDirty())
 	{
-		const FTransform& PreviousRootTransform = OnRootMovedEvent.IsBound() ? ComponentSpaceTransforms.RequestRootTransform() : FTransform::Identity;
+		// Can't be a const reference - we need to make a copy or else the next RequestRootTransform will change the value under our nose.
+		const FTransform PreviousRootTransform = OnRootMovedEvent.IsBound() ? ComponentSpaceTransforms.RequestRootTransform() : FTransform::Identity;
 		ComponentSpaceTransforms.MarkDirty();
 		if (OnRootMovedEvent.IsBound())
 		{
