Geometry collection : Update transform facade to use Ftransform3f for transform attribute since they have changed type
this fixes some fallout from CL 28720845

diff --git a/Engine/Source/Runtime/Experimental/Chaos/Private/GeometryCollection/Facades/CollectionTransformFacade.cpp b/Engine/Source/Runtime/Experimental/Chaos/Private/GeometryCollection/Facades/CollectionTransformFacade.cpp
--- a/Engine/Source/Runtime/Experimental/Chaos/Private/GeometryCollection/Facades/CollectionTransformFacade.cpp
+++ b/Engine/Source/Runtime/Experimental/Chaos/Private/GeometryCollection/Facades/CollectionTransformFacade.cpp
@@ -33,7 +33,7 @@ namespace GeometryCollection::Facades
 	{
 		TArray<FTransform> OutTransforms;
 
-		const TManagedArray<FTransform>& BoneTransforms = TransformAttribute.Get();
+		const TManagedArray<FTransform3f>& BoneTransforms = TransformAttribute.Get();
 		const TManagedArray<int32>& Parents = ParentAttribute.Get();
 
 		GeometryCollectionAlgo::GlobalMatrices(BoneTransforms, Parents, OutTransforms);
@@ -43,7 +43,7 @@ namespace GeometryCollection::Facades
 
 	FTransform FCollectionTransformFacade::ComputeCollectionSpaceTransform(int32 BoneIdx) const
 	{
-		const TManagedArray<FTransform>& BoneTransforms = TransformAttribute.Get();
+		const TManagedArray<FTransform3f>& BoneTransforms = TransformAttribute.Get();
 		const TManagedArray<int32>& Parents = ParentAttribute.Get();
 
 		return GeometryCollectionAlgo::GlobalMatrix(BoneTransforms, Parents, BoneIdx);
@@ -59,11 +59,11 @@ namespace GeometryCollection::Facades
 		// Update only root transforms
 		const TArray<int32>& RootIndices = GetRootIndices();
 
-		TManagedArray<FTransform>& Transforms = TransformAttribute.Modify();
+		TManagedArray<FTransform3f>& Transforms = TransformAttribute.Modify();
 
 		for (int32 Idx : RootIndices)
 		{
-			Transforms[Idx] = Transforms[Idx] * InTransform;
+			Transforms[Idx] = Transforms[Idx] * FTransform3f(InTransform);
 		}
 	}
 
@@ -210,9 +210,9 @@ namespace GeometryCollection::Facades
 
 	void FCollectionTransformFacade::SetBoneTransformToIdentity(int32 BoneIdx)
 	{
-		TManagedArray<FTransform>& BoneTransforms = TransformAttribute.Modify();
+		TManagedArray<FTransform3f>& BoneTransforms = TransformAttribute.Modify();
 
-		BoneTransforms[BoneIdx] = FTransform::Identity;
+		BoneTransforms[BoneIdx] = FTransform3f::Identity;
 	}
 
 
diff --git a/Engine/Source/Runtime/Experimental/Chaos/Public/GeometryCollection/Facades/CollectionTransformFacade.h b/Engine/Source/Runtime/Experimental/Chaos/Public/GeometryCollection/Facades/CollectionTransformFacade.h
--- a/Engine/Source/Runtime/Experimental/Chaos/Public/GeometryCollection/Facades/CollectionTransformFacade.h
+++ b/Engine/Source/Runtime/Experimental/Chaos/Public/GeometryCollection/Facades/CollectionTransformFacade.h
@@ -42,7 +42,7 @@ namespace GeometryCollection::Facades
 		/**
 		* Returns the child indicesfrom the collection. Null if not initialized.
 		*/
-		const TManagedArray< FTransform >* FindTransforms() const { return TransformAttribute.Find(); }
+		const TManagedArray<FTransform3f>* FindTransforms() const { return TransformAttribute.Find(); }
 
 		/**
 		* Returns array of transforms for transforming from bone space to collection space
@@ -92,6 +92,6 @@ namespace GeometryCollection::Facades
 	private:
 		TManagedArrayAccessor<int32>		ParentAttribute;
 		TManagedArrayAccessor<TSet<int32>>	ChildrenAttribute;
-		TManagedArrayAccessor<FTransform>	TransformAttribute;
+		TManagedArrayAccessor<FTransform3f>	TransformAttribute;
 	};
 }
