ISMPool: Per instance removal is a setting on the mesh group instead of a cvar.

diff --git a/Engine/Source/Runtime/Experimental/GeometryCollectionEngine/Private/GeometryCollection/GeometryCollectionComponent.cpp b/Engine/Source/Runtime/Experimental/GeometryCollectionEngine/Private/GeometryCollection/GeometryCollectionComponent.cpp
--- a/Engine/Source/Runtime/Experimental/GeometryCollectionEngine/Private/GeometryCollection/GeometryCollectionComponent.cpp
+++ b/Engine/Source/Runtime/Experimental/GeometryCollectionEngine/Private/GeometryCollection/GeometryCollectionComponent.cpp
@@ -5770,7 +5770,12 @@ void UGeometryCollectionComponent::RefreshCustomRenderer()
 
 					const bool bRenderRootProxy = bEnableRootProxyForCustomRenderer && !bIsBroken;
 
-					RendererInterface->UpdateState(*RestCollection, ComponentTransform, !bRenderRootProxy, !bHiddenInGame);
+					uint32 StateFlags = 0;
+					StateFlags |= bHiddenInGame ? 0 : IGeometryCollectionExternalRenderInterface::EState_Visible;
+					StateFlags |= bRenderRootProxy ? 0 : IGeometryCollectionExternalRenderInterface::EState_Broken;
+					StateFlags |= bEnableRootProxyForCustomRenderer ? 0 : IGeometryCollectionExternalRenderInterface::EState_ForcedBroken;
+
+					RendererInterface->UpdateState(*RestCollection, ComponentTransform, StateFlags);
 
 					if (bRenderRootProxy)
 					{
diff --git a/Engine/Source/Runtime/Experimental/GeometryCollectionEngine/Private/GeometryCollection/GeometryCollectionISMPoolComponent.cpp b/Engine/Source/Runtime/Experimental/GeometryCollectionEngine/Private/GeometryCollection/GeometryCollectionISMPoolComponent.cpp
--- a/Engine/Source/Runtime/Experimental/GeometryCollectionEngine/Private/GeometryCollection/GeometryCollectionISMPoolComponent.cpp
+++ b/Engine/Source/Runtime/Experimental/GeometryCollectionEngine/Private/GeometryCollection/GeometryCollectionISMPoolComponent.cpp
@@ -24,11 +24,6 @@ FAutoConsoleVariableRef CVarISMPoolComponentFreeListTargetSize(
 	GComponentFreeListTargetSize,
 	TEXT("Target size for number of ISM components in the ISMPool."));
 
-static bool GAutoRemoveInstances = false;
-FAutoConsoleVariableRef CVarISMPoolAutoRemoveInstances(
-	TEXT("r.ISMPool.AutoRemoveInstances"),
-	GAutoRemoveInstances,
-	TEXT("Remove instances with zero scale from ISM components in the ISMPool."));
 
 FGeometryCollectionMeshGroup::FMeshId FGeometryCollectionMeshGroup::AddMesh(const FGeometryCollectionStaticMeshInstance& MeshInstance, int32 InstanceCount, const FGeometryCollectionMeshInfo& ISMInstanceInfo)
 {
@@ -36,16 +31,11 @@ FGeometryCollectionMeshGroup::FMeshId FGeometryCollectionMeshGroup::AddMesh(cons
 	return MeshInfoIndex;
 }
 
-bool FGeometryCollectionMeshGroup::BatchUpdateInstancesTransforms(FGeometryCollectionISMPool& ISMPool, FMeshId MeshId, int32 StartInstanceIndex, const TArray<FTransform>& NewInstancesTransforms, bool bWorldSpace, bool bMarkRenderStateDirty, bool bTeleport)
-{
-	return BatchUpdateInstancesTransforms(ISMPool, MeshId, StartInstanceIndex, MakeArrayView(NewInstancesTransforms), bWorldSpace, bMarkRenderStateDirty, bTeleport);
-}
-
 bool FGeometryCollectionMeshGroup::BatchUpdateInstancesTransforms(FGeometryCollectionISMPool& ISMPool, FMeshId MeshId, int32 StartInstanceIndex, TArrayView<const FTransform> NewInstancesTransforms, bool bWorldSpace, bool bMarkRenderStateDirty, bool bTeleport)
 {
 	if (MeshInfos.IsValidIndex(MeshId))
 	{
-		return ISMPool.BatchUpdateInstancesTransforms(MeshInfos[MeshId], StartInstanceIndex, NewInstancesTransforms, bWorldSpace, bMarkRenderStateDirty, bTeleport);
+		return ISMPool.BatchUpdateInstancesTransforms(MeshInfos[MeshId], StartInstanceIndex, NewInstancesTransforms, bWorldSpace, bMarkRenderStateDirty, bTeleport, bAllowPerInstanceRemoval);
 	}
 	UE_LOG(LogChaos, Warning, TEXT("UGeometryCollectionISMPoolComponent : Invalid mesh Id (%d) for this mesh group"), MeshId);
 	return false;
@@ -225,12 +215,7 @@ FGeometryCollectionMeshInfo FGeometryCollectionISMPool::AddISM(UGeometryCollecti
 	return Info;
 }
 
-bool FGeometryCollectionISMPool::BatchUpdateInstancesTransforms(FGeometryCollectionMeshInfo& MeshInfo, int32 StartInstanceIndex, const TArray<FTransform>& NewInstancesTransforms, bool bWorldSpace, bool bMarkRenderStateDirty, bool bTeleport)
-{
-	return BatchUpdateInstancesTransforms(MeshInfo, StartInstanceIndex, MakeArrayView(NewInstancesTransforms), bWorldSpace, bMarkRenderStateDirty, bTeleport);
-}
-
-bool FGeometryCollectionISMPool::BatchUpdateInstancesTransforms(FGeometryCollectionMeshInfo& MeshInfo, int32 StartInstanceIndex, TArrayView<const FTransform> NewInstancesTransforms, bool bWorldSpace, bool bMarkRenderStateDirty, bool bTeleport)
+bool FGeometryCollectionISMPool::BatchUpdateInstancesTransforms(FGeometryCollectionMeshInfo& MeshInfo, int32 StartInstanceIndex, TArrayView<const FTransform> NewInstancesTransforms, bool bWorldSpace, bool bMarkRenderStateDirty, bool bTeleport, bool bAllowPerInstanceRemoval)
 {
 	if (!ISMs.IsValidIndex(MeshInfo.ISMIndex))
 	{
@@ -256,7 +241,7 @@ bool FGeometryCollectionISMPool::BatchUpdateInstancesTransforms(FGeometryCollect
 		FPrimitiveInstanceId InstanceId = ISM.InstanceIds[InstanceGroup.Start + InstanceIndex];
 		FTransform const& Transform = NewInstancesTransforms[InstanceIndex];
 
-		if (GAutoRemoveInstances)
+		if (bAllowPerInstanceRemoval)
 		{
 			if (Transform.GetScale3D().IsZero() && InstanceId.IsValid())
 			{
@@ -429,9 +414,11 @@ void UGeometryCollectionISMPoolComponent::TickComponent(float DeltaTime, enum EL
 	Pool.GarbageCollect();
 }
 
-UGeometryCollectionISMPoolComponent::FMeshGroupId  UGeometryCollectionISMPoolComponent::CreateMeshGroup()
+UGeometryCollectionISMPoolComponent::FMeshGroupId  UGeometryCollectionISMPoolComponent::CreateMeshGroup(bool bAllowPerInstanceRemoval)
 {
-	MeshGroups.Add(NextMeshGroupId);
+	FGeometryCollectionMeshGroup Group;
+	Group.bAllowPerInstanceRemoval = bAllowPerInstanceRemoval;
+	MeshGroups.Add(NextMeshGroupId, Group);
 	return NextMeshGroupId++;
 }
 
diff --git a/Engine/Source/Runtime/Experimental/GeometryCollectionEngine/Private/GeometryCollection/GeometryCollectionISMPoolRenderer.cpp b/Engine/Source/Runtime/Experimental/GeometryCollectionEngine/Private/GeometryCollection/GeometryCollectionISMPoolRenderer.cpp
--- a/Engine/Source/Runtime/Experimental/GeometryCollectionEngine/Private/GeometryCollection/GeometryCollectionISMPoolRenderer.cpp
+++ b/Engine/Source/Runtime/Experimental/GeometryCollectionEngine/Private/GeometryCollection/GeometryCollectionISMPoolRenderer.cpp
@@ -26,10 +26,13 @@ void UGeometryCollectionISMPoolRenderer::OnUnregisterGeometryCollection()
 	ISMPoolActor = nullptr;
 }
 
-void UGeometryCollectionISMPoolRenderer::UpdateState(UGeometryCollection const& InGeometryCollection, FTransform const& InComponentTransform, bool bIsBroken, bool bIsVisible)
+void UGeometryCollectionISMPoolRenderer::UpdateState(UGeometryCollection const& InGeometryCollection, FTransform const& InComponentTransform, uint32 InStateFlags)
 {
 	ComponentTransform = InComponentTransform;
 
+	const bool bIsVisible = (InStateFlags & EState_Visible) != 0;
+	const bool bIsBroken = (InStateFlags & EState_Broken) != 0;
+
 	if (bIsVisible == false)
 	{
 		ReleaseGroup(InstancesGroup);
@@ -104,7 +107,6 @@ void UGeometryCollectionISMPoolRenderer::InitMergedMeshFromGeometryCollection(UG
 
 		FGeometryCollectionStaticMeshInstance StaticMeshInstance;
 		StaticMeshInstance.StaticMesh = StaticMesh;
-		StaticMeshInstance.Desc.Flags |= FISMComponentDescription::UseHISM;
 
 		TArray<float> DummyCustomData;
 		MergedMeshGroup.MeshIds.Add(ISMPoolComponent->AddMeshToGroup(MergedMeshGroup.GroupIndex, StaticMeshInstance, 1, DummyCustomData));
diff --git a/Engine/Source/Runtime/Experimental/GeometryCollectionEngine/Private/GeometryCollection/GeometryCollectionISMPoolRenderer.h b/Engine/Source/Runtime/Experimental/GeometryCollectionEngine/Private/GeometryCollection/GeometryCollectionISMPoolRenderer.h
--- a/Engine/Source/Runtime/Experimental/GeometryCollectionEngine/Private/GeometryCollection/GeometryCollectionISMPoolRenderer.h
+++ b/Engine/Source/Runtime/Experimental/GeometryCollectionEngine/Private/GeometryCollection/GeometryCollectionISMPoolRenderer.h
@@ -21,7 +21,7 @@ public:
 	//~ Begin IGeometryCollectionExternalRenderInterface Interface.
 	virtual void OnRegisterGeometryCollection(UGeometryCollectionComponent const& InComponent) override;
 	virtual void OnUnregisterGeometryCollection() override;
-	virtual void UpdateState(UGeometryCollection const& InGeometryCollection, FTransform const& InComponentTransform, bool bInIsBroken, bool bInIsVisible) override;
+	virtual void UpdateState(UGeometryCollection const& InGeometryCollection, FTransform const& InComponentTransform, uint32 InStateFlags) override;
 	virtual void UpdateRootTransform(UGeometryCollection const& InGeometryCollection, FTransform const& InRootTransform) override;
 	virtual void UpdateTransforms(UGeometryCollection const& InGeometryCollection, TArrayView<const FTransform> InTransforms) override;
 	//~ End IGeometryCollectionExternalRenderInterface Interface.
diff --git a/Engine/Source/Runtime/Experimental/GeometryCollectionEngine/Public/GeometryCollection/GeometryCollectionExternalRenderInterface.h b/Engine/Source/Runtime/Experimental/GeometryCollectionEngine/Public/GeometryCollection/GeometryCollectionExternalRenderInterface.h
--- a/Engine/Source/Runtime/Experimental/GeometryCollectionEngine/Public/GeometryCollection/GeometryCollectionExternalRenderInterface.h
+++ b/Engine/Source/Runtime/Experimental/GeometryCollectionEngine/Public/GeometryCollection/GeometryCollectionExternalRenderInterface.h
@@ -20,12 +20,27 @@ class IGeometryCollectionExternalRenderInterface
 	GENERATED_BODY()
 
 public:
+	enum EStateFlags
+	{
+		EState_Visible		= 1 << 0,
+		EState_Broken		= 1 << 1,
+		EState_ForcedBroken	= 1 << 2,
+	};
+
 	virtual void OnRegisterGeometryCollection(UGeometryCollectionComponent const& InComponent) = 0;
 	virtual void OnUnregisterGeometryCollection() = 0;
-	virtual void UpdateState(UGeometryCollection const& InGeometryCollection, FTransform const& InComponentTransform, bool bInIsBroken, bool bInIsVisible) = 0;
+	virtual void UpdateState(UGeometryCollection const& InGeometryCollection, FTransform const& InComponentTransform, uint32 InStateFlags) = 0;
 	virtual void UpdateRootTransform(UGeometryCollection const& InGeometryCollection, FTransform const& InRootTransform) = 0;
 	virtual void UpdateTransforms(UGeometryCollection const& InGeometryCollection, TArrayView<const FTransform> InTransforms) = 0;
 
+	UE_DEPRECATED(5.4, "Use flags version of UpdateState instead")
+	virtual void UpdateState(UGeometryCollection const& InGeometryCollection, FTransform const& InComponentTransform, bool bInIsBroken, bool bInIsVisible) 
+	{
+		uint32 StateFlags = 0;
+		StateFlags |= bInIsVisible ? IGeometryCollectionExternalRenderInterface::EState_Visible : 0;
+		StateFlags |= bInIsBroken ? IGeometryCollectionExternalRenderInterface::EState_Broken : 0;
+		return UpdateState(InGeometryCollection, InComponentTransform, StateFlags);
+	};
 	UE_DEPRECATED(5.3, "Use FTransform version of UpdateTransforms instead")
 	virtual void UpdateTransforms(UGeometryCollection const& InGeometryCollection, TArrayView<const FMatrix> InMatrices) {};
 };
diff --git a/Engine/Source/Runtime/Experimental/GeometryCollectionEngine/Public/GeometryCollection/GeometryCollectionISMPoolComponent.h b/Engine/Source/Runtime/Experimental/GeometryCollectionEngine/Public/GeometryCollection/GeometryCollectionISMPoolComponent.h
--- a/Engine/Source/Runtime/Experimental/GeometryCollectionEngine/Public/GeometryCollection/GeometryCollectionISMPoolComponent.h
+++ b/Engine/Source/Runtime/Experimental/GeometryCollectionEngine/Public/GeometryCollection/GeometryCollectionISMPoolComponent.h
@@ -230,7 +230,6 @@ struct FGeometryCollectionMeshGroup
 	/** Adds a new mesh with instance count. We expect to only add a unique mesh instance once to each group. Returns a ID that can be used to update the instances. */
 	FMeshId AddMesh(const FGeometryCollectionStaticMeshInstance& MeshInstance, int32 InstanceCount, const FGeometryCollectionMeshInfo& ISMInstanceInfo);
 	/** Update instance transforms for a group of instances. */
-	bool BatchUpdateInstancesTransforms(FGeometryCollectionISMPool& ISMPool, FMeshId MeshId, int32 StartInstanceIndex, const TArray<FTransform>& NewInstancesTransforms, bool bWorldSpace, bool bMarkRenderStateDirty, bool bTeleport);
 	bool BatchUpdateInstancesTransforms(FGeometryCollectionISMPool& ISMPool, FMeshId MeshId, int32 StartInstanceIndex, TArrayView<const FTransform> NewInstancesTransforms, bool bWorldSpace, bool bMarkRenderStateDirty, bool bTeleport);
 	void BatchUpdateInstanceCustomData(FGeometryCollectionISMPool& ISMPool, int32 CustomFloatIndex, float CustomFloatValue);
 
@@ -239,6 +238,9 @@ struct FGeometryCollectionMeshGroup
 
 	/** Array of allocated mesh infos. */
 	TArray<FGeometryCollectionMeshInfo> MeshInfos;
+
+	/** Flag for whether we allow removal of instances when transform scale is set to zero. */
+	bool bAllowPerInstanceRemoval = false;
 };
 
 /** Structure containting all info for a single ISM. */
@@ -272,8 +274,7 @@ struct FGeometryCollectionISMPool
 	/** Remove ISM contents. */
 	void RemoveISM(const FGeometryCollectionMeshInfo& MeshInfo);
 	/** Update ISM contents. */
-	bool BatchUpdateInstancesTransforms(FGeometryCollectionMeshInfo& MeshInfo, int32 StartInstanceIndex, const TArray<FTransform>& NewInstancesTransforms, bool bWorldSpace, bool bMarkRenderStateDirty, bool bTeleport);
-	bool BatchUpdateInstancesTransforms(FGeometryCollectionMeshInfo& MeshInfo, int32 StartInstanceIndex, TArrayView<const FTransform> NewInstancesTransforms, bool bWorldSpace, bool bMarkRenderStateDirty, bool bTeleport);
+	bool BatchUpdateInstancesTransforms(FGeometryCollectionMeshInfo& MeshInfo, int32 StartInstanceIndex, TArrayView<const FTransform> NewInstancesTransforms, bool bWorldSpace, bool bMarkRenderStateDirty, bool bTeleport, bool bAllowPerInstanceRemoval);
 	void BatchUpdateInstanceCustomData(FGeometryCollectionMeshInfo const& MeshInfo, int32 CustomFloatIndex, float CustomFloatValue);
 
 	/** Clear all ISM components and associated data. */
@@ -318,7 +319,7 @@ public:
 	* no resources are created until the meshes are added for this group 
 	* return a mesh group Id used to add and update instances
 	*/
-	GEOMETRYCOLLECTIONENGINE_API FMeshGroupId CreateMeshGroup();
+	GEOMETRYCOLLECTIONENGINE_API FMeshGroupId CreateMeshGroup(bool bAllowPerInstanceRemoval = false);
 
 	/** Destroy  a mesh group and its associated resources */
 	GEOMETRYCOLLECTIONENGINE_API void DestroyMeshGroup(FMeshGroupId MeshGroupId);
