Removed HISM usage from GeomCollection ISMPool renderer.

diff --git a/Engine/Source/Runtime/Experimental/GeometryCollectionEngine/Private/GeometryCollection/GeometryCollectionISMPoolRenderer.cpp b/Engine/Source/Runtime/Experimental/GeometryCollectionEngine/Private/GeometryCollection/GeometryCollectionISMPoolRenderer.cpp
--- a/Engine/Source/Runtime/Experimental/GeometryCollectionEngine/Private/GeometryCollection/GeometryCollectionISMPoolRenderer.cpp
+++ b/Engine/Source/Runtime/Experimental/GeometryCollectionEngine/Private/GeometryCollection/GeometryCollectionISMPoolRenderer.cpp
@@ -151,7 +151,6 @@ void UGeometryCollectionISMPoolRenderer::InitInstancesFromGeometryCollection(UGe
 			{
 				StaticMeshInstance.MaterialsOverrides = AutoInstanceMesh.Materials;
 			}
-			StaticMeshInstance.Desc.Flags |= FISMComponentDescription::UseHISM;
 
 			TArray<float> DummyCustomData;
 			InstancesGroup.MeshIds.Add(ISMPoolComponent->AddMeshToGroup(InstancesGroup.GroupIndex, StaticMeshInstance, AutoInstanceMesh.NumInstances, DummyCustomData));
