Chaos
- add a transient flag to dynamic particle to cache the "moving kinematic" state
- fix broadphase generating overlaps for sleeping particles against statics and stationary kinematics

We loop over all dynamics in the outer collision detection loop (as opposed to awake dynamics and moving kinematics) when there are lots of moving kinematics. However the inner check to accept the overlap assumed all kinematics were moving, so it would generate contacts for sleeping dynamics against stationary kinematics and statics.

diff --git a/Engine/Source/Runtime/Experimental/Chaos/Private/Chaos/Collision/ParticlePairMidPhase.cpp b/Engine/Source/Runtime/Experimental/Chaos/Private/Chaos/Collision/ParticlePairMidPhase.cpp
--- a/Engine/Source/Runtime/Experimental/Chaos/Private/Chaos/Collision/ParticlePairMidPhase.cpp
+++ b/Engine/Source/Runtime/Experimental/Chaos/Private/Chaos/Collision/ParticlePairMidPhase.cpp
@@ -917,6 +917,14 @@ namespace Chaos
 				CullDistance += FMath::Min(VelocityBoundsMultiplier * VMaxDt, MaxVelocityBoundsExpansion);
 			}
 
+#if !UE_BUILD_TEST && !UE_BUILD_SHIPPING
+			// At least one body must be dynamic, and at least one must be awake (or moving if a kinematic)
+			const bool bIsMoving0 = (P0->IsDynamic() && !P0->IsSleeping()) || P0->IsMovingKinematic();
+			const bool bIsMoving1 = (P1->IsDynamic() && !P1->IsSleeping()) || P1->IsMovingKinematic();
+			const bool bAnyMoving = bIsMoving0 || bIsMoving1;
+			ensureMsgf(bAnyMoving, TEXT("GenerateCollisions called on two stationary objects %s %s"), *P0->GetDebugName(), *P1->GetDebugName());
+#endif
+
 			// Run collision detection on all potentially colliding shape pairs
 			NumActiveConstraints = GenerateCollisionsImpl(CullDistance, Dt, Context);
 		}
diff --git a/Engine/Source/Runtime/Experimental/Chaos/Private/Chaos/PBDRigidsEvolution.cpp b/Engine/Source/Runtime/Experimental/Chaos/Private/Chaos/PBDRigidsEvolution.cpp
--- a/Engine/Source/Runtime/Experimental/Chaos/Private/Chaos/PBDRigidsEvolution.cpp
+++ b/Engine/Source/Runtime/Experimental/Chaos/Private/Chaos/PBDRigidsEvolution.cpp
@@ -1065,6 +1065,11 @@ namespace Chaos
 			{
 				Particle->SetSleepCounter(0);
 			}
+
+			// If we are not kinematic, the MovingKinematic flag should be cleared
+			// If we are kinematic and are (or become) moving, the flag will be updated in ApplyKinematicTargets
+			// Either way, we can reset the flag here
+			Particle->ClearIsMovingKinematic();
 		}
 
 		// If we are now dynamic and enabled, we need to be in the graph if not already there. We may not be
diff --git a/Engine/Source/Runtime/Experimental/Chaos/Private/Chaos/PBDRigidsEvolutionGBF.cpp b/Engine/Source/Runtime/Experimental/Chaos/Private/Chaos/PBDRigidsEvolutionGBF.cpp
--- a/Engine/Source/Runtime/Experimental/Chaos/Private/Chaos/PBDRigidsEvolutionGBF.cpp
+++ b/Engine/Source/Runtime/Experimental/Chaos/Private/Chaos/PBDRigidsEvolutionGBF.cpp
@@ -247,36 +247,65 @@ void CheckParticleViewForDupes(const FString& Name, const TParticleView& Particl
 }
 
 void CheckParticleViewsForDupes(FPBDRigidsSOAs& Particles)
+{
+	// Check that all particles know what lists they are in
+	Particles.CheckListMasks();
+
+	// Check that we have no particles in multiple views
+	Particles.CheckViewMasks();
+
+	// A particle appearing twice in either of these results in a race condition because the 
+	// collision detection loop will visit the same particle pair twice on different threads.
+	CheckParticleViewForDupes(TEXT("NonDisabledDynamicView"), Particles.GetNonDisabledDynamicView());
+	CheckParticleViewForDupes(TEXT("ActiveDynamicMovingKinematicParticlesView"), Particles.GetActiveDynamicMovingKinematicParticlesView());
+
+	// No known problems with dupes in these lists, but there should never be dupes
+	CheckParticleViewForDupes(TEXT("NonDisabledView"), Particles.GetNonDisabledView());
+	CheckParticleViewForDupes(TEXT("NonDisabledClusteredView"), Particles.GetNonDisabledClusteredView());
+	CheckParticleViewForDupes(TEXT("ActiveParticlesView"), Particles.GetActiveParticlesView());
+	CheckParticleViewForDupes(TEXT("DirtyParticlesView"), Particles.GetDirtyParticlesView());
+	CheckParticleViewForDupes(TEXT("AllParticlesView"), Particles.GetAllParticlesView());
+	CheckParticleViewForDupes(TEXT("ActiveKinematicParticlesView"), Particles.GetActiveKinematicParticlesView());
+	CheckParticleViewForDupes(TEXT("ActiveMovingKinematicParticlesView"), Particles.GetActiveMovingKinematicParticlesView());
+	CheckParticleViewForDupes(TEXT("ActiveStaticParticlesView"), Particles.GetActiveStaticParticlesView());
+	CheckParticleViewForDupes(TEXT("ActiveDynamicMovingKinematicParticlesView"), Particles.GetActiveDynamicMovingKinematicParticlesView());
+}
+
+void CheckMovingKinematicFlag(FPBDRigidsSOAs& Particles)
+{
+	// Make sure all particles have the correct value for the MovingKinematic flag (NOTE: must be after ApplyKinematicTargets)
+	for (const auto& Particle : Particles.GetActiveParticlesView())
+	{
+		if (const FPBDRigidParticleHandle* Rigid = Particle.Handle()->CastToRigidParticle())
+		{
+			if (Rigid->IsKinematic())
+			{
+				const bool bIsMoving = !Rigid->V().IsNearlyZero() || !Rigid->W().IsNearlyZero();
+				ensureMsgf(bIsMoving == Rigid->IsMovingKinematic(),
+					TEXT("Kinematic IsMoving flag mismatch. IsMoving=%d V=(%f, %f, %f) W=(%f %f %f) %s"),
+					Rigid->IsMovingKinematic(), Rigid->V().X, Rigid->V().Y, Rigid->V().Z, Rigid->W().X, Rigid->W().Y, Rigid->W().Z, *Rigid->GetDebugName());
+			}
+			else
+			{
+				ensureMsgf(!Rigid->IsMovingKinematic(), TEXT("Kinematic IsMoving flag set for non-kinematic %s"), *Rigid->GetDebugName());
+			}
+		}
+	}
+}
+
+void CheckParticleViewsForErrors(FPBDRigidsSOAs& Particles)
 {
 	if (CVars::bChaosSolverCheckParticleViews)
 	{
-		QUICK_SCOPE_CYCLE_COUNTER(STAT_Evolution_CheckParticleViewsForDupes);
-
-		// Check that all particles know what lists they are in
-		Particles.CheckListMasks();
-
-		// Check that we have no particles in multiple views
-		Particles.CheckViewMasks();
-
-		// A particle appearing twice in either of these results in a race condition because the 
-		// collision detection loop will visit the same particle pair twice on different threads.
-		CheckParticleViewForDupes(TEXT("NonDisabledDynamicView"), Particles.GetNonDisabledDynamicView());
-		CheckParticleViewForDupes(TEXT("ActiveDynamicMovingKinematicParticlesView"), Particles.GetActiveDynamicMovingKinematicParticlesView());
-
-		// No known problems with dupes in these lists, but there should never be dupes
-		CheckParticleViewForDupes(TEXT("NonDisabledView"), Particles.GetNonDisabledView());
-		CheckParticleViewForDupes(TEXT("NonDisabledClusteredView"), Particles.GetNonDisabledClusteredView());
-		CheckParticleViewForDupes(TEXT("ActiveParticlesView"), Particles.GetActiveParticlesView());
-		CheckParticleViewForDupes(TEXT("DirtyParticlesView"), Particles.GetDirtyParticlesView());
-		CheckParticleViewForDupes(TEXT("AllParticlesView"), Particles.GetAllParticlesView());
-		CheckParticleViewForDupes(TEXT("ActiveKinematicParticlesView"), Particles.GetActiveKinematicParticlesView());
-		CheckParticleViewForDupes(TEXT("ActiveMovingKinematicParticlesView"), Particles.GetActiveMovingKinematicParticlesView());
-		CheckParticleViewForDupes(TEXT("ActiveStaticParticlesView"), Particles.GetActiveStaticParticlesView());
-		CheckParticleViewForDupes(TEXT("ActiveDynamicMovingKinematicParticlesView"), Particles.GetActiveDynamicMovingKinematicParticlesView());
+		QUICK_SCOPE_CYCLE_COUNTER(STAT_Evolution_CheckParticleViewsForErrors);
+
+		CheckParticleViewsForDupes(Particles);
+		CheckMovingKinematicFlag(Particles);
 	}
 }
 
 
+
 void FPBDRigidsEvolutionGBF::Advance(const FReal Dt,const FReal MaxStepDt,const int32 MaxSteps)
 {
 	// Determine how many steps we would like to take
@@ -451,7 +480,7 @@ void FPBDRigidsEvolutionGBF::AdvanceOneTimeStepImpl(const FReal Dt, const FSubSt
 
 	// Collision detection is sensitive to duplication bugs in the particle views
 	// so this is here to help us track them down when they happen
-	CheckParticleViewsForDupes(Particles);
+	CheckParticleViewsForErrors(Particles);
 
 	{
 		SCOPE_CYCLE_COUNTER(STAT_Evolution_DetectCollisions);
diff --git a/Engine/Source/Runtime/Experimental/Chaos/Public/Chaos/Collision/SpatialAccelerationBroadPhase.h b/Engine/Source/Runtime/Experimental/Chaos/Public/Chaos/Collision/SpatialAccelerationBroadPhase.h
--- a/Engine/Source/Runtime/Experimental/Chaos/Public/Chaos/Collision/SpatialAccelerationBroadPhase.h
+++ b/Engine/Source/Runtime/Experimental/Chaos/Public/Chaos/Collision/SpatialAccelerationBroadPhase.h
@@ -51,63 +51,75 @@ namespace Chaos
 				return false;
 			}
 
-			bool bIsKinematic1 = true;
+			bool bIsMovingKinematic1 = false;
 			bool bIsDynamicAwake1 = false;
 			bool bIsDynamicAsleep1 = false;
 			const FPBDRigidParticleHandle* Rigid1 = Particle1->CastToRigidParticle();
 			if (Rigid1 != nullptr)
 			{
-				bIsKinematic1 = Rigid1->IsKinematic();
-				bIsDynamicAsleep1 = !bIsKinematic1 && Rigid1->IsSleeping();
-				bIsDynamicAwake1 = !bIsKinematic1 && !bIsDynamicAsleep1;
+				bIsMovingKinematic1 = Rigid1->IsMovingKinematic();
+				bIsDynamicAsleep1 = Rigid1->IsDynamic() && Rigid1->IsSleeping();
+				bIsDynamicAwake1 = Rigid1->IsDynamic() && !Rigid1->IsSleeping();;
 			}
 
-			bool bIsKinematic2 = true;
+			bool bIsMovingKinematic2 = false;
 			bool bIsDynamicAwake2 = false;
 			bool bIsDynamicAsleep2 = false;
 			const FPBDRigidParticleHandle* Rigid2 = Particle2->CastToRigidParticle();
 			if (Rigid2 != nullptr)
 			{
-				bIsKinematic2 = Rigid2->IsKinematic();
-				bIsDynamicAsleep2 = !bIsKinematic2 && Rigid2->IsSleeping();
-				bIsDynamicAwake2 = !bIsKinematic2 && !bIsDynamicAsleep2;
+				bIsMovingKinematic2 = Rigid2->IsMovingKinematic();
+				bIsDynamicAsleep2 = Rigid2->IsDynamic() && Rigid2->IsSleeping();
+				bIsDynamicAwake2 = Rigid2->IsDynamic() && !Rigid2->IsSleeping();
 			}
 
-			// In both cases (resim or not) we will generate (1) dynamic-(sleeping,kinematic(steady+moving),static) pairs + (2) sleeping-moving kinematic ones
-			// Sleeping particles could collide with dynamic ones but these collisions are already handled in case 1
-			// Sleeping particles won't collide with static or steady kinematic particles since neither are moving
-			// Sleeping particles will collide against moving kinematic particles
-			bool bAcceptParticlePair = false;
-			if ((bIsDynamicAwake1 && !bIsDynamicAwake2) || (bIsDynamicAsleep1 && bIsKinematic2))
-			{
-				bAcceptParticlePair = true;
-			}
-
-			// Used to determine a winner in cases where we will visit particle pairs in both orders
+			// Used to determine a winner in cases where we visit particle pairs in both orders
 			const bool bIsParticle1Preferred = AreParticlesInPreferredOrder(Particle1, Particle2);
 
+			bool bAcceptParticlePair = false;
 			if (!bIsResimming)
 			{
-				// Normally (not resimming) we iterate over dynamic and asleep|kinematic particles, so:
-				// - Particle1 is dynamic, asleep OR kinematic
-				// - Particle2 may be static, kinematic, dynamic, asleep
+				// Assumptions for the non-resim case:
+				//		- Particle1 is the particle from an outer loop over either
+				//			(A) all dynamic particles, or 
+				//			(B) all awake dynamic and moving kinematic particles.
+				//		- Particle2 is one of the particles whose bounds overlaps Particle1
 
-				// If Particle1 is non dynamic but particle 2 is dynamic, the case should already be handled by (1)
-				if (!bIsDynamicAwake1 && bIsDynamicAwake2)
+				// If the first particle is an awake dynamic
+				//		- accept if the other particle is an awake dynamic and we are the preferred particle (lower ID)
+				//		- accept if the other particle is in any other state (static, kinematic, asleep)
+				if (bIsDynamicAwake1)
+				{
+					if (bIsDynamicAwake2)
+					{
+						bAcceptParticlePair = bIsParticle1Preferred;
+					}
+					else
+					{
+						bAcceptParticlePair = true;
+					}
+				}
+				// If the first particle is an asleep dynamic
+				//		- accept if the other particle is a moving kinematic
+				//		- reject if the other particle is a stationary kinematic or static - sleeping and static particles do not collide
+				//		- reject if the other particle is an awake dynamic - will be picked up when visiting in the opposite order
+				//		- reject if the other particle if an asleep dynamic - two sleeping dynamics do not collide
+				else if (bIsDynamicAsleep1)
 				{
-					bAcceptParticlePair = false;
+					bAcceptParticlePair = bIsMovingKinematic2;
 				}
-				// If Particle1 and Particle2 are dynamic we validate the pair if particle1 has higher ID to discard duplicates since we will visit twice the pairs
-				else if (bIsDynamicAwake1 && bIsDynamicAwake2)
+				// If the first particle is a moving kinematic
+				//		- accept if the other particle is sleeping dynamic
+				//		- reject if the other particle is an awake dynamic - will be picked up when visiting in the opposite order
+				//		- reject if the other particle is a kinematic or static - they do not collide
+				else if (bIsMovingKinematic1)
 				{
-					bAcceptParticlePair = bIsParticle1Preferred;
+					bAcceptParticlePair = bIsDynamicAsleep2;
 				}
-				// If Particle1 is kinematic we should in theory discard the pairs against sleeping particles
-				// since the sleeping-kinematic case has been validated in (2). But Particle1 is asleep OR kinematic so
-				// when entering this condition we are sure that we enver entered (2). It is why we validate the kinematic-sleeping pairs as well
-				else if (bIsKinematic1 && bIsDynamicAsleep2)
+				// If the first particle is static or non-moving kinematic
+				//		- reject always - will be picked up when visiting in the opposite order
+				else
 				{
-					bAcceptParticlePair = true;
 				}
 			}
 			else
@@ -121,6 +133,14 @@ namespace Chaos
 				// Even though Particle1 may be kinematic when resimming, we want to create the contacts in the original order (i.e., dynamic first)
 				// 
 				const bool bIsParticle2Desynced = bIsResimming && (Particle2->SyncState() == ESyncState::HardDesync);
+				const bool bIsKinematic1 = !bIsDynamicAwake1 && !bIsDynamicAsleep1;
+				const bool bIsKinematic2 = !bIsDynamicAwake2 && !bIsDynamicAsleep2;
+
+				// If Particle1 is dynamic, accept if the other is asleep or nor dynamic
+				if ((bIsDynamicAwake1 && !bIsDynamicAwake2) || (bIsDynamicAsleep1 && bIsKinematic2))
+				{
+					bAcceptParticlePair = true;
+				}
 
 				// If Particle1 is non dynamic but particle 2 is dynamic, the case should already be handled by (1) for
 				// the desynced dynamic - synced/desynced (static,kinematic,asleep) pairs. But we still need to process
diff --git a/Engine/Source/Runtime/Experimental/Chaos/Public/Chaos/PBDRigidsEvolution.h b/Engine/Source/Runtime/Experimental/Chaos/Public/Chaos/PBDRigidsEvolution.h
--- a/Engine/Source/Runtime/Experimental/Chaos/Public/Chaos/PBDRigidsEvolution.h
+++ b/Engine/Source/Runtime/Experimental/Chaos/Public/Chaos/PBDRigidsEvolution.h
@@ -455,6 +455,13 @@ public:
 		Particles.DisableParticle(Particle);
 		DisableConstraints(Particle);
 		DestroyTransientConstraints(Particle);
+
+		if (FPBDRigidParticleHandle* Rigid = Particle->CastToRigidParticle())
+		{
+			// This flag is only updated for moving kinematics, so make sure 
+			// we don't leave a residual value if we get enabled again
+			Rigid->ClearIsMovingKinematic();
+		}
 	}
 
 	/**
@@ -763,6 +770,7 @@ public:
 				// Reset velocity and then switch to do-nothing mode
 				Particle.V() = FVec3(0, 0, 0);
 				Particle.W() = FVec3(0, 0, 0);
+				Particle.ClearIsMovingKinematic();
 				KinematicTarget.SetMode(EKinematicTargetMode::None);
 				break;
 			}
@@ -807,6 +815,7 @@ public:
 				Particle.R() = NewR;
 				Particle.V() = NewV;
 				Particle.W() = NewW;
+				Particle.SetIsMovingKinematic();
 
 				break;
 			}
@@ -817,6 +826,8 @@ public:
 				bMoved = true;
 				Particle.X() = Particle.X() + Particle.V() * Dt;
 				Particle.R() = FRotation3::IntegrateRotationWithAngularVelocity(Particle.R(), Particle.W(), Dt);
+				Particle.SetIsMovingKinematic();
+
 				break;
 			}
 			}
diff --git a/Engine/Source/Runtime/Experimental/Chaos/Public/Chaos/ParticleHandle.h b/Engine/Source/Runtime/Experimental/Chaos/Public/Chaos/ParticleHandle.h
--- a/Engine/Source/Runtime/Experimental/Chaos/Public/Chaos/ParticleHandle.h
+++ b/Engine/Source/Runtime/Experimental/Chaos/Public/Chaos/ParticleHandle.h
@@ -1297,6 +1297,21 @@ public:
 		PBDRigidParticles->TransientFlags(ParticleIdx).ClearUseIgnoreCollisionManager();
 	}
 
+	inline bool IsMovingKinematic() const
+	{
+		return IsKinematic() && PBDRigidParticles->TransientFlags(ParticleIdx).GetIsMovingKinematic();
+	}
+
+	inline void SetIsMovingKinematic()
+	{
+		PBDRigidParticles->TransientFlags(ParticleIdx).SetIsMovingKinematic();
+	}
+
+	inline void ClearIsMovingKinematic()
+	{
+		PBDRigidParticles->TransientFlags(ParticleIdx).ClearIsMovingKinematic();
+	}
+
 	ESleepType SleepType() const { return PBDRigidParticles->SleepType(ParticleIdx);}
 
 	void SetSleepType(ESleepType SleepType){ PBDRigidParticles->SetSleepType(ParticleIdx, SleepType); }
@@ -2212,6 +2227,31 @@ public:
 		}
 	}
 
+	inline bool IsMovingKinematic() const
+	{
+		if (auto RigidHandle = MHandle->CastToRigidParticle())
+		{
+			return RigidHandle->IsMovingKinematic();
+		}
+		return false;
+	}
+
+	inline void SetIsMovingKinematic()
+	{
+		if (auto RigidHandle = MHandle->CastToRigidParticle())
+		{
+			RigidHandle->SetIsMovingKinematic();
+		}
+	}
+
+	inline void ClearIsMovingKinematic()
+	{
+		if (auto RigidHandle = MHandle->CastToRigidParticle())
+		{
+			RigidHandle->ClearIsMovingKinematic();
+		}
+	}
+
 	FReal LinearEtherDrag() const
 	{
 		if (auto RigidHandle = MHandle->CastToRigidParticle())
diff --git a/Engine/Source/Runtime/Experimental/Chaos/Public/Chaos/RigidParticleControlFlags.h b/Engine/Source/Runtime/Experimental/Chaos/Public/Chaos/RigidParticleControlFlags.h
--- a/Engine/Source/Runtime/Experimental/Chaos/Public/Chaos/RigidParticleControlFlags.h
+++ b/Engine/Source/Runtime/Experimental/Chaos/Public/Chaos/RigidParticleControlFlags.h
@@ -132,11 +132,18 @@ namespace Chaos
 		void SetUseIgnoreCollisionManager() { Flags.bUseIgnoreCollisionManager = true; }
 		void ClearUseIgnoreCollisionManager() { Flags.bUseIgnoreCollisionManager = false; }
 
+		// Is this particle kinematic and moving (velocity and angular velocity are non-zero). This is updated in ApplyKinematicTargets
+		// to help avoid checking velocity and angular velocity (especially for kinematics).
+		bool GetIsMovingKinematic() const { return Flags.bIsMovingKinematic; }
+		void SetIsMovingKinematic() { Flags.bIsMovingKinematic = true; }
+		void ClearIsMovingKinematic() { Flags.bIsMovingKinematic = false; }
+
 	private:
 		struct FFlags
 		{
 			FStorage bInertiaConditioningDirty : 1;
 			FStorage bUseIgnoreCollisionManager : 1;
+			FStorage bIsMovingKinematic : 1;
 			// Add new properties above this line
 			// Change FStorage typedef if we exceed the max bits
 		};
