Chaos - add cvar to control whether two one-way-interaction particles collide with each other

diff --git a/Engine/Source/Runtime/Experimental/Chaos/Private/Chaos/PBDRigidsEvolutionGBF.cpp b/Engine/Source/Runtime/Experimental/Chaos/Private/Chaos/PBDRigidsEvolutionGBF.cpp
--- a/Engine/Source/Runtime/Experimental/Chaos/Private/Chaos/PBDRigidsEvolutionGBF.cpp
+++ b/Engine/Source/Runtime/Experimental/Chaos/Private/Chaos/PBDRigidsEvolutionGBF.cpp
@@ -123,6 +123,10 @@ namespace Chaos
 		bool bChaosSolverPersistentGraph = true;
 		FAutoConsoleVariableRef CVarChaosSolverPersistentGraph(TEXT("p.Chaos.Solver.PersistentGraph"), bChaosSolverPersistentGraph, TEXT(""));
 
+		// Whether we should ignore collisions between two particles that both have one-way collision enabled
+		bool bChaosIgnoreOneWayPairCollisions = true;
+		FAutoConsoleVariableRef CVarChaosIgnoreOneWayPairCollisions(TEXT("p.Chaos.Solver.IgnoreOneWayPairCollisions"), bChaosIgnoreOneWayPairCollisions, TEXT("Ignore collisions between two one-way-interaction particles"));
+
 		DECLARE_CYCLE_STAT(TEXT("FPBDRigidsEvolutionGBF::AdvanceOneTimeStep"), STAT_Evolution_AdvanceOneTimeStep, STATGROUP_Chaos);
 		DECLARE_CYCLE_STAT(TEXT("FPBDRigidsEvolutionGBF::UnclusterUnions"), STAT_Evolution_UnclusterUnions, STATGROUP_Chaos);
 		DECLARE_CYCLE_STAT(TEXT("FPBDRigidsEvolutionGBF::Integrate"), STAT_Evolution_Integrate, STATGROUP_Chaos);
diff --git a/Engine/Source/Runtime/Experimental/Chaos/Public/Chaos/Collision/SpatialAccelerationBroadPhase.h b/Engine/Source/Runtime/Experimental/Chaos/Public/Chaos/Collision/SpatialAccelerationBroadPhase.h
--- a/Engine/Source/Runtime/Experimental/Chaos/Public/Chaos/Collision/SpatialAccelerationBroadPhase.h
+++ b/Engine/Source/Runtime/Experimental/Chaos/Public/Chaos/Collision/SpatialAccelerationBroadPhase.h
@@ -23,6 +23,7 @@
 namespace Chaos::CVars
 {
 	extern bool bChaosMidPhaseRedistributionEnabled;
+	extern bool bChaosIgnoreOneWayPairCollisions;
 }
 
 namespace Chaos
@@ -679,11 +680,14 @@ namespace Chaos
 
 				if (Overlap.bCollisionsEnabled)
 				{
-					const bool bIsOneWay0 = FConstGenericParticleHandle(Overlap.Particles[0])->OneWayInteraction();
-					const bool bIsOneWay1 = FConstGenericParticleHandle(Overlap.Particles[1])->OneWayInteraction();
-					if (bIsOneWay0 && bIsOneWay1)
+					if (CVars::bChaosIgnoreOneWayPairCollisions)
 					{
-						continue;
+						const bool bIsOneWay0 = FConstGenericParticleHandle(Overlap.Particles[0])->OneWayInteraction();
+						const bool bIsOneWay1 = FConstGenericParticleHandle(Overlap.Particles[1])->OneWayInteraction();
+						if (bIsOneWay0 && bIsOneWay1)
+						{
+							continue;
+						}
 					}
 
 					// Get the midphase for this pair
