Always respect the fail time out regardless of the wait type

#rnx

diff --git a/Engine/Source/Runtime/Experimental/IoStoreOnDemand/Private/CoreHttp/Http.cpp b/Engine/Source/Runtime/Experimental/IoStoreOnDemand/Private/CoreHttp/Http.cpp
--- a/Engine/Source/Runtime/Experimental/IoStoreOnDemand/Private/CoreHttp/Http.cpp
+++ b/Engine/Source/Runtime/Experimental/IoStoreOnDemand/Private/CoreHttp/Http.cpp
@@ -2548,6 +2548,7 @@ struct FTickState
 	uint64						Cancels;
 	int32&						RecvAllowance;
 	int32						PollTimeoutMs;
+	int32						FailTimeoutMs;
 	class FWorkQueue*			Work;
 };
 
@@ -2925,11 +2926,12 @@ public:
 	void						AddActivity(FActivity* Activity);
 
 private:
-	int32						Wait(int32 PollTimeoutMs);
+	int32						Wait(const FTickState& State);
 	TArray<FSocketGroup>		SocketGroups;
 	FWorkQueue					Work;
 	FHost&						Host;
 	uint32						BusyCount = 0;
+	int32						WaitTimeAccum = 0;
 
 	UE_NONCOPYABLE(FHostGroup);
 };
@@ -2943,7 +2945,7 @@ FHostGroup::FHostGroup(FHost& InHost)
 }
 
 ////////////////////////////////////////////////////////////////////////////////
-int32 FHostGroup::Wait(int32 PollTimeoutMs)
+int32 FHostGroup::Wait(const FTickState& State)
 {
 	// Collect groups that are waiting on something
 	TArray<FSocket::FWaiter, TFixedAllocator<64>> Waiters;
@@ -2965,24 +2967,31 @@ int32 FHostGroup::Wait(int32 PollTimeoutMs)
 	ON_SCOPE_EXIT { Trace(0, ETrace::Unwait); };
 
 	// If the poll timeout is negative then treat that as a fatal timeout
-	bool bFailOnTimeout = false;
-	if (PollTimeoutMs < -1)
+	check(State.FailTimeoutMs);
+	int32 PollTimeoutMs = State.PollTimeoutMs;
+	if (PollTimeoutMs < 0)
 	{
-		PollTimeoutMs = -PollTimeoutMs;
-		bFailOnTimeout = true;
+		PollTimeoutMs = State.FailTimeoutMs;
 	}
 
 	// Actually do the wait
 	int32 Result = FSocket::Wait(Waiters, PollTimeoutMs);
 	if (Result <= 0)
 	{
-		if (bFailOnTimeout && Result == 0)
+		// If the user opts to not block then we don't accumulate wait time and
+		// leave it to them to manage time a fail timoue
+		WaitTimeAccum += PollTimeoutMs;
+
+		if (State.PollTimeoutMs < 0 || WaitTimeAccum >= State.FailTimeoutMs)
 		{
 			return MIN_int32;
 		}
+
 		return Result;
 	}
 
+	WaitTimeAccum = 0;
+
 	// For each waiter that's ready, find the associated group "unwait" them.
 	int32 Count = 0;
 	for (int32 i = 0, n = Waiters.Num(); i < n; ++i)
@@ -3028,7 +3037,7 @@ void FHostGroup::Tick(FTickState& State)
 	}
 
 	// Wait on the groups that are
-	if (int32 Result = Wait(State.PollTimeoutMs); Result < 0)
+	if (int32 Result = Wait(State); Result < 0)
 	{
 		const char* Reason = (Result == MIN_int32)
 			? "FailTimeout hit"
@@ -3085,7 +3094,7 @@ private:
 	FActivity*				Pending			= nullptr;
 	FThrottler				Throttler;
 	TArray<FHostGroup>		Groups;
-	int32					FailTimeoutMs	= 0;
+	int32					FailTimeoutMs	= 30 * 1000;
 	uint32					BusyCount		= 0;
 };
 
@@ -3180,7 +3189,7 @@ void FEventLoop::FImpl::Throttle(uint32 KiBPerSec)
 void FEventLoop::FImpl::SetFailTimeout(int32 TimeoutMs)
 {
 	// While TimeoutMs must be >=0, it is signed so that MAX_uint32 can't be used
-	check(TimeoutMs >= 0);
+	check(TimeoutMs > 0);
 	FailTimeoutMs = TimeoutMs;
 }
 
@@ -3236,20 +3245,13 @@ uint32 FEventLoop::FImpl::Tick(int32 PollTimeoutMs)
 	int32 RecvAllowance = Throttler.GetAllowance();
 	uint64 CancelsLoad = Cancels.load(std::memory_order_relaxed);
 
-	// There is a special internal meaning to the poll time out. If a infinite
-	// timeout is given and a fail one is set on the loop, we use a <-1 value
-	// that Wait() knows what to do with.
-	if (PollTimeoutMs < 0 && FailTimeoutMs)
-	{
-		PollTimeoutMs = -(FailTimeoutMs + 1);
-	}
-
 	// Tick groups and then remove ones that are idle
 	FTickState TickState = {
 		.DoneList = nullptr,
 		.Cancels = CancelsLoad,
 		.RecvAllowance = RecvAllowance,
 		.PollTimeoutMs = PollTimeoutMs,
+		.FailTimeoutMs = FailTimeoutMs,
 	};
 	for (FHostGroup& Group : Groups)
 	{
@@ -3655,7 +3657,6 @@ IOSTOREONDEMAND_API void IasHttpTest(const ANSICHAR* TestHost="localhost")
 	};
 
 	MiscTest();
-	ThrottleTest(BuildUrl("/data/"));
 
 	struct
 	{
@@ -3781,6 +3782,22 @@ IOSTOREONDEMAND_API void IasHttpTest(const ANSICHAR* TestHost="localhost")
 		WaitForLoopIdle();
 	}
 
+	// pool
+	for (uint16 i = 1; i < 64; ++i)
+	{
+		FConnectionPool::FParams Params;
+		Params.SetHostFromUrl(BuildUrl());
+		Params.ConnectionCount = (i % 2) + 1;
+		Params.PipelineLength = (i % 5) + 1;
+		FConnectionPool Pool(Params);
+		for (int32 j = 0; j < i; ++j)
+		{
+			FRequest Request = Loop.Get(BuildUrl("/data"), Pool);
+			Loop.Send(MoveTemp(Request), HashSink);
+		}
+		WaitForLoopIdle();
+	}
+
 	// fatal timeout
 	for (int32 i = 0; i < 14; ++i)
 	{
@@ -3805,24 +3822,66 @@ IOSTOREONDEMAND_API void IasHttpTest(const ANSICHAR* TestHost="localhost")
 			}
 		);
 
+		int32 PollTimeoutMs = -1;
 		if (bExpectFailTimeout)
 		{
 			Loop2.SetFailTimeout(1000);
+
+			if ((i & 3) == 1)
+			{
+				PollTimeoutMs = 1000;
+			}
 		}
-		while (Loop2.Tick(-1));
+		while (Loop2.Tick(PollTimeoutMs));
 	}
 
 	// no connect
 	{
-		FRequest Request[] = {
+		FRequest Requests[] = {
 			Loop.Request("GET", BuildUrl(nullptr, 10930)),
 			Loop.Request("GET", "http://thisdoesnotexistihope/"),
 		};
-		Loop.Send(MoveTemp(Request[0]), NullSink);
-		Loop.Send(MoveTemp(Request[1]), NullSink);
+		Loop.Send(MoveTemp(Requests[0]), NullSink);
+		Loop.Send(MoveTemp(Requests[1]), NullSink);
 		WaitForLoopIdle();
 	}
 
+	// head and large requests
+	{
+		auto MixTh = [Th=uint32(0)] () mutable { return (Th = (Th * 75) + 74) & 255; };
+
+		char AsciiData[257];
+		for (char& c : AsciiData)
+		{
+			int32 i = int32(ptrdiff_t(&c - AsciiData));
+			c = 0x41 + (MixTh() % 26);
+			c += (MixTh() & 2) ? 0x20 : 0;
+		}
+
+		for (int32 i = 0; (i += 69493) < 2 << 20;)
+		{
+			FRequest Request = Loop.Request("HEAD", BuildUrl("/data"));
+			for (int32 j = i; j > 0;)
+			{
+				FAnsiStringView Name(AsciiData, MixTh() + 1);
+				FAnsiStringView Value(AsciiData, MixTh() + 1);
+				Request.Header(Name, Value);
+				j -= Name.Len() + Value.Len();
+
+			}
+
+			Loop.Send(MoveTemp(Request), [] (const FTicketStatus& Status) {
+				if (Status.GetId() == FTicketStatus::EId::Response)
+				{
+					FResponse& Response = Status.GetResponse();
+					check(Response.GetStatusCode() == 431); // "too many headers"
+				}
+			});
+
+			WaitForLoopIdle();
+		}
+	}
+
 	// stress 1
 	{
 		const uint32 StressLoad = 32;
@@ -3918,6 +3977,10 @@ IOSTOREONDEMAND_API void IasHttpTest(const ANSICHAR* TestHost="localhost")
 
 	check(Loop.IsIdle());
 
+#if IS_PROGRAM
+	ThrottleTest(BuildUrl("/data/"));
+#endif
+
 	// pre-generated headers
 	// request-with-body
 	// proxy
@@ -3926,7 +3989,6 @@ IOSTOREONDEMAND_API void IasHttpTest(const ANSICHAR* TestHost="localhost")
 	// redirects
 	// loop multi-req.
 	// tls
-	// http pipelining
 	// url auth credentials
 	// transfer-file / splice / sendfile
 	// (header field parser)
