Fix a bug where using the cmdline 'Ias.TocUrl=' would still attempt to read from disk if 's.IasGenerateOnDemandToc' is enabled.

#rnx

- Add a new 'bForceDownload' member to FOnDemandEndpoint (default false), when set to true we will attempt to download that toc even if  's.IasGenerateOnDemandToc' is enabled.
-- This is set to true if the cmdline arg 'Ias.TocUrl=' was used as it indicates that the user wishes to download the file rather than relying on the data setup.
- If loading a toc is deferred we will need to preserve the bForceDownload flag, so FOnDemandIoBackend::DeferredTocs has been changed to be an array of a new struct 'FTocParams' which contains the path of the toc and if it should be forced downloaded or not.
- Renamed FOnDemandIoBackend::AddDeferredTocs to ::FlushDeferredTocs to better describe it's purpose.
- When loading a .iochunktoc from disk we no longer fallback to generating the file if it cannot be found. So enabling 'ias.LoadOnDemandToc' assumes that the file exists. If the file is missing we will return EIoErrorCode::NotFound.
- Cleaned up the error code returned if we fail to open a .iochunktoc for reading (EIoErrorCode::ReadError ->EIoErrorCode::FileOpenFailed) to better describe the problem.

diff --git a/Engine/Source/Runtime/Experimental/IoStoreOnDemand/Private/IoStoreOnDemand.cpp b/Engine/Source/Runtime/Experimental/IoStoreOnDemand/Private/IoStoreOnDemand.cpp
--- a/Engine/Source/Runtime/Experimental/IoStoreOnDemand/Private/IoStoreOnDemand.cpp
+++ b/Engine/Source/Runtime/Experimental/IoStoreOnDemand/Private/IoStoreOnDemand.cpp
@@ -586,6 +586,7 @@ TIoStatusOr<FOnDemandToc> LoadTocFromUrl(const FString& ServiceUrl, const FStrin
 		TAnsiStringBuilder<256> Url;
 
 		Url << "/" << TocPath;
+
 		UE_LOG(LogIas, Log, TEXT("Fetching TOC '%s/%s' (#%d/%d)"), *ServiceUrl, *TocPath, Attempt + 1, RetryCount);
 
 		TIoStatusOr<FOnDemandToc> Toc;
@@ -1624,6 +1625,10 @@ void FIoStoreOnDemandModule::InitializeInternal()
 			{
 				Endpoint.ServiceUrls.Add(FString(UrlView.Left(7 +  Delim)));
 				Endpoint.TocPath = UrlView.RightChop(Endpoint.ServiceUrls[0].Len() + 1);
+
+				// Since the user has provided a url to download the .iochunktoc from we should
+				// assume that they want to download it rather than use anything on disk.
+				Endpoint.bForceTocDownload = true;
 			}
 		}
 	}
diff --git a/Engine/Source/Runtime/Experimental/IoStoreOnDemand/Private/OnDemandIoDispatcherBackend.cpp b/Engine/Source/Runtime/Experimental/IoStoreOnDemand/Private/OnDemandIoDispatcherBackend.cpp
--- a/Engine/Source/Runtime/Experimental/IoStoreOnDemand/Private/OnDemandIoDispatcherBackend.cpp
+++ b/Engine/Source/Runtime/Experimental/IoStoreOnDemand/Private/OnDemandIoDispatcherBackend.cpp
@@ -962,7 +962,7 @@ static TIoStatusOr<FOnDemandToc> GenerateOnDemandTocFromDisk(FStringView TocHash
 			TUniquePtr<FArchive> Ar = CreateReaderFromPlatformPackage(TocPath);
 			if (!Ar.IsValid())
 			{
-				return FIoStatus(EIoErrorCode::ReadError, WriteToString<256>(TEXT("Failed to open '"), TocPath, TEXT("' from disk")));
+				return FIoStatus(EIoErrorCode::FileOpenFailed, WriteToString<256>(TEXT("Failed to open '"), TocPath, TEXT("' from disk")));
 			}
 
 			*Ar << OutToc;
@@ -981,7 +981,7 @@ static TIoStatusOr<FOnDemandToc> GenerateOnDemandTocFromDisk(FStringView TocHash
 		}
 		else
 		{
-			UE_LOG(LogIas, Warning, TEXT("Unable to find '%s' on disk, will attempt to generate it as a fallback"), *TocPath);
+			return FIoStatus(EIoErrorCode::NotFound, WriteToString<256>(TEXT("Unable to find '"), TocPath, TEXT("' on disk")));
 		}
 	}
 
@@ -1505,16 +1505,22 @@ private:
 	FIoStatus ApplyGeneratedOnDemandToc(const FString& CdnUrl, const FString& TocPath);
 	FIoStatus DownloadoadOnDemandToc(const FString& CdnUrl, const FString& TocPath);
 
-	void AddDeferredTocs();
+	void FlushDeferredTocs();
 	void ProcessHttpRequests(FHttpClient& HttpClient, FBitWindow& HttpErrors, int32 MaxConcurrentRequests);
 	int32 WaitForPendingRequests(float WaitTimeSeconds, float PollTimeSeconds);
 
+	struct FTocParams
+	{
+		FString Path;
+		bool bForceDownload;
+	};
+
 	TUniquePtr<IIasCache> Cache;
 	TUniquePtr<FOnDemandIoStore> IoStore;
 	TSharedPtr<const FIoDispatcherBackendContext> BackendContext;
 	TUniquePtr<FRunnableThread> BackendThread;
 	FEventRef TickBackendEvent;
-	TArray<FString> DeferredTocs;
+	TArray<FTocParams> DeferredTocs;
 	FChunkRequests ChunkRequests;
 	FIoRequestQueue CompletedRequests;
 	FChunkRequestQueue HttpRequests;
@@ -1606,7 +1612,7 @@ void FOnDemandIoBackend::Initialize(TSharedRef<const FIoDispatcherBackendContext
 					{
 						AvailableEps.Current = Idx;
 						UE_LOG(LogIas, Log, TEXT("Using endpoint '%s'"), *AvailableEps.GetCurrent());
-						AddDeferredTocs();
+						FlushDeferredTocs();
 					}
 					else
 					{
@@ -1647,7 +1653,7 @@ FString FOnDemandIoBackend::GetEndpointTestPath() const
 		FReadScopeLock _(Lock);
 		if (DeferredTocs.IsEmpty() == false)
 		{
-			TestPath = DeferredTocs[0];
+			TestPath = DeferredTocs[0].Path;
 		}
 	}
 	return TestPath;
@@ -1950,7 +1956,9 @@ FIoStatus FOnDemandIoBackend::ApplyGeneratedOnDemandToc(const FString& CdnUrl, c
 {
 	TRACE_CPUPROFILER_EVENT_SCOPE(IasBackend::ApplyGeneratedOnDemandToc);
 
-	if (bGeneratedOnDemandToc || !GIasGenerateOnDemandToc)
+	// TODO: Assumes we only ever want to generate/load a .iochunktoc once, a safe assumption
+	// now but probably not in the future.
+	if (bGeneratedOnDemandToc)
 	{
 		return FIoStatus::Ok;
 	}
@@ -2005,42 +2013,34 @@ FIoStatus FOnDemandIoBackend::DownloadoadOnDemandToc(const FString& CdnUrl, cons
 	}
 }
 
-void FOnDemandIoBackend::AddDeferredTocs()
+void FOnDemandIoBackend::FlushDeferredTocs()
 {
-	TArray<FString> TocPaths;
+	TArray<FTocParams> Tocs;
 	{
 		FWriteScopeLock _(Lock);
 		if (DeferredTocs.IsEmpty() || AvailableEps.HasCurrent() == false)
 		{
 			return;
 		}
-		TocPaths = MoveTemp(DeferredTocs);
+		Tocs = MoveTemp(DeferredTocs);
 	}
 
-	// We are still reliant on a valid TocPath to generate the FOnDemandToc from disk as the path contains
-	// info that we still need.
-	if (!TocPaths.IsEmpty())
+	for (const FTocParams& TocParams : Tocs)
 	{
-		FIoStatus Result = ApplyGeneratedOnDemandToc(AvailableEps.GetCurrent(), TocPaths[0]);
-		if (!Result.IsOk())
+		if (GIasGenerateOnDemandToc && !TocParams.bForceDownload)
 		{
-			UE_LOG(LogIas, Error, TEXT("Failed to add generated toc', reason '%s'"), *Result.ToString());
+			FIoStatus Result = ApplyGeneratedOnDemandToc(AvailableEps.GetCurrent(), TocParams.Path);
+			if (!Result.IsOk())
+			{
+				UE_LOG(LogIas, Error, TEXT("Failed to add generated toc', reason '%s'"), *Result.ToString());
+			}
 		}
-	}
-
-	// Generating the FOnDemandToc is relying on DeferredTocs having entries and cribbing off the callback from resolving the CDN endpoint
-	// This is to a) make it easier to toggle during development b) easier to validate the results.
-	// So keeping that in mind we need to prevent the toc being downloaded if we are supposed to be using the disk generated version 
-	// instead.
-	// This can all be cleaned up when we either remove the toggle or drop the validation requirements AND clean up the initialization flow.
-	if (!GIasGenerateOnDemandToc)
-	{
-		for (const FString& TocPath : TocPaths)
+		else
 		{
-			FIoStatus Result = DownloadoadOnDemandToc(AvailableEps.GetCurrent(), TocPath);
+			FIoStatus Result = DownloadoadOnDemandToc(AvailableEps.GetCurrent(), TocParams.Path);
 			if (!Result.IsOk())
 			{
-				UE_LOG(LogIas, Error, TEXT("Failed to add TOC '%s/%s', reason '%s'"), *AvailableEps.GetCurrent(), *TocPath, *Result.ToString());
+				UE_LOG(LogIas, Error, TEXT("Failed to add TOC '%s/%s', reason '%s'"), *AvailableEps.GetCurrent(), *TocParams.Path, *Result.ToString());
 			}
 		}
 	}
@@ -2057,6 +2057,12 @@ void FOnDemandIoBackend::Mount(const FOnDemandEndpoint& Endpoint)
 		return;
 	}
 
+	// TODO: Should pass this info around via DeferredTocs
+	if (Endpoint.bForceTocDownload)
+	{
+		GIasGenerateOnDemandToc = false;
+	}
+
 	if (GIasGenerateOnDemandToc && GIasAsyncTocGenerationEnabled)
 	{
 		FString TocHash = FPaths::GetBaseFilename(Endpoint.TocPath);
@@ -2078,7 +2084,7 @@ void FOnDemandIoBackend::Mount(const FOnDemandEndpoint& Endpoint)
 					{
 						DistributionUrl = Endpoint.DistributionUrl;
 					}
-					DeferredTocs.Add(Endpoint.TocPath);
+					DeferredTocs.Add(FTocParams{Endpoint.TocPath, Endpoint.bForceTocDownload});
 					return;
 				}
 			}
@@ -2094,13 +2100,15 @@ void FOnDemandIoBackend::Mount(const FOnDemandEndpoint& Endpoint)
 
 		check(AvailableEps.HasCurrent());
 
-		FIoStatus GeneratedResult = ApplyGeneratedOnDemandToc(AvailableEps.GetCurrent(), Endpoint.TocPath);
-		if (!GeneratedResult.IsOk())
+		if (GIasGenerateOnDemandToc && !Endpoint.bForceTocDownload)
 		{
-			UE_LOG(LogIas, Error, TEXT("Failed to add generated toc', reason '%s'"), *GeneratedResult.ToString());
+			FIoStatus GeneratedResult = ApplyGeneratedOnDemandToc(AvailableEps.GetCurrent(), Endpoint.TocPath);
+			if (!GeneratedResult.IsOk())
+			{
+				UE_LOG(LogIas, Error, TEXT("Failed to add generated toc', reason '%s'"), *GeneratedResult.ToString());
+			}
 		}
-
-		if (!GIasGenerateOnDemandToc)
+		else
 		{
 			FIoStatus Result = DownloadoadOnDemandToc(AvailableEps.GetCurrent(), Endpoint.TocPath);
 			if (!Result.IsOk())
@@ -2108,9 +2116,10 @@ void FOnDemandIoBackend::Mount(const FOnDemandEndpoint& Endpoint)
 				UE_LOG(LogIas, Error, TEXT("Deferring TOC '%s/%s' due to '%s'"), *AvailableEps.GetCurrent(), *Endpoint.TocPath, *Result.ToString());
 					BackendStatus.SetHttpError(true);
 					FWriteScopeLock _(Lock);
-					DeferredTocs.Add(Endpoint.TocPath);
+					DeferredTocs.Add(FTocParams{Endpoint.TocPath, Endpoint.bForceTocDownload});
 			}
 		}
+
 		ConditionallyStartBackendThread();
 	}
 	else
@@ -2330,7 +2339,7 @@ uint32 FOnDemandIoBackend::Run()
 					HttpClient->SetEndpoint(Idx);
 					BackendStatus.SetHttpError(false);
 					UE_LOG(LogIas, Log, TEXT("Successfully reconnected to '%s'"), *AvailableEps.GetCurrent());
-					AddDeferredTocs();
+					FlushDeferredTocs();
 				}
 			}
 			else if (HttpClient->IsUsingPrimaryEndpoint() == false)
diff --git a/Engine/Source/Runtime/Experimental/IoStoreOnDemand/Private/OnDemandIoDispatcherBackend.h b/Engine/Source/Runtime/Experimental/IoStoreOnDemand/Private/OnDemandIoDispatcherBackend.h
--- a/Engine/Source/Runtime/Experimental/IoStoreOnDemand/Private/OnDemandIoDispatcherBackend.h
+++ b/Engine/Source/Runtime/Experimental/IoStoreOnDemand/Private/OnDemandIoDispatcherBackend.h
@@ -24,6 +24,8 @@ struct FOnDemandEndpoint
 	TArray<FString> ServiceUrls;
 	FString TocPath;
 
+	bool bForceTocDownload = false;
+
 	bool IsValid() const
 	{
 		return (DistributionUrl.Len() > 0 || ServiceUrls.Num() > 0) && TocPath.Len() > 0;
