Iris - Filter inclusion groups.

Filter inclusion groups are processed after dynamic filtering, filters implemented via UNetObjectFilter implementations, and their only purpose is to  optionally enable replication of objects that were filtered out. This allows for example an object to be spatially filtered out but then enabled by a team or view target group for example. No such groups are implemented by default, it's up to the user to implement the relevant game logic to create, update and destroy such groups. Only root objects may be added to filter inclusion groups, subobjects will be filtered like the root object.

diff --git a/Engine/Plugins/Runtime/ReplicationSystemTestPlugin/Source/Private/Tests/ReplicationSystem/Filtering/TestFiltering.cpp b/Engine/Plugins/Runtime/ReplicationSystemTestPlugin/Source/Private/Tests/ReplicationSystem/Filtering/TestFiltering.cpp
--- a/Engine/Plugins/Runtime/ReplicationSystemTestPlugin/Source/Private/Tests/ReplicationSystem/Filtering/TestFiltering.cpp
+++ b/Engine/Plugins/Runtime/ReplicationSystemTestPlugin/Source/Private/Tests/ReplicationSystem/Filtering/TestFiltering.cpp
@@ -31,6 +31,26 @@
 	virtual FName GetMockFilterName() const { return FName("MockFilter"); }
 	virtual FName GetMockFilterClassName() const { return FName("/Script/ReplicationSystemTestPlugin.MockNetObjectFilter"); }
 
+	void SetDynamicFilterStatus(ENetFilterStatus FilterStatus)
+	{
+		UMockNetObjectFilter::FFunctionCallSetup CallSetup;
+		CallSetup.AddObject.bReturnValue = true;
+		CallSetup.Filter.bFilterOutByDefault = FilterStatus == ENetFilterStatus::Disallow;
+
+		MockNetObjectFilter->SetFunctionCallSetup(CallSetup);
+		MockNetObjectFilter->ResetFunctionCallStatus();
+	}
+
+	void SetDynamicFragmentFilterStatus(ENetFilterStatus FilterStatus)
+	{
+		UMockNetObjectFilterUsingFragmentData::FFunctionCallSetup CallSetup;
+		CallSetup.AddObject.bReturnValue = true;
+		CallSetup.Filter.bFilterOutByDefault = FilterStatus == ENetFilterStatus::Disallow;
+
+		MockNetObjectFilterWithFragments->SetFunctionCallSetup(CallSetup);
+		MockNetObjectFilterWithFragments->ResetFunctionCallStatus();
+	}
+
 private:
 	void InitNetObjectFilterDefinitions()
 	{
@@ -419,7 +439,7 @@
 	// Setup group filter
 	FNetObjectGroupHandle GroupHandle = Server->ReplicationSystem->CreateGroup();
 	Server->ReplicationSystem->AddToGroup(GroupHandle, ServerObject0->NetRefHandle);
-	Server->ReplicationSystem->AddGroupFilter(GroupHandle);
+	Server->ReplicationSystem->AddExclusionFilterGroup(GroupHandle);
 
 	// Spawn object on server
 	UReplicatedTestObject* ServerObject = Server->CreateObject(0, 0);
@@ -443,7 +463,7 @@
 	// Object should not have been created on the clients 
 	UE_NET_ASSERT_EQ(Client->GetReplicationBridge()->GetReplicatedObject(ServerObject0->NetRefHandle), nullptr);
 		
-	// Object should have been crated on all clients
+	// Object should have been created on all clients
 	UE_NET_ASSERT_NE(Client->GetReplicationBridge()->GetReplicatedObject(ServerObject->NetRefHandle), nullptr);
 	UE_NET_ASSERT_NE(Client->GetReplicationBridge()->GetReplicatedObject(ServerSubObject->NetRefHandle), nullptr);
 }
@@ -470,12 +490,7 @@
 	UE_NET_ASSERT_EQ(ClientArray[0]->ConnectionIdOnServer, 1U);
 
 	// Send and deliver packets
-	Server->PreSendUpdate();
-	for (FReplicationSystemTestClient* Client : ClientArray)
-	{
-		Server->SendAndDeliverTo(Client, DeliverPacket);
-	}
-	Server->PostSendUpdate();
+	Server->UpdateAndSend(ClientArray, DeliverPacket);
 
 	// Object should have been created on the client with connection ID 1.
 	for (FReplicationSystemTestClient*& Client : ClientArray)
@@ -510,12 +525,7 @@
 	Server->ReplicationSystem->SetConnectionFilter(ServerObject->NetRefHandle, AllowedConnections, ENetFilterStatus::Allow);
 
 	// Send and deliver packets
-	Server->PreSendUpdate();
-	for (FReplicationSystemTestClient* Client : ClientArray)
-	{
-		Server->SendAndDeliverTo(Client, DeliverPacket);
-	}
-	Server->PostSendUpdate();
+	Server->UpdateAndSend(ClientArray, DeliverPacket);
 
 	// Object should only have been created on the last client
 	for (FReplicationSystemTestClient*& Client : ClientArray)
@@ -693,7 +703,7 @@
 
 	FNetObjectGroupHandle GroupHandle = Server->ReplicationSystem->CreateGroup();
 	Server->ReplicationSystem->AddToGroup(GroupHandle, ServerObject->NetRefHandle);
-	Server->ReplicationSystem->AddGroupFilter(GroupHandle);
+	Server->ReplicationSystem->AddExclusionFilterGroup(GroupHandle);
 
 	// Filter out objects in group
 	Server->ReplicationSystem->SetGroupFilterStatus(GroupHandle, Client->ConnectionIdOnServer, ENetFilterStatus::Disallow);
@@ -717,7 +727,7 @@
 
 	FNetObjectGroupHandle GroupHandle = Server->ReplicationSystem->CreateGroup();
 	Server->ReplicationSystem->AddToGroup(GroupHandle, ServerObject->NetRefHandle);
-	Server->ReplicationSystem->AddGroupFilter(GroupHandle);
+	Server->ReplicationSystem->AddExclusionFilterGroup(GroupHandle);
 
 	// Filter out objects in group
 	Server->ReplicationSystem->SetGroupFilterStatus(GroupHandle, Client->ConnectionIdOnServer, ENetFilterStatus::Disallow);
@@ -759,12 +769,7 @@
 	Server->ReplicationSystem->SetConnectionFilter(ServerObject->NetRefHandle, AllowedConnections, ENetFilterStatus::Allow);
 
 	// Send and deliver packets
-	Server->PreSendUpdate();
-	for (FReplicationSystemTestClient* Client : ClientArray)
-	{
-		Server->SendAndDeliverTo(Client, DeliverPacket);
-	}
-	Server->PostSendUpdate();
+	Server->UpdateAndSend(ClientArray, DeliverPacket);
 
 	// Object should only have been created on the last client
 	for (FReplicationSystemTestClient*& Client : ClientArray)
@@ -783,16 +788,11 @@
 	// Create and set group filter for last client only
 	FNetObjectGroupHandle GroupHandle = Server->ReplicationSystem->CreateGroup();
 	Server->ReplicationSystem->AddToGroup(GroupHandle, ServerObject->NetRefHandle);
-	Server->ReplicationSystem->AddGroupFilter(GroupHandle);
+	Server->ReplicationSystem->AddExclusionFilterGroup(GroupHandle);
 	Server->ReplicationSystem->SetGroupFilterStatus(GroupHandle, ConnectionIdForLastClient, ENetFilterStatus::Disallow);
 
 	// Send and deliver packets
-	Server->PreSendUpdate();
-	for (FReplicationSystemTestClient* Client : ClientArray)
-	{
-		Server->SendAndDeliverTo(Client, DeliverPacket);
-	}
-	Server->PostSendUpdate();
+	Server->UpdateAndSend(ClientArray, DeliverPacket);
 
 	// Object Should now have been destroyed on the last client
 	for (FReplicationSystemTestClient*& Client : ClientArray)
@@ -804,12 +804,7 @@
 	Server->ReplicationSystem->SetGroupFilterStatus(GroupHandle, ConnectionIdForLastClient, ENetFilterStatus::Allow);
 
 	// Send and deliver packets
-	Server->PreSendUpdate();
-	for (FReplicationSystemTestClient* Client : ClientArray)
-	{
-		Server->SendAndDeliverTo(Client, DeliverPacket);
-	}
-	Server->PostSendUpdate();
+	Server->UpdateAndSend(ClientArray, DeliverPacket);
 
 	// Object should now be recreated on last client
 	for (FReplicationSystemTestClient*& Client : ClientArray)
@@ -846,13 +841,7 @@
 UE_NET_TEST_FIXTURE(FTestFilteringFixture, DynamicFilterAddObjectAndRemoveObjectIsCalledWhenObjectIsDeleted)
 {
 	// Setup dynamic filter for the test
-	{
-		UMockNetObjectFilter::FFunctionCallSetup CallSetup;
-		CallSetup.AddObject.bReturnValue = true;
-		CallSetup.Filter.bFilterOutByDefault = false;
-		MockNetObjectFilter->SetFunctionCallSetup(CallSetup);
-		MockNetObjectFilter->ResetFunctionCallStatus();
-	}
+	SetDynamicFilterStatus(ENetFilterStatus::Allow);
 
 	UReplicatedTestObject* ServerObject = Server->CreateObject(0, 0);
 	Server->ReplicationSystem->SetFilter(ServerObject->NetRefHandle, MockFilterHandle);
@@ -888,13 +877,7 @@
 UE_NET_TEST_FIXTURE(FTestFilteringFixture, DynamicFilterCanAllowObjectToReplicate)
 {
 	// Setup dynamic filter for the test
-	{
-		UMockNetObjectFilter::FFunctionCallSetup CallSetup;
-		CallSetup.AddObject.bReturnValue = true;
-		CallSetup.Filter.bFilterOutByDefault = false;
-		MockNetObjectFilter->SetFunctionCallSetup(CallSetup);
-		MockNetObjectFilter->ResetFunctionCallStatus();
-	}
+	SetDynamicFilterStatus(ENetFilterStatus::Allow);
 
 	// Add client
 	FReplicationSystemTestClient* Client = CreateClient();
@@ -915,13 +898,7 @@
 UE_NET_TEST_FIXTURE(FTestFilteringFixture, DynamicFilterCanDisallowObjectToReplicate)
 {
 	// Setup dynamic filter for the test
-	{
-		UMockNetObjectFilter::FFunctionCallSetup CallSetup;
-		CallSetup.AddObject.bReturnValue = true;
-		CallSetup.Filter.bFilterOutByDefault = true;
-		MockNetObjectFilter->SetFunctionCallSetup(CallSetup);
-		MockNetObjectFilter->ResetFunctionCallStatus();
-	}
+	SetDynamicFilterStatus(ENetFilterStatus::Disallow);
 
 	// Add client
 	FReplicationSystemTestClient* Client = CreateClient();
@@ -942,13 +919,7 @@
 UE_NET_TEST_FIXTURE(FTestFilteringFixture, SwitchingFiltersCallsRemoveObjectOnPreviousFilter)
 {
 	// Setup dynamic filter for the test
-	{
-		UMockNetObjectFilter::FFunctionCallSetup CallSetup;
-		CallSetup.AddObject.bReturnValue = true;
-		CallSetup.Filter.bFilterOutByDefault = true;
-		MockNetObjectFilter->SetFunctionCallSetup(CallSetup);
-		MockNetObjectFilter->ResetFunctionCallStatus();
-	}
+	SetDynamicFilterStatus(ENetFilterStatus::Disallow);
 
 	// Spawn object on server and set filter
 	UReplicatedTestObject* ServerObject = Server->CreateObject(0, 0);
@@ -977,13 +948,7 @@
 UE_NET_TEST_FIXTURE(FTestFilteringFixture, SubObjectsAreReplicatedWhenOwnerDynamicFilterAllows)
 {
 	// Setup dynamic filter for the test
-	{
-		UMockNetObjectFilter::FFunctionCallSetup CallSetup;
-		CallSetup.AddObject.bReturnValue = true;
-		CallSetup.Filter.bFilterOutByDefault = false;
-		MockNetObjectFilter->SetFunctionCallSetup(CallSetup);
-		MockNetObjectFilter->ResetFunctionCallStatus();
-	}
+	SetDynamicFilterStatus(ENetFilterStatus::Allow);
 
 	// Add client
 	FReplicationSystemTestClient* Client = CreateClient();
@@ -1007,13 +972,7 @@
 UE_NET_TEST_FIXTURE(FTestFilteringFixture, SubObjectsAreNotReplicatedWhenOwnerDynamicFilterDisallows)
 {
 	// Setup dynamic filter for the test
-	{
-		UMockNetObjectFilter::FFunctionCallSetup CallSetup;
-		CallSetup.AddObject.bReturnValue = true;
-		CallSetup.Filter.bFilterOutByDefault = true;
-		MockNetObjectFilter->SetFunctionCallSetup(CallSetup);
-		MockNetObjectFilter->ResetFunctionCallStatus();
-	}
+	SetDynamicFilterStatus(ENetFilterStatus::Disallow);
 
 	// Add client
 	FReplicationSystemTestClient* Client = CreateClient();
@@ -1037,13 +996,7 @@
 UE_NET_TEST_FIXTURE(FTestFilteringFixture, DependentObjectIsUnaffectedByDynamicFilter)
 {
 	// Setup dynamic filter for the test
-	{
-		UMockNetObjectFilter::FFunctionCallSetup CallSetup;
-		CallSetup.AddObject.bReturnValue = true;
-		CallSetup.Filter.bFilterOutByDefault = true;
-		MockNetObjectFilter->SetFunctionCallSetup(CallSetup);
-		MockNetObjectFilter->ResetFunctionCallStatus();
-	}
+	SetDynamicFilterStatus(ENetFilterStatus::Disallow);
 
 	// Add client
 	FReplicationSystemTestClient* Client = CreateClient();
@@ -1090,13 +1043,7 @@
 UE_NET_TEST_FIXTURE(FTestFilteringFixture, NestedDependentObjectIsFilteredAsParentsOrIndependent)
 {
 	// Setup dynamic filter for the test
-	{
-		UMockNetObjectFilter::FFunctionCallSetup CallSetup;
-		CallSetup.AddObject.bReturnValue = true;
-		CallSetup.Filter.bFilterOutByDefault = true;
-		MockNetObjectFilter->SetFunctionCallSetup(CallSetup);
-		MockNetObjectFilter->ResetFunctionCallStatus();
-	}
+	SetDynamicFilterStatus(ENetFilterStatus::Disallow);
 
 	// Add client
 	FReplicationSystemTestClient* Client = CreateClient();
@@ -1166,13 +1113,7 @@
 UE_NET_TEST_FIXTURE(FTestFilteringFixture, ObjectGetsFilterOutSettingFromStart)
 {
 	// Setup dynamic filter for the test. For this test we want the value of the object's filter property to rule.
-	{
-		UMockNetObjectFilter::FFunctionCallSetup CallSetup;
-		CallSetup.AddObject.bReturnValue = true;
-		CallSetup.Filter.bFilterOutByDefault = false;
-		MockNetObjectFilterWithFragments->SetFunctionCallSetup(CallSetup);
-		MockNetObjectFilterWithFragments->ResetFunctionCallStatus();
-	}
+	SetDynamicFragmentFilterStatus(ENetFilterStatus::Allow);
 
 	// Add client
 	FReplicationSystemTestClient* Client = CreateClient();
@@ -1197,13 +1138,7 @@
 UE_NET_TEST_FIXTURE(FTestFilteringFixture, ObjectGetsUpdatedFilterOutSetting)
 {
 	// Setup dynamic filter for the test. For this test we want the value of the object's filter property to rule.
-	{
-		UMockNetObjectFilter::FFunctionCallSetup CallSetup;
-		CallSetup.AddObject.bReturnValue = true;
-		CallSetup.Filter.bFilterOutByDefault = false;
-		MockNetObjectFilterWithFragments->SetFunctionCallSetup(CallSetup);
-		MockNetObjectFilterWithFragments->ResetFunctionCallStatus();
-	}
+	SetDynamicFragmentFilterStatus(ENetFilterStatus::Allow);
 
 	// Add client
 	FReplicationSystemTestClient* Client = CreateClient();
@@ -1237,18 +1172,9 @@
 
 UE_NET_TEST_FIXTURE(FTestFilteringFixture, MixPreAndPostFilters)
 {
-	// Setup dynamic filter for the test.
-	{
-		UMockNetObjectFilter::FFunctionCallSetup CallSetup;
-		CallSetup.AddObject.bReturnValue = true;
-		CallSetup.Filter.bFilterOutByDefault = true;
+	// Setup dynamic filters for the test.
+	SetDynamicFilterStatus(ENetFilterStatus::Disallow);
+	SetDynamicFragmentFilterStatus(ENetFilterStatus::Disallow);
 
-		MockNetObjectFilter->SetFunctionCallSetup(CallSetup);
-		MockNetObjectFilter->ResetFunctionCallStatus();
-
-		MockNetObjectFilterWithFragments->SetFunctionCallSetup(CallSetup);
-		MockNetObjectFilterWithFragments->ResetFunctionCallStatus();
-	}
-
 	// Add client
 	FReplicationSystemTestClient* Client = CreateClient();
@@ -1281,14 +1207,8 @@
 
 UE_NET_TEST_FIXTURE(FTestFilteringFixture, DynamicFilteredOutSubObjectsAreResetWhenIndexIsReused)
 {
-	// Setup dynamic filter for the test
-	{
-		UMockNetObjectFilter::FFunctionCallSetup CallSetup;
-		CallSetup.AddObject.bReturnValue = true;
-		CallSetup.Filter.bFilterOutByDefault = true;
-		MockNetObjectFilter->SetFunctionCallSetup(CallSetup);
-		MockNetObjectFilter->ResetFunctionCallStatus();
-	}
+	// Setup dynamic filters for the test
+	SetDynamicFilterStatus(ENetFilterStatus::Disallow);
 
 	// Add client
 	FReplicationSystemTestClient* Client = CreateClient();
@@ -1326,15 +1246,8 @@
 
 UE_NET_TEST_FIXTURE(FTestFilteringWithConditionFixture, TestCulledDirtyActors)
 {
-	// Setup dynamic filter for the test.
-	{
-		UMockNetObjectFilter::FFunctionCallSetup CallSetup;
-		CallSetup.AddObject.bReturnValue = true;
-		CallSetup.Filter.bFilterOutByDefault = true;
+	// Setup dynamic filter for the test
+	SetDynamicFilterStatus(ENetFilterStatus::Disallow);
 
-		MockNetObjectFilter->SetFunctionCallSetup(CallSetup);
-		MockNetObjectFilter->ResetFunctionCallStatus();
-	}
-
 	const uint32 RepSystemID = Server->GetReplicationSystemId();
 
@@ -1440,4 +1353,1046 @@
 	}
 }
 
+UE_NET_TEST_FIXTURE(FTestFilteringFixture, InclusionGroupDoesNotFilterOutObject)
+{
+	// Add client
+	FReplicationSystemTestClient* Client = CreateClient();
+
+	// Spawn object on server
+	UReplicatedTestObject* ServerObject = Server->CreateObject(UTestReplicatedIrisObject::FComponents{});
+	
+	// Setup inclusion group filter
+	FNetObjectGroupHandle GroupHandle = Server->ReplicationSystem->CreateGroup();
+	Server->ReplicationSystem->AddToGroup(GroupHandle, ServerObject->NetRefHandle);
+	Server->ReplicationSystem->AddInclusionFilterGroup(GroupHandle);
+	Server->ReplicationSystem->SetGroupFilterStatus(GroupHandle, ENetFilterStatus::Allow);
+
+	// Send and deliver packets
+	Server->UpdateAndSend({Client}, DeliverPacket);
+
+	// Object should have been created
+	UE_NET_ASSERT_NE(Client->GetReplicationBridge()->GetReplicatedObject(ServerObject->NetRefHandle), nullptr);
+
+	// Change filter to not allow replication. As it's an inclusion filter this should not change things.
+	Server->ReplicationSystem->SetGroupFilterStatus(GroupHandle, ENetFilterStatus::Disallow);
+
+	// Send and deliver packets
+	Server->UpdateAndSend({Client}, DeliverPacket);
+
+	// Object should not have been destroyed
+	UE_NET_ASSERT_NE(Client->GetReplicationBridge()->GetReplicatedObject(ServerObject->NetRefHandle), nullptr);
+}
+
+UE_NET_TEST_FIXTURE(FTestFilteringFixture, InclusionGroupDoesNotOverrideOwnerFilter)
+{
+	// Add clients
+	constexpr uint32 OwningClientIndex = 0;
+	constexpr uint32 NonOwningClientIndex = 1;
+
+	FReplicationSystemTestClient* ClientArray[2];
+	for (FReplicationSystemTestClient*& Client : ClientArray)
+	{
+		Client = CreateClient();
+	}
+
+	// Spawn object on server
+	UReplicatedTestObject* ServerObject = Server->CreateObject(UTestReplicatedIrisObject::FComponents{});
+	Server->ReplicationSystem->SetFilter(ServerObject->NetRefHandle, ToOwnerFilterHandle);
+	Server->ReplicationSystem->SetOwningNetConnection(ServerObject->NetRefHandle, ClientArray[OwningClientIndex]->ConnectionIdOnServer);
+
+	// Setup inclusion group filter
+	FNetObjectGroupHandle GroupHandle = Server->ReplicationSystem->CreateGroup();
+	Server->ReplicationSystem->AddToGroup(GroupHandle, ServerObject->NetRefHandle);
+	Server->ReplicationSystem->AddInclusionFilterGroup(GroupHandle);
+	Server->ReplicationSystem->SetGroupFilterStatus(GroupHandle, ENetFilterStatus::Allow);
+
+	// Send and deliver packet
+	Server->UpdateAndSend(ClientArray, DeliverPacket);
+
+	// Object should have been created on the owning client
+	UE_NET_ASSERT_NE(Clients[OwningClientIndex]->GetReplicationBridge()->GetReplicatedObject(ServerObject->NetRefHandle), nullptr);
+	// Object should not have been created on the non-owning client
+	UE_NET_ASSERT_EQ(Clients[NonOwningClientIndex]->GetReplicationBridge()->GetReplicatedObject(ServerObject->NetRefHandle), nullptr);
+
+	// Change filter to not allow replication. As it's an inclusion filter this should not change things.
+	Server->ReplicationSystem->SetGroupFilterStatus(GroupHandle, ENetFilterStatus::Disallow);
+
+	// Send and deliver packets
+	Server->UpdateAndSend(ClientArray, DeliverPacket);
+
+	// Object creation status should remain the same as before
+	UE_NET_ASSERT_NE(Clients[OwningClientIndex]->GetReplicationBridge()->GetReplicatedObject(ServerObject->NetRefHandle), nullptr);
+	UE_NET_ASSERT_EQ(Clients[NonOwningClientIndex]->GetReplicationBridge()->GetReplicatedObject(ServerObject->NetRefHandle), nullptr);
+}
+
+UE_NET_TEST_FIXTURE(FTestFilteringFixture, InclusionGroupDoesNotOverrideConnectionFilter)
+{
+	// Add clients
+	constexpr uint32 AllowedClientIndex = 0;
+	constexpr uint32 DisallowedOwningClientIndex = 1;
+
+	FReplicationSystemTestClient* ClientArray[2];
+	for (FReplicationSystemTestClient*& Client : ClientArray)
+	{
+		Client = CreateClient();
+	}
+
+	// Spawn object on server
+	UReplicatedTestObject* ServerObject = Server->CreateObject(UTestReplicatedIrisObject::FComponents{});
+
+	// Apply connection filter
+	TBitArray<> AllowedConnections;
+	AllowedConnections.Init(false, ClientArray[AllowedClientIndex]->ConnectionIdOnServer + 1);
+	AllowedConnections[ClientArray[AllowedClientIndex]->ConnectionIdOnServer] = true;
+	Server->ReplicationSystem->SetConnectionFilter(ServerObject->NetRefHandle, AllowedConnections, ENetFilterStatus::Allow);
+
+	// Setup inclusion group filter
+	FNetObjectGroupHandle GroupHandle = Server->ReplicationSystem->CreateGroup();
+	Server->ReplicationSystem->AddToGroup(GroupHandle, ServerObject->NetRefHandle);
+	Server->ReplicationSystem->AddInclusionFilterGroup(GroupHandle);
+	Server->ReplicationSystem->SetGroupFilterStatus(GroupHandle, ENetFilterStatus::Allow);
+
+	// Send and deliver packet
+	Server->UpdateAndSend(ClientArray, DeliverPacket);
+
+	// Object should have been created on the owning client
+	UE_NET_ASSERT_NE(Clients[AllowedClientIndex]->GetReplicationBridge()->GetReplicatedObject(ServerObject->NetRefHandle), nullptr);
+	// Object should not have been created on the non-owning client
+	UE_NET_ASSERT_EQ(Clients[DisallowedOwningClientIndex]->GetReplicationBridge()->GetReplicatedObject(ServerObject->NetRefHandle), nullptr);
+
+	// Change filter to not allow replication. As it's an inclusion filter this should not change things.
+	Server->ReplicationSystem->SetGroupFilterStatus(GroupHandle, ENetFilterStatus::Disallow);
+
+	// Send and deliver packets
+	Server->UpdateAndSend(ClientArray, DeliverPacket);
+
+	// Object creation status should remain the same as before
+	UE_NET_ASSERT_NE(Clients[AllowedClientIndex]->GetReplicationBridge()->GetReplicatedObject(ServerObject->NetRefHandle), nullptr);
+	UE_NET_ASSERT_EQ(Clients[DisallowedOwningClientIndex]->GetReplicationBridge()->GetReplicatedObject(ServerObject->NetRefHandle), nullptr);
+}
+
+UE_NET_TEST_FIXTURE(FTestFilteringFixture, InclusionGroupDoesNotOverrideExclusionGroupFilter)
+{
+	// Add clients
+	constexpr uint32 AllowedClientIndex = 0;
+	constexpr uint32 DisallowedOwningClientIndex = 1;
+
+	FReplicationSystemTestClient* ClientArray[2];
+	for (FReplicationSystemTestClient*& Client : ClientArray)
+	{
+		Client = CreateClient();
+	}
+
+	// Spawn object on server
+	UReplicatedTestObject* ServerObject = Server->CreateObject(UTestReplicatedIrisObject::FComponents{});
+
+	// Setup exclusion group filter
+	FNetObjectGroupHandle ExclusionGroupHandle = Server->ReplicationSystem->CreateGroup();
+	Server->ReplicationSystem->AddToGroup(ExclusionGroupHandle, ServerObject->NetRefHandle);
+	Server->ReplicationSystem->AddExclusionFilterGroup(ExclusionGroupHandle);
+	Server->ReplicationSystem->SetGroupFilterStatus(ExclusionGroupHandle, ClientArray[AllowedClientIndex]->ConnectionIdOnServer, ENetFilterStatus::Allow);
+
+	// Setup inclusion group filter
+	FNetObjectGroupHandle InclusionGroupHandle = Server->ReplicationSystem->CreateGroup();
+	Server->ReplicationSystem->AddToGroup(InclusionGroupHandle, ServerObject->NetRefHandle);
+	Server->ReplicationSystem->AddInclusionFilterGroup(InclusionGroupHandle);
+	Server->ReplicationSystem->SetGroupFilterStatus(InclusionGroupHandle, ENetFilterStatus::Allow);
+
+	// Send and deliver packet
+	Server->PreSendUpdate();
+	for (FReplicationSystemTestClient* Client : ClientArray)
+	{
+		Server->SendAndDeliverTo(Client, DeliverPacket);
+	}
+	Server->PostSendUpdate();
+
+	// Object should have been created on the allowed client
+	UE_NET_ASSERT_NE(Clients[AllowedClientIndex]->GetReplicationBridge()->GetReplicatedObject(ServerObject->NetRefHandle), nullptr);
+	// Object should not have been created on the disallowed client
+	UE_NET_ASSERT_EQ(Clients[DisallowedOwningClientIndex]->GetReplicationBridge()->GetReplicatedObject(ServerObject->NetRefHandle), nullptr);
+
+	// Change inclusion filter to not allow replication. As it's an inclusion filter this should not change things.
+	Server->ReplicationSystem->SetGroupFilterStatus(InclusionGroupHandle, ENetFilterStatus::Disallow);
+
+	// Send and deliver packets
+	Server->UpdateAndSend(ClientArray, DeliverPacket);
+
+	// Object creation status should remain the same as before
+	UE_NET_ASSERT_NE(Clients[AllowedClientIndex]->GetReplicationBridge()->GetReplicatedObject(ServerObject->NetRefHandle), nullptr);
+	UE_NET_ASSERT_EQ(Clients[DisallowedOwningClientIndex]->GetReplicationBridge()->GetReplicatedObject(ServerObject->NetRefHandle), nullptr);
+}
+
+UE_NET_TEST_FIXTURE(FTestFilteringFixture, DisabledInclusionGroupDoesNotOverrideDynamicFilter)
+{
+	// Setup dynamic filter for the test
+	SetDynamicFilterStatus(ENetFilterStatus::Disallow);
+
+	// Add clients
+	FReplicationSystemTestClient* ClientArray[2];
+	for (FReplicationSystemTestClient*& Client : ClientArray)
+	{
+		Client = CreateClient();
+	}
+
+	// Spawn object on server
+	UReplicatedTestObject* ServerObject = Server->CreateObject(UTestReplicatedIrisObject::FComponents{});
+
+	// Set filter which will filter out all objects
+	Server->ReplicationSystem->SetFilter(ServerObject->NetRefHandle, MockFilterHandle);
+
+	// Setup inclusion group filter
+	FNetObjectGroupHandle InclusionGroupHandle = Server->ReplicationSystem->CreateGroup();
+	Server->ReplicationSystem->AddToGroup(InclusionGroupHandle, ServerObject->NetRefHandle);
+	Server->ReplicationSystem->AddInclusionFilterGroup(InclusionGroupHandle);
+	Server->ReplicationSystem->SetGroupFilterStatus(InclusionGroupHandle, ENetFilterStatus::Disallow);
+
+	// Send and deliver packet
+	Server->UpdateAndSend(ClientArray, DeliverPacket);
+
+	// Object should not have been created on any clients
+	for (FReplicationSystemTestClient* Client : ClientArray)
+	{
+		UE_NET_ASSERT_EQ(Client->GetReplicationBridge()->GetReplicatedObject(ServerObject->NetRefHandle), nullptr);
+	}
+}
+
+UE_NET_TEST_FIXTURE(FTestFilteringFixture, EnabledInclusionGroupDoesOverrideDynamicFilter)
+{
+	// Setup dynamic filter for the test
+	SetDynamicFilterStatus(ENetFilterStatus::Disallow);
+
+	// Add clients
+	constexpr uint32 AllowedClientIndex = 1;
+	constexpr uint32 DisallowedOwningClientIndex = 0;
+
+	FReplicationSystemTestClient* ClientArray[2];
+	for (FReplicationSystemTestClient*& Client : ClientArray)
+	{
+		Client = CreateClient();
+	}
+
+	// Spawn object on server
+	UReplicatedTestObject* ServerObject = Server->CreateObject(UTestReplicatedIrisObject::FComponents{});
+
+	// Set filter which will filter out all objects
+	Server->ReplicationSystem->SetFilter(ServerObject->NetRefHandle, MockFilterHandle);
+
+	// Setup inclusion group filter
+	FNetObjectGroupHandle InclusionGroupHandle = Server->ReplicationSystem->CreateGroup();
+	Server->ReplicationSystem->AddToGroup(InclusionGroupHandle, ServerObject->NetRefHandle);
+	Server->ReplicationSystem->AddInclusionFilterGroup(InclusionGroupHandle);
+	// Disallow by default
+	Server->ReplicationSystem->SetGroupFilterStatus(InclusionGroupHandle, ENetFilterStatus::Disallow);
+	// Allow one client
+	Server->ReplicationSystem->SetGroupFilterStatus(InclusionGroupHandle, ClientArray[AllowedClientIndex]->ConnectionIdOnServer, ENetFilterStatus::Allow);
+
+	// Send and deliver packet
+	Server->UpdateAndSend(ClientArray, DeliverPacket);
+
+	// Object should have been created on the allowed client
+	UE_NET_ASSERT_NE(Clients[AllowedClientIndex]->GetReplicationBridge()->GetReplicatedObject(ServerObject->NetRefHandle), nullptr);
+	// Object should not have been created on the disallowed client
+	UE_NET_ASSERT_EQ(Clients[DisallowedOwningClientIndex]->GetReplicationBridge()->GetReplicatedObject(ServerObject->NetRefHandle), nullptr);
+}
+
+UE_NET_TEST_FIXTURE(FTestFilteringFixture, LateAddingToEnabledInclusionGroupDoesOverrideDynamicFilter)
+{
+	// Setup dynamic filter for the test
+	SetDynamicFilterStatus(ENetFilterStatus::Disallow);
+
+	// Add clients
+	constexpr uint32 AllowedClientIndex = 1;
+	constexpr uint32 DisallowedOwningClientIndex = 0;
+
+	FReplicationSystemTestClient* ClientArray[2];
+	for (FReplicationSystemTestClient*& Client : ClientArray)
+	{
+		Client = CreateClient();
+	}
+
+	// Spawn object on server
+	UReplicatedTestObject* ServerObject = Server->CreateObject(UTestReplicatedIrisObject::FComponents{});
+
+	// Set filter which will filter out all objects
+	Server->ReplicationSystem->SetFilter(ServerObject->NetRefHandle, MockFilterHandle);
+
+	// Setup inclusion group filter
+	FNetObjectGroupHandle InclusionGroupHandle = Server->ReplicationSystem->CreateGroup();
+	Server->ReplicationSystem->AddInclusionFilterGroup(InclusionGroupHandle);
+	// Disallow by default
+	Server->ReplicationSystem->SetGroupFilterStatus(InclusionGroupHandle, ENetFilterStatus::Disallow);
+	// Allow one client
+	Server->ReplicationSystem->SetGroupFilterStatus(InclusionGroupHandle, ClientArray[AllowedClientIndex]->ConnectionIdOnServer, ENetFilterStatus::Allow);
+
+	// Send and deliver packet
+	Server->UpdateAndSend(ClientArray, DeliverPacket);
+
+	// At this point the object should not have been created on any client
+	for (FReplicationSystemTestClient* Client : ClientArray)
+	{
+		UE_NET_ASSERT_EQ(Client->GetReplicationBridge()->GetReplicatedObject(ServerObject->NetRefHandle), nullptr);
+	}
+
+	// Late adding to group
+	Server->ReplicationSystem->AddToGroup(InclusionGroupHandle, ServerObject->NetRefHandle);
+
+	// Send and deliver packet
+	Server->UpdateAndSend(ClientArray, DeliverPacket);
+
+	// Object should have been created on the allowed client
+	UE_NET_ASSERT_NE(Clients[AllowedClientIndex]->GetReplicationBridge()->GetReplicatedObject(ServerObject->NetRefHandle), nullptr);
+	// Object should not have been created on the disallowed client
+	UE_NET_ASSERT_EQ(Clients[DisallowedOwningClientIndex]->GetReplicationBridge()->GetReplicatedObject(ServerObject->NetRefHandle), nullptr);
+}
+
+UE_NET_TEST_FIXTURE(FTestFilteringFixture, LateEnablingInclusionGroupDoesOverrideDynamicFilter)
+{
+	// Setup dynamic filter for the test
+	SetDynamicFilterStatus(ENetFilterStatus::Disallow);
+
+	// Add clients
+	constexpr uint32 AllowedClientIndex = 1;
+	constexpr uint32 DisallowedOwningClientIndex = 0;
+
+	FReplicationSystemTestClient* ClientArray[2];
+	for (FReplicationSystemTestClient*& Client : ClientArray)
+	{
+		Client = CreateClient();
+	}
+
+	// Spawn object on server
+	UReplicatedTestObject* ServerObject = Server->CreateObject(UTestReplicatedIrisObject::FComponents{});
+
+	// Set filter which will filter out all objects
+	Server->ReplicationSystem->SetFilter(ServerObject->NetRefHandle, MockFilterHandle);
+
+	// Setup inclusion group filter
+	FNetObjectGroupHandle InclusionGroupHandle = Server->ReplicationSystem->CreateGroup();
+	Server->ReplicationSystem->AddToGroup(InclusionGroupHandle, ServerObject->NetRefHandle);
+	Server->ReplicationSystem->AddInclusionFilterGroup(InclusionGroupHandle);
+	// Disallow by default
+	Server->ReplicationSystem->SetGroupFilterStatus(InclusionGroupHandle, ENetFilterStatus::Disallow);
+
+	// Send and deliver packet
+	Server->UpdateAndSend(ClientArray, DeliverPacket);
+
+	// At this point the object should not have been created on any client
+	for (FReplicationSystemTestClient* Client : ClientArray)
+	{
+		UE_NET_ASSERT_EQ(Client->GetReplicationBridge()->GetReplicatedObject(ServerObject->NetRefHandle), nullptr);
+	}
+
+	// Late enabling client
+	Server->ReplicationSystem->SetGroupFilterStatus(InclusionGroupHandle, ClientArray[AllowedClientIndex]->ConnectionIdOnServer, ENetFilterStatus::Allow);
+
+	// Send and deliver packet
+	Server->UpdateAndSend(ClientArray, DeliverPacket);
+
+	// Object should have been created on the allowed client
+	UE_NET_ASSERT_NE(Clients[AllowedClientIndex]->GetReplicationBridge()->GetReplicatedObject(ServerObject->NetRefHandle), nullptr);
+	// Object should not have been created on the disallowed client
+	UE_NET_ASSERT_EQ(Clients[DisallowedOwningClientIndex]->GetReplicationBridge()->GetReplicatedObject(ServerObject->NetRefHandle), nullptr);
+}
+
+UE_NET_TEST_FIXTURE(FTestFilteringFixture, RemovingFromInclusionGroupRemovesDynamicFilterOverride)
+{
+	// Setup dynamic filter for the test
+	SetDynamicFilterStatus(ENetFilterStatus::Disallow);
+
+	// Add client
+	FReplicationSystemTestClient* Client = CreateClient();
+
+	// Spawn object on server
+	UReplicatedTestObject* ServerObject = Server->CreateObject(UTestReplicatedIrisObject::FComponents{});
+
+	// Set filter which will filter out all objects
+	Server->ReplicationSystem->SetFilter(ServerObject->NetRefHandle, MockFilterHandle);
+
+	// Setup inclusion group filter
+	FNetObjectGroupHandle InclusionGroupHandle = Server->ReplicationSystem->CreateGroup();
+	Server->ReplicationSystem->AddToGroup(InclusionGroupHandle, ServerObject->NetRefHandle);
+	Server->ReplicationSystem->AddInclusionFilterGroup(InclusionGroupHandle);
+	Server->ReplicationSystem->SetGroupFilterStatus(InclusionGroupHandle, ENetFilterStatus::Allow);
+
+	// Send and deliver packet
+	Server->UpdateAndSend({Client}, DeliverPacket);
+
+	// Object should have been created
+	UE_NET_ASSERT_NE(Client->GetReplicationBridge()->GetReplicatedObject(ServerObject->NetRefHandle), nullptr);
+
+	// Remove object from inclusion group. This should cause the object to be filtered out again.
+	Server->ReplicationSystem->RemoveFromGroup(InclusionGroupHandle, ServerObject->NetRefHandle);
+
+	// Send and deliver packet
+	Server->UpdateAndSend({Client}, DeliverPacket);
+
+	// Object should have been created
+	UE_NET_ASSERT_EQ(Client->GetReplicationBridge()->GetReplicatedObject(ServerObject->NetRefHandle), nullptr);
+}
+
+UE_NET_TEST_FIXTURE(FTestFilteringFixture, SubObjectAddedToAllowedInclusionGroupFollowsOwnerNotInInclusionGroup)
+{
+	// Setup dynamic filter for the test
+	SetDynamicFilterStatus(ENetFilterStatus::Disallow);
+
+	// Add client
+	FReplicationSystemTestClient* Client = CreateClient();
+
+	// Spawn object with subobject on server
+	UReplicatedTestObject* ServerObject = Server->CreateObject(UTestReplicatedIrisObject::FComponents{});
+	UReplicatedTestObject* ServerSubObject = Server->CreateSubObject(ServerObject->NetRefHandle, UTestReplicatedIrisObject::FComponents{});
+
+	// Set filter which will filter out all objects
+	Server->ReplicationSystem->SetFilter(ServerObject->NetRefHandle, MockFilterHandle);
+
+	// Setup inclusion group filter
+	FNetObjectGroupHandle InclusionGroupHandle = Server->ReplicationSystem->CreateGroup();
+	Server->ReplicationSystem->AddToGroup(InclusionGroupHandle, ServerSubObject->NetRefHandle);
+	Server->ReplicationSystem->AddInclusionFilterGroup(InclusionGroupHandle);
+	Server->ReplicationSystem->SetGroupFilterStatus(InclusionGroupHandle, Client->ConnectionIdOnServer, ENetFilterStatus::Allow);
+
+	// Send and deliver packet
+	Server->UpdateAndSend({Client}, DeliverPacket);
+
+	// Neither object nor subobject should have been created on the client.
+	UE_NET_ASSERT_EQ(Client->GetReplicationBridge()->GetReplicatedObject(ServerObject->NetRefHandle), nullptr);
+	UE_NET_ASSERT_EQ(Client->GetReplicationBridge()->GetReplicatedObject(ServerSubObject->NetRefHandle), nullptr);
+}
+
+UE_NET_TEST_FIXTURE(FTestFilteringFixture, SubObjectLateAddedToAllowedInclusionGroupFollowsOwnerNotInInclusionGroup)
+{
+	// Setup dynamic filter for the test
+	SetDynamicFilterStatus(ENetFilterStatus::Disallow);
+
+	// Add client
+	FReplicationSystemTestClient* Client = CreateClient();
+
+	// Spawn object with subobject on server
+	UReplicatedTestObject* ServerObject = Server->CreateObject(UTestReplicatedIrisObject::FComponents{});
+	UReplicatedTestObject* ServerSubObject = Server->CreateSubObject(ServerObject->NetRefHandle, UTestReplicatedIrisObject::FComponents{});
+
+	// Set filter which will filter out all objects
+	Server->ReplicationSystem->SetFilter(ServerObject->NetRefHandle, MockFilterHandle);
+
+	// Setup inclusion group filter
+	FNetObjectGroupHandle InclusionGroupHandle = Server->ReplicationSystem->CreateGroup();
+	Server->ReplicationSystem->AddInclusionFilterGroup(InclusionGroupHandle);
+	Server->ReplicationSystem->SetGroupFilterStatus(InclusionGroupHandle, Client->ConnectionIdOnServer, ENetFilterStatus::Allow);
+
+	// Send and deliver packet
+	Server->UpdateAndSend({Client}, DeliverPacket);
+
+	// Late add subobject to group
+	Server->ReplicationSystem->AddToGroup(InclusionGroupHandle, ServerSubObject->NetRefHandle);
+
+	// Send and deliver packet
+	Server->UpdateAndSend({Client}, DeliverPacket);
+
+	// Neither object nor subobject should have been created on the client.
+	UE_NET_ASSERT_EQ(Client->GetReplicationBridge()->GetReplicatedObject(ServerObject->NetRefHandle), nullptr);
+	UE_NET_ASSERT_EQ(Client->GetReplicationBridge()->GetReplicatedObject(ServerSubObject->NetRefHandle), nullptr);
+}
+
+UE_NET_TEST_FIXTURE(FTestFilteringFixture, LateAddedSubObjectFollowsOwnerInAllowedInclusionGroup)
+{
+	// Setup dynamic filter for the test
+	SetDynamicFilterStatus(ENetFilterStatus::Disallow);
+
+	// Add client
+	FReplicationSystemTestClient* Client = CreateClient();
+
+	// Spawn object on server
+	UReplicatedTestObject* ServerObject = Server->CreateObject(UTestReplicatedIrisObject::FComponents{});
+
+	// Set filter which will filter out all objects
+	Server->ReplicationSystem->SetFilter(ServerObject->NetRefHandle, MockFilterHandle);
+
+	// Setup inclusion group filter and add object
+	FNetObjectGroupHandle InclusionGroupHandle = Server->ReplicationSystem->CreateGroup();
+	Server->ReplicationSystem->AddInclusionFilterGroup(InclusionGroupHandle);
+	Server->ReplicationSystem->AddToGroup(InclusionGroupHandle, ServerObject->NetRefHandle);
+	Server->ReplicationSystem->SetGroupFilterStatus(InclusionGroupHandle, Client->ConnectionIdOnServer, ENetFilterStatus::Allow);
+
+	// Send and deliver packet
+	Server->UpdateAndSend({Client}, DeliverPacket);
+
+	// Object should now have been created on the client.
+	UE_NET_ASSERT_NE(Client->GetReplicationBridge()->GetReplicatedObject(ServerObject->NetRefHandle), nullptr);
+
+	// Late add subobject to owner
+	UReplicatedTestObject* ServerSubObject = Server->CreateSubObject(ServerObject->NetRefHandle, UTestReplicatedIrisObject::FComponents{});
+
+	// Send and deliver packet
+	Server->UpdateAndSend({Client}, DeliverPacket);
+
+	// Subobject should now have been created on the client.
+	UE_NET_ASSERT_NE(Client->GetReplicationBridge()->GetReplicatedObject(ServerSubObject->NetRefHandle), nullptr);
+}
+
+UE_NET_TEST_FIXTURE(FTestFilteringFixture, LateAddedSubObjectFollowsOwnerInDisallowedInclusionGroup)
+{
+	// Setup dynamic filter for the test
+	SetDynamicFilterStatus(ENetFilterStatus::Disallow);
+
+	// Add client
+	FReplicationSystemTestClient* Client = CreateClient();
+
+	// Spawn object on server
+	UReplicatedTestObject* ServerObject = Server->CreateObject(UTestReplicatedIrisObject::FComponents{});
+
+	// Set filter which will filter out all objects
+	Server->ReplicationSystem->SetFilter(ServerObject->NetRefHandle, MockFilterHandle);
+
+	// Setup inclusion group filter and add object
+	FNetObjectGroupHandle InclusionGroupHandle = Server->ReplicationSystem->CreateGroup();
+	Server->ReplicationSystem->AddInclusionFilterGroup(InclusionGroupHandle);
+	Server->ReplicationSystem->AddToGroup(InclusionGroupHandle, ServerObject->NetRefHandle);
+	Server->ReplicationSystem->SetGroupFilterStatus(InclusionGroupHandle, Client->ConnectionIdOnServer, ENetFilterStatus::Disallow);
+
+	// Send and deliver packet
+	Server->UpdateAndSend({Client}, DeliverPacket);
+
+	// Object should not have been created on the client.
+	UE_NET_ASSERT_EQ(Client->GetReplicationBridge()->GetReplicatedObject(ServerObject->NetRefHandle), nullptr);
+
+	// Late add subobject to owner
+	UReplicatedTestObject* ServerSubObject = Server->CreateSubObject(ServerObject->NetRefHandle, UTestReplicatedIrisObject::FComponents{});
+
+	// Send and deliver packet
+	Server->UpdateAndSend({Client}, DeliverPacket);
+
+	// Subobject should not have been created on the client.
+	UE_NET_ASSERT_EQ(Client->GetReplicationBridge()->GetReplicatedObject(ServerSubObject->NetRefHandle), nullptr);
+}
+
+UE_NET_TEST_FIXTURE(FTestFilteringFixture, SubObjectAddedToDisallowedInclusionGroupFollowsOwnerNotInInclusionGroup)
+{
+	// Setup dynamic filter for the test
+	SetDynamicFilterStatus(ENetFilterStatus::Allow);
+
+	// Add client
+	FReplicationSystemTestClient* Client = CreateClient();
+
+	// Spawn object with subobject on server
+	UReplicatedTestObject* ServerObject = Server->CreateObject(UTestReplicatedIrisObject::FComponents{});
+	UReplicatedTestObject* ServerSubObject = Server->CreateSubObject(ServerObject->NetRefHandle, UTestReplicatedIrisObject::FComponents{});
+
+	// Set filter which will filter out all objects
+	Server->ReplicationSystem->SetFilter(ServerObject->NetRefHandle, MockFilterHandle);
+
+	// Setup inclusion group filter
+	FNetObjectGroupHandle InclusionGroupHandle = Server->ReplicationSystem->CreateGroup();
+	Server->ReplicationSystem->AddToGroup(InclusionGroupHandle, ServerSubObject->NetRefHandle);
+	Server->ReplicationSystem->AddInclusionFilterGroup(InclusionGroupHandle);
+	Server->ReplicationSystem->SetGroupFilterStatus(InclusionGroupHandle, Client->ConnectionIdOnServer, ENetFilterStatus::Disallow);
+
+	// Send and deliver packet
+	Server->UpdateAndSend({Client}, DeliverPacket);
+
+	// Both object and subobject should have been created. Disallowing replication of members in an inclusion group does not filter out, not objects nor subobjects.
+	UE_NET_ASSERT_NE(Client->GetReplicationBridge()->GetReplicatedObject(ServerObject->NetRefHandle), nullptr);
+	UE_NET_ASSERT_NE(Client->GetReplicationBridge()->GetReplicatedObject(ServerSubObject->NetRefHandle), nullptr);
+}
+
+UE_NET_TEST_FIXTURE(FTestFilteringFixture, SubObjectRemovedFromAllowedInclusionGroupFollowsOwnerNotInInclusionGroup)
+{
+	// Setup dynamic filter for the test
+	SetDynamicFilterStatus(ENetFilterStatus::Disallow);
+
+	// Add client
+	FReplicationSystemTestClient* Client = CreateClient();
+
+	// Spawn object with subobject on server
+	UReplicatedTestObject* ServerObject = Server->CreateObject(UTestReplicatedIrisObject::FComponents{});
+	UReplicatedTestObject* ServerSubObject = Server->CreateSubObject(ServerObject->NetRefHandle, UTestReplicatedIrisObject::FComponents{});
+
+	// Set filter which will filter out all objects
+	Server->ReplicationSystem->SetFilter(ServerObject->NetRefHandle, MockFilterHandle);
+
+	// Setup inclusion group filter
+	FNetObjectGroupHandle InclusionGroupHandle = Server->ReplicationSystem->CreateGroup();
+	Server->ReplicationSystem->AddToGroup(InclusionGroupHandle, ServerSubObject->NetRefHandle);
+	Server->ReplicationSystem->AddInclusionFilterGroup(InclusionGroupHandle);
+	Server->ReplicationSystem->SetGroupFilterStatus(InclusionGroupHandle, Client->ConnectionIdOnServer, ENetFilterStatus::Allow);
+
+	// Send and deliver packet
+	Server->UpdateAndSend({Client}, DeliverPacket);
+
+	// Remove subobject from group
+	Server->ReplicationSystem->RemoveFromGroup(InclusionGroupHandle, ServerSubObject->NetRefHandle);
+ 
+	// Send and deliver packet
+	Server->UpdateAndSend({Client}, DeliverPacket);
+
+	// Neither object nor subobject should have been created on the client.
+	UE_NET_ASSERT_EQ(Client->GetReplicationBridge()->GetReplicatedObject(ServerObject->NetRefHandle), nullptr);
+	UE_NET_ASSERT_EQ(Client->GetReplicationBridge()->GetReplicatedObject(ServerSubObject->NetRefHandle), nullptr);
+}
+
+UE_NET_TEST_FIXTURE(FTestFilteringFixture, SubObjectRemovedFromDisallowedInclusionGroupFollowsOwnerNotInInclusionGroup)
+{
+	// Setup dynamic filter for the test
+	SetDynamicFilterStatus(ENetFilterStatus::Allow);
+
+	// Add client
+	FReplicationSystemTestClient* Client = CreateClient();
+
+	// Spawn object with subobject on server
+	UReplicatedTestObject* ServerObject = Server->CreateObject(UTestReplicatedIrisObject::FComponents{});
+	UReplicatedTestObject* ServerSubObject = Server->CreateSubObject(ServerObject->NetRefHandle, UTestReplicatedIrisObject::FComponents{});
+
+	// Set filter which will filter out all objects
+	Server->ReplicationSystem->SetFilter(ServerObject->NetRefHandle, MockFilterHandle);
+
+	// Setup inclusion group filter
+	FNetObjectGroupHandle InclusionGroupHandle = Server->ReplicationSystem->CreateGroup();
+	Server->ReplicationSystem->AddToGroup(InclusionGroupHandle, ServerSubObject->NetRefHandle);
+	Server->ReplicationSystem->AddInclusionFilterGroup(InclusionGroupHandle);
+	Server->ReplicationSystem->SetGroupFilterStatus(InclusionGroupHandle, Client->ConnectionIdOnServer, ENetFilterStatus::Disallow);
+
+	// Send and deliver packet
+	Server->UpdateAndSend({Client}, DeliverPacket);
+
+	// Remove subobject from group
+	Server->ReplicationSystem->RemoveFromGroup(InclusionGroupHandle, ServerSubObject->NetRefHandle);
+ 
+	// Send and deliver packet
+	Server->UpdateAndSend({Client}, DeliverPacket);
+
+	// Both object and subobject should have been created.
+	UE_NET_ASSERT_NE(Client->GetReplicationBridge()->GetReplicatedObject(ServerObject->NetRefHandle), nullptr);
+	UE_NET_ASSERT_NE(Client->GetReplicationBridge()->GetReplicatedObject(ServerSubObject->NetRefHandle), nullptr);
+}
+
+UE_NET_TEST_FIXTURE(FTestFilteringFixture, SubObjectAddedToInclusionGroupFollowsOwnerInOtherInclusionGroupp)
+{
+	// Setup dynamic filter for the test
+	SetDynamicFilterStatus(ENetFilterStatus::Disallow);
+
+	// Add client
+	FReplicationSystemTestClient* Client = CreateClient();
+
+	// Spawn object with subobject on server
+	UReplicatedTestObject* ServerObject = Server->CreateObject(UTestReplicatedIrisObject::FComponents{});
+	UReplicatedTestObject* ServerSubObject = Server->CreateSubObject(ServerObject->NetRefHandle, UTestReplicatedIrisObject::FComponents{});
+
+	// Set filter which will filter out all objects
+	Server->ReplicationSystem->SetFilter(ServerObject->NetRefHandle, MockFilterHandle);
+
+	// Setup separate inclusion group filters for object and subobject
+	FNetObjectGroupHandle ObjectInclusionGroupHandle = Server->ReplicationSystem->CreateGroup();
+	Server->ReplicationSystem->AddToGroup(ObjectInclusionGroupHandle, ServerObject->NetRefHandle);
+	Server->ReplicationSystem->AddInclusionFilterGroup(ObjectInclusionGroupHandle);
+	Server->ReplicationSystem->SetGroupFilterStatus(ObjectInclusionGroupHandle, Client->ConnectionIdOnServer, ENetFilterStatus::Disallow);
+	
+	FNetObjectGroupHandle SubObjectInclusionGroupHandle = Server->ReplicationSystem->CreateGroup();
+	Server->ReplicationSystem->AddToGroup(SubObjectInclusionGroupHandle, ServerSubObject->NetRefHandle);
+	Server->ReplicationSystem->AddInclusionFilterGroup(SubObjectInclusionGroupHandle);
+	Server->ReplicationSystem->SetGroupFilterStatus(SubObjectInclusionGroupHandle, Client->ConnectionIdOnServer, ENetFilterStatus::Disallow);
+
+	// Send and deliver packet
+	Server->UpdateAndSend({Client}, DeliverPacket);
+
+	// Neither object nor subobject should have been created on the client.
+	UE_NET_ASSERT_EQ(Client->GetReplicationBridge()->GetReplicatedObject(ServerObject->NetRefHandle), nullptr);
+	UE_NET_ASSERT_EQ(Client->GetReplicationBridge()->GetReplicatedObject(ServerSubObject->NetRefHandle), nullptr);
+
+	// Allow subobject group to be replicated. This should not change anything.
+	Server->ReplicationSystem->SetGroupFilterStatus(SubObjectInclusionGroupHandle, Client->ConnectionIdOnServer, ENetFilterStatus::Allow);
+
+	// Send and deliver packet
+	Server->UpdateAndSend({Client}, DeliverPacket);
+
+	// Neither object nor subobject should have been created on the client.
+	UE_NET_ASSERT_EQ(Client->GetReplicationBridge()->GetReplicatedObject(ServerObject->NetRefHandle), nullptr);
+	UE_NET_ASSERT_EQ(Client->GetReplicationBridge()->GetReplicatedObject(ServerSubObject->NetRefHandle), nullptr);
+
+	// Allow object group to be replicated and subobject group to not be replicated. This should result in both the object and subobject being created on the client.
+	Server->ReplicationSystem->SetGroupFilterStatus(ObjectInclusionGroupHandle, Client->ConnectionIdOnServer, ENetFilterStatus::Allow);
+	Server->ReplicationSystem->SetGroupFilterStatus(SubObjectInclusionGroupHandle, Client->ConnectionIdOnServer, ENetFilterStatus::Disallow);
+
+	// Send and deliver packet
+	Server->UpdateAndSend({Client}, DeliverPacket);
+
+	// Both object and subobject should have been created on the client.
+	UE_NET_ASSERT_NE(Client->GetReplicationBridge()->GetReplicatedObject(ServerObject->NetRefHandle), nullptr);
+	UE_NET_ASSERT_NE(Client->GetReplicationBridge()->GetReplicatedObject(ServerSubObject->NetRefHandle), nullptr);
+}
+
+UE_NET_TEST_FIXTURE(FTestFilteringFixture, SubObjectLateAddedToInclusionGroupFollowsOwnerInOtherInclusionGroupp)
+{
+	// Setup dynamic filter for the test
+	SetDynamicFilterStatus(ENetFilterStatus::Disallow);
+
+	// Add client
+	FReplicationSystemTestClient* Client = CreateClient();
+
+	// Spawn object with subobject on server
+	UReplicatedTestObject* ServerObject = Server->CreateObject(UTestReplicatedIrisObject::FComponents{});
+	UReplicatedTestObject* ServerSubObject = Server->CreateSubObject(ServerObject->NetRefHandle, UTestReplicatedIrisObject::FComponents{});
+
+	// Set filter which will filter out all objects
+	Server->ReplicationSystem->SetFilter(ServerObject->NetRefHandle, MockFilterHandle);
+
+	// Setup separate inclusion group filters for object and subobject
+	FNetObjectGroupHandle ObjectInclusionGroupHandle = Server->ReplicationSystem->CreateGroup();
+	Server->ReplicationSystem->AddToGroup(ObjectInclusionGroupHandle, ServerObject->NetRefHandle);
+	Server->ReplicationSystem->AddInclusionFilterGroup(ObjectInclusionGroupHandle);
+	Server->ReplicationSystem->SetGroupFilterStatus(ObjectInclusionGroupHandle, Client->ConnectionIdOnServer, ENetFilterStatus::Disallow);
+	
+	FNetObjectGroupHandle SubObjectInclusionGroupHandle = Server->ReplicationSystem->CreateGroup();
+	Server->ReplicationSystem->AddInclusionFilterGroup(SubObjectInclusionGroupHandle);
+	Server->ReplicationSystem->SetGroupFilterStatus(SubObjectInclusionGroupHandle, Client->ConnectionIdOnServer, ENetFilterStatus::Allow);
+
+	// Send and deliver packet
+	Server->UpdateAndSend({Client}, DeliverPacket);
+
+	// Late add subobject to group
+	Server->ReplicationSystem->AddToGroup(SubObjectInclusionGroupHandle, ServerSubObject->NetRefHandle);
+ 
+	// Send and deliver packet
+	Server->UpdateAndSend({Client}, DeliverPacket);
+
+	// Neither object nor subobject should have been created on the client.
+	UE_NET_ASSERT_EQ(Client->GetReplicationBridge()->GetReplicatedObject(ServerObject->NetRefHandle), nullptr);
+	UE_NET_ASSERT_EQ(Client->GetReplicationBridge()->GetReplicatedObject(ServerSubObject->NetRefHandle), nullptr);
+}
+
+UE_NET_TEST_FIXTURE(FTestFilteringFixture, SubObjectRemovedFromInclusionGroupFollowsOwnerInOtherInclusionGroup)
+{
+	// Setup dynamic filter for the test
+	SetDynamicFilterStatus(ENetFilterStatus::Disallow);
+
+	// Add client
+	FReplicationSystemTestClient* Client = CreateClient();
+
+	// Spawn object with subobject on server
+	UReplicatedTestObject* ServerObject = Server->CreateObject(UTestReplicatedIrisObject::FComponents{});
+	UReplicatedTestObject* ServerSubObject = Server->CreateSubObject(ServerObject->NetRefHandle, UTestReplicatedIrisObject::FComponents{});
+
+	// Set filter which will filter out all objects
+	Server->ReplicationSystem->SetFilter(ServerObject->NetRefHandle, MockFilterHandle);
+
+	// Setup separate inclusion group filters for object and subobject
+	FNetObjectGroupHandle ObjectInclusionGroupHandle = Server->ReplicationSystem->CreateGroup();
+	Server->ReplicationSystem->AddToGroup(ObjectInclusionGroupHandle, ServerObject->NetRefHandle);
+	Server->ReplicationSystem->AddInclusionFilterGroup(ObjectInclusionGroupHandle);
+	Server->ReplicationSystem->SetGroupFilterStatus(ObjectInclusionGroupHandle, Client->ConnectionIdOnServer, ENetFilterStatus::Disallow);
+	
+	FNetObjectGroupHandle SubObjectInclusionGroupHandle = Server->ReplicationSystem->CreateGroup();
+	Server->ReplicationSystem->AddToGroup(SubObjectInclusionGroupHandle, ServerSubObject->NetRefHandle);
+	Server->ReplicationSystem->AddInclusionFilterGroup(SubObjectInclusionGroupHandle);
+	Server->ReplicationSystem->SetGroupFilterStatus(SubObjectInclusionGroupHandle, Client->ConnectionIdOnServer, ENetFilterStatus::Disallow);
+
+	// Send and deliver packet
+	Server->UpdateAndSend({Client}, DeliverPacket);
+
+	// Neither object nor subobject should have been created on the client.
+	UE_NET_ASSERT_EQ(Client->GetReplicationBridge()->GetReplicatedObject(ServerObject->NetRefHandle), nullptr);
+	UE_NET_ASSERT_EQ(Client->GetReplicationBridge()->GetReplicatedObject(ServerSubObject->NetRefHandle), nullptr);
+
+	// Allow subobject group to be replicated. This should not change anything.
+	Server->ReplicationSystem->SetGroupFilterStatus(SubObjectInclusionGroupHandle, Client->ConnectionIdOnServer, ENetFilterStatus::Allow);
+
+	// Send and deliver packet
+	Server->UpdateAndSend({Client}, DeliverPacket);
+
+	// Neither object nor subobject should have been created on the client.
+	UE_NET_ASSERT_EQ(Client->GetReplicationBridge()->GetReplicatedObject(ServerObject->NetRefHandle), nullptr);
+	UE_NET_ASSERT_EQ(Client->GetReplicationBridge()->GetReplicatedObject(ServerSubObject->NetRefHandle), nullptr);
+
+	// Remove subobject from group. This should not affect anything.
+	Server->ReplicationSystem->RemoveFromGroup(SubObjectInclusionGroupHandle, ServerSubObject->NetRefHandle);
+
+	// Send and deliver packet
+	Server->UpdateAndSend({Client}, DeliverPacket);
+
+	// Neither object nor subobject should have been created on the client.
+	UE_NET_ASSERT_EQ(Client->GetReplicationBridge()->GetReplicatedObject(ServerObject->NetRefHandle), nullptr);
+	UE_NET_ASSERT_EQ(Client->GetReplicationBridge()->GetReplicatedObject(ServerSubObject->NetRefHandle), nullptr);
+
+	// Re-add subobject to its group for a second round of tests...
+	Server->ReplicationSystem->AddToGroup(SubObjectInclusionGroupHandle, ServerSubObject->NetRefHandle);
+	
+	// Allow object group to be replicated and subobject group to not be replicated. This should result in both the object and subobject being created on the client.
+	Server->ReplicationSystem->SetGroupFilterStatus(ObjectInclusionGroupHandle, Client->ConnectionIdOnServer, ENetFilterStatus::Allow);
+	Server->ReplicationSystem->SetGroupFilterStatus(SubObjectInclusionGroupHandle, Client->ConnectionIdOnServer, ENetFilterStatus::Disallow);
+
+	// Send and deliver packet
+	Server->UpdateAndSend({Client}, DeliverPacket);
+
+	// Both object and subobject should have been created on the client.
+	UE_NET_ASSERT_NE(Client->GetReplicationBridge()->GetReplicatedObject(ServerObject->NetRefHandle), nullptr);
+	UE_NET_ASSERT_NE(Client->GetReplicationBridge()->GetReplicatedObject(ServerSubObject->NetRefHandle), nullptr);
+
+	// Remove subobject from group. This should not affect anything.
+	Server->ReplicationSystem->RemoveFromGroup(SubObjectInclusionGroupHandle, ServerSubObject->NetRefHandle);
+
+	// Send and deliver packet
+	Server->UpdateAndSend({Client}, DeliverPacket);
+
+	// Both object and subobject should have been created on the client.
+	UE_NET_ASSERT_NE(Client->GetReplicationBridge()->GetReplicatedObject(ServerObject->NetRefHandle), nullptr);
+	UE_NET_ASSERT_NE(Client->GetReplicationBridge()->GetReplicatedObject(ServerSubObject->NetRefHandle), nullptr);
+}
+
+UE_NET_TEST_FIXTURE(FTestFilteringFixture, InclusionGroupsWorksWithMultipleObjects)
+{
+	// Setup dynamic filter for the test
+	SetDynamicFilterStatus(ENetFilterStatus::Disallow);
+
+	// Add clients
+	FReplicationSystemTestClient* ClientArray[3];
+	for (FReplicationSystemTestClient*& Client : ClientArray)
+	{
+		Client = CreateClient();
+	}
+
+	// Spawn objects on server
+	UReplicatedTestObject* ServerObjects[3];
+	for (UReplicatedTestObject*& ServerObject : ServerObjects)
+	{
+		ServerObject = Server->CreateObject(UTestReplicatedIrisObject::FComponents{});
+		// Filter out by default
+		Server->ReplicationSystem->SetFilter(ServerObject->NetRefHandle, MockFilterHandle);
+	}
+
+	// Setup inclusion group filters
+	FNetObjectGroupHandle InclusionGroupHandles[UE_ARRAY_COUNT(ServerObjects)];
+	for (FNetObjectGroupHandle& InclusionGroupHandle : InclusionGroupHandles)
+	{
+		const SIZE_T Index = &InclusionGroupHandle - &InclusionGroupHandles[0];
+		InclusionGroupHandle = Server->ReplicationSystem->CreateGroup();
+		Server->ReplicationSystem->AddInclusionFilterGroup(InclusionGroupHandle);
+		Server->ReplicationSystem->AddToGroup(InclusionGroupHandle, ServerObjects[Index]->NetRefHandle);
+		// Disallow by default
+		Server->ReplicationSystem->SetGroupFilterStatus(InclusionGroupHandle, ENetFilterStatus::Disallow);
+		// Allow one client
+		Server->ReplicationSystem->SetGroupFilterStatus(InclusionGroupHandle, ClientArray[Index]->ConnectionIdOnServer, ENetFilterStatus::Allow);
+	}
+
+	// Send and deliver packet
+	Server->UpdateAndSend(ClientArray, DeliverPacket);
+
+	// Exactly one object should have been replicated to each
+	for (UReplicatedTestObject*& ServerObject : ServerObjects)
+	{
+		const SIZE_T Index = &ServerObject - &ServerObjects[0];
+		UE_NET_ASSERT_NE(Clients[Index]->GetReplicationBridge()->GetReplicatedObject(ServerObject->NetRefHandle), nullptr);
+		UE_NET_ASSERT_EQ(Clients[(Index + 1U) % UE_ARRAY_COUNT(ServerObjects)]->GetReplicationBridge()->GetReplicatedObject(ServerObject->NetRefHandle), nullptr);
+		UE_NET_ASSERT_EQ(Clients[(Index + 2U) % UE_ARRAY_COUNT(ServerObjects)]->GetReplicationBridge()->GetReplicatedObject(ServerObject->NetRefHandle), nullptr);
+	}
+}
+
+UE_NET_TEST_FIXTURE(FTestFilteringFixture, InclusionGroupsAreCumulative)
+{
+	// Setup dynamic filter for the test
+	SetDynamicFilterStatus(ENetFilterStatus::Disallow);
+
+	// Add clients
+	FReplicationSystemTestClient* ClientArray[3];
+	for (FReplicationSystemTestClient*& Client : ClientArray)
+	{
+		Client = CreateClient();
+	}
+
+	// Spawn objects on server
+	UReplicatedTestObject* ServerObjects[3];
+	for (UReplicatedTestObject*& ServerObject : ServerObjects)
+	{
+		ServerObject = Server->CreateObject(UTestReplicatedIrisObject::FComponents{});
+		// Filter out by default
+		Server->ReplicationSystem->SetFilter(ServerObject->NetRefHandle, MockFilterHandle);
+	}
+
+	// Setup inclusion group filters
+	FNetObjectGroupHandle InclusionGroupHandles[UE_ARRAY_COUNT(ServerObjects)];
+	for (FNetObjectGroupHandle& InclusionGroupHandle : InclusionGroupHandles)
+	{
+		const SIZE_T Index = &InclusionGroupHandle - &InclusionGroupHandles[0];
+		InclusionGroupHandle = Server->ReplicationSystem->CreateGroup();
+		Server->ReplicationSystem->AddInclusionFilterGroup(InclusionGroupHandle);
+		Server->ReplicationSystem->AddToGroup(InclusionGroupHandle, ServerObjects[Index]->NetRefHandle);
+		// Disallow by default
+		Server->ReplicationSystem->SetGroupFilterStatus(InclusionGroupHandle, ENetFilterStatus::Disallow);
+		// Allow one client
+		Server->ReplicationSystem->SetGroupFilterStatus(InclusionGroupHandle, ClientArray[Index]->ConnectionIdOnServer, ENetFilterStatus::Allow);
+	}
+
+	// Send and deliver packet
+	Server->UpdateAndSend(ClientArray, DeliverPacket);
+
+	// Exactly one object should have been replicated to each connection
+	for (UReplicatedTestObject*& ServerObject : ServerObjects)
+	{
+		const SIZE_T Index = &ServerObject - &ServerObjects[0];
+		UE_NET_ASSERT_NE(Clients[Index]->GetReplicationBridge()->GetReplicatedObject(ServerObject->NetRefHandle), nullptr);
+		UE_NET_ASSERT_EQ(Clients[(Index + 1U) % UE_ARRAY_COUNT(ServerObjects)]->GetReplicationBridge()->GetReplicatedObject(ServerObject->NetRefHandle), nullptr);
+		UE_NET_ASSERT_EQ(Clients[(Index + 2U) % UE_ARRAY_COUNT(ServerObjects)]->GetReplicationBridge()->GetReplicatedObject(ServerObject->NetRefHandle), nullptr);
+	}
+}
+
+UE_NET_TEST_FIXTURE(FTestFilteringFixture, LateAddedConnectionWorksWithSimpleGroupInclusionFilter)
+{
+	// Setup dynamic filter for the test
+	SetDynamicFilterStatus(ENetFilterStatus::Disallow);
+
+	// Add client
+	FReplicationSystemTestClient* Client = CreateClient();
+
+	// Spawn objects on server
+	UReplicatedTestObject* ServerObjects[4];
+	for (UReplicatedTestObject*& ServerObject : ServerObjects)
+	{
+		ServerObject = Server->CreateObject(UTestReplicatedIrisObject::FComponents{});
+		// Filter out by default
+		Server->ReplicationSystem->SetFilter(ServerObject->NetRefHandle, MockFilterHandle);
+	}
+
+	// Setup inclusion group filters differently for each object
+	FNetObjectGroupHandle InclusionGroupHandles[UE_ARRAY_COUNT(ServerObjects)];
+	for (FNetObjectGroupHandle& InclusionGroupHandle : InclusionGroupHandles)
+	{
+		const SIZE_T Index = &InclusionGroupHandle - &InclusionGroupHandles[0];
+		InclusionGroupHandle = Server->ReplicationSystem->CreateGroup();
+		Server->ReplicationSystem->AddInclusionFilterGroup(InclusionGroupHandle);
+		Server->ReplicationSystem->AddToGroup(InclusionGroupHandle, ServerObjects[Index]->NetRefHandle);
+
+		switch (Index)
+		{
+		case 0:
+			Server->ReplicationSystem->SetGroupFilterStatus(InclusionGroupHandle, ENetFilterStatus::Disallow);
+			break;
+		case 1:
+			Server->ReplicationSystem->SetGroupFilterStatus(InclusionGroupHandle, ENetFilterStatus::Allow);
+			break;
+		case 2:
+			Server->ReplicationSystem->SetGroupFilterStatus(InclusionGroupHandle, ENetFilterStatus::Disallow);
+			Server->ReplicationSystem->SetGroupFilterStatus(InclusionGroupHandle, Client->ConnectionIdOnServer, ENetFilterStatus::Allow);
+			break;
+		case 3:
+			Server->ReplicationSystem->SetGroupFilterStatus(InclusionGroupHandle, ENetFilterStatus::Disallow);
+			// For testing purposes we predict the next client ID and allow the object to be replicated to it.
+			Server->ReplicationSystem->SetGroupFilterStatus(InclusionGroupHandle, Client->ConnectionIdOnServer + 1, ENetFilterStatus::Allow);
+			break;
+		};
+	}
+
+	// Send and deliver packet
+	Server->UpdateAndSend({Client}, DeliverPacket);
+
+	// Add client
+	FReplicationSystemTestClient* LateAddedClient = CreateClient();
+
+	// Send and deliver packet
+	Server->UpdateAndSend({Client, LateAddedClient}, DeliverPacket);
+
+	// Verify objects were created or not according to inclusion filters
+	UE_NET_ASSERT_EQ(Client->GetReplicationBridge()->GetReplicatedObject(ServerObjects[0]->NetRefHandle), nullptr);
+	UE_NET_ASSERT_EQ(LateAddedClient->GetReplicationBridge()->GetReplicatedObject(ServerObjects[0]->NetRefHandle), nullptr);
+	UE_NET_ASSERT_NE(Client->GetReplicationBridge()->GetReplicatedObject(ServerObjects[1]->NetRefHandle), nullptr);
+	UE_NET_ASSERT_NE(LateAddedClient->GetReplicationBridge()->GetReplicatedObject(ServerObjects[1]->NetRefHandle), nullptr);
+	UE_NET_ASSERT_NE(Client->GetReplicationBridge()->GetReplicatedObject(ServerObjects[2]->NetRefHandle), nullptr);
+	UE_NET_ASSERT_EQ(LateAddedClient->GetReplicationBridge()->GetReplicatedObject(ServerObjects[2]->NetRefHandle), nullptr);
+	UE_NET_ASSERT_EQ(Client->GetReplicationBridge()->GetReplicatedObject(ServerObjects[3]->NetRefHandle), nullptr);
+	UE_NET_ASSERT_NE(LateAddedClient->GetReplicationBridge()->GetReplicatedObject(ServerObjects[3]->NetRefHandle), nullptr);
+}
+
+UE_NET_TEST_FIXTURE(FTestFilteringFixture, LateAddedConnectionWorksWithComplexGroupInclusionFilter)
+{
+	// Setup dynamic filter for the test
+	SetDynamicFilterStatus(ENetFilterStatus::Disallow);
+
+	// Add client
+	FReplicationSystemTestClient* Client = CreateClient();
+
+	// Spawn objects with one subobject on server
+	UReplicatedTestObject* ServerObjects[4];
+	UReplicatedTestObject* ServerSubObjects[4];
+	for (UReplicatedTestObject*& ServerObject : ServerObjects)
+	{
+		const SIZE_T Index = &ServerObject - &ServerObjects[0];
+
+		ServerObject = Server->CreateObject(UTestReplicatedIrisObject::FComponents{});
+		ServerSubObjects[Index] = Server->CreateSubObject(ServerObject->NetRefHandle, UTestReplicatedIrisObject::FComponents{});
+
+		// Filter out by default
+		Server->ReplicationSystem->SetFilter(ServerObject->NetRefHandle, MockFilterHandle);
+	}
+
+	// Setup inclusion group filters differently for each object, and subobject differently than the root object
+	FNetObjectGroupHandle ObjectInclusionGroupHandles[UE_ARRAY_COUNT(ServerObjects)];
+	FNetObjectGroupHandle SubObjectInclusionGroupHandles[UE_ARRAY_COUNT(ServerSubObjects)];
+	for (FNetObjectGroupHandle& ObjectInclusionGroupHandle : ObjectInclusionGroupHandles)
+	{
+		const SIZE_T Index = &ObjectInclusionGroupHandle - &ObjectInclusionGroupHandles[0];
+
+		ObjectInclusionGroupHandle = Server->ReplicationSystem->CreateGroup();
+		Server->ReplicationSystem->AddInclusionFilterGroup(ObjectInclusionGroupHandle);
+		Server->ReplicationSystem->AddToGroup(ObjectInclusionGroupHandle, ServerObjects[Index]->NetRefHandle);
+
+		FNetObjectGroupHandle SubObjectInclusionGroupHandle = Server->ReplicationSystem->CreateGroup();
+		Server->ReplicationSystem->AddInclusionFilterGroup(SubObjectInclusionGroupHandle);
+		Server->ReplicationSystem->AddToGroup(SubObjectInclusionGroupHandle, ServerSubObjects[Index]->NetRefHandle);
+
+		
+		switch (Index)
+		{
+		case 0:
+			Server->ReplicationSystem->SetGroupFilterStatus(ObjectInclusionGroupHandle, ENetFilterStatus::Disallow);
+			Server->ReplicationSystem->SetGroupFilterStatus(SubObjectInclusionGroupHandle, ENetFilterStatus::Allow);
+			break;
+		case 1:
+			Server->ReplicationSystem->SetGroupFilterStatus(ObjectInclusionGroupHandle, ENetFilterStatus::Allow);
+			Server->ReplicationSystem->SetGroupFilterStatus(SubObjectInclusionGroupHandle, ENetFilterStatus::Disallow);
+			break;
+		case 2:
+			Server->ReplicationSystem->SetGroupFilterStatus(ObjectInclusionGroupHandle, ENetFilterStatus::Disallow);
+			Server->ReplicationSystem->SetGroupFilterStatus(ObjectInclusionGroupHandle, Client->ConnectionIdOnServer, ENetFilterStatus::Allow);
+
+			Server->ReplicationSystem->SetGroupFilterStatus(SubObjectInclusionGroupHandle, ENetFilterStatus::Allow);
+			Server->ReplicationSystem->SetGroupFilterStatus(SubObjectInclusionGroupHandle, Client->ConnectionIdOnServer, ENetFilterStatus::Disallow);
+			break;
+		case 3:
+			Server->ReplicationSystem->SetGroupFilterStatus(ObjectInclusionGroupHandle, ENetFilterStatus::Disallow);
+			// For testing purposes we predict the next client ID and allow the object to be replicated to it.
+			Server->ReplicationSystem->SetGroupFilterStatus(ObjectInclusionGroupHandle, Client->ConnectionIdOnServer + 1, ENetFilterStatus::Allow);
+
+			Server->ReplicationSystem->SetGroupFilterStatus(SubObjectInclusionGroupHandle, ENetFilterStatus::Allow);
+			Server->ReplicationSystem->SetGroupFilterStatus(SubObjectInclusionGroupHandle, Client->ConnectionIdOnServer + 1, ENetFilterStatus::Disallow);
+			break;
+		};
+	}
+
+	// Send and deliver packet
+	Server->UpdateAndSend({Client}, DeliverPacket);
+
+	// Add client
+	FReplicationSystemTestClient* LateAddedClient = CreateClient();
+
+	// Send and deliver packet
+	Server->UpdateAndSend({Client, LateAddedClient}, DeliverPacket);
+
+	// Verify objects were created or not according to inclusion filters
+	UE_NET_ASSERT_EQ(Client->GetReplicationBridge()->GetReplicatedObject(ServerObjects[0]->NetRefHandle), nullptr);
+	UE_NET_ASSERT_EQ(LateAddedClient->GetReplicationBridge()->GetReplicatedObject(ServerObjects[0]->NetRefHandle), nullptr);
+	UE_NET_ASSERT_EQ(Client->GetReplicationBridge()->GetReplicatedObject(ServerSubObjects[0]->NetRefHandle), nullptr);
+	UE_NET_ASSERT_EQ(LateAddedClient->GetReplicationBridge()->GetReplicatedObject(ServerSubObjects[0]->NetRefHandle), nullptr);
+
+	UE_NET_ASSERT_NE(Client->GetReplicationBridge()->GetReplicatedObject(ServerObjects[1]->NetRefHandle), nullptr);
+	UE_NET_ASSERT_NE(LateAddedClient->GetReplicationBridge()->GetReplicatedObject(ServerObjects[1]->NetRefHandle), nullptr);
+	UE_NET_ASSERT_NE(Client->GetReplicationBridge()->GetReplicatedObject(ServerSubObjects[1]->NetRefHandle), nullptr);
+	UE_NET_ASSERT_NE(LateAddedClient->GetReplicationBridge()->GetReplicatedObject(ServerSubObjects[1]->NetRefHandle), nullptr);
+	
+	UE_NET_ASSERT_NE(Client->GetReplicationBridge()->GetReplicatedObject(ServerObjects[2]->NetRefHandle), nullptr);
+	UE_NET_ASSERT_EQ(LateAddedClient->GetReplicationBridge()->GetReplicatedObject(ServerObjects[2]->NetRefHandle), nullptr);
+	UE_NET_ASSERT_NE(Client->GetReplicationBridge()->GetReplicatedObject(ServerSubObjects[2]->NetRefHandle), nullptr);
+	UE_NET_ASSERT_EQ(LateAddedClient->GetReplicationBridge()->GetReplicatedObject(ServerSubObjects[2]->NetRefHandle), nullptr);
+	
+	UE_NET_ASSERT_EQ(Client->GetReplicationBridge()->GetReplicatedObject(ServerObjects[3]->NetRefHandle), nullptr);
+	UE_NET_ASSERT_NE(LateAddedClient->GetReplicationBridge()->GetReplicatedObject(ServerObjects[3]->NetRefHandle), nullptr);
+	UE_NET_ASSERT_EQ(Client->GetReplicationBridge()->GetReplicatedObject(ServerSubObjects[3]->NetRefHandle), nullptr);
+	UE_NET_ASSERT_NE(LateAddedClient->GetReplicationBridge()->GetReplicatedObject(ServerSubObjects[3]->NetRefHandle), nullptr);
+}
+
 } // end namespace UE::Net::Private
diff --git a/Engine/Plugins/Runtime/ReplicationSystemTestPlugin/Source/Private/Tests/ReplicationSystem/Filtering/TestGroups.cpp b/Engine/Plugins/Runtime/ReplicationSystemTestPlugin/Source/Private/Tests/ReplicationSystem/Filtering/TestGroups.cpp
--- a/Engine/Plugins/Runtime/ReplicationSystemTestPlugin/Source/Private/Tests/ReplicationSystem/Filtering/TestGroups.cpp
+++ b/Engine/Plugins/Runtime/ReplicationSystemTestPlugin/Source/Private/Tests/ReplicationSystem/Filtering/TestGroups.cpp
@@ -157,7 +157,7 @@ UE_NET_TEST_FIXTURE(FTestGroupsFixture, GroupFilterAPI)
 
 	// Create group
 	FNetObjectGroupHandle GroupHandle = Server->ReplicationSystem->CreateGroup();
-	UE_NET_ASSERT_EQ(true, GroupHandle.IsValid());	
+	UE_NET_ASSERT_TRUE(GroupHandle.IsValid());	
 
 	// Add Objects to group
 	Server->ReplicationSystem->AddToGroup(GroupHandle, ServerObject0->NetRefHandle);
@@ -170,7 +170,7 @@ UE_NET_TEST_FIXTURE(FTestGroupsFixture, GroupFilterAPI)
 
 	// Create another group
 	FNetObjectGroupHandle GroupHandle2 = Server->ReplicationSystem->CreateGroup();
-	UE_NET_ASSERT_EQ(true, GroupHandle2.IsValid());	
+	UE_NET_ASSERT_TRUE(GroupHandle2.IsValid());	
 
 	// Add Objects to group
 	Server->ReplicationSystem->AddToGroup(GroupHandle2, ServerObject3->NetRefHandle);
@@ -189,9 +189,9 @@ UE_NET_TEST_FIXTURE(FTestGroupsFixture, GroupFilterAPI)
 	const FNetBitArrayView ObjectGroupFilter = ServerFiltering->GetGroupFilteredOutObjects(Client->ConnectionIdOnServer);
 
 	// Add group to group filter
-	Server->ReplicationSystem->AddGroupFilter(GroupHandle);
-	Server->ReplicationSystem->AddGroupFilter(GroupHandle2);
-	Server->ReplicationSystem->AddGroupFilter(GroupHandle3);
+	Server->ReplicationSystem->AddExclusionFilterGroup(GroupHandle);
+	Server->ReplicationSystem->AddExclusionFilterGroup(GroupHandle2);
+	Server->ReplicationSystem->AddExclusionFilterGroup(GroupHandle3);
 
 	Server->ReplicationSystem->SetGroupFilterStatus(GroupHandle, Client->ConnectionIdOnServer, ENetFilterStatus::Disallow);
 	Server->ReplicationSystem->SetGroupFilterStatus(GroupHandle2, Client->ConnectionIdOnServer, ENetFilterStatus::Disallow);
@@ -275,7 +275,7 @@ UE_NET_TEST_FIXTURE(FTestGroupsFixture, GroupFilterAPIManyFilters)
 	for (uint32 It = 0; It < GroupCount; ++It)
 	{
 		FNetObjectGroupHandle GroupHandle = Server->ReplicationSystem->CreateGroup();
-		Server->ReplicationSystem->AddGroupFilter(GroupHandle);
+		Server->ReplicationSystem->AddExclusionFilterGroup(GroupHandle);
 	
 		GroupHandles[It] = GroupHandle;	
 	}
@@ -306,8 +306,8 @@ UE_NET_TEST_FIXTURE(FTestGroupsFixture, GroupFilterAPIObjectMemberOfMoreThanOneG
 	VerifyObjectFilterStatus(ServerHandle, Client->ConnectionIdOnServer, true);
 
 	// Add groups to group filter
-	Server->ReplicationSystem->AddGroupFilter(GroupHandle0);
-	Server->ReplicationSystem->AddGroupFilter(GroupHandle1);
+	Server->ReplicationSystem->AddExclusionFilterGroup(GroupHandle0);
+	Server->ReplicationSystem->AddExclusionFilterGroup(GroupHandle1);
 
 	// Filter out objects in the group
 	Server->ReplicationSystem->SetGroupFilterStatus(GroupHandle0, Client->ConnectionIdOnServer, ENetFilterStatus::Disallow);
@@ -364,8 +364,8 @@ UE_NET_TEST_FIXTURE(FTestGroupsFixture, GroupFilterAPIAddThenRemoveFilterTrait)
 	VerifyObjectFilterStatus(ServerHandle, Client->ConnectionIdOnServer, true);
 
 	// Add groups to group filter
-	Server->ReplicationSystem->AddGroupFilter(GroupHandle0);
-	Server->ReplicationSystem->AddGroupFilter(GroupHandle1);
+	Server->ReplicationSystem->AddExclusionFilterGroup(GroupHandle0);
+	Server->ReplicationSystem->AddExclusionFilterGroup(GroupHandle1);
 
 	// Filter out objects in the group
 	Server->ReplicationSystem->SetGroupFilterStatus(GroupHandle0, Client->ConnectionIdOnServer, ENetFilterStatus::Disallow);
@@ -422,15 +422,15 @@ UE_NET_TEST_FIXTURE(FTestGroupsFixture, GroupFilterTestObjectListViaTrait)
 	// Update filters
 	Filter();
 
-	const FNetBitArrayView GlobalGroupFilterList = ServerGroups->GetGroupFilteredObjects();
+	const FNetBitArrayView GlobalGroupFilterList = ServerGroups->GetGroupFilteredOutObjects();
 
 	// Object is not in a group filter yet
 	UE_NET_ASSERT_FALSE(GlobalGroupFilterList.IsBitSet(ServerInternalIndex));
 	VerifyObjectFilterStatus(ServerHandle, Client->ConnectionIdOnServer, true);
 
 	// Set two groups to start filtering
-	Server->ReplicationSystem->AddGroupFilter(GroupHandle0);
-	Server->ReplicationSystem->AddGroupFilter(GroupHandle2);
+	Server->ReplicationSystem->AddExclusionFilterGroup(GroupHandle0);
+	Server->ReplicationSystem->AddExclusionFilterGroup(GroupHandle2);
 
 	Filter();
 
@@ -457,7 +457,7 @@ UE_NET_TEST_FIXTURE(FTestGroupsFixture, GroupFilterTestObjectListViaTrait)
 	VerifyObjectFilterStatus(ServerHandle, Client->ConnectionIdOnServer, true);
 
 	// Add the filter trait to the last group
-	Server->ReplicationSystem->AddGroupFilter(GroupHandle1);
+	Server->ReplicationSystem->AddExclusionFilterGroup(GroupHandle1);
 
 	Filter();
 
@@ -490,13 +490,13 @@ UE_NET_TEST_FIXTURE(FTestGroupsFixture, GroupFilterTestObjectListViaMembership)
 	Server->ReplicationSystem->AddToGroup(GroupHandle2, ServerHandle);
 
 	// Set two groups to start filtering
-	Server->ReplicationSystem->AddGroupFilter(GroupHandle0);
-	Server->ReplicationSystem->AddGroupFilter(GroupHandle2);
+	Server->ReplicationSystem->AddExclusionFilterGroup(GroupHandle0);
+	Server->ReplicationSystem->AddExclusionFilterGroup(GroupHandle2);
 
 	// Update filters
 	Filter();
 
-	const FNetBitArrayView GlobalGroupFilterList = ServerGroups->GetGroupFilteredObjects();
+	const FNetBitArrayView GlobalGroupFilterList = ServerGroups->GetGroupFilteredOutObjects();
 
 	// Object should be in a group filter
 	UE_NET_ASSERT_TRUE(GlobalGroupFilterList.IsBitSet(ServerInternalIndex));
@@ -580,7 +580,7 @@ UE_NET_TEST_FIXTURE(FTestGroupsFixture, GroupFilterAPISubObjects)
 	
 	// Create groups
 	FNetObjectGroupHandle GroupHandle = Server->ReplicationSystem->CreateGroup();
-	Server->ReplicationSystem->AddGroupFilter(GroupHandle);
+	Server->ReplicationSystem->AddExclusionFilterGroup(GroupHandle);
 
 	// Add Object to group
 	Server->ReplicationSystem->AddToGroup(GroupHandle, ServerHandle);
@@ -601,7 +601,7 @@ UE_NET_TEST_FIXTURE(FTestGroupsFixture, GroupFilterAPIRemovingSubObjectsRestores
 {
 	// Create groups
 	FNetObjectGroupHandle GroupHandle = Server->ReplicationSystem->CreateGroup();
-	Server->ReplicationSystem->AddGroupFilter(GroupHandle);
+	Server->ReplicationSystem->AddExclusionFilterGroup(GroupHandle);
 
 	// Update filters
 	Filter();
@@ -641,7 +641,7 @@ UE_NET_TEST_FIXTURE(FTestGroupsFixture, GroupFilterAPILateAddedSubObjectIsFilter
 	
 	// Create groups
 	FNetObjectGroupHandle GroupHandle = Server->ReplicationSystem->CreateGroup();
-	Server->ReplicationSystem->AddGroupFilter(GroupHandle);
+	Server->ReplicationSystem->AddExclusionFilterGroup(GroupHandle);
 
 	// Add Object to group
 	Server->ReplicationSystem->AddToGroup(GroupHandle, ServerHandle);
@@ -678,7 +678,7 @@ UE_NET_TEST_FIXTURE(FTestGroupsFixture, GroupFilterAPISubObjectIsFilteredOutIfOw
 	FNetObjectGroupHandle GroupHandle = Server->ReplicationSystem->CreateGroup();
 	FNetObjectGroupHandle SubObjectGroupHandle = Server->ReplicationSystem->CreateGroup();
 
-	Server->ReplicationSystem->AddGroupFilter(GroupHandle);
+	Server->ReplicationSystem->AddExclusionFilterGroup(GroupHandle);
 	Server->ReplicationSystem->SetGroupFilterStatus(GroupHandle, ENetFilterStatus::Allow);
 
 	// Add Object to group
@@ -722,8 +722,8 @@ UE_NET_TEST_FIXTURE(FTestGroupsFixture, GroupFilterAPISubObjectIsFilteredOutWith
 	FNetObjectGroupHandle GroupHandle = Server->ReplicationSystem->CreateGroup();
 	FNetObjectGroupHandle SubObjectGroupHandle = Server->ReplicationSystem->CreateGroup();
 
-	Server->ReplicationSystem->AddGroupFilter(GroupHandle);
-	Server->ReplicationSystem->AddGroupFilter(SubObjectGroupHandle);
+	Server->ReplicationSystem->AddExclusionFilterGroup(GroupHandle);
+	Server->ReplicationSystem->AddExclusionFilterGroup(SubObjectGroupHandle);
 	Server->ReplicationSystem->SetGroupFilterStatus(GroupHandle, ENetFilterStatus::Allow);
 
 	// Add Object to group
diff --git a/Engine/Plugins/Runtime/ReplicationSystemTestPlugin/Source/Private/Tests/ReplicationSystem/ReplicationSystemServerClientTestFixture.cpp b/Engine/Plugins/Runtime/ReplicationSystemTestPlugin/Source/Private/Tests/ReplicationSystem/ReplicationSystemServerClientTestFixture.cpp
--- a/Engine/Plugins/Runtime/ReplicationSystemTestPlugin/Source/Private/Tests/ReplicationSystem/ReplicationSystemServerClientTestFixture.cpp
+++ b/Engine/Plugins/Runtime/ReplicationSystemTestPlugin/Source/Private/Tests/ReplicationSystem/ReplicationSystemServerClientTestFixture.cpp
@@ -301,6 +301,12 @@ void FReplicationSystemTestNode::PostSendUpdate()
 void FReplicationSystemTestNode::DeliverTo(FReplicationSystemTestNode& Dest, uint32 LocalConnectionId, uint32 RemoteConnectionId, bool bDeliver)
 {
 	FConnectionInfo& Connection = GetConnectionInfo(LocalConnectionId);
+	if (Connection.WrittenPackets.IsEmpty())
+	{
+		UE_LOG(LogIris, Log, TEXT("ReplicationSystemTestFixture: Conn: %u Unable to %hs packet as there are no packets."), LocalConnectionId, (bDeliver ? "deliver" : "drop"));
+		return;
+	}
+
 	const FPacketData& Packet = Connection.WrittenPackets.Peek();
 
 	if (bDeliver)
@@ -414,7 +420,7 @@ void FReplicationSystemTestServer::DeliverTo(FReplicationSystemTestClient* Clien
 	FReplicationSystemTestNode::DeliverTo(*Client, Client->ConnectionIdOnServer, Client->LocalConnectionId, bDeliver);
 }
 
-bool FReplicationSystemTestServer::UpdateAndSend(const TArray<FReplicationSystemTestClient*>& Clients, bool bDeliver /*= true*/)
+bool FReplicationSystemTestServer::UpdateAndSend(const TArrayView<FReplicationSystemTestClient*const>& Clients, bool bDeliver /*= true*/)
 {
 	bool bSuccess = true;
 
diff --git a/Engine/Plugins/Runtime/ReplicationSystemTestPlugin/Source/Private/Tests/ReplicationSystem/ReplicationSystemServerClientTestFixture.h b/Engine/Plugins/Runtime/ReplicationSystemTestPlugin/Source/Private/Tests/ReplicationSystem/ReplicationSystemServerClientTestFixture.h
--- a/Engine/Plugins/Runtime/ReplicationSystemTestPlugin/Source/Private/Tests/ReplicationSystem/ReplicationSystemServerClientTestFixture.h
+++ b/Engine/Plugins/Runtime/ReplicationSystemTestPlugin/Source/Private/Tests/ReplicationSystem/ReplicationSystemServerClientTestFixture.h
@@ -191,7 +191,7 @@ public:
 	void DeliverTo(FReplicationSystemTestClient* Client, bool bDeliver);
 
 	// Tick and send packets to one or all clients
-	bool UpdateAndSend(const TArray<FReplicationSystemTestClient*>& Clients, bool bDeliver=true);
+	bool UpdateAndSend(const TArrayView<FReplicationSystemTestClient*const>& Clients, bool bDeliver=true);
 };
 
 class FReplicationSystemServerClientTestFixture : public FNetworkAutomationTestSuiteFixture
diff --git a/Engine/Plugins/Runtime/ReplicationSystemTestPlugin/Source/Private/Tests/ReplicationSystem/TestDependentObjects.cpp b/Engine/Plugins/Runtime/ReplicationSystemTestPlugin/Source/Private/Tests/ReplicationSystem/TestDependentObjects.cpp
--- a/Engine/Plugins/Runtime/ReplicationSystemTestPlugin/Source/Private/Tests/ReplicationSystem/TestDependentObjects.cpp
+++ b/Engine/Plugins/Runtime/ReplicationSystemTestPlugin/Source/Private/Tests/ReplicationSystem/TestDependentObjects.cpp
@@ -83,7 +83,7 @@ UE_NET_TEST_FIXTURE(FReplicationSystemServerClientTestFixture, TestDependentObje
 
 	// Filter out Server object to start with
 	FNetObjectGroupHandle FilterGroup = ReplicationSystem->CreateGroup();
-	ReplicationSystem->AddGroupFilter(FilterGroup);
+	ReplicationSystem->AddExclusionFilterGroup(FilterGroup);
 	ReplicationSystem->AddToGroup(FilterGroup, ServerObject->NetRefHandle);
 
 	// Setup dependent object to only replicate with ServerObject
@@ -131,7 +131,7 @@ UE_NET_TEST_FIXTURE(FReplicationSystemServerClientTestFixture, TestChainedDepend
 
 	// Filter out Server object to start with
 	FNetObjectGroupHandle FilterGroup = ReplicationSystem->CreateGroup();
-	ReplicationSystem->AddGroupFilter(FilterGroup);
+	ReplicationSystem->AddExclusionFilterGroup(FilterGroup);
 	ReplicationSystem->AddToGroup(FilterGroup, ServerObject->NetRefHandle);
 
 	// Setup dependent object to only replicate with ServerObject
diff --git a/Engine/Plugins/Runtime/ReplicationSystemTestPlugin/Source/Private/Tests/ReplicationSystem/TestReplicationSystem.cpp b/Engine/Plugins/Runtime/ReplicationSystemTestPlugin/Source/Private/Tests/ReplicationSystem/TestReplicationSystem.cpp
--- a/Engine/Plugins/Runtime/ReplicationSystemTestPlugin/Source/Private/Tests/ReplicationSystem/TestReplicationSystem.cpp
+++ b/Engine/Plugins/Runtime/ReplicationSystemTestPlugin/Source/Private/Tests/ReplicationSystem/TestReplicationSystem.cpp
@@ -715,7 +715,7 @@ UE_NET_TEST_FIXTURE(FReplicationSystemServerClientTestFixture, AddRemoveFromConn
 	FNetObjectGroupHandle Group = ReplicationSystem->CreateGroup();
 	ReplicationSystem->AddToGroup(Group, ServerObject->NetRefHandle);
 
-	ReplicationSystem->AddGroupFilter(Group);
+	ReplicationSystem->AddExclusionFilterGroup(Group);
 	ReplicationSystem->SetGroupFilterStatus(Group, ENetFilterStatus::Allow);
 
 	// Start replicating object
diff --git a/Engine/Source/Runtime/Experimental/Iris/Core/Private/Iris/ReplicationSystem/Filtering/NetObjectGroups.cpp b/Engine/Source/Runtime/Experimental/Iris/Core/Private/Iris/ReplicationSystem/Filtering/NetObjectGroups.cpp
--- a/Engine/Source/Runtime/Experimental/Iris/Core/Private/Iris/ReplicationSystem/Filtering/NetObjectGroups.cpp
+++ b/Engine/Source/Runtime/Experimental/Iris/Core/Private/Iris/ReplicationSystem/Filtering/NetObjectGroups.cpp
@@ -1,6 +1,7 @@
 // Copyright Epic Games, Inc. All Rights Reserved.
 
 #include "NetObjectGroups.h"
+#include "Math/UnrealMathUtility.h"
 #include "Iris/ReplicationSystem/Filtering/NetObjectFilter.h"
 #include "Iris/ReplicationSystem/NetRefHandleManager.h" // for InvalidInternalIndex
 #include "Containers/ArrayView.h"
@@ -71,22 +72,23 @@ void FNetObjectGroups::ResetGroupMembership(FNetObjectGroupMembership& Target)
 
 void FNetObjectGroups::Init(const FNetObjectGroupInitParams& Params)
 {
-	MaxGroupCount = Params.MaxGroupCount;
+	ensureMsgf(Params.MaxGroupCount < std::numeric_limits<FNetObjectGroupHandle::FGroupIndexType>::max(), TEXT("MaxGroupCount cannot exceed %u"), std::numeric_limits<FNetObjectGroupHandle::FGroupIndexType>::max());
+	MaxGroupCount = FMath::Clamp<uint32>(Params.MaxGroupCount, 0U, std::numeric_limits<FNetObjectGroupHandle::FGroupIndexType>::max());
 
 	// Reserve first as invalid group
 	Groups.Add(FNetObjectGroup());
 
 	GroupMemberships.SetNumZeroed(Params.MaxObjectCount);
-	GroupFilteredObjects.Init(Params.MaxObjectCount);
+	GroupFilteredOutObjects.Init(Params.MaxObjectCount);
 }
 
 FNetObjectGroupHandle FNetObjectGroups::CreateGroup()
 {
 	if (ensure((uint32)Groups.Num() < MaxGroupCount))
 	{
-		const int Index = Groups.Add(FNetObjectGroup());
+		const uint32 Index = static_cast<uint32>(Groups.Add(FNetObjectGroup()));
 		FNetObjectGroupHandle GroupHandle;
-		GroupHandle.Index = Index;
+		GroupHandle.Index = static_cast<FNetObjectGroupHandle::FGroupIndexType>(Index);
 		GroupHandle.Epoch = CurrentEpoch;
 		return GroupHandle;
 	}
@@ -124,6 +126,16 @@ void FNetObjectGroups::SetGroupName(FNetObjectGroupHandle GroupHandle, FName Gro
 	}
 }
 
+FName FNetObjectGroups::GetGroupName(FNetObjectGroupHandle GroupHandle) const
+{
+	if (const FNetObjectGroup* Group = GetGroup(GroupHandle))
+	{
+		return Group->GroupName;
+	}
+
+	return FName();
+}
+
 FNetObjectGroupHandle FNetObjectGroups::CreateNamedGroup(FName GroupName)
 {
 	if (NamedGroups.Contains(GroupName))
@@ -219,7 +231,7 @@ void FNetObjectGroups::AddToGroup(FNetObjectGroupHandle GroupHandle, FInternalNe
 
 			if (IsFilterGroup(*Group))
 			{
-				GroupFilteredObjects.SetBit(InternalIndex);
+				GroupFilteredOutObjects.SetBit(InternalIndex);
 			}
 		}
 		else
@@ -243,50 +255,82 @@ void FNetObjectGroups::RemoveFromGroup(FNetObjectGroupHandle GroupHandle, FInter
 		// Check to see if the object is still part of a filter group
 		if (!IsInAnyFilterGroup(GroupMembership))
 		{
-			GroupFilteredObjects.ClearBit(InternalIndex);
+			GroupFilteredOutObjects.ClearBit(InternalIndex);
 		}
 	}
 }
 
-void FNetObjectGroups::AddFilterTrait(FNetObjectGroupHandle GroupHandle)
+void FNetObjectGroups::AddExclusionFilterTrait(FNetObjectGroupHandle GroupHandle)
 {
 	if (FNetObjectGroup* Group = GetGroup(GroupHandle))
 	{
 		if (!IsFilterGroup(*Group))
 		{
-			Group->Traits |= ENetObjectGroupTraits::IsFiltering;
+			Group->Traits |= ENetObjectGroupTraits::IsExclusionFiltering;
 
 			// Flag all current members of this group that they are now filterable
 			for (FInternalNetRefIndex MemberIndex : Group->Members)
 			{
-				GroupFilteredObjects.SetBit(MemberIndex);
+				GroupFilteredOutObjects.SetBit(MemberIndex);
 			}
 		}
 	}
 }
 
-void FNetObjectGroups::RemoveFilterTrait(FNetObjectGroupHandle GroupHandle)
+void FNetObjectGroups::RemoveExclusionFilterTrait(FNetObjectGroupHandle GroupHandle)
 {
-	if (FNetObjectGroup* Group = GetGroup(GroupHandle))
+	FNetObjectGroup* Group = GetGroup(GroupHandle);
+	if (Group == nullptr)
 	{
-		Group->Traits &= ~(ENetObjectGroupTraits::IsFiltering);
+		return;
+	}
+
+	if (!IsExclusionFilterGroup(*Group))
+	{
+		return;
+	}
 
-		for (FInternalNetRefIndex MemberIndex : Group->Members)
+	Group->Traits &= ~(ENetObjectGroupTraits::IsExclusionFiltering);
+
+	for (FInternalNetRefIndex MemberIndex : Group->Members)
+	{
+		// Check to see if the object is still part of a filter group
+		if (!IsInAnyFilterGroup(GroupMemberships[MemberIndex]))
 		{
-			// Check to see if the object is still part of a filter group
-			if (!IsInAnyFilterGroup(GroupMemberships[MemberIndex]))
-			{
-				GroupFilteredObjects.ClearBit(MemberIndex);
-			}
+			GroupFilteredOutObjects.ClearBit(MemberIndex);
 		}
 	}
 }
 
+void FNetObjectGroups::AddInclusionFilterTrait(FNetObjectGroupHandle GroupHandle)
+{
+	if (FNetObjectGroup* Group = GetGroup(GroupHandle))
+	{
+		// Can't be both inclusion and exclusion so let's do nothing if the group has any sort of filter trait.
+		if (!IsFilterGroup(*Group))
+		{
+			Group->Traits |= ENetObjectGroupTraits::IsInclusionFiltering;
+		}
+	}
+}
+
+void FNetObjectGroups::RemoveInclusionFilterTrait(FNetObjectGroupHandle GroupHandle)
+{
+	FNetObjectGroup* Group = GetGroup(GroupHandle);
+	if (Group == nullptr)
+	{
+		return;
+	}
+
+	// Simply remove the trait.
+	Group->Traits &= ~(ENetObjectGroupTraits::IsInclusionFiltering);
+}
+
 bool FNetObjectGroups::IsFilterGroup(FNetObjectGroupHandle GroupHandle) const
 {
 	if (const FNetObjectGroup* Group = GetGroup(GroupHandle))
 	{
-		return EnumHasAnyFlags(Group->Traits, ENetObjectGroupTraits::IsFiltering);
+		return EnumHasAnyFlags(Group->Traits, ENetObjectGroupTraits::IsExclusionFiltering | ENetObjectGroupTraits::IsInclusionFiltering);
 	}
 
 	return false;
diff --git a/Engine/Source/Runtime/Experimental/Iris/Core/Private/Iris/ReplicationSystem/Filtering/NetObjectGroups.h b/Engine/Source/Runtime/Experimental/Iris/Core/Private/Iris/ReplicationSystem/Filtering/NetObjectGroups.h
--- a/Engine/Source/Runtime/Experimental/Iris/Core/Private/Iris/ReplicationSystem/Filtering/NetObjectGroups.h
+++ b/Engine/Source/Runtime/Experimental/Iris/Core/Private/Iris/ReplicationSystem/Filtering/NetObjectGroups.h
@@ -23,20 +23,19 @@ namespace UE::Net::Private
 
 enum class ENetObjectGroupTraits : uint32
 {
-	None				= 0x0000,
-	IsFindableByName	= 0x0001,
-	IsFiltering			= 0x0002,
+	None                 = 0x0000,
+	IsFindableByName     = 0x0001,
+	IsExclusionFiltering = 0x0002,
+	IsInclusionFiltering = 0x0004,
 };
 ENUM_CLASS_FLAGS(ENetObjectGroupTraits);
 
 struct FNetObjectGroup
 {
-	FNetObjectGroup() : Traits(ENetObjectGroupTraits::None) {}
-
 	// Group members can only be replicated objects that have internal indices
 	TArray<FInternalNetRefIndex> Members;
 	FName GroupName;
-	ENetObjectGroupTraits Traits;
+	ENetObjectGroupTraits Traits = ENetObjectGroupTraits::None;
 };
 
 struct FNetObjectGroupInitParams
@@ -67,22 +66,35 @@ public:
 	FNetObjectGroup* GetGroupByIndex(FNetObjectGroupHandle::FGroupIndexType GroupIndex);
 
 	void SetGroupName(FNetObjectGroupHandle GroupHandle, FName GroupName);
+	FName GetGroupName(FNetObjectGroupHandle GroupHandle) const;
 
-	inline bool IsValidGroup(FNetObjectGroupHandle GroupHandle) const { return GroupHandle.IsValid() && GroupHandle.Epoch == CurrentEpoch && Groups.IsValidIndex(GroupHandle.GetGroupIndex()); }
+	bool IsValidGroup(FNetObjectGroupHandle GroupHandle) const;
 
 	bool Contains(FNetObjectGroupHandle GroupHandle, FInternalNetRefIndex InternalIndex) const;
 	void AddToGroup(FNetObjectGroupHandle GroupHandle, FInternalNetRefIndex InternalIndex);
 	void RemoveFromGroup(FNetObjectGroupHandle GroupHandle, FInternalNetRefIndex InternalIndex);
 
-	/** Called when a group is set to filter it's members */
-	void AddFilterTrait(FNetObjectGroupHandle GroupHandle);
+	/** Called when a group is to be used as an exclusion filter group */
+	void AddExclusionFilterTrait(FNetObjectGroupHandle GroupHandle);
+
+	/** Called when a group is no longer used as an exclusion filter group */
+	void RemoveExclusionFilterTrait(FNetObjectGroupHandle GroupHandle);
 
-	/** Called when a group no longer filters it's members */
-	void RemoveFilterTrait(FNetObjectGroupHandle GroupHandle);
+	/** Called when a group is to be used as an inclusion filter group */
+	void AddInclusionFilterTrait(FNetObjectGroupHandle GroupHandle);
 
-	/** Does the gorup have the filter trait */
+	/** Called when a group is no longer used as an inclusion filter group */
+	void RemoveInclusionFilterTrait(FNetObjectGroupHandle GroupHandle);
+
+	/** Does the group have a filter trait, either exclusion or inclusion */
 	bool IsFilterGroup(FNetObjectGroupHandle GroupHandle) const;
 
+	/** Does the group have the exclusion filter trait */
+	bool IsExclusionFilterGroup(FNetObjectGroupHandle GroupHandle) const;
+
+	/** Does the group have the inclusion filter trait */
+	bool IsInclusionFilterGroup(FNetObjectGroupHandle GroupHandle) const;
+
 	/** Returns how many groups the given handle is a member of */
 	uint32 GetNumGroupMemberships(FInternalNetRefIndex InternalIndex) const;
 
@@ -98,9 +110,9 @@ public:
 	void DestroyNamedGroup(FName GroupName);
 
 	/** Returns a list of all objects currently part of a group with the filter trait */
-	const FNetBitArrayView GetGroupFilteredObjects() const
+	const FNetBitArrayView GetGroupFilteredOutObjects() const
 	{
-		return MakeNetBitArrayView(GroupFilteredObjects);
+		return MakeNetBitArrayView(GroupFilteredOutObjects);
 	}
 
 private:
@@ -116,6 +128,8 @@ private:
 	static bool IsMemberOf(const FNetObjectGroupMembership& Target, FNetObjectGroupHandle Group);
 
 	bool IsFilterGroup(const FNetObjectGroup& Group) const;
+	bool IsExclusionFilterGroup(const FNetObjectGroup& Group) const;
+	bool IsInclusionFilterGroup(const FNetObjectGroup& Group) const;
 
 	bool IsInAnyFilterGroup(const FNetObjectGroupMembership& GroupMembership) const;
 
@@ -127,18 +141,33 @@ private:
 	uint32 MaxGroupCount;
 
 	// List of objects that are members of a group with a filter trait
-	FNetBitArray GroupFilteredObjects;
+	FNetBitArray GroupFilteredOutObjects;
 
 	TMap<FName, FNetObjectGroupHandle> NamedGroups;
 
-	uint32 CurrentEpoch = 0U;
+	FNetObjectGroupHandle::FGroupIndexType CurrentEpoch = 0U;
 
-	inline static uint32 NextEpoch = 1U;
+	inline static FNetObjectGroupHandle::FGroupIndexType NextEpoch = 1U;
 };
 
+inline bool FNetObjectGroups::IsValidGroup(FNetObjectGroupHandle GroupHandle) const
+{
+	return GroupHandle.IsValid() && GroupHandle.Epoch == CurrentEpoch && Groups.IsValidIndex(GroupHandle.GetGroupIndex());
+}
+
 inline bool FNetObjectGroups::IsFilterGroup(const FNetObjectGroup& Group) const
 {
-	return EnumHasAnyFlags(Group.Traits, ENetObjectGroupTraits::IsFiltering);
+	return EnumHasAnyFlags(Group.Traits, ENetObjectGroupTraits::IsExclusionFiltering | ENetObjectGroupTraits::IsInclusionFiltering);
+}
+
+inline bool FNetObjectGroups::IsExclusionFilterGroup(const FNetObjectGroup& Group) const
+{
+	return EnumHasAnyFlags(Group.Traits, ENetObjectGroupTraits::IsExclusionFiltering);
+}
+
+inline bool FNetObjectGroups::IsInclusionFilterGroup(const FNetObjectGroup& Group) const
+{
+	return EnumHasAnyFlags(Group.Traits, ENetObjectGroupTraits::IsInclusionFiltering);
 }
 
 } // end namespace UE::Net::Private
diff --git a/Engine/Source/Runtime/Experimental/Iris/Core/Private/Iris/ReplicationSystem/Filtering/ReplicationFiltering.cpp b/Engine/Source/Runtime/Experimental/Iris/Core/Private/Iris/ReplicationSystem/Filtering/ReplicationFiltering.cpp
--- a/Engine/Source/Runtime/Experimental/Iris/Core/Private/Iris/ReplicationSystem/Filtering/ReplicationFiltering.cpp
+++ b/Engine/Source/Runtime/Experimental/Iris/Core/Private/Iris/ReplicationSystem/Filtering/ReplicationFiltering.cpp
@@ -219,6 +219,8 @@ FReplicationFiltering::FReplicationFiltering()
 , bHasDynamicFilters(0)
 , bHasDynamicRawFilters(0)
 , bHasDynamicFragmentFilters(0)
+, bHasDirtyExclusionFilterGroup(0)
+, bHasDirtyInclusionFilterGroup(0)
 {
 	StaticChecks();
 }
@@ -259,13 +261,19 @@ void FReplicationFiltering::Init(FReplicationFilteringInitParams& Params)
 	DynamicFilterEnabledObjects.Init(MaxObjectCount);
 
 	// Group filtering
-	MaxGroupCount = Params.MaxGroupCount;
-	GroupInfos.SetNumZeroed(MaxGroupCount);
-	FilterGroups.Init(GroupInfos.Num());
-	DirtyFilterGroups.Init(GroupInfos.Num());
-	// SubObjectFilter groups
-	SubObjectFilterGroups.Init(GroupInfos.Num());
-	DirtySubObjectFilterGroups.Init(GroupInfos.Num());
+	{
+		const uint32 InMaxGroupCount = Params.MaxGroupCount;
+		check(InMaxGroupCount < std::numeric_limits<FNetObjectGroupHandle::FGroupIndexType>::max());
+		MaxGroupCount = InMaxGroupCount;
+		GroupInfos.SetNumZeroed(InMaxGroupCount);
+		ExclusionFilterGroups.Init(InMaxGroupCount);
+		InclusionFilterGroups.Init(InMaxGroupCount);
+		DirtyExclusionFilterGroups.Init(InMaxGroupCount);
+		DirtyInclusionFilterGroups.Init(InMaxGroupCount);
+		// SubObjectFilter groups
+		SubObjectFilterGroups.Init(InMaxGroupCount);
+		DirtySubObjectFilterGroups.Init(InMaxGroupCount);
+	}
 
 	ObjectIndexToPerObjectInfoIndex.SetNumZeroed(MaxObjectCount);
 
@@ -296,8 +304,9 @@ void FReplicationFiltering::FilterPrePoll()
 
 	UpdateObjectsInScope();
 
-	// Update group filters if needed
-	UpdateGroupFiltering();
+	UpdateGroupExclusionFiltering();
+
+	UpdateGroupInclusionFiltering();
 
 	UpdateOwnerAndConnectionFiltering();
 
@@ -336,18 +345,18 @@ void FReplicationFiltering::FilterNonRelevantObjects()
 	const uint32* const WithOwnerData = ObjectsWithOwnerFilter.GetData();
 	const uint32* const ConnectionFiltersData = AllConnectionFilteredObjects.GetData();
 	const uint32* const DynamicFilteredData = DynamicFilterEnabledObjects.GetData();
-	const uint32* const GroupFilteredData = Groups->GetGroupFilteredObjects().GetData();
+	const uint32* const GroupFilteredOutData = Groups->GetGroupFilteredOutObjects().GetData();
 
 	FNetBitArrayView GlobalRelevantObjects = NetRefHandleManager->GetRelevantObjectsInternalIndices();
 	
 	uint32* GlobalRelevantData = GlobalRelevantObjects.GetData();
 
-	//$IRIS TODO: Could be worth sampling if it's faster for the AlwaysRelevant list to be kept cached instead of recalculating it here frame.
+	//$IRIS TODO: Could be worth sampling if it's faster for the AlwaysRelevant list to be kept cached instead of recalculating it here every frame.
 	const uint32 MaxWords = GlobalRelevantObjects.GetNumWords();
 	for (uint32 WordIndex = 0; WordIndex < MaxWords; ++WordIndex)
 	{
 		// Build the list of always relevant objects (objects that have no filters)
-		GlobalRelevantData[WordIndex] = CurrentFrameScopeData[WordIndex] & ~(WithOwnerData[WordIndex] | ConnectionFiltersData[WordIndex] | DynamicFilteredData[WordIndex] | GroupFilteredData[WordIndex]);
+		GlobalRelevantData[WordIndex] = CurrentFrameScopeData[WordIndex] & ~(WithOwnerData[WordIndex] | ConnectionFiltersData[WordIndex] | DynamicFilteredData[WordIndex] | GroupFilteredOutData[WordIndex]);
 	}
 
 	// Build the list of currently relevant objects. e.g. always relevant objects + filterable objects relevant to at least one connection.	
@@ -629,8 +638,9 @@ void FReplicationFiltering::RemoveConnection(uint32 ConnectionId)
 	// Reset connection info
 	FPerConnectionInfo& ConnectionInfo = ConnectionInfos[ConnectionId];
 	ConnectionInfo.ConnectionFilteredObjects.Empty();
-	ConnectionInfo.GroupFilteredOutObjects.Empty();
+	ConnectionInfo.GroupExcludedObjects.Empty();
 	ConnectionInfo.ObjectsInScopeBeforeDynamicFiltering.Empty();
+	ConnectionInfo.GroupIncludedObjects.Empty();
 	ConnectionInfo.ObjectsInScope.Empty();
 
 	for (FFilterInfo& Info : DynamicFilterInfos)
@@ -641,7 +651,7 @@ void FReplicationFiltering::RemoveConnection(uint32 ConnectionId)
 	// Reset SubObject filter for removed connection
 	SubObjectFilterGroups.ForAllSetBits([this, ConnectionId](uint32 GroupIndex)
 	{
-		if (!FNetObjectGroupHandle::IsReservedNetObjectGroupIndex(GroupIndex))
+		if (!FNetObjectGroupHandle::IsReservedNetObjectGroupIndex(static_cast<FNetObjectGroupHandle::FGroupIndexType>(GroupIndex)))
 		{
 			SetConnectionFilterStatus(*GetPerObjectInfo(GroupInfos[GroupIndex].ConnectionStateIndex), ConnectionId, ENetFilterStatus::Disallow);
 			DirtySubObjectFilterGroups.SetBit(GroupIndex);
@@ -670,10 +680,13 @@ void FReplicationFiltering::InitNewConnections()
 		ConnectionInfo.ConnectionFilteredObjects.Copy(ScopableInternalIndices);
 		ConnectionInfo.ConnectionFilteredObjects.ClearBit(FNetRefHandleManager::InvalidInternalIndex);
 
-		// Everything enabled by default in groups, i.e. not filtered out.
-		ConnectionInfo.GroupFilteredOutObjects.Init(MaxObjectCount);
+		// Do not filter out anything by default.
+		ConnectionInfo.GroupExcludedObjects.Init(MaxObjectCount);
 
-		// The combined result of scoped objects, connection filtering and group filtering.
+		// Do not override dynamic filtering by default.
+		ConnectionInfo.GroupIncludedObjects.Init(MaxObjectCount);
+
+		// The combined result of scoped objects, connection filtering and group exclusion filtering.
 		ConnectionInfo.ObjectsInScopeBeforeDynamicFiltering.Init(MaxObjectCount);
 
 		// The final result of all filtering.
@@ -685,32 +698,71 @@ void FReplicationFiltering::InitNewConnections()
 			ConnectionInfo.InProgressDynamicFilteredOutObjects.Init(MaxObjectCount);
 		}
 
-		// Update group filter by iterating over all registered groups for the connection
+		// Update group exclusion filtering
 		{
-			auto InitGroupFilterForConnection = [this, &ConnectionInfo, ConnectionId](uint32 GroupIndexIn)
+			auto InitExclusionGroupFilterForConnection = [this, &ConnectionInfo, ConnectionId](uint32 InGroupIndex)
 			{
-				const FNetObjectGroupHandle::FGroupIndexType GroupIndex = static_cast<FNetObjectGroupHandle::FGroupIndexType>(GroupIndexIn);
+				const FNetObjectGroupHandle::FGroupIndexType GroupIndex = static_cast<FNetObjectGroupHandle::FGroupIndexType>(InGroupIndex);
 				const FPerObjectInfo* ConnectionState = GetPerObjectInfo(this->GroupInfos[GroupIndex].ConnectionStateIndex);
 
-				// Setup filter for connection
+				// Setup exclusion filter for connection
 				if (GetConnectionFilterStatus(*ConnectionState, ConnectionId) == ENetFilterStatus::Disallow)
 				{
 					// Apply filter
 					const FNetObjectGroup* Group = this->Groups->GetGroupByIndex(GroupIndex);
-					FNetBitArray& GroupFilteredOutObjects = ConnectionInfo.GroupFilteredOutObjects;
-					for (uint32 ObjectIndex : MakeArrayView(Group->Members.GetData(), Group->Members.Num()))
+					FNetBitArray& GroupExcludedObjects = ConnectionInfo.GroupExcludedObjects;
+					for (const FInternalNetRefIndex ObjectIndex : Group->Members)
 					{
-						GroupFilteredOutObjects.SetBit(ObjectIndex);
+						GroupExcludedObjects.SetBit(ObjectIndex);
+
 						// Filter subobjects
 						for (const FInternalNetRefIndex SubObjectIndex : NetRefHandleManager->GetSubObjects(ObjectIndex))
 						{
-							GroupFilteredOutObjects.SetBit(SubObjectIndex);
+							GroupExcludedObjects.SetBit(SubObjectIndex);
 						}
 					}
 				}
 			};
 
-			FilterGroups.ForAllSetBits(InitGroupFilterForConnection);
+			ExclusionFilterGroups.ForAllSetBits(InitExclusionGroupFilterForConnection);
+		}
+
+		// Update group inclusion filtering
+		{
+			auto InitInclusionGroupFilterForConnection = [this, &ConnectionInfo, ConnectionId](uint32 InGroupIndex)
+			{
+				const FNetObjectGroupHandle::FGroupIndexType GroupIndex = static_cast<FNetObjectGroupHandle::FGroupIndexType>(InGroupIndex);
+				const FPerObjectInfo* ConnectionState = GetPerObjectInfo(this->GroupInfos[GroupIndex].ConnectionStateIndex);
+
+				// Setup inclusion filter for connection
+				if (GetConnectionFilterStatus(*ConnectionState, ConnectionId) == ENetFilterStatus::Allow)
+				{
+					const FNetBitArray& ObjectsInScope = ConnectionInfo.ConnectionFilteredObjects;
+					const FNetBitArrayView SubObjectInternalIndices = NetRefHandleManager->GetSubObjectInternalIndicesView();
+
+					// Apply filter
+					const FNetObjectGroup* Group = this->Groups->GetGroupByIndex(GroupIndex);
+					FNetBitArray& GroupIncludedObjects = ConnectionInfo.GroupIncludedObjects;
+					for (const FInternalNetRefIndex ObjectIndex : Group->Members)
+					{
+						// SubObjects follow root object.
+						if (UNLIKELY(SubObjectInternalIndices.GetBit(ObjectIndex)))
+						{
+							continue;
+						}
+
+						GroupIncludedObjects.SetBit(ObjectIndex);
+
+						// Filter subobjects
+						for (const FInternalNetRefIndex SubObjectIndex : NetRefHandleManager->GetSubObjects(ObjectIndex))
+						{
+							GroupIncludedObjects.SetBit(SubObjectIndex);
+						}
+					}
+				}
+			};
+
+			InclusionFilterGroups.ForAllSetBits(InitInclusionGroupFilterForConnection);
 		}
 
 		// Update connection scope with owner filtering
@@ -739,7 +791,7 @@ void FReplicationFiltering::InitNewConnections()
 
 		// Combine connection and group filtering
 		ConnectionInfo.ObjectsInScopeBeforeDynamicFiltering.Copy(ConnectionInfo.ConnectionFilteredObjects);
-		ConnectionInfo.ObjectsInScopeBeforeDynamicFiltering.Combine(ConnectionInfo.GroupFilteredOutObjects, FNetBitArray::AndNotOp);
+		ConnectionInfo.ObjectsInScopeBeforeDynamicFiltering.Combine(ConnectionInfo.GroupExcludedObjects, FNetBitArray::AndNotOp);
 	};
 
 	NewConnections.ForAllSetBits(InitNewConnection);
@@ -773,7 +825,7 @@ void FReplicationFiltering::ResetRemovedConnections()
 		}
 	};
 
-	FilterGroups.ForAllSetBits(ResetGroupFilterStatus);
+	FNetBitArray::ForAllSetBits(ExclusionFilterGroups, InclusionFilterGroups, FNetBitArrayBase::OrOp, ResetGroupFilterStatus);
 }
 
 void FReplicationFiltering::UpdateObjectsInScope()
@@ -887,11 +939,18 @@ void FReplicationFiltering::UpdateObjectsInScope()
 					// If parent is a member of a group filter we need to refresh group filtering to include subobject
 					uint32 GroupMembershipCount = 0U;
 					const FNetObjectGroupHandle* GroupHandles = Groups->GetGroupMemberships(ParentIndex, GroupMembershipCount);
-					for (FNetObjectGroupHandle GroupHandle : MakeArrayView(GroupHandles, GroupMembershipCount))
+					for (const FNetObjectGroupHandle GroupHandle : MakeArrayView(GroupHandles, GroupMembershipCount))
 					{
-						if (FilterGroups.GetBit(GroupHandle.GetGroupIndex()))
+						const FNetObjectGroupHandle::FGroupIndexType GroupIndex = GroupHandle.GetGroupIndex();
+						if (ExclusionFilterGroups.GetBit(GroupIndex))
+						{
+							DirtyExclusionFilterGroups.SetBit(GroupIndex);
+							bHasDirtyExclusionFilterGroup = 1;
+						}
+						else if (InclusionFilterGroups.GetBit(GroupIndex))
 						{
-							DirtyFilterGroups.SetBit(GroupHandle.GetGroupIndex());
+							DirtyInclusionFilterGroups.SetBit(GroupIndex);
+							bHasDirtyInclusionFilterGroup = 1;
 						}
 					}
 
@@ -926,7 +985,8 @@ void FReplicationFiltering::UpdateObjectsInScope()
 
 		// Or in brand new objects and mask off deleted objects by anding with now existing objects.
 		uint32* FilteredObjectsStorage = ConnectionInfo.ConnectionFilteredObjects.GetData();
-		uint32* GroupFilteredOutObjectsStorage = ConnectionInfo.GroupFilteredOutObjects.GetData();
+		uint32* GroupExcludedObjectsStorage = ConnectionInfo.GroupExcludedObjects.GetData();
+		uint32* GroupIncludedObjectsStorage = ConnectionInfo.GroupIncludedObjects.GetData();
 		uint32* ObjectsInScopeBeforeDynamicFiltering = ConnectionInfo.ObjectsInScopeBeforeDynamicFiltering.GetData();
 		uint32* DynamicFilteredOutObjects = ConnectionInfo.DynamicFilteredOutObjects.GetData();
 
@@ -937,11 +997,18 @@ void FReplicationFiltering::UpdateObjectsInScope()
 			const uint32 PrevExistingObjects = FakePrevObjectsInScopeStorage[WordIndex];
 			const uint32 ExistingObjects = ObjectsInScopeStorage[WordIndex];
 			const uint32 NewObjects = ExistingObjects & ~PrevExistingObjects;
+
 			const uint32 FilteredObjectsWord = (FilteredObjectsStorage[WordIndex] | NewObjects) & ExistingObjects;
-			const uint32 GroupFilteredOutWord = GroupFilteredOutObjectsStorage[WordIndex] & ExistingObjects;
 			FilteredObjectsStorage[WordIndex] = FilteredObjectsWord;
-			GroupFilteredOutObjectsStorage[WordIndex] = GroupFilteredOutWord;
-			ObjectsInScopeBeforeDynamicFiltering[WordIndex] = FilteredObjectsWord & ~GroupFilteredOutWord;
+
+			const uint32 GroupExcludedObjectsWord = GroupExcludedObjectsStorage[WordIndex] & ExistingObjects;
+			GroupExcludedObjectsStorage[WordIndex] = GroupExcludedObjectsWord;
+
+			const uint32 GroupIncludedObjectsWord = GroupIncludedObjectsStorage[WordIndex] & ExistingObjects;
+			GroupIncludedObjectsStorage[WordIndex] = GroupIncludedObjectsWord;
+
+			// Note that we only filter out objects from exclusion groups here. Inclusion groups only overrides dynamic filtering.
+			ObjectsInScopeBeforeDynamicFiltering[WordIndex] = FilteredObjectsWord & ~GroupExcludedObjectsWord;
 
 			// Make sure the we restore dynamic filtering for new objects
 			if (HasDynamicFilters())
@@ -1010,12 +1077,12 @@ void FReplicationFiltering::UpdateOwnerAndConnectionFiltering()
 		{
 			FPerConnectionInfo& ConnectionInfo = this->ConnectionInfos[ConnectionId];
 			FNetBitArrayView ConnectionScope = MakeNetBitArrayView(ConnectionInfo.ConnectionFilteredObjects);
-			FNetBitArrayView GroupFilteredOutObjects = MakeNetBitArrayView(ConnectionInfo.GroupFilteredOutObjects);
+			FNetBitArrayView GroupExcludedObjects = MakeNetBitArrayView(ConnectionInfo.GroupExcludedObjects);
 			FNetBitArrayView ObjectsInScopeBeforeDynamicFiltering = MakeNetBitArrayView(ConnectionInfo.ObjectsInScopeBeforeDynamicFiltering);
 
 			// Update filter info
 			{
-				auto MaskObject = [this, &ConnectionScope, &GroupFilteredOutObjects, &ObjectsInScopeBeforeDynamicFiltering, ConnectionId, &CurrentFrameObjectsInScope](uint32 ObjectIndex)
+				auto MaskObject = [this, &ConnectionScope, &GroupExcludedObjects, &ObjectsInScopeBeforeDynamicFiltering, ConnectionId, &CurrentFrameObjectsInScope](uint32 ObjectIndex)
 				{
 					bool bObjectIsInScope = true;
 					if (HasOwnerFilter(ObjectIndex))
@@ -1034,7 +1101,7 @@ void FReplicationFiltering::UpdateOwnerAndConnectionFiltering()
 
 					// Update scope for parent object.
 					{
-						const bool bIsGroupEnabled = !GroupFilteredOutObjects.GetBit(ObjectIndex);
+						const bool bIsGroupEnabled = !GroupExcludedObjects.GetBit(ObjectIndex);
 						ConnectionScope.SetBitValue(ObjectIndex, bObjectIsInScope);
 						ObjectsInScopeBeforeDynamicFiltering.SetBitValue(ObjectIndex, bObjectIsInScope & bIsGroupEnabled);
 					}
@@ -1043,7 +1110,7 @@ void FReplicationFiltering::UpdateOwnerAndConnectionFiltering()
 					for (const FInternalNetRefIndex SubObjectIndex : NetRefHandleManager->GetSubObjects(ObjectIndex))
 					{
 						const bool bEnableObject = bObjectIsInScope && CurrentFrameObjectsInScope.GetBit(SubObjectIndex);
-						const bool bIsGroupEnabled = !GroupFilteredOutObjects.GetBit(SubObjectIndex);
+						const bool bIsGroupEnabled = !GroupExcludedObjects.GetBit(SubObjectIndex);
 
 						ConnectionScope.SetBitValue(SubObjectIndex, bEnableObject);
 						ObjectsInScopeBeforeDynamicFiltering.SetBitValue(SubObjectIndex, bEnableObject & bIsGroupEnabled);
@@ -1065,52 +1132,108 @@ void FReplicationFiltering::UpdateOwnerAndConnectionFiltering()
 	ObjectsWithDirtyOwner.Reset();
 }
 
-void FReplicationFiltering::UpdateGroupFiltering()
+void FReplicationFiltering::UpdateGroupExclusionFiltering()
 {
+	if (!bHasDirtyExclusionFilterGroup)
+	{
+		return;
+	}
+
 	// Adding objects to an active group filter is deferred in order to avoid triggering constant filter updates.
 	auto UpdateGroupFilter = [this](uint32 GroupIndex)
 	{
-		if (this->FilterGroups.GetBit(GroupIndex))
-		{
-			IRIS_PROFILER_SCOPE(FReplicationFiltering_UpdateGroupFiltering);
+		IRIS_PROFILER_SCOPE(FReplicationFiltering_UpdateGroupExclusionFiltering);
 
-			const FPerObjectInfo* ConnectionStateInfo = GetPerObjectInfo(this->GroupInfos[GroupIndex].ConnectionStateIndex);
-			const FNetObjectGroup* Group = this->Groups->GetGroupByIndex(FNetObjectGroupHandle::FGroupIndexType(GroupIndex));
-			FPerConnectionInfo* LocalConnectionInfos = this->ConnectionInfos.GetData();
+		const FPerObjectInfo* ConnectionStateInfo = GetPerObjectInfo(this->GroupInfos[GroupIndex].ConnectionStateIndex);
+		const FNetObjectGroup* Group = this->Groups->GetGroupByIndex(FNetObjectGroupHandle::FGroupIndexType(GroupIndex));
+		FPerConnectionInfo* LocalConnectionInfos = this->ConnectionInfos.GetData();
 
-			const FNetBitArrayView CurrentFrameScopableObjects = NetRefHandleManager->GetCurrentFrameScopableInternalIndices();
+		const FNetBitArrayView CurrentFrameScopableObjects = NetRefHandleManager->GetCurrentFrameScopableInternalIndices();
 
-			auto UpdateGroupFilterForConnection = [this, ConnectionStateInfo, Group, LocalConnectionInfos, CurrentFrameScopableObjects](uint32 ConnectionId)
+		auto UpdateGroupFilterForConnection = [this, ConnectionStateInfo, Group, LocalConnectionInfos, CurrentFrameScopableObjects](uint32 ConnectionId)
+		{
+			if (this->GetConnectionFilterStatus(*ConnectionStateInfo, ConnectionId) == ENetFilterStatus::Disallow)
 			{
-				if (this->GetConnectionFilterStatus(*ConnectionStateInfo, ConnectionId) == ENetFilterStatus::Disallow)
+				FNetBitArray& GroupExcludedObjects = LocalConnectionInfos[ConnectionId].GroupExcludedObjects;
+				FNetBitArray& ObjectsInScopeBeforeDynamicFiltering = LocalConnectionInfos[ConnectionId].ObjectsInScopeBeforeDynamicFiltering;
+
+				for (const FInternalNetRefIndex ObjectIndex : Group->Members)
 				{
-					FNetBitArray& GroupFilteredOutObjects = LocalConnectionInfos[ConnectionId].GroupFilteredOutObjects;
-					FNetBitArray& ObjectsInScopeBeforeDynamicFiltering = LocalConnectionInfos[ConnectionId].ObjectsInScopeBeforeDynamicFiltering;
+					GroupExcludedObjects.SetBit(ObjectIndex);
+					ObjectsInScopeBeforeDynamicFiltering.ClearBit(ObjectIndex);
 
-					for (uint32 ObjectIndex : Group->Members)
+					// Filter subobjects
+					for (const FInternalNetRefIndex SubObjectIndex : NetRefHandleManager->GetSubObjects(ObjectIndex))
 					{
-						GroupFilteredOutObjects.SetBit(ObjectIndex);
-						ObjectsInScopeBeforeDynamicFiltering.ClearBit(ObjectIndex);
+						const bool bIsScopable = CurrentFrameScopableObjects.GetBit(SubObjectIndex);
+						GroupExcludedObjects.SetBitValue(SubObjectIndex, bIsScopable);
+						ObjectsInScopeBeforeDynamicFiltering.ClearBit(SubObjectIndex);
+					}
+				}
+			}
+		};
 
-						// Filter subobjects
-						for (const FInternalNetRefIndex SubObjectIndex : NetRefHandleManager->GetSubObjects(ObjectIndex))
-						{
-							const bool bIsScopable = CurrentFrameScopableObjects.GetBit(SubObjectIndex);
-							GroupFilteredOutObjects.SetBitValue(SubObjectIndex, bIsScopable);
-							ObjectsInScopeBeforeDynamicFiltering.ClearBit(SubObjectIndex);
-						}
+		this->ValidConnections.ForAllSetBits(UpdateGroupFilterForConnection);
+	};
+
+	DirtyExclusionFilterGroups.ForAllSetBits(UpdateGroupFilter);
+
+	// Clear out dirtiness
+	bHasDirtyExclusionFilterGroup = 0;
+	DirtyExclusionFilterGroups.Reset();
+}
+
+void FReplicationFiltering::UpdateGroupInclusionFiltering()
+{
+	if (!bHasDirtyInclusionFilterGroup)
+	{
+		return;
+	}
+
+	auto UpdateGroupFilter = [this](uint32 GroupIndex)
+	{
+		IRIS_PROFILER_SCOPE(FReplicationFiltering_UpdateGroupInclusionFiltering);
+
+		const FPerObjectInfo* ConnectionStateInfo = GetPerObjectInfo(this->GroupInfos[GroupIndex].ConnectionStateIndex);
+		const FNetObjectGroup* Group = this->Groups->GetGroupByIndex(FNetObjectGroupHandle::FGroupIndexType(GroupIndex));
+		FPerConnectionInfo* LocalConnectionInfos = this->ConnectionInfos.GetData();
+
+		const FNetBitArrayView CurrentFrameScopableObjects = NetRefHandleManager->GetCurrentFrameScopableInternalIndices();
+		const FNetBitArrayView SubObjectInternalIndices = NetRefHandleManager->GetSubObjectInternalIndicesView();
+
+		auto UpdateGroupFilterForConnection = [this, ConnectionStateInfo, Group, LocalConnectionInfos, &CurrentFrameScopableObjects, &SubObjectInternalIndices](uint32 ConnectionId)
+		{
+			if (this->GetConnectionFilterStatus(*ConnectionStateInfo, ConnectionId) == ENetFilterStatus::Allow)
+			{
+				FNetBitArray& GroupIncludedObjects = LocalConnectionInfos[ConnectionId].GroupIncludedObjects;
+
+				for (const FInternalNetRefIndex ObjectIndex : Group->Members)
+				{
+					// SubObjects follow root object.
+					if (UNLIKELY(SubObjectInternalIndices.GetBit(ObjectIndex)))
+					{
+						continue;
+					}
+
+					GroupIncludedObjects.SetBit(ObjectIndex);
+
+					// Filter subobjects
+					for (const FInternalNetRefIndex SubObjectIndex : NetRefHandleManager->GetSubObjects(ObjectIndex))
+					{
+						GroupIncludedObjects.SetBit(SubObjectIndex);
 					}
 				}
-			};
+			}
+		};
 
-			this->ValidConnections.ForAllSetBits(UpdateGroupFilterForConnection);
-		}
+		this->ValidConnections.ForAllSetBits(UpdateGroupFilterForConnection);
 	};
 
-	DirtyFilterGroups.ForAllSetBits(UpdateGroupFilter);
+	DirtyInclusionFilterGroups.ForAllSetBits(UpdateGroupFilter);
 
 	// Clear out dirtiness
-	DirtyFilterGroups.Reset();
+	bHasDirtyInclusionFilterGroup = 0;
+	DirtyInclusionFilterGroups.Reset();
 }
 
 void FReplicationFiltering::PreUpdateDynamicFiltering(ENetFilterType FilterType)
@@ -1292,14 +1415,19 @@ void FReplicationFiltering::UpdateDynamicFiltering(ENetFilterType FilterType)
 			uint32* ObjectsInScopeData = ConnectionInfo.ObjectsInScope.GetData();
 			const uint32* ObjectsInScopeBeforeDynamicFilteringData = ConnectionInfo.ObjectsInScopeBeforeDynamicFiltering.GetData();
 			const uint32* DynamicFilteredOutObjectsData = ConnectionInfo.DynamicFilteredOutObjects.GetData();
+			const uint32* GroupIncludedObjectsData = ConnectionInfo.GroupIncludedObjects.GetData();
+			// $IRIS TODO Vectorization opportunity
 			for (SIZE_T WordIt = 0, WordEndIt = WordCountForObjectBitArrays; WordIt != WordEndIt; ++WordIt)
 			{
-				ObjectsInScopeData[WordIt] = ObjectsInScopeBeforeDynamicFilteringData[WordIt] & ~DynamicFilteredOutObjectsData[WordIt];
+				const uint32 ObjectsInScopeBeforeWord = ObjectsInScopeBeforeDynamicFilteringData[WordIt];
+				const uint32 DynamicFilteredOutWord = DynamicFilteredOutObjectsData[WordIt];
+				const uint32 GroupIncludedWord = GroupIncludedObjectsData[WordIt];
+				ObjectsInScopeData[WordIt] = ObjectsInScopeBeforeWord & (GroupIncludedWord | ~DynamicFilteredOutWord);
 			}
 		}
 
 		// The scope for the connection is now fully updated, apart from disabled dependent objects.
-		// If any object that has a dependency isn't filtered out we must re-enable the dependent object.
+		// If any object that has a dependency that isn't filtered out we must re-enable the dependent object.
 		{
 			for (const uint32 DependentObjectIndex : DisabledDependentObjects)
 			{
@@ -1308,7 +1436,8 @@ void FReplicationFiltering::UpdateDynamicFiltering(ENetFilterType FilterType)
 					ConnectionInfo.ObjectsInScope.SetBit(DependentObjectIndex);
 					for (const FInternalNetRefIndex SubObjectIndex : NetRefHandleManager->GetSubObjects(DependentObjectIndex))
 					{
-						ConnectionInfo.ObjectsInScope.SetBit(SubObjectIndex);
+						const bool bIsInScope = ConnectionInfo.ObjectsInScopeBeforeDynamicFiltering.GetBit(SubObjectIndex);
+						ConnectionInfo.ObjectsInScope.SetBitValue(SubObjectIndex, bIsInScope);
 					}
 				}
 			}
@@ -1494,22 +1623,26 @@ const FReplicationFiltering::FPerObjectInfo* FReplicationFiltering::GetPerObject
 void FReplicationFiltering::AddSubObjectFilter(FNetObjectGroupHandle GroupHandle)
 {
 	const FNetObjectGroupHandle::FGroupIndexType GroupIndex = GroupHandle.GetGroupIndex();
-	check(!FilterGroups.GetBit(GroupIndex));
+	const bool bIsValidGroup = ensureMsgf(Groups->IsValidGroup(GroupHandle), TEXT("AddGroupFilter received an invalid group handle GroupIndex: %u"), GroupIndex);
+	if (!bIsValidGroup)
+	{
+		return;
+	}
 
-	if (ensure(GroupHandle.IsValid() && !SubObjectFilterGroups.GetBit(GroupIndex)))
+	const bool bIsFiltering = Groups->IsFilterGroup(GroupHandle) || SubObjectFilterGroups.GetBit(GroupIndex);
+	ensureMsgf(!bIsFiltering, TEXT("NetObjectGroup Name '%s' GroupIndex %u was asked to start subobject filtering but it was already used for filtering."), ToCStr(Groups->GetGroupName(GroupHandle).ToString()), GroupIndex);
+	if (bIsFiltering)
 	{
-		const bool bGroupExists = SubObjectFilterGroups.GetBit(GroupIndex);
-		if (ensure(bGroupExists == false))
-		{
-			SubObjectFilterGroups.SetBit(GroupIndex);
-			GroupInfos[GroupIndex].ConnectionStateIndex = AllocPerObjectInfo();
+		return;
+	}
 
-			// By default filter is active, i.e. we filter out all connections
-			SetPerObjectInfoFilterStatus(*GetPerObjectInfo(GroupInfos[GroupIndex].ConnectionStateIndex), ENetFilterStatus::Disallow);
+	SubObjectFilterGroups.SetBit(GroupIndex);
+	GroupInfos[GroupIndex].ConnectionStateIndex = AllocPerObjectInfo();
 
-			UE_LOG(LogIrisFiltering, Verbose, TEXT("ReplicationFiltering::AddSubObjectFilter (%s) GroupIndex: %u, FilterStatus: DisallowReplication"), *(Groups->GetGroup(GroupHandle)->GroupName.ToString()), GroupIndex);
-		}
-	}
+	// By default we filter out all connections
+	SetPerObjectInfoFilterStatus(*GetPerObjectInfo(GroupInfos[GroupIndex].ConnectionStateIndex), ENetFilterStatus::Disallow);
+
+	UE_LOG(LogIrisFiltering, Verbose, TEXT("ReplicationFiltering::AddSubObjectFilter (%s) GroupIndex: %u, FilterStatus: DisallowReplication"), *(Groups->GetGroup(GroupHandle)->GroupName.ToString()), GroupIndex);
 }
 
 void FReplicationFiltering::RemoveSubObjectFilter(FNetObjectGroupHandle GroupHandle)
@@ -1535,7 +1668,7 @@ void FReplicationFiltering::UpdateSubObjectFilters()
 		FNetObjectGroupHandle GroupHandle = Groups->MakeNetObjectGroupHandle(static_cast<FNetObjectGroupHandle::FGroupIndexType>(GroupIndex));
 		if (const FNetObjectGroup* Group = Groups->GetGroup(GroupHandle))
 		{
-			if (Group->Members.Num() == 0U && !IsConnectionFilterStatusAllowedForAnyConnection(*GetPerObjectInfo(GroupInfos[GroupHandle.GetGroupIndex()].ConnectionStateIndex)))
+			if (Group->Members.IsEmpty() && !IsAnyConnectionFilterStatusAllowed(*GetPerObjectInfo(GroupInfos[GroupIndex].ConnectionStateIndex)))
 			{
 				ReplicationSystem->DestroyGroup(GroupHandle);
 			}
@@ -1550,43 +1683,45 @@ void FReplicationFiltering::SetSubObjectFilterStatus(FNetObjectGroupHandle Group
 {
 	IRIS_PROFILER_SCOPE(SetSubObjectFilterStatus)
 
+	const FNetObjectGroupHandle::FGroupIndexType GroupIndex = GroupHandle.GetGroupIndex();
 	if (GroupHandle.IsReservedNetObjectGroup())
 	{
-		UE_LOG(LogIrisFiltering, Warning, TEXT("FReplicationFiltering::SetSubObjectFilterStatus - Trying to set filter for reserved GroupIndex: %u which is not allowed."), GroupHandle.GetGroupIndex());
+		UE_LOG(LogIrisFiltering, Warning, TEXT("FReplicationFiltering::SetSubObjectFilterStatus - Trying to set filter for reserved GroupIndex: %u which is not allowed."), GroupIndex);
 		return;
 	}
 
-	if (!SubObjectFilterGroups.GetBit(GroupHandle.GetGroupIndex()))
+	if (!SubObjectFilterGroups.GetBit(GroupIndex))
 	{
-		UE_LOG(LogIrisFiltering, Warning, TEXT("FReplicationFiltering::SetSubObjectFilterStatus - Trying to ReplicationStatus for GroupIndex: %u that is not a SubObjectFilterGroup"), GroupHandle.GetGroupIndex());
+		UE_LOG(LogIrisFiltering, Warning, TEXT("FReplicationFiltering::SetSubObjectFilterStatus - Trying to ReplicationStatus for GroupIndex: %u that is not a SubObjectFilterGroup"), GroupIndex);
 		return;
 	}
 
 	FPerObjectInfo* FilterInfo = GetPerObjectInfo(GroupInfos[GroupHandle.GetGroupIndex()].ConnectionStateIndex);
 	SetPerObjectInfoFilterStatus(*FilterInfo, ReplicationStatus);
-	if (!IsConnectionFilterStatusAllowedForAnyConnection(*FilterInfo))
+	if (!IsAnyConnectionFilterStatusAllowed(*FilterInfo))
 	{
-		DirtySubObjectFilterGroups.SetBit(GroupHandle.GetGroupIndex());
+		DirtySubObjectFilterGroups.SetBit(GroupIndex);
 	}
 }
 
 void FReplicationFiltering::SetSubObjectFilterStatus(FNetObjectGroupHandle GroupHandle, const FNetBitArrayView& ConnectionsBitArray, ENetFilterStatus ReplicationStatus)
 {
+	const FNetObjectGroupHandle::FGroupIndexType GroupIndex = GroupHandle.GetGroupIndex();
 	if (ConnectionsBitArray.GetNumBits() > ValidConnections.GetNumBits())
 	{
-		UE_LOG(LogIrisFiltering, Warning, TEXT("FReplicationFiltering::SetSubObjectFilterStatus - Trying to set filter for %u, with invalid Connections parameters."), GroupHandle.GetGroupIndex());
+		UE_LOG(LogIrisFiltering, Warning, TEXT("FReplicationFiltering::SetSubObjectFilterStatus - Trying to set filter for %u, with invalid Connections parameters."), GroupIndex);
 		return;
 	}
 
 	if (GroupHandle.IsReservedNetObjectGroup())
 	{
-		UE_LOG(LogIrisFiltering, Warning, TEXT("FReplicationFiltering::SetSubObjectFilterStatus - Trying to set filter for reserved GroupIndex: %u which is not allowed."), GroupHandle.GetGroupIndex());
+		UE_LOG(LogIrisFiltering, Warning, TEXT("FReplicationFiltering::SetSubObjectFilterStatus - Trying to set filter for reserved GroupIndex: %u which is not allowed."), GroupIndex);
 		return;
 	}
 
-	if (!SubObjectFilterGroups.GetBit(GroupHandle.GetGroupIndex()))
+	if (!SubObjectFilterGroups.GetBit(GroupIndex))
 	{
-		UE_LOG(LogIrisFiltering, Warning, TEXT("FReplicationFiltering::SetSubObjectFilterStatus - Trying to ReplicationStatus for GroupIndex: %u that is not a SubObjectFilterGroup"), GroupHandle.GetGroupIndex());
+		UE_LOG(LogIrisFiltering, Warning, TEXT("FReplicationFiltering::SetSubObjectFilterStatus - Trying to ReplicationStatus for GroupIndex: %u that is not a SubObjectFilterGroup"), GroupIndex);
 		return;
 	}
 
@@ -1610,20 +1745,21 @@ void FReplicationFiltering::SetSubObjectFilterStatus(FNetObjectGroupHandle Group
 
 void FReplicationFiltering::SetSubObjectFilterStatus(FNetObjectGroupHandle GroupHandle, uint32 ConnectionId, ENetFilterStatus ReplicationStatus)
 {
+	const FNetObjectGroupHandle::FGroupIndexType GroupIndex = GroupHandle.GetGroupIndex();
 	if (GroupHandle.IsReservedNetObjectGroup())
 	{
-		UE_LOG(LogIrisFiltering, Warning, TEXT("FReplicationFiltering::SetSubObjectFilterStatus - Trying to set filter for reserved GroupIndex: %u which is not allowed."), GroupHandle.GetGroupIndex());
+		UE_LOG(LogIrisFiltering, Warning, TEXT("FReplicationFiltering::SetSubObjectFilterStatus - Trying to set filter for reserved GroupIndex: %u which is not allowed."), GroupIndex);
 		return;
 	}
 
-	if (ensure(ValidConnections.GetBit(ConnectionId) && SubObjectFilterGroups.GetBit(GroupHandle.GetGroupIndex())))
+	if (ensure(ValidConnections.GetBit(ConnectionId) && SubObjectFilterGroups.GetBit(GroupIndex)))
 	{
 		UE_LOG(LogIrisFiltering, Verbose, TEXT("ReplicationFiltering::SetSubObjectFilterStatus GroupIndex: %u, ConnectionId: %u, FilterStatus: %u"), GroupHandle.GetRawValue(), ConnectionId, ReplicationStatus == ENetFilterStatus::Allow ? 1U : 0U);
-		FPerObjectInfo* FilterInfo = GetPerObjectInfo(GroupInfos[GroupHandle.GetGroupIndex()].ConnectionStateIndex);
+		FPerObjectInfo* FilterInfo = GetPerObjectInfo(GroupInfos[GroupIndex].ConnectionStateIndex);
 		SetConnectionFilterStatus(*FilterInfo, ConnectionId, ReplicationStatus);
-		if (!IsConnectionFilterStatusAllowedForAnyConnection(*FilterInfo))
+		if (!IsAnyConnectionFilterStatusAllowed(*FilterInfo))
 		{
-			DirtySubObjectFilterGroups.SetBit(GroupHandle.GetGroupIndex());
+			DirtySubObjectFilterGroups.SetBit(GroupIndex);
 		}
 	}
 }
@@ -1641,51 +1777,110 @@ bool FReplicationFiltering::GetSubObjectFilterStatus(FNetObjectGroupHandle Group
 	return true;
 }
 
-void FReplicationFiltering::AddGroupFilter(FNetObjectGroupHandle GroupHandle)
+bool FReplicationFiltering::AddExclusionFilterGroup(FNetObjectGroupHandle GroupHandle)
 {
-	const bool bIsValidGroup = ensureMsgf(Groups->IsValidGroup(GroupHandle), TEXT("AddGroupFilter received an invalid group handle"));
-	if (bIsValidGroup)
+	const FNetObjectGroupHandle::FGroupIndexType GroupIndex = GroupHandle.GetGroupIndex();
+	const bool bIsValidGroup = ensureMsgf(Groups->IsValidGroup(GroupHandle), TEXT("AddExclusionFilterGroup received an invalid group handle GroupIndex: %u"), GroupIndex);
+	if (!bIsValidGroup)
 	{
-		const bool bIsFiltering = FilterGroups.GetBit(GroupHandle.GetGroupIndex());
-		ensureMsgf(!bIsFiltering, TEXT("NetObjectGroup %u was asked to start filtering but it was already doing so."), GroupHandle.GetGroupIndex());
+		return false;
+	}
 
-		check(!SubObjectFilterGroups.GetBit(GroupHandle.GetGroupIndex()));
-		
-		if (!bIsFiltering)
-		{
-			Groups->AddFilterTrait(GroupHandle);
+	const bool bIsFiltering = Groups->IsFilterGroup(GroupHandle) || SubObjectFilterGroups.GetBit(GroupIndex);
+	ensureMsgf(!bIsFiltering, TEXT("NetObjectGroup Name '%s' GroupIndex %u was asked to start exclusion filtering but it was already used for filtering."), ToCStr(Groups->GetGroupName(GroupHandle).ToString()), GroupIndex);
+	if (bIsFiltering)
+	{
+		return false;
+	}
 
-			FilterGroups.SetBit(GroupHandle.GetGroupIndex());
-			DirtyFilterGroups.SetBit(GroupHandle.GetGroupIndex());
-			GroupInfos[GroupHandle.GetGroupIndex()].ConnectionStateIndex = AllocPerObjectInfo();
+	Groups->AddExclusionFilterTrait(GroupHandle);
 
-			// By default filter is active, i.e. we filter out all connections
-			SetPerObjectInfoFilterStatus(*GetPerObjectInfo(GroupInfos[GroupHandle.GetGroupIndex()].ConnectionStateIndex), ENetFilterStatus::Disallow);
+	ExclusionFilterGroups.SetBit(GroupIndex);
+	DirtyExclusionFilterGroups.SetBit(GroupIndex);
+	bHasDirtyExclusionFilterGroup = 1;
 
-			UE_LOG(LogIrisFiltering, Verbose, TEXT("ReplicationFiltering::AddGroupFilter GroupIndex: %u, FilterStatus: DisallowReplication"), GroupHandle.GetGroupIndex());
-		}
+	// By default we filter out the group members for all connections
+	GroupInfos[GroupIndex].ConnectionStateIndex = AllocPerObjectInfo();
+	SetPerObjectInfoFilterStatus(*GetPerObjectInfo(GroupInfos[GroupIndex].ConnectionStateIndex), ENetFilterStatus::Disallow);
+
+	UE_LOG(LogIrisFiltering, Verbose, TEXT("ReplicationFiltering::AddExclusionGroupFilter GroupIndex: %u, FilterStatus: DisallowReplication"), GroupIndex);
+	return true;
+}
+
+bool FReplicationFiltering::AddInclusionFilterGroup(FNetObjectGroupHandle GroupHandle)
+{
+	const FNetObjectGroupHandle::FGroupIndexType GroupIndex = GroupHandle.GetGroupIndex();
+	const bool bIsValidGroup = ensureMsgf(Groups->IsValidGroup(GroupHandle), TEXT("AddInclusionFilterGroup received an invalid group handle GroupIndex: %u"), GroupIndex);
+	if (!bIsValidGroup)
+	{
+		return false;
+	}
+
+	const bool bIsFiltering = Groups->IsFilterGroup(GroupHandle) || SubObjectFilterGroups.GetBit(GroupIndex);
+	ensureMsgf(!bIsFiltering, TEXT("NetObjectGroup Name '%s' GroupIndex %u was asked to start exclusion filtering but it was already used for filtering."), ToCStr(Groups->GetGroupName(GroupHandle).ToString()), GroupIndex);
+	if (bIsFiltering)
+	{
+		return false;
 	}
+
+	Groups->AddInclusionFilterTrait(GroupHandle);
+
+	InclusionFilterGroups.SetBit(GroupIndex);
+	DirtyInclusionFilterGroups.SetBit(GroupIndex);
+	bHasDirtyInclusionFilterGroup = 1;
+
+	// By default we do not override dynamic filtering.
+	GroupInfos[GroupIndex].ConnectionStateIndex = AllocPerObjectInfo();
+	SetPerObjectInfoFilterStatus(*GetPerObjectInfo(GroupInfos[GroupIndex].ConnectionStateIndex), ENetFilterStatus::Disallow);
+
+	UE_LOG(LogIrisFiltering, Verbose, TEXT("ReplicationFiltering::AddInclusionFilterGroup GroupIndex: %u, FilterStatus: DoNotOverride"), GroupIndex);
+	return true;
 }
 
 void FReplicationFiltering::RemoveGroupFilter(FNetObjectGroupHandle GroupHandle)
 {
 	// Remove the filter trait from the group
-	if (Groups->IsValidGroup(GroupHandle) && FilterGroups.GetBit(GroupHandle.GetGroupIndex()))
+	if (!Groups->IsValidGroup(GroupHandle))
 	{
-		ValidConnections.ForAllSetBits( [GroupHandle, this](uint32 ConnectionIndex) 
-		{ 
+		return;
+	}
+
+	const FNetObjectGroupHandle::FGroupIndexType GroupIndex = GroupHandle.GetGroupIndex();
+	if (ExclusionFilterGroups.GetBit(GroupIndex))
+	{
+		DirtyExclusionFilterGroups.ClearBit(GroupIndex);
+
+		ValidConnections.ForAllSetBits([GroupHandle, this](uint32 ConnectionIndex) 
+		{
 			this->SetGroupFilterStatus(GroupHandle, ConnectionIndex, ENetFilterStatus::Allow); 
 		});
 
-		// Mark group as no longer a group filter group
-		FilterGroups.ClearBit(GroupHandle.GetGroupIndex());
-		const PerObjectInfoIndexType ConnectionStateIndex = GroupInfos[GroupHandle.GetGroupIndex()].ConnectionStateIndex;
-		GroupInfos[GroupHandle.GetGroupIndex()].ConnectionStateIndex = 0U;
+		// Mark group as no longer an exclusion filter
+		ExclusionFilterGroups.ClearBit(GroupIndex);
+		const PerObjectInfoIndexType ConnectionStateIndex = GroupInfos[GroupIndex].ConnectionStateIndex;
+		GroupInfos[GroupIndex].ConnectionStateIndex = 0U;
+		FreePerObjectInfo(ConnectionStateIndex);
+
+		Groups->RemoveExclusionFilterTrait(GroupHandle);
+
+		UE_LOG(LogIrisFiltering, Verbose, TEXT("ReplicationFiltering::RemoveGroupFilter ExclusionGroup GroupIndex: %u"), GroupHandle.GetGroupIndex());
+	}
+	else if (InclusionFilterGroups.GetBit(GroupIndex))
+	{
+		DirtyInclusionFilterGroups.ClearBit(GroupIndex);
+
+		// Clear filter effects
+		SetGroupFilterStatus(GroupHandle, ENetFilterStatus::Disallow);
+
+		// Mark group as no longer an inclusion filter
+		InclusionFilterGroups.ClearBit(GroupIndex);
+		const PerObjectInfoIndexType ConnectionStateIndex = GroupInfos[GroupIndex].ConnectionStateIndex;
+		GroupInfos[GroupIndex].ConnectionStateIndex = 0U;
 		FreePerObjectInfo(ConnectionStateIndex);
 
-		Groups->RemoveFilterTrait(GroupHandle);
+		Groups->RemoveInclusionFilterTrait(GroupHandle);
 
-		UE_LOG(LogIrisFiltering, Verbose, TEXT("ReplicationFiltering::RemoveGroupFilter GroupIndex: %u"), GroupHandle.GetGroupIndex());
+		UE_LOG(LogIrisFiltering, Verbose, TEXT("ReplicationFiltering::RemoveGroupFilter InclusionGroup GroupIndex: %u"), GroupHandle.GetGroupIndex());
 	}
 }
 
@@ -1693,79 +1888,107 @@ void FReplicationFiltering::SetGroupFilterStatus(FNetObjectGroupHandle GroupHand
 {
 	IRIS_PROFILER_SCOPE(SetGroupFilterStatus)
 
+	const FNetObjectGroupHandle::FGroupIndexType GroupIndex = GroupHandle.GetGroupIndex();
 	if (GroupHandle.IsReservedNetObjectGroup())
 	{
-		UE_LOG(LogIrisFiltering, Warning, TEXT("FReplicationFiltering::SetGroupFilterStatus - Trying to set filter for reserved GroupIndex: %u which is not allowed."), GroupHandle.GetGroupIndex());
+		UE_LOG(LogIrisFiltering, Warning, TEXT("FReplicationFiltering::SetGroupFilterStatus - Trying to set filter for reserved GroupIndex: %u which is not allowed."), GroupIndex);
 		return;
 	}
 
-	if (!FilterGroups.GetBit(GroupHandle.GetGroupIndex()))
+	if (ExclusionFilterGroups.GetBit(GroupIndex))
 	{
-		UE_LOG(LogIrisFiltering, Warning, TEXT("FReplicationFiltering::SetGroupFilterStatus - Trying to set filter for invalid GroupIndex: %u, Make sure group is added to filtering"), GroupHandle.GetGroupIndex());
+		// It is intentional that we iterate over all connections, not only the valid ones as we want the group filter status to be initialized correctly for new connections as well.
+		for (uint32 ConnectionId = 1U; ConnectionId < ValidConnections.GetNumBits(); ++ConnectionId)
+		{
+			InternalSetExclusionGroupFilterStatus(GroupHandle, ConnectionId, ReplicationStatus);
+		}
+
 		return;
 	}
-	
-	// It is intentional that we iterate overall connections, not only the valid ones as we want the group filter status to be initialized correctly for new connections as well.
-	for (uint32 ConnectionId = 1U; ConnectionId < ValidConnections.GetNumBits(); ++ConnectionId)
+
+	if (InclusionFilterGroups.GetBit(GroupIndex))
 	{
-		InternalSetGroupFilterStatus(GroupHandle, ConnectionId, ReplicationStatus);
+		InternalSetInclusionGroupFilterStatus(GroupHandle, ReplicationStatus);
+		return;
 	}
+	
+	UE_LOG(LogIrisFiltering, Warning, TEXT("FReplicationFiltering::SetGroupFilterStatus - Trying to set filter for invalid GroupIndex: %u, Make sure group is added to filtering"), GroupHandle.GetGroupIndex());
 }
 
 void FReplicationFiltering::SetGroupFilterStatus(FNetObjectGroupHandle GroupHandle, const FNetBitArrayView& ConnectionsBitArray, ENetFilterStatus ReplicationStatus)
 {
 	IRIS_PROFILER_SCOPE(SetGroupFilterStatus)
 		
+	const FNetObjectGroupHandle::FGroupIndexType GroupIndex = GroupHandle.GetGroupIndex();
 	if (GroupHandle.IsReservedNetObjectGroup())
 	{
-		UE_LOG(LogIrisFiltering, Warning, TEXT("FReplicationFiltering::SetGroupFilterStatus - Trying to set filter for reserved GroupIndex: %u which is not allowed."), GroupHandle.GetGroupIndex());		
+		UE_LOG(LogIrisFiltering, Warning, TEXT("FReplicationFiltering::SetGroupFilterStatus - Trying to set filter for reserved GroupIndex: %u which is not allowed."), GroupIndex);		
 		return;
 	}
 
-	if (!FilterGroups.GetBit(GroupHandle.GetGroupIndex()))
+	if (ConnectionsBitArray.GetNumBits() > ValidConnections.GetNumBits())
 	{
-		UE_LOG(LogIrisFiltering, Warning, TEXT("FReplicationFiltering::SetGroupFilterStatus - Trying to set filter for invalid GroupIndex: %u, Make sure group is added to filtering"), GroupHandle.GetGroupIndex());
+		UE_LOG(LogIrisFiltering, Warning, TEXT("FReplicationFiltering::SetGroupFilterStatus - Trying to set filter for %u, with invalid Connections parameters."), GroupIndex);
 		return;
 	}
 
-	if (ConnectionsBitArray.GetNumBits() > ValidConnections.GetNumBits())
+	if (ExclusionFilterGroups.GetBit(GroupIndex))
 	{
-		UE_LOG(LogIrisFiltering, Warning, TEXT("FReplicationFiltering::SetGroupFilterStatus - Trying to set filter for %u, with invalid Connections parameters."), GroupHandle.GetGroupIndex());
-		return;
+		// It is intentional that we iterate over all connections, not only the valid ones as we want the group filter status to be initialized correctly for new connections as well.
+		for (uint32 ConnectionId = 1U, EndConnectionId = ValidConnections.GetNumBits(); ConnectionId < EndConnectionId; ++ConnectionId)
+		{
+			ENetFilterStatus ReplicationStatusToSet = ReplicationStatus;
+			if (ConnectionId >= ConnectionsBitArray.GetNumBits() || !ConnectionsBitArray.GetBit(ConnectionId))
+			{
+				ReplicationStatusToSet = ReplicationStatus == ENetFilterStatus::Allow ? ENetFilterStatus::Disallow : ENetFilterStatus::Allow;
+			}
+			InternalSetExclusionGroupFilterStatus(GroupHandle, ConnectionId, ReplicationStatusToSet);
+		}
 	}
-
-	// It is intentional that we iterate overall connections, not only the valid ones as we want the group filter status to be initialized correctly for new connections as well.
-	for (uint32 ConnectionId = 1U; ConnectionId < ValidConnections.GetNumBits(); ++ConnectionId)
+	else if (InclusionFilterGroups.GetBit(GroupIndex))
 	{
-		ENetFilterStatus ReplicationStatusToSet = ReplicationStatus;
-		if (ConnectionId >= ConnectionsBitArray.GetNumBits() || !ConnectionsBitArray.GetBit(ConnectionId))
+		// It is intentional that we iterate over all connections, not only the valid ones as we want the group filter status to be initialized correctly for new connections as well.
+		for (uint32 ConnectionId = 1U, EndConnectionId = ValidConnections.GetNumBits(); ConnectionId < EndConnectionId; ++ConnectionId)
 		{
-			ReplicationStatusToSet = ReplicationStatus == ENetFilterStatus::Allow ? ENetFilterStatus::Disallow : ENetFilterStatus::Allow;
+			ENetFilterStatus ReplicationStatusToSet = ReplicationStatus;
+			if (ConnectionId >= ConnectionsBitArray.GetNumBits() || !ConnectionsBitArray.GetBit(ConnectionId))
+			{
+				ReplicationStatusToSet = ReplicationStatus == ENetFilterStatus::Allow ? ENetFilterStatus::Disallow : ENetFilterStatus::Allow;
+			}
+			InternalSetInclusionGroupFilterStatus(GroupHandle, ConnectionId, ReplicationStatusToSet);
 		}
-		InternalSetGroupFilterStatus(GroupHandle, ConnectionId, ReplicationStatusToSet);
 	}
+
+	UE_LOG(LogIrisFiltering, Warning, TEXT("FReplicationFiltering::SetGroupFilterStatus - Trying to set filter for invalid GroupIndex: %u, Make sure group is added to filtering"), GroupIndex);
 }
 
 void FReplicationFiltering::SetGroupFilterStatus(FNetObjectGroupHandle GroupHandle, uint32 ConnectionId, ENetFilterStatus ReplicationStatus)
 {
 	IRIS_PROFILER_SCOPE(SetGroupFilterStatus)
 
+	const FNetObjectGroupHandle::FGroupIndexType GroupIndex = GroupHandle.GetGroupIndex();
 	if (GroupHandle.IsReservedNetObjectGroup())
 	{
-		UE_LOG(LogIrisFiltering, Warning, TEXT("FReplicationFiltering::SetGroupFilterStatus - Trying to set filter for reserved GroupIndex: %u which is not allowed."), GroupHandle.GetGroupIndex());
+		UE_LOG(LogIrisFiltering, Warning, TEXT("FReplicationFiltering::SetGroupFilterStatus - Trying to set filter for reserved GroupIndex: %u which is not allowed."), GroupIndex);
 		return;
 	}
 
-	if (!FilterGroups.GetBit(GroupHandle.GetGroupIndex()))
+	if (ExclusionFilterGroups.GetBit(GroupIndex))
 	{
-		UE_LOG(LogIrisFiltering, Warning, TEXT("FReplicationFiltering::SetGroupFilterStatus - Trying to set invalid filter ConnectionId: %u, GroupIndex: %u"), ConnectionId, GroupHandle.GetGroupIndex());
+		InternalSetExclusionGroupFilterStatus(GroupHandle, ConnectionId, ReplicationStatus);
 		return;
 	}
 	
-	InternalSetGroupFilterStatus(GroupHandle, ConnectionId, ReplicationStatus);
+	if (InclusionFilterGroups.GetBit(GroupIndex))
+	{
+		InternalSetInclusionGroupFilterStatus(GroupHandle, ConnectionId, ReplicationStatus);
+		return;
+	}
+
+	UE_LOG(LogIrisFiltering, Warning, TEXT("FReplicationFiltering::SetGroupFilterStatus - Trying to set invalid filter ConnectionId: %u, GroupIndex: %u"), ConnectionId, GroupIndex);
 }
 
-bool FReplicationFiltering::GetIsFilteredOutByAnyGroup(uint32 ObjectInternalIndex, uint32 ConnectionId) const
+bool FReplicationFiltering::IsExcludedByAnyGroup(uint32 ObjectInternalIndex, uint32 ConnectionId) const
 {
 	uint32 GroupMembershipCount = 0U;
 	const FNetObjectGroupHandle* GroupHandles = Groups->GetGroupMemberships(ObjectInternalIndex, GroupMembershipCount);
@@ -1775,11 +1998,12 @@ bool FReplicationFiltering::GetIsFilteredOutByAnyGroup(uint32 ObjectInternalInde
 		return false;
 	}
 
-	for (FNetObjectGroupHandle GroupHandle : MakeArrayView(GroupHandles, GroupMembershipCount))
+	for (const FNetObjectGroupHandle GroupHandle : MakeArrayView(GroupHandles, GroupMembershipCount))
 	{
-		if (FilterGroups.GetBit(GroupHandle.GetGroupIndex()))
+		const FNetObjectGroupHandle::FGroupIndexType GroupIndex = GroupHandle.GetGroupIndex();
+		if (ExclusionFilterGroups.GetBit(GroupIndex))
 		{
-			const FPerObjectInfo* ConnectionState = GetPerObjectInfo(GroupInfos[GroupHandle.GetGroupIndex()].ConnectionStateIndex);
+			const FPerObjectInfo* ConnectionState = GetPerObjectInfo(GroupInfos[GroupIndex].ConnectionStateIndex);
 			if (GetConnectionFilterStatus(*ConnectionState, ConnectionId) == ENetFilterStatus::Disallow)
 			{
 				return true;
@@ -1790,27 +2014,53 @@ bool FReplicationFiltering::GetIsFilteredOutByAnyGroup(uint32 ObjectInternalInde
 	return false;
 }
 
-bool FReplicationFiltering::UpdateGroupFilterEffectsForObject(uint32 ObjectIndex, uint32 ConnectionId)
+bool FReplicationFiltering::IsIncludedByAnyGroup(uint32 ObjectInternalIndex, uint32 ConnectionId) const
+{
+	uint32 GroupMembershipCount = 0U;
+	const FNetObjectGroupHandle* GroupHandles = Groups->GetGroupMemberships(ObjectInternalIndex, GroupMembershipCount);
+
+	if (!GroupHandles)
+	{
+		return false;
+	}
+
+	for (const FNetObjectGroupHandle GroupHandle : MakeArrayView(GroupHandles, GroupMembershipCount))
+	{
+		const FNetObjectGroupHandle::FGroupIndexType GroupIndex = GroupHandle.GetGroupIndex();
+		if (InclusionFilterGroups.GetBit(GroupIndex))
+		{
+			const FPerObjectInfo* ConnectionState = GetPerObjectInfo(GroupInfos[GroupIndex].ConnectionStateIndex);
+			if (GetConnectionFilterStatus(*ConnectionState, ConnectionId) == ENetFilterStatus::Allow)
+			{
+				return true;
+			}
+		}
+	}
+
+	return false;
+}
+
+bool FReplicationFiltering::ClearGroupExclusionFilterEffectsForObject(uint32 ObjectIndex, uint32 ConnectionId)
 {
 	FPerConnectionInfo& ConnectionInfo = ConnectionInfos[ConnectionId];
-	FNetBitArray& GroupFilteredOutObjects = ConnectionInfo.GroupFilteredOutObjects;
+	FNetBitArray& GroupExcludedObjects = ConnectionInfo.GroupExcludedObjects;
 	FNetBitArray& ConnectionFilteredObjects = ConnectionInfo.ConnectionFilteredObjects;
 	FNetBitArray& ObjectsInScopeBeforeDynamicFiltering = ConnectionInfo.ObjectsInScopeBeforeDynamicFiltering;
 
-	auto ClearGroupFilterForObject = [&GroupFilteredOutObjects, &ConnectionFilteredObjects, &ObjectsInScopeBeforeDynamicFiltering](uint32 InternalObjectIndex)
+	auto ClearGroupFilterForObject = [&GroupExcludedObjects, &ConnectionFilteredObjects, &ObjectsInScopeBeforeDynamicFiltering](uint32 InternalObjectIndex)
 	{
-		GroupFilteredOutObjects.ClearBit(InternalObjectIndex);
+		GroupExcludedObjects.ClearBit(InternalObjectIndex);
 		ObjectsInScopeBeforeDynamicFiltering.SetBitValue(InternalObjectIndex, ConnectionFilteredObjects.GetBit(InternalObjectIndex));
 	};
 
-	if (!GetIsFilteredOutByAnyGroup(ObjectIndex, ConnectionId))
+	if (!IsExcludedByAnyGroup(ObjectIndex, ConnectionId))
 	{
 		ClearGroupFilterForObject(ObjectIndex);
 	
 		// Filter subobjects
 		for (const FInternalNetRefIndex SubObjectIndex : NetRefHandleManager->GetSubObjects(ObjectIndex))
 		{
-			if (!GetIsFilteredOutByAnyGroup(SubObjectIndex, ConnectionId))
+			if (!IsExcludedByAnyGroup(SubObjectIndex, ConnectionId))
 			{
 				ClearGroupFilterForObject(SubObjectIndex);
 			}
@@ -1822,55 +2072,70 @@ bool FReplicationFiltering::UpdateGroupFilterEffectsForObject(uint32 ObjectIndex
 	return false;
 }
 
-void FReplicationFiltering::InternalSetGroupFilterStatus(FNetObjectGroupHandle GroupHandle, uint32 ConnectionId, ENetFilterStatus ReplicationStatus)
+bool FReplicationFiltering::ClearGroupInclusionFilterEffectsForObject(uint32 ObjectIndex, uint32 ConnectionId)
 {
-	if (GroupHandle.IsNotReplicatedNetObjectGroup())
+	FPerConnectionInfo& ConnectionInfo = ConnectionInfos[ConnectionId];
+
+	// Ignore subobjects. Inclusion groups overrides dynamic filters which only operate on root objects. SubObjects need to be filtered like the root.
+	if (NetRefHandleManager->GetSubObjectInternalIndices().GetBit(ObjectIndex))
 	{
-		UE_LOG(LogIrisFiltering, Warning, TEXT("FReplicationFiltering::SetGroupFilterStatus - Trying to set filter statua on NotReplicatedNetObjectGroupHandle GroupIndex: %u which is not allowed."), GroupHandle.GetGroupIndex());		
-		return;
+		return false;
 	}
 
-	if (!FilterGroups.GetBit(GroupHandle.GetGroupIndex()))
+	if (!IsIncludedByAnyGroup(ObjectIndex, ConnectionId))
 	{
-		UE_LOG(LogIrisFiltering, Warning, TEXT("FReplicationFiltering::SetGroupFilterStatus - Trying to set invalid filter ConnectionId: %u, GroupIndex: %u"), ConnectionId, GroupHandle.GetGroupIndex());
-		return;
+		FNetBitArray& GroupIncludedObjects = ConnectionInfo.GroupIncludedObjects;
+		
+		GroupIncludedObjects.ClearBit(ObjectIndex);
+
+		for (const FInternalNetRefIndex SubObjectIndex : NetRefHandleManager->GetSubObjects(ObjectIndex))
+		{
+			GroupIncludedObjects.ClearBit(SubObjectIndex);
+		}
+		
+		return true;
 	}
 
+	return false;
+}
+
+void FReplicationFiltering::InternalSetExclusionGroupFilterStatus(FNetObjectGroupHandle GroupHandle, uint32 ConnectionId, ENetFilterStatus ReplicationStatus)
+{
 	if (ReplicationStatus == ENetFilterStatus::Disallow)
 	{
 		// Filter out objects in group
 		FPerObjectInfo* ConnectionState = GetPerObjectInfo(GroupInfos[GroupHandle.GetGroupIndex()].ConnectionStateIndex);
 		if (GetConnectionFilterStatus(*ConnectionState, ConnectionId) == ENetFilterStatus::Allow)
 		{
-			IRIS_PROFILER_SCOPE(FReplicationFiltering_InternalSetGroupFilterStatus_Disallow);
+			IRIS_PROFILER_SCOPE(FReplicationFiltering_InternalSetExclusionGroupFilterStatus_Disallow);
 
 			// Mark filter as active for the connection
 			SetConnectionFilterStatus(*ConnectionState, ConnectionId, ENetFilterStatus::Disallow);
 
-			UE_LOG(LogIrisFiltering, Verbose, TEXT("ReplicationFiltering::SetGroupFilterStatus GroupIndex: %u, ConnectionId: %u, FilterStatus: DisallowReplication"), GroupHandle.GetGroupIndex(), ConnectionId);
+			UE_LOG(LogIrisFiltering, Verbose, TEXT("ReplicationFiltering::SetGroupFilterStatus ExclusionGroup GroupIndex: %u, ConnectionId: %u, FilterStatus: DisallowReplication"), GroupHandle.GetGroupIndex(), ConnectionId);
 
 			if (ValidConnections.GetBit(ConnectionId))
 			{
-				// New Connections will be initialized the next filter update
+				// New Connections will be initialized in the next filter update
 				if (!NewConnections.GetBit(ConnectionId))
 				{
 					const FNetObjectGroup* Group = Groups->GetGroup(GroupHandle);
 			
 					// Update group filter mask for the connection
 					FPerConnectionInfo& ConnectionInfo = ConnectionInfos[ConnectionId];
-					FNetBitArray& GroupFilteredOutObjects = ConnectionInfo.GroupFilteredOutObjects;
+					FNetBitArray& GroupExcludedObjects = ConnectionInfo.GroupExcludedObjects;
 					FNetBitArray& ObjectsInScopeBeforeDynamicFiltering = ConnectionInfo.ObjectsInScopeBeforeDynamicFiltering;
 					const FNetBitArrayView GlobalScopableObjects = NetRefHandleManager->GetGlobalScopableInternalIndices();
-					for (uint32 ObjectIndex : MakeArrayView(Group->Members.GetData(), Group->Members.Num()))
+					for (const FInternalNetRefIndex ObjectIndex : Group->Members)
 					{
-						GroupFilteredOutObjects.SetBit(ObjectIndex);
+						GroupExcludedObjects.SetBit(ObjectIndex);
 
 						//$IRIS TODO: ObjectsInScopeBeforeDynamicFiltering should not be accessed outside PreSendUpdate since its reset there. Is this needed ?
 						ObjectsInScopeBeforeDynamicFiltering.ClearBit(ObjectIndex);
 						// Filter subobjects
 						for (const FInternalNetRefIndex SubObjectIndex : NetRefHandleManager->GetSubObjects(ObjectIndex))
 						{
-							GroupFilteredOutObjects.SetBitValue(SubObjectIndex, GlobalScopableObjects.GetBit(SubObjectIndex));
+							GroupExcludedObjects.SetBitValue(SubObjectIndex, GlobalScopableObjects.GetBit(SubObjectIndex));
 							ObjectsInScopeBeforeDynamicFiltering.ClearBit(SubObjectIndex);
 						}
 					}
@@ -1884,12 +2149,12 @@ void FReplicationFiltering::InternalSetGroupFilterStatus(FNetObjectGroupHandle G
 		FPerObjectInfo* ConnectionState = GetPerObjectInfo(GroupInfos[GroupHandle.GetGroupIndex()].ConnectionStateIndex);
 		if (GetConnectionFilterStatus(*ConnectionState, ConnectionId) == ENetFilterStatus::Disallow)
 		{
-			IRIS_PROFILER_SCOPE(FReplicationFiltering_InternalSetGroupFilterStatus_Allow);
+			IRIS_PROFILER_SCOPE(FReplicationFiltering_InternalSetExclusionGroupFilterStatus_Allow);
 
 			// Mark filter as no longer being active for the connection
 			SetConnectionFilterStatus(*ConnectionState, ConnectionId, ENetFilterStatus::Allow);
 
-			UE_LOG(LogIrisFiltering, Verbose, TEXT("ReplicationFiltering::SetGroupFilterStatus GroupIndex: %u, ConnectionId: %u, FilterStatus: AllowReplication"), GroupHandle.GetGroupIndex(), ConnectionId);
+			UE_LOG(LogIrisFiltering, Verbose, TEXT("ReplicationFiltering::SetGroupFilterStatus ExclusionGroup GroupIndex: %u, ConnectionId: %u, FilterStatus: AllowReplication"), GroupHandle.GetGroupIndex(), ConnectionId);
 
 			if (ValidConnections.GetBit(ConnectionId))
 			{
@@ -1899,9 +2164,9 @@ void FReplicationFiltering::InternalSetGroupFilterStatus(FNetObjectGroupHandle G
 					const FNetObjectGroup* Group = Groups->GetGroup(GroupHandle);
 
 					// Update group filter mask for the connection
-					for (uint32 ObjectIndex : MakeArrayView(Group->Members.GetData(), Group->Members.Num()))
+					for (const FInternalNetRefIndex ObjectIndex : Group->Members)
 					{
-						UpdateGroupFilterEffectsForObject(ObjectIndex, ConnectionId);
+						ClearGroupExclusionFilterEffectsForObject(ObjectIndex, ConnectionId);
 					}
 				}
 			}
@@ -1909,9 +2174,84 @@ void FReplicationFiltering::InternalSetGroupFilterStatus(FNetObjectGroupHandle G
 	}
 }
 
+// Set same status for all connections
+void FReplicationFiltering::InternalSetInclusionGroupFilterStatus(FNetObjectGroupHandle GroupHandle, ENetFilterStatus ReplicationStatus)
+{
+	FPerObjectInfo* ConnectionState = GetPerObjectInfo(GroupInfos[GroupHandle.GetGroupIndex()].ConnectionStateIndex);
+
+	// If status is intact we can avoid expensive processing.
+	if (ReplicationStatus == ENetFilterStatus::Disallow && !IsAnyConnectionFilterStatusAllowed(*ConnectionState))
+	{
+		return;
+	}
+
+	auto UpdateFilterStatusForConnection = [this, GroupHandle, ReplicationStatus](uint32 ConnectionId)
+	{
+		InternalSetInclusionGroupFilterStatus(GroupHandle, ConnectionId, ReplicationStatus);
+	};
+
+	// Update filter status for all initialized connections
+	FNetBitArray::ForAllSetBits(ValidConnections, NewConnections, FNetBitArrayBase::AndNotOp, UpdateFilterStatusForConnection);
+
+	// Make sure status is set for all connections.
+	SetPerObjectInfoFilterStatus(*ConnectionState, ReplicationStatus);
+}
+
+// Set status for a single connection
+void FReplicationFiltering::InternalSetInclusionGroupFilterStatus(FNetObjectGroupHandle GroupHandle, uint32 ConnectionId, ENetFilterStatus ReplicationStatus)
+{
+	FPerObjectInfo* ConnectionState = GetPerObjectInfo(GroupInfos[GroupHandle.GetGroupIndex()].ConnectionStateIndex);
+	// If status remains there's nothing to do.
+	if (GetConnectionFilterStatus(*ConnectionState, ConnectionId) == ReplicationStatus)
+	{
+		return;
+	}
+
+	SetConnectionFilterStatus(*ConnectionState, ConnectionId, ReplicationStatus);
+
+	// Ignore processing of not yet initialized connections
+	if (!ValidConnections.GetBit(ConnectionId) || NewConnections.GetBit(ConnectionId))
+	{
+		return;
+	}
+
+	const FNetObjectGroup* Group = Groups->GetGroup(GroupHandle);
+	if (ReplicationStatus == ENetFilterStatus::Disallow)
+	{
+		UE_LOG(LogIrisFiltering, Verbose, TEXT("ReplicationFiltering::SetGroupFilterStatus InclusionGroup GroupIndex: %u, ConnectionId: %u, FilterStatus: DisallowReplication"), GroupHandle.GetGroupIndex(), ConnectionId);
+		for (const FInternalNetRefIndex ObjectIndex : Group->Members)
+		{
+			ClearGroupInclusionFilterEffectsForObject(ObjectIndex, ConnectionId);
+		}
+	}
+	else
+	{
+		UE_LOG(LogIrisFiltering, Verbose, TEXT("ReplicationFiltering::SetGroupFilterStatus InclusionGroup GroupIndex: %u, ConnectionId: %u, FilterStatus: AllowReplication"), GroupHandle.GetGroupIndex(), ConnectionId);
+
+		FPerConnectionInfo& ConnectionInfo = ConnectionInfos[ConnectionId];
+		FNetBitArray& GroupIncludedObjects = ConnectionInfo.GroupIncludedObjects;
+		const FNetBitArrayView GlobalScopableObjects = NetRefHandleManager->GetGlobalScopableInternalIndices();
+		const FNetBitArrayView SubObjectInternalIndices = NetRefHandleManager->GetSubObjectInternalIndicesView();
+		for (const FInternalNetRefIndex ObjectIndex : Group->Members)
+		{
+			// SubObjects follow root object.
+			if (SubObjectInternalIndices.GetBit(ObjectIndex))
+			{
+				continue;
+			}
+
+			GroupIncludedObjects.SetBitValue(ObjectIndex, GlobalScopableObjects.GetBit(ObjectIndex));
+			for (const FInternalNetRefIndex SubObjectIndex : NetRefHandleManager->GetSubObjects(ObjectIndex))
+			{
+				GroupIncludedObjects.SetBitValue(SubObjectIndex, GlobalScopableObjects.GetBit(SubObjectIndex));
+			}
+		}
+	}
+}
+
 bool FReplicationFiltering::GetGroupFilterStatus(FNetObjectGroupHandle GroupHandle, uint32 ConnectionId, ENetFilterStatus& OutReplicationStatus) const
 {
-	if (!(ValidConnections.GetBit(ConnectionId) && FilterGroups.GetBit(GroupHandle.GetGroupIndex())))
+	if (!(ValidConnections.GetBit(ConnectionId) && ExclusionFilterGroups.GetBit(GroupHandle.GetGroupIndex())))
 	{
 		return false;
 	}
@@ -1924,42 +2264,81 @@ bool FReplicationFiltering::GetGroupFilterStatus(FNetObjectGroupHandle GroupHand
 
 void FReplicationFiltering::NotifyObjectAddedToGroup(FNetObjectGroupHandle GroupHandle, FInternalNetRefIndex ObjectIndex)
 {
-	// As we want to avoid the fundamentally bad case of doing tons of single object adds to a group and updating group filter for every call we deffer the update until next call to filter
-	if (FilterGroups.GetBit(GroupHandle.GetGroupIndex()))
+	// As we want to avoid the fundamentally bad case of doing tons of single object adds to a group and updating group filter for every call we defer the update until next call to filter
+	const FNetObjectGroupHandle::FGroupIndexType GroupIndex = GroupHandle.GetGroupIndex();
+	if (SubObjectFilterGroups.GetBit(GroupIndex))
 	{
-		DirtyFilterGroups.SetBit(GroupHandle.GetGroupIndex());
+		UE_LOG(LogIrisFiltering, Verbose, TEXT("ReplicationFiltering::NotifyObjectAddedToGroup Added %s to SubObjectFilter with GroupIndex: %u"), *(NetRefHandleManager->GetReplicatedObjectDataNoCheck(ObjectIndex).RefHandle.ToString()), GroupIndex);
+	}
+	else if (ExclusionFilterGroups.GetBit(GroupIndex))
+	{
+		const FPerObjectInfo* ConnectionState = GetPerObjectInfo(GroupInfos[GroupHandle.GetGroupIndex()].ConnectionStateIndex);
+
+		// Avoid expensive filter updates if possible.
+		if (IsAnyConnectionFilterStatusDisallowed(*ConnectionState))
+		{
+			DirtyExclusionFilterGroups.SetBit(GroupIndex);
+			bHasDirtyExclusionFilterGroup = 1;
+		}
 
-		UE_LOG(LogIrisFiltering, Verbose, TEXT("ReplicationFiltering::NotifyObjectAddedToGroup Added %s to Groupfilter with GroupIndex: %u"), *(NetRefHandleManager->GetReplicatedObjectDataNoCheck(ObjectIndex).RefHandle.ToString()), GroupHandle.GetGroupIndex());
+		UE_LOG(LogIrisFiltering, Verbose, TEXT("ReplicationFiltering::NotifyObjectAddedToGroup Added %s to ExclusionGroup filter with GroupIndex: %u"), *(NetRefHandleManager->GetReplicatedObjectDataNoCheck(ObjectIndex).RefHandle.ToString()), GroupIndex);
 	}
-	else if (SubObjectFilterGroups.GetBit(GroupHandle.GetGroupIndex()))
+	else if (InclusionFilterGroups.GetBit(GroupIndex))
 	{
-		UE_LOG(LogIrisFiltering, Verbose, TEXT("ReplicationFiltering::NotifyObjectAddedToGroup Added %s to SubObjectFilter with GroupIndex: %u"), *(NetRefHandleManager->GetReplicatedObjectDataNoCheck(ObjectIndex).RefHandle.ToString()), GroupHandle.GetGroupIndex());
+		const FPerObjectInfo* ConnectionState = GetPerObjectInfo(GroupInfos[GroupHandle.GetGroupIndex()].ConnectionStateIndex);
+
+		// Avoid expensive filter updates if possible.
+		if (IsAnyConnectionFilterStatusAllowed(*ConnectionState))
+		{
+			DirtyInclusionFilterGroups.SetBit(GroupIndex);
+			bHasDirtyInclusionFilterGroup = 1;
+		}
+
+		UE_LOG(LogIrisFiltering, Verbose, TEXT("ReplicationFiltering::NotifyObjectAddedToGroup Added %s to InclusionGroup filter with GroupIndex: %u"), *(NetRefHandleManager->GetReplicatedObjectDataNoCheck(ObjectIndex).RefHandle.ToString()), GroupIndex);
 	}
 }
 
 void FReplicationFiltering::NotifyObjectRemovedFromGroup(FNetObjectGroupHandle GroupHandle, FInternalNetRefIndex ObjectIndex)
 {
-	UE_LOG(LogIrisFiltering, Verbose, TEXT("ReplicationFiltering::NotifyObjectRemovedFromGroup GroupIndex: %u, %s"), GroupHandle.GetRawValue(), *(NetRefHandleManager->GetReplicatedObjectDataNoCheck(ObjectIndex).RefHandle.ToString()));
+	const FNetObjectGroupHandle::FGroupIndexType GroupIndex = GroupHandle.GetGroupIndex();
+
+	UE_LOG(LogIrisFiltering, Verbose, TEXT("ReplicationFiltering::NotifyObjectRemovedFromGroup Removing %s from GroupIndex: %u"), ToCStr(NetRefHandleManager->GetReplicatedObjectDataNoCheck(ObjectIndex).RefHandle.ToString()), GroupIndex);
 
-	if (FilterGroups.GetBit(GroupHandle.GetGroupIndex()))
+	if (SubObjectFilterGroups.GetBit(GroupIndex))
 	{
-		FPerObjectInfo* ConnectionState = GetPerObjectInfo(GroupInfos[GroupHandle.GetGroupIndex()].ConnectionStateIndex);
+		DirtySubObjectFilterGroups.SetBit(GroupIndex);
+	}
+	else if (ExclusionFilterGroups.GetBit(GroupIndex))
+	{
+		const FPerObjectInfo* ConnectionState = GetPerObjectInfo(GroupInfos[GroupIndex].ConnectionStateIndex);
 
-		auto MarkGroupsDirty = [this, ObjectIndex, GroupHandle, ConnectionState](uint32 ConnectionId)
+		auto UpdateGroupEffects = [this, ObjectIndex, GroupHandle, ConnectionState](uint32 ConnectionId)
 		{
 			if (GetConnectionFilterStatus(*ConnectionState, ConnectionId) == ENetFilterStatus::Disallow)
 			{
-				FNetBitArray& GroupFilteredOutObjects = this->ConnectionInfos[ConnectionId].GroupFilteredOutObjects;
-				UpdateGroupFilterEffectsForObject(ObjectIndex, ConnectionId);
+				FNetBitArray& GroupExcludedObjects = this->ConnectionInfos[ConnectionId].GroupExcludedObjects;
+				ClearGroupExclusionFilterEffectsForObject(ObjectIndex, ConnectionId);
 			}
 		};
 
-		// We need to update groupfilter status for all initialized connections
-		FNetBitArray::ForAllSetBits(ValidConnections, NewConnections, FNetBitArrayBase::AndNotOp, MarkGroupsDirty);
+		// We need to update filter status for all initialized connections
+		FNetBitArray::ForAllSetBits(ValidConnections, NewConnections, FNetBitArrayBase::AndNotOp, UpdateGroupEffects);
 	}
-	else if (SubObjectFilterGroups.GetBit(GroupHandle.GetGroupIndex()))
-	{		
-		DirtySubObjectFilterGroups.SetBit(GroupHandle.GetGroupIndex());
+	else if (InclusionFilterGroups.GetBit(GroupIndex))
+	{
+		const FPerObjectInfo* ConnectionState = GetPerObjectInfo(GroupInfos[GroupIndex].ConnectionStateIndex);
+
+		auto UpdateGroupEffects = [this, ObjectIndex, GroupHandle, ConnectionState](uint32 ConnectionId)
+		{
+			if (GetConnectionFilterStatus(*ConnectionState, ConnectionId) == ENetFilterStatus::Allow)
+			{
+				FNetBitArray& GroupIncludedObjects = this->ConnectionInfos[ConnectionId].GroupIncludedObjects;
+				ClearGroupInclusionFilterEffectsForObject(ObjectIndex, ConnectionId);
+			}
+		};
+
+		// We need to update filter status for all initialized connections
+		FNetBitArray::ForAllSetBits(ValidConnections, NewConnections, FNetBitArrayBase::AndNotOp, UpdateGroupEffects);
 	}
 }
 
@@ -1978,17 +2357,18 @@ ENetFilterStatus FReplicationFiltering::GetConnectionFilterStatus(const FPerObje
 	return (ObjectInfo.ConnectionIds[ConnectionId >> 5U] & (1U << (ConnectionId & 31U))) != 0U ? ENetFilterStatus::Allow : ENetFilterStatus::Disallow;
 }
 
-bool FReplicationFiltering::IsConnectionFilterStatusAllowedForAnyConnection(const FPerObjectInfo& ObjectInfo) const
+bool FReplicationFiltering::IsAnyConnectionFilterStatusAllowed(const FPerObjectInfo& ObjectInfo) const
 {
-	for (uint32 WordIt = 0, WordEndIt = PerObjectInfoStorageCountPerItem; WordIt != WordEndIt; ++WordIt)
-	{
-		if (ObjectInfo.ConnectionIds[WordIt] != 0U)
-		{
-			return true;
-		}
-	}
+	static_assert(uint32(ENetFilterStatus::Disallow) == 0U && uint32(ENetFilterStatus::Allow) == 1U, "Inappropriate assumptions regarding ENetFilterStatus values");
+	constexpr uint32 FirstValidConnectionIndex = 1U;
+	return FNetBitArrayView(const_cast<FNetBitArrayView::StorageWordType*>(ObjectInfo.ConnectionIds), ValidConnections.GetNumBits(), FNetBitArrayBase::NoResetNoValidate).FindFirstOne(FirstValidConnectionIndex) != FNetBitArrayBase::InvalidIndex;
+}
 
-	return false;
+bool FReplicationFiltering::IsAnyConnectionFilterStatusDisallowed(const FPerObjectInfo& ObjectInfo) const
+{
+	static_assert(uint32(ENetFilterStatus::Disallow) == 0U && uint32(ENetFilterStatus::Allow) == 1U, "Inappropriate assumptions regarding ENetFilterStatus values");
+	constexpr uint32 FirstValidConnectionIndex = 1U;
+	return FNetBitArrayView(const_cast<FNetBitArrayView::StorageWordType*>(ObjectInfo.ConnectionIds), ValidConnections.GetNumBits(), FNetBitArrayBase::NoResetNoValidate).FindFirstZero(FirstValidConnectionIndex) != FNetBitArrayBase::InvalidIndex;
 }
 
 void FReplicationFiltering::SetConnectionFilterStatus(FPerObjectInfo& ObjectInfo, uint32 ConnectionId, ENetFilterStatus ReplicationStatus)
diff --git a/Engine/Source/Runtime/Experimental/Iris/Core/Private/Iris/ReplicationSystem/Filtering/ReplicationFiltering.h b/Engine/Source/Runtime/Experimental/Iris/Core/Private/Iris/ReplicationSystem/Filtering/ReplicationFiltering.h
--- a/Engine/Source/Runtime/Experimental/Iris/Core/Private/Iris/ReplicationSystem/Filtering/ReplicationFiltering.h
+++ b/Engine/Source/Runtime/Experimental/Iris/Core/Private/Iris/ReplicationSystem/Filtering/ReplicationFiltering.h
@@ -70,10 +70,7 @@ public:
 	 */
 	void FilterPostPoll();
 
-	/** 
-	* Returns the list of objects relevant to a given connection.
-	* This represents the global scope list minus the objects that were filtered out for the given connection.
-	*/
+	/**  Returns the list of objects relevant to a given connection. This represents the global scope list minus the objects that were filtered out for the given connection. */
 	const FNetBitArrayView GetRelevantObjectsInScope(uint32 ConnectionId) const
 	{
 		return MakeNetBitArrayView(ConnectionInfos[ConnectionId].ObjectsInScope);
@@ -81,14 +78,13 @@ public:
 
 	const FNetBitArrayView GetGroupFilteredOutObjects(uint32 ConnectionId) const
 	{
-		return MakeNetBitArrayView(ConnectionInfos[ConnectionId].GroupFilteredOutObjects);
+		return MakeNetBitArrayView(ConnectionInfos[ConnectionId].GroupExcludedObjects);
 	}
 
-	// Who owns what?
+	// Who owns what
 	void SetOwningConnection(FInternalNetRefIndex ObjectIndex, uint32 ConnectionId);
 	uint32 GetOwningConnection(FInternalNetRefIndex ObjectIndex) const { return !bHasDirtyOwner ? ObjectIndexToOwningConnection[ObjectIndex] : GetOwningConnectionIfDirty(ObjectIndex); }
 
-	// Various filters
 	bool SetFilter(FInternalNetRefIndex ObjectIndex, FNetObjectFilterHandle Filter);
 
 	// Set whether an object is allowed to be replicated to certain connections or not.
@@ -105,9 +101,12 @@ public:
 	void RemoveConnection(uint32 ConnectionId);
 
 	// Group based filtering
-	void AddGroupFilter(FNetObjectGroupHandle GroupHandle);
+	bool AddExclusionFilterGroup(FNetObjectGroupHandle GroupHandle);
+	bool AddInclusionFilterGroup(FNetObjectGroupHandle GroupHandle);
+
 	void RemoveGroupFilter(FNetObjectGroupHandle GroupHandle);
-	bool IsGroupFilterGroup(FNetObjectGroupHandle GroupHandle) const { return GroupHandle.IsValid() && FilterGroups.GetBit(GroupHandle.GetGroupIndex()); }
+	bool IsExclusionFilterGroup(FNetObjectGroupHandle GroupHandle) const { return GroupHandle.IsValid() && ExclusionFilterGroups.GetBit(GroupHandle.GetGroupIndex()); }
+	bool IsInclusionFilterGroup(FNetObjectGroupHandle GroupHandle) const { return GroupHandle.IsValid() && InclusionFilterGroups.GetBit(GroupHandle.GetGroupIndex()); }
 
 	void SetGroupFilterStatus(FNetObjectGroupHandle GroupHandle, ENetFilterStatus ReplicationStatus);
 	void SetGroupFilterStatus(FNetObjectGroupHandle GroupHandle, const FNetBitArrayView& ConnectionsBitArray, ENetFilterStatus);
@@ -135,17 +134,19 @@ private:
 	{
 		// Objects filtered depending on owning connection or user set connection filtering
 		FNetBitArray ConnectionFilteredObjects;
-		// Objects filtered out due to one or more groups it belongs to is filtered out
-		FNetBitArray GroupFilteredOutObjects;
-		// Connection and group filtering is assumed to happen seldom. Avoid recalculating from scratch every frame.
+		// Objects filtered out due to one or more exclusion groups it belongs to is filtered out
+		FNetBitArray GroupExcludedObjects;
+		// Connection and group exclusion filtering is assumed to happen seldom. Avoid recalculating from scratch every frame.
 		FNetBitArray ObjectsInScopeBeforeDynamicFiltering;
+		// Group inclusion filtering is assumed to happen seldom. Applied after dynamic filtering, but before dependent objects.
+		FNetBitArray GroupIncludedObjects;
 		// Objects in scope after all kinds of filtering, including dynamic filtering, has been applied
 		FNetBitArray ObjectsInScope;
 
-		// Which objects are filtered out by dynamic filters.
+		// Which objects are filtered out after dynamic filters and inclusion groups have been processed.
 		FNetBitArray DynamicFilteredOutObjects;
 
-		// List of objects currently filtered out by a either dynamic filter passes
+		// List of objects currently filtered out after processing dynamic filter passes and inclusion groups
 		FNetBitArray InProgressDynamicFilteredOutObjects;
 	};
 
@@ -185,7 +186,8 @@ private:
 	void ResetRemovedConnections();
 	void UpdateObjectsInScope();
 	void UpdateOwnerAndConnectionFiltering();
-	void UpdateGroupFiltering();
+	void UpdateGroupExclusionFiltering();
+	void UpdateGroupInclusionFiltering();
 	void UpdateSubObjectFilters();
 
 	void UpdateDynamicFilters(ENetFilterType FilterPass);
@@ -200,8 +202,10 @@ private:
 	bool HasRawFilters() const;
 	bool HasFragmentFilters() const;
 
-	// Helper to update and reset group filter effects if objects are removed from a filter or after a filter status change, returns true if the group filter was changed
-	bool UpdateGroupFilterEffectsForObject(uint32 ObjectIndex, uint32 ConnectionId);
+	// Helper to update and reset group exclusion filter effects if objects are removed from a filter or after a filter status change, returns true if the group filter was changed
+	bool ClearGroupExclusionFilterEffectsForObject(uint32 ObjectIndex, uint32 ConnectionId);
+	// Helper to update and reset group inclusion filter effects if objects are removed from a filter or after a filter status change, returns true if the group filter was changed
+	bool ClearGroupInclusionFilterEffectsForObject(uint32 ObjectIndex, uint32 ConnectionId);
 	bool HasOwnerFilter(uint32 ObjectIndex) const;
 	bool HasConnectionFilter(uint32 ObjectIndex) const;
 
@@ -217,10 +221,17 @@ private:
 	void SetPerObjectInfoFilterStatus(FPerObjectInfo& ObjectInfo, ENetFilterStatus ReplicationStatus);
 
 	ENetFilterStatus GetConnectionFilterStatus(const FPerObjectInfo& ObjectInfo, uint32 ConnectionId) const;
-	bool IsConnectionFilterStatusAllowedForAnyConnection(const FPerObjectInfo& ObjectInfo) const;
+	bool IsAnyConnectionFilterStatusAllowed(const FPerObjectInfo& ObjectInfo) const;
+	bool IsAnyConnectionFilterStatusDisallowed(const FPerObjectInfo& ObjectInfo) const;
 	void SetConnectionFilterStatus(FPerObjectInfo& ObjectInfo, uint32 ConnectionId, ENetFilterStatus ReplicationStatus);
-	bool GetIsFilteredOutByAnyGroup(uint32 ObjectInternalIndex, uint32 ConnectionId) const;
-	void InternalSetGroupFilterStatus(FNetObjectGroupHandle GroupHandle, uint32 ConnectionId, ENetFilterStatus ReplicationStatus);
+	// Returns true if object is filtered out by any exclusion group.
+	bool IsExcludedByAnyGroup(uint32 ObjectInternalIndex, uint32 ConnectionId) const;
+	// Returns true if object is allowed to replicate by any inclusion group.
+	bool IsIncludedByAnyGroup(uint32 ObjectInternalIndex, uint32 ConnectionId) const;
+	void InternalSetExclusionGroupFilterStatus(FNetObjectGroupHandle GroupHandle, uint32 ConnectionId, ENetFilterStatus ReplicationStatus);
+	void InternalSetInclusionGroupFilterStatus(FNetObjectGroupHandle GroupHandle, ENetFilterStatus ReplicationStatus);
+	void InternalSetInclusionGroupFilterStatus(FNetObjectGroupHandle GroupHandle, uint32 ConnectionId, ENetFilterStatus ReplicationStatus);
+
 	uint32 GetOwningConnectionIfDirty(uint32 ObjectIndex) const;
 
 	void RemoveFromDynamicFilter(uint32 ObjectIndex, uint32 FilterIndex);
@@ -238,7 +249,7 @@ private:
 	static constexpr uint8 InvalidDynamicFilterIndex = 255U;
 
 	// General
-	TObjectPtr<UReplicationSystem> ReplicationSystem;
+	TObjectPtr<UReplicationSystem> ReplicationSystem = nullptr;
 	const FNetRefHandleManager* NetRefHandleManager = nullptr;
 
 	// Groups
@@ -276,11 +287,17 @@ private:
 	TArray<FPerGroupInfo> GroupInfos;
 	uint32 MaxGroupCount = 0;
 
-	/** NetObjectGroups who have the group filtering trait enabled */
-	FNetBitArray FilterGroups;
+	/** NetObjectGroups used for filtering out objects. */
+	FNetBitArray ExclusionFilterGroups;
+
+	/** NetObjectGroups used to allow allowing replication of dynamically filtered out objects. */
+	FNetBitArray InclusionFilterGroups;
+
+	/** Exclusion filtering groups with newly added members and that need to filter out objects for at least one connection. */
+	FNetBitArray DirtyExclusionFilterGroups;
 
-	/** NetObjectGroups who have dirty members and their filtering updated this frame */
-	FNetBitArray DirtyFilterGroups;
+	/** Inclusion filtering groups with newly added members and that need to include objects for at least one connection. */
+	FNetBitArray DirtyInclusionFilterGroups;
 
 	//$IRIS TODO: These need better documentation
 	FNetBitArray SubObjectFilterGroups;
@@ -307,6 +324,8 @@ private:
 	uint32 bHasDynamicFilters : 1;
 	uint32 bHasDynamicRawFilters : 1;
 	uint32 bHasDynamicFragmentFilters : 1;
+	uint32 bHasDirtyExclusionFilterGroup : 1;
+	uint32 bHasDirtyInclusionFilterGroup : 1;
 };
 
 inline bool FReplicationFiltering::HasDynamicFilters() const
diff --git a/Engine/Source/Runtime/Experimental/Iris/Core/Private/Iris/ReplicationSystem/ReplicationBridge.cpp b/Engine/Source/Runtime/Experimental/Iris/Core/Private/Iris/ReplicationSystem/ReplicationBridge.cpp
--- a/Engine/Source/Runtime/Experimental/Iris/Core/Private/Iris/ReplicationSystem/ReplicationBridge.cpp
+++ b/Engine/Source/Runtime/Experimental/Iris/Core/Private/Iris/ReplicationSystem/ReplicationBridge.cpp
@@ -937,7 +937,7 @@ UE::Net::FNetObjectGroupHandle UReplicationBridge::CreateLevelGroup(const UObjec
 	FNetObjectGroupHandle LevelGroupHandle = ReplicationSystem->CreateGroup();
 	if (ensure(LevelGroupHandle.IsValid()))
 	{
-		ReplicationSystem->AddGroupFilter(LevelGroupHandle);
+		ReplicationSystem->AddExclusionFilterGroup(LevelGroupHandle);
 		LevelGroups.Emplace(FObjectKey(Level), LevelGroupHandle);
 	}
 
diff --git a/Engine/Source/Runtime/Experimental/Iris/Core/Private/Iris/ReplicationSystem/ReplicationSystem.cpp b/Engine/Source/Runtime/Experimental/Iris/Core/Private/Iris/ReplicationSystem/ReplicationSystem.cpp
--- a/Engine/Source/Runtime/Experimental/Iris/Core/Private/Iris/ReplicationSystem/ReplicationSystem.cpp
+++ b/Engine/Source/Runtime/Experimental/Iris/Core/Private/Iris/ReplicationSystem/ReplicationSystem.cpp
@@ -69,7 +69,7 @@ public:
 	{
 		NotReplicatedNetObjectGroupHandle = ReplicationSystem->CreateGroup();
 		check(NotReplicatedNetObjectGroupHandle.IsNotReplicatedNetObjectGroup());
-		ReplicationSystem->AddGroupFilter(NotReplicatedNetObjectGroupHandle);
+		ReplicationSystem->AddExclusionFilterGroup(NotReplicatedNetObjectGroupHandle);
 		
 		// Setup SubObjectFiltering groups
 		NetGroupOwnerNetObjectGroupHandle = ReplicationSystem->GetOrCreateSubObjectFilter(UE::Net::NetGroupOwner);
@@ -1212,23 +1212,25 @@ void UReplicationSystem::AddToGroup(FNetObjectGroupHandle GroupHandle, FNetRefHa
 
 	using namespace UE::Net::Private;
 
-	FNetObjectGroups& Groups = Impl->ReplicationSystemInternal.GetGroups();
 	FNetRefHandleManager& NetRefHandleManager = Impl->ReplicationSystemInternal.GetNetRefHandleManager();
-	FReplicationFiltering& Filtering = Impl->ReplicationSystemInternal.GetFiltering();	
-
 	const FInternalNetRefIndex ObjectInternalIndex = NetRefHandleManager.GetInternalIndex(Handle);
-	
+
 	if (Impl->ReplicationSystemInternal.AreFilterChangesBlocked())
 	{
 		ensureMsgf(false, TEXT("Setting filter conditions is not yet supported during this operation. Filter condition on %s (%s) failed."), *GetNameSafe(NetRefHandleManager.GetReplicatedObjectInstance(ObjectInternalIndex)), *Handle.ToString());
 		return;
 	}
 
-	if (ObjectInternalIndex)
+	if (ObjectInternalIndex == FNetRefHandleManager::InvalidInternalIndex)
 	{
-		Groups.AddToGroup(GroupHandle, ObjectInternalIndex);
-		Filtering.NotifyObjectAddedToGroup(GroupHandle, ObjectInternalIndex);
+		return;
 	}
+	
+	FNetObjectGroups& Groups = Impl->ReplicationSystemInternal.GetGroups();
+	FReplicationFiltering& Filtering = Impl->ReplicationSystemInternal.GetFiltering();	
+
+	Groups.AddToGroup(GroupHandle, ObjectInternalIndex);
+	Filtering.NotifyObjectAddedToGroup(GroupHandle, ObjectInternalIndex);
 }
 
 void UReplicationSystem::RemoveFromGroup(FNetObjectGroupHandle GroupHandle, FNetRefHandle Handle)
@@ -1242,23 +1244,23 @@ void UReplicationSystem::RemoveFromGroup(FNetObjectGroupHandle GroupHandle, FNet
 	using namespace UE::Net::Private;
 
 	FNetRefHandleManager& NetRefHandleManager = Impl->ReplicationSystemInternal.GetNetRefHandleManager();
-
 	const FInternalNetRefIndex ObjectInternalIndex = NetRefHandleManager.GetInternalIndex(Handle);
-
 	if (Impl->ReplicationSystemInternal.AreFilterChangesBlocked())
 	{
 		ensureMsgf(false, TEXT("Setting filter conditions is not yet supported during this operation. Filter condition on %s (%s) failed."), *GetNameSafe(NetRefHandleManager.GetReplicatedObjectInstance(ObjectInternalIndex)), *Handle.ToString());
 		return;
 	}
 
-	if (ObjectInternalIndex)
+	if (ObjectInternalIndex == FNetRefHandleManager::InvalidInternalIndex)
 	{
-		FNetObjectGroups& Groups = Impl->ReplicationSystemInternal.GetGroups();
-		FReplicationFiltering& Filtering = Impl->ReplicationSystemInternal.GetFiltering();	
-
-		Groups.RemoveFromGroup(GroupHandle, ObjectInternalIndex);
-		Filtering.NotifyObjectRemovedFromGroup(GroupHandle, ObjectInternalIndex);
+		return;
 	}
+
+	FNetObjectGroups& Groups = Impl->ReplicationSystemInternal.GetGroups();
+	FReplicationFiltering& Filtering = Impl->ReplicationSystemInternal.GetFiltering();
+
+	Groups.RemoveFromGroup(GroupHandle, ObjectInternalIndex);
+	Filtering.NotifyObjectRemovedFromGroup(GroupHandle, ObjectInternalIndex);
 }
 
 void UReplicationSystem::RemoveFromAllGroups(FNetRefHandle Handle)
@@ -1350,22 +1352,40 @@ UE::Net::FNetObjectGroupHandle UReplicationSystem::GetNetGroupReplayNetObjectGro
 	return Impl->NetGroupReplayNetObjectGroupHandle;
 }
 
-void UReplicationSystem::AddGroupFilter(FNetObjectGroupHandle GroupHandle)
+bool UReplicationSystem::AddExclusionFilterGroup(FNetObjectGroupHandle GroupHandle)
 {
 	// Early out if this is invalid group
 	if (!ensure(IsValidGroup(GroupHandle)))
 	{
-		return;
+		return false;
 	}
 
 	if (Impl->ReplicationSystemInternal.AreFilterChangesBlocked())
 	{
 		ensureMsgf(false, TEXT("Setting filter conditions is not yet supported during this operation."));
-		return;
+		return false;
+	}
+
+	UE::Net::Private::FReplicationFiltering& Filtering = Impl->ReplicationSystemInternal.GetFiltering();	
+	return Filtering.AddExclusionFilterGroup(GroupHandle);
+}
+
+bool UReplicationSystem::AddInclusionFilterGroup(FNetObjectGroupHandle GroupHandle)
+{
+	// Early out if this is invalid group
+	if (!ensure(IsValidGroup(GroupHandle)))
+	{
+		return false;
+	}
+
+	if (Impl->ReplicationSystemInternal.AreFilterChangesBlocked())
+	{
+		ensureMsgf(false, TEXT("Setting filter conditions is not yet supported during this operation."));
+		return false;
 	}
 
 	UE::Net::Private::FReplicationFiltering& Filtering = Impl->ReplicationSystemInternal.GetFiltering();	
-	Filtering.AddGroupFilter(GroupHandle);
+	return Filtering.AddInclusionFilterGroup(GroupHandle);
 }
 
 void UReplicationSystem::RemoveGroupFilter(FNetObjectGroupHandle GroupHandle)
diff --git a/Engine/Source/Runtime/Experimental/Iris/Core/Public/Iris/ReplicationSystem/ReplicationSystem.h b/Engine/Source/Runtime/Experimental/Iris/Core/Public/Iris/ReplicationSystem/ReplicationSystem.h
--- a/Engine/Source/Runtime/Experimental/Iris/Core/Public/Iris/ReplicationSystem/ReplicationSystem.h
+++ b/Engine/Source/Runtime/Experimental/Iris/Core/Public/Iris/ReplicationSystem/ReplicationSystem.h
@@ -340,7 +340,7 @@ public:
 	 */
 	IRISCORE_API bool IsValidGroup(FNetObjectGroupHandle GroupHandle) const;
 
-	/** Special group, NetHandles assigned to this group will be filtered out for all connections */
+	/** Special group, root objects assigned to this group will be filtered out for all connections */
 	IRISCORE_API FNetObjectGroupHandle GetNotReplicatedNetObjectGroup() const;
 
 	/** Special group, SubObjects assigned to this group will replicate to owner of RootParent */
@@ -406,13 +406,32 @@ public:
 	// Group Filtering
 
 	/**
-	 * Add a group to the filtering system. By default the filter disallows replication for all objects in the group.
+	 * Add a group to the filtering system. This group is used only for filtering out objects. Exclusion groups are processed before dynamic filters, those implemented by UNetObjectFilter. 
+	 * By default an exclusion group disallows replication for all objects in it. Use SetGroupFilterStatus to change the behavior.
+	 * @note A group can only be either an exclusion group or an inclusion group, not both at the same time.
 	 * @param GroupHandle A valid handle to a group.
-	 * @see CreateGroup 
+	 * @see CreateGroup
+	 * @see AddInclusionFilterGroup
+	 * @see SetGroupFilterStatus
+	 * @return true if the group was successfully added as an exclusion group, false in all other cases such as being an invalid group, reserved group or used as an inclusion filter.
 	 */
-	IRISCORE_API void AddGroupFilter(FNetObjectGroupHandle GroupHandle);
+	IRISCORE_API bool AddExclusionFilterGroup(FNetObjectGroupHandle GroupHandle);
 
-	/** Remove group from filtering system, will cancel effects of the group. */
+	/**
+	 * Add a group to the filtering system. This group is used only for allowing replication of objects. Inclusion groups are processed after dynamic filters, those implemented by UNetObjectFilter. 
+	 * Inclusion groups are used to allow overriding the effect of dynamic filtering which can be useful to always allow replication of team specific objects for example.
+	 * By default the group will not override the effects of dynamic filtering. Use SetGroupFilterStatus to set which connection the objects should be allowed to replicate to, overriding the dynamic filtering.
+	 * @note A group can only be either an exclusion group or an inclusion group, not both at the same time.
+	 * @note Subobjects added to inclusion groups will be ignored during processing. A subobject's filter status is determined by the root object. For subobject filtering one can use SubObjectFilters.
+	 * @param GroupHandle A valid handle to a group.
+	 * @see CreateGroup
+	 * @see AddExclusionFilterGroup
+	 * @see SetGroupFilterStatus
+	 * @return true if the group was successfully added as an inclusion group, false in all other cases such as being an invalid group, reserved group or used as an exclusion filter.
+	 */
+	IRISCORE_API bool AddInclusionFilterGroup(FNetObjectGroupHandle GroupHandle);
+
+	/** Remove group from filtering system, cancelling all effects of the group. */
 	IRISCORE_API void RemoveGroupFilter(FNetObjectGroupHandle GroupHandle);
 
 	/** Set status of GroupFilter for specific connection. */
