IAS throttling waits until allowance or external time out is met.

#rnx

diff --git a/Engine/Source/Runtime/Experimental/IoStoreOnDemand/Private/CoreHttp/Http.cpp b/Engine/Source/Runtime/Experimental/IoStoreOnDemand/Private/CoreHttp/Http.cpp
--- a/Engine/Source/Runtime/Experimental/IoStoreOnDemand/Private/CoreHttp/Http.cpp
+++ b/Engine/Source/Runtime/Experimental/IoStoreOnDemand/Private/CoreHttp/Http.cpp
@@ -2433,6 +2433,7 @@ public:
 private:
 	friend	void ThrottleTest(FAnsiStringView);
 	int32	GetAllowance(uint64 CycleDelta);
+	int32	GetWaitEstimateMs() const;
 	uint64	CycleFreq;
 	uint64	CycleLast;
 	uint64	CycleIdle;
@@ -2440,7 +2441,7 @@ private:
 	int32	Available = 0;
 
 	enum {
-		LIMITLESS	= 0x7fff'ffff,
+		LIMITLESS	= MAX_int32,
 		THRESHOLD	= 2 << 10,
 	};
 };
@@ -2495,7 +2496,7 @@ int32 FThrottler::GetAllowance(uint64 CycleDelta)
 	if (Delta == 0)
 	{
 		CycleLast -= CycleDelta;
-		return 0;
+		return 0 - GetWaitEstimateMs();
 	}
 
 	// Don't let available run away
@@ -2506,7 +2507,7 @@ int32 FThrottler::GetAllowance(uint64 CycleDelta)
 	// Doesn't make sense to trickle out tiny allowances
 	if (Available < THRESHOLD)
 	{
-		return 0;
+		return 0 - GetWaitEstimateMs();
 	}
 
 	int32 Released = Available;
@@ -2520,6 +2521,15 @@ void FThrottler::ReturnUnused(uint32 Unused)
 	Available += Unused;
 }
 
+////////////////////////////////////////////////////////////////////////////////
+int32 FThrottler::GetWaitEstimateMs() const
+{
+	// Calculate an approximate time to wait for more allowance
+	int64 Estimate = THRESHOLD - Available;
+	Estimate = (Estimate * 1000ll) / int64(Limit);
+	return FMath::Max(int32(Estimate), 0);
+}
+
 
 
 // {{{1 groups .................................................................
@@ -2861,7 +2871,7 @@ bool FSocketGroup::Tick(FTickState& State)
 		SendInternal(State);
 	}
 
-	if (Recv != nullptr)
+	if (Recv != nullptr && State.RecvAllowance)
 	{
 		RecvInternal(State);
 	}
@@ -3242,7 +3252,29 @@ uint32 FEventLoop::FImpl::Tick(int32 PollTimeoutMs)
 
 	ReceiveWork();
 
+	// We limit recv sizes as a way to control bandwidth use.
 	int32 RecvAllowance = Throttler.GetAllowance();
+	if (RecvAllowance <= 0)
+	{
+		if (PollTimeoutMs == 0)
+		{
+			return BusyCount;
+		}
+
+		int32 ThrottleWaitMs = -RecvAllowance;
+		if (PollTimeoutMs > 0)
+		{
+			ThrottleWaitMs = FMath::Min(ThrottleWaitMs, PollTimeoutMs);
+		}
+		FPlatformProcess::SleepNoStats(float(ThrottleWaitMs) / 1000.0f);
+
+		RecvAllowance = Throttler.GetAllowance();
+		if (RecvAllowance <= 0)
+		{
+			return BusyCount;
+		}
+	}
+
 	uint64 CancelsLoad = Cancels.load(std::memory_order_relaxed);
 
 	// Tick groups and then remove ones that are idle
@@ -3586,8 +3618,8 @@ static void ThrottleTest(FAnsiStringView TestUrl)
 	{
 		int64 Delta = (OneSecond + 15) >> 4;
 		Counter -= Delta;
-		uint32 Allowance = Throttler.GetAllowance(Delta);
-		if (Allowance != 0)
+		int32 Allowance = Throttler.GetAllowance(Delta);
+		if (Allowance > 0)
 		{
 			check(Allowance == Limit);
 			check(Counter < 10);
@@ -3602,7 +3634,7 @@ static void ThrottleTest(FAnsiStringView TestUrl)
 
 	// timing test
 	FIoBuffer RecvData;
-	for (uint32 SizeKiB : { 25, 60 })
+	for (uint32 SizeKiB : { 10, 25, 60 })
 	{
 		const uint32 ThrottleKiB = 5;
 
@@ -3622,8 +3654,12 @@ static void ThrottleTest(FAnsiStringView TestUrl)
 			}
 		});
 
+		int32 Timeout = -1;
+		if (SizeKiB < 25) Timeout = 123;
+		if (SizeKiB > 25) Timeout = 4567;
+
 		uint64 Time = FPlatformTime::Cycles64();
-		while (Loop.Tick(-1));
+		while (Loop.Tick(Timeout));
 		Time = FPlatformTime::Cycles64() - Time;
 		Time /= OneSecond;
 
