Iris - Mute filtering warnings.

#rnx

diff --git a/Engine/Source/Runtime/Experimental/Iris/Core/Private/Iris/ReplicationSystem/Filtering/ReplicationFiltering.cpp b/Engine/Source/Runtime/Experimental/Iris/Core/Private/Iris/ReplicationSystem/Filtering/ReplicationFiltering.cpp
--- a/Engine/Source/Runtime/Experimental/Iris/Core/Private/Iris/ReplicationSystem/Filtering/ReplicationFiltering.cpp
+++ b/Engine/Source/Runtime/Experimental/Iris/Core/Private/Iris/ReplicationSystem/Filtering/ReplicationFiltering.cpp
@@ -1944,8 +1944,11 @@ void FReplicationFiltering::SetGroupFilterStatus(FNetObjectGroupHandle GroupHand
 			}
 			InternalSetExclusionGroupFilterStatus(GroupHandle, ConnectionId, ReplicationStatusToSet);
 		}
+
+		return;
 	}
-	else if (InclusionFilterGroups.GetBit(GroupIndex))
+	
+	if (InclusionFilterGroups.GetBit(GroupIndex))
 	{
 		// It is intentional that we iterate over all connections, not only the valid ones as we want the group filter status to be initialized correctly for new connections as well.
 		for (uint32 ConnectionId = 1U, EndConnectionId = ValidConnections.GetNumBits(); ConnectionId < EndConnectionId; ++ConnectionId)
@@ -1957,6 +1960,8 @@ void FReplicationFiltering::SetGroupFilterStatus(FNetObjectGroupHandle GroupHand
 			}
 			InternalSetInclusionGroupFilterStatus(GroupHandle, ConnectionId, ReplicationStatusToSet);
 		}
+		
+		return;
 	}
 
 	UE_LOG(LogIrisFiltering, Warning, TEXT("FReplicationFiltering::SetGroupFilterStatus - Trying to set filter for invalid GroupIndex: %u, Make sure group is added to filtering"), GroupIndex);
