Speed up the loading of .iochunktoc from disk

#rnx

- The default buffer size of FArchiveFileReaderGeneric on some platforms can be as low as 1 KiB, which performs very poorly for the .iochunktoc file which is often multiple MiB in size so raising the buffer size to 256 KiB.
- Add some missing cpu profiling scope markers

diff --git a/Engine/Source/Runtime/Experimental/IoStoreOnDemand/Private/OnDemandIoDispatcherBackend.cpp b/Engine/Source/Runtime/Experimental/IoStoreOnDemand/Private/OnDemandIoDispatcherBackend.cpp
--- a/Engine/Source/Runtime/Experimental/IoStoreOnDemand/Private/OnDemandIoDispatcherBackend.cpp
+++ b/Engine/Source/Runtime/Experimental/IoStoreOnDemand/Private/OnDemandIoDispatcherBackend.cpp
@@ -922,6 +922,8 @@ static TArray<FString> FindOnDemandUtocFilesOnDisk()
  */
 static TUniquePtr<FArchive> CreateReaderFromPlatformPackage(const FString& RelPath)
 {
+	TRACE_CPUPROFILER_EVENT_SCOPE(IasBackend::CreateReaderFromPlatformPackage);
+
 	const FString AbsPath = FPaths::Combine(FGenericPlatformMisc::RootDir(), RelPath);
 
 	IFileHandle* File = IPlatformFile::GetPlatformPhysical().OpenRead(*AbsPath);
@@ -932,8 +934,8 @@ static TUniquePtr<FArchive> CreateReaderFromPlatformPackage(const FString& RelPa
 		// offset to the start of the asset file.
 		File->Seek(0);
 #endif //PLATFORM_ANDROID
-
-		return MakeUnique<FArchiveFileReaderGeneric>(File, *AbsPath, File->Size());
+		const uint32 ReadBufferSize = 256 * 1024;
+		return MakeUnique<FArchiveFileReaderGeneric>(File, *AbsPath, File->Size(), ReadBufferSize);
 	}
 	else
 	{
@@ -950,6 +952,8 @@ static TIoStatusOr<FOnDemandToc> GenerateOnDemandTocFromDisk(FStringView TocHash
 
 	if (GIasLoadOnDemandToc)
 	{
+		TRACE_CPUPROFILER_EVENT_SCOPE(IasBackend::LoadTocFromDisk);
+
 		UE_LOG(LogIas, Log, TEXT("Serializing .iochunktoc from disk"));
 
 		FString TocFileName = FString(TocHash);
