Added an idle time out to IAS

#rnx

diff --git a/Engine/Source/Runtime/Experimental/IoStoreOnDemand/Private/CoreHttp/Http.cpp b/Engine/Source/Runtime/Experimental/IoStoreOnDemand/Private/CoreHttp/Http.cpp
--- a/Engine/Source/Runtime/Experimental/IoStoreOnDemand/Private/CoreHttp/Http.cpp
+++ b/Engine/Source/Runtime/Experimental/IoStoreOnDemand/Private/CoreHttp/Http.cpp
@@ -179,12 +179,18 @@ static void Trace(T Id, ETrace Action, U Param=0)
 // {{{1 misc ...................................................................
 
 ////////////////////////////////////////////////////////////////////////////////
-static int32 GIasHttpRecvWorkThresholdKiB = 80;
-static FAutoConsoleVariableRef CVar_IasHttpRecvWorkThresholdKiB(
-	TEXT("ias.HttpWorkThreshold"),
-	GIasHttpRecvWorkThresholdKiB,
-	TEXT("Threshold of data remaining at which next request is sent (in KiB)")
-);
+#define IAS_CVAR(Type, Name, Default, Desc, ...) \
+	Type G##Name = Default; \
+	static FAutoConsoleVariableRef CVar_Ias##Name( \
+		TEXT("ias.Http" #Name), \
+		G##Name, \
+		TEXT(Desc) \
+		__VA_ARGS__ \
+	)
+
+////////////////////////////////////////////////////////////////////////////////
+static IAS_CVAR(int32, RecvWorkThresholdKiB,80,		"Threshold of data remaining at which next request is sent (in KiB)");
+static IAS_CVAR(int32, IdleMs,				50'000,	"Time in seconds to close idle connections or fail waits");
 
 ////////////////////////////////////////////////////////////////////////////////
 class FResult
@@ -1040,13 +1046,14 @@ int32 FSocket::Wait(TArrayView<FWaiter> Waiters, int32 TimeoutMs)
 #if !UE_BUILD_SHIPPING
 
 ///////////////////////////////////////////////////////////////////////////////
-static uint32	GSocksIpAddress	= 0;
-static int32	GSocksPort		= 1080; // default SOCKS5 port
-static int32	GSocksVersion	= 5;
+static IAS_CVAR(int32,		SocksVersion,	5,		"SOCKS proxy protocol version to use");
+static IAS_CVAR(FString,	SocksIp,		"",		"Routes all IAS HTTP traffic through the given SOCKS proxy");
+static IAS_CVAR(int32,		SocksPort,		1080,	"Port of the SOCKS proxy to use");
 
-#if !UE_BUILD_SHIPPING
-static void SetSocksIpAddress(const TCHAR* Value)
+////////////////////////////////////////////////////////////////////////////////
+static uint32 GetSocksIpAddress()
 {
+	const TCHAR* Value = *GSocksIp;
 	uint32 IpAddress = 0;
 	uint32 Accumulator = 0;
 	while (true)
@@ -1072,34 +1079,10 @@ static void SetSocksIpAddress(const TCHAR* Value)
 			continue;
 		}
 
-		GSocksIpAddress = 0;
-		return;
+		return 0;
 	}
-	GSocksIpAddress = IpAddress;
-}
-
-static FString GSocksIpAddressStr;
-static FAutoConsoleVariableRef CVar_IasSocksIpAddress(
-	TEXT("ias.SocksIp"),
-	GSocksIpAddressStr,
-	TEXT("Routes all IAS HTTP traffic through the given SOCKS5 proxy"),
-	FConsoleVariableDelegate::CreateStatic([] (IConsoleVariable*) {
-		SetSocksIpAddress(*GSocksIpAddressStr);
-	})
-);
-
-static FAutoConsoleVariableRef CVar_IasSocksPort(
-	TEXT("ias.SocksPort"),
-	GSocksPort,
-	TEXT("Port of the SOCKS5 proxy to use")
-);
-
-static FAutoConsoleVariableRef CVar_IasSocksVersion(
-	TEXT("ias.SocksVersion"),
-	GSocksVersion,
-	TEXT("SOCKS proxy protocol version to use")
-);
-#endif
+	return IpAddress;
+}
 
 ////////////////////////////////////////////////////////////////////////////////
 static int32 ConnectSocks4(FSocket& Socket, uint32 IpAddress, uint32 Port)
@@ -1120,7 +1103,8 @@ static int32 ConnectSocks4(FSocket& Socket, uint32 IpAddress, uint32 Port)
 		uint32	IpAddress;
 	};
 
-	if (!Socket.Connect(GSocksIpAddress, GSocksPort))
+	uint32 SocksIpAddress = GetSocksIpAddress();
+	if (!SocksIpAddress || !Socket.Connect(SocksIpAddress, GSocksPort))
 	{
 		return -1;
 	}
@@ -1157,7 +1141,8 @@ static int32 ConnectSocks5(FSocket& Socket, uint32 IpAddress, uint32 Port)
 #pragma warning(disable : 6385)
 #endif
 
-	if (!Socket.Connect(GSocksIpAddress, GSocksPort))
+	uint32 SocksIpAddress = GetSocksIpAddress();
+	if (!SocksIpAddress || !Socket.Connect(SocksIpAddress, GSocksPort))
 	{
 		return -1;
 	}
@@ -1225,7 +1210,7 @@ static int32 MaybeConnectSocks(FSocket& Socket, uint32 IpAddress, uint32 Port)
 #if UE_BUILD_SHIPPING
 	return 0;
 #else
-	if (GSocksIpAddress == 0)
+	if (GSocksIp.IsEmpty())
 	{
 		return 0;
 	}
@@ -2559,6 +2544,7 @@ struct FTickState
 	int32&						RecvAllowance;
 	int32						PollTimeoutMs;
 	int32						FailTimeoutMs;
+	uint32						NowMs;
 	class FWorkQueue*			Work;
 };
 
@@ -2692,6 +2678,7 @@ private:
 	FActivity*			Send = nullptr;
 	FActivity*			Recv = nullptr;
 	FSocket				Socket;
+	uint32				LastUseMs = 0;
 	uint8				IsKeepAlive = 0;
 	bool				bWaiting = false;
 
@@ -2775,12 +2762,13 @@ void FSocketGroup::RecvInternal(FTickState& State)
 	}
 
 	IsKeepAlive &= Activity->IsKeepAlive;
+	LastUseMs = State.NowMs;
 
 	// If we've only a small amount left to receive we can start more work
 	if (IsKeepAlive & (Recv->Next == nullptr))
 	{
 		uint32 Remaining = Activity_RemainingKiB(Activity);
-		if (Remaining < uint32(GIasHttpRecvWorkThresholdKiB))
+		if (Remaining < uint32(GRecvWorkThresholdKiB))
 		{
 			if (FActivity* Next = State.Work->PopActivity(); Next != nullptr)
 			{
@@ -2892,6 +2880,13 @@ void FSocketGroup::TickSend(FTickState& State, FHost& Host)
 	FActivity* Pending = State.Work->PopActivity();
 	check(Pending != nullptr);
 
+	// Close idle sockets
+	if (Socket.IsValid() && LastUseMs + GIdleMs < State.NowMs)
+	{
+		LastUseMs = State.NowMs;
+		Socket = FSocket();
+	}
+
 	// We don't have a connected socket on first use, or if a keep-alive:close
 	// was received from the server. So we connect here.
 	bool bWillBlock = false;
@@ -3104,7 +3099,7 @@ private:
 	FActivity*				Pending			= nullptr;
 	FThrottler				Throttler;
 	TArray<FHostGroup>		Groups;
-	int32					FailTimeoutMs	= 30 * 1000;
+	int32					FailTimeoutMs	= GIdleMs;
 	uint32					BusyCount		= 0;
 };
 
@@ -3277,6 +3272,21 @@ uint32 FEventLoop::FImpl::Tick(int32 PollTimeoutMs)
 
 	uint64 CancelsLoad = Cancels.load(std::memory_order_relaxed);
 
+	uint32 NowMs;
+	{
+		// 4.2MM seconds will give us 50 days of uptime.
+		static uint64 Freq = 0;
+		static uint64 Base = 0;
+		if (Freq == 0)
+		{
+			Freq = uint64(1.0 / FPlatformTime::GetSecondsPerCycle());
+			Base = FPlatformTime::Cycles64();
+		}
+		uint64 NowBig = ((FPlatformTime::Cycles64() - Base) * 1000) / Freq;
+		NowMs = uint32(NowBig);
+		check(NowMs == NowBig);
+	}
+
 	// Tick groups and then remove ones that are idle
 	FTickState TickState = {
 		.DoneList = nullptr,
@@ -3284,6 +3294,7 @@ uint32 FEventLoop::FImpl::Tick(int32 PollTimeoutMs)
 		.RecvAllowance = RecvAllowance,
 		.PollTimeoutMs = PollTimeoutMs,
 		.FailTimeoutMs = FailTimeoutMs,
+		.NowMs = NowMs,
 	};
 	for (FHostGroup& Group : Groups)
 	{
