Move some potentially slow IAS initialization code from the GameThread to the Ias.Http thread.

#rnx

- To select a primary endpoint we will iterate over all known endpoints in order and run the latency test on them, stopping on the first endpoint that responds. This normally should take 10s of ms but in a bad network environment could easily take longer.
- In addition the call to ::AddDeferredTocs usually should just be collecting the result of previous work but depending on the local settings could try to generate a .iochunktoc from disk or download it from the newly selected primary endpoint both of which can take 100s of ms to complete.
- Moved the code into a new method FOnDemandIoBackend::InitializePrimaryEndpoint and instead of calling it once the distribution endpoint request has resolved (which will be from Tick on the GameThread) we run it right at the start of the Ias.Http thread.
- Added some additional logging to make initialization easier to track.
- Promoted the warning given if we fail to find any endpoint from the distribution point to an error.
- Add a log error if we fail LatencyTests with all endpoints.

diff --git a/Engine/Source/Runtime/Experimental/IoStoreOnDemand/Private/OnDemandIoDispatcherBackend.cpp b/Engine/Source/Runtime/Experimental/IoStoreOnDemand/Private/OnDemandIoDispatcherBackend.cpp
--- a/Engine/Source/Runtime/Experimental/IoStoreOnDemand/Private/OnDemandIoDispatcherBackend.cpp
+++ b/Engine/Source/Runtime/Experimental/IoStoreOnDemand/Private/OnDemandIoDispatcherBackend.cpp
@@ -1509,7 +1509,21 @@ private:
 	FIoStatus ApplyGeneratedOnDemandToc(const FString& CdnUrl, const FString& TocPath);
 	FIoStatus DownloadoadOnDemandToc(const FString& CdnUrl, const FString& TocPath);
 
-	void FlushDeferredTocs();
+	void InitializePrimaryEndpoint();
+
+	/** Mode to control which types of .iochunktoc are flush */
+	enum class EFlushMode : uint8
+	{
+		/** Flush tocs that require disk access */
+		Disk	= 1 << 0,
+		/** Flush tocs that require network access */
+		Network = 1 << 1,
+
+		/** Flush all types of tocs */
+		All		= Disk | Network,
+	};
+
+	void FlushDeferredTocs(EFlushMode FlushMode);
 	void ProcessHttpRequests(FHttpClient& HttpClient, FBitWindow& HttpErrors, int32 MaxConcurrentRequests);
 	int32 WaitForPendingRequests(float WaitTimeSeconds, float PollTimeSeconds);
 
@@ -1599,7 +1613,7 @@ void FOnDemandIoBackend::Initialize(TSharedRef<const FIoDispatcherBackendContext
 		Resolver->ResolveEndpoints(DistributionUrl,
 			[this, Resolver](const FString& DistributionEp, TConstArrayView<FString> Eps)
 			{
-				UE_CLOG(Eps.IsEmpty(), LogIas, Warning, TEXT("Failed to resolve available endpoint(s) from '%s'"), *DistributionEp);
+				UE_CLOG(Eps.IsEmpty(), LogIas, Error, TEXT("Failed to resolve available endpoint(s) from '%s'"), *DistributionEp);
 				{
 					FWriteScopeLock _(Lock);
 					for (const FString& Ep : Eps)
@@ -1608,21 +1622,11 @@ void FOnDemandIoBackend::Initialize(TSharedRef<const FIoDispatcherBackendContext
 					}
 				}
 
-				{
-					TConstArrayView<FString> Urls = AvailableEps.Urls;
-					const int32 MaxUrls = FMath::Min(GIasMaxEndpointTestCountAtStartup, AvailableEps.Urls.Num());
-					const FString TestPath = GetEndpointTestPath();
-					if (int32 Idx = LatencyTest(Urls.Left(MaxUrls), TestPath, bStopRequested); Idx != INDEX_NONE)
-					{
-						AvailableEps.Current = Idx;
-						UE_LOG(LogIas, Log, TEXT("Using endpoint '%s'"), *AvailableEps.GetCurrent());
-						FlushDeferredTocs();
-					}
-					else
-					{
-						BackendStatus.SetHttpError(true);
-					}
-				}
+				// A disk only flush should be fairly quick at this point, especially if 's.IasEnableThreadedTocGeneration' is true
+				// and this will help mount .iochunktoc as soon as possible rather than waiting for the primary endpoint to be
+				// chosen.
+				FlushDeferredTocs(EFlushMode::Disk);
+
 				ConditionallyStartBackendThread();
 			});
 		Resolver->ResolveDeferredEndpoints();
@@ -2017,34 +2021,82 @@ FIoStatus FOnDemandIoBackend::DownloadoadOnDemandToc(const FString& CdnUrl, cons
 	}
 }
 
-void FOnDemandIoBackend::FlushDeferredTocs()
+void FOnDemandIoBackend::InitializePrimaryEndpoint()
+{
+	TRACE_CPUPROFILER_EVENT_SCOPE(IasBackend::InitializePrimaryEndpoint);
+
+	UE_LOG(LogIas, Log, TEXT("Attempting to find the primary endpoint..."));
+
+	check(!AvailableEps.HasCurrent());
+
+	TConstArrayView<FString> Urls = AvailableEps.Urls;
+	const int32 MaxUrls = FMath::Min(GIasMaxEndpointTestCountAtStartup, AvailableEps.Urls.Num());
+	const FString TestPath = GetEndpointTestPath();
+	if (int32 Idx = LatencyTest(Urls.Left(MaxUrls), TestPath, bStopRequested); Idx != INDEX_NONE)
+	{
+		AvailableEps.Current = Idx;
+		UE_LOG(LogIas, Log, TEXT("Using endpoint '%s'"), *AvailableEps.GetCurrent());
+		
+		// We need to call this after a valid endpoint has been selected in case we need to download
+		// the .iochunktoc file.
+		FlushDeferredTocs(EFlushMode::All);
+	}
+	else
+	{
+		UE_LOG(LogIas, Error, TEXT("Unable to connect to any valid endpoint"));
+		BackendStatus.SetHttpError(true);
+	}
+}
+
+void FOnDemandIoBackend::FlushDeferredTocs(EFlushMode FlushMode)
 {
+	TRACE_CPUPROFILER_EVENT_SCOPE(IasBackend::FlushDeferredTocs);
+
 	TArray<FTocParams> Tocs;
 	{
 		FWriteScopeLock _(Lock);
-		if (DeferredTocs.IsEmpty() || AvailableEps.HasCurrent() == false)
+		if (DeferredTocs.IsEmpty())
 		{
 			return;
 		}
 		Tocs = MoveTemp(DeferredTocs);
 	}
 
-	for (const FTocParams& TocParams : Tocs)
+	for (FTocParams& TocParams : Tocs)
 	{
 		if (GIasGenerateOnDemandToc && !TocParams.bForceDownload)
 		{
-			FIoStatus Result = ApplyGeneratedOnDemandToc(AvailableEps.GetCurrent(), TocParams.Path);
-			if (!Result.IsOk())
+			if (EnumHasAnyFlags(FlushMode, EFlushMode::Disk))
 			{
-				UE_LOG(LogIas, Error, TEXT("Failed to add generated toc', reason '%s'"), *Result.ToString());
+				// TODO: This value if only used if UE_VALIDATE_GENERATED_TOC is enabled. The validation
+				// should be moved elsewhere so we don't need to do this work when it won't be used.
+				const FString CdnUrl = AvailableEps.HasCurrent() ? AvailableEps.GetCurrent() : FString();
+				FIoStatus Result = ApplyGeneratedOnDemandToc(CdnUrl, TocParams.Path);
+				if (!Result.IsOk())
+				{
+					UE_LOG(LogIas, Error, TEXT("Failed to add generated toc', reason '%s'"), *Result.ToString());
+				}
+			}
+			else
+			{
+				FWriteScopeLock _(Lock);
+				DeferredTocs.Emplace(MoveTemp(TocParams));
 			}
 		}
 		else
 		{
-			FIoStatus Result = DownloadoadOnDemandToc(AvailableEps.GetCurrent(), TocParams.Path);
-			if (!Result.IsOk())
+			if (EnumHasAnyFlags(FlushMode, EFlushMode::Network) && AvailableEps.HasCurrent())
 			{
-				UE_LOG(LogIas, Error, TEXT("Failed to add TOC '%s/%s', reason '%s'"), *AvailableEps.GetCurrent(), *TocParams.Path, *Result.ToString());
+				FIoStatus Result = DownloadoadOnDemandToc(AvailableEps.GetCurrent(), TocParams.Path);
+				if (!Result.IsOk())
+				{
+					UE_LOG(LogIas, Error, TEXT("Failed to add TOC '%s/%s', reason '%s'"), *AvailableEps.GetCurrent(), *TocParams.Path, *Result.ToString());
+				}
+			}
+			else
+			{
+				FWriteScopeLock _(Lock);
+				DeferredTocs.Emplace(MoveTemp(TocParams));
 			}
 		}
 	}
@@ -2294,6 +2346,8 @@ uint32 FOnDemandIoBackend::Run()
 {
 	LLM_SCOPE_BYTAG(Ias);
 
+	InitializePrimaryEndpoint();
+
 	FBitWindow HttpErrors;
 	HttpErrors.Reset(GIasHttpErrorSampleCount);
 
@@ -2343,7 +2397,7 @@ uint32 FOnDemandIoBackend::Run()
 					HttpClient->SetEndpoint(Idx);
 					BackendStatus.SetHttpError(false);
 					UE_LOG(LogIas, Log, TEXT("Successfully reconnected to '%s'"), *AvailableEps.GetCurrent());
-					FlushDeferredTocs();
+					FlushDeferredTocs(EFlushMode::All);
 				}
 			}
 			else if (HttpClient->IsUsingPrimaryEndpoint() == false)
