Remove HISM support from ISMPool.
Raw ISM is faster for all our use cases.
HISM doesn't support permanant instance IDs.
#tests PIE

diff --git a/Engine/Source/Runtime/Experimental/GeometryCollectionEngine/Private/GeometryCollection/GeometryCollectionISMPoolComponent.cpp b/Engine/Source/Runtime/Experimental/GeometryCollectionEngine/Private/GeometryCollection/GeometryCollectionISMPoolComponent.cpp
--- a/Engine/Source/Runtime/Experimental/GeometryCollectionEngine/Private/GeometryCollection/GeometryCollectionISMPoolComponent.cpp
+++ b/Engine/Source/Runtime/Experimental/GeometryCollectionEngine/Private/GeometryCollection/GeometryCollectionISMPoolComponent.cpp
@@ -58,18 +58,11 @@ void FGeometryCollectionMeshGroup::RemoveAllMeshes(FGeometryCollectionISMPool& I
 	MeshInfos.Empty();
 }
 
-void FGeometryCollectionISM::CreateISM(AActor* InOwningActor, bool bInUseHISM)
+void FGeometryCollectionISM::CreateISM(AActor* InOwningActor)
 {
 	check(InOwningActor);
 
-	if (bInUseHISM)
-	{
-		ISMComponent = NewObject<UHierarchicalInstancedStaticMeshComponent>(InOwningActor, NAME_None, RF_Transient | RF_DuplicateTransient);
-	}
-	else
-	{
-		ISMComponent = NewObject<UInstancedStaticMeshComponent>(InOwningActor, NAME_None, RF_Transient | RF_DuplicateTransient);
-	}
+	ISMComponent = NewObject<UInstancedStaticMeshComponent>(InOwningActor, NAME_None, RF_Transient | RF_DuplicateTransient);
 
 	ISMComponent->SetRemoveSwap();
 	ISMComponent->SetCanEverAffectNavigation(false);
@@ -176,15 +169,8 @@ FGeometryCollectionISMPool::FISMIndex FGeometryCollectionISMPool::AddISM(UGeomet
 	}
 
 	// Take an ISM from the current FreeLists if available instead of allocating a new slot.
-	const bool bIsHISM = (MeshInstance.Desc.Flags & FISMComponentDescription::UseHISM) != 0;
-
 	FISMIndex ISMIndex = INDEX_NONE;
-	if (bIsHISM && FreeListHISM.Num())
-	{
-		ISMIndex = FreeListHISM.Last();
-		FreeListHISM.RemoveAt(FreeListHISM.Num() - 1);
-	}
-	else if (!bIsHISM && FreeListISM.Num())
+	if (FreeListISM.Num())
 	{
 		ISMIndex = FreeListISM.Last();
 		FreeListISM.RemoveAt(FreeListISM.Num() - 1);
@@ -193,12 +179,12 @@ FGeometryCollectionISMPool::FISMIndex FGeometryCollectionISMPool::AddISM(UGeomet
 	{
 		ISMIndex = FreeList.Last();
 		FreeList.RemoveAt(FreeList.Num() - 1);
-		ISMs[ISMIndex].CreateISM(OwningComponent->GetOwner(), bIsHISM);
+		ISMs[ISMIndex].CreateISM(OwningComponent->GetOwner());
 	}
 	else
 	{
 		ISMIndex = ISMs.AddDefaulted();
-		ISMs[ISMIndex].CreateISM(OwningComponent->GetOwner(), bIsHISM);
+		ISMs[ISMIndex].CreateISM(OwningComponent->GetOwner());
 	}
 	
 	ISMs[ISMIndex].InitISM(MeshInstance);
@@ -329,16 +315,7 @@ void FGeometryCollectionISMPool::RemoveISM(const FGeometryCollectionMeshInfo& Me
 			// Remove component and push this ISM slot to the free list.
 			ensure(ISM.ISMComponent->PerInstanceSMData.Num() == 0);
 			MeshToISMIndex.Remove(ISM.MeshInstance);
-
-			const bool bIsHISM = (ISM.MeshInstance.Desc.Flags & FISMComponentDescription::UseHISM) != 0;
-			if (bIsHISM)
-			{
-				FreeListHISM.Add(MeshInfo.ISMIndex);
-			}
-			else
-			{
-				FreeListISM.Add(MeshInfo.ISMIndex);
-			}
+			FreeListISM.Add(MeshInfo.ISMIndex);
 
 #if WITH_EDITOR
 			ISM.ISMComponent->Rename(nullptr);
@@ -352,7 +329,6 @@ void FGeometryCollectionISMPool::Clear()
 	MeshToISMIndex.Reset();
 	FreeList.Reset();
 	FreeListISM.Reset();
-	FreeListHISM.Reset();
 	if (ISMs.Num() > 0)
 	{
 		if (AActor* OwningActor = ISMs[0].ISMComponent->GetOwner())
@@ -371,22 +347,10 @@ void FGeometryCollectionISMPool::Clear()
 void FGeometryCollectionISMPool::GarbageCollect()
 {
 	// Release one component per call until we reach minimum pool size.
-	const int32 NumFreeISMSlots = FreeListISM.Num();
-	const int32 NumFreeHISMSlots = FreeListHISM.Num();
-
-	if (NumFreeISMSlots + NumFreeHISMSlots > GComponentFreeListTargetSize)
+	if (FreeListISM.Num() > GComponentFreeListTargetSize)
 	{
-		int32 ISMIndex = INDEX_NONE;
-		if (NumFreeHISMSlots >= NumFreeISMSlots)
-		{
-			ISMIndex = FreeListHISM.Last();
-			FreeListHISM.RemoveAt(FreeListHISM.Num() - 1);
-		}
-		else
-		{
-			ISMIndex = FreeListISM.Last();
-			FreeListISM.RemoveAt(FreeListISM.Num() - 1);
-		}
+		const int32 ISMIndex = FreeListISM.Last();
+		FreeListISM.RemoveAt(FreeListISM.Num() - 1);
 		
 		UInstancedStaticMeshComponent* ISM = ISMs[ISMIndex].ISMComponent;
 		ISM->UnregisterComponent();
@@ -487,8 +451,7 @@ void UGeometryCollectionISMPoolComponent::GetResourceSizeEx(FResourceSizeEx& Cum
 		+ Pool.MeshToISMIndex.GetAllocatedSize()
 		+ Pool.ISMs.GetAllocatedSize()
 		+ Pool.FreeList.GetAllocatedSize()
-		+ Pool.FreeListISM.GetAllocatedSize()
-		+ Pool.FreeListHISM.GetAllocatedSize();
+		+ Pool.FreeListISM.GetAllocatedSize();
 	
 	for (FGeometryCollectionISM ISM : Pool.ISMs)
 	{
diff --git a/Engine/Source/Runtime/Experimental/GeometryCollectionEngine/Public/GeometryCollection/GeometryCollectionISMPoolComponent.h b/Engine/Source/Runtime/Experimental/GeometryCollectionEngine/Public/GeometryCollection/GeometryCollectionISMPoolComponent.h
--- a/Engine/Source/Runtime/Experimental/GeometryCollectionEngine/Public/GeometryCollection/GeometryCollectionISMPoolComponent.h
+++ b/Engine/Source/Runtime/Experimental/GeometryCollectionEngine/Public/GeometryCollection/GeometryCollectionISMPoolComponent.h
@@ -102,7 +102,7 @@ struct FISMComponentDescription
 {
 	enum EFlags
 	{
-		UseHISM = 1 << 1,
+		UseHISM = 1 << 1,							// HISM is no longer supported. This flag is ignored.
 		GpuLodSelection = 1 << 2,
 		ReverseCulling = 1 << 3,
 		StaticMobility = 1 << 4,
@@ -247,7 +247,7 @@ struct FGeometryCollectionMeshGroup
 struct FGeometryCollectionISM
 {
 	/** Create the ISMComponent according to settings on the mesh instance. */
-	void CreateISM(AActor* InOwningActor, bool bInUseHISM);
+	void CreateISM(AActor* InOwningActor);
 	/** Initialize the ISMComponent according to settings on the mesh instance. */
 	void InitISM(const FGeometryCollectionStaticMeshInstance& InMeshInstance);
 	/** Add a group to the ISM. Returns the group index. */
@@ -291,8 +291,6 @@ struct FGeometryCollectionISMPool
 	TArray<int32> FreeList;
 	/** Free list of indices in ISMs that have registered ISM components. */
 	TArray<int32> FreeListISM;
-	/** Free list of indices in ISMs that have registered HISM components. */
-	TArray<int32> FreeListHISM;
 };
 
 
