Chaos Cloth: Minor optimizations to TriangleMesh collisions.

Some additional vectorization and reducing dynamic allocations.

Moved some things around so they're private rather than protected to avoid future versioning/deprecation issues.

On our Kendra test on my 64-core threadripper:
Build spatial 1.5ms -> 1.ms
Find intersections 1ms -> 0.8ms

diff --git a/Engine/Source/Runtime/Experimental/Chaos/Private/Chaos/PBDCollisionSpringConstraintsBase.cpp b/Engine/Source/Runtime/Experimental/Chaos/Private/Chaos/PBDCollisionSpringConstraintsBase.cpp
--- a/Engine/Source/Runtime/Experimental/Chaos/Private/Chaos/PBDCollisionSpringConstraintsBase.cpp
+++ b/Engine/Source/Runtime/Experimental/Chaos/Private/Chaos/PBDCollisionSpringConstraintsBase.cpp
@@ -70,9 +70,9 @@ void FPBDCollisionSpringConstraintsBase::Init(const FSolverParticles& Particles,
 
 		// Preallocate enough space for all possible connections.
 		constexpr int32 MaxConnectionsPerPoint = 3;
-		Constraints.SetNum(NumParticles * MaxConnectionsPerPoint);
-		Barys.SetNum(NumParticles * MaxConnectionsPerPoint);
-		FlipNormal.SetNum(NumParticles * MaxConnectionsPerPoint);
+		Constraints.SetNumUninitialized(NumParticles * MaxConnectionsPerPoint);
+		Barys.SetNumUninitialized(NumParticles * MaxConnectionsPerPoint);
+		FlipNormal.SetNumUninitialized(NumParticles * MaxConnectionsPerPoint);
 
 		std::atomic<int32> ConstraintIndex(0);
 
@@ -179,9 +179,9 @@ void FPBDCollisionSpringConstraintsBase::Init(const FSolverParticles& Particles,
 
 		// Shrink the arrays to the actual number of found constraints.
 		const int32 ConstraintNum = ConstraintIndex.load();
-		Constraints.SetNum(ConstraintNum, /*bAllowShrinking*/ true);
-		Barys.SetNum(ConstraintNum, /*bAllowShrinking*/ true);
-		FlipNormal.SetNum(ConstraintNum, /*bAllowShrinking*/ true);
+		Constraints.SetNum(ConstraintNum, /*bAllowShrinking*/ false);
+		Barys.SetNum(ConstraintNum, /*bAllowShrinking*/ false);
+		FlipNormal.SetNum(ConstraintNum, /*bAllowShrinking*/ false);
 	}
 }
 template void CHAOS_API FPBDCollisionSpringConstraintsBase::Init<FTriangleMesh::TBVHType<FSolverReal>>(const FSolverParticles& Particles, const FTriangleMesh::TBVHType<FSolverReal>& Spatial, 
diff --git a/Engine/Source/Runtime/Experimental/Chaos/Private/Chaos/PBDTriangleMeshCollisions.cpp b/Engine/Source/Runtime/Experimental/Chaos/Private/Chaos/PBDTriangleMeshCollisions.cpp
--- a/Engine/Source/Runtime/Experimental/Chaos/Private/Chaos/PBDTriangleMeshCollisions.cpp
+++ b/Engine/Source/Runtime/Experimental/Chaos/Private/Chaos/PBDTriangleMeshCollisions.cpp
@@ -23,7 +23,7 @@ struct FEdgeFaceIntersection
 
 // Returned array has NOT been shrunk
 template<typename SpatialAccelerator>
-static TArray<FEdgeFaceIntersection> FindEdgeFaceIntersections(const FTriangleMesh& TriangleMesh, const SpatialAccelerator& Spatial, const FSolverParticles& Particles)
+static void FindEdgeFaceIntersections(const FTriangleMesh& TriangleMesh, const SpatialAccelerator& Spatial, const FSolverParticles& Particles, TArray<FEdgeFaceIntersection>& Intersections)
 {
 	TRACE_CPUPROFILER_EVENT_SCOPE(ChaosFPBDTriangleMeshCollisions_IntersectionQuery);
 
@@ -31,18 +31,18 @@ static TArray<FEdgeFaceIntersection> FindEdgeFaceIntersections(const FTriangleMe
 	const TArray<TVec2<int32>>& EdgeToFaces = TriangleMesh.GetEdgeToFaces();
 	const TArray<TVec3<int32>>& Elements = TriangleMesh.GetElements();
 
-	TArray<FEdgeFaceIntersection> Intersections;
-
 	// Preallocate enough space for (more than) typical number of expected intersections.
 	constexpr int32 PreallocatedIntersectionsPerEdge = 3;
-	Intersections.SetNum(PreallocatedIntersectionsPerEdge * SegmentMesh.GetNumElements());
+	const int32 PreallocatedIntersectionsNum = PreallocatedIntersectionsPerEdge * SegmentMesh.GetNumElements();
+	Intersections.Reset();
+	Intersections.SetNumUninitialized(PreallocatedIntersectionsNum);
 	std::atomic<int32> IntersectionIndex(0);
 
 	// Extra intersections that require a lock to write to if you have more than PreallocatedIntersectionsPerEdge 
 	TArray<FEdgeFaceIntersection> ExtraIntersections;
 	FCriticalSection CriticalSection;
 	PhysicsParallelFor(SegmentMesh.GetNumElements(),
-		[&Spatial, &TriangleMesh, &Particles, &SegmentMesh, &EdgeToFaces, &Elements, &IntersectionIndex, &Intersections, PreallocatedIntersectionsPerEdge, &ExtraIntersections, &CriticalSection](int32 EdgeIndex)
+		[&Spatial, &TriangleMesh, &Particles, &SegmentMesh, &EdgeToFaces, &Elements, &IntersectionIndex, &Intersections, PreallocatedIntersectionsNum, &ExtraIntersections, &CriticalSection](int32 EdgeIndex)
 		{
 
 			TArray< TTriangleCollisionPoint<FSolverReal> > Result;
@@ -78,10 +78,9 @@ static TArray<FEdgeFaceIntersection> FindEdgeFaceIntersections(const FTriangleMe
 					Intersection.FaceCoordinate = { CollisionPoint.Bary[2], CollisionPoint.Bary[3] };
 					Intersection.FaceNormal = CollisionPoint.Normal;
 					Intersection.IntersectionPoint = CollisionPoint.Location;
-
-					if (CollisionPointIndex < PreallocatedIntersectionsPerEdge)
+					const int32 IndexToWrite = IntersectionIndex.fetch_add(1);
+					if (IndexToWrite < PreallocatedIntersectionsNum)
 					{
-						const int32 IndexToWrite = IntersectionIndex.fetch_add(1);
 						Intersections[IndexToWrite] = Intersection;
 					}
 					else
@@ -101,7 +100,6 @@ static TArray<FEdgeFaceIntersection> FindEdgeFaceIntersections(const FTriangleMe
 
 	// Append any ExtraIntersections
 	Intersections.Append(ExtraIntersections);
-	return Intersections;
 }
 
 // Global intersection analysis (identifying global contours, flood filling)
@@ -239,7 +237,8 @@ namespace GIA
 		} // namespace __internal
 
 		// Build intersection contour data from edge-face intersection data.
-		static TArray<FIntersectionContourPair> BuildIntersectionContours(const FTriangleMesh& TriangleMesh, const TArray<FEdgeFaceIntersection>& IntersectionArray, TArray<TArray<FPBDTriangleMeshCollisions::FBarycentricPoint>>& ContourPoints)
+		static void BuildIntersectionContours(const FTriangleMesh& TriangleMesh, const TArray<FEdgeFaceIntersection>& IntersectionArray, 
+			TArray<FIntersectionContourPair>& Contours, TArray<TArray<FPBDTriangleMeshCollisions::FBarycentricPoint>>& ContourPoints)
 		{
 			TRACE_CPUPROFILER_EVENT_SCOPE(ChaosFPBDTriangleMeshCollisions_BuildIntersectionContours);
 
@@ -257,8 +256,6 @@ namespace GIA
 			const TArray<TVec3<int32>>& FaceToEdges = TriangleMesh.GetFaceToEdges();
 			const TArray<TVec3<int32>>& Elements = TriangleMesh.GetElements();
 
-			TArray<FIntersectionContourPair> Contours;
-
 			while (!IntersectionMap.IsEmpty())
 			{
 				// Add a new contour
@@ -586,8 +583,7 @@ namespace GIA
 					// We failed to find the next intersection in this contour. 
 					break;
 				}
-			}
-			return Contours;
+			};
 		}
 	} // namespace ContourBuilding
 
@@ -1252,6 +1248,18 @@ namespace ContourMinimization
 	}
 } // namespace ContourMinimization
 
+struct FPBDTriangleMeshCollisions::FScratchBuffers
+{
+	TArray<FEdgeFaceIntersection> EdgeFaceIntersections;
+	TArray<GIA::FIntersectionContourPair> IntersectionContours;
+
+	void Reset()
+	{
+		EdgeFaceIntersections.Reset();
+		IntersectionContours.Reset();
+	}
+};
+
 void FPBDTriangleMeshCollisions::Init(const FSolverParticles& Particles, const FSolverReal MinProximityQueryRadius)
 {
 	if (TriangleMesh.GetNumElements() == 0)
@@ -1277,21 +1285,26 @@ void FPBDTriangleMeshCollisions::Init(const FSolverParticles& Particles, const F
 		return;
 	}
 
+	if (!ScratchBuffers)
+	{
+		ScratchBuffers = MakePimpl<FScratchBuffers>();
+	}
+	ScratchBuffers->Reset();
+
 	// Detect all EdgeFace Intersections
-	TArray<FEdgeFaceIntersection> Intersections = FindEdgeFaceIntersections(TriangleMesh, SpatialHash, Particles);
-	if (Intersections.Num() == 0)
+	FindEdgeFaceIntersections(TriangleMesh, SpatialHash, Particles, ScratchBuffers->EdgeFaceIntersections);
+	if (ScratchBuffers->EdgeFaceIntersections.Num() == 0)
 	{
 		return;
 	}
 
-	TArray<GIA::FIntersectionContourPair> IntersectionContours;
 	if (bGlobalIntersectionAnalysis)
 	{ 
 		// Walk EdgeFace intersections to build global contours
-		IntersectionContours = GIA::ContourBuilding::BuildIntersectionContours(TriangleMesh, Intersections, IntersectionContourPoints);
+		GIA::ContourBuilding::BuildIntersectionContours(TriangleMesh, ScratchBuffers->EdgeFaceIntersections, ScratchBuffers->IntersectionContours, IntersectionContourPoints);
 		// Flood fill global contours to determine intersecting regions.
-		GIA::FloodFill::FloodFillContours(TriangleMesh, NumParticles, Offset, IntersectionContours, VertexGIAColors, TriangleGIAColors);
-		GIA::FloodFill::AssignContourPointTypes(IntersectionContours, IntersectionContourPoints, IntersectionContourTypes);
+		GIA::FloodFill::FloodFillContours(TriangleMesh, NumParticles, Offset, ScratchBuffers->IntersectionContours, VertexGIAColors, TriangleGIAColors);
+		GIA::FloodFill::AssignContourPointTypes(ScratchBuffers->IntersectionContours, IntersectionContourPoints, IntersectionContourTypes);
 	}
 
 	if (bContourMinimization)
@@ -1299,12 +1312,12 @@ void FPBDTriangleMeshCollisions::Init(const FSolverParticles& Particles, const F
 		if (bGlobalIntersectionAnalysis)
 		{
 			// Global contours which are non-closed or loop are handled via ContourMinimization impulses. Build global gradient (by adding contribution across all intersections per contour).
-			ContourMinimization::BuildGlobalContourMinimizationIntersections(TriangleMesh, Particles, IntersectionContours, ContourMinimizationIntersections);
+			ContourMinimization::BuildGlobalContourMinimizationIntersections(TriangleMesh, Particles, ScratchBuffers->IntersectionContours, ContourMinimizationIntersections);
 		}
 		else
 		{
 			// Just build local gradient for all Intersections
-			ContourMinimization::BuildLocalContourMinimizationIntersections(TriangleMesh, Particles, Intersections, ContourMinimizationIntersections);
+			ContourMinimization::BuildLocalContourMinimizationIntersections(TriangleMesh, Particles, ScratchBuffers->EdgeFaceIntersections, ContourMinimizationIntersections);
 		}
 	}
 }
@@ -1327,19 +1340,24 @@ void FPBDTriangleMeshCollisions::PostStepInit(const FSolverParticles& Particles)
 			constexpr FSolverReal MinSpatialLodSize = 1.f;
 			TriangleMesh.BuildSpatialHash(static_cast<const TArrayView<const FSolverVec3>&>(Particles.XArray()), SpatialHash, MinSpatialLodSize);
 
+			if (!ScratchBuffers)
+			{
+				ScratchBuffers = MakePimpl<FScratchBuffers>();
+			}
+			ScratchBuffers->Reset();
 
 			// Detect all EdgeFace Intersections
-			TArray<FEdgeFaceIntersection> Intersections = FindEdgeFaceIntersections(TriangleMesh, SpatialHash, Particles);
+			FindEdgeFaceIntersections(TriangleMesh, SpatialHash, Particles, ScratchBuffers->EdgeFaceIntersections);
 
 			if (bUseGlobalPostStepContours)
 			{
-				TArray<GIA::FIntersectionContourPair> IntersectionContours = GIA::ContourBuilding::BuildIntersectionContours(TriangleMesh, Intersections, PostStepIntersectionContourPoints);
-				ContourMinimization::BuildGlobalContourMinimizationIntersections(TriangleMesh, Particles, IntersectionContours, PostStepContourMinimizationIntersections);
+				GIA::ContourBuilding::BuildIntersectionContours(TriangleMesh, ScratchBuffers->EdgeFaceIntersections, ScratchBuffers->IntersectionContours, PostStepIntersectionContourPoints);
+				ContourMinimization::BuildGlobalContourMinimizationIntersections(TriangleMesh, Particles, ScratchBuffers->IntersectionContours, PostStepContourMinimizationIntersections);
 			}
 			else
 			{
 				// Just build local gradient for all Intersections
-				ContourMinimization::BuildLocalContourMinimizationIntersections(TriangleMesh, Particles, Intersections, PostStepContourMinimizationIntersections);
+				ContourMinimization::BuildLocalContourMinimizationIntersections(TriangleMesh, Particles, ScratchBuffers->EdgeFaceIntersections, PostStepContourMinimizationIntersections);
 			}
 		}
 	}
diff --git a/Engine/Source/Runtime/Experimental/Chaos/Private/Chaos/TriangleMesh.cpp b/Engine/Source/Runtime/Experimental/Chaos/Private/Chaos/TriangleMesh.cpp
--- a/Engine/Source/Runtime/Experimental/Chaos/Private/Chaos/TriangleMesh.cpp
+++ b/Engine/Source/Runtime/Experimental/Chaos/Private/Chaos/TriangleMesh.cpp
@@ -64,10 +64,39 @@ struct TTriangleMeshBvEntry
 		return Bounds;
 	}
 
+	typename FTriangleMesh::TSpatialHashType<T>::FVectorAABB VectorAABB() const
+	{
+		const TVec3<int32>& Tri = TmData->GetElements()[Index];
+		const TVec3<T>& A = (*Points)[Tri[0]];
+		const TVec3<T>& B = (*Points)[Tri[1]];
+		const TVec3<T>& C = (*Points)[Tri[2]];
+		typename FTriangleMesh::TSpatialHashType<T>::FVectorAABB AABB(A);
+		AABB.GrowToInclude(B);
+		AABB.GrowToInclude(C);
+		return AABB;
+	}
+
 	template<typename TPayloadType>
 	int32 GetPayload(int32 Idx) const { return Idx; }
 };
 
+template<typename T>
+struct TTriangleMeshBvData
+{
+	const FTriangleMesh* TmData;
+	const TConstArrayView<TVec3<T>>* Points;
+
+	TTriangleMeshBvEntry<T> operator[](const int32 ParticleIndex) const
+	{
+		return TTriangleMeshBvEntry<T>{ TmData, Points, ParticleIndex };
+	}
+
+	int32 Num() const
+	{
+		return TmData->GetNumElements();
+	}
+};
+
 FTriangleMesh::FTriangleMesh()
     : MStartIdx(0)
     , MNumIndices(0)
@@ -1675,15 +1704,9 @@ template CHAOS_API bool FTriangleMesh::SmoothProject<FRealDouble>(const TBVHType
 template<typename T>
 void FTriangleMesh::BuildSpatialHash(const TConstArrayView<TVec3<T>>& Points, TSpatialHashType<T>& SpatialHash, const T MinSpatialLodSize) const
 {
-	TArray<TTriangleMeshBvEntry<T>> BVEntries;
-	const int32 NumTris = MElements.Num();
-	BVEntries.Reset(NumTris);
-	for (int32 Tri = 0; Tri < NumTris; ++Tri)
-	{
-		BVEntries.Add({ this, &Points, Tri });
-	}
+	const TTriangleMeshBvData<T> BvData({ this, &Points });
 
-	SpatialHash.Initialize(BVEntries, MinSpatialLodSize);
+	SpatialHash.Initialize(BvData, MinSpatialLodSize);
 }
 template void FTriangleMesh::BuildSpatialHash<FRealSingle>(const TConstArrayView<TVec3<FRealSingle>>& Points, TSpatialHashType<FRealSingle>& SpatialHash, const FRealSingle MinSpatialLodSize) const;
 template void FTriangleMesh::BuildSpatialHash<FRealDouble>(const TConstArrayView<TVec3<FRealDouble>>& Points, TSpatialHashType<FRealDouble>& SpatialHash, const FRealDouble MinSpatialLodSize) const;
@@ -1694,19 +1717,19 @@ bool FTriangleMesh::PointProximityQuery(const TSpatialHashType<T>& SpatialHash,
 {
 	const T TotalThickness = ThisThickness + PointThickness;
 	const T TotalThicknessSq = TotalThickness * TotalThickness;
-	TAABB<T,3> QueryBounds(PointPosition, PointPosition);
+	typename TSpatialHashType<T>::FVectorAABB QueryBounds(PointPosition);
 	QueryBounds.Thicken(TotalThickness);
 
-	const TArray<int32> PotentialIntersections = SpatialHash.FindAllIntersections(QueryBounds);
+	const TArray<int32> PotentialIntersections = SpatialHash.FindAllIntersections(QueryBounds,
+		[PointIndex, &BroadphaseTest](int32 Payload)
+	{
+		return BroadphaseTest(PointIndex, Payload);
+	});
 
 	Result.Reset(PotentialIntersections.Num());
 
 	for (int32 TriIdx : PotentialIntersections)
 	{
-		if (!BroadphaseTest(PointIndex, TriIdx))
-		{
-			continue;
-		}
 
 		const TVec3<T>& A = Points[MElements[TriIdx][0]];
 		const TVec3<T>& B = Points[MElements[TriIdx][1]];
@@ -1743,20 +1766,20 @@ template<typename T>
 bool FTriangleMesh::EdgeIntersectionQuery(const TSpatialHashType<T>& SpatialHash, const TConstArrayView<TVec3<T>>& Points, const int32 EdgeIndex, const TVec3<T>& EdgePosition1, const TVec3<T>& EdgePosition2,
 	TFunctionRef<bool(const int32 EdgeIndex, const int32 TriangleIndex)> BroadphaseTest, TArray<TTriangleCollisionPoint<T>>& Result) const
 {
-	TAABB<T, 3> QueryBounds(EdgePosition1, EdgePosition1);
+	typename TSpatialHashType<T>::FVectorAABB QueryBounds(EdgePosition1);
 	QueryBounds.GrowToInclude(EdgePosition2);
 
-	const TArray<int32> PotentialIntersections = SpatialHash.FindAllIntersections(QueryBounds);
+	const TArray<int32> PotentialIntersections = SpatialHash.FindAllIntersections(QueryBounds,
+		[EdgeIndex, &BroadphaseTest](int32 Payload)
+	{
+		return BroadphaseTest(EdgeIndex, Payload);
+	}
+	);
 
 	Result.Reset(PotentialIntersections.Num());
 
 	for (int32 TriIdx : PotentialIntersections)
 	{
-		if (!BroadphaseTest(EdgeIndex, TriIdx))
-		{
-			continue;
-		}
-
 		const TVec3<T>& A = Points[MElements[TriIdx][0]];
 		const TVec3<T>& B = Points[MElements[TriIdx][1]];
 		const TVec3<T>& C = Points[MElements[TriIdx][2]];
diff --git a/Engine/Source/Runtime/Experimental/Chaos/Private/Chaos/WeightedLatticeImplicitObject.cpp b/Engine/Source/Runtime/Experimental/Chaos/Private/Chaos/WeightedLatticeImplicitObject.cpp
--- a/Engine/Source/Runtime/Experimental/Chaos/Private/Chaos/WeightedLatticeImplicitObject.cpp
+++ b/Engine/Source/Runtime/Experimental/Chaos/Private/Chaos/WeightedLatticeImplicitObject.cpp
@@ -330,6 +330,22 @@ struct FWeightedLatticeBvEntry
 		return Bounds;
 	}
 
+	THierarchicalSpatialHash<int32, FReal>::FVectorAABB VectorAABB() const
+	{
+		THierarchicalSpatialHash<int32, FReal>::FVectorAABB Bounds(Lattice->GetDeformedPoints()(CellIndex));
+		for (int32 I = 0; I < 2; ++I)
+		{
+			for (int32 J = 0; J < 2; ++J)
+			{
+				for (int32 K = 0; K < 2; ++K)
+				{
+					Bounds.GrowToInclude(Lattice->GetDeformedPoints()(CellIndex + TVec3<int32>(I, J, K)));
+				}
+			}
+		}
+		return Bounds;
+	}
+
 	template<typename TPayloadType>
 	TPayloadType GetPayload(int32 Idx) const
 	{
diff --git a/Engine/Source/Runtime/Experimental/Chaos/Public/Chaos/HierarchicalSpatialHash.h b/Engine/Source/Runtime/Experimental/Chaos/Public/Chaos/HierarchicalSpatialHash.h
--- a/Engine/Source/Runtime/Experimental/Chaos/Public/Chaos/HierarchicalSpatialHash.h
+++ b/Engine/Source/Runtime/Experimental/Chaos/Public/Chaos/HierarchicalSpatialHash.h
@@ -11,7 +11,7 @@
 
 namespace Chaos
 {
-struct FHierarchicalSpatialHashCellIdx
+struct UE_DEPRECATED(5.4, "This is deprecated and its replacement is private.") FHierarchicalSpatialHashCellIdx
 {
 	TVec3<int32> Cell;
 	int32 Lod;
@@ -59,6 +59,8 @@ struct FHierarchicalSpatialHashCellIdx
 	}
 };
 
+UE_DEPRECATED(5.4, "This is deprecated and its replacement is private.")
+PRAGMA_DISABLE_DEPRECATION_WARNINGS
 inline uint32 GetTypeHash(const Chaos::FHierarchicalSpatialHashCellIdx& Idx)
 {
 	// xorHash
@@ -68,17 +70,56 @@ inline uint32 GetTypeHash(const Chaos::FHierarchicalSpatialHashCellIdx& Idx)
 	Hash = Hash ^ (Idx.Lod * 67867979);
 	return (uint32)Hash;
 }
+PRAGMA_ENABLE_DEPRECATION_WARNINGS
 
 template<typename PayloadType>
 class TSpatialHashGridBase
 {
-protected:	
+private:	
 	struct FPayloadAndBounds
 	{
 		PayloadType Payload;
 		int32 BoundsLookupIdx;
 	};
 
+	struct FHashIndex
+	{
+		VectorRegister4Int VectorIndex;
+
+		FHashIndex() = default;
+
+		FHashIndex(const TVec3<int32>& Cell, int32 Lod)
+			:VectorIndex(MakeVectorRegisterInt(Cell.X, Cell.Y, Cell.Z, Lod))
+		{}
+		FHashIndex(int32 CellX, int32 CellY, int32 CellZ, int32 Lod)
+			:VectorIndex(MakeVectorRegisterInt(CellX, CellY, CellZ, Lod))
+		{}
+		FHashIndex(const VectorRegister4Int& InVectorIndex)
+			:VectorIndex(InVectorIndex)
+		{}
+
+		bool operator==(const FHashIndex& Other) const
+		{
+			union
+			{
+				VectorRegister4Float v;
+				VectorRegister4Int i;
+			}Result;
+			Result.i = VectorIntCompareNEQ(VectorIndex, Other.VectorIndex);
+			return VectorMaskBits(Result.v) == 0;
+		}
+
+		inline uint32 GetTypeHash() const
+		{
+			// xorHash
+			static constexpr VectorRegister4Int Multiplier = MakeVectorRegisterIntConstant(73856093, 19349663, 83492791, 67867979);
+			const VectorRegister4Int Multiplied = VectorIntMultiply(VectorIndex, Multiplier);
+			alignas(VectorRegister4Int) int32 Indexable[4];
+			VectorIntStoreAligned(Multiplied, Indexable);
+			return (uint32)(Indexable[0] ^ Indexable[1] ^ Indexable[2] ^ Indexable[3]);
+		}
+	};
+
 	// This hashmap is designed to be fast at just what we need to do to build all at once in exactly the order it is done in Initialize and then query. 
 	// Based off of Core/HashTable.h
 	class FHashMap
@@ -114,10 +155,12 @@ protected:
 				{
 					for (int32 ZIdx = MinCellIdx[2]; ZIdx <= MaxCellIdx[2]; ++ZIdx)
 					{
-						Elements[IndexToWrite].Key = FHierarchicalSpatialHashCellIdx{ {XIdx, YIdx, ZIdx}, Lod };
-						Elements[IndexToWrite].Value = Value;
-						const uint32 HashKey = GetTypeHash(Elements[IndexToWrite].Key) & (Hash.Num() - 1);
-						Elements[IndexToWrite].NextIndex = Hash[HashKey].exchange(IndexToWrite);
+						FElement Element;
+						Element.Key = FHashIndex(XIdx, YIdx, ZIdx, Lod);
+						const uint32 HashKey = Element.Key.GetTypeHash() & (Hash.Num() - 1);
+						Element.NextIndex = Hash[HashKey].exchange(IndexToWrite);
+						Element.Value = Value;
+						Elements[IndexToWrite] = MoveTemp(Element);
 						++IndexToWrite;
 					}
 				}
@@ -127,10 +170,13 @@ protected:
 		int32 ConcurrentAddElement(const TVec3<int32>& CellIdx, int32 Lod, const FPayloadAndBounds& Value)
 		{
 			const int32 IndexToWrite = ConcurrentElementAddIdx.fetch_add(1, std::memory_order_relaxed);
-			Elements[IndexToWrite].Key = FHierarchicalSpatialHashCellIdx{ CellIdx, Lod };
-			Elements[IndexToWrite].Value = Value;
-			const uint32 HashKey = GetTypeHash(Elements[IndexToWrite].Key) & (Hash.Num() - 1);
-			Elements[IndexToWrite].NextIndex = Hash[HashKey].exchange(IndexToWrite);
+			FElement Element;
+			Element.Key = FHashIndex( CellIdx, Lod );
+			const uint32 HashKey = Element.Key.GetTypeHash() & (Hash.Num() - 1);
+			Element.NextIndex = Hash[HashKey].exchange(IndexToWrite);
+			Element.Value = Value;
+			Elements[IndexToWrite] = MoveTemp(Element);
+
 			return IndexToWrite;
 		}
 
@@ -140,10 +186,10 @@ protected:
 			Elements.SetNum(ConcurrentElementAddIdx.load(), bAllowShrinking);
 		}
 
-		int32 First(const FHierarchicalSpatialHashCellIdx& Key) const
+		int32 First(const FHashIndex& Key) const
 		{
 			checkSlow(FMath::IsPowerOfTwo(Hash.Num()));
-			const uint32 HashKey = GetTypeHash(Key) & (Hash.Num() - 1);
+			const uint32 HashKey = Key.GetTypeHash() &(Hash.Num() - 1);
 			return Hash[HashKey].load(std::memory_order_relaxed);
 		}
 
@@ -157,7 +203,7 @@ protected:
 			return Index != INDEX_NONE;
 		}
 
-		bool KeyMatches(int32 Index, const FHierarchicalSpatialHashCellIdx& Key) const
+		bool KeyMatches(int32 Index, const FHashIndex& Key) const
 		{
 			return Elements[Index].Key == Key;
 		}
@@ -167,7 +213,7 @@ protected:
 			return Elements[Index].Value;
 		}
 
-		const FHierarchicalSpatialHashCellIdx& GetKey(int32 Index) const
+		const FHashIndex& GetKey(int32 Index) const
 		{
 			return Elements[Index].Key;
 		}
@@ -175,9 +221,9 @@ protected:
 	private:
 
 		TArray<std::atomic<int32>> Hash;
-		struct alignas(32) FElement
+		struct FElement
 		{
-			FHierarchicalSpatialHashCellIdx Key;
+			FHashIndex Key;
 			FPayloadAndBounds Value;
 			int32 NextIndex;
 		};
@@ -192,6 +238,31 @@ protected:
 	{
 		HashMap.Reset();
 	}
+
+	static inline TVec3<int32> CellIdxForPoint(const VectorRegister4Float& Point, const VectorRegister4Float& CellSize)
+	{
+		VectorRegister4Float Tmp = VectorDivide(Point, CellSize);
+
+		VectorRegister4Int CellVector = VectorRoundToIntHalfToEven(VectorAdd(Tmp, VectorSubtract(Tmp, GlobalVectorConstants::FloatOneHalf)));
+		CellVector = VectorShiftRightImmArithmetic(CellVector, 1);
+		TVec4<int32> Result4;
+		VectorIntStore(CellVector, &Result4[0]);
+		return TVec3<int32>(Result4[0], Result4[1], Result4[2]);
+	}
+
+	static inline TVec3<int32> CellIdxForPoint(const VectorRegister4Double& Point, const VectorRegister4Double& CellSize)
+	{
+		VectorRegister4Double Tmp = VectorDivide(Point, CellSize);
+
+		VectorRegister4Int CellVector = VectorRoundToIntHalfToEven(MakeVectorRegisterFloatFromDouble(VectorAdd(Tmp, VectorSubtract(Tmp, GlobalVectorConstants::DoubleOneHalf))));
+		CellVector = VectorShiftRightImmArithmetic(CellVector, 1);
+		TVec4<int32> Result4;
+		VectorIntStore(CellVector, &Result4[0]);
+		return TVec3<int32>(Result4[0], Result4[1], Result4[2]);
+	}
+
+	template<typename PT, typename T> friend class THierarchicalSpatialHash;
+	template<typename PT, typename T> friend class TSpatialHashGridPoints;
 };
 
 // Currently this assumes a 3 dimensional grid. 
@@ -201,8 +272,33 @@ class THierarchicalSpatialHash : public TSpatialHashGridBase<PayloadType>
 	typedef TSpatialHashGridBase<PayloadType> Base;
 	using typename Base::FPayloadAndBounds;
 	using typename Base::FHashMap;
+	using typename Base::FHashIndex;
 	using Base::HashMap;
 
+	static int32 LodForCellSize(T CellSize)
+	{
+		// want Ceil( log2(CellSize) )
+		if (CellSize <= (T).5)
+		{
+			return -(int32)FMath::FloorLog2((uint32)FMath::FloorToInt32((T)1 / CellSize));
+		}
+		else
+		{
+			return (int32)FMath::CeilLogTwo((uint32)FMath::CeilToInt32(CellSize));
+		}
+	}
+
+	static T CellSizeForLod(const int32 Lod)
+	{
+		if (Lod >= 0)
+		{
+			return (T)(1ULL << Lod);
+		}
+		return (T)1 / (T)(1ULL << (-Lod));
+	}
+
+public:
+
 	class FVectorAABB
 	{
 	public:
@@ -214,6 +310,36 @@ class THierarchicalSpatialHash : public TSpatialHashGridBase<PayloadType>
 			, Max(MakeVectorRegister(BBox.Max()[0], BBox.Max()[1], BBox.Max()[2], (T)0.))
 		{}
 
+		explicit FVectorAABB(const TVec3<T>& Point)
+			: Min(MakeVectorRegister(Point[0], Point[1], Point[2], (T)0.))
+			, Max(Min)
+		{}
+
+		void GrowToInclude(const TVec3<T>& Point)
+		{
+			const TVectorRegisterType<T> Vector = MakeVectorRegister(Point[0], Point[1], Point[2], (T)0);
+			Min = VectorMin(Min, Vector);
+			Max = VectorMax(Max, Vector);
+		}
+
+		void Thicken(const T Thickness)
+		{
+			const TVectorRegisterType<T> ThicknessVec = MakeVectorRegister(Thickness, Thickness, Thickness, (T)0);
+			Min = VectorSubtract(Min, ThicknessVec);
+			Max = VectorAdd(Max, ThicknessVec);
+		}
+
+		T LongestSideLength() const
+		{
+			const TVectorRegisterType<T> ExtentsVector = VectorSubtract(Max, Min);
+			TVec3<T> Extents;
+			VectorStoreFloat3(ExtentsVector, &Extents);
+			return Extents.Max();
+		}
+
+		const TVectorRegisterType<T>& GetMin() const { return Min; }
+		const TVectorRegisterType<T>& GetMax() const { return Max; }
+
 		bool Intersects(const FVectorAABB& Other) const
 		{
 			if (VectorAnyLesserThan(Other.Max, Min) || VectorAnyGreaterThan(Other.Min, Max))
@@ -238,8 +364,6 @@ class THierarchicalSpatialHash : public TSpatialHashGridBase<PayloadType>
 		TVectorRegisterType<T> Max;
 	};
 
-public:
-
 	void Reset()
 	{
 		Base::Reset();
@@ -255,7 +379,7 @@ public:
 		Bounds.SetNumUninitialized(Particles.Num());
 		HashMap.PreallocateElementsForConcurrentAdd(Particles.Num() * 8, Particles.Num() * 2);
 
-		const int32 MinAllowableLod = MinLodSize > (T)0 ? FHierarchicalSpatialHashCellIdx::LodForCellSize(MinLodSize) : std::numeric_limits<int32>::min();
+		const int32 MinAllowableLod = MinLodSize > (T)0 ? LodForCellSize(MinLodSize) : std::numeric_limits<int32>::min();
 
 		// For most common UsedLods, store whether or not the Lod is used in an array without a lock.
 		// For uncommon Lods, use a critical section to store directly into the UsedLods map.
@@ -269,19 +393,21 @@ public:
 			[this, &Particles, MinAllowableLod, MaxPreAllocatedUsedLodValue, MinPreAllocatedUsedLodValue, &UsedLodsArray, &UsedLodsCriticalSection](int32 ParticleIdx)
 			{
 				const auto& Particle = Particles[ParticleIdx];
-				const TAABB<T, 3> ParticleBounds = Particle.BoundingBox();
+				const FVectorAABB ParticleBounds = Particle.VectorAABB();
 
-				const int32 Lod = FMath::Max(MinAllowableLod, FHierarchicalSpatialHashCellIdx::LodForAABB(ParticleBounds));
-				const T CellSize = FHierarchicalSpatialHashCellIdx::CellSizeForLod<T>(Lod);
+				const int32 Lod = FMath::Max(MinAllowableLod, LodForCellSize(ParticleBounds.LongestSideLength()));
+				const T CellSize = CellSizeForLod(Lod);
+				const TVectorRegisterType<T> CellSizeVector = MakeVectorRegister(CellSize, CellSize, CellSize, CellSize);
 
 				const PayloadType Payload = Particle.template GetPayload<PayloadType>(ParticleIdx);
-				Bounds[ParticleIdx] = FVectorAABB(ParticleBounds);
 
-				const TVec3<int32> MinCellIdx = FHierarchicalSpatialHashCellIdx::CellIdxForPoint(ParticleBounds.Min(), CellSize);
-				const TVec3<int32> MaxCellIdx = FHierarchicalSpatialHashCellIdx::CellIdxForPoint(ParticleBounds.Max(), CellSize);
+				const TVec3<int32> MinCellIdx = Base::CellIdxForPoint(ParticleBounds.GetMin(), CellSizeVector);
+				const TVec3<int32> MaxCellIdx = Base::CellIdxForPoint(ParticleBounds.GetMax(), CellSizeVector);
+
 				const FPayloadAndBounds Value = { Payload, ParticleIdx };
 				HashMap.ConcurrentAddElementRange(MinCellIdx, MaxCellIdx, Lod, Value);
 
+				Bounds[ParticleIdx] = ParticleBounds;
 				if (Lod >= MinPreAllocatedUsedLodValue && Lod <= MaxPreAllocatedUsedLodValue)
 				{
 					UsedLodsArray[Lod - MinPreAllocatedUsedLodValue] = true;
@@ -301,22 +427,23 @@ public:
 		{
 			if (UsedLodsArray[Lod - MinPreAllocatedUsedLodValue])
 			{
-				UsedLods.Add(Lod, FHierarchicalSpatialHashCellIdx::CellSizeForLod<T>(Lod));
+				UsedLods.Add(Lod, CellSizeForLod(Lod));
 			}
 		}
-	}
-
-	TArray<PayloadType> FindAllIntersections(const TAABB<T, 3>& LookupBounds) const 
+	}	
+	
+	TArray<PayloadType> FindAllIntersections(const FVectorAABB LookupBounds,
+		TFunctionRef<bool(const int32 Payload)> BroadphaseTest) const
 	{
+		// LookupBounds intentionally passed by value--it seems faster if it gets copied locally.
 		TArray<PayloadType> Result;
-		FVectorAABB LookupBoundsVec(LookupBounds);
 		for (const typename TMap<int32, T>::ElementType& UsedLod : UsedLods)
 		{
 			const int32 Lod = UsedLod.Key;
-			const T CellSize = UsedLod.Value;
+			const TVectorRegisterType<T> CellSize = MakeVectorRegister(UsedLod.Value, UsedLod.Value, UsedLod.Value, UsedLod.Value);
 
-			const TVec3<int32> MinCellIdx = FHierarchicalSpatialHashCellIdx::CellIdxForPoint(LookupBounds.Min(), CellSize);
-			const TVec3<int32> MaxCellIdx = FHierarchicalSpatialHashCellIdx::CellIdxForPoint(LookupBounds.Max(), CellSize);
+			const TVec3<int32> MinCellIdx = Base::CellIdxForPoint(LookupBounds.GetMin(), CellSize);
+			const TVec3<int32> MaxCellIdx = Base::CellIdxForPoint(LookupBounds.GetMax(), CellSize);
 
 			for (int32 XIdx = MinCellIdx[0]; XIdx <= MaxCellIdx[0]; ++XIdx)
 			{
@@ -324,15 +451,20 @@ public:
 				{
 					for (int32 ZIdx = MinCellIdx[2]; ZIdx <= MaxCellIdx[2]; ++ZIdx)
 					{
-						const FHierarchicalSpatialHashCellIdx Key{ {XIdx, YIdx, ZIdx}, Lod };
+						const FHashIndex Key{ {XIdx, YIdx, ZIdx}, Lod };
 						for (int32 HashIndex = HashMap.First(Key); HashMap.IsValid(HashIndex); HashIndex = HashMap.Next(HashIndex))
 						{
 							if (HashMap.KeyMatches(HashIndex, Key))
 							{
-								const FPayloadAndBounds Value = HashMap.GetValue(HashIndex);
+								const FPayloadAndBounds& Value = HashMap.GetValue(HashIndex);
+								if (!BroadphaseTest(Value.Payload))
+								{
+									continue;
+								}
+
 								const FVectorAABB& PayloadBounds = Bounds[Value.BoundsLookupIdx];
 
-								if (PayloadBounds.Intersects(LookupBoundsVec))
+								if (PayloadBounds.Intersects(LookupBounds))
 								{
 									// Typical cloth examples do not find a lot of unique results. 
 									// In testing, the overhead to do this search is lower than using a TSet and converting the result to an array
@@ -347,6 +479,11 @@ public:
 		return Result;
 	}
 
+	TArray<PayloadType> FindAllIntersections(const TAABB<T, 3>& LookupBounds) const
+	{
+		return FindAllIntersections(FVectorAABB(LookupBounds), [](const int32) { return true; });
+	}
+
 	TArray<PayloadType> FindAllIntersections(const TVec3<T>& Point) const
 	{
 		TVectorRegisterType<T> PointVec(MakeVectorRegister(Point.X, Point.Y, Point.Z, (T)0));
@@ -355,10 +492,10 @@ public:
 		for (const typename TMap<int32, T>::ElementType& UsedLod : UsedLods)
 		{
 			const int32 Lod = UsedLod.Key;
-			const T CellSize = UsedLod.Value;
+			const TVectorRegisterType<T> CellSize = MakeVectorRegister(UsedLod.Value, UsedLod.Value, UsedLod.Value, UsedLod.Value);
 
-			const TVec3<int32> CellIdx = FHierarchicalSpatialHashCellIdx::CellIdxForPoint(Point, CellSize);
-			const FHierarchicalSpatialHashCellIdx Key{ CellIdx, Lod };
+			const TVec3<int32> CellIdx = Base::CellIdxForPoint(PointVec, CellSize);
+			const FHashIndex Key{ CellIdx, Lod };
 			for (int32 HashIndex = HashMap.First(Key); HashMap.IsValid(HashIndex); HashIndex = HashMap.Next(HashIndex))
 			{
 				if (HashMap.KeyMatches(HashIndex, Key))
@@ -379,6 +516,7 @@ public:
 	}
 
 private:
+
 	TMap<int32, T> UsedLods;
 	TArray<FVectorAABB> Bounds;
 };
@@ -389,12 +527,13 @@ class TSpatialHashGridPoints : public TSpatialHashGridBase<PayloadType>
 	typedef TSpatialHashGridBase<PayloadType> Base;
 	using typename Base::FPayloadAndBounds;
 	using typename Base::FHashMap;
+	using typename Base::FHashIndex;
 	using Base::HashMap;
 
 public:
 	TSpatialHashGridPoints(const T InCellSize)
 		: Base()
-		, CellSize(InCellSize)
+		, CellSize(MakeVectorRegister(InCellSize, InCellSize, InCellSize, InCellSize))
 	{}
 
 	void Reset()
@@ -416,7 +555,8 @@ public:
 		{
 			const auto& Particle = Particles[ParticleIdx];
 			const PayloadType Payload = Particle.template GetPayload<PayloadType>(ParticleIdx);
-			const TVec3<int32> CellIdx = FHierarchicalSpatialHashCellIdx::CellIdxForPoint(Particle.X(), CellSize);
+			const auto& ParticleX = Particle.X();
+			const TVec3<int32> CellIdx = Base::CellIdxForPoint(MakeVectorRegister(ParticleX[0], ParticleX[1], ParticleX[2], 0), CellSize);
 			const FPayloadAndBounds Value = { Payload, ParticleIdx };
 			const int32 ElementIdx = HashMap.ConcurrentAddElement(CellIdx, INDEX_NONE, Value);
 			ParticleIndexToElement[ParticleIdx] = ElementIdx;
@@ -426,110 +566,128 @@ public:
 		HashMap.ShrinkElementsAfterConcurrentAdd();
 	}
 
-	TArray<TVec2<PayloadType>> FindAllSelfProximities(int32 CellRadius, int32 MaxNumExpectedConnections, TFunctionRef<bool(const int32 Payload0, const int32 Payload1)> NarrowTest)
+	TArray<TVec2<PayloadType>> FindAllSelfProximities(int32 CellRadius, int32 MaxNumExpectedConnections, TFunctionRef<bool(const int32 Payload0, const int32 Payload1)> NarrowTest) const
 	{
 		check(CellRadius >= 0);
 
-		// Resize using a RWLock. Start with MaxNumExpectedConnections
-		TArray<TVec2<PayloadType>> Result;
-		Result.SetNum(MaxNumExpectedConnections);
-		Result.SetNum(Result.Max()); // Remove all slack.
-		std::atomic<int32> ResultIndex(0);
-		std::atomic<int32> CurrentResultNum(Result.Num()); // Arrays aren't thread-safe, so use this to track size instead.
-		FRWLock ResultResizeRWLock;
-
-		auto AddResult = [this, &NarrowTest, &Result, &ResultIndex, &CurrentResultNum, &ResultResizeRWLock]
-		(const TVec2<PayloadType>& Proximity)
+		// Working around MSVC lambda bug with this functor.
+		struct FParticleProximities
 		{
-			if (!NarrowTest(Proximity[0], Proximity[1]))
+			const int32 CellRadius;
+			TFunctionRef<bool(const int32 Payload0, const int32 Payload1)> NarrowTest;
+			const FHashMap& HashMap;
+			const TArray<int32>& ParticleIndexToElement;
+
+			// Resize using a RWLock. Start with MaxNumExpectedConnections
+			TArray<TVec2<PayloadType>> Result;
+			std::atomic<int32> ResultIndex;
+			std::atomic<int32> CurrentResultNum; // Arrays aren't thread-safe, so use this to track size instead.
+			FRWLock ResultResizeRWLock;
+
+			FParticleProximities(const int32 InCellRadius, const int32 InMaxNumExpectedConnections,
+				TFunctionRef<bool(const int32 Payload0, const int32 Payload1)> InNarrowTest,
+				const FHashMap& InHashMap, const TArray<int32>& InParticleIndexToElement)
+				:CellRadius(InCellRadius), NarrowTest(InNarrowTest), HashMap(InHashMap), ParticleIndexToElement(InParticleIndexToElement), ResultIndex(0)
 			{
-				return;
+				Result.SetNum(InMaxNumExpectedConnections);
+				Result.SetNum(Result.Max()); // Remove all slack.
+				CurrentResultNum.store(Result.Num());
 			}
-			const int32 NewResultIndex = ResultIndex.fetch_add(1);
-			if (NewResultIndex < CurrentResultNum.load())
-			{
-				ResultResizeRWLock.ReadLock();
-				Result[NewResultIndex] = Proximity;
-				ResultResizeRWLock.ReadUnlock();
-			}
-			else
+
+			void AddResult(const TVec2<PayloadType>& Proximity)
 			{
-				ResultResizeRWLock.WriteLock();
-				if (NewResultIndex >= CurrentResultNum.load())
+				if (!NarrowTest(Proximity[0], Proximity[1]))
 				{
-					// Use array's allocator to decide how much slack to create
-					Result.SetNum(NewResultIndex+1);
-					Result.SetNum(Result.Max());
-					CurrentResultNum.store(Result.Num());
+					return;
 				}
-				Result[NewResultIndex] = Proximity;
-				ResultResizeRWLock.WriteUnlock();
-			}
-		};
-
-		auto AddResultsInCell = [this,&AddResult](PayloadType ThisPayload, const FHierarchicalSpatialHashCellIdx& Cell)
-		{
-			for (int32 HashIndex = HashMap.First(Cell); HashMap.IsValid(HashIndex); HashIndex = HashMap.Next(HashIndex))
-			{
-				if (HashMap.KeyMatches(HashIndex, Cell))
+				const int32 NewResultIndex = ResultIndex.fetch_add(1);
+				if (NewResultIndex < CurrentResultNum.load())
 				{
-					AddResult(TVec2<PayloadType>(ThisPayload, HashMap.GetValue(HashIndex).Payload));
+					ResultResizeRWLock.ReadLock();
+					Result[NewResultIndex] = Proximity;
+					ResultResizeRWLock.ReadUnlock();
+				}
+				else
+				{
+					ResultResizeRWLock.WriteLock();
+					if (NewResultIndex >= CurrentResultNum.load())
+					{
+						// Use array's allocator to decide how much slack to create
+						Result.SetNum(NewResultIndex + 1);
+						Result.SetNum(Result.Max());
+						CurrentResultNum.store(Result.Num());
+					}
+					Result[NewResultIndex] = Proximity;
+					ResultResizeRWLock.WriteUnlock();
 				}
 			}
-		};
 
-		PhysicsParallelFor(ParticleIndexToElement.Num(),
-			[this, &AddResult, &AddResultsInCell, CellRadius](int32 ParticleIdx)
-		{
-			const int32 ThisElementIdx = ParticleIndexToElement[ParticleIdx];
-			const int32 ThisPayload = HashMap.GetValue(ThisElementIdx).Payload;
-			const FHierarchicalSpatialHashCellIdx& ThisKey = HashMap.GetKey(ThisElementIdx);
-
-			// All neighbors in our cell after us in the list
-			int32 OtherElementIdx = HashMap.Next(ThisElementIdx);
-			for (; HashMap.IsValid(OtherElementIdx); OtherElementIdx = HashMap.Next(OtherElementIdx))
+			void AddResultsInCell(PayloadType ThisPayload, const FHashIndex& Cell)
 			{
-				if (HashMap.KeyMatches(OtherElementIdx, ThisKey))
+				for (int32 HashIndex = HashMap.First(Cell); HashMap.IsValid(HashIndex); HashIndex = HashMap.Next(HashIndex))
 				{
-					AddResult(TVec2<PayloadType>(ThisPayload, HashMap.GetValue(OtherElementIdx).Payload));
+					if (HashMap.KeyMatches(HashIndex, Cell))
+					{
+						AddResult(TVec2<PayloadType>(ThisPayload, HashMap.GetValue(HashIndex).Payload));
+					}
 				}
 			}
 
-			// Iterate over cells within CellRadius away. Only look "forward". Particles behind us will find us by looking forward.
-			// I dunno... there's probably a better way to do this.....
-			for (int32 ZIndex = 1; ZIndex <= CellRadius; ++ZIndex)
+			void operator()(int32 ParticleIdx)
 			{
-				AddResultsInCell(ThisPayload, FHierarchicalSpatialHashCellIdx{ ThisKey.Cell + TVec3<int32>(0,0,ZIndex), INDEX_NONE });
-			}
+				const int32 ThisElementIdx = ParticleIndexToElement[ParticleIdx];
+				const int32 ThisPayload = HashMap.GetValue(ThisElementIdx).Payload;
+				const FHashIndex& ThisKey = HashMap.GetKey(ThisElementIdx);
 
-			for (int32 YIndex = 1; YIndex <= CellRadius; ++YIndex)
-			{
-				for (int32 ZIndex = -CellRadius; ZIndex <= CellRadius; ++ZIndex)
+				// All neighbors in our cell after us in the list
+				int32 OtherElementIdx = HashMap.Next(ThisElementIdx);
+				for (; HashMap.IsValid(OtherElementIdx); OtherElementIdx = HashMap.Next(OtherElementIdx))
 				{
-					AddResultsInCell(ThisPayload, FHierarchicalSpatialHashCellIdx{ ThisKey.Cell + TVec3<int32>(0,YIndex,ZIndex), INDEX_NONE });
+					if (HashMap.KeyMatches(OtherElementIdx, ThisKey))
+					{
+						AddResult(TVec2<PayloadType>(ThisPayload, HashMap.GetValue(OtherElementIdx).Payload));
+					}
 				}
-			}
 
-			for (int32 XIndex = 1; XIndex <= CellRadius; ++XIndex)
-			{
-				for (int32 YIndex = -CellRadius; YIndex <= CellRadius; ++YIndex)
+				// Iterate over cells within CellRadius away. Only look "forward". Particles behind us will find us by looking forward.
+				// I dunno... there's probably a better way to do this.....
+				for (int32 ZIndex = 1; ZIndex <= CellRadius; ++ZIndex)
+				{
+					AddResultsInCell(ThisPayload, FHashIndex(VectorIntAdd(ThisKey.VectorIndex, MakeVectorRegisterInt(0, 0, ZIndex, 0))));
+				}
+
+				for (int32 YIndex = 1; YIndex <= CellRadius; ++YIndex)
 				{
 					for (int32 ZIndex = -CellRadius; ZIndex <= CellRadius; ++ZIndex)
 					{
-						AddResultsInCell(ThisPayload, FHierarchicalSpatialHashCellIdx{ ThisKey.Cell + TVec3<int32>(XIndex,YIndex,ZIndex), INDEX_NONE });
+						AddResultsInCell(ThisPayload, FHashIndex(VectorIntAdd(ThisKey.VectorIndex, MakeVectorRegisterInt(0, YIndex, ZIndex, 0))));
+					}
+				}
+
+				for (int32 XIndex = 1; XIndex <= CellRadius; ++XIndex)
+				{
+					for (int32 YIndex = -CellRadius; YIndex <= CellRadius; ++YIndex)
+					{
+						for (int32 ZIndex = -CellRadius; ZIndex <= CellRadius; ++ZIndex)
+						{
+							AddResultsInCell(ThisPayload, FHashIndex(VectorIntAdd(ThisKey.VectorIndex, MakeVectorRegisterInt(XIndex, YIndex, ZIndex, 0))));
+						}
 					}
 				}
+
 			}
-		}
-		);
+			
+		} ConcurrentResults{ CellRadius, MaxNumExpectedConnections, NarrowTest, HashMap, ParticleIndexToElement };
+
+		PhysicsParallelFor(ParticleIndexToElement.Num(), ConcurrentResults);
 
 		// Shrink Result array to actual number of found proximities
-		const int32 ResultNum = ResultIndex.load();
-		Result.SetNum(ResultNum, /*bAllowShrinking*/ true);
-		return Result;
+		const int32 ResultNum = ConcurrentResults.ResultIndex.load();
+		ConcurrentResults.Result.SetNum(ResultNum, /*bAllowShrinking*/ true);
+		return ConcurrentResults.Result;
 	}
 private:
-	const T CellSize;
+	const TVectorRegisterType<T> CellSize;
 
 	// Particle Index is dense array index passed in at build time. 
 	TArray<int32> ParticleIndexToElement;
diff --git a/Engine/Source/Runtime/Experimental/Chaos/Public/Chaos/PBDTriangleMeshCollisions.h b/Engine/Source/Runtime/Experimental/Chaos/Public/Chaos/PBDTriangleMeshCollisions.h
--- a/Engine/Source/Runtime/Experimental/Chaos/Public/Chaos/PBDTriangleMeshCollisions.h
+++ b/Engine/Source/Runtime/Experimental/Chaos/Public/Chaos/PBDTriangleMeshCollisions.h
@@ -6,6 +6,7 @@
 #include "Chaos/TriangleMesh.h"
 #include "Chaos/HierarchicalSpatialHash.h"
 #include "Chaos/CollectionPropertyFacade.h"
+#include "Templates/PimplPtr.h"
 
 // This does initialization for PBDCollisionSpringConstraints and PBDTriangleMeshIntersections, 
 // including intersection detection and global intersection analysis
@@ -229,6 +230,10 @@ private:
 	TArray<FGIAColor> VertexGIAColors;
 	TArray<FGIAColor> TriangleGIAColors;
 
+	// Scratch buffers used by Init and PostInit. They live here so they can be reused rather than reallocated.
+	struct FScratchBuffers;
+	TPimplPtr<FScratchBuffers> ScratchBuffers;
+
 	// Debug display of intersection contours
 	TArray<TArray<FBarycentricPoint>> IntersectionContourPoints;
 	TArray<FContourType> IntersectionContourTypes;
