[Chaos] Adding a way for Geometry collection components to not permanently change the net dormancy of their owners.

This CL introduces two modes:

- ForceDormancyAwake : This is the current behavior, as soon we attempt to update replicated data we set the owner dormancy to awake
- FlushNetDormancy : Before updating replicated data, we flush the dormancy of our owner

The modes can be changed with the cvar p.Chaos.GC.NetAwakeningMode (0-1)

This change also makes UClusterUnionReplicatedProxyComponent flush its owner dormancy before updating its replicated data.

This can be disabled with the cvar p.Chaos.CU.UseFlushNetDormancy set to false

diff --git a/Engine/Source/Runtime/Engine/Classes/PhysicsEngine/ClusterUnionReplicatedProxyComponent.h b/Engine/Source/Runtime/Engine/Classes/PhysicsEngine/ClusterUnionReplicatedProxyComponent.h
--- a/Engine/Source/Runtime/Engine/Classes/PhysicsEngine/ClusterUnionReplicatedProxyComponent.h
+++ b/Engine/Source/Runtime/Engine/Classes/PhysicsEngine/ClusterUnionReplicatedProxyComponent.h
@@ -110,4 +110,7 @@ public:
 public:
 	ENGINE_API virtual void PostRepNotifies() override;
 	//~ End UObject Interface
+
+private:
+	void FlushNetDormancyIfNeeded();
 };
diff --git a/Engine/Source/Runtime/Engine/Private/PhysicsEngine/ClusterUnionReplicatedProxyComponent.cpp b/Engine/Source/Runtime/Engine/Private/PhysicsEngine/ClusterUnionReplicatedProxyComponent.cpp
--- a/Engine/Source/Runtime/Engine/Private/PhysicsEngine/ClusterUnionReplicatedProxyComponent.cpp
+++ b/Engine/Source/Runtime/Engine/Private/PhysicsEngine/ClusterUnionReplicatedProxyComponent.cpp
@@ -10,6 +10,12 @@
 
 #include UE_INLINE_GENERATED_CPP_BY_NAME(ClusterUnionReplicatedProxyComponent)
 
+namespace ClusterUnionCVars
+{
+	bool ClusterUnionUseFlushNetDormancy = true;
+	FAutoConsoleVariableRef CVarClusterUnionUseFlushNetDormancy(TEXT("p.Chaos.CU.UseFlushNetDormancy"), ClusterUnionUseFlushNetDormancy, TEXT("When true it will flush the net dormancy of the owner the next frame instead of awaking the actor"));
+}
+
 UClusterUnionReplicatedProxyComponent::UClusterUnionReplicatedProxyComponent(const FObjectInitializer& ObjectInitializer)
 	: UActorComponent(ObjectInitializer)
 {
@@ -56,18 +62,24 @@ void UClusterUnionReplicatedProxyComponent::EndPlay(const EEndPlayReason::Type E
 
 void UClusterUnionReplicatedProxyComponent::SetParentClusterUnion(UClusterUnionComponent* InComponent)
 {
+	FlushNetDormancyIfNeeded();
+
 	ParentClusterUnion = InComponent;
 	MARK_PROPERTY_DIRTY_FROM_NAME(UClusterUnionReplicatedProxyComponent, ParentClusterUnion, this);
 }
 
 void UClusterUnionReplicatedProxyComponent::SetChildClusteredComponent(UPrimitiveComponent* InComponent)
 {
+	FlushNetDormancyIfNeeded();
+
 	ChildClusteredComponent = InComponent;
 	MARK_PROPERTY_DIRTY_FROM_NAME(UClusterUnionReplicatedProxyComponent, ChildClusteredComponent, this);
 }
 
 void UClusterUnionReplicatedProxyComponent::SetParticleBoneIds(const TArray<int32>& InIds)
 {
+	FlushNetDormancyIfNeeded();
+
 	ParticleBoneIds = InIds;
 
 	ParticleChildToParents.Empty();
@@ -85,6 +97,8 @@ void UClusterUnionReplicatedProxyComponent::SetParticleChildToParent(int32 BoneI
 	int32 Index = INDEX_NONE;
 	if (ParticleBoneIds.Find(BoneId, Index))
 	{
+		FlushNetDormancyIfNeeded();
+
 		ParticleChildToParents[Index] = ChildToParent;
 		MARK_PROPERTY_DIRTY_FROM_NAME(UClusterUnionReplicatedProxyComponent, ParticleChildToParents, this);
 	}
@@ -162,6 +176,19 @@ void UClusterUnionReplicatedProxyComponent::PostRepNotifies()
 	}
 }
 
+void UClusterUnionReplicatedProxyComponent::FlushNetDormancyIfNeeded()
+{
+	if (!ClusterUnionCVars::ClusterUnionUseFlushNetDormancy)
+	{
+		return;
+	}
+
+	if (AActor* Owner = GetOwner())
+	{
+		Owner->FlushNetDormancy();
+	}
+}
+
 void UClusterUnionReplicatedProxyComponent::AddComponentToCluster()
 {
 	if (!ParentClusterUnion.IsValid() || !ChildClusteredComponent.IsValid() || IsPendingDeletion())
diff --git a/Engine/Source/Runtime/Experimental/GeometryCollectionEngine/Private/GeometryCollection/GeometryCollectionComponent.cpp b/Engine/Source/Runtime/Experimental/GeometryCollectionEngine/Private/GeometryCollection/GeometryCollectionComponent.cpp
--- a/Engine/Source/Runtime/Experimental/GeometryCollectionEngine/Private/GeometryCollection/GeometryCollectionComponent.cpp
+++ b/Engine/Source/Runtime/Experimental/GeometryCollectionEngine/Private/GeometryCollection/GeometryCollectionComponent.cpp
@@ -135,6 +135,9 @@ FAutoConsoleVariableRef CVarGeometryCollectionCreatePhysicsStateInEditor(TEXT("p
 bool GeometryCollectionUseReplicationV2 = true;
 FAutoConsoleVariableRef CVarGeometryCollectionUseReplicationV2(TEXT("p.Chaos.GC.UseReplicationV2"), GeometryCollectionUseReplicationV2, TEXT("When true use new replication data model"));
 
+int32 GeometryCollectionNetAwakeningMode = 1;
+FAutoConsoleVariableRef CVarGeometryCollectionNetAwakeningMode(TEXT("p.Chaos.GC.NetAwakeningMode"), GeometryCollectionNetAwakeningMode, TEXT("Changes how GC components ensure that their owner is awake for replication. 0 = ForceDormancyAwake, 1 = Use Flush Net Dormancy"));
+
 DEFINE_LOG_CATEGORY_STATIC(UGCC_LOG, Error, All);
 
 extern FGeometryCollectionDynamicDataPool GDynamicDataPool;
@@ -1988,6 +1991,8 @@ void UGeometryCollectionComponent::UpdateRepData()
 	
 	if (Owner && GetIsReplicated() && Owner->GetLocalRole() == ROLE_Authority)
 	{
+		FlushNetDormancyIfNeeded();
+
 		const bool bReplicateMovement = Owner->IsReplicatingMovement();
 		bool bFirstUpdate = false;
 		if(ClustersToRep == nullptr)
@@ -2196,10 +2201,14 @@ void UGeometryCollectionComponent::UpdateRepData()
 			MARK_PROPERTY_DIRTY_FROM_NAME(UGeometryCollectionComponent, RepData, this);
 			++RepData.Version;
 
-			if(Owner->NetDormancy != DORM_Awake)
+			const bool bSetDormancyToAwake = Owner->NetDormancy > DORM_Awake && GetDesiredNetAwakeningMode() == ENetAwakeningMode::ForceDormancyAwake;
+			if (bSetDormancyToAwake)
 			{
-				//If net dormancy is Initial it must be for perf reasons, but since a cluster changed we need to replicate down
-				Owner->SetNetDormancy(DORM_Awake);
+				if (Owner->NetDormancy > DORM_Awake)
+				{
+					//If net dormancy is Initial it must be for perf reasons, but since a cluster changed we need to replicate down	
+					Owner->SetNetDormancy(DORM_Awake);
+				}
 			}
 		}
 	}
@@ -2225,6 +2234,8 @@ void UGeometryCollectionComponent::UpdateRepStateAndDynamicData()
 
 	if (GetIsReplicated() && Owner->GetLocalRole() == ROLE_Authority)
 	{
+		FlushNetDormancyIfNeeded();
+
 		if (RestCollection && RestCollection->GetGeometryCollection())
 		{
 			FPBDRigidsSolver* Solver = PhysicsProxy->GetSolver<Chaos::FPBDRigidsSolver>();
@@ -2339,7 +2350,8 @@ void UGeometryCollectionComponent::UpdateRepStateAndDynamicData()
 
 			if (bDynamicChanged || bStateChanged)
 			{
-				if (Owner->NetDormancy != DORM_Awake)
+				const bool bSetDormancyToAwake = Owner->NetDormancy > DORM_Awake && GetDesiredNetAwakeningMode() == ENetAwakeningMode::ForceDormancyAwake;
+				if (bSetDormancyToAwake)
 				{
 					//If net dormancy is Initial it must be for perf reasons, but since a cluster changed we need to replicate down
 					Owner->SetNetDormancy(DORM_Awake);
@@ -6537,6 +6549,38 @@ void UGeometryCollectionComponent::PostLoad()
 	}
 }
 
+void UGeometryCollectionComponent::FlushNetDormancyIfNeeded() const
+{
+	if (GetDesiredNetAwakeningMode() != ENetAwakeningMode::FlushNetDormancy)
+	{
+		return;
+	}
+
+	if (AActor* Owner = GetOwner())
+	{
+
+		Owner->FlushNetDormancy();
+	}
+}
+
+UGeometryCollectionComponent::ENetAwakeningMode UGeometryCollectionComponent::GetDesiredNetAwakeningMode() const
+{
+	constexpr int32 MinNetAwakeningMode = 0;
+	constexpr int32 MaxNetAwakeningMode = 1;
+	ENetAwakeningMode Mode;
+	if (GeometryCollectionNetAwakeningMode < MinNetAwakeningMode || GeometryCollectionNetAwakeningMode > MaxNetAwakeningMode)
+	{
+		ensureMsgf(false, TEXT("Invalid NetAwakening mode configured, falling back to the default mode"));
+		Mode = ENetAwakeningMode::ForceDormancyAwake;
+	}
+	else
+	{
+		Mode = static_cast<ENetAwakeningMode>(GeometryCollectionNetAwakeningMode);
+	}
+
+	return Mode;
+}
+
 Chaos::FPhysicsObject* UGeometryCollectionComponent::GetPhysicsObjectById(Chaos::FPhysicsObjectId Id) const
 {
 	if (!PhysicsProxy)
diff --git a/Engine/Source/Runtime/Experimental/GeometryCollectionEngine/Public/GeometryCollection/GeometryCollectionComponent.h b/Engine/Source/Runtime/Experimental/GeometryCollectionEngine/Public/GeometryCollection/GeometryCollectionComponent.h
--- a/Engine/Source/Runtime/Experimental/GeometryCollectionEngine/Public/GeometryCollection/GeometryCollectionComponent.h
+++ b/Engine/Source/Runtime/Experimental/GeometryCollectionEngine/Public/GeometryCollection/GeometryCollectionComponent.h
@@ -1708,6 +1708,19 @@ private:
 
 	bool bUpdateCustomRenderer;
 
+private:
+
+	enum class ENetAwakeningMode
+	{
+		ForceDormancyAwake,
+		FlushNetDormancy,
+	};
+
+	/** Flushes the net dormancy of our owner if this action is enabled and we are dormant */
+	void FlushNetDormancyIfNeeded() const;
+
+	ENetAwakeningMode GetDesiredNetAwakeningMode() const;
+
 	//~ Begin IPhysicsComponent Interface.
 public:
 	GEOMETRYCOLLECTIONENGINE_API virtual Chaos::FPhysicsObject* GetPhysicsObjectById(Chaos::FPhysicsObjectId Id) const override;
