cluster union connectivity released particles should also respect the throttling for auxiliary particles

diff --git a/Engine/Source/Runtime/Experimental/Chaos/Private/Chaos/ClusterUnionManager.cpp b/Engine/Source/Runtime/Experimental/Chaos/Private/Chaos/ClusterUnionManager.cpp
--- a/Engine/Source/Runtime/Experimental/Chaos/Private/Chaos/ClusterUnionManager.cpp
+++ b/Engine/Source/Runtime/Experimental/Chaos/Private/Chaos/ClusterUnionManager.cpp
@@ -759,10 +759,22 @@ namespace Chaos
 		}
 
 		ParticleIndicesToRemove.Sort();
+
+		TArray<FPBDRigidParticleHandle*> NonMainReleasedParticles;
+		NonMainReleasedParticles.Reserve(ParticleIndicesToRemove.Num());
+
 		for (int32 Index = ParticleIndicesToRemove.Num() - 1; Index >= 0; --Index)
 		{
 			const int32 ParticleIndex = ParticleIndicesToRemove[Index];
-			Cluster->ChildProperties.Remove(Cluster->ChildParticles[ParticleIndex]);
+			FPBDRigidParticleHandle* Particle = Cluster->ChildParticles[ParticleIndex];
+			if (const FClusterUnionParticleProperties* Properties = Cluster->ChildProperties.Find(Particle))
+			{
+				if (Properties->bIsAuxiliaryParticle)
+				{
+					NonMainReleasedParticles.Add(Particle);
+				}
+				Cluster->ChildProperties.Remove(Particle);
+			}
 
 			// This can't be RemoveAtSwap otherwise there will be a mismatch between the index of a particle
 			// and the index of its corresponding shape in the cluster union's shape array. There is currently
@@ -771,6 +783,11 @@ namespace Chaos
 		}
 
 		MClustering.RemoveParticlesFromCluster(Cluster->InternalCluster, ParticleSet.Array());
+		if (MClustering.ShouldThrottleParticleRelease())
+		{
+			// Only want to throttle the release of non-main particles.
+			MClustering.ThrottleReleasedParticlesIfNecessary(NonMainReleasedParticles);
+		}
 
 		// Removing a particle should have no bearing on the proxy of the cluster.
 		// This gets changed because we go through an internal initialization route when we update the cluster union particle's properties.
diff --git a/Engine/Source/Runtime/Experimental/Chaos/Private/Chaos/PBDRigidClustering.cpp b/Engine/Source/Runtime/Experimental/Chaos/Private/Chaos/PBDRigidClustering.cpp
--- a/Engine/Source/Runtime/Experimental/Chaos/Private/Chaos/PBDRigidClustering.cpp
+++ b/Engine/Source/Runtime/Experimental/Chaos/Private/Chaos/PBDRigidClustering.cpp
@@ -126,7 +126,7 @@ namespace Chaos
 			return Current;
 		}
 
-		bool ShouldThrottleParticleRelease()
+		bool CVarShouldThrottleParticleRelease()
 		{
 			return (ClusteringParticleReleaseThrottlingMinCount >= 0 && ClusteringParticleReleaseThrottlingMaxCount >= 0);
 		}
@@ -150,6 +150,36 @@ namespace Chaos
 
 			return 1.0f;
 		}
+
+		template<typename TParticleContainer>
+		void GenericThrottleReleasedParticlesIfNecessary(TParticleContainer& Container, typename FRigidClustering::FRigidEvolution& MEvolution)
+		{
+			if (!CVarShouldThrottleParticleRelease())
+			{
+				return;
+			}
+
+			const float RatioOfParticlesToDisable = GetRatioOfReleasedParticlesToDisable(MEvolution);
+			const int32 NumberOfParticlesToDisable = (int32)((float)Container.Num() * RatioOfParticlesToDisable);
+			if (NumberOfParticlesToDisable > 0)
+			{ 
+				int32 DisabledParticleCount = 0;
+				for (auto ChildIt = Container.CreateIterator(); ChildIt; ++ChildIt)
+				{
+					if (FPBDRigidParticleHandle* Child = *ChildIt)
+					{
+						DisabledParticleCount++;
+						MEvolution.DisableParticle(Child);
+						MEvolution.GetParticles().MarkTransientDirtyParticle(Child);
+						ChildIt.RemoveCurrent();
+					}
+					if (DisabledParticleCount >= NumberOfParticlesToDisable)
+					{
+						break;
+					}
+				}
+			}
+		}
 	}
 	
 	//==========================================================================
@@ -1200,30 +1230,8 @@ namespace Chaos
 		}
 
 		// optimization : start disabling activated children if the number of active released particle is too high
-		if (ShouldThrottleParticleRelease())
-		{
-			const float RatioOfParticlesToDisable = GetRatioOfReleasedParticlesToDisable(MEvolution);
-			const int32 NumberOfParticlesToDisable = (int32)((float)ActivatedChildren.Num() * RatioOfParticlesToDisable);
-			if (NumberOfParticlesToDisable > 0)
-			{ 
-				int32 DisabledParticleCount = 0;
-				for (auto ChildIt = ActivatedChildren.CreateIterator(); ChildIt; ++ChildIt)
-				{
-					if (FPBDRigidParticleHandle* Child = *ChildIt)
-					{
-						DisabledParticleCount++;
-						MEvolution.DisableParticle(Child);
-						MEvolution.GetParticles().MarkTransientDirtyParticle(Child);
-						ChildIt.RemoveCurrent();
-					}
-					if (DisabledParticleCount >= NumberOfParticlesToDisable)
-					{
-						break;
-					}
-				}
-			}
-		}
-		
+		ThrottleReleasedParticlesIfNecessary(ActivatedChildren);
+
 		FrameReleasedChildren += ActivatedChildren.Num();
 
 		return ActivatedChildren;
@@ -2915,4 +2923,19 @@ namespace Chaos
 		RemoveFilteredNodeConnections(ClusteredChild, true);
 	}
 
+	bool FRigidClustering::ShouldThrottleParticleRelease() const
+	{
+		return CVarShouldThrottleParticleRelease();
+	}
+
+	void FRigidClustering::ThrottleReleasedParticlesIfNecessary(TSet<FPBDRigidParticleHandle*>& Particles) const
+	{
+		GenericThrottleReleasedParticlesIfNecessary(Particles, MEvolution);
+	}
+
+	void FRigidClustering::ThrottleReleasedParticlesIfNecessary(TArray<FPBDRigidParticleHandle*>& Particles) const
+	{
+		GenericThrottleReleasedParticlesIfNecessary(Particles, MEvolution);
+	}
+
 } // namespace Chaos
diff --git a/Engine/Source/Runtime/Experimental/Chaos/Public/Chaos/PBDRigidClustering.h b/Engine/Source/Runtime/Experimental/Chaos/Public/Chaos/PBDRigidClustering.h
--- a/Engine/Source/Runtime/Experimental/Chaos/Public/Chaos/PBDRigidClustering.h
+++ b/Engine/Source/Runtime/Experimental/Chaos/Public/Chaos/PBDRigidClustering.h
@@ -491,6 +491,11 @@ public:
 	CHAOS_API TSet<FPBDRigidParticleHandle*> HandleConnectivityOnReleaseClusterParticle(FPBDRigidClusteredParticleHandle* ClusteredParticle, bool bCreateNewClusters);
 
 	CHAOS_API void DisableCluster(FPBDRigidClusteredParticleHandle* ClusteredParticle);
+
+	bool ShouldThrottleParticleRelease() const;
+	void ThrottleReleasedParticlesIfNecessary(TSet<FPBDRigidParticleHandle*>& Particles) const;
+	void ThrottleReleasedParticlesIfNecessary(TArray<FPBDRigidParticleHandle*>& Particles) const;
+
  protected:
 
 	CHAOS_API void ComputeStrainFromCollision(const FPBDCollisionConstraints& CollisionRule, const FReal Dt);
