Add GetKinematicParticle to FReadPhysicsObjectInterface

diff --git a/Engine/Source/Runtime/Experimental/Chaos/Private/Chaos/PhysicsObjectInterface.cpp b/Engine/Source/Runtime/Experimental/Chaos/Private/Chaos/PhysicsObjectInterface.cpp
--- a/Engine/Source/Runtime/Experimental/Chaos/Private/Chaos/PhysicsObjectInterface.cpp
+++ b/Engine/Source/Runtime/Experimental/Chaos/Private/Chaos/PhysicsObjectInterface.cpp
@@ -349,6 +349,23 @@ namespace Chaos
 		return Handle->GetParticle<Id>();
 	}
 
+	template<EThreadContext Id>
+	TThreadKinematicParticle<Id>* FReadPhysicsObjectInterface<Id>::GetKinematicParticle(const FConstPhysicsObjectHandle Handle)
+	{
+		if (Handle)
+		{
+			if (TThreadParticle<Id>* Particle = Handle->GetParticle<Id>())
+			{
+				if (TThreadKinematicParticle<Id>* KinematicParticle = Particle->CastToKinematicParticle())
+				{
+					return KinematicParticle;
+				}
+			}
+		}
+
+		return nullptr;
+	}
+
 	template<EThreadContext Id>
 	TThreadRigidParticle<Id>* FReadPhysicsObjectInterface<Id>::GetRigidParticle(const FConstPhysicsObjectHandle Handle)
 	{
diff --git a/Engine/Source/Runtime/Experimental/Chaos/Public/Chaos/PhysicsObjectInterface.h b/Engine/Source/Runtime/Experimental/Chaos/Public/Chaos/PhysicsObjectInterface.h
--- a/Engine/Source/Runtime/Experimental/Chaos/Public/Chaos/PhysicsObjectInterface.h
+++ b/Engine/Source/Runtime/Experimental/Chaos/Public/Chaos/PhysicsObjectInterface.h
@@ -73,6 +73,7 @@ namespace Chaos
 		CHAOS_API FSpatialAccelerationIdx GetSpatialIndex(const FConstPhysicsObjectHandle Object);
 
 		CHAOS_API TThreadParticle<Id>* GetParticle(const FConstPhysicsObjectHandle Object);
+		CHAOS_API TThreadKinematicParticle<Id>* GetKinematicParticle(const FConstPhysicsObjectHandle Object);
 		CHAOS_API TThreadRigidParticle<Id>* GetRigidParticle(const FConstPhysicsObjectHandle Object);
 		CHAOS_API TArray<TThreadParticle<Id>*> GetAllParticles(TArrayView<const FConstPhysicsObjectHandle> InObjects);
 		CHAOS_API TArray<TThreadRigidParticle<Id>*> GetAllRigidParticles(TArrayView<const FConstPhysicsObjectHandle> InObjects);
