Prevent Transforms allocation from dynamic Collection inside the physics thread

diff --git a/Engine/Source/Runtime/Experimental/Chaos/Private/PhysicsProxy/GeometryCollectionPhysicsProxy.cpp b/Engine/Source/Runtime/Experimental/Chaos/Private/PhysicsProxy/GeometryCollectionPhysicsProxy.cpp
--- a/Engine/Source/Runtime/Experimental/Chaos/Private/PhysicsProxy/GeometryCollectionPhysicsProxy.cpp
+++ b/Engine/Source/Runtime/Experimental/Chaos/Private/PhysicsProxy/GeometryCollectionPhysicsProxy.cpp
@@ -3517,7 +3517,8 @@ void FGeometryCollectionPhysicsProxy::BufferPhysicsResults_Internal(Chaos::FPBDR
 			if (bHasChanged)
 			{
 				// default is what we have on the collection
-				FTransform ParentSpaceTransform = FTransform(PhysicsThreadCollection.GetTransform(TransformGroupIndex));
+				const FTransform3f& ParentSpaceTransform3f = PhysicsThreadCollection.GetTransform(TransformGroupIndex);
+				FTransform ParentSpaceTransform = FTransform(ParentSpaceTransform3f);
 
 				// recompute parent space transform if there's no more parent or if the parent is an internal cluster 
 				if (!ClusterParent || ClusterParent->InternalCluster())
@@ -3536,8 +3537,11 @@ void FGeometryCollectionPhysicsProxy::BufferPhysicsResults_Internal(Chaos::FPBDR
 				Results.SetState(EntryIndex, StateData);
 				Results.SetPositions(EntryIndex, PositionData);
 				Results.SetVelocities(EntryIndex, VelocityData);
-				// todo(chaos) : we shoudl eventually get rid of the transform in the Physics collection
-				PhysicsThreadCollection.SetTransform(TransformGroupIndex, FTransform3f(ParentSpaceTransform));
+				FTransform3f NewParentSpaceTransform3f(ParentSpaceTransform);
+				if (!NewParentSpaceTransform3f.Equals(ParentSpaceTransform3f))
+				{
+					PhysicsThreadCollection.SetTransform(TransformGroupIndex, ParentSpaceTransform3f);
+				}
 				IsObjectDynamic = true;
 			}
 
