Convert GC ComponentSpaceTransform in single precision

diff --git a/Engine/Plugins/Experimental/ChaosNiagara/Source/ChaosNiagara/Private/NiagaraDataInterfaceChaosDestruction.cpp b/Engine/Plugins/Experimental/ChaosNiagara/Source/ChaosNiagara/Private/NiagaraDataInterfaceChaosDestruction.cpp
--- a/Engine/Plugins/Experimental/ChaosNiagara/Source/ChaosNiagara/Private/NiagaraDataInterfaceChaosDestruction.cpp
+++ b/Engine/Plugins/Experimental/ChaosNiagara/Source/ChaosNiagara/Private/NiagaraDataInterfaceChaosDestruction.cpp
@@ -2052,7 +2052,7 @@ void UNiagaraDataInterfaceChaosDestruction::HandleTrailingEvents(const Chaos::FT
 						}
 					}
 
-					const TArray<FTransform>& ComponentSpaceTransforms = GeometryCollectionComponent->GetComponentSpaceTransforms();
+					const TArray<FTransform3f>& ComponentSpaceTransforms = GeometryCollectionComponent->GetComponentSpaceTransforms3f();
 					const FTransform ActorTransform = GeometryCollectionComponent->GetComponentToWorld();
 
 					const FGeometryDynamicCollection* DynamicCollection = GeometryCollectionComponent->GetDynamicCollection();
@@ -2085,7 +2085,7 @@ void UNiagaraDataInterfaceChaosDestruction::HandleTrailingEvents(const Chaos::FT
 									{
 										Chaos::FTrailingDataExt TrailingData;
 
-										const FTransform CurrTransform = CollectionMassToLocal[Idx] * ComponentSpaceTransforms[Idx] * ActorTransform;
+										const FTransform CurrTransform = CollectionMassToLocal[Idx] * FTransform(ComponentSpaceTransforms[Idx]) * ActorTransform;
 										TrailingData.Location = CurrTransform.GetTranslation();
 
 										TrailingData.Velocity = (*LinearVelocity)[Idx];
diff --git a/Engine/Plugins/Experimental/ChaosNiagara/Source/ChaosNiagara/Private/NiagaraDataInterfaceGeometryCollection.cpp b/Engine/Plugins/Experimental/ChaosNiagara/Source/ChaosNiagara/Private/NiagaraDataInterfaceGeometryCollection.cpp
--- a/Engine/Plugins/Experimental/ChaosNiagara/Source/ChaosNiagara/Private/NiagaraDataInterfaceGeometryCollection.cpp
+++ b/Engine/Plugins/Experimental/ChaosNiagara/Source/ChaosNiagara/Private/NiagaraDataInterfaceGeometryCollection.cpp
@@ -176,7 +176,7 @@ void FNDIGeometryCollectionData::Update(UNiagaraDataInterfaceGeometryCollection*
 				Collection = Interface->GeometryCollectionActor->GetGeometryCollectionComponent()->RestCollection->GetGeometryCollection();
 			const TManagedArray<FBox>& BoundingBoxes = Collection->BoundingBox;
 			const TManagedArray<int32>& TransformIndexArray = Collection->TransformIndex;
-			const TArray<FTransform>& ComponentSpaceTransforms = Interface->GeometryCollectionActor->GetGeometryCollectionComponent()->GetComponentSpaceTransforms();
+			const TArray<FTransform3f>& ComponentSpaceTransforms = Interface->GeometryCollectionActor->GetGeometryCollectionComponent()->GetComponentSpaceTransforms3f();
 			const TManagedArray<TSet<int32>>& Children = Collection->Children;
 			
 			int NumPieces = 0;
@@ -225,7 +225,7 @@ void FNDIGeometryCollectionData::Update(UNiagaraDataInterfaceGeometryCollection*
 					FVector LocalTranslation = (CurrBox.Max + CurrBox.Min) * .5;
 					FTransform LocalOffset(LocalTranslation);
 
-					const FTransform3f CurrTransform(LocalOffset * ComponentSpaceTransforms[CurrTransformIndex] * ActorTransform);
+					const FTransform3f CurrTransform(LocalOffset * FTransform(ComponentSpaceTransforms[CurrTransformIndex]) * ActorTransform);
 					CurrTransform.ToMatrixWithScale().To3x4MatrixTranspose(&AssetArrays->WorldTransformBuffer[TransformIndex].X);
 
 					const FTransform3f CurrInverse = CurrTransform.Inverse();
diff --git a/Engine/Plugins/Experimental/GeometryCollectionPlugin/Source/GeometryCollectionNodes/Private/Dataflow/GeometryCollectionMeshNodes.cpp b/Engine/Plugins/Experimental/GeometryCollectionPlugin/Source/GeometryCollectionNodes/Private/Dataflow/GeometryCollectionMeshNodes.cpp
--- a/Engine/Plugins/Experimental/GeometryCollectionPlugin/Source/GeometryCollectionNodes/Private/Dataflow/GeometryCollectionMeshNodes.cpp
+++ b/Engine/Plugins/Experimental/GeometryCollectionPlugin/Source/GeometryCollectionNodes/Private/Dataflow/GeometryCollectionMeshNodes.cpp
@@ -173,7 +173,7 @@ void FCollectionToMeshDataflowNode::Evaluate(Dataflow::FContext& Context, const
 		{
 			if (TUniquePtr<FGeometryCollection> GeomCollection = TUniquePtr<FGeometryCollection>(InCollection.NewCopy<FGeometryCollection>()))
 			{
-				const TManagedArray<FTransform>& BoneTransforms = InCollection.GetAttribute<FTransform>("Transform", FGeometryCollection::TransformGroup);
+				const TManagedArray<FTransform3f>& BoneTransforms = InCollection.GetAttribute<FTransform3f>("Transform", FGeometryCollection::TransformGroup);
 
 				TArray<int32> TransformIndices;
 				TransformIndices.AddUninitialized(BoneTransforms.Num());
diff --git a/Engine/Plugins/MeshPainting/Source/MeshPaintingToolset/Private/MeshPaintGeometryCollectionAdapter.cpp b/Engine/Plugins/MeshPainting/Source/MeshPaintingToolset/Private/MeshPaintGeometryCollectionAdapter.cpp
--- a/Engine/Plugins/MeshPainting/Source/MeshPaintingToolset/Private/MeshPaintGeometryCollectionAdapter.cpp
+++ b/Engine/Plugins/MeshPainting/Source/MeshPaintingToolset/Private/MeshPaintGeometryCollectionAdapter.cpp
@@ -86,15 +86,15 @@ bool FMeshPaintGeometryCollectionComponentAdapter::InitializeVertexData()
 	const TManagedArray<int32>& BoneMap = Collection->BoneMap;
 	const TManagedArray<int32>& SimTypes = Collection->SimulationType;
 	const TManagedArray<bool>& Visible = Collection->Visible;
-	const TArray<FTransform>& ComponentSpaceTransforms = GeometryCollectionComponent->GetComponentSpaceTransforms();
+	const TArray<FTransform3f>& ComponentSpaceTransforms = GeometryCollectionComponent->GetComponentSpaceTransforms3f();
 	const int32 NumVertices = Vertices.Num();
 	MeshVertices.Reset();
 	MeshVertices.AddZeroed(NumVertices);
 	for (int32 Index = 0; Index < NumVertices; ++Index)
 	{
 		const int32 BoneIndex = BoneMap[Index];
-		const FTransform& Transform = ComponentSpaceTransforms[BoneIndex];
-		const FVector Position = Transform.TransformPosition((FVector)Vertices[Index]);
+		const FTransform3f& Transform = ComponentSpaceTransforms[BoneIndex];
+		const FVector Position = FVector(Transform.TransformPosition(Vertices[Index]));
 		MeshVertices[Index] = Position;
 	}
 
diff --git a/Engine/Source/Runtime/Experimental/Chaos/Private/GeometryCollection/GeometryCollectionAlgo.cpp b/Engine/Source/Runtime/Experimental/Chaos/Private/GeometryCollection/GeometryCollectionAlgo.cpp
--- a/Engine/Source/Runtime/Experimental/Chaos/Private/GeometryCollection/GeometryCollectionAlgo.cpp
+++ b/Engine/Source/Runtime/Experimental/Chaos/Private/GeometryCollection/GeometryCollectionAlgo.cpp
@@ -348,7 +348,8 @@ namespace GeometryCollectionAlgo
 		return true;
 	}
 
-	void GlobalMatricesHelper(const int32 Index, const FGeometryDynamicCollection& DynamicCollection, TArray<bool>& IsTransformComputed, const TManagedArray<FTransform>* UniformScale, TArray<FTransform>& OutGlobalTransforms)
+	template<typename TransformType>
+	void GlobalMatricesHelper(const int32 Index, const FGeometryDynamicCollection& DynamicCollection, TArray<bool>& IsTransformComputed, const TManagedArray<FTransform>* UniformScale, TArray<TransformType>& OutGlobalTransforms)
 	{
 		if (IsTransformComputed[Index])
 		{
@@ -365,7 +366,7 @@ namespace GeometryCollectionAlgo
 		{
 			const int32 ProcessIndex = ToProcess.Pop(false);
 			const int32 ParentIndex = DynamicCollection.GetParent(ProcessIndex);
-			FTransform Result = FTransform(DynamicCollection.GetTransform(ProcessIndex));
+			TransformType Result = TransformType(DynamicCollection.GetTransform(ProcessIndex));
 			if (ParentIndex != FGeometryCollection::Invalid)
 			{
 				Result *= OutGlobalTransforms[ParentIndex];
@@ -373,7 +374,7 @@ namespace GeometryCollectionAlgo
 
 			if (UniformScale)
 			{
-				OutGlobalTransforms[ProcessIndex] = (*UniformScale)[ProcessIndex] * Result;
+				OutGlobalTransforms[ProcessIndex] = TransformType((*UniformScale)[ProcessIndex]) * Result;
 			}
 			else
 			{
@@ -403,15 +404,15 @@ namespace GeometryCollectionAlgo
 		{
 			const int32 ProcessIndex = ToProcess.Pop(false);
 			const int32 ParentIndex = Parents[ProcessIndex];
-			FTransform Result = FTransform(Transform[ProcessIndex]);
+			TransformTypeOut Result = TransformTypeOut(Transform[ProcessIndex]);
 			if (ParentIndex != FGeometryCollection::Invalid)
 			{
-				Result *= FTransform(OutGlobalTransforms[ParentIndex]);
+				Result *= OutGlobalTransforms[ParentIndex];
 			}
 
 			if (UniformScale)
 			{
-				OutGlobalTransforms[ProcessIndex] = TransformTypeOut((*UniformScale)[ProcessIndex] * Result);
+				OutGlobalTransforms[ProcessIndex] = TransformTypeOut((*UniformScale)[ProcessIndex]) * Result;
 			}
 			else
 			{
@@ -674,6 +675,21 @@ namespace GeometryCollectionAlgo
 				GlobalMatricesHelper(BoneIdx, DynamicCollection, IsTransformComputed, nullptr, OutGlobalTransforms);
 			}
 		}
+
+		void GlobalMatrices(const FGeometryDynamicCollection& DynamicCollection, TArray<FTransform3f>& OutGlobalTransforms)
+		{
+			int32 NumTransforms = DynamicCollection.GetNumTransforms();
+
+			TArray<bool> IsTransformComputed;
+			IsTransformComputed.AddDefaulted(NumTransforms);
+
+			OutGlobalTransforms.SetNumUninitialized(NumTransforms, false);
+
+			for (int BoneIdx = 0; BoneIdx < NumTransforms; ++BoneIdx)
+			{
+				GlobalMatricesHelper(BoneIdx, DynamicCollection, IsTransformComputed, nullptr, OutGlobalTransforms);
+			}
+		}
 	}
 
 	template<typename MatrixType, typename TransformType>
@@ -699,6 +715,8 @@ namespace GeometryCollectionAlgo
 	template void CHAOS_API GlobalMatrices<FMatrix, FTransform>(const TManagedArray<FTransform>&, const TManagedArray<int32>&, TArray<FMatrix>&);
 	template void CHAOS_API GlobalMatrices<FTransform, FTransform3f>(const TManagedArray<FTransform3f>&, const TManagedArray<int32>&, TArray<FTransform>&);
 	template void CHAOS_API GlobalMatrices<FMatrix, FTransform3f>(const TManagedArray<FTransform3f>&, const TManagedArray<int32>&, TArray<FMatrix>&);
+	template void CHAOS_API GlobalMatrices<FTransform3f, FTransform3f>(const TManagedArray<FTransform3f>&, const TManagedArray<int32>&, TArray<FTransform3f>&);
+	template void CHAOS_API GlobalMatrices<FTransform3f, FTransform>(const TManagedArray<FTransform>&, const TManagedArray<int32>&, TArray<FTransform3f>&);
 
 
 	void FloodForOverlappedPairs(int Level, int32 BoneIndex, TMap<int32, int32> &BoneToGroup, const TManagedArray<int32>& Levels, const TMap<int32, FBox>& BoundingBoxes, TSet<TTuple<int32, int32>>& OutOverlappedPairs)
diff --git a/Engine/Source/Runtime/Experimental/Chaos/Public/GeometryCollection/GeometryCollectionAlgo.h b/Engine/Source/Runtime/Experimental/Chaos/Public/GeometryCollection/GeometryCollectionAlgo.h
--- a/Engine/Source/Runtime/Experimental/Chaos/Public/GeometryCollection/GeometryCollectionAlgo.h
+++ b/Engine/Source/Runtime/Experimental/Chaos/Public/GeometryCollection/GeometryCollectionAlgo.h
@@ -292,6 +292,7 @@ namespace GeometryCollectionAlgo
 	namespace Private
 	{
 		void CHAOS_API GlobalMatrices(const FGeometryDynamicCollection& DynamicCollection, TArray<FTransform>& Transforms);
+		void CHAOS_API GlobalMatrices(const FGeometryDynamicCollection& DynamicCollection, TArray<FTransform3f>& Transforms);
 		void CHAOS_API GlobalMatrices(const FGeometryDynamicCollection& DynamicCollection, const TArray<int32>& Indices, TArray<FTransform>& OutGlobalTransforms);
 		FTransform CHAOS_API GlobalMatrix(const FGeometryDynamicCollection& DynamicCollection, int32 Index);
 	}
diff --git a/Engine/Source/Runtime/Experimental/GeometryCollectionEngine/Private/GeometryCollection/GeometryCollectionComponent.cpp b/Engine/Source/Runtime/Experimental/GeometryCollectionEngine/Private/GeometryCollection/GeometryCollectionComponent.cpp
--- a/Engine/Source/Runtime/Experimental/GeometryCollectionEngine/Private/GeometryCollection/GeometryCollectionComponent.cpp
+++ b/Engine/Source/Runtime/Experimental/GeometryCollectionEngine/Private/GeometryCollection/GeometryCollectionComponent.cpp
@@ -482,7 +482,7 @@ public:
 		TotalTransforms = 0;
 	}
 
-	bool Break(UGeometryCollectionComponent* t, const TArray<FTransform>& CompSpaceTransforms)
+	bool Break(UGeometryCollectionComponent* t, const TArray<FTransform3f>& CompSpaceTransforms)
 	{
 		if (Components.Contains(t))
 			return true;
@@ -744,7 +744,7 @@ namespace
 				{
 					for (int32 TransformIndex = 0; TransformIndex < ExplodedVectors.Num(); ++TransformIndex)
 					{
-						GlobalMatricesInOut[TransformIndex].AddToTranslation((FVector)ExplodedVectors[TransformIndex]);
+						GlobalMatricesInOut[TransformIndex].AddToTranslation((FVector3f)ExplodedVectors[TransformIndex]);
 					}
 				}
 			}
@@ -754,17 +754,17 @@ namespace
 	/**
 	* compute the bounding box from the bounding boxes stored in the geometry group
 	*/
-	template <typename TTransformType>
+	template <typename TWorldTransformType, typename TLocalTransformType>
 	inline FBox ComputeBoundsFromGeometryBoundingBoxes(
 		const TManagedArray<int32>& TransformToGeometryIndex,
 		const TManagedArray<int32>& TransformIndices,
 		const TManagedArray<FBox>& BoundingBoxes,
-		const TArray<TTransformType>& GlobalMatrices,
-		const TTransformType& LocalToWorldWithScale)
+		const TArray<TLocalTransformType>& GlobalMatrices,
+		const TWorldTransformType& LocalToWorldWithScale)
 	{
 		FBox BoundingBox(ForceInit);
 		// todo(chaos ) implement ISPC function using FTransform 
-		constexpr bool bIsMatrixType = std::is_same<TTransformType, FMatrix>::value;
+		constexpr bool bIsMatrixType = std::is_same<TLocalTransformType, FMatrix>::value && std::is_same<TWorldTransformType, FMatrix>::value;
 		if constexpr (bIsMatrixType)
 		{
 			if (bChaos_BoxCalcBounds_ISPC_Enabled && !bGeometryCollectionSingleThreadedBoundsCalculation)
@@ -795,7 +795,7 @@ namespace
 			
 			if (TransformToGeometryIndex[TransformIndex] != INDEX_NONE)
 			{
-				BoundingBox += BoundingBoxes[BoxIdx].TransformBy(GlobalMatrices[TransformIndex] * LocalToWorldWithScale);
+				BoundingBox += BoundingBoxes[BoxIdx].TransformBy(TWorldTransformType(GlobalMatrices[TransformIndex]) * LocalToWorldWithScale);
 			}
 		}
 		return BoundingBox;
@@ -805,17 +805,17 @@ namespace
 	* compute the bounding box from the bounding boxes stored in the transform group
 	* (used for nanite or when the geometry group data has been stripped on cook )
 	*/
-	template <typename TTransformType>
+	template <typename TWorldTransformType, typename TLocalTransformType>
 	inline FBox ComputeBoundsFromTransformBoundingBoxes(
 		const TManagedArray<int32>& TransformToGeometryIndex,
 		const TManagedArray<FBox>& BoundingBoxes,
-		const TArray<TTransformType>& GlobalMatrices,
-		const TTransformType& LocalToWorldWithScale)
+		const TArray<TLocalTransformType>& GlobalMatrices,
+		const TWorldTransformType& LocalToWorldWithScale)
 	{
 		FBox BoundingBox(ForceInit);
 
 		// todo(chaos ) implement ISPC function using FTransform 
-		constexpr bool bIsMatrixType = std::is_same<TTransformType, FMatrix>::value;
+		constexpr bool bIsMatrixType = std::is_same<TWorldTransformType, FMatrix>::value && std::is_same<TLocalTransformType, FMatrix>::value;
 		if constexpr (bIsMatrixType)
 		{
 			if (bChaos_BoxCalcBounds_ISPC_Enabled && !bGeometryCollectionSingleThreadedBoundsCalculation)
@@ -838,14 +838,14 @@ namespace
 		{
 			if (TransformToGeometryIndex[TransformIndex] != INDEX_NONE)
 			{
-				BoundingBox += BoundingBoxes[TransformIndex].TransformBy(GlobalMatrices[TransformIndex] * LocalToWorldWithScale);
+				BoundingBox += BoundingBoxes[TransformIndex].TransformBy(TWorldTransformType(GlobalMatrices[TransformIndex]) * LocalToWorldWithScale);
 			}
 		}
 		return BoundingBox;
 	}
 
-	template <typename TTransformType>
-	inline FBox ComputeBoundsFromTransforms(const UGeometryCollectionComponent& Component, const TTransformType& LocalToWorldWithScale, const TArray<TTransformType>& GlobalMatricesArray)
+	template <typename TWorldTransformType, typename TLocalTransformType>
+	inline FBox ComputeBoundsFromTransforms(const UGeometryCollectionComponent& Component, const TWorldTransformType& LocalToWorldWithScale, const TArray<TLocalTransformType>& GlobalMatricesArray)
 	{
 		static FName BoundingBoxAttributeName = "BoundingBox";
 
@@ -866,14 +866,20 @@ namespace
 
 FBox UGeometryCollectionComponent::ComputeBoundsFromGlobalMatrices(const FMatrix& LocalToWorldWithScale, const TArray<FMatrix>& GlobalMatricesArray) const
 {
-	return ComputeBoundsFromTransforms<FMatrix>(*this, LocalToWorldWithScale, GlobalMatricesArray);
+	return ComputeBoundsFromTransforms<FMatrix, FMatrix>(*this, LocalToWorldWithScale, GlobalMatricesArray);
 }
 
 FBox UGeometryCollectionComponent::ComputeBoundsFromComponentSpaceTransforms(const FTransform& LocalToWorldWithScale, const TArray<FTransform>& ComponentSpaceTransformsArray) const
 {
-	return ComputeBoundsFromTransforms<FTransform>(*this, LocalToWorldWithScale, ComponentSpaceTransformsArray);
+	return ComputeBoundsFromTransforms<FTransform, FTransform>(*this, LocalToWorldWithScale, ComponentSpaceTransformsArray);
+}
+
+FBox UGeometryCollectionComponent::ComputeBoundsFromComponentSpaceTransforms(const FTransform& LocalToWorldWithScale, const TArray<FTransform3f>& ComponentSpaceTransformsArray) const
+{
+	return ComputeBoundsFromTransforms<FTransform, FTransform3f>(*this, LocalToWorldWithScale, ComponentSpaceTransformsArray);
 }
 
+
 FBox UGeometryCollectionComponent::ComputeBounds(const FMatrix& LocalToWorldWithScale) const
 {
 	return ComputeBounds(FTransform(LocalToWorldWithScale));
@@ -884,7 +890,7 @@ FBox UGeometryCollectionComponent::ComputeBounds(const FTransform& LocalToWorldW
 	FBox BoundingBox(ForceInit);
 	if (RestCollection)
 	{
-		const TArray<FTransform>& CompSpaceTransforms = ComponentSpaceTransforms.RequestAllTransforms();
+		const TArray<FTransform3f>& CompSpaceTransforms = ComponentSpaceTransforms.RequestAllTransforms();
 		if (CompSpaceTransforms.Num() > 0)
 		{
 			BoundingBox = ComputeBoundsFromComponentSpaceTransforms(LocalToWorldWithScale, CompSpaceTransforms);
@@ -1387,7 +1393,7 @@ bool UGeometryCollectionComponent::DoCustomNavigableGeometryExport(FNavigableGeo
 	}
 
 	// Get all the geometry transforms in component space (they are stored natively in parent-bone space)
-	TArray<FTransform> GeomToComponent;
+	TArray<FTransform3f> GeomToComponent;
 	ComputeCurrentGlobalsMatrices(GeomToComponent);
 
 	OutVertexBuffer.AddUninitialized(VertexCount);
@@ -1403,16 +1409,15 @@ bool UGeometryCollectionComponent::DoCustomNavigableGeometryExport(FNavigableGeo
 		int32 SourceGeometryVertexStart = VertexStartArray[GeometryIndex];
 		int32 SourceGeometryVertexCount = VertexCountArray[GeometryIndex];
 
-		ParallelFor(SourceGeometryVertexCount, [&](int32 PointIdx)
+		ParallelFor(SourceGeometryVertexCount, [&OutVertexBuffer, &GeomToComponent, &Vertex, SourceGeometryVertexStart, SubsetIndex, DestVertex](int32 PointIdx)
 			{
 				//extract vertex from source
 				int32 SourceGeometryVertexIndex = SourceGeometryVertexStart + PointIdx;
-				FVector const VertexInWorldSpace = GeomToComponent[SubsetIndex].TransformPosition((FVector)Vertex[SourceGeometryVertexIndex]);
-
+				const FVector3f VertexInComponentSpace = GeomToComponent[SubsetIndex].TransformPosition(Vertex[SourceGeometryVertexIndex]);
 				int32 DestVertexIndex = DestVertex + PointIdx;
-				OutVertexBuffer[DestVertexIndex].X = VertexInWorldSpace.X;
-				OutVertexBuffer[DestVertexIndex].Y = VertexInWorldSpace.Y;
-				OutVertexBuffer[DestVertexIndex].Z = VertexInWorldSpace.Z;
+				OutVertexBuffer[DestVertexIndex].X = VertexInComponentSpace.X;
+				OutVertexBuffer[DestVertexIndex].Y = VertexInComponentSpace.Y;
+				OutVertexBuffer[DestVertexIndex].Z = VertexInComponentSpace.Z;
 			});
 
 		DestVertex += SourceGeometryVertexCount;
@@ -1526,7 +1531,7 @@ void UGeometryCollectionComponent::RefreshEmbeddedGeometry()
 	EmbeddedInstanceIndex.Init(INDEX_NONE, RestCollection->GetGeometryCollection()->NumElements(FGeometryCollection::TransformGroup));
 #endif
 
-	const TArray<FTransform>& CompSpaceTransforms = ComponentSpaceTransforms.RequestAllTransforms();
+	const TArray<FTransform3f>& CompSpaceTransforms = ComponentSpaceTransforms.RequestAllTransforms();
 
 	for (int32 ExemplarIndex = 0; ExemplarIndex < ExemplarCount; ++ExemplarIndex)
 	{		
@@ -1545,7 +1550,7 @@ void UGeometryCollectionComponent::RefreshEmbeddedGeometry()
 			{
 				if (!HideArray || !(*HideArray)[Idx])
 				{ 
-					InstanceTransforms.Add(CompSpaceTransforms[Idx]);
+					InstanceTransforms.Add(FTransform(CompSpaceTransforms[Idx]));
 #if WITH_EDITOR
 					int32 InstanceIndex = EmbeddedBoneMaps[ExemplarIndex].Add(Idx);
 					EmbeddedInstanceIndex[Idx] = InstanceIndex;
@@ -1772,8 +1777,18 @@ static void DispatchGeometryCollectionCrumblingEvent(const FChaosCrumblingEvent&
 	}
 }
 
-const TArray<FTransform>& UGeometryCollectionComponent::GetComponentSpaceTransforms()
+// DEPRECATED Use GetComponentSpaceTransforms3f instead
+TArray<FTransform> UGeometryCollectionComponent::GetComponentSpaceTransforms()
 { 
+	const TArray<FTransform3f>& Transforms3f = GetComponentSpaceTransforms3f();
+	TArray<FTransform> Transforms;
+	Transforms.Reserve(Transforms3f.Num());
+	Algo::Transform(Transforms3f, Transforms, [](const FTransform3f& T) { return FTransform(T); });
+	return Transforms;
+}
+
+const TArray<FTransform3f>& UGeometryCollectionComponent::GetComponentSpaceTransforms3f()
+{
 	return ComponentSpaceTransforms.RequestAllTransforms();
 }
 
@@ -3062,7 +3077,7 @@ FGeometryCollectionDynamicData* UGeometryCollectionComponent::InitDynamicData(bo
 		DynamicData->IsDynamic = true;
 		DynamicData->IsLoading = GetIsObjectLoading();
 
-		const TArray<FTransform>& CompSpaceTransforms = ComponentSpaceTransforms.RequestAllTransforms();
+		const TArray<FTransform3f>& CompSpaceTransforms = ComponentSpaceTransforms.RequestAllTransforms();
 
 		// If we have no transforms stored in the dynamic data, then assign both prev and current to the same global matrices
 		// Copy existing global matrices into prev transforms
@@ -3216,14 +3231,14 @@ FTransform UGeometryCollectionComponent::GetSocketTransform(FName InSocketName,
 				const int32 TransformIndex = Collection->BoneName.Find(InSocketName.ToString());
 				if (TransformIndex != INDEX_NONE)
 				{
-					const TArray<FTransform>& CompSpaceTransforms = ComponentSpaceTransforms.RequestAllTransforms();
+					const TArray<FTransform3f>& CompSpaceTransforms = ComponentSpaceTransforms.RequestAllTransforms();
 					if (CompSpaceTransforms.IsValidIndex(TransformIndex))
 					{
-						SocketComponentSpaceTransform = CompSpaceTransforms[TransformIndex];
+						SocketComponentSpaceTransform = FTransform(CompSpaceTransforms[TransformIndex]);
 						const int32 ParentTransformIndex = Collection->Parent[TransformIndex];
 						if (CompSpaceTransforms.IsValidIndex(ParentTransformIndex))
 						{
-							ParentComponentSpaceTransform = CompSpaceTransforms[ParentTransformIndex];
+							ParentComponentSpaceTransform = FTransform(CompSpaceTransforms[ParentTransformIndex]);
 						}
 						bFoundSocket = true;
 					}
@@ -3861,11 +3876,11 @@ void UGeometryCollectionComponent::OnPostPhysicsSync()
 	if (DynamicCollection && DynamicCollection->IsDirty())
 	{
 		// Can't be a const reference - we need to make a copy or else the next RequestRootTransform will change the value under our nose.
-		const FTransform PreviousRootTransform = OnRootMovedEvent.IsBound() ? ComponentSpaceTransforms.RequestRootTransform() : FTransform::Identity;
+		const FTransform3f PreviousRootTransform = OnRootMovedEvent.IsBound() ? ComponentSpaceTransforms.RequestRootTransform() : FTransform3f::Identity;
 		ComponentSpaceTransforms.MarkDirty();
 		if (OnRootMovedEvent.IsBound())
 		{
-			const FTransform& NewRootTransform = ComponentSpaceTransforms.RequestRootTransform();
+			const FTransform3f& NewRootTransform = ComponentSpaceTransforms.RequestRootTransform();
 			if (!PreviousRootTransform.EqualsNoScale(NewRootTransform))
 			{
 				OnRootMovedEvent.Broadcast();
@@ -4026,7 +4041,7 @@ void UGeometryCollectionComponent::UpdateRemovalIfNeeded()
 			const FTransform InverseComponentTransform = (RestCollection->bScaleOnRemoval) ? GetComponentTransform().Inverse() : FTransform::Identity;
 
 			const TManagedArray<FTransform>* MassToLocal = nullptr;
-			const TArray<FTransform>* CompSpaceTransform = nullptr;
+			const TArray<FTransform3f>* CompSpaceTransform = nullptr;
 			if (RestCollection->bScaleOnRemoval)
 			{
 				MassToLocal = RestCollection->GetGeometryCollection()->FindAttribute<FTransform>(MassToLocalAttributeName, FGeometryCollection::TransformGroup);
@@ -4058,12 +4073,13 @@ void UGeometryCollectionComponent::UpdateRemovalIfNeeded()
 							ShrinkRadius = -AccumulatedSphere.W;
 						}
 
-						const FQuat LocalRotation = (InverseComponentTransform * (*CompSpaceTransform)[TransformIndex].Inverse()).GetRotation();
+						// InverseComponentTransform calculation must in double precision for LWC
+						const FQuat LocalRotation = (InverseComponentTransform * FTransform((*CompSpaceTransform)[TransformIndex].Inverse())).GetRotation();
 						const FVector LocalDown = LocalRotation.RotateVector(FVector(0.f, 0.f, ShrinkRadius));
 						const FVector CenterOfMass = (*MassToLocal)[TransformIndex].GetTranslation();
-						const FVector ScaleCenter = LocalDown + CenterOfMass;
-						const FTransform ScaleTransform(FQuat::Identity, ScaleCenter * FVector::FReal(1.f - Scale), FVector(Scale));
-						DynamicCollection->SetTransform(TransformIndex, FTransform3f(ScaleTransform) * DynamicCollection->GetTransform(TransformIndex));
+						const FVector3f ScaleCenter = FVector3f(LocalDown + CenterOfMass);
+						const FTransform3f ScaleTransform(FQuat4f::Identity, ScaleCenter * FVector3f::FReal(1.f - Scale), FVector3f(Scale));
+						DynamicCollection->SetTransform(TransformIndex, ScaleTransform * DynamicCollection->GetTransform(TransformIndex));
 					}
 				}
 			}
@@ -5413,22 +5429,22 @@ static void CheckForNaNs(const TArray<FTransform>& Transforms)
 }
 #endif
 
-FTransform UGeometryCollectionComponent::GetCurrentTransform(int32 Index) const
+FTransform3f UGeometryCollectionComponent::GetCurrentTransform(int32 Index) const
 {
 	if (DynamicCollection)
 	{
-		return FTransform(DynamicCollection->GetTransform(Index));
+		return DynamicCollection->GetTransform(Index);
 	}
 
 	const bool bRestTransformsOverriden = (RestTransforms.Num() > 0);
 	if (bRestTransformsOverriden)
 	{
-		return RestTransforms[Index];
+		return FTransform3f(RestTransforms[Index]);
 	}
-	return FTransform(RestCollection->GetGeometryCollection()->Transform[Index]);
+	return RestCollection->GetGeometryCollection()->Transform[Index];
 }
 
-void UGeometryCollectionComponent::ComputeCurrentGlobalsMatrices(TArray<FTransform>& OutTransforms) const
+void UGeometryCollectionComponent::ComputeCurrentGlobalsMatrices(TArray<FTransform3f>& OutTransforms) const
 {
 	if (DynamicCollection)
 	{
@@ -5448,12 +5464,12 @@ void UGeometryCollectionComponent::ComputeCurrentGlobalsMatrices(TArray<FTransfo
 	}
 }
 
-const FTransform& UGeometryCollectionComponent::FComponentSpaceTransforms::RequestRootTransform() const
+const FTransform3f& UGeometryCollectionComponent::FComponentSpaceTransforms::RequestRootTransform() const
 {
 	if (!Component)
 	{
 		ensure(false); // we should not be able to reach here 
-		return FTransform::Identity;
+		return FTransform3f::Identity;
 	}
 
 	if (Transforms.IsValidIndex(RootIndex))
@@ -5467,14 +5483,14 @@ const FTransform& UGeometryCollectionComponent::FComponentSpaceTransforms::Reque
 		return Transforms[RootIndex];
 	}
 
-	return FTransform::Identity;
+	return FTransform3f::Identity;
 }
 
-const TArray<FTransform>& UGeometryCollectionComponent::FComponentSpaceTransforms::RequestAllTransforms() const
+const TArray<FTransform3f>& UGeometryCollectionComponent::FComponentSpaceTransforms::RequestAllTransforms() const
 {
 	SCOPE_CYCLE_COUNTER(STAT_GCCUGlobalMatrices);
 
-	static TArray<FTransform> EmptyArray;
+	static TArray<FTransform3f> EmptyArray;
 	if (!Component || Transforms.Num() == 0)
 	{
 		return EmptyArray;
@@ -5512,7 +5528,7 @@ const TArray<FTransform>& UGeometryCollectionComponent::FComponentSpaceTransform
 			const int32 ParentTransformIndex = Component->GetParent(TransformIndex);
 
 
-			FTransform CurrentTransform = Component->GetCurrentTransform(TransformIndex);
+			FTransform3f CurrentTransform = Component->GetCurrentTransform(TransformIndex);
 			if (ParentTransformIndex == INDEX_NONE)
 			{
 				
@@ -5520,7 +5536,7 @@ const TArray<FTransform>& UGeometryCollectionComponent::FComponentSpaceTransform
 			}
 			else
 			{
-				const FTransform& ParentTransform = Transforms[ParentTransformIndex];
+				const FTransform3f& ParentTransform = Transforms[ParentTransformIndex];
 				Transforms[TransformIndex] = CurrentTransform * ParentTransform;
 			}
 		}
@@ -5791,12 +5807,12 @@ void UGeometryCollectionComponent::RefreshCustomRenderer()
 
 					if (bRenderRootProxy)
 					{
-						const FTransform& CompSpaceRootTransform = ComponentSpaceTransforms.RequestRootTransform();
+						const FTransform CompSpaceRootTransform(ComponentSpaceTransforms.RequestRootTransform());
 						RendererInterface->UpdateRootTransform(*RestCollection, CompSpaceRootTransform);
 					}
 					else
 					{
-						const TArray<FTransform>& CompSpaceTransforms = ComponentSpaceTransforms.RequestAllTransforms();
+						const TArray<FTransform3f>& CompSpaceTransforms = ComponentSpaceTransforms.RequestAllTransforms();
 						RendererInterface->UpdateTransforms(*RestCollection, CompSpaceTransforms);
 					}
 				}
@@ -6225,7 +6241,7 @@ FTransform UGeometryCollectionComponent::GetRootCurrentTransform() const
 	FTransform RootInitialTransform{ FTransform::Identity };
 	if (RestCollection)
 	{
-		const FTransform& CompSpaceRootTransform = ComponentSpaceTransforms.RequestRootTransform();
+		const FTransform CompSpaceRootTransform(ComponentSpaceTransforms.RequestRootTransform());
 		RootInitialTransform = CompSpaceRootTransform * GetComponentTransform();
 	}
 	return RootInitialTransform;
@@ -6233,7 +6249,7 @@ FTransform UGeometryCollectionComponent::GetRootCurrentTransform() const
 
 FTransform UGeometryCollectionComponent::GetRootCurrentComponentSpaceTransform() const
 {
-	return (RestCollection) ? ComponentSpaceTransforms.RequestRootTransform() : FTransform::Identity;
+	return (RestCollection) ? FTransform(ComponentSpaceTransforms.RequestRootTransform()) : FTransform::Identity;
 }
 
 TArray<FTransform> UGeometryCollectionComponent::GetInitialLocalRestTransforms() const
@@ -6309,10 +6325,10 @@ TArray<FMatrix> UGeometryCollectionComponent::ComputeGlobalMatricesFromComponent
 	TArray<FMatrix> ComponentSpaceMatrices;
 	ComponentSpaceMatrices.SetNumUninitialized(ComponentSpaceTransforms.Num());
 
-	const TArray<FTransform>& CompSpaceTransforms = ComponentSpaceTransforms.RequestAllTransforms();
+	const TArray<FTransform3f>& CompSpaceTransforms = ComponentSpaceTransforms.RequestAllTransforms();
 	for (int32 TransformIndex = 0; TransformIndex < CompSpaceTransforms.Num(); TransformIndex++)
 	{
-		ComponentSpaceMatrices[TransformIndex] = CompSpaceTransforms[TransformIndex].ToMatrixWithScale();
+		ComponentSpaceMatrices[TransformIndex] = FTransform(CompSpaceTransforms[TransformIndex]).ToMatrixWithScale();
 	}
 	return ComponentSpaceMatrices;
 }
diff --git a/Engine/Source/Runtime/Experimental/GeometryCollectionEngine/Private/GeometryCollection/GeometryCollectionISMPoolRenderer.cpp b/Engine/Source/Runtime/Experimental/GeometryCollectionEngine/Private/GeometryCollection/GeometryCollectionISMPoolRenderer.cpp
--- a/Engine/Source/Runtime/Experimental/GeometryCollectionEngine/Private/GeometryCollection/GeometryCollectionISMPoolRenderer.cpp
+++ b/Engine/Source/Runtime/Experimental/GeometryCollectionEngine/Private/GeometryCollection/GeometryCollectionISMPoolRenderer.cpp
@@ -65,7 +65,7 @@ void UGeometryCollectionISMPoolRenderer::UpdateRootTransform(UGeometryCollection
 	UpdateMergedMeshTransforms(InRootTransform * ComponentTransform);
 }
 
-void UGeometryCollectionISMPoolRenderer::UpdateTransforms(UGeometryCollection const& InGeometryCollection, TArrayView<const FTransform> InTransforms)
+void UGeometryCollectionISMPoolRenderer::UpdateTransforms(UGeometryCollection const& InGeometryCollection, TArrayView<const FTransform3f> InTransforms)
 {
 	UpdateInstanceTransforms(InGeometryCollection, ComponentTransform, InTransforms);
 }
@@ -178,7 +178,7 @@ void UGeometryCollectionISMPoolRenderer::UpdateMergedMeshTransforms(FTransform c
 	}
 }
 
-void UGeometryCollectionISMPoolRenderer::UpdateInstanceTransforms(UGeometryCollection const& InGeometryCollection, FTransform const& InBaseTransform, TArrayView<const FTransform> InTransforms)
+void UGeometryCollectionISMPoolRenderer::UpdateInstanceTransforms(UGeometryCollection const& InGeometryCollection, FTransform const& InBaseTransform, TArrayView<const FTransform3f> InTransforms)
 {
 	if (InstancesGroup.GroupIndex == INDEX_NONE)
 	{
@@ -209,7 +209,7 @@ void UGeometryCollectionISMPoolRenderer::UpdateInstanceTransforms(UGeometryColle
 			const int32 AutoInstanceMeshIndex = InstancedMeshFacade.GetIndex(TransformIndex);
 			if (AutoInstanceMeshIndex == MeshIndex && Children[TransformIndex].Num() == 0)
 			{
-				InstanceTransforms.Add(InTransforms[TransformIndex] * InBaseTransform);
+				InstanceTransforms.Add(FTransform(InTransforms[TransformIndex]) * InBaseTransform);
 			}
 		}
 		
diff --git a/Engine/Source/Runtime/Experimental/GeometryCollectionEngine/Private/GeometryCollection/GeometryCollectionISMPoolRenderer.h b/Engine/Source/Runtime/Experimental/GeometryCollectionEngine/Private/GeometryCollection/GeometryCollectionISMPoolRenderer.h
--- a/Engine/Source/Runtime/Experimental/GeometryCollectionEngine/Private/GeometryCollection/GeometryCollectionISMPoolRenderer.h
+++ b/Engine/Source/Runtime/Experimental/GeometryCollectionEngine/Private/GeometryCollection/GeometryCollectionISMPoolRenderer.h
@@ -23,7 +23,7 @@ public:
 	virtual void OnUnregisterGeometryCollection() override;
 	virtual void UpdateState(UGeometryCollection const& InGeometryCollection, FTransform const& InComponentTransform, uint32 InStateFlags) override;
 	virtual void UpdateRootTransform(UGeometryCollection const& InGeometryCollection, FTransform const& InRootTransform) override;
-	virtual void UpdateTransforms(UGeometryCollection const& InGeometryCollection, TArrayView<const FTransform> InTransforms) override;
+	virtual void UpdateTransforms(UGeometryCollection const& InGeometryCollection, TArrayView<const FTransform3f> InTransforms) override;
 	//~ End IGeometryCollectionExternalRenderInterface Interface.
 
 	/** Description for a group of meshes that are added/updated together. */
@@ -53,6 +53,6 @@ private:
 	void InitMergedMeshFromGeometryCollection(UGeometryCollection const& InGeometryCollection);
 	void InitInstancesFromGeometryCollection(UGeometryCollection const& InGeometryCollection);
 	void UpdateMergedMeshTransforms(FTransform const& InBaseTransform);
-	void UpdateInstanceTransforms(UGeometryCollection const& InGeometryCollection, FTransform const& InBaseTransform, TArrayView<const FTransform> InTransforms);
+	void UpdateInstanceTransforms(UGeometryCollection const& InGeometryCollection, FTransform const& InBaseTransform, TArrayView<const FTransform3f> InTransforms);
 	void ReleaseGroup(FISMPoolGroup& InOutGroup);
 };
diff --git a/Engine/Source/Runtime/Experimental/GeometryCollectionEngine/Private/GeometryCollection/GeometryCollectionSceneProxy.h b/Engine/Source/Runtime/Experimental/GeometryCollectionEngine/Private/GeometryCollection/GeometryCollectionSceneProxy.h
--- a/Engine/Source/Runtime/Experimental/GeometryCollectionEngine/Private/GeometryCollection/GeometryCollectionSceneProxy.h
+++ b/Engine/Source/Runtime/Experimental/GeometryCollectionEngine/Private/GeometryCollection/GeometryCollectionSceneProxy.h
@@ -69,6 +69,15 @@ inline void CopyTransformsWithConversionWhenNeeded(TArray<FMatrix44f>& DstTransf
 	}
 }
 
+inline void CopyTransformsWithConversionWhenNeeded(TArray<FMatrix44f>& DstTransforms, const TArray<FTransform3f>& SrcTransforms)
+{
+	DstTransforms.SetNumUninitialized(SrcTransforms.Num());
+	for (int TransformIndex = 0; TransformIndex < SrcTransforms.Num(); ++TransformIndex)
+	{
+		DstTransforms[TransformIndex] = SrcTransforms[TransformIndex].ToMatrixWithScale();
+	}
+}
+
 /** Mutable rendering data */
 struct FGeometryCollectionDynamicData
 {
@@ -104,6 +113,11 @@ struct FGeometryCollectionDynamicData
 		CopyTransformsWithConversionWhenNeeded(Transforms, InTransforms);
 	}
 
+	void SetTransforms(const TArray<FTransform3f>& InTransforms)
+	{
+		CopyTransformsWithConversionWhenNeeded(Transforms, InTransforms);
+	}
+
 	UE_DEPRECATED(5.3, "Use FTransform version of SetPrevTransforms instead")
 	void SetPrevTransforms(const TArray<FMatrix>& InTransforms)
 	{
@@ -117,6 +131,11 @@ struct FGeometryCollectionDynamicData
 		CopyTransformsWithConversionWhenNeeded(PrevTransforms, InTransforms);
 	}
 
+	void SetPrevTransforms(const TArray<FTransform3f>& InTransforms)
+	{
+		CopyTransformsWithConversionWhenNeeded(PrevTransforms, InTransforms);
+	}
+
 	UE_DEPRECATED(5.3, "Use FTransform version of SetAllTransforms instead")
 	void SetAllTransforms(const TArray<FMatrix>& InTransforms)
 	{
diff --git a/Engine/Source/Runtime/Experimental/GeometryCollectionEngine/Public/GeometryCollection/GeometryCollectionComponent.h b/Engine/Source/Runtime/Experimental/GeometryCollectionEngine/Public/GeometryCollection/GeometryCollectionComponent.h
--- a/Engine/Source/Runtime/Experimental/GeometryCollectionEngine/Public/GeometryCollection/GeometryCollectionComponent.h
+++ b/Engine/Source/Runtime/Experimental/GeometryCollectionEngine/Public/GeometryCollection/GeometryCollectionComponent.h
@@ -1270,7 +1270,10 @@ public:
 	UE_DEPRECATED(5.3, "Use GetComponentSpaceTransforms instead")
 	TArray<FMatrix> GetGlobalMatrices() { return ComputeGlobalMatricesFromComponentSpaceTransforms(); }
 
-	GEOMETRYCOLLECTIONENGINE_API const TArray<FTransform>& GetComponentSpaceTransforms();
+	UE_DEPRECATED(5.4, "Use GetComponentSpaceTransforms3f instead")
+	GEOMETRYCOLLECTIONENGINE_API TArray<FTransform> GetComponentSpaceTransforms();
+
+	GEOMETRYCOLLECTIONENGINE_API const TArray<FTransform3f>& GetComponentSpaceTransforms3f();
 
 	GEOMETRYCOLLECTIONENGINE_API const FGeometryDynamicCollection* GetDynamicCollection() const;
 	GEOMETRYCOLLECTIONENGINE_API FGeometryDynamicCollection* GetDynamicCollection();  // TEMP HACK?
@@ -1424,6 +1427,7 @@ protected:
 	GEOMETRYCOLLECTIONENGINE_API FBox ComputeBoundsFromGlobalMatrices(const FMatrix& LocalToWorldWithScale, const TArray<FMatrix>& GlobalMatricesArray) const;
 
 	GEOMETRYCOLLECTIONENGINE_API FBox ComputeBoundsFromComponentSpaceTransforms(const FTransform& LocalToWorldWithScale, const TArray<FTransform>& ComponentSpaceTransformsArray) const;
+	FBox ComputeBoundsFromComponentSpaceTransforms(const FTransform& LocalToWorldWithScale, const TArray<FTransform3f>& ComponentSpaceTransformsArray) const;
 
 	UE_DEPRECATED(5.3, "Use FTransform version of ComputeBounds instead")
 	GEOMETRYCOLLECTIONENGINE_API FBox ComputeBounds(const FMatrix& LocalToWorldWithScale) const;
@@ -1545,8 +1549,8 @@ private:
 
 	// return the most actual transforms
 	// this can be the rest collection ones, the overridden RestTransforms or the dynamic collection ones
-	FTransform GetCurrentTransform(int32 Index) const;
-	void ComputeCurrentGlobalsMatrices(TArray<FTransform>& OutTransforms) const;
+	FTransform3f GetCurrentTransform(int32 Index) const;
+	void ComputeCurrentGlobalsMatrices(TArray<FTransform3f>& OutTransforms) const;
 
 	bool bRenderStateDirty;
 	bool bEnableBoneSelection;
@@ -1601,16 +1605,16 @@ private:
 
 		// request all transform to be update
 		// this will trigger an update if it is still marked dirty
-		const TArray<FTransform>& RequestAllTransforms() const;
+		const TArray<FTransform3f>& RequestAllTransforms() const;
 
 		// request the root transform, this may compute it if it is still marked dirty
-		const FTransform& RequestRootTransform() const;
+		const FTransform3f& RequestRootTransform() const;
 
 	private:
 		int32 RootIndex;
 		mutable uint8 bIsRootDirty : 1;
 		mutable uint8 bIsDirty : 1;
-		mutable TArray<FTransform> Transforms;
+		mutable TArray<FTransform3f> Transforms;
 		const UGeometryCollectionComponent* Component;
 	};
 
diff --git a/Engine/Source/Runtime/Experimental/GeometryCollectionEngine/Public/GeometryCollection/GeometryCollectionExternalRenderInterface.h b/Engine/Source/Runtime/Experimental/GeometryCollectionEngine/Public/GeometryCollection/GeometryCollectionExternalRenderInterface.h
--- a/Engine/Source/Runtime/Experimental/GeometryCollectionEngine/Public/GeometryCollection/GeometryCollectionExternalRenderInterface.h
+++ b/Engine/Source/Runtime/Experimental/GeometryCollectionEngine/Public/GeometryCollection/GeometryCollectionExternalRenderInterface.h
@@ -31,7 +31,7 @@ public:
 	virtual void OnUnregisterGeometryCollection() = 0;
 	virtual void UpdateState(UGeometryCollection const& InGeometryCollection, FTransform const& InComponentTransform, uint32 InStateFlags) = 0;
 	virtual void UpdateRootTransform(UGeometryCollection const& InGeometryCollection, FTransform const& InRootTransform) = 0;
-	virtual void UpdateTransforms(UGeometryCollection const& InGeometryCollection, TArrayView<const FTransform> InTransforms) = 0;
+	virtual void UpdateTransforms(UGeometryCollection const& InGeometryCollection, TArrayView<const FTransform3f> InTransforms) = 0;
 
 	UE_DEPRECATED(5.4, "Use flags version of UpdateState instead")
 	virtual void UpdateState(UGeometryCollection const& InGeometryCollection, FTransform const& InComponentTransform, bool bInIsBroken, bool bInIsVisible) 
