Fix non leaves transform calculation in UGeometryCollectionComponent::SetLocalRestTransforms when bOnlyLeaves = true

diff --git a/Engine/Source/Runtime/Experimental/GeometryCollectionEngine/Private/GeometryCollection/GeometryCollectionComponent.cpp b/Engine/Source/Runtime/Experimental/GeometryCollectionEngine/Private/GeometryCollection/GeometryCollectionComponent.cpp
--- a/Engine/Source/Runtime/Experimental/GeometryCollectionEngine/Private/GeometryCollection/GeometryCollectionComponent.cpp
+++ b/Engine/Source/Runtime/Experimental/GeometryCollectionEngine/Private/GeometryCollection/GeometryCollectionComponent.cpp
@@ -6292,12 +6292,12 @@ void UGeometryCollectionComponent::SetLocalRestTransforms(const TArray<FTransfor
 			for (int32 ProcessIndex = 0; ProcessIndex < TransformsToProcess.Num(); ProcessIndex++)
 			{
 				const int32 TransformIndex = TransformsToProcess[ProcessIndex];
+				const FTransform InvMassToLocal = MassToLocal ? (*MassToLocal)[TransformIndex].Inverse() : FTransform::Identity;
+
 				const bool bIsLeaf = (GeometryCollection.SimulationType[TransformIndex] == FGeometryCollection::ESimulationTypes::FST_Rigid);
 				const bool bNeedUpdate = !bOnlyLeaves || bIsLeaf;
 				if (bNeedUpdate)
 				{
-					const FTransform InvMassToLocal = MassToLocal ? (*MassToLocal)[TransformIndex].Inverse() : FTransform::Identity;
-
 					LocalTransforms[TransformIndex] = InvMassToLocal * NewTransforms[TransformIndex];
 
 					// because we update top-down, the parent up-to-date transform is stored in LocalTransforms
@@ -6307,6 +6307,10 @@ void UGeometryCollectionComponent::SetLocalRestTransforms(const TArray<FTransfor
 						LocalTransforms[TransformIndex] = LocalTransforms[TransformIndex].GetRelativeTransform(LocalTransforms[ParentTransformIndex]);
 					}
 				}
+				else
+				{
+					LocalTransforms[TransformIndex] = InvMassToLocal * LocalTransforms[TransformIndex];
+				}
 
 				for (const int32 ChildTransformIndex : GeometryCollection.Children[TransformIndex])
 				{
