Geometry collection : optimize computation of bounds when GC is unbroken
Also fix an issue with dynamic collection being always dirty after it was once
- Cache root space bounds
- Remove old LocalBounds member as it is a duplicate of the ComponentSpaceBounds
- Remove cvar option to use cached bounds ( now always using them )
- Set the dynamic collection clean when syncing from PT to GT ( and dirty it if any thing changed )

diff --git a/Engine/Source/Runtime/Experimental/Chaos/Private/PhysicsProxy/GeometryCollectionPhysicsProxy.cpp b/Engine/Source/Runtime/Experimental/Chaos/Private/PhysicsProxy/GeometryCollectionPhysicsProxy.cpp
--- a/Engine/Source/Runtime/Experimental/Chaos/Private/PhysicsProxy/GeometryCollectionPhysicsProxy.cpp
+++ b/Engine/Source/Runtime/Experimental/Chaos/Private/PhysicsProxy/GeometryCollectionPhysicsProxy.cpp
@@ -4013,6 +4013,8 @@ bool FGeometryCollectionPhysicsProxy::PullFromPhysicsState(const Chaos::FDirtyGe
 
 		// if physics world transform was dirtied this frame we need to force an update to make sure transforms of 
 		// broken sleeping particles on kinemically driven GCs remained updated and the object feel grounded to the world
+
+		GameThreadCollection.MakeClean();
 		if (bIsCollectionDirty || bIsPhysicsThreadWorldTransformDirty)
 		{
 			GameThreadCollection.MakeDirty();
diff --git a/Engine/Source/Runtime/Experimental/GeometryCollectionEngine/Private/GeometryCollection/GeometryCollectionComponent.cpp b/Engine/Source/Runtime/Experimental/GeometryCollectionEngine/Private/GeometryCollection/GeometryCollectionComponent.cpp
--- a/Engine/Source/Runtime/Experimental/GeometryCollectionEngine/Private/GeometryCollection/GeometryCollectionComponent.cpp
+++ b/Engine/Source/Runtime/Experimental/GeometryCollectionEngine/Private/GeometryCollection/GeometryCollectionComponent.cpp
@@ -107,9 +107,6 @@ static bool bChaos_BoxCalcBounds_ISPC_Enabled = CHAOS_BOX_CALC_BOUNDS_ISPC_ENABL
 static FAutoConsoleVariableRef CVarChaosBoxCalcBoundsISPCEnabled(TEXT("p.Chaos.BoxCalcBounds.ISPC"), bChaos_BoxCalcBounds_ISPC_Enabled, TEXT("Whether to use ISPC optimizations in calculating box bounds in geometry collections"));
 #endif
 
-bool bChaos_GC_CacheComponentSpaceBounds = true;
-FAutoConsoleVariableRef CVarChaosGCCacheComponentSpaceBounds(TEXT("p.Chaos.GC.CacheComponentSpaceBounds"), bChaos_GC_CacheComponentSpaceBounds, TEXT("Cache component space bounds for performance"));
-
 bool bChaos_GC_UseCustomRenderer = true;
 FAutoConsoleVariableRef CVarChaosGCUseCustomRenderer(TEXT("p.Chaos.GC.UseCustomRenderer"), bChaos_GC_UseCustomRenderer, TEXT("When enabled, use a custom renderer if specified"));
 
@@ -890,10 +887,33 @@ FBox UGeometryCollectionComponent::ComputeBounds(const FTransform& LocalToWorldW
 	FBox BoundingBox(ForceInit);
 	if (RestCollection)
 	{
-		const TArray<FTransform3f>& CompSpaceTransforms = ComponentSpaceTransforms.RequestAllTransforms();
-		if (CompSpaceTransforms.Num() > 0)
+		if (IsRootBroken())
+		{
+			// root is broken we need to go through all the leafs and compute bounds from them 
+			const TArray<FTransform3f>& CompSpaceTransforms = ComponentSpaceTransforms.RequestAllTransforms();
+			if (CompSpaceTransforms.Num() > 0)
+			{
+				BoundingBox = ComputeBoundsFromComponentSpaceTransforms(LocalToWorldWithScale, CompSpaceTransforms);
+			}
+		}
+		else
 		{
-			BoundingBox = ComputeBoundsFromComponentSpaceTransforms(LocalToWorldWithScale, CompSpaceTransforms);
+			// fast path where we use the root space bounds top avoid paying the high cost of recomputing everything from scratch
+			if (!RootSpaceBounds.IsValid)
+			{
+				const TArray<FTransform3f>& CompSpaceTransforms = ComponentSpaceTransforms.RequestAllTransforms();
+				if (CompSpaceTransforms.Num() > 0)
+				{
+					const FTransform InverseRootTransform(ComponentSpaceTransforms.RequestRootTransform().Inverse());
+					RootSpaceBounds = ComputeBoundsFromComponentSpaceTransforms(InverseRootTransform, CompSpaceTransforms);
+				}
+			}
+
+			if (RootSpaceBounds.IsValid)
+			{
+				const FTransform RootToWorld = (FTransform(ComponentSpaceTransforms.RequestRootTransform()) * LocalToWorldWithScale);
+				BoundingBox = RootSpaceBounds.TransformBy(RootToWorld);
+			}
 		}
 	}
 	return BoundingBox;
@@ -903,27 +923,21 @@ FBoxSphereBounds UGeometryCollectionComponent::CalcBounds(const FTransform& Loca
 {	
 	SCOPE_CYCLE_COUNTER(STAT_GCCUpdateBounds);
 
-	if (bChaos_GC_CacheComponentSpaceBounds)
-	{
-		bool NeedBoundsUpdate = false;
-		NeedBoundsUpdate |= (ComponentSpaceBounds.GetSphere().W < 1e-5);
-		NeedBoundsUpdate |= CachePlayback;
-//		NeedBoundsUpdate |= (PhysicsProxy == nullptr);
-		NeedBoundsUpdate |= (DynamicCollection && DynamicCollection->IsDirty());
-		
-		if (NeedBoundsUpdate)
-		{
-			ComponentSpaceBounds = ComputeBounds(FTransform::Identity);
-		}
-		else
-		{
-			NeedBoundsUpdate = false;
-		}
+	bool NeedBoundsUpdate = false;
+	NeedBoundsUpdate |= (ComponentSpaceBounds.GetSphere().W < 1e-5);
+	NeedBoundsUpdate |= CachePlayback;
+	NeedBoundsUpdate |= (DynamicCollection && DynamicCollection->IsDirty());
 
-		return ComponentSpaceBounds.TransformBy(LocalToWorldIn);
+	if (NeedBoundsUpdate)
+	{
+		ComponentSpaceBounds = ComputeBounds(FTransform::Identity);
+	}
+	else
+	{
+		NeedBoundsUpdate = false;
 	}
 
-	return FBoxSphereBounds(ComputeBounds(LocalToWorldIn));
+	return ComponentSpaceBounds.TransformBy(LocalToWorldIn);
 }
 
 int32 UGeometryCollectionComponent::GetNumElements(FName Group) const
@@ -1004,7 +1018,6 @@ void UGeometryCollectionComponent::SetDensityFromPhysicsMaterial(bool bInDensity
 void UGeometryCollectionComponent::UpdateCachedBounds()
 {
 	ComponentSpaceBounds = ComputeBounds(FTransform::Identity);
-	CalculateLocalBounds();
 	UpdateBounds();
 }
 
@@ -1679,6 +1692,7 @@ void UGeometryCollectionComponent::RestTransformsChanged()
 		// only need to mark it dirty and let whoever needs it to compute it on demand
 		ComponentSpaceTransforms.MarkDirty();
 	}
+	RootSpaceBounds.Init();
 
 	RefreshEmbeddedGeometry();
 	RefreshCustomRenderer();
@@ -3592,6 +3606,7 @@ void UGeometryCollectionComponent::ResetDynamicCollection()
 
 	ComponentSpaceTransforms.Reset(NumRestCollectionTransforms, GetRootIndex());
 
+	RootSpaceBounds.Init();
 	UpdateCachedBounds();
 }
 
@@ -4403,12 +4418,6 @@ void UGeometryCollectionComponent::SendRenderDynamicData_Concurrent()
 			}
 		}		
 	}
-
-	// mark collection clean now that we have rendered
-	if (DynamicCollection)
-	{
-		DynamicCollection->MakeClean();
-	}
 }
 
 void UGeometryCollectionComponent::SetCollisionObjectType(ECollisionChannel Channel)
@@ -5411,12 +5420,6 @@ AChaosSolverActor* UGeometryCollectionComponent::GetPhysicsSolverActor() const
 	return nullptr;
 }
 
-void UGeometryCollectionComponent::CalculateLocalBounds()
-{
-	LocalBounds.Init();
-	LocalBounds = ComputeBounds(FTransform::Identity);
-}
-
 #define GEOMETRY_COLLECTION_CHECK_FOR_NANS_IN_TRANSFORMS 0
 
 #if GEOMETRY_COLLECTION_CHECK_FOR_NANS_IN_TRANSFORMS
diff --git a/Engine/Source/Runtime/Experimental/GeometryCollectionEngine/Public/GeometryCollection/GeometryCollectionComponent.h b/Engine/Source/Runtime/Experimental/GeometryCollectionEngine/Public/GeometryCollection/GeometryCollectionComponent.h
--- a/Engine/Source/Runtime/Experimental/GeometryCollectionEngine/Public/GeometryCollection/GeometryCollectionComponent.h
+++ b/Engine/Source/Runtime/Experimental/GeometryCollectionEngine/Public/GeometryCollection/GeometryCollectionComponent.h
@@ -613,7 +613,7 @@ public:
 	
 	//~ Begin USceneComponent Interface.
 	GEOMETRYCOLLECTIONENGINE_API virtual FBoxSphereBounds CalcBounds(const FTransform& LocalToWorld) const override;
-	virtual FBoxSphereBounds CalcLocalBounds() const { return LocalBounds; }
+	virtual FBoxSphereBounds CalcLocalBounds() const { return ComponentSpaceBounds; }
 
 	GEOMETRYCOLLECTIONENGINE_API virtual bool HasAnySockets() const override;
 	GEOMETRYCOLLECTIONENGINE_API virtual bool DoesSocketExist(FName InSocketName) const override;
@@ -666,7 +666,7 @@ public:
 	* Get local bounds of the geometry collection
 	*/
 	UFUNCTION(BlueprintCallable, Category = "ChaosPhysics")
-	FBox GetLocalBounds() const { return LocalBounds; }
+	FBox GetLocalBounds() const { return ComponentSpaceBounds.GetBox(); }
 
 	/**
 	 * Apply an external strain to specific piece of the geometry collection
@@ -1420,9 +1420,10 @@ protected:
 	GEOMETRYCOLLECTIONENGINE_API void DispatchFieldCommand(const FFieldSystemCommand& InCommand);
 
 	GEOMETRYCOLLECTIONENGINE_API Chaos::FPhysicsSolver* GetSolver(const UGeometryCollectionComponent& GeometryCollectionComponent);
-	GEOMETRYCOLLECTIONENGINE_API void CalculateLocalBounds();
-	//GEOMETRYCOLLECTIONENGINE_API void CalculateGlobalMatrices();
-	
+
+	UE_DEPRECATED(5.4, "CalculateLocalBounds is now Deprecated as it does not need to be called anymore, see ComponentSpaceBounds which replace LocalBounds")
+	GEOMETRYCOLLECTIONENGINE_API void CalculateLocalBounds() {};
+
 	UE_DEPRECATED(5.3, "Use ComputeBoundsFromComponentSpaceTransforms instead")
 	GEOMETRYCOLLECTIONENGINE_API FBox ComputeBoundsFromGlobalMatrices(const FMatrix& LocalToWorldWithScale, const TArray<FMatrix>& GlobalMatricesArray) const;
 
@@ -1620,8 +1621,13 @@ private:
 
 	FComponentSpaceTransforms ComponentSpaceTransforms;
 
-	FBox LocalBounds;
+	/** bounds for unbroken state bounds in root space */
+	mutable FBox RootSpaceBounds;
 
+	/** 
+	* Bounds in component space 
+	* if unbroken this will use computed from RootSpaceBounds
+	*/
 	mutable FBoxSphereBounds ComponentSpaceBounds;
 
 	float CurrentCacheTime;
