Fix unnecessary proxy recreate overhead when moving ISMPool components.

diff --git a/Engine/Source/Runtime/Experimental/GeometryCollectionEngine/Private/GeometryCollection/GeometryCollectionISMPoolComponent.cpp b/Engine/Source/Runtime/Experimental/GeometryCollectionEngine/Private/GeometryCollection/GeometryCollectionISMPoolComponent.cpp
--- a/Engine/Source/Runtime/Experimental/GeometryCollectionEngine/Private/GeometryCollection/GeometryCollectionISMPoolComponent.cpp
+++ b/Engine/Source/Runtime/Experimental/GeometryCollectionEngine/Private/GeometryCollection/GeometryCollectionISMPoolComponent.cpp
@@ -86,6 +86,8 @@ void FGeometryCollectionISM::InitISM(const FGeometryCollectionStaticMeshInstance
 #endif
 
 	ISMComponent->SetStaticMesh(MeshInstance.StaticMesh);
+	ISMComponent->SetMobility((MeshInstance.Desc.Flags & FISMComponentDescription::StaticMobility) != 0 ? EComponentMobility::Static : EComponentMobility::Movable);
+
 	ISMComponent->EmptyOverrideMaterials();
 	for (int32 MaterialIndex = 0; MaterialIndex < MeshInstance.MaterialsOverrides.Num(); MaterialIndex++)
 	{
@@ -108,7 +110,6 @@ void FGeometryCollectionISM::InitISM(const FGeometryCollectionStaticMeshInstance
 		ISMComponent->SetRelativeTransform(FTransform(FQuat::Identity, MeshInstance.Desc.Position, Scale));
 	}
 
-	ISMComponent->SetMobility((MeshInstance.Desc.Flags & FISMComponentDescription::StaticMobility) != 0 ? EComponentMobility::Static : EComponentMobility::Stationary);
 	ISMComponent->SetCachedMaxDrawDistance(MeshInstance.Desc.EndCullDistance);
 	ISMComponent->SetCullDistances(MeshInstance.Desc.StartCullDistance, MeshInstance.Desc.EndCullDistance);
 	ISMComponent->SetCastShadow((MeshInstance.Desc.Flags & FISMComponentDescription::AffectShadow) != 0);
@@ -309,18 +310,18 @@ void FGeometryCollectionISMPool::RemoveISM(const FGeometryCollectionMeshInfo& Me
 			ISM.InstanceGroups.Reset();
 			ISM.InstanceIds.Reset();
 		}
-
+	
 		if (GUseComponentFreeList && ISM.InstanceGroups.IsEmpty())
-		{
-			// Remove component and push this ISM slot to the free list.
+			{
+				// Remove component and push this ISM slot to the free list.
 			ensure(ISM.ISMComponent->PerInstanceSMData.Num() == 0);
-			MeshToISMIndex.Remove(ISM.MeshInstance);
-			FreeListISM.Add(MeshInfo.ISMIndex);
+				MeshToISMIndex.Remove(ISM.MeshInstance);
+				FreeListISM.Add(MeshInfo.ISMIndex);
 
 #if WITH_EDITOR
-			ISM.ISMComponent->Rename(nullptr);
+				ISM.ISMComponent->Rename(nullptr);
 #endif
-		}
+			}
 	}
 }
 
@@ -437,7 +438,7 @@ void UGeometryCollectionISMPoolComponent::PreallocateMeshInstance(const FGeometr
 	// If we are recycling components with a free list then we don't expect to have zero instance components.
 	// So don't do preallocation of components either in that case.
 	if (!GUseComponentFreeList)
-	{
+{
 		Pool.AddISM(this, MeshInstance);
 	}
 }
