ISMPool per primitive max distance cull is now controlled with a flag rather than always set.

diff --git a/Engine/Source/Runtime/Experimental/GeometryCollectionEngine/Private/GeometryCollection/GeometryCollectionISMPoolComponent.cpp b/Engine/Source/Runtime/Experimental/GeometryCollectionEngine/Private/GeometryCollection/GeometryCollectionISMPoolComponent.cpp
--- a/Engine/Source/Runtime/Experimental/GeometryCollectionEngine/Private/GeometryCollection/GeometryCollectionISMPoolComponent.cpp
+++ b/Engine/Source/Runtime/Experimental/GeometryCollectionEngine/Private/GeometryCollection/GeometryCollectionISMPoolComponent.cpp
@@ -109,8 +109,10 @@ void FGeometryCollectionISM::InitISM(const FGeometryCollectionStaticMeshInstance
 	{
 		ISMComponent->SetRelativeTransform(FTransform(FQuat::Identity, MeshInstance.Desc.Position, Scale));
 	}
-
-	ISMComponent->SetCachedMaxDrawDistance(MeshInstance.Desc.EndCullDistance);
+	if ((MeshInstance.Desc.Flags & FISMComponentDescription::DistanceCullPrimitive) != 0)
+	{
+		ISMComponent->SetCachedMaxDrawDistance(MeshInstance.Desc.EndCullDistance);
+	}
 	ISMComponent->SetCullDistances(MeshInstance.Desc.StartCullDistance, MeshInstance.Desc.EndCullDistance);
 	ISMComponent->SetCastShadow((MeshInstance.Desc.Flags & FISMComponentDescription::AffectShadow) != 0);
 	ISMComponent->bAffectDynamicIndirectLighting = (MeshInstance.Desc.Flags & FISMComponentDescription::AffectDynamicIndirectLighting) != 0;
diff --git a/Engine/Source/Runtime/Experimental/GeometryCollectionEngine/Public/GeometryCollection/GeometryCollectionISMPoolComponent.h b/Engine/Source/Runtime/Experimental/GeometryCollectionEngine/Public/GeometryCollection/GeometryCollectionISMPoolComponent.h
--- a/Engine/Source/Runtime/Experimental/GeometryCollectionEngine/Public/GeometryCollection/GeometryCollectionISMPoolComponent.h
+++ b/Engine/Source/Runtime/Experimental/GeometryCollectionEngine/Public/GeometryCollection/GeometryCollectionISMPoolComponent.h
@@ -110,6 +110,7 @@ struct FISMComponentDescription
 		AffectShadow = 1 << 6,
 		AffectDistanceFieldLighting = 1 << 7,
 		AffectDynamicIndirectLighting = 1 << 8,
+		DistanceCullPrimitive = 1 << 9,
 	};
 
 	uint32 Flags = WorldPositionOffsetWritesVelocity|AffectShadow;
