Chaos - fix debug draw for kinematic center of mass and inertia
- sim/query-only debug draw settings work for inertia and particle transform

diff --git a/Engine/Source/Runtime/Experimental/Chaos/Private/Chaos/ChaosDebugDraw.cpp b/Engine/Source/Runtime/Experimental/Chaos/Private/Chaos/ChaosDebugDraw.cpp
--- a/Engine/Source/Runtime/Experimental/Chaos/Private/Chaos/ChaosDebugDraw.cpp
+++ b/Engine/Source/Runtime/Experimental/Chaos/Private/Chaos/ChaosDebugDraw.cpp
@@ -350,6 +350,46 @@ namespace Chaos
 			return DebugDrawSettings.ShapesColorsPerState.GetColorFromState(Particle->ObjectState(), ParticleIsOneWayInteraction(Particle));
 		}
 
+		//
+		//
+		//
+
+		bool IsParticleQueryOnly(const FGeometryParticleHandle* InParticle)
+		{
+			for (const FShapeInstancePtr& ShapeInstance : InParticle->ShapeInstances())
+			{
+				if (ShapeInstance->GetSimEnabled() || ShapeInstance->GetIsProbe())
+				{
+					return false;
+				}
+			}
+			return true;
+		}
+
+		bool IsParticleSimOnly(const FGeometryParticleHandle* InParticle)
+		{
+			for (const FShapeInstancePtr& ShapeInstance : InParticle->ShapeInstances())
+			{
+				if (ShapeInstance->GetQueryEnabled() || ShapeInstance->GetIsProbe())
+				{
+					return false;
+				}
+			}
+			return true;
+		}
+
+		bool IsParticleProbeOnly(const FGeometryParticleHandle* InParticle)
+		{
+			for (const FShapeInstancePtr& ShapeInstance : InParticle->ShapeInstances())
+			{
+				if (ShapeInstance->GetQueryEnabled() || ShapeInstance->GetSimEnabled())
+				{
+					return false;
+				}
+			}
+			return true;
+		}
+
 		//
 		//
 		//
@@ -1020,6 +1060,28 @@ namespace Chaos
 
 		void DrawParticleTransformImpl(const FRigidTransform3& SpaceTransform, const FGeometryParticleHandle* InParticle, int32 Index, FRealSingle ColorScale, const FChaosDebugDrawSettings& Settings)
 		{
+			if (!bChaosDebugDebugDrawShowQueryOnlyShapes)
+			{
+				if (IsParticleQueryOnly(InParticle))
+				{
+					return;
+				}
+			}
+			if (!bChaosDebugDebugDrawShowSimOnlyShapes)
+			{
+				if (IsParticleSimOnly(InParticle))
+				{
+					return;
+				}
+			}
+			if (!bChaosDebugDebugDrawShowProbeOnlyShapes)
+			{
+				if (IsParticleProbeOnly(InParticle))
+				{
+					return;
+				}
+			}
+
 			const TPBDRigidParticleHandle<FReal, 3>* Rigid = InParticle->CastToRigidParticle();
 			if (Rigid && Rigid->Disabled())
 			{
@@ -1031,8 +1093,18 @@ namespace Chaos
 			FColor Blue = (ColorScale * FColor::Blue).ToFColor(false);
 
 			FConstGenericParticleHandle Particle(InParticle);
-			FVec3 PCOM = SpaceTransform.TransformPosition(FParticleUtilities::GetCoMWorldPosition(Particle));
-			FRotation3 QCOM = SpaceTransform.GetRotation() * FParticleUtilities::GetCoMWorldRotation(Particle);
+
+			FRigidTransform3 SpaceCOM;
+			if (Rigid != nullptr)
+			{
+				SpaceCOM = Rigid->GetTransformPQCom() * SpaceTransform;
+			}
+			else
+			{
+				SpaceCOM = InParticle->GetTransformXR() * SpaceTransform;
+			}
+			FVec3 PCOM = SpaceCOM.GetLocation();
+			FRotation3 QCOM = SpaceCOM.GetRotation();
 			FMatrix33 QCOMm = QCOM.ToMatrix();
 			FDebugDrawQueue::GetInstance().DrawDebugDirectionalArrow(PCOM, PCOM + Settings.DrawScale * Settings.BodyAxisLen * QCOMm.GetAxis(0), Settings.DrawScale * Settings.ArrowSize, Red, false, UE_KINDA_SMALL_NUMBER, uint8(Settings.DrawPriority), Settings.LineThickness);
 			FDebugDrawQueue::GetInstance().DrawDebugDirectionalArrow(PCOM, PCOM + Settings.DrawScale * Settings.BodyAxisLen * QCOMm.GetAxis(1), Settings.DrawScale * Settings.ArrowSize, Green, false, UE_KINDA_SMALL_NUMBER, uint8(Settings.DrawPriority), Settings.LineThickness);
@@ -1399,16 +1471,7 @@ namespace Chaos
 			{
 				if (!bChaosDebugDebugDrawShowQueryOnlyShapes)
 				{
-					bool bIsQueryOnly = true;
-					for (const FShapeInstancePtr& ShapeInstance : InParticle->ShapeInstances())
-					{
-						if (ShapeInstance->GetSimEnabled())
-						{
-							bIsQueryOnly = false;
-							break;
-						}
-					}
-					if (bIsQueryOnly)
+					if (IsParticleQueryOnly(InParticle))
 					{
 						return;
 					}
