Iris dynamic filters
* Added ability to override the default class setting for individual actors via FActorBeginReplicationParams
* Added Net.Iris.PrintDynamicFilterClassConfig to be able to print the current class mappings

#rnx

diff --git a/Engine/Plugins/Runtime/ReplicationSystemTestPlugin/Source/Private/Tests/ReplicationSystem/Filtering/TestGridFilter.cpp b/Engine/Plugins/Runtime/ReplicationSystemTestPlugin/Source/Private/Tests/ReplicationSystem/Filtering/TestGridFilter.cpp
--- a/Engine/Plugins/Runtime/ReplicationSystemTestPlugin/Source/Private/Tests/ReplicationSystem/Filtering/TestGridFilter.cpp
+++ b/Engine/Plugins/Runtime/ReplicationSystemTestPlugin/Source/Private/Tests/ReplicationSystem/Filtering/TestGridFilter.cpp
@@ -104,7 +104,7 @@ UE_NET_TEST_FIXTURE(FTestGridFilterFixture, TestWorldLocGridFilter)
 	UObjectReplicationBridge::FCreateNetRefHandleParams Params;
 	Params.bCanReceive = true;
 	Params.bNeedsWorldLocationUpdate = true;
-	Params.bAllowDynamicFilter = true;
+	Params.bUseClassConfigDynamicFilter = true;
 
 	UReplicatedTestObject* ServerObjectZero = Server->CreateObject(Params);
 	UReplicatedTestObject* ServerObjectLimit = Server->CreateObject(Params);
diff --git a/Engine/Plugins/Runtime/ReplicationSystemTestPlugin/Source/Private/Tests/ReplicationSystem/TestDirtyNetObjectTracker.cpp b/Engine/Plugins/Runtime/ReplicationSystemTestPlugin/Source/Private/Tests/ReplicationSystem/TestDirtyNetObjectTracker.cpp
--- a/Engine/Plugins/Runtime/ReplicationSystemTestPlugin/Source/Private/Tests/ReplicationSystem/TestDirtyNetObjectTracker.cpp
+++ b/Engine/Plugins/Runtime/ReplicationSystemTestPlugin/Source/Private/Tests/ReplicationSystem/TestDirtyNetObjectTracker.cpp
@@ -147,7 +147,7 @@ UE_NET_TEST_FIXTURE(FReplicationSystemServerClientTestFixture, GlobalDirtyTracke
 	const float PollFrequency = Server->ConvertPollPeriodIntoFrequency(PollPeriod);
 	Params.PollFrequency = PollFrequency;
 	Params.bCanReceive = true;
-	Params.bAllowDynamicFilter = true;
+	Params.bUseClassConfigDynamicFilter = true;
 	Params.bNeedsPreUpdate = true;
 
 	UTestReplicatedIrisObject* ServerObject = Server->CreateObject(Params);
@@ -192,7 +192,7 @@ UE_NET_TEST_FIXTURE(FReplicationSystemServerClientTestFixture, DirtyInsidePreUpd
 	const float PollFrequency = Server->ConvertPollPeriodIntoFrequency(PollPeriod);
 	Params.PollFrequency = PollFrequency;
 	Params.bCanReceive = true;
-	Params.bAllowDynamicFilter = true;	
+	Params.bUseClassConfigDynamicFilter = true;
 	Params.bNeedsPreUpdate = true;
 	UTestReplicatedIrisObject* ServerObjectA = Server->CreateObject(Params);
 	UTestReplicatedIrisObject* ServerObjectB = Server->CreateObject(Params);
diff --git a/Engine/Plugins/Runtime/ReplicationSystemTestPlugin/Source/Private/Tests/ReplicationSystem/TestObjectPollFrequencyLimiter.cpp b/Engine/Plugins/Runtime/ReplicationSystemTestPlugin/Source/Private/Tests/ReplicationSystem/TestObjectPollFrequencyLimiter.cpp
--- a/Engine/Plugins/Runtime/ReplicationSystemTestPlugin/Source/Private/Tests/ReplicationSystem/TestObjectPollFrequencyLimiter.cpp
+++ b/Engine/Plugins/Runtime/ReplicationSystemTestPlugin/Source/Private/Tests/ReplicationSystem/TestObjectPollFrequencyLimiter.cpp
@@ -124,7 +124,7 @@ UE_NET_TEST_FIXTURE(FReplicationSystemServerClientTestFixture, TestSubObjectsPol
 	UObjectReplicationBridge::FCreateNetRefHandleParams Params;
 	Params.PollFrequency = PollFrequency;
 	Params.bCanReceive = true;
-	Params.bAllowDynamicFilter = true;
+	Params.bUseClassConfigDynamicFilter = true;
 	Params.bNeedsPreUpdate = true;
 	UTestReplicatedIrisObject* ServerRootObject = Server->CreateObject(Params);
 
diff --git a/Engine/Plugins/Runtime/ReplicationSystemTestPlugin/Source/Private/Tests/ReplicationSystem/TestPreUpdateCallback.cpp b/Engine/Plugins/Runtime/ReplicationSystemTestPlugin/Source/Private/Tests/ReplicationSystem/TestPreUpdateCallback.cpp
--- a/Engine/Plugins/Runtime/ReplicationSystemTestPlugin/Source/Private/Tests/ReplicationSystem/TestPreUpdateCallback.cpp
+++ b/Engine/Plugins/Runtime/ReplicationSystemTestPlugin/Source/Private/Tests/ReplicationSystem/TestPreUpdateCallback.cpp
@@ -26,7 +26,7 @@ UE_NET_TEST_FIXTURE(FReplicationSystemServerClientTestFixture, CreateObjectsInsi
 	const float PollFrequency = Server->ConvertPollPeriodIntoFrequency(PollPeriod);
 	Params.PollFrequency = PollFrequency;
 	Params.bCanReceive = true;
-	Params.bAllowDynamicFilter = true;
+	Params.bUseClassConfigDynamicFilter = true;
 	Params.bNeedsPreUpdate = true;
 	UTestReplicatedIrisObject* ServerObject = Server->CreateObject(Params);
 	UTestReplicatedIrisObject* ServerSubObject(nullptr);
diff --git a/Engine/Source/Runtime/Engine/Private/Net/Iris/ReplicationSystem/ActorReplicationBridge.cpp b/Engine/Source/Runtime/Engine/Private/Net/Iris/ReplicationSystem/ActorReplicationBridge.cpp
--- a/Engine/Source/Runtime/Engine/Private/Net/Iris/ReplicationSystem/ActorReplicationBridge.cpp
+++ b/Engine/Source/Runtime/Engine/Private/Net/Iris/ReplicationSystem/ActorReplicationBridge.cpp
@@ -242,7 +242,7 @@ UE::Net::FNetRefHandle UActorReplicationBridge::BeginReplication(AActor* Actor,
 	}
 
 	// Create handles for the registered fragments
-	Super::FCreateNetRefHandleParams CreateNetRefHandleParams = DefaultCreateNetRefHandleParams;
+	Super::FCreateNetRefHandleParams CreateNetRefHandleParams = UObjectReplicationBridge::DefaultCreateNetRefHandleParams;
 	CreateNetRefHandleParams.bNeedsPreUpdate = 1U;
 	CreateNetRefHandleParams.bNeedsWorldLocationUpdate = 1U;
 	CreateNetRefHandleParams.StaticPriority = (Actor->bAlwaysRelevant || Actor->bOnlyRelevantToOwner) ? Actor->NetPriority : 0.0f;
diff --git a/Engine/Source/Runtime/Engine/Public/Net/Iris/ReplicationSystem/ActorReplicationBridge.h b/Engine/Source/Runtime/Engine/Public/Net/Iris/ReplicationSystem/ActorReplicationBridge.h
--- a/Engine/Source/Runtime/Engine/Public/Net/Iris/ReplicationSystem/ActorReplicationBridge.h
+++ b/Engine/Source/Runtime/Engine/Public/Net/Iris/ReplicationSystem/ActorReplicationBridge.h
@@ -32,9 +32,14 @@ ENGINE_API bool ShouldUseIrisReplication(const UObject* Actor);
 /** Parameters passed to UActorReplicationBridge::BeginReplication. */
 struct FActorBeginReplicationParams
 {
-	FActorBeginReplicationParams() : bIncludeInLevelGroupFilter(1U) {}
+	/** Will the actor be part to the level visibility filters. This makes it that only clients who loaded the dynamic level of the actor can have the actor replicated to them.*/
+	bool bIncludeInLevelGroupFilter = true;
 
-	uint32 bIncludeInLevelGroupFilter : 1U;
+	/** When true we ignore the configured dynamic filter for this actor type and use the explicit filter instead */
+	bool bOverrideDynamicFilterConfig = false;
+
+	/** Only used when bOverrideDynamicFilterConfig is true. The name of the dynamic filter to assign to this actor. Can be NAME_None if no filters should be set. */
+	FName ExplicitDynamicFilterName;
 };
 
 #endif // UE_WITH_IRIS
diff --git a/Engine/Source/Runtime/Experimental/Iris/Core/Private/Iris/ReplicationSystem/ObjectReplicationBridge.cpp b/Engine/Source/Runtime/Experimental/Iris/Core/Private/Iris/ReplicationSystem/ObjectReplicationBridge.cpp
--- a/Engine/Source/Runtime/Experimental/Iris/Core/Private/Iris/ReplicationSystem/ObjectReplicationBridge.cpp
+++ b/Engine/Source/Runtime/Experimental/Iris/Core/Private/Iris/ReplicationSystem/ObjectReplicationBridge.cpp
@@ -37,6 +37,8 @@
 #include "Iris/Serialization/InternalNetSerializationContext.h"
 #include "Iris/Serialization/NetBitStreamUtil.h"
 
+DEFINE_LOG_CATEGORY(LogIrisFilterConfig)
+
 #define UE_LOG_OBJECTREPLICATIONBRIDGE(Category, Format, ...)  UE_LOG(LogIrisBridge, Category, TEXT("ObjectReplicationBridge(%u)::") Format, GetReplicationSystem()->GetId(), ##__VA_ARGS__)
 
 static bool bUseFrequencyBasedPolling = true;
@@ -76,12 +78,12 @@ static FAutoConsoleVariableRef CVarEnableForceNetUpdate(
 
 UObjectReplicationBridge::FCreateNetRefHandleParams UObjectReplicationBridge::DefaultCreateNetRefHandleParams =
 {
-	false, // bCanReceive
-	false, // bNeedsPreUpdate
-	false, // bNeedsWorldLocationUpdate
-	true, // bAllowDynamicFilter
-	0.0f, // StaticPriority
-	0U, // PollFramePeriod
+	.bCanReceive=false, 
+	.bNeedsPreUpdate=false,
+	.bNeedsWorldLocationUpdate=false, 
+	.bUseClassConfigDynamicFilter=true,
+	.StaticPriority=0.0f, 
+	.PollFrequency=0.0f,
 };
 
 namespace UE::Net::Private
@@ -357,15 +359,7 @@ UE::Net::FNetRefHandle UObjectReplicationBridge::BeginReplication(UObject* Insta
 				}
 			}
 
-			if (bEnableFilterMappings && Params.bAllowDynamicFilter)
-			{
-				const FNetObjectFilterHandle FilterHandle = GetDynamicFilter(Instance->GetClass());
-				if (FilterHandle != InvalidNetObjectFilterHandle)
-				{
-					UE_LOG_OBJECTREPLICATIONBRIDGE(Verbose, TEXT("BeginReplication Filter: %s will be used for Object: %s "), *(ReplicationSystem->GetFilterName(FilterHandle).ToString()), *Instance->GetName());
-					ReplicationSystem->SetFilter(RefHandle, FilterHandle);
-				}
-			}
+			AssignDynamicFilter(Instance, Params, RefHandle);
 
 			if (ShouldClassBeDeltaCompressed(Instance->GetClass()))
 			{
@@ -386,6 +380,40 @@ UE::Net::FNetRefHandle UObjectReplicationBridge::BeginReplication(UObject* Insta
 	return FNetRefHandle();
 }
 
+void UObjectReplicationBridge::AssignDynamicFilter(UObject* Instance, const FCreateNetRefHandleParams& Params, FNetRefHandle RefHandle)
+{
+	using namespace UE::Net;
+
+	if (!bEnableFilterMappings)
+	{
+		return;
+	}
+
+	FNetObjectFilterHandle FilterHandle = InvalidNetObjectFilterHandle;
+
+	if (Params.bUseExplicitDynamicFilter)
+	{
+		if (Params.ExplicitDynamicFilterName != NAME_None)
+		{
+			FilterHandle = ReplicationSystem->GetFilterHandle(Params.ExplicitDynamicFilterName);
+			
+			UE_CLOG(FilterHandle == InvalidNetObjectFilterHandle, LogIrisBridge, Error, TEXT("Could not assign explicit dynamic filter to %s. No filters named %s exist"), *GetPathNameSafe(Instance), *Params.ExplicitDynamicFilterName.ToString() );
+			ensure(FilterHandle != InvalidNetObjectFilterHandle);
+		}
+	}
+	else if (Params.bUseClassConfigDynamicFilter)
+	{
+		FilterHandle = GetDynamicFilter(Instance->GetClass());
+	}
+
+	if (FilterHandle != InvalidNetObjectFilterHandle)
+	{
+		UE_LOG_OBJECTREPLICATIONBRIDGE(Verbose, TEXT("BeginReplication Filter: %s will be used for Object: %s "), *(ReplicationSystem->GetFilterName(FilterHandle).ToString()), *Instance->GetName());
+		ReplicationSystem->SetFilter(RefHandle, FilterHandle);
+	}
+	
+}
+
 UE::Net::FNetRefHandle UObjectReplicationBridge::BeginReplication(FNetRefHandle OwnerRefHandle, UObject* Instance, FNetRefHandle InsertRelativeToSubObjectRefHandle, ESubObjectInsertionOrder InsertionOrder, const FCreateNetRefHandleParams& Params)
 {
 	LLM_SCOPE_BYTAG(IrisState);
@@ -416,7 +444,7 @@ UE::Net::FNetRefHandle UObjectReplicationBridge::BeginReplication(FNetRefHandle
 
 		FCreateNetRefHandleParams SubObjectCreateParams = Params;
 		// The filtering system ignores subobjects so let's not waste cycles figuring out which filter to use.
-		SubObjectCreateParams.bAllowDynamicFilter = 0U;
+		SubObjectCreateParams.bUseClassConfigDynamicFilter = false;
 		SubObjectRefHandle = BeginReplication(Instance, SubObjectCreateParams);
 	}
 
@@ -1325,6 +1353,61 @@ bool UObjectReplicationBridge::ShouldClassBeDeltaCompressed(const UClass* Class)
 	return false;
 }
 
+void UObjectReplicationBridge::SetClassDynamicFilterConfig(FName ClassPathName, const UE::Net::FNetObjectFilterHandle FilterHandle)
+{
+	if (ClassPathName.IsNone())
+	{
+		return;
+	}
+
+	if (UE_LOG_ACTIVE(LogIrisFilterConfig, Log))
+	{
+		if (const UE::Net::FNetObjectFilterHandle* OldFilter = ClassesWithDynamicFilter.Find(ClassPathName))
+		{
+			if (*OldFilter != FilterHandle)
+			{
+				UE_LOG(LogIrisFilterConfig, Log, TEXT("SetClassDynamicFilterConfig assigned %s to use filter %s. Previously using filter %s."),
+					*ClassPathName.ToString(), *(ReplicationSystem->GetFilterName(FilterHandle).ToString()), *(ReplicationSystem->GetFilterName(*OldFilter).ToString()));
+			}
+			else
+			{
+				UE_LOG(LogIrisFilterConfig, Log, TEXT("SetClassDynamicFilterConfig assigned %s to use filter %s but the class was already assigned to this filter."), *ClassPathName.ToString(), *(ReplicationSystem->GetFilterName(FilterHandle).ToString()));
+			}
+		}
+		else
+		{
+			UE_LOG(LogIrisFilterConfig, Log, TEXT("SetClassDynamicFilterConfig assigned %s to use filter %s."), *ClassPathName.ToString(), *(ReplicationSystem->GetFilterName(FilterHandle).ToString()));
+		}
+	}
+
+	ClassesWithDynamicFilter.Add(ClassPathName, FilterHandle);
+}
+
+void UObjectReplicationBridge::SetClassDynamicFilterConfig(FName ClassPathName, FName FilterName)
+{
+	using namespace UE::Net;
+
+	if (ClassPathName.IsNone())
+	{
+		return;
+	}
+
+	if (FilterName != NAME_None)
+	{
+		FNetObjectFilterHandle FilterHandle = GetReplicationSystem()->GetFilterHandle(FilterName);;
+
+		if (ensureMsgf(FilterHandle != InvalidNetObjectFilterHandle, TEXT("SetClassDynamicFilterConfig for %s received invalid filter named %s"), *ClassPathName.ToString(), *FilterName.ToString()))
+		{
+			SetClassDynamicFilterConfig(ClassPathName, FilterHandle);
+		}
+	}
+	else
+	{
+		// Reset the filter so the class does not get assigned a dynamic filter anymore.
+		SetClassDynamicFilterConfig(ClassPathName, InvalidNetObjectFilterHandle);
+	}
+}
+
 UE::Net::FNetObjectFilterHandle UObjectReplicationBridge::GetDynamicFilter(const UClass* Class)
 {
 	using namespace UE::Net;
diff --git a/Engine/Source/Runtime/Experimental/Iris/Core/Private/Iris/ReplicationSystem/ObjectReplicationBridgeDebugging.cpp b/Engine/Source/Runtime/Experimental/Iris/Core/Private/Iris/ReplicationSystem/ObjectReplicationBridgeDebugging.cpp
new file mode 100644
--- /dev/null
+++ b/Engine/Source/Runtime/Experimental/Iris/Core/Private/Iris/ReplicationSystem/ObjectReplicationBridgeDebugging.cpp
@@ -0,0 +1,78 @@
+// Copyright Epic Games, Inc. All Rights Reserved.
+
+#include "Iris/ReplicationSystem/ObjectReplicationBridge.h"
+
+#include "HAL/IConsoleManager.h"
+
+#include "Iris/Core/IrisLog.h"
+#include "Iris/Core/IrisDebugging.h"
+
+#include "Iris/ReplicationSystem/ReplicationSystem.h"
+#include "Iris/ReplicationSystem/ReplicationSystemTypes.h"
+#include "Iris/ReplicationSystem/ReplicationSystemInternal.h"
+
+#include "Net/Core/NetBitArrayPrinter.h"
+
+namespace UE::Net::Private::ObjectBridgeDebugging
+{
+
+UReplicationSystem* FindReplicationSystem(const TArray<FString>& Args)
+{
+	uint32 RepSystemId = 0;
+
+	// If the ReplicationSystemId was specified
+	if (const FString* ArgRepSystemId = Args.FindByPredicate([](const FString& Str) { return Str.Contains(TEXT("RepSystemId=")); }))
+	{
+		FParse::Value(**ArgRepSystemId, TEXT("RepSystemId="), RepSystemId);
+	}
+
+	return UE::Net::GetReplicationSystem(RepSystemId);
+}
+
+
+} // end namespace UE::Net::Private::ObjectBridgeDebugging
+
+// --------------------------------------------------------------------------------------------------------------------------------------------
+// Debug commands
+// --------------------------------------------------------------------------------------------------------------------------------------------
+
+FAutoConsoleCommand ObjectBridgePrintDynamicFilter(TEXT("Net.Iris.PrintDynamicFilterClassConfig"), TEXT("Prints the dynamic filter configured to be assigned to specific classes."), FConsoleCommandWithArgsDelegate::CreateLambda([](const TArray< FString >& Args)
+{
+	using namespace UE::Net::Private::ObjectBridgeDebugging;
+
+	UReplicationSystem* RepSystem = FindReplicationSystem(Args);
+	if (!RepSystem)
+	{
+		UE_LOG(LogIrisBridge, Error, TEXT("Could not find ReplicationSystem."));
+		return;
+	}
+
+	UObjectReplicationBridge* ObjectBridge = CastChecked<UObjectReplicationBridge>(RepSystem->GetReplicationBridge());
+	if (!ObjectBridge)
+	{
+		UE_LOG(LogIrisBridge, Error, TEXT("Could not find ObjectReplicationBridge."));
+		return;
+	}
+
+	ObjectBridge->PrintDynamicFilterClassConfig();
+}));
+
+void UObjectReplicationBridge::PrintDynamicFilterClassConfig() const
+{
+	const UReplicationSystem* RepSystem = GetReplicationSystem();
+
+	UE_LOG(LogIrisFilterConfig, Display, TEXT(""));
+	UE_LOG(LogIrisFilterConfig, Display, TEXT("Default Dynamic Filter Class Config:"));
+	{
+		TMap<FName, UE::Net::FNetObjectFilterHandle> SortedClassConfig = ClassesWithDynamicFilter;
+
+		SortedClassConfig.KeyStableSort([](FName lhs, FName rhs){return lhs.Compare(rhs) < 0;});
+		for (auto MapIt = SortedClassConfig.CreateConstIterator(); MapIt; ++MapIt)
+		{
+			const FName ClassName = MapIt.Key();
+			const UE::Net::FNetObjectFilterHandle FilterHandle = MapIt.Value();
+
+			UE_LOG(LogIrisFilterConfig, Display, TEXT("\t%s -> %s"), *ClassName.ToString(), *RepSystem->GetFilterName(FilterHandle).ToString());
+		}
+	}
+}
\ No newline at end of file
diff --git a/Engine/Source/Runtime/Experimental/Iris/Core/Public/Iris/ReplicationSystem/ObjectReplicationBridge.h b/Engine/Source/Runtime/Experimental/Iris/Core/Public/Iris/ReplicationSystem/ObjectReplicationBridge.h
--- a/Engine/Source/Runtime/Experimental/Iris/Core/Public/Iris/ReplicationSystem/ObjectReplicationBridge.h
+++ b/Engine/Source/Runtime/Experimental/Iris/Core/Public/Iris/ReplicationSystem/ObjectReplicationBridge.h
@@ -11,6 +11,8 @@
 
 #include "ObjectReplicationBridge.generated.h"
 
+DECLARE_LOG_CATEGORY_EXTERN(LogIrisFilterConfig, Log, All);
+
 namespace UE::Net
 {
 	struct FNetObjectResolveContext;
@@ -49,7 +51,19 @@ public:
 		bool bCanReceive = false;
 		bool bNeedsPreUpdate = false;
 		bool bNeedsWorldLocationUpdate = false;
-		bool bAllowDynamicFilter = false;
+
+		/** When true we ask the class config if a dynamic filter was assigned to this class or one of it's parent inherited class. */
+		bool bUseClassConfigDynamicFilter = false;
+
+		/** When enabled we ignore the class config for this object and instead use the one specified in ExplicitDynamicFilter */
+		bool bUseExplicitDynamicFilter = false;
+
+		/** 
+		 * The name of the dynamic filter to use for this object (instead of asking the class config). 
+		 * Can be none so that no dynamic filter is assigned to the object.
+		 * Only used when bUseExplicitDynamicFilter is true.
+		 */
+		FName ExplicitDynamicFilterName;
 
 		/**
 		 * If StaticPriority is > 0 the ReplicationSystem will use that as priority when scheduling objects. 
@@ -154,6 +168,15 @@ public:
 	/** Set poll frequency on root object and its subobjects. They will be polled on the same frame. */
 	IRISCORE_API void SetPollFrequency(FNetRefHandle RootHandle, float PollFrequency);
 
+	/** Set the object filter to use for objects of this class and any derived classes without an explicit config. */
+	IRISCORE_API void SetClassDynamicFilterConfig(FName ClassPathName, const UE::Net::FNetObjectFilterHandle FilterHandle);
+	IRISCORE_API void SetClassDynamicFilterConfig(FName ClassPathName, FName FilterName);
+
+public:
+
+	// Debug functions exposed via console commands
+	void PrintDynamicFilterClassConfig() const;
+
 protected:
 	IRISCORE_API virtual ~UObjectReplicationBridge();
 
@@ -303,6 +326,9 @@ private:
 	/** Retrieves the prioritizer to set for the given class. If bRequireForceEnabled the config needs to have bForceEnableOnAllInstances set in order for this method to return the configured prioritizer. Returns an invalid handle if no prioritizer should be set. */
 	UE::Net::FNetObjectPrioritizerHandle GetPrioritizer(const UClass* Class, bool bRequireForceEnabled);
 
+	/** Assign the proper dynamic filter to a new object */
+	void AssignDynamicFilter(UObject* Instance, const FCreateNetRefHandleParams& Params, FNetRefHandle RefHandle);
+
 	/** Returns true if instances of this class should be delta compressed */
 	bool ShouldClassBeDeltaCompressed(const UClass* Class);
 
@@ -348,6 +374,8 @@ private:
 	TSet<FName> ClassesWithoutPollPeriodOverride;
 
 	// Filter mapping
+	//$IRIS TODO: Look into improving this class map by balancing runtime speed (implicit addition of new classes) with runtime modifications of base classes.
+    //		      Right-now any changes to say APawn's filter will not be read if a APlayerPawn entry was created based on the APawn entry.
 	TMap<FName, UE::Net::FNetObjectFilterHandle> ClassesWithDynamicFilter;
 	TFunction<bool(const UClass*)> ShouldUseDefaultSpatialFilterFunction;
 	TFunction<bool(const UClass*,const UClass*)> ShouldSubclassUseSameFilterFunction;
