Geometry collection optimization : improve reuse of cached  component space bounds
- changed ComponentSpaceBounds to be an Fbox
- CalcBounds now do not check for dynamic collection but when the dynamic collection is modified the component space bounds are invalidated
- invalidation happens in OnPostPhysicsState and when rest transforms are updated
- Cached IsRootBroken state

diff --git a/Engine/Source/Runtime/Experimental/GeometryCollectionEngine/Private/GeometryCollection/GeometryCollectionComponent.cpp b/Engine/Source/Runtime/Experimental/GeometryCollectionEngine/Private/GeometryCollection/GeometryCollectionComponent.cpp
--- a/Engine/Source/Runtime/Experimental/GeometryCollectionEngine/Private/GeometryCollection/GeometryCollectionComponent.cpp
+++ b/Engine/Source/Runtime/Experimental/GeometryCollectionEngine/Private/GeometryCollection/GeometryCollectionComponent.cpp
@@ -596,6 +596,7 @@ UGeometryCollectionComponent::UGeometryCollectionComponent(const FObjectInitiali
 	, bIsTransformSelectionModeEnabled(false)
 #endif  // #if GEOMETRYCOLLECTION_EDITOR_SELECTION
 	, bIsMoving(false)
+	, bIsRootBroken(false)
 	, bUpdateCustomRenderer(true)
 {
 	// by default tick is registered but disabled, we only need it when we need to update the removal timers
@@ -923,19 +924,11 @@ FBoxSphereBounds UGeometryCollectionComponent::CalcBounds(const FTransform& Loca
 {	
 	SCOPE_CYCLE_COUNTER(STAT_GCCUpdateBounds);
 
-	bool NeedBoundsUpdate = false;
-	NeedBoundsUpdate |= (ComponentSpaceBounds.GetSphere().W < 1e-5);
-	NeedBoundsUpdate |= CachePlayback;
-	NeedBoundsUpdate |= (DynamicCollection && DynamicCollection->IsDirty());
-
+	const bool NeedBoundsUpdate = (!ComponentSpaceBounds.IsValid) || CachePlayback;
 	if (NeedBoundsUpdate)
 	{
 		ComponentSpaceBounds = ComputeBounds(FTransform::Identity);
 	}
-	else
-	{
-		NeedBoundsUpdate = false;
-	}
 
 	return ComponentSpaceBounds.TransformBy(LocalToWorldIn);
 }
@@ -1657,7 +1650,9 @@ void UGeometryCollectionComponent::RestTransformsChanged()
 	{
 		FGeometryCollectionDynamicData* DynamicData = GDynamicDataPool.Allocate();
 		DynamicData->SetPrevTransforms(ComponentSpaceTransforms.RequestAllTransforms());
-		ComponentSpaceTransforms.MarkDirty();
+
+		OnTransformsDirty();
+
 		DynamicData->SetTransforms(ComponentSpaceTransforms.RequestAllTransforms());
 		DynamicData->IsDynamic = true;
 
@@ -1689,8 +1684,8 @@ void UGeometryCollectionComponent::RestTransformsChanged()
 	}
 	else
 	{
-		// only need to mark it dirty and let whoever needs it to compute it on demand
-		ComponentSpaceTransforms.MarkDirty();
+		// only need to mark transform dirty and let whoever needs it to compute it on demand
+		OnTransformsDirty();
 	}
 	RootSpaceBounds.Init();
 
@@ -3608,6 +3603,8 @@ void UGeometryCollectionComponent::ResetDynamicCollection()
 
 	RootSpaceBounds.Init();
 	UpdateCachedBounds();
+
+	bIsRootBroken = false;
 }
 
 void UGeometryCollectionComponent::OnCreatePhysicsState()
@@ -3883,6 +3880,25 @@ void UGeometryCollectionComponent::RegisterAndInitializePhysicsProxy()
 	RegisterForEvents();
 }
 
+void UGeometryCollectionComponent::UpdateIsRootBroken()
+{
+	bIsRootBroken = false;
+	if (DynamicCollection && DynamicCollection->Active.Num() > 0)
+	{
+		const int32 RootIndex = GetRootIndex();
+		if (RootIndex != INDEX_NONE)
+		{
+			bIsRootBroken = !DynamicCollection->Active[RootIndex];
+		}
+	}
+}
+
+void UGeometryCollectionComponent::OnTransformsDirty()
+{
+	ComponentSpaceTransforms.MarkDirty();
+	ComponentSpaceBounds.Init();
+}
+
 void UGeometryCollectionComponent::OnPostPhysicsSync()
 {
 	SCOPE_CYCLE_COUNTER(STAT_GCPostPhysicsSync);
@@ -3892,7 +3908,9 @@ void UGeometryCollectionComponent::OnPostPhysicsSync()
 	{
 		// Can't be a const reference - we need to make a copy or else the next RequestRootTransform will change the value under our nose.
 		const FTransform3f PreviousRootTransform = OnRootMovedEvent.IsBound() ? ComponentSpaceTransforms.RequestRootTransform() : FTransform3f::Identity;
-		ComponentSpaceTransforms.MarkDirty();
+
+		OnTransformsDirty();
+		
 		if (OnRootMovedEvent.IsBound())
 		{
 			const FTransform3f& NewRootTransform = ComponentSpaceTransforms.RequestRootTransform();
@@ -3914,6 +3932,7 @@ void UGeometryCollectionComponent::OnPostPhysicsSync()
 
 	// Once the GC is broken, removal feature will need the tick to properly update the timers 
 	// even if the physics does not get any updates
+	UpdateIsRootBroken();
 	if (IsRootBroken())
 	{
 		if (!PrimaryComponentTick.IsTickFunctionEnabled())
@@ -5848,19 +5867,6 @@ TArray<UStaticMeshComponent*> UGeometryCollectionComponent::CreateProxyComponent
 	return Components;
 }
 
-bool UGeometryCollectionComponent::IsRootBroken() const
-{
-	if (DynamicCollection && DynamicCollection->Active.Num() > 0)
-	{
-		const int32 RootIndex = GetRootIndex();
-		if (RootIndex != INDEX_NONE)
-		{
-			return !DynamicCollection->Active[RootIndex];
-		}
-	}
-	return false;
-}
-
 struct FGeometryCollectionDecayContext
 {
 	FGeometryCollectionDecayContext(FGeometryCollectionPhysicsProxy& PhysicsProxyIn, FGeometryCollectionDecayDynamicFacade& DecayFacadeIn)
diff --git a/Engine/Source/Runtime/Experimental/GeometryCollectionEngine/Public/GeometryCollection/GeometryCollectionComponent.h b/Engine/Source/Runtime/Experimental/GeometryCollectionEngine/Public/GeometryCollection/GeometryCollectionComponent.h
--- a/Engine/Source/Runtime/Experimental/GeometryCollectionEngine/Public/GeometryCollection/GeometryCollectionComponent.h
+++ b/Engine/Source/Runtime/Experimental/GeometryCollectionEngine/Public/GeometryCollection/GeometryCollectionComponent.h
@@ -666,7 +666,7 @@ public:
 	* Get local bounds of the geometry collection
 	*/
 	UFUNCTION(BlueprintCallable, Category = "ChaosPhysics")
-	FBox GetLocalBounds() const { return ComponentSpaceBounds.GetBox(); }
+	FBox GetLocalBounds() const { return ComponentSpaceBounds; }
 
 	/**
 	 * Apply an external strain to specific piece of the geometry collection
@@ -789,7 +789,7 @@ public:
 
 	/** return true if the root cluster is not longer active at runtime */
 	UFUNCTION(BlueprintCallable, Category = "ChaosPhysics")
-	GEOMETRYCOLLECTIONENGINE_API bool IsRootBroken() const;
+	GEOMETRYCOLLECTIONENGINE_API bool IsRootBroken() const { return bIsRootBroken; }
 
 	/** 
 	* Get the initial rest transforms in component (local) space  space, 
@@ -1628,7 +1628,7 @@ private:
 	* Bounds in component space 
 	* if unbroken this will use computed from RootSpaceBounds
 	*/
-	mutable FBoxSphereBounds ComponentSpaceBounds;
+	mutable FBox ComponentSpaceBounds;
 
 	float CurrentCacheTime;
 	TArray<bool> EventsPlayed;
@@ -1706,6 +1706,9 @@ private:
 
 	void MoveComponentToRootTransform();
 
+	/** called when the dynamic collection is found to be dirty */
+	void OnTransformsDirty();
+
 	/** The clusters we need to replicate */
 	TUniquePtr<TSet<Chaos::FPBDRigidClusteredParticleHandle*>> ClustersToRep;
 
@@ -1716,6 +1719,14 @@ private:
 	/** True if GeometryCollection transforms have changed from previous tick. */
 	bool bIsMoving;
 
+	/** 
+	* root bone broken state 
+	* this is updated post physics sync and use as is for quick lookup 
+	*/
+	bool bIsRootBroken;
+
+	void UpdateIsRootBroken();
+
 	bool bUpdateCustomRenderer;
 
 private:
