ClusterUnion: Fix rendering of GC when in a cluster union
Also fix an uninitialized variable

diff --git a/Engine/Source/Runtime/Experimental/GeometryCollectionEngine/Private/GeometryCollection/GeometryCollectionComponent.cpp b/Engine/Source/Runtime/Experimental/GeometryCollectionEngine/Private/GeometryCollection/GeometryCollectionComponent.cpp
--- a/Engine/Source/Runtime/Experimental/GeometryCollectionEngine/Private/GeometryCollection/GeometryCollectionComponent.cpp
+++ b/Engine/Source/Runtime/Experimental/GeometryCollectionEngine/Private/GeometryCollection/GeometryCollectionComponent.cpp
@@ -598,6 +598,7 @@ UGeometryCollectionComponent::UGeometryCollectionComponent(const FObjectInitiali
 	, bIsMoving(false)
 	, bIsRootBroken(false)
 	, bUpdateCustomRenderer(true)
+	, bUpdateCustomRendererOnPostPhysicsSync(true)
 {
 	// by default tick is registered but disabled, we only need it when we need to update the removal timers
 	// tick will be then enabled only when the root is broken from OnPostPhysicsSync callback
@@ -4028,7 +4029,10 @@ void UGeometryCollectionComponent::UpdateRenderSystemsIfNeeded(bool bDynamicColl
 		// #todo review: When we've made changes to ISMC, we need to move this function call to SetRenderDynamicData_Concurrent
 		RefreshEmbeddedGeometry();
 
-		RefreshCustomRenderer();
+		if (bUpdateCustomRendererOnPostPhysicsSync)
+		{
+			RefreshCustomRenderer();
+		}
 
 		if (SceneProxy && SceneProxy->IsNaniteMesh())
 		{
diff --git a/Engine/Source/Runtime/Experimental/GeometryCollectionEngine/Public/GeometryCollection/GeometryCollectionComponent.h b/Engine/Source/Runtime/Experimental/GeometryCollectionEngine/Public/GeometryCollection/GeometryCollectionComponent.h
--- a/Engine/Source/Runtime/Experimental/GeometryCollectionEngine/Public/GeometryCollection/GeometryCollectionComponent.h
+++ b/Engine/Source/Runtime/Experimental/GeometryCollectionEngine/Public/GeometryCollection/GeometryCollectionComponent.h
@@ -1290,9 +1290,12 @@ public:
 	/** Force all GC components to reregister their custom renderer objects. */
 	static GEOMETRYCOLLECTIONENGINE_API void ReregisterAllCustomRenderers();
 
-	/** allow update of the custom renderer if enabled */
+	/** allow update of the custom renderer ( valid if custom redner is being used ) - true by default */
 	GEOMETRYCOLLECTIONENGINE_API void SetUpdateCustomRenderer(bool bValue) { bUpdateCustomRenderer = bValue; }
 
+	/** update of the custom renderer when post physics sync callback is executing ( valid if custom redner is being used ) - true by default */
+	GEOMETRYCOLLECTIONENGINE_API void SetUpdateCustomRendererOnPostPhysicsSync(bool bValue) { bUpdateCustomRendererOnPostPhysicsSync = bValue; }
+
 	GEOMETRYCOLLECTIONENGINE_API bool ShouldUpdateComponentTransformToRootBone() const { return bUpdateComponentTransformToRootBone; }
 
 public:
@@ -1729,6 +1732,8 @@ private:
 
 	bool bUpdateCustomRenderer;
 
+	bool bUpdateCustomRendererOnPostPhysicsSync;
+
 private:
 
 	enum class ENetAwakeningMode
