Iris
- Fixed disconnect issue on clients due to server sending more data to not yet acked hugeobjects.

diff --git a/Engine/Source/Runtime/Experimental/Iris/Core/Private/Iris/ReplicationSystem/ReplicationReader.cpp b/Engine/Source/Runtime/Experimental/Iris/Core/Private/Iris/ReplicationSystem/ReplicationReader.cpp
--- a/Engine/Source/Runtime/Experimental/Iris/Core/Private/Iris/ReplicationSystem/ReplicationReader.cpp
+++ b/Engine/Source/Runtime/Experimental/Iris/Core/Private/Iris/ReplicationSystem/ReplicationReader.cpp
@@ -1675,6 +1675,8 @@ void FReplicationReader::ProcessHugeObjectAttachment(FNetSerializationContext& C
 		return;
 	}
 
+	IRIS_PROFILER_SCOPE(FReplicationReader_ProcessHugeObjectAttachment)
+
 	FNetTraceCollector* HugeObjectTraceCollector = nullptr;
 #if UE_NET_TRACE_ENABLED
 	FNetTraceCollector HugeObjectTraceCollectorOnStack;
diff --git a/Engine/Source/Runtime/Experimental/Iris/Core/Private/Iris/ReplicationSystem/ReplicationWriter.cpp b/Engine/Source/Runtime/Experimental/Iris/Core/Private/Iris/ReplicationSystem/ReplicationWriter.cpp
--- a/Engine/Source/Runtime/Experimental/Iris/Core/Private/Iris/ReplicationSystem/ReplicationWriter.cpp
+++ b/Engine/Source/Runtime/Experimental/Iris/Core/Private/Iris/ReplicationSystem/ReplicationWriter.cpp
@@ -3622,7 +3622,8 @@ void FReplicationWriter::FHugeObjectSendQueue::AckObjects(TFunctionRef<void (con
 
 		if (!bObjectIsAcked)
 		{
-			continue;
+			// As clients deliver hugeobjects parts in order we cannot ack later objects until previous ones have been fully acked.
+			break;
 		}
 
 		AckHugeObject(Context);
