Lazy initializer for Dynamic removals attributes

diff --git a/Engine/Source/Runtime/Experimental/GeometryCollectionEngine/Private/GeometryCollection/GeometryCollectionComponent.cpp b/Engine/Source/Runtime/Experimental/GeometryCollectionEngine/Private/GeometryCollection/GeometryCollectionComponent.cpp
--- a/Engine/Source/Runtime/Experimental/GeometryCollectionEngine/Private/GeometryCollection/GeometryCollectionComponent.cpp
+++ b/Engine/Source/Runtime/Experimental/GeometryCollectionEngine/Private/GeometryCollection/GeometryCollectionComponent.cpp
@@ -581,6 +581,7 @@ UGeometryCollectionComponent::UGeometryCollectionComponent(const FObjectInitiali
 	, ReplicationAbandonClusterLevel_DEPRECATED(0)
 	, ReplicationAbandonAfterLevel(0)
 	, ReplicationMaxPositionAndVelocityCorrectionLevel(100)
+	, bInitializedRemovalDynamicAttribute(false)
 	, bRenderStateDirty(true)
 	, bEnableBoneSelection(false)
 	, ViewLevel(-1)
@@ -3366,6 +3367,7 @@ void UGeometryCollectionComponent::TickComponent(float DeltaTime, enum ELevelTic
 	// todo(chaos) : cache root broken state ? 
 	if (IsRootBroken())
 	{
+		InitializeRemovalDynamicAttributesIfNeeded();
 		const float AdjustedDeltaTime = FMath::Max(GeometryCollectionRemovalMultiplier, 0.0001f) * DeltaTime;
 		// todo(chaos) : move removal logic on the physics thread
 		IncrementSleepTimer(AdjustedDeltaTime);
@@ -3524,6 +3526,38 @@ void UGeometryCollectionComponent::ReregisterAllCustomRenderers()
 	}
 }
 
+
+void UGeometryCollectionComponent::InitializeRemovalDynamicAttributesIfNeeded()
+{
+	if (bInitializedRemovalDynamicAttribute == false)
+	{
+		FGeometryCollectionDecayDynamicFacade DecayDynamicFacade(*DynamicCollection);
+
+		// we are not testing for bAllowRemovalOnSleep, so that we can enable it at runtime if necessary
+		if (RestCollection->bRemoveOnMaxSleep)
+		{
+			DecayDynamicFacade.AddAttributes();
+
+			FGeometryCollectionRemoveOnSleepDynamicFacade RemoveOnSleepDynamicFacade(*DynamicCollection);
+			RemoveOnSleepDynamicFacade.DefineSchema();
+			RemoveOnSleepDynamicFacade.SetAttributeValues(RestCollection->MaximumSleepTime, RestCollection->RemovalDuration);
+		}
+
+		// Remove on break feature related dynamic attribute arrays
+		// we are not testing for bAllowRemovalOnBreak, so that we can enable it at runtime if necessary
+		GeometryCollection::Facades::FCollectionRemoveOnBreakFacade RemoveOnBreakFacade(*RestCollection->GetGeometryCollection());
+		if (RemoveOnBreakFacade.IsValid())
+		{
+			DecayDynamicFacade.AddAttributes();
+
+			FGeometryCollectionRemoveOnBreakDynamicFacade RemoveOnBreakDynamicFacade(*DynamicCollection);
+			RemoveOnBreakDynamicFacade.DefineSchema();
+			RemoveOnBreakDynamicFacade.SetAttributeValues(RemoveOnBreakFacade);
+		}
+		bInitializedRemovalDynamicAttribute = true;
+	}
+}
+
 void UGeometryCollectionComponent::ResetDynamicCollection()
 {
 	bool bCreateDynamicCollection = true;
@@ -3554,30 +3588,6 @@ void UGeometryCollectionComponent::ResetDynamicCollection()
 		}
 		DynamicCollection->AddAttribute<uint8>("InternalClusterParentTypeArray", FTransformCollection::TransformGroup);
 
-		FGeometryCollectionDecayDynamicFacade DecayDynamicFacade(*DynamicCollection);
-		
-		// we are not testing for bAllowRemovalOnSleep, so that we can enable it at runtime if necessary
-		if (RestCollection->bRemoveOnMaxSleep)
-		{
-			DecayDynamicFacade.AddAttributes();
-
-			FGeometryCollectionRemoveOnSleepDynamicFacade RemoveOnSleepDynamicFacade(*DynamicCollection);
-			RemoveOnSleepDynamicFacade.DefineSchema();
-			RemoveOnSleepDynamicFacade.SetAttributeValues(RestCollection->MaximumSleepTime, RestCollection->RemovalDuration);
-		}
-		
-		// Remove on break feature related dynamic attribute arrays
-		// we are not testing for bAllowRemovalOnBreak, so that we can enable it at runtime if necessary
-		GeometryCollection::Facades::FCollectionRemoveOnBreakFacade RemoveOnBreakFacade(*RestCollection->GetGeometryCollection());
-		if (RemoveOnBreakFacade.IsValid())
-		{
-			DecayDynamicFacade.AddAttributes();
-
-			FGeometryCollectionRemoveOnBreakDynamicFacade RemoveOnBreakDynamicFacade(*DynamicCollection);
-			RemoveOnBreakDynamicFacade.DefineSchema();
-			RemoveOnBreakDynamicFacade.SetAttributeValues(RemoveOnBreakFacade);
-		}
-
 		DynamicCollection->MakeDirty();
 		MarkRenderStateDirty();
 		MarkRenderDynamicDataDirty();
@@ -3604,7 +3614,7 @@ void UGeometryCollectionComponent::ResetDynamicCollection()
 
 	RootSpaceBounds.Init();
 	UpdateCachedBounds();
-
+	bInitializedRemovalDynamicAttribute = false;
 	bIsRootBroken = false;
 }
 
diff --git a/Engine/Source/Runtime/Experimental/GeometryCollectionEngine/Public/GeometryCollection/GeometryCollectionComponent.h b/Engine/Source/Runtime/Experimental/GeometryCollectionEngine/Public/GeometryCollection/GeometryCollectionComponent.h
--- a/Engine/Source/Runtime/Experimental/GeometryCollectionEngine/Public/GeometryCollection/GeometryCollectionComponent.h
+++ b/Engine/Source/Runtime/Experimental/GeometryCollectionEngine/Public/GeometryCollection/GeometryCollectionComponent.h
@@ -1546,6 +1546,7 @@ private:
 	void ProcessRepDataOnPT();
 	void ProcessRepStateDataOnPT();
 	void ProcessRepDynamicDataOnPT();
+	void InitializeRemovalDynamicAttributesIfNeeded();
 
 	// called when the rest transform are updated from SetRestState / ResetRestTransforms
 	// this updates only the renderer, the dynamic collection should be initialized when calling this function
@@ -1556,6 +1557,7 @@ private:
 	FTransform3f GetCurrentTransform(int32 Index) const;
 	void ComputeCurrentGlobalsMatrices(TArray<FTransform3f>& OutTransforms) const;
 
+	bool bInitializedRemovalDynamicAttribute;
 	bool bRenderStateDirty;
 	bool bEnableBoneSelection;
 	int ViewLevel;
