IAS analytics will now only send the data tracked since the last analytics payload was sent.

#rnx

- Previously the IAS analytics system was sending a running total for each counter each time an analytics event was raised which over time and multiple events would start to weight the data heavily towards the stats gathered earlier on in the process's life span.
- Some of the counters are used by other systems so we cannot simply reset them every time the analytics event is called or those stat gathering systems will start showing incorrect data.
- These counters need to record the previous value that was sent and then only send a delta.
- Counters that are only used for analytic events will be reset.
- FRunningTotals can track counters over time via their name, along with the TRACK_DELTA this allows us to avoid changing the formatting of the analytics code too much and keeps it looking like analytics code found elsewhere in the engine.

diff --git a/Engine/Source/Runtime/Experimental/IoStoreOnDemand/Private/Statistics.cpp b/Engine/Source/Runtime/Experimental/IoStoreOnDemand/Private/Statistics.cpp
--- a/Engine/Source/Runtime/Experimental/IoStoreOnDemand/Private/Statistics.cpp
+++ b/Engine/Source/Runtime/Experimental/IoStoreOnDemand/Private/Statistics.cpp
@@ -176,6 +176,53 @@ private:
 	double VarianceAccumulator;
 };
 
+class FDeltaTracking
+{
+public:
+	int64 Get(FStringView Name, int64 Value)
+	{
+		if (int64* PrevValue = IntTotals.FindByHash(GetTypeHash(Name), Name))
+		{
+			const int64 Delta = Value - *PrevValue;
+			*PrevValue = Value;
+
+			return Delta;
+		}
+		else
+		{
+			IntTotals.Add(FString(Name), 0);
+			return Value;
+		}
+	}
+
+	uint32 Get(FStringView Name, uint32 Value)
+	{
+		return static_cast<uint32>(Get(Name, static_cast<int64>(Value)));
+	}
+
+	double Get(FStringView Name, double Value)
+	{
+		if (double* PrevValue = RealTotals.FindByHash(GetTypeHash(Name), Name))
+		{
+			const double Delta = Value - *PrevValue;
+			*PrevValue = Value;
+
+			return Delta;
+		}
+		else
+		{
+			RealTotals.Add(FString(Name), 0.0);
+			return Value;
+		}
+	}
+
+private:
+
+	TMap<FString, int64> IntTotals;
+	TMap<FString, double> RealTotals;
+
+} static GDeltaTracking;
+
 ////////////////////////////////////////////////////////////////////////////////
 // TRACE STATS
 
@@ -380,13 +427,18 @@ FOnDemandIoBackendStats* FOnDemandIoBackendStats::Get()
 
 void FOnDemandIoBackendStats::ReportAnalytics(TArray<FAnalyticsEventAttribute>& OutAnalyticsArray) const
 {
+	// We use this macro with counters that are used elsewhere meaning we can't just reset them
+	// each time we send an analytics payload. The macro will track the running total and only
+	// pass on the delta to the analytics array since the last time this method was called.
+#define TRACK_DELTA(Name, Value) TEXT(Name), GDeltaTracking.Get(TEXTVIEW(Name), Value)
+
 	if (GIasReportHttpAnalyticsEnabled)
 	{
 		AppendAnalyticsEventAttributeArray(OutAnalyticsArray
-			,TEXT("IasHttpErrorCount"), GHttpErrorCount.Get()
-			,TEXT("IasHttpRetryCount"), GHttpRetryCount.Get()
-			,TEXT("IasHttpGetCount"), GHttpGetCount.Get()
-			,TEXT("IasHttpDownloadedBytes"), GHttpDownloadedBytes.Get()
+			,TRACK_DELTA("IasHttpErrorCount", GHttpErrorCount.Get())
+			,TRACK_DELTA("IasHttpRetryCount", GHttpRetryCount.Get())
+			,TRACK_DELTA("IasHttpGetCount", GHttpGetCount.Get())
+			,TRACK_DELTA("IasHttpDownloadedBytes", GHttpDownloadedBytes.Get())
 			,TEXT("IasHttpDurationMeanAvg"), GHttpAvgDuration.GetMean()
 			,TEXT("IasHttpDurationStdDev"), GHttpAvgDuration.GetDeviation()
 
@@ -396,21 +448,27 @@ void FOnDemandIoBackendStats::ReportAnalytics(TArray<FAnalyticsEventAttribute>&
 			,TEXT("IasHttpDuration3"), GHttpDurationBuckets[3]
 			,TEXT("IasHttpDuration4"), GHttpDurationBuckets[4]
 		);
+
+		// These values we can just reset as they are only being used with analytics
+		GHttpAvgDuration.Reset();
+		FMemory::Memzero(GHttpDurationBuckets);
 	}
 
 	if (GIasReportCacheAnalyticsEnabled)
 	{
 		AppendAnalyticsEventAttributeArray(OutAnalyticsArray
 
-			,TEXT("IasCacheErrorCount"), GCacheErrorCount.Get()
+			,TRACK_DELTA("IasCacheErrorCount", GCacheErrorCount.Get())
 
-			,TEXT("IasCacheCachedBytes"), GCacheCachedBytes.Get()
-			,TEXT("IasCacheMaxBytes"), GCacheMaxBytes
+			,TRACK_DELTA("IasCacheCachedBytes", GCacheCachedBytes.Get())
+			,TRACK_DELTA("IasCacheMaxBytes", GCacheMaxBytes)
 
-			,TEXT("IasCacheReadBytes"), GCacheReadBytes.Get()
-			,TEXT("IasCacheRejectBytes"), GCacheRejectBytes.Get()
+			,TRACK_DELTA("IasCacheReadBytes", GCacheReadBytes.Get())
+			,TRACK_DELTA("IasCacheRejectBytes", GCacheRejectBytes.Get())
 		);
 	}
+
+#undef TRACK_DELTA
 }
 
 void FOnDemandIoBackendStats::OnIoRequestEnqueue()
