Geometry collection optimization : Remove findAttribute calls in when marshaling data between PT and GT

diff --git a/Engine/Source/Runtime/Experimental/Chaos/Private/GeometryCollectionProxyData.cpp b/Engine/Source/Runtime/Experimental/Chaos/Private/GeometryCollectionProxyData.cpp
--- a/Engine/Source/Runtime/Experimental/Chaos/Private/GeometryCollectionProxyData.cpp
+++ b/Engine/Source/Runtime/Experimental/Chaos/Private/GeometryCollectionProxyData.cpp
@@ -118,6 +118,7 @@ const FName FGeometryDynamicCollection::ShapesSimDataAttribute("ShapesSimData");
 const FName FGeometryDynamicCollection::SimplicialsAttribute("CollisionParticles");
 const FName FGeometryDynamicCollection::SimulatableParticlesAttribute("SimulatableParticlesAttribute");
 const FName FGeometryDynamicCollection::SharedImplicitsAttribute("SharedImplicits");
+const FName FGeometryDynamicCollection::InternalClusterParentTypeAttribute("InternalClusterParentTypeArray");
 
 // Deprecated
 const FName FGeometryDynamicCollection::CollisionMaskAttribute("CollisionMask");
@@ -125,13 +126,16 @@ const FName FGeometryDynamicCollection::CollisionGroupAttribute("CollisionGroup"
 
 FGeometryDynamicCollection::FGeometryDynamicCollection(const FGeometryCollection* InRestCollection)
 	: FTransformDynamicCollection(InRestCollection)
+	, OptionalLinearVelocityAttribute(nullptr)
+	, OptionalAngularVelocityAttribute(nullptr)
+	, OptionalAnimateTransformAttribute(nullptr)
 {
 	// Transform Group
 	AddExternalAttribute<bool>(FGeometryDynamicCollection::ActiveAttribute, FTransformCollection::TransformGroup, Active);
 	AddExternalAttribute<uint8>(FGeometryDynamicCollection::DynamicStateAttribute, FTransformCollection::TransformGroup, DynamicState);
 	AddExternalAttribute(SimplicialsAttribute, FTransformCollection::TransformGroup, Simplicials);
 	AddExternalAttribute(SimulatableParticlesAttribute, FGeometryCollection::TransformGroup, SimulatableParticles);
-
+	AddExternalAttribute(InternalClusterParentTypeAttribute, FGeometryCollection::TransformGroup, InternalClusterParentType);
 }
 
 const TManagedArrayAccessor<int32> FGeometryDynamicCollection::GetInitialLevels() const
@@ -140,6 +144,32 @@ const TManagedArrayAccessor<int32> FGeometryDynamicCollection::GetInitialLevels(
 	return TManagedArrayAccessor<int32>(*RestCollection, LevelAttributeName, FGeometryCollection::TransformGroup);
 }
 
+void FGeometryDynamicCollection::AddVelocitiesAttributes()
+{
+	if (OptionalLinearVelocityAttribute == nullptr && OptionalAngularVelocityAttribute == nullptr)
+	{
+		static const FName LinearVelocityAttributeName = "LinearVelocity";
+		static const FName AngularVelocityAttributeName = "AngularVelocity";
+
+		OptionalLinearVelocityAttribute = &AddAttribute<FVector3f>(LinearVelocityAttributeName, FTransformCollection::TransformGroup);
+		OptionalAngularVelocityAttribute = &AddAttribute<FVector3f>(AngularVelocityAttributeName, FTransformCollection::TransformGroup);
+	}
+}
+
+void FGeometryDynamicCollection::AddAnimateTransformAttribute()
+{
+	if (OptionalAnimateTransformAttribute == nullptr)
+	{
+		static const FName AnimateTransformAttributeName = "AnimateTransformAttribute";
+
+		OptionalAnimateTransformAttribute = &AddAttribute<bool>("AnimateTransformAttribute", FGeometryCollection::TransformGroup);
+		if (OptionalAnimateTransformAttribute)
+		{
+			OptionalAnimateTransformAttribute->Fill(false);
+		}
+	}
+}
+
 void FGeometryDynamicCollection::CopyInitialVelocityAttributesFrom(const FGeometryDynamicCollection& SourceCollection)
 {
 	FInitialVelocityFacade InitialVelocityFacade(*this);
diff --git a/Engine/Source/Runtime/Experimental/Chaos/Private/PhysicsProxy/GeometryCollectionPhysicsProxy.cpp b/Engine/Source/Runtime/Experimental/Chaos/Private/PhysicsProxy/GeometryCollectionPhysicsProxy.cpp
--- a/Engine/Source/Runtime/Experimental/Chaos/Private/PhysicsProxy/GeometryCollectionPhysicsProxy.cpp
+++ b/Engine/Source/Runtime/Experimental/Chaos/Private/PhysicsProxy/GeometryCollectionPhysicsProxy.cpp
@@ -60,10 +60,6 @@
 
 namespace
 {
-	const FName LinearVelocityAttributeName = "LinearVelocity";
-	const FName AngularVelocitiesAttributeName = "AngularVelocity";
-	const FName InternalClusterParentTypeArrayAttributeName = "InternalClusterParentTypeArray";
-	const FName AnimateTransformAttributeName = "AnimateTransformAttribute";
 	const FName MassToLocalAttributeName = "MassToLocal";
 	const FName MassAttributeName = "Mass";
 	const FName InertiaTensorAttributeName = "InertiaTensor";
@@ -3737,10 +3733,10 @@ bool FGeometryCollectionPhysicsProxy::PullNonInterpolatableDataFromSinglePhysics
 
 	bool bIsCollectionDirty = false;
 
-	TManagedArray<FVector3f>* LinearVelocities = GameThreadCollection.FindAttributeTyped<FVector3f>(LinearVelocityAttributeName, FTransformCollection::TransformGroup);
-	TManagedArray<FVector3f>* AngularVelocities = GameThreadCollection.FindAttributeTyped<FVector3f>(AngularVelocitiesAttributeName, FTransformCollection::TransformGroup);
-	TManagedArray<uint8>* InternalClusterParentTypeArray = GameThreadCollection.FindAttributeTyped<uint8>(InternalClusterParentTypeArrayAttributeName, FTransformCollection::TransformGroup);
-	const TManagedArray<bool>* AnimationsActive = GameThreadCollection.FindAttribute<bool>(AnimateTransformAttributeName, FGeometryCollection::TransformGroup);
+	TManagedArray<FVector3f>* LinearVelocities = GameThreadCollection.GetLinearVelocitiesAttribute();
+	TManagedArray<FVector3f>* AngularVelocities = GameThreadCollection.GetAngularVelocitiesAttribute();
+	TManagedArray<uint8>& InternalClusterParentTypeArray = GameThreadCollection.GetInternalClusterParentTypeAttribute();
+	const TManagedArray<bool>* AnimationsActive = GameThreadCollection.GetAnimateTransformAttribute();
 
 	// first step : process the non values that do not need to be interpolated 
 	for (int32 EntryIndex = 0; EntryIndex < CurrentResults.GetNumEntries(); EntryIndex++)
@@ -3779,26 +3775,23 @@ bool FGeometryCollectionPhysicsProxy::PullNonInterpolatableDataFromSinglePhysics
 			bIsCollectionDirty = true;
 		}
 
-		if (InternalClusterParentTypeArray)
+		Chaos::EInternalClusterType ParentType = Chaos::EInternalClusterType::None;
+		if (StateData.State.HasInternalClusterParent != 0)
 		{
-			Chaos::EInternalClusterType ParentType = Chaos::EInternalClusterType::None;
-			if (StateData.State.HasInternalClusterParent != 0)
+			if (StateData.State.HasClusterUnionParent)
 			{
-				if (StateData.State.HasClusterUnionParent)
-				{
-					ParentType = Chaos::EInternalClusterType::ClusterUnion;
-				}
-				else
-				{
-					ParentType = (StateData.State.DynamicInternalClusterParent != 0) ? Chaos::EInternalClusterType::Dynamic : Chaos::EInternalClusterType::KinematicOrStatic;
-				}
+				ParentType = Chaos::EInternalClusterType::ClusterUnion;
 			}
-			const uint8 ParentTypeUInt8 = static_cast<uint8>(ParentType);
-			if (UpdateValue((*InternalClusterParentTypeArray)[TransformGroupIndex], ParentTypeUInt8))
+			else
 			{
-				bIsCollectionDirty = true;
+				ParentType = (StateData.State.DynamicInternalClusterParent != 0) ? Chaos::EInternalClusterType::Dynamic : Chaos::EInternalClusterType::KinematicOrStatic;
 			}
 		}
+		const uint8 ParentTypeUInt8 = static_cast<uint8>(ParentType);
+		if (UpdateValue(InternalClusterParentTypeArray[TransformGroupIndex], ParentTypeUInt8))
+		{
+			bIsCollectionDirty = true;
+		}
 
 		// if interpolation is off , we need to apply XR, VW and transforms
 		if (bForcePullXRVW)
@@ -3879,7 +3872,7 @@ bool FGeometryCollectionPhysicsProxy::PullFromPhysicsState(const Chaos::FDirtyGe
 		}
 		bIsCollectionDirty |= PullNonInterpolatableDataFromSinglePhysicsState(PullData, !bNeedInterpolation, NextPullData ? &NextPullData->Results().GetModifiedTransformIndices() : nullptr);
 
-		const TManagedArray<bool>* AnimationsActive = GameThreadCollection.FindAttribute<bool>(AnimateTransformAttributeName, FGeometryCollection::TransformGroup);
+		const TManagedArray<bool>* AnimationsActive = GameThreadCollection.GetAnimateTransformAttribute();
 		const TManagedArray<FTransform>& MassToLocal = Parameters.RestCollection->GetAttribute<FTransform>(MassToLocalAttributeName, FTransformCollection::TransformGroup);
 
 		// second : interpolate-able ones
@@ -3890,8 +3883,8 @@ bool FGeometryCollectionPhysicsProxy::PullFromPhysicsState(const Chaos::FDirtyGe
 			const FGeometryCollectionResults& PrevResults = PullData.Results();
 			const FGeometryCollectionResults& NextResults = NextPullData->Results();
 
-			TManagedArray<FVector3f>* LinearVelocities = GameThreadCollection.FindAttributeTyped<FVector3f>(LinearVelocityAttributeName, FTransformCollection::TransformGroup);
-			TManagedArray<FVector3f>* AngularVelocities = GameThreadCollection.FindAttributeTyped<FVector3f>(AngularVelocitiesAttributeName, FTransformCollection::TransformGroup);
+			TManagedArray<FVector3f>* LinearVelocities = GameThreadCollection.GetLinearVelocitiesAttribute();
+			TManagedArray<FVector3f>* AngularVelocities = GameThreadCollection.GetAngularVelocitiesAttribute();
 
 			// for that case we cannot just go through the list of entries since Results and NextResults may have different number of entries that don't always match
 			// so we need to go through the transform indices and find the matching entries on both side 
diff --git a/Engine/Source/Runtime/Experimental/Chaos/Public/GeometryCollectionProxyData.h b/Engine/Source/Runtime/Experimental/Chaos/Public/GeometryCollectionProxyData.h
--- a/Engine/Source/Runtime/Experimental/Chaos/Public/GeometryCollectionProxyData.h
+++ b/Engine/Source/Runtime/Experimental/Chaos/Public/GeometryCollectionProxyData.h
@@ -96,6 +96,7 @@ public:
 	static CHAOS_API const FName SharedImplicitsAttribute;
 	static CHAOS_API const FName SimplicialsAttribute;
 	static CHAOS_API const FName SimulatableParticlesAttribute;
+	static CHAOS_API const FName InternalClusterParentTypeAttribute;
 
 	UE_DEPRECATED(5.4, "CollisionMaskAttribute is no longer supported")
 	static CHAOS_API const FName CollisionMaskAttribute;
@@ -125,6 +126,19 @@ public:
 
 	CHAOS_API const TManagedArrayAccessor<int32> GetInitialLevels() const;
 
+	CHAOS_API const TManagedArray<uint8>& GetInternalClusterParentTypeAttribute() const { return InternalClusterParentType; }
+	CHAOS_API TManagedArray<uint8>& GetInternalClusterParentTypeAttribute() { return InternalClusterParentType; }
+
+	CHAOS_API void AddVelocitiesAttributes();
+	CHAOS_API const TManagedArray<FVector3f>* GetLinearVelocitiesAttribute() const { return OptionalLinearVelocityAttribute; }
+	CHAOS_API const TManagedArray<FVector3f>* GetAngularVelocitiesAttribute() const { return OptionalAngularVelocityAttribute; };
+	CHAOS_API TManagedArray<FVector3f>* GetLinearVelocitiesAttribute() { return OptionalLinearVelocityAttribute; }
+	CHAOS_API TManagedArray<FVector3f>* GetAngularVelocitiesAttribute() { return OptionalAngularVelocityAttribute; };
+
+	CHAOS_API void AddAnimateTransformAttribute();
+	CHAOS_API const TManagedArray<bool>* GetAnimateTransformAttribute() const { return OptionalAnimateTransformAttribute; }
+	CHAOS_API TManagedArray<bool>* GetAnimateTransformAttribute() { return OptionalAnimateTransformAttribute; }
+
 	struct FInitialVelocityFacade
 	{
 		FInitialVelocityFacade(FGeometryDynamicCollection& DynamicCollection);
@@ -143,6 +157,13 @@ public:
 	FInitialVelocityFacade GetInitialVelocityFacade() const { return FInitialVelocityFacade(*this); }
 
 	CHAOS_API void CopyInitialVelocityAttributesFrom(const FGeometryDynamicCollection& SourceCollection);
+
+private:
+	TManagedArray<uint8>	  InternalClusterParentType;
+
+	TManagedArray<FVector3f>* OptionalLinearVelocityAttribute;
+	TManagedArray<FVector3f>* OptionalAngularVelocityAttribute;
+	TManagedArray<bool>*	  OptionalAnimateTransformAttribute;
 };
 
 /**
