Add WPO evaluate flag for ISMPool.
Setting this can allow us to take the optimized path on nanite.

diff --git a/Engine/Source/Runtime/Experimental/GeometryCollectionEngine/Private/GeometryCollection/GeometryCollectionISMPoolComponent.cpp b/Engine/Source/Runtime/Experimental/GeometryCollectionEngine/Private/GeometryCollection/GeometryCollectionISMPoolComponent.cpp
--- a/Engine/Source/Runtime/Experimental/GeometryCollectionEngine/Private/GeometryCollection/GeometryCollectionISMPoolComponent.cpp
+++ b/Engine/Source/Runtime/Experimental/GeometryCollectionEngine/Private/GeometryCollection/GeometryCollectionISMPoolComponent.cpp
@@ -118,6 +118,7 @@ void FGeometryCollectionISM::InitISM(const FGeometryCollectionStaticMeshInstance
 	ISMComponent->bAffectDynamicIndirectLighting = (MeshInstance.Desc.Flags & FISMComponentDescription::AffectDynamicIndirectLighting) != 0;
 	ISMComponent->bAffectDistanceFieldLighting = (MeshInstance.Desc.Flags & FISMComponentDescription::AffectDistanceFieldLighting) != 0;
 	ISMComponent->bWorldPositionOffsetWritesVelocity = (MeshInstance.Desc.Flags & FISMComponentDescription::WorldPositionOffsetWritesVelocity) != 0;
+	ISMComponent->bEvaluateWorldPositionOffset = (MeshInstance.Desc.Flags & FISMComponentDescription::EvaluateWorldPositionOffset) != 0;
 	ISMComponent->bUseGpuLodSelection = (MeshInstance.Desc.Flags & FISMComponentDescription::GpuLodSelection) != 0;
 	ISMComponent->bOverrideMinLOD = MeshInstance.Desc.MinLod > 0;
 	ISMComponent->MinLOD = MeshInstance.Desc.MinLod;
@@ -312,18 +313,18 @@ void FGeometryCollectionISMPool::RemoveISM(const FGeometryCollectionMeshInfo& Me
 			ISM.InstanceGroups.Reset();
 			ISM.InstanceIds.Reset();
 		}
-	
+
 		if (GUseComponentFreeList && ISM.InstanceGroups.IsEmpty())
-			{
-				// Remove component and push this ISM slot to the free list.
+		{
+			// Remove component and push this ISM slot to the free list.
 			ensure(ISM.ISMComponent->PerInstanceSMData.Num() == 0);
-				MeshToISMIndex.Remove(ISM.MeshInstance);
-				FreeListISM.Add(MeshInfo.ISMIndex);
+			MeshToISMIndex.Remove(ISM.MeshInstance);
+			FreeListISM.Add(MeshInfo.ISMIndex);
 
 #if WITH_EDITOR
-				ISM.ISMComponent->Rename(nullptr);
+			ISM.ISMComponent->Rename(nullptr);
 #endif
-			}
+		}
 	}
 }
 
@@ -354,7 +355,7 @@ void FGeometryCollectionISMPool::GarbageCollect()
 	{
 		const int32 ISMIndex = FreeListISM.Last();
 		FreeListISM.RemoveAt(FreeListISM.Num() - 1);
-		
+
 		UInstancedStaticMeshComponent* ISM = ISMs[ISMIndex].ISMComponent;
 		ISM->UnregisterComponent();
 		ISM->DestroyComponent();
diff --git a/Engine/Source/Runtime/Experimental/GeometryCollectionEngine/Public/GeometryCollection/GeometryCollectionISMPoolComponent.h b/Engine/Source/Runtime/Experimental/GeometryCollectionEngine/Public/GeometryCollection/GeometryCollectionISMPoolComponent.h
--- a/Engine/Source/Runtime/Experimental/GeometryCollectionEngine/Public/GeometryCollection/GeometryCollectionISMPoolComponent.h
+++ b/Engine/Source/Runtime/Experimental/GeometryCollectionEngine/Public/GeometryCollection/GeometryCollectionISMPoolComponent.h
@@ -107,13 +107,14 @@ struct FISMComponentDescription
 		ReverseCulling = 1 << 3,
 		StaticMobility = 1 << 4,
 		WorldPositionOffsetWritesVelocity = 1 << 5,
-		AffectShadow = 1 << 6,
-		AffectDistanceFieldLighting = 1 << 7,
-		AffectDynamicIndirectLighting = 1 << 8,
-		DistanceCullPrimitive = 1 << 9,
+		EvaluateWorldPositionOffset = 1 << 6,
+		AffectShadow = 1 << 7,
+		AffectDistanceFieldLighting = 1 << 8,
+		AffectDynamicIndirectLighting = 1 << 9,
+		DistanceCullPrimitive = 1 << 10,
 	};
 
-	uint32 Flags = WorldPositionOffsetWritesVelocity|AffectShadow;
+	uint32 Flags = WorldPositionOffsetWritesVelocity|EvaluateWorldPositionOffset|AffectShadow;
 	int32 NumCustomDataFloats = 0;
 	FVector Position = FVector::ZeroVector;
 	int32 StartCullDistance = 0;
