Geometry collection optimization : add missing change from 28889155

diff --git a/Engine/Source/Runtime/Experimental/GeometryCollectionEngine/Private/GeometryCollection/GeometryCollectionComponent.cpp b/Engine/Source/Runtime/Experimental/GeometryCollectionEngine/Private/GeometryCollection/GeometryCollectionComponent.cpp
--- a/Engine/Source/Runtime/Experimental/GeometryCollectionEngine/Private/GeometryCollection/GeometryCollectionComponent.cpp
+++ b/Engine/Source/Runtime/Experimental/GeometryCollectionEngine/Private/GeometryCollection/GeometryCollectionComponent.cpp
@@ -3583,10 +3583,8 @@ void UGeometryCollectionComponent::ResetDynamicCollection()
 
 		if (bStoreVelocities || bNotifyTrailing)
 		{
-			DynamicCollection->AddAttribute<FVector3f>("LinearVelocity", FTransformCollection::TransformGroup);
-			DynamicCollection->AddAttribute<FVector3f>("AngularVelocity", FTransformCollection::TransformGroup);
+			DynamicCollection->AddVelocitiesAttributes();
 		}
-		DynamicCollection->AddAttribute<uint8>("InternalClusterParentTypeArray", FTransformCollection::TransformGroup);
 
 		DynamicCollection->MakeDirty();
 		MarkRenderStateDirty();
