Dataflow - Fixed a crash when duplicating a node that has a UObject pointer property.

diff --git a/Engine/Source/Runtime/Experimental/Dataflow/Core/Private/Dataflow/DataflowNode.cpp b/Engine/Source/Runtime/Experimental/Dataflow/Core/Private/Dataflow/DataflowNode.cpp
--- a/Engine/Source/Runtime/Experimental/Dataflow/Core/Private/Dataflow/DataflowNode.cpp
+++ b/Engine/Source/Runtime/Experimental/Dataflow/Core/Private/Dataflow/DataflowNode.cpp
@@ -5,8 +5,8 @@
 #include "ChaosLog.h"
 #include "Dataflow/DataflowInputOutput.h"
 #include "Dataflow/DataflowArchive.h"
-#include "Serialization/MemoryWriter.h"
-#include "Serialization/MemoryReader.h"
+#include "Serialization/ObjectWriter.h"
+#include "Serialization/ObjectReader.h"
 
 #include UE_INLINE_GENERATED_CPP_BY_NAME(DataflowNode)
 
@@ -696,10 +696,10 @@ void FDataflowNode::CopyNodeProperties(const TSharedPtr<FDataflowNode> CopyFromD
 {
 	TArray<uint8> NodeData;
 
-	FMemoryWriter ArWriter(NodeData);
+	FObjectWriter ArWriter(NodeData);
 	CopyFromDataflowNode->SerializeInternal(ArWriter);
 
-	FMemoryReader ArReader(NodeData);
+	FObjectReader ArReader(NodeData);
 	this->SerializeInternal(ArReader);
 }
 
