Chaos - option to treat collisions between two one-way interaction particles as spheres (currently disable via cvar)
- add bOneWayInteraction flag to BodyInstance

diff --git a/Engine/Source/Runtime/Engine/Classes/PhysicsEngine/BodyInstance.h b/Engine/Source/Runtime/Engine/Classes/PhysicsEngine/BodyInstance.h
--- a/Engine/Source/Runtime/Engine/Classes/PhysicsEngine/BodyInstance.h
+++ b/Engine/Source/Runtime/Engine/Classes/PhysicsEngine/BodyInstance.h
@@ -440,6 +440,10 @@ protected:
 	UPROPERTY(EditAnywhere, AdvancedDisplay, BlueprintReadWrite, Category = Physics)
 	uint8 bInertiaConditioning : 1;
 
+	/** If set to true, this body will treat bodies that do not have the flag set as having infinite mass */
+	UPROPERTY(EditAnywhere, AdvancedDisplay, BlueprintReadWrite, Category = Physics)
+	uint8 bOneWayInteraction : 1;
+
 public:
 	/** Set the desired delta time for the body. **/
 	UPROPERTY(EditAnywhere, AdvancedDisplay, BlueprintReadOnly, Category = Physics)
@@ -876,11 +880,6 @@ public:
 	*/
 	ENGINE_API void SetMaxDepenetrationVelocity(float MaxVelocity);
 
-private:
-	// Send the MaxDepenetrationVelocity setting to the solver
-	ENGINE_API void UpdateMaxDepenetrationVelocity();
-
-public:
 	/** Set whether we should get a notification about physics collisions */
 	ENGINE_API void SetInstanceNotifyRBCollision(bool bNewNotifyCollision);
 	/** Enables/disables whether this body is affected by gravity. */
@@ -1278,6 +1277,9 @@ public:
 
 private:
 
+	ENGINE_API void UpdateOneWayInteraction();
+	ENGINE_API void UpdateMaxDepenetrationVelocity();
+
 	/**
 	 * Invalidate Collision Profile Name
 	 * This gets called when it invalidates the reason of Profile Name
diff --git a/Engine/Source/Runtime/Engine/Private/PhysicsEngine/BodyInstance.cpp b/Engine/Source/Runtime/Engine/Private/PhysicsEngine/BodyInstance.cpp
--- a/Engine/Source/Runtime/Engine/Private/PhysicsEngine/BodyInstance.cpp
+++ b/Engine/Source/Runtime/Engine/Private/PhysicsEngine/BodyInstance.cpp
@@ -341,6 +341,7 @@ FBodyInstance::FBodyInstance()
 	, bInterpolateWhenSubStepping(true)
 	, bPendingCollisionProfileSetup(false)
 	, bInertiaConditioning(true)	
+	, bOneWayInteraction(false)
 	, bOverrideSolverAsyncDeltaTime(false)
 	, SolverAsyncDeltaTime(1.f / 60)
 	, Scale3D(1.0f)
@@ -3400,15 +3401,27 @@ void FBodyInstance::ClearForces(bool bAllowSubstepping)
 	});
 }
 
-void FBodyInstance::SetOneWayInteraction(bool InOneWayInteraction /*= true*/)
+void FBodyInstance::SetOneWayInteraction(bool bInOneWayInteraction)
 {
-	FPhysicsCommand::ExecuteWrite(ActorHandle, [&](const FPhysicsActorHandle& Actor)
+	if (bOneWayInteraction != bInOneWayInteraction)
+	{
+		bOneWayInteraction = bInOneWayInteraction;
+
+		UpdateOneWayInteraction();
+	}
+}
+
+void FBodyInstance::UpdateOneWayInteraction()
+{
+	const bool bCurrentOneWayInteraction = bOneWayInteraction;
+
+	FPhysicsCommand::ExecuteWrite(ActorHandle, [bCurrentOneWayInteraction](const FPhysicsActorHandle& Actor)
+	{
+		if (FPhysicsInterface::IsRigidBody(Actor) && !IsRigidBodyKinematic_AssumesLocked(Actor))
 		{
-			if (FPhysicsInterface::IsRigidBody(Actor) && !IsRigidBodyKinematic_AssumesLocked(Actor))
-			{
-				FPhysicsInterface::SetOneWayInteraction_AssumesLocked(Actor, InOneWayInteraction);
-			}
-		});
+			FPhysicsInterface::SetOneWayInteraction_AssumesLocked(Actor, bCurrentOneWayInteraction);
+		}
+	});
 }
 
 void FBodyInstance::AddTorqueInRadians(const FVector& Torque, bool bAllowSubstepping, bool bAccelChange, const FAsyncPhysicsTimestamp TimeStamp, APlayerController* PlayerController)
@@ -4105,6 +4118,7 @@ void FBodyInstance::InitDynamicProperties_AssumesLocked()
 			UpdateDampingProperties();
 			SetMaxAngularVelocityInRadians(GetMaxAngularVelocityInRadians(), false, false);
 			UpdateMaxDepenetrationVelocity();
+			UpdateOneWayInteraction();
 		}
 		else
 		{
diff --git a/Engine/Source/Runtime/Experimental/Chaos/Private/Chaos/ChaosDebugDraw.cpp b/Engine/Source/Runtime/Experimental/Chaos/Private/Chaos/ChaosDebugDraw.cpp
--- a/Engine/Source/Runtime/Experimental/Chaos/Private/Chaos/ChaosDebugDraw.cpp
+++ b/Engine/Source/Runtime/Experimental/Chaos/Private/Chaos/ChaosDebugDraw.cpp
@@ -35,6 +35,11 @@
 
 namespace Chaos
 {
+	namespace CVars
+	{
+		extern int32 ChaosOneWayInteractionPairCollisionMode;
+	}
+
 	namespace DebugDraw
 	{
 		bool bChaosDebugDebugDrawShapeBounds = false;
@@ -1331,6 +1336,12 @@ namespace Chaos
 					FDebugDrawQueue::GetInstance().DrawDebugLine(ImpulsePos, ImpulsePos + Settings.DrawScale * Settings.ImpulseScale * SpaceTransform.TransformVectorNoScale(Contact.GetAccumulatedImpulse()), Color, false, Duration, uint8(Settings.DrawPriority), Settings.LineThickness);
 				}
 
+				// Show the sphere approximation if enabled
+				if ((CVars::ChaosOneWayInteractionPairCollisionMode == (int32)EOneWayInteractionPairCollisionMode::SphereCollision) && Particle0->OneWayInteraction() && Particle1->OneWayInteraction())
+				{
+					DrawShape(Contact.GetShapeWorldTransform0(), Contact.GetImplicit0(), nullptr, FColor::Green, 0, &Settings);
+					DrawShape(Contact.GetShapeWorldTransform1(), Contact.GetImplicit1(), nullptr, FColor::Green, 0, &Settings);
+				}
 			}
 			if (Settings.ContactOwnerWidth > 0)
 			{
diff --git a/Engine/Source/Runtime/Experimental/Chaos/Private/Chaos/Collision/ParticlePairMidPhase.cpp b/Engine/Source/Runtime/Experimental/Chaos/Private/Chaos/Collision/ParticlePairMidPhase.cpp
--- a/Engine/Source/Runtime/Experimental/Chaos/Private/Chaos/Collision/ParticlePairMidPhase.cpp
+++ b/Engine/Source/Runtime/Experimental/Chaos/Private/Chaos/Collision/ParticlePairMidPhase.cpp
@@ -5,6 +5,7 @@
 #include "Chaos/Collision/CollisionContext.h"
 #include "Chaos/Collision/CollisionFilter.h"
 #include "Chaos/Collision/CollisionUtil.h"
+#include "Chaos/Collision/ContactPointsMiscShapes.h"
 #include "Chaos/Collision/ContactTriangles.h"
 #include "Chaos/Collision/PBDCollisionConstraint.h"
 #include "Chaos/CollisionResolution.h"
@@ -58,6 +59,8 @@ namespace Chaos
 		int32 Chaos_Collision_MidPhase_MaxShapePairs = 100;
 		FAutoConsoleVariableRef CVarChaos_Collision_EnableShapePairs(TEXT("p.Chaos.Collision.EnableShapePairs"), bChaos_Collision_MidPhase_EnableShapePairs, TEXT(""));
 		FAutoConsoleVariableRef CVarChaos_Collision_MaxShapePairs(TEXT("p.Chaos.Collision.MaxShapePairs"), Chaos_Collision_MidPhase_MaxShapePairs, TEXT(""));
+
+		extern int32 ChaosOneWayInteractionPairCollisionMode;
 	}
 
 	using namespace CVars;
@@ -747,21 +750,44 @@ namespace Chaos
 
 		// Is either of the implicits a Union of Unions?
 		const bool bIsTree = bIsTree0 || bIsTree1;
+	
+		const bool bCanPrebuildShapePairs = !bTooManyImplicitPairs && !bIsTree;
+
+		// If both particles have one-way interaction enabled, we might want to treat them as spheres
+		if (bCanPrebuildShapePairs && (CVars::ChaosOneWayInteractionPairCollisionMode == (int32)EOneWayInteractionPairCollisionMode::SphereCollision))
+		{
+			const bool bIsOneWay0 = FConstGenericParticleHandle(InParticle0)->OneWayInteraction();
+			const bool bIsOneWay1 = FConstGenericParticleHandle(InParticle1)->OneWayInteraction();
+			if (bIsOneWay0 && bIsOneWay1)
+			{
+				return EParticlePairMidPhaseType::SphereApproximation;
+			}
+		}
+
+		// If we have two small flat hierarchies we will expand and prefilter all shape pairs
+		if (bCanPrebuildShapePairs)
+		{
+			return EParticlePairMidPhaseType::ShapePair;
+		}
 
-		const bool bPrebuildShapePairs = !bTooManyImplicitPairs && !bIsTree;
-		return bPrebuildShapePairs ? EParticlePairMidPhaseType::ShapePair : EParticlePairMidPhaseType::Generic;
+		// We have at least one complicated geometry hierarchy so use the general purpose midphase
+		return EParticlePairMidPhaseType::Generic;
 	}
 
 	FParticlePairMidPhase* FParticlePairMidPhase::Make(FGeometryParticleHandle* InParticle0, FGeometryParticleHandle* InParticle1)
 	{
 		const EParticlePairMidPhaseType MidPhaseType = CalculateMidPhaseType(InParticle0, InParticle1);
-		if (MidPhaseType == EParticlePairMidPhaseType::ShapePair)
+		switch (MidPhaseType)
 		{
-			return new FShapePairParticlePairMidPhase();
-		}
-		else
-		{
-			return new FGenericParticlePairMidPhase();
+			case EParticlePairMidPhaseType::Generic:
+				return new FGenericParticlePairMidPhase();
+			case EParticlePairMidPhaseType::ShapePair:
+				return new FShapePairParticlePairMidPhase();
+			case EParticlePairMidPhaseType::SphereApproximation:
+				return new FSphereApproximationParticlePairMidPhase();
+			default:
+				check(false);
+				return nullptr;
 		}
 	}
 
@@ -1948,6 +1974,217 @@ namespace Chaos
 		UE_LOG(LogChaos, Warning, TEXT("Unsupported, handle Rewind/Resim in FGenericParticlePairMidPhase::InjectCollisionImpl"));
 #endif
 	}
+
+
+	////////////////////////////////////////////////////////////////////////////////////////
+	////////////////////////////////////////////////////////////////////////////////////////
+	////////////////////////////////////////////////////////////////////////////////////////
+	////////////////////////////////////////////////////////////////////////////////////////
+
+
+	FSphereApproximationParticlePairMidPhase::FSphereApproximationParticlePairMidPhase()
+		: FParticlePairMidPhase(EParticlePairMidPhaseType::SphereApproximation)
+		, Sphere0(FVec3(0), 0.0)
+		, Sphere1(FVec3(0), 0.0)
+		, SphereShape0(nullptr)
+		, SphereShape1(nullptr)
+		, LastUsedEpoch(INDEX_NONE)
+		, bHasSpheres(false)
+	{
+	}
+
+	void FSphereApproximationParticlePairMidPhase::ResetImpl()
+	{
+		Constraint.Reset();
+	}
+
+	void FSphereApproximationParticlePairMidPhase::BuildDetectorsImpl()
+	{
+		if (!IsValid())
+		{
+			return;
+		}
+
+		if (!GetParticle0()->HasBounds() || !GetParticle1()->HasBounds())
+		{
+			return;
+		}
+
+		const FShapeInstanceArray& Shapes0 = GetParticle0()->ShapeInstances();
+		const FShapeInstanceArray& Shapes1 = GetParticle1()->ShapeInstances();
+		if (Shapes0.IsEmpty() || Shapes1.IsEmpty())
+		{
+			return;
+		}
+
+		// See if we should collide (do any shape pairs collide)
+		bool bDoCollide = false;
+		for (int32 ShapeIndex0 = 0; ShapeIndex0 < Shapes0.Num(); ++ShapeIndex0)
+		{
+			const FShapeInstance* ShapeInstance0 = Shapes0[ShapeIndex0].Get();
+			for (int32 ShapeIndex1 = 0; ShapeIndex1 < Shapes1.Num(); ++ShapeIndex1)
+			{
+				const FShapeInstance* ShapeInstance1 = Shapes1[ShapeIndex1].Get();
+
+				const FImplicitObject* Implicit0 = ShapeInstance0->GetLeafGeometry();
+				const EImplicitObjectType ImplicitType0 = Private::GetImplicitCollisionType(Particle0, Implicit0);
+
+				const FImplicitObject* Implicit1 = ShapeInstance1->GetLeafGeometry();
+				const EImplicitObjectType ImplicitType1 = Private::GetImplicitCollisionType(Particle1, Implicit1);
+
+				// Use materials etc from the first overlapping shape pairs
+				if (SphereShape0 == nullptr)
+				{
+					SphereShape0 = ShapeInstance0;
+					SphereShape1 = ShapeInstance1;
+				}
+
+				const bool bDoPassFilter = ShapePairNarrowPhaseFilter(ImplicitType0, ShapeInstance0, ImplicitType1, ShapeInstance1);
+				if (bDoPassFilter)
+				{
+					bDoCollide = true;
+					break;
+				}
+			}
+		}
+
+		if (!bDoCollide)
+		{
+			return;
+		}
+
+		// Initialize the sphere centers and radii
+		InitSphere(GetParticle0(), Sphere0);
+		InitSphere(GetParticle1(), Sphere1);
+		bHasSpheres = true;
+	}
+
+	void FSphereApproximationParticlePairMidPhase::InitSphere(
+		const FGeometryParticleHandle* InParticle, 
+		FImplicitSphere3& OutSphere)
+	{
+		// @todo(chaos): maybe we should only consider the bounds of the shapes that pass the ShapePairNarrowPhaseFilter?
+		const FVec3 Center = InParticle->LocalBounds().Center();
+		const FVec3 Extents = InParticle->LocalBounds().Extents();
+		const FReal Radius = 0.5 * Extents.GetAbsMax();
+
+		OutSphere = FImplicitSphere3(Center, Radius);
+	}
+
+	int32 FSphereApproximationParticlePairMidPhase::GenerateCollisionsImpl(
+		const FReal CullDistance,
+		const FReal Dt,
+		const FCollisionContext& Context)
+	{
+		if (!bHasSpheres)
+		{
+			return 0;
+		}
+
+		// NOTE: We are still using the bounds generated from the real collision shapes for the cull test
+		const FAABB3& WorldBounds0 = GetParticle0()->WorldSpaceInflatedBounds();
+		const FAABB3& WorldBounds1 = GetParticle1()->WorldSpaceInflatedBounds();
+		if (!WorldBounds0.Intersects(WorldBounds1))
+		{
+			return 0;
+		}
+
+		const int32 CurrentEpoch = Context.GetAllocator()->GetCurrentEpoch();
+		const int32 LastEpoch = CurrentEpoch - 1;
+		const bool bWasUpdatedLastTick = IsUsedSince(LastEpoch);
+
+		const FRigidTransform3 ShapeWorldTransform0 = FConstGenericParticleHandle(GetParticle0())->GetTransformPQ();
+		const FRigidTransform3 ShapeWorldTransform1 = FConstGenericParticleHandle(GetParticle1())->GetTransformPQ();
+
+		// Sphere distance culling
+		const FVec3 SphereWorldPos0 = ShapeWorldTransform0.TransformPositionNoScale(Sphere0.GetCenter());
+		const FVec3 SphereWorldPos1 = ShapeWorldTransform1.TransformPositionNoScale(Sphere1.GetCenter());
+		const FVec3 DR = SphereWorldPos0 - SphereWorldPos1;
+		const FReal DRLenSq = DR.SizeSquared();
+		const FReal CullSeparation = Sphere0.GetRadius() + Sphere1.GetRadius() + CullDistance;
+		if (DRLenSq > FMath::Square(CullSeparation))
+		{
+			return 0;
+		}
+
+		// Create and set up constraint
+		if (!Constraint.IsValid())
+		{
+			Constraint = Context.GetAllocator()->CreateConstraint(
+				GetParticle0(), &Sphere0, SphereShape0, nullptr, FRigidTransform3(), 
+				GetParticle1(), &Sphere1, SphereShape1, nullptr, FRigidTransform3(),
+				CullDistance, true, EContactShapesType::SphereSphere);
+
+			Constraint->GetContainerCookie().MidPhase = this;
+			Constraint->GetContainerCookie().bIsMultiShapePair = false;
+			Constraint->GetContainerCookie().CreationEpoch = CurrentEpoch;
+			Constraint->SetCollisionSortKey(Private::FCollisionSortKey(Particle0, 0, Particle1, 0));
+		}
+
+		Constraint->SetShapeWorldTransforms(ShapeWorldTransform0, ShapeWorldTransform1);
+		Constraint->SetCullDistance(CullDistance);
+
+		if (!bWasUpdatedLastTick || (Constraint->GetManifoldPoints().Num() == 0))
+		{
+			// Clear all manifold data including saved contact data
+			Constraint->ResetManifold();
+		}
+
+		// We are not trying to reuse manifold points, so reset them but leave stored data intact (for friction)
+		Constraint->ResetActiveManifoldContacts();
+
+		// Sphere collision
+		FContactPoint ContactPoint = SphereSphereContactPoint(Sphere0, ShapeWorldTransform0, Sphere1, ShapeWorldTransform1, CullDistance);
+		Constraint->SetOneShotManifoldContacts({ ContactPoint });
+
+		// Activate constraint if we have a contact
+		if (Constraint->GetPhi() <= CullDistance)
+		{
+			Constraint->SetIsInitialContact(!bWasUpdatedLastTick);
+
+			if (Context.GetAllocator()->ActivateConstraint(Constraint.Get()))
+			{
+				LastUsedEpoch = CurrentEpoch;
+				return 1;
+			}
+		}
+
+		return 0;
+	}
+
+	void FSphereApproximationParticlePairMidPhase::WakeCollisionsImpl(
+		const int32 CurrentEpoch)
+	{
+		if (Constraint.IsValid() && (Constraint->GetContainerCookie().LastUsedEpoch >= LastUsedEpoch))
+		{
+			Constraint->GetContainerCookie().LastUsedEpoch = CurrentEpoch;
+			Constraint->GetContainerCookie().ConstraintIndex = INDEX_NONE;
+			Constraint->GetContainerCookie().CCDConstraintIndex = INDEX_NONE;
+		}
+	}
+
+	void FSphereApproximationParticlePairMidPhase::InjectCollisionImpl(
+		const FPBDCollisionConstraint& InConstraint, 
+		const FCollisionContext& Context)
+	{
+		const int32 CurrentEpoch = Context.GetAllocator()->GetCurrentEpoch();
+
+		if (!Constraint.IsValid())
+		{
+			Constraint = Context.GetAllocator()->CreateConstraint();
+			Constraint->GetContainerCookie().MidPhase = this;
+			Constraint->GetContainerCookie().bIsMultiShapePair = false;
+			Constraint->GetContainerCookie().CreationEpoch = CurrentEpoch;
+		}
+
+		// Copy the constraint data over the existing one (ensure we do not replace data required by the graph and the allocator/container)
+		Constraint->RestoreFrom(InConstraint);
+
+		// Add the constraint to the active list
+		// If the constraint already existed and was already active, this will do nothing
+		Context.GetAllocator()->ActivateConstraint(Constraint.Get());
+		LastUsedEpoch = CurrentEpoch;
+	}
 }
 
 
diff --git a/Engine/Source/Runtime/Experimental/Chaos/Private/Chaos/PBDRigidsEvolutionGBF.cpp b/Engine/Source/Runtime/Experimental/Chaos/Private/Chaos/PBDRigidsEvolutionGBF.cpp
--- a/Engine/Source/Runtime/Experimental/Chaos/Private/Chaos/PBDRigidsEvolutionGBF.cpp
+++ b/Engine/Source/Runtime/Experimental/Chaos/Private/Chaos/PBDRigidsEvolutionGBF.cpp
@@ -123,9 +123,11 @@ namespace Chaos
 		bool bChaosSolverPersistentGraph = true;
 		FAutoConsoleVariableRef CVarChaosSolverPersistentGraph(TEXT("p.Chaos.Solver.PersistentGraph"), bChaosSolverPersistentGraph, TEXT(""));
 
-		// Whether we should ignore collisions between two particles that both have one-way collision enabled
-		bool bChaosIgnoreOneWayPairCollisions = true;
-		FAutoConsoleVariableRef CVarChaosIgnoreOneWayPairCollisions(TEXT("p.Chaos.Solver.IgnoreOneWayPairCollisions"), bChaosIgnoreOneWayPairCollisions, TEXT("Ignore collisions between two one-way-interaction particles"));
+		// Determines what happens when two one-way particles collide
+		// See EOneWayInteractionPairCollisionMode
+		int32 ChaosOneWayInteractionPairCollisionMode = (int32)EOneWayInteractionPairCollisionMode::IgnoreCollision;
+		FAutoConsoleVariableRef CVarChaosIgnoreOneWayPairCollisions(TEXT("p.Chaos.Solver.OneWayPairCollisionMode"), ChaosOneWayInteractionPairCollisionMode, TEXT("How to treat collisions between two one-way interaction particles. See EOneWayInteractionPairCollisionMode (0: Ignore collisions; 1: Collide as normal; 2: Collide as spheres)"));
+
 
 		DECLARE_CYCLE_STAT(TEXT("FPBDRigidsEvolutionGBF::AdvanceOneTimeStep"), STAT_Evolution_AdvanceOneTimeStep, STATGROUP_Chaos);
 		DECLARE_CYCLE_STAT(TEXT("FPBDRigidsEvolutionGBF::UnclusterUnions"), STAT_Evolution_UnclusterUnions, STATGROUP_Chaos);
diff --git a/Engine/Source/Runtime/Experimental/Chaos/Public/Chaos/Collision/CollisionConstraintAllocator.h b/Engine/Source/Runtime/Experimental/Chaos/Public/Chaos/Collision/CollisionConstraintAllocator.h
--- a/Engine/Source/Runtime/Experimental/Chaos/Public/Chaos/Collision/CollisionConstraintAllocator.h
+++ b/Engine/Source/Runtime/Experimental/Chaos/Public/Chaos/Collision/CollisionConstraintAllocator.h
@@ -678,6 +678,22 @@ namespace Chaos
 			}
 		}
 
+		if (GetMidPhaseType() == EParticlePairMidPhaseType::SphereApproximation)
+		{
+			FSphereApproximationParticlePairMidPhase* This = static_cast<FSphereApproximationParticlePairMidPhase*>(this);
+			FPBDCollisionConstraint* Constraint = This->Constraint.Get();
+			if (Constraint != nullptr)
+			{
+				if (CollisionVisitorShouldVisit(Constraint, VisitFlags))
+				{
+					if (Visitor(*Constraint) == ECollisionVisitorResult::Stop)
+					{
+						return ECollisionVisitorResult::Stop;
+					}
+				}
+			}
+		}
+
 		return ECollisionVisitorResult::Continue;
 	}
 
@@ -717,6 +733,22 @@ namespace Chaos
 			}
 		}
 
+		if (GetMidPhaseType() == EParticlePairMidPhaseType::SphereApproximation)
+		{
+			const FSphereApproximationParticlePairMidPhase* This = static_cast<const FSphereApproximationParticlePairMidPhase*>(this);
+			const FPBDCollisionConstraint* Constraint = This->Constraint.Get();
+			if (Constraint != nullptr)
+			{
+				if (CollisionVisitorShouldVisit(Constraint, VisitFlags))
+				{
+					if (Visitor(*Constraint) == ECollisionVisitorResult::Stop)
+					{
+						return ECollisionVisitorResult::Stop;
+					}
+				}
+			}
+		}
+
 		return ECollisionVisitorResult::Continue;
 	}
 
diff --git a/Engine/Source/Runtime/Experimental/Chaos/Public/Chaos/Collision/PBDCollisionConstraint.h b/Engine/Source/Runtime/Experimental/Chaos/Public/Chaos/Collision/PBDCollisionConstraint.h
--- a/Engine/Source/Runtime/Experimental/Chaos/Public/Chaos/Collision/PBDCollisionConstraint.h
+++ b/Engine/Source/Runtime/Experimental/Chaos/Public/Chaos/Collision/PBDCollisionConstraint.h
@@ -202,6 +202,7 @@ namespace Chaos
 		friend class FPBDCollisionConstraints;
 		friend class FShapePairParticlePairMidPhase;
 		friend class FSingleShapePairCollisionDetector;
+		friend class FSphereApproximationParticlePairMidPhase;
 
 		friend CHAOS_API bool ContactConstraintSortPredicate(const FPBDCollisionConstraint& L, const FPBDCollisionConstraint& R);
 
diff --git a/Engine/Source/Runtime/Experimental/Chaos/Public/Chaos/Collision/ParticlePairMidPhase.h b/Engine/Source/Runtime/Experimental/Chaos/Public/Chaos/Collision/ParticlePairMidPhase.h
--- a/Engine/Source/Runtime/Experimental/Chaos/Public/Chaos/Collision/ParticlePairMidPhase.h
+++ b/Engine/Source/Runtime/Experimental/Chaos/Public/Chaos/Collision/ParticlePairMidPhase.h
@@ -36,6 +36,9 @@ namespace Chaos
 		// number of shapes. Pre-expands the set of potentially
 		// colliding shape pairs.
 		ShapePair,
+
+		// A midphase used to collide particles as sphere approximations
+		SphereApproximation,
 	};
 
 	/**
@@ -182,7 +185,7 @@ namespace Chaos
 		}
 
 		/**
-		 * @brief Set up the midphase based on the SHapesArrays of the two particles
+		 * @brief Set up the midphase based on the ShapesArrays of the two particles
 		 * Only intended to be called once right after constructor. We don't do this work in
 		 * the constructor so that we can reduce the time that the lock is held when allocating
 		 * new MidPhases.
@@ -540,4 +543,39 @@ namespace Chaos
 		TMap<uint64, FPBDCollisionConstraintPtr> Constraints;
 		TArray<FPBDCollisionConstraint*> NewConstraints;
 	};
+
+	/**
+	* A midphase for a particle pair that replaces both particles with a sphere approximation
+	*/
+	class FSphereApproximationParticlePairMidPhase : public FParticlePairMidPhase
+	{
+	public:
+		friend class FParticlePairMidPhase;
+
+		CHAOS_API FSphereApproximationParticlePairMidPhase();
+
+		CHAOS_API virtual void ResetImpl() override final;
+		CHAOS_API virtual void BuildDetectorsImpl() override final;
+
+	protected:
+		CHAOS_API virtual int32 GenerateCollisionsImpl(
+			const FReal CullDistance,
+			const FReal Dt,
+			const FCollisionContext& Context) override final;
+
+		CHAOS_API virtual void WakeCollisionsImpl(const int32 CurrentEpoch) override final;
+
+		CHAOS_API virtual void InjectCollisionImpl(const FPBDCollisionConstraint& Constraint, const FCollisionContext& Context) override final;
+
+		static void InitSphere(const FGeometryParticleHandle* InParticle, FImplicitSphere3& OutSphere);
+
+	private:
+		FImplicitSphere3 Sphere0;
+		FImplicitSphere3 Sphere1;
+		const FShapeInstance* SphereShape0;
+		const FShapeInstance* SphereShape1;
+		FPBDCollisionConstraintPtr Constraint;
+		int32 LastUsedEpoch;
+		bool bHasSpheres;
+	};
 }
diff --git a/Engine/Source/Runtime/Experimental/Chaos/Public/Chaos/Collision/SpatialAccelerationBroadPhase.h b/Engine/Source/Runtime/Experimental/Chaos/Public/Chaos/Collision/SpatialAccelerationBroadPhase.h
--- a/Engine/Source/Runtime/Experimental/Chaos/Public/Chaos/Collision/SpatialAccelerationBroadPhase.h
+++ b/Engine/Source/Runtime/Experimental/Chaos/Public/Chaos/Collision/SpatialAccelerationBroadPhase.h
@@ -23,7 +23,7 @@
 namespace Chaos::CVars
 {
 	extern bool bChaosMidPhaseRedistributionEnabled;
-	extern bool bChaosIgnoreOneWayPairCollisions;
+	extern int32 ChaosOneWayInteractionPairCollisionMode;
 }
 
 namespace Chaos
@@ -680,7 +680,7 @@ namespace Chaos
 
 				if (Overlap.bCollisionsEnabled)
 				{
-					if (CVars::bChaosIgnoreOneWayPairCollisions)
+					if (CVars::ChaosOneWayInteractionPairCollisionMode == (int32)EOneWayInteractionPairCollisionMode::IgnoreCollision)
 					{
 						const bool bIsOneWay0 = FConstGenericParticleHandle(Overlap.Particles[0])->OneWayInteraction();
 						const bool bIsOneWay1 = FConstGenericParticleHandle(Overlap.Particles[1])->OneWayInteraction();
diff --git a/Engine/Source/Runtime/Experimental/Chaos/Public/Chaos/CollisionResolutionTypes.h b/Engine/Source/Runtime/Experimental/Chaos/Public/Chaos/CollisionResolutionTypes.h
--- a/Engine/Source/Runtime/Experimental/Chaos/Public/Chaos/CollisionResolutionTypes.h
+++ b/Engine/Source/Runtime/Experimental/Chaos/Public/Chaos/CollisionResolutionTypes.h
@@ -71,6 +71,19 @@ namespace Chaos
 		NumShapesTypes
 	};
 
+	/** How to treat collisions between two particles that both have OneWayInteraction enabled */
+	enum class EOneWayInteractionPairCollisionMode
+	{
+		// Ignore collisions
+		IgnoreCollision,
+
+		// Collide using particle's regular shapes
+		NormalCollision,
+
+		// Collide as spheres
+		SphereCollision,
+	};
+
 	//
 	//
 	//
