Bottom-Up convex construction  + shape index on the simplified convex

diff --git a/Engine/Source/Runtime/Experimental/Chaos/Private/Chaos/ChaosDebugDraw.cpp b/Engine/Source/Runtime/Experimental/Chaos/Private/Chaos/ChaosDebugDraw.cpp
--- a/Engine/Source/Runtime/Experimental/Chaos/Private/Chaos/ChaosDebugDraw.cpp
+++ b/Engine/Source/Runtime/Experimental/Chaos/Private/Chaos/ChaosDebugDraw.cpp
@@ -398,7 +398,7 @@ namespace Chaos
 		//
 		//
 		//
-		void DrawShapesImpl(const FGeometryParticleHandle* Particle, const FRigidTransform3& ShapeTransform, const FImplicitObject* Implicit, const FPerShapeData* Shape, const FReal Margin, const FColor& Color, const FRealSingle Duration, const FChaosDebugDrawSettings& Settings, const bool bHasConvexOptimizer = false);
+		void DrawShapesImpl(const FGeometryParticleHandle* Particle, const FRigidTransform3& ShapeTransform, const FImplicitObject* Implicit, const FPerShapeData* Shape, const FReal Margin, const FColor& Color, const FRealSingle Duration, const FChaosDebugDrawSettings& Settings);
 
 
 		void DrawShape(const FRigidTransform3& ShapeTransform, const FImplicitObject* Implicit, const FPerShapeData* Shape, const FColor& Color, const float Duration, const FChaosDebugDrawSettings* Settings)
@@ -626,7 +626,7 @@ namespace Chaos
 			}
 		}
 
-		void DrawShapesImpl(const FGeometryParticleHandle* Particle, const FRigidTransform3& ShapeTransform, const FImplicitObject* Implicit, const FPerShapeData* Shape, const FReal Margin, const FColor& Color, const FRealSingle Duration, const FChaosDebugDrawSettings& Settings, const bool bHasConvexOptimizer)
+		void DrawShapesImpl(const FGeometryParticleHandle* Particle, const FRigidTransform3& ShapeTransform, const FImplicitObject* Implicit, const FPerShapeData* Shape, const FReal Margin, const FColor& Color, const FRealSingle Duration, const FChaosDebugDrawSettings& Settings)
 		{
 			if (Implicit == nullptr)
 			{
@@ -750,31 +750,23 @@ namespace Chaos
 						}
 					}
 				}
-			}
-			if (bChaosDebugDebugDrawColorShapesByConvexType)
-			{
-				if(InnerType == ImplicitObjectType::Convex)
-				{
-					ShapeColor = bHasConvexOptimizer ? FColor::Green : FColor::Orange;
-				}
-				else
-				{
-					ShapeColor = FColor::Yellow;
-				}
-			}
-			if (Shape && bChaosDebugDebugDrawColorShapesBySimQueryType)
-			{
-				if (Shape->GetSimEnabled())
-				{
-					ShapeColor = FColor::Green;
-				}
-				else if (Shape->GetQueryEnabled())
+				if (bChaosDebugDebugDrawColorShapesByConvexType)
 				{
-					ShapeColor = FColor::Red;
-				}
-				else if (Shape->GetIsProbe())
-				{
-					ShapeColor = FColor::Orange;
+					if(InnerType == ImplicitObjectType::Convex)
+					{
+						ShapeColor = FColor::Orange;
+						if (const FPBDRigidClusteredParticleHandle* ClusteredParticle = Particle->CastToClustered())
+						{
+							if(ClusteredParticle->ConvexOptimizer() && ClusteredParticle->ConvexOptimizer()->IsValid())
+							{
+								ShapeColor = FColor::Green;
+							}
+						}
+					}
+					else
+					{
+						ShapeColor = FColor::Yellow;
+					}
 				}
 			}
 			if (Shape && !bChaosDebugDebugDrawShowQueryOnlyShapes)
@@ -915,11 +907,11 @@ namespace Chaos
 					{
 						if(RootObjectIndex == INDEX_NONE && !ConvexOptimizer->GetShapeInstances().IsEmpty())
 						{
-							DrawShapesImpl(Particle, RelativeTransform * FRigidTransform3(P, Q), ImplicitObject, ConvexOptimizer->GetShapeInstances()[0].Get(), 0.0f, InColor, 0.0f, Settings, true);
+							DrawShapesImpl(Particle, RelativeTransform * FRigidTransform3(P, Q), ImplicitObject, ConvexOptimizer->GetShapeInstances()[0].Get(), 0.0f, InColor, 0.0f, Settings);
 						}
 						else if(Particle->ShapeInstances().IsValidIndex(RootObjectIndex))
 						{
-							DrawShapesImpl(Particle, RelativeTransform * FRigidTransform3(P, Q), ImplicitObject, Particle->ShapeInstances()[RootObjectIndex].Get(), 0.0f, InColor, 0.0f, Settings, false);
+							DrawShapesImpl(Particle, RelativeTransform * FRigidTransform3(P, Q), ImplicitObject, Particle->ShapeInstances()[RootObjectIndex].Get(), 0.0f, InColor, 0.0f, Settings);
 						}
 					});
 					return true;
diff --git a/Engine/Source/Runtime/Experimental/Chaos/Private/Chaos/ClusterUnionManager.cpp b/Engine/Source/Runtime/Experimental/Chaos/Private/Chaos/ClusterUnionManager.cpp
--- a/Engine/Source/Runtime/Experimental/Chaos/Private/Chaos/ClusterUnionManager.cpp
+++ b/Engine/Source/Runtime/Experimental/Chaos/Private/Chaos/ClusterUnionManager.cpp
@@ -896,9 +896,6 @@ namespace Chaos
 		{
 			FlushIncrementalGeometryOperations(ClusterUnion);
 		}
-
-		// Build the convex optimizer if required
-		FRigidClustering::BuildConvexOptimizer(ClusterUnion.InternalCluster);
 		
 		if (ClusterUnion.bGenerateConnectivityEdges)
 		{
@@ -919,7 +916,9 @@ namespace Chaos
 				MClustering.HandleConnectivityOnReleaseClusterParticle(ClusterUnion.InternalCluster, false);
 			}
 		}
-
+		// Build the convex optimizer if required
+		MClustering.BuildConvexOptimizer(ClusterUnion.InternalCluster);
+		
 		for (FPBDRigidClusteredParticleHandle* ChildParticle : PendingParticlesToUndoChildToParentLock)
 		{
 			if (ChildParticle)
diff --git a/Engine/Source/Runtime/Experimental/Chaos/Private/Chaos/ConvexOptimizer.cpp b/Engine/Source/Runtime/Experimental/Chaos/Private/Chaos/ConvexOptimizer.cpp
--- a/Engine/Source/Runtime/Experimental/Chaos/Private/Chaos/ConvexOptimizer.cpp
+++ b/Engine/Source/Runtime/Experimental/Chaos/Private/Chaos/ConvexOptimizer.cpp
@@ -16,20 +16,24 @@ namespace CVars
 	FAutoConsoleVariableRef CVarChaosSimplifyUnion(TEXT("p.Chaos.Convex.SimplifyUnion"), bChaosConvexSimplifyUnion, TEXT("If true replace all the convexes within an implcit hierarchy with a simplified one (kdop18 tribox for now) for collision"));
 
 	// Max number of convex LODs used during simplification  for dynamic particles
-	int32 ChaosConvexMaxDynamicLODs = 0;
-	FAutoConsoleVariableRef CVarChaosConvexMaxDynamicLODs(TEXT("p.Chaos.Convex.MaxDynamicLODs"), ChaosConvexMaxDynamicLODs, TEXT("Max number of convex LODs used during simplication for dynamic particles"));
+	int32 ChaosConvexKinematicMode = 2;
+	FAutoConsoleVariableRef CVarChaosConvexKinematicMode(TEXT("p.Chaos.Convex.KinematicMode"), ChaosConvexKinematicMode, TEXT("Simplification mode for the kinematic shapes (0: Single Convex, 1: One convex per children, 2: Merge connected children using the splitting threshold"));
 
 	// Max number of convex LODs used during simplification  for kinematic particles
-	int32 ChaosConvexMaxKinematicLODs = -1;
-	FAutoConsoleVariableRef CVarChaosConvexMaxKinematicLODs(TEXT("p.Chaos.Convex.MaxKinematicLODs"), ChaosConvexMaxKinematicLODs, TEXT("Max number of convex LODs used during simplication for kinematic particles"));
+	int32 ChaosConvexDynamicMode = 2;
+	FAutoConsoleVariableRef CVarChaosConvexDynamicMode(TEXT("p.Chaos.Convex.DynamicMode"), ChaosConvexDynamicMode, TEXT("Simplification mode for the dynamic shapes (0: Single Convex, 1: One convex per children, 2: Merge connected children using the splitting threshold)"));
 	
 	// Tribox volume / convex hull threshold to trigger a volume splitting during tree construction
-	float ChaosConvexSplittingThreshold = 1.2f;
+	float ChaosConvexSplittingThreshold = 1.0f;
 	FAutoConsoleVariableRef CVarChaosConvexSplittingThreshold(TEXT("p.Chaos.Convex.SplittingThreshold"), ChaosConvexSplittingThreshold, TEXT("Tribox volume / convex hull threshold to trigger a volume splitting during tree construction"));
 
 	// Min volume of the simplified convexes
-	float ChaosConvexMinVolume = 0.2f;
+	float ChaosConvexMinVolume = 10000.0f;
 	FAutoConsoleVariableRef CVarChaosConvexMinVolume(TEXT("p.Chaos.Convex.MinVolume"), ChaosConvexMinVolume, TEXT("Min volume of the simplified convexes"));
+
+	// Boolean to check if we are merging (bottom-up) or splitting (top-bottom) the convexes 
+	bool ChaosConvexEnableMerging = true;
+	FAutoConsoleVariableRef CVarChaosConvexEnableMerging(TEXT("p.Chaos.Convex.EnableMerging"), ChaosConvexEnableMerging, TEXT("Boolean to check if we are merging (bottom-up) or splitting (top-bottom) the convexes"));
 	
 	extern int32 ChaosUnionBVHMaxDepth;
 	extern int32 ChaosUnionBVHMinShapes;
@@ -37,196 +41,6 @@ namespace CVars
 
 namespace Private
 {
-/**
- * @brief Simple tribox node to store tribox hierarchy information
- */
-struct FTriboxNode
-{
-	// Constructor with a given node tribox and the underling objects range 
-	FTriboxNode(const FTribox& NodeBounds, const TArray<FTribox>& NodeObjects) : BoundsTribox(NodeBounds), ObjectTriboxes(NodeObjects)
-	{}
-
-	// Base constructor
-	FTriboxNode() : BoundsTribox(), ObjectTriboxes()
-	{}
-
-	// Get the closest plane to the node center 
-	void FindClosestPlane(int32& PlaneAxis, FTribox::FRealType& PlaneProjection) const;
-
-	// Split the object triboxes into a left and a right tribox
-	void SplitTriboxNode(const int32& PlaneAxis, const FTribox::FRealType& PlaneProjection,
-							FTriboxNode& LeftTribox, FTribox::FRealType& LeftVolume,
-							FTriboxNode& RightTribox, FTribox::FRealType& RightVolume) const;
-	
-	// Enclosing node tribox 
-	FTribox BoundsTribox;
-
-	// List of objects that are part of that node
-	TArray<FTribox> ObjectTriboxes;
-};
-
-FORCEINLINE void FTriboxNode::FindClosestPlane(int32& PlaneAxis, FTribox::FRealType& PlaneProjection) const
-{
-	const FTribox::FVec3Type PointPosition = BoundsTribox.GetCenter();
-	
-	int32 LocalAxis = INDEX_NONE;
-	FTribox::FRealType LocalProjection = 0.0;
-	FTribox::FRealType MinDistance = FLT_MAX;
-			
-	for(const FTribox& ObjectTribox : ObjectTriboxes)
-	{
-		const FTribox::FRealType ClosestDistance = ObjectTribox.GetClosestPlane(PointPosition, LocalAxis, LocalProjection);
-		if(ClosestDistance < MinDistance)
-		{
-			MinDistance = ClosestDistance;
-			PlaneAxis = LocalAxis;
-			PlaneProjection = LocalProjection;
-		}
-	}
-}
-	
-FORCEINLINE void FTriboxNode::SplitTriboxNode(const int32& PlaneAxis, const FTribox::FRealType& PlaneProjection,
-	FTriboxNode& LeftTribox, FTribox::FRealType& LeftVolume, FTriboxNode& RightTribox,  FTribox::FRealType& RightVolume) const
-{
-	for(const FTribox& ObjectTribox : ObjectTriboxes)
-	{
-		FTribox LocalLeftTribox, LocalRightTribox;
-		ObjectTribox.SplitTriboxSlab(PlaneAxis, PlaneProjection, LocalLeftTribox, LocalRightTribox);
-				
-		if(LocalLeftTribox.HasDatas())
-		{
-			LeftTribox.BoundsTribox += LocalLeftTribox;
-			LeftTribox.ObjectTriboxes.Add(LocalLeftTribox);
-			LeftVolume += LocalLeftTribox.ComputeVolume();
-		}
-		if(LocalRightTribox.HasDatas())
-		{
-			RightTribox.BoundsTribox += LocalRightTribox;
-			RightTribox.ObjectTriboxes.Add(LocalRightTribox);
-			RightVolume += LocalRightTribox.ComputeVolume();
-		}
-	}
-}
-
-/**
- * @brief Simple tribox binary tree to store tribox hierarchy
- */
-struct FTriboxTree
-{
-	// Constructor given a LOD and a number of objects
-	FTriboxTree(const int32 NumLODs, const FConvexOptimizer::FCachedTriboxes& RootTriboxes);
-	
-	// Build all the tree LODs
-	void BuildTreeLODs(TArray<FImplicitObjectPtr>& SimplifiedConvexes, const TUniquePtr<Private::FCollisionObjects>& CollisionObjects);
-
-	// Add a node to the tree nodes
-	void AddTreeNode(const int32 LODIndex, const FTriboxNode& TriboxNode, const FTribox::FRealType& ObjectsVolume,
-		TArray<FImplicitObjectPtr>& SimplifiedConvexes, const TUniquePtr<Private::FCollisionObjects>& CollisionObjects);
-	
-	// Swap the 2 buffers
-	void SwapNodeBuffers()
-	{
-		BufferIndex = 1 - BufferIndex;
-		TreeNodes[BufferIndex].Reset();
-	}
-
-	// Current buffer used to build tthe tree
-	int32 CurrentBuffer() const {return BufferIndex;}
-
-	// Previous buffer used to build the tree
-	int32 PreviousBuffer() const {return 1-BufferIndex;}
-	
-	// List of tree nodes 
-	TArray<FTriboxNode> TreeNodes[2];
-
-	// Tree buffer index 
-	int32 BufferIndex = 0;
-	
-	// Tree max LODs
-	int32 MaxLODs = 0;
-};
-	
-FORCEINLINE FTriboxTree::FTriboxTree(const int32 NumLODs, const FConvexOptimizer::FCachedTriboxes& RootTriboxes) : TreeNodes(), BufferIndex(0)
-{
-	const int32 NumNodes = 1 << NumLODs;
-	
-	TreeNodes[0].Reserve(NumNodes);
-	TreeNodes[1].Reserve(NumNodes);
-	
-	FTriboxNode TriboxNode;
-	for(auto& RootTribox : RootTriboxes)
-	{
-		if(RootTribox.Value.Tribox.HasDatas())
-		{
-			TriboxNode.BoundsTribox += RootTribox.Value.Tribox;
-			TriboxNode.ObjectTriboxes.Add(RootTribox.Value.Tribox);
-		}
-	}
-	TreeNodes[CurrentBuffer()].Emplace(TriboxNode);
-	MaxLODs = NumLODs;
-}
-
-FORCEINLINE void FTriboxTree::BuildTreeLODs(TArray<FImplicitObjectPtr>& SimplifiedConvexes, const TUniquePtr<Private::FCollisionObjects>& CollisionObjects)
-{
-	//UE_LOG(LogTemp, Log, TEXT("	Building tribox binary tree"))
-	for(int32 LODIndex = 1; LODIndex < MaxLODs; ++LODIndex)
-	{
-		// Swap the node buffers (working/result)
-		SwapNodeBuffers();
-
-		//int32 NodeIndex = 0;
-		for(FTriboxNode& TriboxNode : TreeNodes[PreviousBuffer()])
-		{
-			//UE_LOG(LogTemp, Log, TEXT("			Iterating over node %d with num objects = %d"), NodeIndex++, TriboxNode.ObjectTriboxes.Num())
-
-			int32 PlaneAxis = INDEX_NONE;
-			FTribox::FRealType PlaneProjection = 0.0;
-
-			// Find the closest plane to the node center
-			TriboxNode.FindClosestPlane(PlaneAxis, PlaneProjection);
-					
-			FTriboxNode LeftTribox, RightTribox;
-			FTribox::FRealType LeftVolume = 0.0, RightVolume = 0.0;
-
-			// Split the node along the cutting plane
-			TriboxNode.SplitTriboxNode(PlaneAxis, PlaneProjection, LeftTribox, LeftVolume, RightTribox, RightVolume);
-
-			// Add left node to the tree
-			AddTreeNode(LODIndex, LeftTribox, LeftVolume, SimplifiedConvexes, CollisionObjects);
-
-			// Add right node to the tree
-			AddTreeNode(LODIndex, RightTribox, RightVolume, SimplifiedConvexes, CollisionObjects);
-		}
-	}
-}
-
-FORCEINLINE void FTriboxTree::AddTreeNode(const int32 LODIndex, const FTriboxNode& TriboxNode, const FTribox::FRealType& ObjectsVolume,
-	TArray<FImplicitObjectPtr>& SimplifiedConvexes, const TUniquePtr<Private::FCollisionObjects>& CollisionObjects)
-{
-	if(TriboxNode.BoundsTribox.HasDatas())
-	{
-		// Disgard all the objects with small volumes
-		if(ObjectsVolume > CVars::ChaosConvexMinVolume)
-		{
-			const FTribox::FRealType NodeVolume = TriboxNode.BoundsTribox.ComputeVolume();
-			//UE_LOG(LogTemp, Log, TEXT("				Node volume = [%f %f] -> %f"), NodeVolume, ObjectsVolume, NodeVolume/ObjectsVolume)
-			
-			if((NodeVolume/ObjectsVolume > CVars::ChaosConvexSplittingThreshold) &&
-				(LODIndex < (MaxLODs-1)))
-			{
-				// If volume ratio (concavity) is big enough add this node to the current list to be processed by the next LOD
-				TreeNodes[CurrentBuffer()].Add(TriboxNode);
-			}
-			else
-			{
-				// Otherwise add the node to the final convexes list
-				SimplifiedConvexes.Add(TriboxNode.BoundsTribox.MakeConvex());
-				Private::FImplicitBVH::CollectLeafObject(SimplifiedConvexes.Last(),
-						FRigidTransform3::Identity, INDEX_NONE, CollisionObjects->ImplicitObjects);
-			}
-		}
-	}
-}
 
 FConvexOptimizer::FConvexOptimizer() :
 	SimplifiedConvexes(), CollisionObjects(MakeUnique<Private::FCollisionObjects>()), ShapesArray()
@@ -240,10 +54,10 @@ void FConvexOptimizer::VisitCollisionObjects(const FImplicitHierarchyVisitor& Vi
 	{
 		const TArray<Private::FImplicitBVHObject>& ImplicitObjects = BVH.IsValid() ? BVH->GetObjects() : CollisionObjects->ImplicitObjects;
 		
-		int32 ObjectIndex = 0;
+		int32 ObjectIndex = 1;
+		int32 LeafObjectIndex = 0;
 		for(const Private::FImplicitBVHObject& CollisionObject : ImplicitObjects)
 		{
-			int32 LeafObjectIndex = CollisionObject.GetObjectIndex();
 			CollisionObject.GetGeometry()->VisitLeafObjectsImpl(CollisionObject.GetTransform(),
 				CollisionObject.GetRootObjectIndex(), ObjectIndex, LeafObjectIndex, VisitorFunc);
 		}
@@ -254,21 +68,22 @@ void FConvexOptimizer::VisitOverlappingObjects(const FAABB3& LocalBounds, const
 {
 	if(CVars::bChaosConvexSimplifyUnion)
 	{
-		int32 ObjectIndex = 0;
+		int32 ObjectIndex = 1;
+		int32 LeafObjectIndex = 0;
 		if(BVH.IsValid())
 		{
 			BVH->VisitAllIntersections(LocalBounds,
-			[this, &ObjectIndex, &VisitorFunc](const FImplicitObject* Implicit, const FRigidTransform3f& RelativeTransformf,
-				const FAABB3f& RelativeBoundsf, const int32 RootObjectIndex, const int32 LeafObjectIndex)
+			[this, &ObjectIndex, &VisitorFunc, &LocalBounds, &LeafObjectIndex](const FImplicitObject* Implicit, const FRigidTransform3f& RelativeTransformf,
+				const FAABB3f& RelativeBoundsf, const int32 RootObjectIndex, const int32)
 			{
-				VisitorFunc(Implicit, FRigidTransform3(RelativeTransformf), RootObjectIndex, ObjectIndex++, LeafObjectIndex);
+				Implicit->VisitOverlappingLeafObjectsImpl(LocalBounds, FRigidTransform3(RelativeTransformf),
+					RootObjectIndex, ObjectIndex, LeafObjectIndex, VisitorFunc);
 			});
 		}
 		else
 		{
 			for(const Private::FImplicitBVHObject& CollisionObject : CollisionObjects->ImplicitObjects)
 			{
-				int32 LeafObjectIndex = CollisionObject.GetObjectIndex();
 				CollisionObject.GetGeometry()->VisitOverlappingLeafObjectsImpl(LocalBounds, CollisionObject.GetTransform(),
 					CollisionObject.GetRootObjectIndex(), ObjectIndex, LeafObjectIndex, VisitorFunc);
 			}
@@ -293,36 +108,35 @@ void FConvexOptimizer::BuildConvexShapes(const FShapesArray& UnionShapes)
 	ShapeData->SetSimData(TemplateData->GetSimData());
 	ShapeData->SetCollisionTraceType(TemplateData->GetCollisionTraceType());
 
-	// Only sim enabled if all the underlying shapes could be used for physics
+	// Only simulation enabled if all the underlying shapes could be used for physics
 	ShapeData->SetSimEnabled(true);
 
 	// Disable simple shapes for query 
 	ShapeData->SetQueryEnabled(false);
 }
 
-void FConvexOptimizer::SimplifyRootConvexes(const Chaos::FImplicitObjectUnionPtr& UnionGeometry, const FShapesArray& UnionShapes, const EObjectStateType ObjectState)
+void FConvexOptimizer::SimplifyRootConvexes(const Chaos::FImplicitObjectUnionPtr& UnionGeometry, const FShapesArray& UnionShapes,
+		const EObjectStateType ObjectState)
 {
 	SCOPE_CYCLE_COUNTER(STAT_Collisions_SimplifyConvexes);
+
 	if(CVars::bChaosConvexSimplifyUnion && UnionGeometry && (UnionShapes.Num() > 0) &&
 		(UnionShapes.Num() <= UnionGeometry->GetObjects().Num()))
 	{
 		SimplifiedConvexes.Reset();
 		CollisionObjects->ImplicitObjects.Reset();
 		ShapesArray.Reset();
-		const int32 MaxLODs = (ObjectState == EObjectStateType::Dynamic) ? CVars::ChaosConvexMaxDynamicLODs : CVars::ChaosConvexMaxKinematicLODs;
 
-		if(MaxLODs == 0)
+		if((UnionShapes.Num() == 1) || ((ObjectState == EObjectStateType::Dynamic) && (CVars::ChaosConvexDynamicMode == 0)) 
+				|| ((ObjectState != EObjectStateType::Dynamic) && (CVars::ChaosConvexKinematicMode == 0)))
 		{
 			BuildSingleConvex(UnionGeometry, UnionShapes);
 		}
 		else
 		{
-			BuildMultipleConvex(UnionGeometry, UnionShapes, MaxLODs);
-		}
-		if(!SimplifiedConvexes.IsEmpty())
-		{
-			// Build the simplified shapes
-			BuildConvexShapes(UnionShapes);
+			const bool bEnableMerging = ((ObjectState == EObjectStateType::Dynamic) && (CVars::ChaosConvexDynamicMode == 2))
+				|| ((ObjectState != EObjectStateType::Dynamic) && (CVars::ChaosConvexKinematicMode == 2));
+			BuildMultipleConvex(UnionGeometry, UnionShapes, bEnableMerging);
 		}
 	}
 
@@ -336,25 +150,47 @@ void FConvexOptimizer::SimplifyRootConvexes(const Chaos::FImplicitObjectUnionPtr
 	}
 }
 
-FORCEINLINE void InvalidateCachedTriboxes(FConvexOptimizer::FCachedTriboxes& RootTriboxes)
+FORCEINLINE void InvalidateCachedTriboxes(FConvexOptimizer::FTriboxNodes& RootTriboxes)
 {
 	for(auto& RootTribox : RootTriboxes)
 	{
-		RootTribox.Value.Tribox.SetValid(false);
+		RootTribox.Value.NodeTribox.SetValid(false);
 	}
 }
 
-FORCEINLINE void ResizeCachedTriboxes(FConvexOptimizer::FCachedTriboxes& RootTriboxes)
+DECLARE_CYCLE_STAT(TEXT("Collisions::ResizeCachedTriboxes"), STAT_ResizeCachedTriboxes, STATGROUP_ChaosCollision);
+FORCEINLINE void ResizeCachedTriboxes(FConvexOptimizer::FTriboxNodes& RootTriboxes)
 {
-	TArray<FImplicitObject*> KeysToRemove;
+	SCOPE_CYCLE_COUNTER(STAT_ResizeCachedTriboxes);
+
+	TArray<FImplicitObject*> NodesToRemove;
+	TArray<FImplicitObject*> EdgesToRemove;
 	for(auto& RootTribox : RootTriboxes)
 	{
-		if(!RootTribox.Value.Tribox.IsValid())
+		if(!RootTribox.Value.NodeTribox.IsValid())
 		{
-			KeysToRemove.Add(RootTribox.Key);
+			NodesToRemove.Add(RootTribox.Key);
+		}
+		else
+		{
+			EdgesToRemove.Reset();
+			for(auto& NodeEdge : RootTribox.Value.bValidEdges)
+			{
+				if(auto* EdgeNode = RootTriboxes.Find(NodeEdge.Key))
+				{
+					if(!EdgeNode->NodeTribox.IsValid())
+					{
+						EdgesToRemove.Add(NodeEdge.Key);
+					}
+				}
+			}
+			for(auto& KeyToRemove : EdgesToRemove)
+			{
+				RootTribox.Value.bValidEdges.Remove(KeyToRemove);
+			}
 		}
 	}
-	for(auto& KeyToRemove : KeysToRemove)
+	for(auto& KeyToRemove : NodesToRemove)
 	{
 		RootTriboxes.Remove(KeyToRemove);
 	}
@@ -362,7 +198,7 @@ FORCEINLINE void ResizeCachedTriboxes(FConvexOptimizer::FCachedTriboxes& RootTri
 
 DECLARE_CYCLE_STAT(TEXT("Collisions::BuildConvexTriboxes"), STAT_BuildConvexTriboxes, STATGROUP_ChaosCollision);
 FORCEINLINE void BuildConvexTriboxes(const Chaos::FImplicitObjectUnionPtr& UnionGeometry, const FShapesArray& UnionShapes,
-	TUniquePtr<Private::FCollisionObjects>& CollisionObjects, FConvexOptimizer::FCachedTriboxes& RootTriboxes)
+	TUniquePtr<Private::FCollisionObjects>& CollisionObjects, FConvexOptimizer::FTriboxNodes& RootTriboxes)
 {
 	SCOPE_CYCLE_COUNTER(STAT_BuildConvexTriboxes);
 
@@ -376,16 +212,16 @@ FORCEINLINE void BuildConvexTriboxes(const Chaos::FImplicitObjectUnionPtr& Union
 		{
 			if(FImplicitObject* RootObject = UnionGeometry->GetObjects()[RootObjectIndex].GetReference())
 			{
-				FConvexOptimizer::FCachedTribox* RootTribox = RootTriboxes.Find(RootObject);
+				FConvexOptimizer::FTriboxNode* RootTribox = RootTriboxes.Find(RootObject);
 				const bool bHasRootTribox = (RootTribox != nullptr);
+				if(bHasRootTribox) RootTribox->bValidNode = false;
 				
 				FTribox LocalTribox;
 				int32 LeafObjectIndex = 0, ObjectIndex = 0;
-				
 				RootObject->VisitLeafObjectsImpl(FRigidTransform3::Identity,RootObjectIndex, ObjectIndex, LeafObjectIndex,
-					[&CollisionObjects, &LocalTribox, &bHasRootTribox](
+					[&CollisionObjects, &LocalTribox, &bHasRootTribox, &ShapeIndex](
 					const FImplicitObject* ImplicitObject, const FRigidTransform3& RelativeTransform,
-					const int32 RootObjectIndex, const int32, const int32)
+					const int32 RootIndex, const int32, const int32)
 				{
 					if(const FConvex* Convex = ImplicitObject->AsA<FConvex>())
 					{
@@ -394,7 +230,6 @@ FORCEINLINE void BuildConvexTriboxes(const Chaos::FImplicitObjectUnionPtr& Union
 							const FTribox::FRigidTransform3Type ConvexTransform(RelativeTransform);
 							// Disable collision for all the convexes that are going to be used to build the tribox
 							// For now only used for debug draw
-							const_cast<FImplicitObject*>(ImplicitObject)->SetDoCollide(false);
 							LocalTribox.AddConvex(Convex, ConvexTransform);
 						}
 					}
@@ -406,7 +241,6 @@ FORCEINLINE void BuildConvexTriboxes(const Chaos::FImplicitObjectUnionPtr& Union
 								RelativeTransform.GetRotation(), RelativeTransform.GetScale3D() * ScaledObject->GetScale());
 							// Disable collision for all the convexes that are going to be used to build the tribox
 							// For now only used for debug draw
-							const_cast<FImplicitObject*>(ImplicitObject)->SetDoCollide(false);
 							LocalTribox.AddConvex(ScaledObject->GetUnscaledObject(), ConvexTransform);
 						}
 					}
@@ -417,7 +251,6 @@ FORCEINLINE void BuildConvexTriboxes(const Chaos::FImplicitObjectUnionPtr& Union
 							const FTribox::FRigidTransform3Type ConvexTransform(RelativeTransform);
 							// Disable collision for all the convexes that are going to be used to build the tribox
 							// For now only used for debug draw
-							const_cast<FImplicitObject*>(ImplicitObject)->SetDoCollide(false);
 							LocalTribox.AddConvex(InstancedObject->Object(), ConvexTransform);
 						}
 					}
@@ -425,17 +258,18 @@ FORCEINLINE void BuildConvexTriboxes(const Chaos::FImplicitObjectUnionPtr& Union
 					{
 						// Add non convex implicits to the list of collision objects
 						Private::FImplicitBVH::CollectLeafObject(ImplicitObject,
-							RelativeTransform, RootObjectIndex, CollisionObjects->ImplicitObjects);	  
+							RelativeTransform, ShapeIndex, CollisionObjects->ImplicitObjects);	  
 					}
 				});
 				if(!bHasRootTribox && LocalTribox.HasDatas())
 				{
 					LocalTribox.BuildTribox();
-					RootTribox = &(RootTriboxes.Add(RootObject, {LocalTribox,nullptr}));
+					RootTribox = &(RootTriboxes.Add(RootObject, {LocalTribox, nullptr, ShapeIndex, LocalTribox.ComputeVolume(),true}));
 				}
 				if(RootTribox)
 				{
-					RootTribox->Tribox.SetValid(true);
+					RootTribox->NodeTribox.SetValid(true);
+					RootTribox->ShapeIndex = ShapeIndex;
 				}
 			}
 		}
@@ -450,9 +284,9 @@ void FConvexOptimizer::BuildSingleConvex(const Chaos::FImplicitObjectUnionPtr& U
 	FTribox MainTribox;
 	for(auto& RootTribox : RootTriboxes)
 	{
-		if(RootTribox.Value.Tribox.HasDatas())
+		if(RootTribox.Value.NodeTribox.HasDatas())
 		{
-			MainTribox += RootTribox.Value.Tribox;
+			MainTribox += RootTribox.Value.NodeTribox;
 		}
 	}
 
@@ -460,54 +294,192 @@ void FConvexOptimizer::BuildSingleConvex(const Chaos::FImplicitObjectUnionPtr& U
 	{
 		// Build the tribox and add it to the list of optimized convexes
 		SimplifiedConvexes.Add(MainTribox.MakeConvex());
-		Private::FImplicitBVH::CollectLeafObject(SimplifiedConvexes.Last(), FRigidTransform3::Identity, INDEX_NONE, CollisionObjects->ImplicitObjects);	
+		Private::FImplicitBVH::CollectLeafObject(SimplifiedConvexes.Last(), FRigidTransform3::Identity, 0, CollisionObjects->ImplicitObjects);	
 	}
 }
 
-FORCEINLINE void FindClosestPlane(const FTribox::FVec3Type& PointPosition, const TArray<FTribox>& Triboxes,
-						int32& PlaneAxis, FTribox::FRealType& PlaneProjection)
+DECLARE_CYCLE_STAT(TEXT("Collisions::BuildUnionConnectivity"), STAT_BuildUnionConnectivity, STATGROUP_ChaosCollision);
+void FConvexOptimizer::BuildUnionConnectivity(const Chaos::FImplicitObjectUnionPtr& UnionGeometry)
 {
-	int32 LocalAxis = INDEX_NONE;
-	FTribox::FRealType LocalProjection = 0.0;
-	FTribox::FRealType MinDistance = FLT_MAX;
-				
-	for(auto& Tribox : Triboxes)
+	SCOPE_CYCLE_COUNTER(STAT_BuildUnionConnectivity);
+	
+	const int32 NumTriboxes = UnionGeometry->GetObjects().Num();
+	TUniquePtr<Private::FImplicitBVH> LocalBVH = nullptr;
+
+	if(NumTriboxes > CVars::ChaosUnionBVHMinShapes)
+	{
+		TArray<Private::FImplicitBVHObject> ImplicitObjects;
+		ImplicitObjects.Reserve(NumTriboxes);
+
+		for (int32 RootIndex = 0; RootIndex < NumTriboxes; ++RootIndex)
+		{
+			if (auto* RootTribox = RootTriboxes.Find(UnionGeometry->GetObjects()[RootIndex]))
+			{
+				Private::FImplicitBVH::CollectLeafObject(UnionGeometry->GetObjects()[RootIndex], FRigidTransform3::Identity, RootIndex, ImplicitObjects);
+			}
+		}
+		LocalBVH = FImplicitBVH::TryMakeFromLeaves(MoveTemp(ImplicitObjects), CVars::ChaosUnionBVHMinShapes, CVars::ChaosUnionBVHMaxDepth);
+	}
+
+	auto AddUnionEdges = [this, &UnionGeometry](const int32 RootIndexA, const int32 RootIndexB, const FAABB3& LocalBoundsA, FTriboxNode* RootTriboxA)
+	{
+		if (auto* RootTriboxB = RootTriboxes.Find(UnionGeometry->GetObjects()[RootIndexB]))
+		{
+			if(!RootTriboxB->bValidNode || (RootTriboxB->bValidNode && RootIndexA < RootIndexB))
+			{ 
+				if (LocalBoundsA.Intersects(RootTriboxB->NodeTribox.GetBounds()))
+				{
+					RootTriboxA->bValidEdges.Add(UnionGeometry->GetObjects()[RootIndexB], true);
+					RootTriboxB->bValidEdges.Add(UnionGeometry->GetObjects()[RootIndexA], true);
+				}
+			}
+		}
+	};
+	for (int32 RootIndexA = 0; RootIndexA < NumTriboxes; ++RootIndexA)
 	{
-		const FTribox::FRealType ClosestDistance = Tribox.GetClosestPlane(PointPosition, LocalAxis, LocalProjection);
-		if(ClosestDistance < MinDistance)
+		if (auto* RootTriboxA = RootTriboxes.Find(UnionGeometry->GetObjects()[RootIndexA]))
 		{
-			MinDistance = ClosestDistance;
-			PlaneAxis = LocalAxis;
-			PlaneProjection = LocalProjection;
+			if (RootTriboxA->bValidNode)
+			{
+				const FAABB3 LocalBoundsA = RootTriboxA->NodeTribox.GetBounds();
+				if (LocalBVH.IsValid())
+				{
+					LocalBVH->VisitAllIntersections(LocalBoundsA,
+					[this, &RootIndexA, &LocalBoundsA, &UnionGeometry, &RootTriboxA, &AddUnionEdges](const FImplicitObject* ImplicitObjectB, const FRigidTransform3f& RelativeTransformB,
+						const FAABB3f& RelativeBoundsB, const int32 RootIndexB, const int32 LeafObjectIndexB)
+					{
+						AddUnionEdges(RootIndexA, RootIndexB, LocalBoundsA, RootTriboxA);
+					});
+				}
+				else
+				{
+					for (int32 RootIndexB = 0; RootIndexB < NumTriboxes; ++RootIndexB)
+					{
+						AddUnionEdges(RootIndexA, RootIndexB, LocalBoundsA, RootTriboxA);
+					}
+				}
+			}
+		}
+	}
+	for (auto& RootTribox : RootTriboxes)
+	{
+		RootTribox.Value.bValidNode = true;
+
+		// TODO :  keep it for now since we must make it faster
+		//RootTribox.Value.bValidEdges.KeySort([this, &UnionGeometry](const FImplicitObject& ImplicitObjectA, 
+		//const FImplicitObject& ImplicitObjectB) -> bool{ return RootTriboxes.Find(&ImplicitObjectA)->ShapeIndex < RootTriboxes.Find(&ImplicitObjectB)->ShapeIndex;});
+		for (auto& ValidEdge : RootTribox.Value.bValidEdges)
+		{
+			ValidEdge.Value = true;
 		}
 	}
 }
 	
-void FConvexOptimizer::BuildMultipleConvex(const Chaos::FImplicitObjectUnionPtr& UnionGeometry, const FShapesArray& UnionShapes, const int32 MaxLODs)
+DECLARE_CYCLE_STAT(TEXT("Collisions::MergeConnectedShapes"), STAT_MergeConnectedShapes, STATGROUP_ChaosCollision);
+void FConvexOptimizer::MergeConnectedShapes(const Chaos::FImplicitObjectUnionPtr& UnionGeometry, TArray<FTriboxNode>& MergedNodes)
+{
+	SCOPE_CYCLE_COUNTER(STAT_MergeConnectedShapes);
+	const int32 NumTriboxes = RootTriboxes.Num();
+		
+	MergedNodes.Reserve(RootTriboxes.Num());
+	FTriboxNode* CurrentNode = nullptr;
+
+	TArray<const FImplicitObject*> NodeQueue;
+	NodeQueue.Reserve(NumTriboxes);
+
+	for (int32 RootIndex = 0; RootIndex < NumTriboxes; ++RootIndex)
+	{
+		if (auto* RootTribox = RootTriboxes.Find(UnionGeometry->GetObjects()[RootIndex]))
+		{
+			if (RootTribox->bValidNode)
+			{
+				RootTribox->bValidNode = false;
+				NodeQueue.Add(UnionGeometry->GetObjects()[RootIndex]);
+
+				MergedNodes.Add({RootTribox->NodeTribox, nullptr, RootTribox->ShapeIndex, RootTribox->NodeVolume});
+				CurrentNode = &MergedNodes.Last();
+
+				while (!NodeQueue.IsEmpty())
+				{
+					const FImplicitObject* NextNode = NodeQueue.Pop(false);
+
+					for (auto& ValidEdge : RootTriboxes.Find(NextNode)->bValidEdges)
+					{
+						if (ValidEdge.Value)
+						{
+							ValidEdge.Value = false;
+							const FImplicitObject* EdgeOtherNode = ValidEdge.Key;
+							auto* OtherNode = RootTriboxes.Find(EdgeOtherNode);
+
+							if (OtherNode->bValidNode)
+							{
+								FTribox LocalTribox = CurrentNode->NodeTribox;
+								LocalTribox += OtherNode->NodeTribox;
+
+								const FTribox::FRealType VolumeRatio = LocalTribox.ComputeVolume() / (CurrentNode->NodeVolume + OtherNode->NodeVolume);
+								if (VolumeRatio < CVars::ChaosConvexSplittingThreshold)
+								{
+									CurrentNode->NodeTribox += OtherNode->NodeTribox;
+									CurrentNode->NodeVolume += OtherNode->NodeVolume;
+
+									NodeQueue.Add(EdgeOtherNode);
+									OtherNode->bValidNode = false;
+								}
+							}
+						}
+					}
+				}
+			}
+		}
+	}
+}
+
+	
+void FConvexOptimizer::BuildMultipleConvex(const Chaos::FImplicitObjectUnionPtr& UnionGeometry, const FShapesArray& UnionShapes, const bool bEnableMerging)
 {
 	BuildConvexTriboxes(UnionGeometry, UnionShapes, CollisionObjects, RootTriboxes);
 
-	if(MaxLODs < 0)
+	if(!bEnableMerging)
 	{
 		for(auto& RootTribox : RootTriboxes)
-        {
-			if(RootTribox.Value.Tribox.HasDatas() && (RootTribox.Value.Tribox.ComputeVolume() > CVars::ChaosConvexMinVolume))
+		{
+			if(RootTribox.Value.NodeTribox.HasDatas())
 			{
-				// Build the tribox and add it to the list of optimized convexes
-				if(!RootTribox.Value.Convex)
+				// Build the Tribox and add it to the list of optimized convexes
+				if(!RootTribox.Value.TriboxConvex)
 				{
-					RootTribox.Value.Convex = RootTribox.Value.Tribox.MakeConvex();
+					RootTribox.Value.TriboxConvex = RootTribox.Value.NodeTribox.MakeConvex();
 				}
-				SimplifiedConvexes.Add(RootTribox.Value.Convex);
-				Private::FImplicitBVH::CollectLeafObject(SimplifiedConvexes.Last(), FRigidTransform3::Identity, INDEX_NONE, CollisionObjects->ImplicitObjects);
+				SimplifiedConvexes.Add(RootTribox.Value.TriboxConvex);
+				Private::FImplicitBVH::CollectLeafObject(SimplifiedConvexes.Last(), FRigidTransform3::Identity, RootTribox.Value.ShapeIndex, CollisionObjects->ImplicitObjects);
 			}
-        }
+		}
 	}
 	else
 	{
-		// Build the binary tree and add the leaves to the simplified convexes
-		FTriboxTree TriboxTree(MaxLODs, RootTriboxes);
-		TriboxTree.BuildTreeLODs(SimplifiedConvexes, CollisionObjects);
+		// Bottom-up merging to build the simplified convexes
+		TArray<FConvexOptimizer::FTriboxNode> MergedNodes;
+
+		// Create the connectivity graph
+		BuildUnionConnectivity(UnionGeometry);
+
+		// Merge the shapes
+        MergeConnectedShapes(UnionGeometry, MergedNodes);
+
+		// Add to the simplified convexes
+        for(int32 NodeIndex = 0, NumMerged = MergedNodes.Num(); NodeIndex < NumMerged; ++NodeIndex)
+        {
+            if(MergedNodes[NodeIndex].NodeTribox.HasDatas() && (MergedNodes[NodeIndex].NodeVolume > CVars::ChaosConvexMinVolume))
+            {
+            	auto* RootTribox = RootTriboxes.Find(UnionGeometry->GetObjects()[MergedNodes[NodeIndex].ShapeIndex]);
+            	if(!RootTribox->TriboxConvex || (RootTribox->TriboxConvex && (RootTribox->TriboxConvex->AsA<FConvex>()->GetVolume() != MergedNodes[NodeIndex].NodeTribox.ComputeVolume())))
+            	{
+            		RootTribox->TriboxConvex = MergedNodes[NodeIndex].NodeTribox.MakeConvex();
+				}
+            	SimplifiedConvexes.Add(RootTribox->TriboxConvex);
+            	Private::FImplicitBVH::CollectLeafObject(SimplifiedConvexes.Last(), FRigidTransform3::Identity, MergedNodes[NodeIndex].ShapeIndex, CollisionObjects->ImplicitObjects);
+            }
+        }
 	}
 }
 
diff --git a/Engine/Source/Runtime/Experimental/Chaos/Private/Chaos/PBDRigidClustering.cpp b/Engine/Source/Runtime/Experimental/Chaos/Private/Chaos/PBDRigidClustering.cpp
--- a/Engine/Source/Runtime/Experimental/Chaos/Private/Chaos/PBDRigidClustering.cpp
+++ b/Engine/Source/Runtime/Experimental/Chaos/Private/Chaos/PBDRigidClustering.cpp
@@ -70,6 +70,11 @@ namespace Chaos
 	int32 ClusteringParticleReleaseThrottlingMaxCount = INDEX_NONE;
 	FAutoConsoleVariableRef CVarClusteringParticleReleaseThrottlingMaxCount(TEXT("p.Clustering.ParticleReleaseThrottlingMaxCount"), ClusteringParticleReleaseThrottlingMaxCount, TEXT("Maximum number of active geometry collection to reach before all released clustering disable all released particle instantly"));
 
+	namespace CVars
+	{
+		extern bool bChaosConvexSimplifyUnion;
+	}
+
 	template <typename TProxy=FGeometryCollectionPhysicsProxy>
 	TProxy* GetConcreteProxy(FPBDRigidClusteredParticleHandle* ClusteredParticle)
 	{
@@ -290,11 +295,11 @@ namespace Chaos
 		UpdateKinematicProperties(NewParticle, MChildren, MEvolution);
 		UpdateGeometry(NewParticle, ChildrenSet, MChildren, ProxyGeometry, Parameters);
 
-		// Build the convex optimizer if required
-		FRigidClustering::BuildConvexOptimizer(NewParticle);
-		
 		GenerateConnectionGraph(NewParticle, Parameters);
 
+		// Build the convex optimizer if required
+		//FRigidClustering::BuildConvexOptimizer(NewParticle);
+
 		NewParticle->SetSleeping(bClusterIsAsleep);
 
 		auto AddToClusterUnion = [this](int32 ClusterID, FPBDRigidClusteredParticleHandle* Handle)
@@ -378,14 +383,26 @@ namespace Chaos
 
 	void FRigidClustering::BuildConvexOptimizer(FPBDRigidClusteredParticleHandle* Particle)
 	{
-		if(Particle && Particle->GetGeometry())
+		bool bHasOptimizer = false;
+		if(Particle && Particle->GetGeometry() && CVars::bChaosConvexSimplifyUnion)
 		{ 
-			if (!Particle->ConvexOptimizer())
+			if(FImplicitObjectUnion* Union = Particle->GetGeometry()->template AsA<FImplicitObjectUnion>())
 			{
-				Particle->ConvexOptimizer() = MakePimpl<Private::FConvexOptimizer>();
+				if(Union->GetNumLeafObjects() > 1)
+				{ 
+					if (!Particle->ConvexOptimizer())
+					{
+						Particle->ConvexOptimizer() = MakePimpl<Private::FConvexOptimizer>();
+					}
+					Particle->ConvexOptimizer()->SimplifyRootConvexes(Union,
+						Particle->ShapesArray(), Particle->ObjectState());
+					bHasOptimizer = Particle->ConvexOptimizer()->IsValid();
+				}
 			}
-			Particle->ConvexOptimizer()->SimplifyRootConvexes(Particle->GetGeometry()->template AsA<FImplicitObjectUnion>(),
-				Particle->ShapesArray(), Particle->ObjectState());
+		}
+		if(Particle && !bHasOptimizer)
+		{	
+			Particle->ConvexOptimizer().Reset();
 		}
 	}
 
diff --git a/Engine/Source/Runtime/Experimental/Chaos/Private/Chaos/Tribox.cpp b/Engine/Source/Runtime/Experimental/Chaos/Private/Chaos/Tribox.cpp
--- a/Engine/Source/Runtime/Experimental/Chaos/Private/Chaos/Tribox.cpp
+++ b/Engine/Source/Runtime/Experimental/Chaos/Private/Chaos/Tribox.cpp
@@ -398,66 +398,16 @@ FORCEINLINE void CompressFaces(TArray<TArray<int32>>& FaceIndices)
 FImplicitObjectPtr FTribox::CreateConvexFromTopology(
 	TArray<FConvex::FPlaneType>&& ConvexPlanes, TArray<TArray<int32>>&& FaceIndices, TArray<FConvex::FVec3Type>&& ConvexVertices) const
 {
-	FRotation3 RotationMatrix = FRotation3::Identity;
-	FVec3Type InertiaTensor(1.0f, 1.0f, 1.0f);
-	FRealType MinVolume = FLT_MAX;
-	{
-		const FRealType HX = (MaxDists[0]+MaxDists[1]);
-		const FRealType PX = (MaxDists[6]+MaxDists[7]);
-		const FRealType MX = (MaxDists[8]+MaxDists[9]);
-
-		const FRealType HY = (MaxDists[2]+MaxDists[3]);
-		const FRealType PY = (MaxDists[10]+MaxDists[11]);
-		const FRealType MY = (MaxDists[12]+MaxDists[13]);
-
-		const FRealType HZ = (MaxDists[4]+MaxDists[5]);
-		const FRealType PZ = (MaxDists[14]+MaxDists[15]);
-		const FRealType MZ = (MaxDists[16]+MaxDists[17]);
-		
-		FRealType Volumes[4];
-		Volumes[0] = HX * HY * HZ;
-		Volumes[1] = HX * PX * MX;
-		Volumes[2] = HY * PY * MY;
-		Volumes[3] = HZ * PZ * MZ;
-		
-		int8 MinIndex = INDEX_NONE;
-		for(int8 VolumeIndex = 0; VolumeIndex < 4; ++VolumeIndex)
-		{
-			if(Volumes[VolumeIndex] < MinVolume)
-			{
-				MinVolume = Volumes[VolumeIndex];
-				MinIndex = VolumeIndex;
-			}
-		}
-		
-		if(MinIndex == 0)
-		{
-			static const FRotation3 ConstRotationMatrix = FRotation3::Identity;
-			RotationMatrix = ConstRotationMatrix;
-			InertiaTensor = FVec3Type(HY * HY + HZ * HZ, HX * HX + HZ * HZ, HY * HY + HX * HX) / (12.0f * Volumes[0]);
-		}
-		else if(MinIndex == 1)
-		{
-			static const FRotation3 ConstRotationMatrix = FRotation3::FromAxisAngle(FVec3(1.0f, 0.0f, 0.0f), PI / 4.0f);
-			RotationMatrix = ConstRotationMatrix;
-			InertiaTensor = FVec3Type(PX * PX + MX * MX, HX * HX + MX * MX, PX * PX + HX * HX) / (12.0f * Volumes[1]);
-		}
-		else if(MinIndex == 2)
-		{
-			static const FRotation3 ConstRotationMatrix = FRotation3::FromAxisAngle(FVec3(0.0f, 1.0f, 0.0f), PI / 4.0f);
-			RotationMatrix = ConstRotationMatrix;
-			InertiaTensor = FVec3Type(MY * MY + HY * HY, MY * MY + PY * PY, HY * HY + PY * PY) / (12.0f * Volumes[2]);
-		}
-		else if(MinIndex == 3)
-		{
-			static const FRotation3 ConstRotationMatrix = FRotation3::FromAxisAngle(FVec3(0.0f, 0.0f, 1.0f), PI / 4.0f);
-			RotationMatrix = ConstRotationMatrix;
-			InertiaTensor = FVec3Type(MZ * MZ + HZ * HZ, PZ * PZ + HZ * HZ, PZ * PZ + MZ * MZ) / (12.0f * Volumes[3]);
-		}
-	}
+	const FRealType HX = (MaxDists[0]+MaxDists[1]);
+	const FRealType HY = (MaxDists[2]+MaxDists[3]);
+	const FRealType HZ = (MaxDists[4]+MaxDists[5]);
+	
+	const FVec3Type InertiaTensor = FVec3Type(HY * HY + HZ * HZ, HX * HX + HZ * HZ, HY * HY + HX * HX) / (12.0f);
+	const FRealType Volume = HX * HY * HZ;
+	
 	return MakeImplicitObjectPtr<Chaos::FConvex>(MoveTemp(ConvexPlanes), MoveTemp(FaceIndices), MoveTemp(ConvexVertices), 
 		FConvex::FVec3Type(-MaxDists[1], -MaxDists[3], -MaxDists[5]), FConvex::FVec3Type(MaxDists[0], MaxDists[2], MaxDists[4]), 
-		MinVolume, InertiaTensor, RotationMatrix, true);
+		Volume, InertiaTensor, FRotation3::Identity, true);
 }
 
 DECLARE_CYCLE_STAT(TEXT("Collisions::MakeConvex"), STAT_MakeTriboxConvex, STATGROUP_ChaosCollision);
@@ -532,6 +482,11 @@ bool FTribox::IsTriboxOverlapping(const FTribox& OtherTribox) const
 	}
 	return true;
 }
+
+FAABB3 FTribox::GetBounds() const
+{
+	return FAABB3( FVector(-MaxDists[1], -MaxDists[3], -MaxDists[5]), FVector(MaxDists[0], MaxDists[2], MaxDists[4]));
+}
 	
 bool FTribox::OverlapTribox(const FTribox& OtherTribox, FTribox& OverlapTribox) const
 {
@@ -661,28 +616,6 @@ FTribox::FRealType FTribox::ComputeVolume() const
 	const FVec3Type BoxSize(MaxDists[0]+MaxDists[1], MaxDists[2]+MaxDists[3], MaxDists[4]+MaxDists[5]);
 	FRealType Volume = BoxSize[0] * BoxSize[1] * BoxSize[2];
 
-	auto RemovePrism = [&Volume, &BoxSize, this](const int32 AxisX, const int32 AxisY, const int32 AxisZ,
-		const int32 ChamferA, const int32 ChamferB, const int32 ChamferC, const int32 ChamferD)
-	{
-		FRealType CornerDistance = (MaxDists[AxisX] + MaxDists[AxisY]) * InvSqrt2 - MaxDists[ChamferA];
-		Volume -= CornerDistance * CornerDistance * BoxSize[AxisZ];
-		CornerDistance = (MaxDists[AxisX] + MaxDists[AxisY+1]) * InvSqrt2 - MaxDists[ChamferB];
-		Volume -= CornerDistance * CornerDistance * BoxSize[AxisZ];
-		CornerDistance = (MaxDists[AxisX+1] + MaxDists[AxisY]) * InvSqrt2 - MaxDists[ChamferC];
-		Volume -= CornerDistance * CornerDistance * BoxSize[AxisZ];
-		CornerDistance = (MaxDists[AxisX+1] + MaxDists[AxisY+1]) * InvSqrt2 - MaxDists[ChamferD];
-		Volume -= CornerDistance * CornerDistance * BoxSize[AxisZ];
-	};
-
-	// Remove all th prisms along the z directions
-	RemovePrism(0,2,2,14,16,17,15);
-	
-	// Remove all th prisms along the y directions
-	RemovePrism(0,4,1,10,12,13,11);
-	
-	// Remove all th prisms along the x directions
-	RemovePrism(2,4,0,6,8,9,7);
-	
 	return Volume;
 }
 	
diff --git a/Engine/Source/Runtime/Experimental/Chaos/Public/Chaos/ConvexOptimizer.h b/Engine/Source/Runtime/Experimental/Chaos/Public/Chaos/ConvexOptimizer.h
--- a/Engine/Source/Runtime/Experimental/Chaos/Public/Chaos/ConvexOptimizer.h
+++ b/Engine/Source/Runtime/Experimental/Chaos/Public/Chaos/ConvexOptimizer.h
@@ -25,16 +25,30 @@ namespace Private
 	{
 		public :
 		
-		struct FCachedTribox
+		struct FTriboxNode
 		{
+			using FValidEdges = TMap<FImplicitObject*, bool>;
+			
 			// Cached tribox
-			Private::FTribox Tribox;
+			Private::FTribox NodeTribox;
 
 			// Cached convex
-			FImplicitObjectPtr Convex;
+			FImplicitObjectPtr TriboxConvex = nullptr;
+
+			// Shape Index
+			int32 ShapeIndex = INDEX_NONE;
+			
+			// Node volume
+			FTribox::FRealType NodeVolume = 0.0f;
+
+			// Valid flag to check if the node has already been processed during merging
+			bool bValidNode = true;
+
+			// Valid flag to check if the edge has already been processed during merging
+			FValidEdges bValidEdges = {};
 		};
 
-		using FCachedTriboxes = TMap<FImplicitObject*,FCachedTribox>;
+		using FTriboxNodes = TMap<FImplicitObject*, FTriboxNode>;
 		
 		CHAOS_API FConvexOptimizer();
 
@@ -42,7 +56,8 @@ namespace Private
 		CHAOS_API ~FConvexOptimizer();
 
 		// Simplify all the convexes in the hierarchy 
-		CHAOS_API void SimplifyRootConvexes(const Chaos::FImplicitObjectUnionPtr& UnionGeometry, const FShapesArray& UnionShapes, const EObjectStateType ObjectState);
+		CHAOS_API void SimplifyRootConvexes(const Chaos::FImplicitObjectUnionPtr& UnionGeometry, const FShapesArray& UnionShapes,
+			const EObjectStateType ObjectState);
 
 		// Check if the manager is valid or not
 		CHAOS_API bool IsValid() const {return !SimplifiedConvexes.IsEmpty();}
@@ -61,11 +76,17 @@ namespace Private
 
 	private:
 
+		// Build union connectivity for merging
+		void BuildUnionConnectivity(const Chaos::FImplicitObjectUnionPtr& UnionGeometry);
+
+		// Merge the connected shapes
+		void MergeConnectedShapes(const Chaos::FImplicitObjectUnionPtr& UnionGeometry, TArray<FTriboxNode>& MergedNodes);
+
 		// Build a single convex
 		void BuildSingleConvex(const Chaos::FImplicitObjectUnionPtr& UnionGeometry, const FShapesArray& UnionShapes);
 
 		// Build several convexes 
-		void BuildMultipleConvex(const Chaos::FImplicitObjectUnionPtr& UnionGeometry, const FShapesArray& UnionShapes, const int32 MaxLODs);
+		void BuildMultipleConvex(const Chaos::FImplicitObjectUnionPtr& UnionGeometry, const FShapesArray& UnionShapes, const bool bEnableMerging);
 
 		// Build the simplified shapes
 		void BuildConvexShapes(const FShapesArray& UnionShapes);
@@ -80,7 +101,7 @@ namespace Private
 		FShapeInstanceArray ShapesArray;
 
 		// Intermediate root triboxes to reuse the intermediate computation
-		FCachedTriboxes RootTriboxes;
+		FTriboxNodes RootTriboxes;
 
 		// BVH used to accelerate the collisions queries
 		TUniquePtr<Private::FImplicitBVH> BVH;
diff --git a/Engine/Source/Runtime/Experimental/Chaos/Public/Chaos/PBDRigidClustering.h b/Engine/Source/Runtime/Experimental/Chaos/Public/Chaos/PBDRigidClustering.h
--- a/Engine/Source/Runtime/Experimental/Chaos/Public/Chaos/PBDRigidClustering.h
+++ b/Engine/Source/Runtime/Experimental/Chaos/Public/Chaos/PBDRigidClustering.h
@@ -402,7 +402,7 @@ public:
 	*    all the convexes within the hierarchy
 	*    @param Particle particle on which the geometry will be simplified   
 	*/
-	static CHAOS_API void BuildConvexOptimizer(FPBDRigidClusteredParticleHandle* Particle);
+	CHAOS_API void BuildConvexOptimizer(FPBDRigidClusteredParticleHandle* Particle);
 	
 	static CHAOS_API bool ShouldUnionsHaveCollisionParticles();
 
diff --git a/Engine/Source/Runtime/Experimental/Chaos/Public/Chaos/Tribox.h b/Engine/Source/Runtime/Experimental/Chaos/Public/Chaos/Tribox.h
--- a/Engine/Source/Runtime/Experimental/Chaos/Public/Chaos/Tribox.h
+++ b/Engine/Source/Runtime/Experimental/Chaos/Public/Chaos/Tribox.h
@@ -62,6 +62,9 @@ public :
 	
 	// Get the tribox center
 	FVec3Type GetCenter() const;
+
+	// Get the bounding box
+	FAABB3 GetBounds() const;
 	
 	// Get the closest plane along the +X,-X,+Y,-Y,+Z,-Z directions
 	FRealType GetClosestPlane(const FVec3Type& PointPosition, int32& PlaneAxis, FRealType& PlaneProjection) const;
