Chaos - fix sleeping with a threshold of 1

The sleep metrics (VSmooth) were updated UpdatePositions which is after the sleep check is processed, so sleep checks did not take the current frame's impulses into account. E.g., that an object with a sleep counter of 1 could go to sleep immediately on spawn, even if it has collisions that should pop it out of the ground/wall etc.

diff --git a/Engine/Source/Runtime/Experimental/Chaos/Private/Chaos/Island/IslandManager.cpp b/Engine/Source/Runtime/Experimental/Chaos/Private/Chaos/Island/IslandManager.cpp
--- a/Engine/Source/Runtime/Experimental/Chaos/Private/Chaos/Island/IslandManager.cpp
+++ b/Engine/Source/Runtime/Experimental/Chaos/Private/Chaos/Island/IslandManager.cpp
@@ -39,6 +39,7 @@ namespace Chaos::CVars
 	extern int32 ChaosSolverCollisionPositionShockPropagationIterations;
 	extern int32 ChaosSolverCollisionVelocityShockPropagationIterations;
 	extern bool bChaosSolverPersistentGraph;
+	extern FRealSingle SmoothedPositionLerpRate;
 
 	bool bChaosConstraintGraphValidate = (CHAOS_CONSTRAINTGRAPH_CHECK_ENABLED != 0);
 	FAutoConsoleVariableRef CVarChaosConstraintGraphValidate(TEXT("p.Chaos.ConstraintGraph.Validate"), bChaosConstraintGraphValidate, TEXT("Enable per-tick ConstraintGraph validation checks/assertions"));
@@ -186,6 +187,18 @@ namespace Chaos::Private
 		return false;
 	}
 
+	template<typename TRigidParticleHandle>
+	void UpdateParticleSleepMetrics(TRigidParticleHandle& Rigid, FReal Dt)
+	{
+		if (Dt > UE_SMALL_NUMBER)
+		{
+			const FReal SmoothRate = FMath::Clamp(CVars::SmoothedPositionLerpRate, 0.0f, 1.0f);
+			const FVec3 VImp = FVec3::CalculateVelocity(Rigid.X(), Rigid.P(), Dt);
+			const FVec3 WImp = FRotation3::CalculateAngularVelocity(Rigid.R(), Rigid.Q(), Dt);
+			Rigid.SetVSmooth(FMath::Lerp(Rigid.VSmooth(), VImp, SmoothRate));
+			Rigid.SetWSmooth(FMath::Lerp(Rigid.WSmooth(), WImp, SmoothRate));
+		}
+	}
 
 	//////////////////////////////////////////////////////////////////////////////////////////////////
 	//////////////////////////////////////////////////////////////////////////////////////////////////
@@ -820,9 +833,9 @@ namespace Chaos::Private
 		ProcessIslands();
 	}
 
-	void FPBDIslandManager::UpdateSleep()
+	void FPBDIslandManager::UpdateSleep(const FReal Dt)
 	{
-		ProcessSleep();
+		ProcessSleep(FRealSingle(Dt));
 	}
 
 	void FPBDIslandManager::UpdateDisable(TFunctionRef<void(FPBDRigidParticleHandle*)> ParticleDisableFunctor)
@@ -1936,7 +1949,7 @@ namespace Chaos::Private
 		Validate();
 	}
 
-	void FPBDIslandManager::ProcessSleep()
+	void FPBDIslandManager::ProcessSleep(const FRealSingle Dt)
 	{
 		if (!CVars::bChaosSolverSleepEnabled)
 		{
@@ -1944,7 +1957,7 @@ namespace Chaos::Private
 		}
 
 		// Isloated particles are not kept in any island and need to be handled separately.
-		ProcessParticlesSleep();
+		ProcessParticlesSleep(Dt);
 
 		// @todo(chaos): can go wide except for PropagateSleepState
 		for (FPBDIsland* Island : Islands)
@@ -1952,7 +1965,7 @@ namespace Chaos::Private
 			if (!Island->Flags.bIsSleeping && !!Island->Flags.bIsSleepAllowed && !Island->Flags.bIsUsingCache)
 			{
 				// Update the sleep state based on particle movement etc
-				ProcessIslandSleep(Island);
+				ProcessIslandSleep(Island, Dt);
 
 				if (Island->Flags.bIsSleeping)
 				{
@@ -1965,21 +1978,17 @@ namespace Chaos::Private
 		}
 	}
 
-	void FPBDIslandManager::ProcessParticlesSleep()
+	void FPBDIslandManager::ProcessParticlesSleep(const FRealSingle Dt)
 	{
-		// We only need to process particles that have zero gravity because under gravity
-		// the particles will only sleep when held in place by a constraint and that will
-		// be handled in ProcessIslandSleep.
-		// @todo(chaos): keep track of particles with zero gravity? There usually are very few
-		// so it's a shame to have to visit all isolated particles here
+		// Check the sleepiness of particles that are not in any islands.
+		// @todo(chaos): this is very expensive because we have to search for a material in GetParticleSleepThresholds.
+		// We should probably cache a particle's sleep (and disable) thresholds somewhere.
 		TArray<FGeometryParticleHandle*> SleptParticles;
 		TArray<FGeometryParticleHandle*> DisabledParticles;
 		for (FTransientPBDRigidParticleHandle& Rigid : Particles.GetActiveDynamicMovingKinematicParticlesView())
 		{
 			if (Rigid.IsDynamic() && !Rigid.IsInConstraintGraph())
 			{
-				// @todo(chaos): this is very expensive because we have to search for a material
-				// We should probably cache a particle's sleep (and disable) thresholds somewhere.
 				FRealSingle SleepLinearThreshold, SleepAngularThreshold;
 				int32 SleepCounterThreshold;
 				GetParticleSleepThresholds(Rigid.Handle(), PhysicsMaterials, PerParticlePhysicsMaterials, SimMaterials, SleepLinearThreshold, SleepAngularThreshold, SleepCounterThreshold);
@@ -1987,6 +1996,8 @@ namespace Chaos::Private
 				// Check for sleep
 				if ((SleepLinearThreshold > 0) || (SleepAngularThreshold > 0))
 				{
+					UpdateParticleSleepMetrics(Rigid, Dt);
+
 					int32 SleepCounter = 0;
 					if (SleepCounterThreshold < TNumericLimits<int32>::Max())
 					{
@@ -2021,7 +2032,7 @@ namespace Chaos::Private
 		}
 	}
 
-	void FPBDIslandManager::ProcessIslandSleep(FPBDIsland* Island)
+	void FPBDIslandManager::ProcessIslandSleep(FPBDIsland* Island, const FRealSingle Dt)
 	{
 		bool bWithinSleepThreshold = true;
 		int32 SleepCounterThreshold = 0;
@@ -2035,11 +2046,17 @@ namespace Chaos::Private
 			}
 
 			// Check the particle state against the thresholds
-			FConstGenericParticleHandle P = Node->GetParticle();
+			FPBDRigidParticleHandle* Rigid = Node->GetParticle()->CastToRigidParticle();
+			if (Rigid == nullptr)
+			{
+				continue;
+			}
+			
+			UpdateParticleSleepMetrics(*Rigid, Dt);
 
 			// Did we exceed the velocity threshold?
-			if ((P->VSmooth().SizeSquared() > Node->SleepLinearThresholdSq) 
-				|| (P->WSmooth().SizeSquared() > Node->SleepAngularThresholdSq))
+			if ((Rigid->VSmooth().SizeSquared() > Node->SleepLinearThresholdSq)
+				|| (Rigid->WSmooth().SizeSquared() > Node->SleepAngularThresholdSq))
 			{
 				bWithinSleepThreshold = false;
 				break;
diff --git a/Engine/Source/Runtime/Experimental/Chaos/Private/Chaos/PBDRigidsEvolutionGBF.cpp b/Engine/Source/Runtime/Experimental/Chaos/Private/Chaos/PBDRigidsEvolutionGBF.cpp
--- a/Engine/Source/Runtime/Experimental/Chaos/Private/Chaos/PBDRigidsEvolutionGBF.cpp
+++ b/Engine/Source/Runtime/Experimental/Chaos/Private/Chaos/PBDRigidsEvolutionGBF.cpp
@@ -647,7 +647,7 @@ void FPBDRigidsEvolutionGBF::AdvanceOneTimeStepImpl(const FReal Dt, const FSubSt
 		CSV_SCOPED_TIMING_STAT(PhysicsVerbose, StepSolver_DeactivateSleep);
 
 		// Put any stationary islands to sleep (based on material settings)
-		GetIslandManager().UpdateSleep();
+		GetIslandManager().UpdateSleep(Dt);
 		
 		// Disable stationary particles (based on material settings)
 		GetIslandManager().UpdateDisable([this](FPBDRigidParticleHandle* Rigid) { DisableParticle(Rigid); });
@@ -776,15 +776,6 @@ FPBDRigidsEvolutionGBF::FPBDRigidsEvolutionGBF(
 	{
 		ParticlesInput.ParallelFor([&](auto& Particle, int32 Index)
 		{
-			if (Dt > UE_SMALL_NUMBER)
-			{
-				const FReal SmoothRate = FMath::Clamp(SmoothedPositionLerpRate, 0.0f, 1.0f);
-				const FVec3 VImp = FVec3::CalculateVelocity(Particle.X(), Particle.P(), Dt);
-				const FVec3 WImp = FRotation3::CalculateAngularVelocity(Particle.R(), Particle.Q(), Dt);
-				Particle.VSmooth() = FMath::Lerp(Particle.VSmooth(), VImp, SmoothRate);
-				Particle.WSmooth() = FMath::Lerp(Particle.WSmooth(), WImp, SmoothRate);
-			}
-
 			Particle.X() = Particle.P();
 			Particle.R() = Particle.Q();
 
diff --git a/Engine/Source/Runtime/Experimental/Chaos/Public/Chaos/Island/IslandManager.h b/Engine/Source/Runtime/Experimental/Chaos/Public/Chaos/Island/IslandManager.h
--- a/Engine/Source/Runtime/Experimental/Chaos/Public/Chaos/Island/IslandManager.h
+++ b/Engine/Source/Runtime/Experimental/Chaos/Public/Chaos/Island/IslandManager.h
@@ -424,7 +424,7 @@ namespace Chaos::Private
 		CHAOS_API void SetAssignLevels(const bool bInAssignLevels);
 		CHAOS_API void UpdateParticles();
 		CHAOS_API void UpdateIslands();
-		CHAOS_API void UpdateSleep();
+		CHAOS_API void UpdateSleep(const FReal Dt = 0);
 		CHAOS_API void UpdateDisable(TFunctionRef<void(FPBDRigidParticleHandle*)> ParticleDisableFunctor);
 		CHAOS_API void EndTick();
 
@@ -531,9 +531,9 @@ namespace Chaos::Private
 		CHAOS_API void SortIslandEdges(FPBDIsland* Island);
 
 		// Sleeping
-		CHAOS_API void ProcessSleep();
-		CHAOS_API void ProcessParticlesSleep();
-		CHAOS_API void ProcessIslandSleep(FPBDIsland* Island);
+		CHAOS_API void ProcessSleep(const FRealSingle Dt);
+		CHAOS_API void ProcessParticlesSleep(const FRealSingle Dt);
+		CHAOS_API void ProcessIslandSleep(FPBDIsland* Island, const FRealSingle Dt);
 		CHAOS_API void PropagateIslandSleep(FPBDIsland* Island);
 
 		// Disabling
