Chaos - GeometryCollections now use InitialOverlapDepenetrationVelocity from their BodyInstance during initialization

diff --git a/Engine/Source/Runtime/Engine/Classes/PhysicsEngine/BodyInstance.h b/Engine/Source/Runtime/Engine/Classes/PhysicsEngine/BodyInstance.h
--- a/Engine/Source/Runtime/Engine/Classes/PhysicsEngine/BodyInstance.h
+++ b/Engine/Source/Runtime/Engine/Classes/PhysicsEngine/BodyInstance.h
@@ -880,6 +880,9 @@ public:
 	*/
 	ENGINE_API void SetMaxDepenetrationVelocity(float MaxVelocity);
 
+	/** The maximum velocity at which initally-overlapping bodies will separate. Does not affect normal contact resolution. */
+	ENGINE_API float GetMaxDepenetrationVelocity() const { return MaxDepenetrationVelocity; }
+
 	/** Set whether we should get a notification about physics collisions */
 	ENGINE_API void SetInstanceNotifyRBCollision(bool bNewNotifyCollision);
 	/** Enables/disables whether this body is affected by gravity. */
diff --git a/Engine/Source/Runtime/Experimental/Chaos/Private/Chaos/ClusterUnionManager.cpp b/Engine/Source/Runtime/Experimental/Chaos/Private/Chaos/ClusterUnionManager.cpp
--- a/Engine/Source/Runtime/Experimental/Chaos/Private/Chaos/ClusterUnionManager.cpp
+++ b/Engine/Source/Runtime/Experimental/Chaos/Private/Chaos/ClusterUnionManager.cpp
@@ -655,6 +655,8 @@ namespace Chaos
 			{
 				Cluster->InternalCluster->SetGravityGroupIndex(FinalParticlesToAdd[0]->GravityGroupIndex());
 			}
+
+			Cluster->InternalCluster->SetInitialOverlapDepenetrationVelocity(FinalParticlesToAdd[0]->InitialOverlapDepenetrationVelocity());
 		}
 	}
 
diff --git a/Engine/Source/Runtime/Experimental/Chaos/Private/PhysicsProxy/GeometryCollectionPhysicsProxy.cpp b/Engine/Source/Runtime/Experimental/Chaos/Private/PhysicsProxy/GeometryCollectionPhysicsProxy.cpp
--- a/Engine/Source/Runtime/Experimental/Chaos/Private/PhysicsProxy/GeometryCollectionPhysicsProxy.cpp
+++ b/Engine/Source/Runtime/Experimental/Chaos/Private/PhysicsProxy/GeometryCollectionPhysicsProxy.cpp
@@ -1465,6 +1465,7 @@ void FGeometryCollectionPhysicsProxy::InitializeBodiesPT(Chaos::FPBDRigidsSolver
 				Handle->SetInertiaConditioningEnabled(Parameters.UseInertiaConditioning);
 				Handle->SetLinearEtherDrag(Parameters.LinearDamping);
 				Handle->SetAngularEtherDrag(Parameters.AngularDamping);
+				Handle->SetInitialOverlapDepenetrationVelocity(Parameters.InitialOverlapDepenetrationVelocity);
 
 #if !(UE_BUILD_SHIPPING || UE_BUILD_TEST)
 				Handle->SetDebugName(MakeShared<FString, ESPMode::ThreadSafe>(FString::Printf(TEXT("%s-%d"), *Parameters.Name, HandleIndex)));
diff --git a/Engine/Source/Runtime/Experimental/Chaos/Public/GeometryCollection/GeometryCollectionSimulationCoreTypes.h b/Engine/Source/Runtime/Experimental/Chaos/Public/GeometryCollection/GeometryCollectionSimulationCoreTypes.h
--- a/Engine/Source/Runtime/Experimental/Chaos/Public/GeometryCollection/GeometryCollectionSimulationCoreTypes.h
+++ b/Engine/Source/Runtime/Experimental/Chaos/Public/GeometryCollection/GeometryCollectionSimulationCoreTypes.h
@@ -216,6 +216,7 @@ struct FSimulationParameters
 		, bIsOneWayInteraction(false)
 		, LinearDamping(0.01f)
 		, AngularDamping(0)
+		, InitialOverlapDepenetrationVelocity(-1.0f)
 		, bUseDamagePropagation(false)
 		, BreakDamagePropagationFactor(1.0f)
 		, ShockDamagePropagationFactor(0.0f)
@@ -276,6 +277,7 @@ struct FSimulationParameters
 		, bIsOneWayInteraction(Other.bIsOneWayInteraction)
 		, LinearDamping(Other.LinearDamping)
 		, AngularDamping(Other.AngularDamping)
+		, InitialOverlapDepenetrationVelocity(Other.InitialOverlapDepenetrationVelocity)
 		, bUseDamagePropagation(Other.bUseDamagePropagation)
 		, BreakDamagePropagationFactor(Other.BreakDamagePropagationFactor)
 		, ShockDamagePropagationFactor(Other.ShockDamagePropagationFactor)
@@ -366,6 +368,7 @@ struct FSimulationParameters
 	bool bIsOneWayInteraction;
 	float LinearDamping;
 	float AngularDamping;
+	float InitialOverlapDepenetrationVelocity;
 
 	bool bUseDamagePropagation;
 	float BreakDamagePropagationFactor;
diff --git a/Engine/Source/Runtime/Experimental/GeometryCollectionEngine/Private/GeometryCollection/GeometryCollectionComponent.cpp b/Engine/Source/Runtime/Experimental/GeometryCollectionEngine/Private/GeometryCollection/GeometryCollectionComponent.cpp
--- a/Engine/Source/Runtime/Experimental/GeometryCollectionEngine/Private/GeometryCollection/GeometryCollectionComponent.cpp
+++ b/Engine/Source/Runtime/Experimental/GeometryCollectionEngine/Private/GeometryCollection/GeometryCollectionComponent.cpp
@@ -3796,6 +3796,7 @@ void UGeometryCollectionComponent::RegisterAndInitializePhysicsProxy()
 		SimulationParameters.UseCCD = BodyInstance.bUseCCD;
 		SimulationParameters.LinearDamping = BodyInstance.LinearDamping;
 		SimulationParameters.AngularDamping = BodyInstance.AngularDamping;
+		SimulationParameters.InitialOverlapDepenetrationVelocity = BodyInstance.GetMaxDepenetrationVelocity();
 		SimulationParameters.bUseDamagePropagation = DamagePropagationData.bEnabled;
 		SimulationParameters.BreakDamagePropagationFactor = DamagePropagationData.BreakDamagePropagationFactor;
 		SimulationParameters.ShockDamagePropagationFactor = DamagePropagationData.ShockDamagePropagationFactor;
