Ias: Don't issue http requests for chunk requests that have already been cancelled.
Avoids building up a giant queue of http requests in the locker when switching between skins that have not been completely loaded.

#rnx

diff --git a/Engine/Source/Runtime/Experimental/IoStoreOnDemand/Private/OnDemandIoDispatcherBackend.cpp b/Engine/Source/Runtime/Experimental/IoStoreOnDemand/Private/OnDemandIoDispatcherBackend.cpp
--- a/Engine/Source/Runtime/Experimental/IoStoreOnDemand/Private/OnDemandIoDispatcherBackend.cpp
+++ b/Engine/Source/Runtime/Experimental/IoStoreOnDemand/Private/OnDemandIoDispatcherBackend.cpp
@@ -2240,7 +2240,12 @@ void FOnDemandIoBackend::ProcessHttpRequests(FHttpClient& HttpClient, FBitWindow
 
 				Stats.OnHttpDequeue();
 
-				if (BackendStatus.IsHttpEnabled() == false)
+				if (ChunkRequest->bCancelled)
+				{
+					CompleteRequest(ChunkRequest);
+					Stats.OnHttpCancel();
+				}
+				else if (BackendStatus.IsHttpEnabled() == false)
 				{
 					UE::Tasks::Launch(UE_SOURCE_LOCATION, [this, ChunkRequest]()
 					{
diff --git a/Engine/Source/Runtime/Experimental/IoStoreOnDemand/Private/Statistics.cpp b/Engine/Source/Runtime/Experimental/IoStoreOnDemand/Private/Statistics.cpp
--- a/Engine/Source/Runtime/Experimental/IoStoreOnDemand/Private/Statistics.cpp
+++ b/Engine/Source/Runtime/Experimental/IoStoreOnDemand/Private/Statistics.cpp
@@ -264,6 +264,7 @@ FCounterAtomicInt		GCacheRejectBytes(TEXT("Ias/CachePutRejectBytes"), TraceCount
 FCounterInt				GHttpGetCount(TEXT("Ias/HttpGetCount"), TraceCounterDisplayHint_None);
 FCounterInt				GHttpErrorCount(TEXT("Ias/HttpErrorCount"), TraceCounterDisplayHint_None);
 FCounterInt				GHttpRetryCount(TEXT("Ias/HttpRetryCount"), TraceCounterDisplayHint_None);
+FCounterInt				GHttpCancelCount(TEXT("Ias/HttpCancelCount"), TraceCounterDisplayHint_None);
 FCounterAtomicInt		GHttpPendingCount(TEXT("Ias/HttpPendingCount"), TraceCounterDisplayHint_None);
 FCounterInt				GHttpInflightCount(TEXT("Ias/HttpInflightCount"), TraceCounterDisplayHint_None);
 FCounterInt				GHttpDownloadedBytes(TEXT("Ias/HttpDownloadedBytes"), TraceCounterDisplayHint_Memory);
@@ -322,6 +323,7 @@ CSV_DEFINE_STAT(Ias, CacheRejectedMB);
 // http stat totals
 CSV_DEFINE_STAT(Ias, HttpGetCount);
 CSV_DEFINE_STAT(Ias, HttpRetryCount);
+CSV_DEFINE_STAT(Ias, HttpCancelCount);
 CSV_DEFINE_STAT(Ias, HttpErrorCount);
 CSV_DEFINE_STAT(Ias, HttpPendingCount);
 CSV_DEFINE_STAT(Ias, HttpDownloadedMB);
@@ -360,6 +362,7 @@ FOnDemandIoBackendStats::FOnDemandIoBackendStats()
 		// http stat totals
 		int32 HGetCount = int32(GHttpGetCount.Get());
 		int32 HRetryCount = int32(GHttpRetryCount.Get());
+		int32 HCancelCount = int32(GHttpCancelCount.Get());
 		int32 HErrorCount = int32(GHttpErrorCount.Get());
 		int32 HPendingCount = int32(GHttpPendingCount.Get());
 		int32 HDownloadedMB = BytesToApproxMB(GHttpDownloadedBytes.Get());
@@ -367,7 +370,7 @@ FOnDemandIoBackendStats::FOnDemandIoBackendStats()
 		int32 HDurationMsAvg = int32(GHttpDurationMsAvg.Get());
 		int32 HDurationMsMax = int32(GHttpDurationMsMax.Get());
 		CSV_CUSTOM_STAT_DEFINED(HttpGetCount, HGetCount, ECsvCustomStatOp::Set);
-		CSV_CUSTOM_STAT_DEFINED(HttpRetryCount, HRetryCount, ECsvCustomStatOp::Set);
+		CSV_CUSTOM_STAT_DEFINED(HttpCancelCount, HCancelCount, ECsvCustomStatOp::Set);
 		CSV_CUSTOM_STAT_DEFINED(HttpErrorCount, HErrorCount, ECsvCustomStatOp::Set);
 		CSV_CUSTOM_STAT_DEFINED(HttpPendingCount, HPendingCount, ECsvCustomStatOp::Set);
 		CSV_CUSTOM_STAT_DEFINED(HttpDownloadedMB, HDownloadedMB, ECsvCustomStatOp::Set);
@@ -382,8 +385,8 @@ FOnDemandIoBackendStats::FOnDemandIoBackendStats()
 			{
 				UE_LOG(LogIas, Log, TEXT("CacheStats: CachedMB=%d, RejectedMB=%d, ReadMB=%d, Get=%d, Error=%d, Put=%d, PutReject=%d, PutExisting=%d"),
 					CCachedMB, CRejectedMB, CReadMB, CGetCount, CErrorCount, CPutCount, CPutRejectCount, CPutExistingCount);
-				UE_LOG(LogIas, Log, TEXT("HttpStats: DownloadedMB=%d, Get=%d, Retry=%d, Error=%d, CurPending=%d, CurDurationMsAvg=%d, CurDurationMsMax=%d"),
-					HDownloadedMB, HGetCount, HRetryCount, HErrorCount, HPendingCount, HDurationMsAvg, HDurationMsMax);
+				UE_LOG(LogIas, Log, TEXT("HttpStats: DownloadedMB=%d, Get=%d, Retry=%d, Cancel=%d, Error=%d, CurPending=%d, CurDurationMsAvg=%d, CurDurationMsMax=%d"),
+					HDownloadedMB, HGetCount, HRetryCount, HCancelCount, HErrorCount, HPendingCount, HDurationMsAvg, HDurationMsMax);
 				LastLogTime = Time;
 			}
 		}
@@ -589,6 +592,13 @@ void FOnDemandIoBackendStats::OnHttpGet(uint64 SizeBytes, uint64 DurationMs)
 	GHttpDurationBuckets[FindDurationBucket(DurationMs)]++;
 }
 
+void FOnDemandIoBackendStats::OnHttpCancel()
+{
+	GHttpInflightCount.Add(-1);
+	GHttpPendingCount.Add(-1);
+	GHttpCancelCount.Add(1);
+}
+
 void FOnDemandIoBackendStats::OnHttpRetry()
 {
 	GHttpRetryCount.Add(1);
diff --git a/Engine/Source/Runtime/Experimental/IoStoreOnDemand/Private/Statistics.h b/Engine/Source/Runtime/Experimental/IoStoreOnDemand/Private/Statistics.h
--- a/Engine/Source/Runtime/Experimental/IoStoreOnDemand/Private/Statistics.h
+++ b/Engine/Source/Runtime/Experimental/IoStoreOnDemand/Private/Statistics.h
@@ -50,6 +50,7 @@ public:
 	void OnCacheSetMaxBytes(uint64 TotalSize) IAS_STATISTICS_IMPL()
 
 	void OnHttpEnqueue() IAS_STATISTICS_IMPL()
+	void OnHttpCancel() IAS_STATISTICS_IMPL()
 	void OnHttpDequeue() IAS_STATISTICS_IMPL()
 	void OnHttpGet(uint64 SizeBytes, uint64 DurationMs) IAS_STATISTICS_IMPL()
 	void OnHttpRetry() IAS_STATISTICS_IMPL()
