Geometry collection optimization : tick geometry collection only when there's at least a piece decaying
- Tracks broken and decayed states in an internal structure
- Simplifies the FullDecay check
- Removal function remain mostly intact but are open for optimization after this change since we can now query the broken state directly instead of inferring it from other more expensive state checks

diff --git a/Engine/Source/Runtime/Experimental/GeometryCollectionEngine/Private/GeometryCollection/GeometryCollectionComponent.cpp b/Engine/Source/Runtime/Experimental/GeometryCollectionEngine/Private/GeometryCollection/GeometryCollectionComponent.cpp
--- a/Engine/Source/Runtime/Experimental/GeometryCollectionEngine/Private/GeometryCollection/GeometryCollectionComponent.cpp
+++ b/Engine/Source/Runtime/Experimental/GeometryCollectionEngine/Private/GeometryCollection/GeometryCollectionComponent.cpp
@@ -597,7 +597,6 @@ UGeometryCollectionComponent::UGeometryCollectionComponent(const FObjectInitiali
 	, bIsTransformSelectionModeEnabled(false)
 #endif  // #if GEOMETRYCOLLECTION_EDITOR_SELECTION
 	, bIsMoving(false)
-	, bIsRootBroken(false)
 	, bUpdateCustomRenderer(true)
 	, bUpdateCustomRendererOnPostPhysicsSync(true)
 {
@@ -3364,8 +3363,7 @@ void UGeometryCollectionComponent::TickComponent(float DeltaTime, enum ELevelTic
 	//UE_LOG(UGCC_LOG, Log, TEXT("GeometryCollectionComponent[%p]::TickComponent()"), this);
 	Super::TickComponent(DeltaTime, TickType, ThisTickFunction);
 
-	// todo(chaos) : cache root broken state ? 
-	if (IsRootBroken())
+	if (BrokenAndDecayedStates.HasAnyDecaying())
 	{
 		InitializeRemovalDynamicAttributesIfNeeded();
 		const float AdjustedDeltaTime = FMath::Max(GeometryCollectionRemovalMultiplier, 0.0001f) * DeltaTime;
@@ -3373,56 +3371,24 @@ void UGeometryCollectionComponent::TickComponent(float DeltaTime, enum ELevelTic
 		IncrementSleepTimer(AdjustedDeltaTime);
 		IncrementBreakTimer(AdjustedDeltaTime);
 	}
+	// if there's no more decaying pieces, we can disable tick component
+	if (!BrokenAndDecayedStates.HasAnyDecaying())
+	{
+		PrimaryComponentTick.SetTickFunctionEnable(false);
+	}
 }
 
 void UGeometryCollectionComponent::CheckFullyDecayed()
 {
-	if (bAlreadyFullyDecayed || !OnFullyDecayedEvent.IsBound())
-	{
-		// Already fully decayed - don't bother doing extra work.
-		return;
-	}
-
-	if (DynamicCollection && PhysicsProxy)
+	if (!bAlreadyFullyDecayed)
 	{
-		FGeometryCollectionDynamicStateFacade DynamicStateFacade(*DynamicCollection);
-
-		if (DynamicStateFacade.IsValid())
+		if (BrokenAndDecayedStates.HasFullyDecayed())
 		{
-			bool bFullyDecayed = true;
-			const int32 NumTransforms = DynamicCollection->NumElements(FGeometryCollection::TransformGroup);
-
-			for (int32 TransformIdx = 0; TransformIdx < NumTransforms; ++TransformIdx)
+			bAlreadyFullyDecayed = true;
+			if (OnFullyDecayedEvent.IsBound())
 			{
-				// If we didn't create a particle for this transform, we shouldn't consider this particle.
-				if (!PhysicsProxy->GetExternalParticles()[TransformIdx])
-				{
-					continue;
-				}
-
-				// In an internal cluster, decay hasn't gotten to this particle yet (e.g. could be in a cluster union).
-				if (DynamicStateFacade.HasInternalClusterParent(TransformIdx))
-				{
-					bFullyDecayed = false;
-					break;
-				}
-
-				// If the particle is active, it's definitely not decayed either.
-				if (DynamicStateFacade.IsActive(TransformIdx))
-				{
-					bFullyDecayed = false;
-					break;
-				}
-			}
-
-			if (bFullyDecayed)
-			{
-				bAlreadyFullyDecayed = true;
-				if (OnFullyDecayedEvent.IsBound())
-				{
-					SCOPE_CYCLE_COUNTER(STAT_GCFullyDecayedBroadcast);
-					OnFullyDecayedEvent.Broadcast();
-				}
+				SCOPE_CYCLE_COUNTER(STAT_GCFullyDecayedBroadcast);
+				OnFullyDecayedEvent.Broadcast();
 			}
 		}
 	}
@@ -3613,7 +3579,7 @@ void UGeometryCollectionComponent::ResetDynamicCollection()
 	RootSpaceBounds.Init();
 	UpdateCachedBounds();
 	bInitializedRemovalDynamicAttribute = false;
-	bIsRootBroken = false;
+	BrokenAndDecayedStates.Reset(NumRestCollectionTransforms);
 }
 
 void UGeometryCollectionComponent::OnCreatePhysicsState()
@@ -3890,19 +3856,146 @@ void UGeometryCollectionComponent::RegisterAndInitializePhysicsProxy()
 	RegisterForEvents();
 }
 
-void UGeometryCollectionComponent::UpdateIsRootBroken()
+void UGeometryCollectionComponent::UpdateBrokenAndDecayedStates()
 {
-	bIsRootBroken = false;
-	if (DynamicCollection && DynamicCollection->Active.Num() > 0)
+	if (PhysicsProxy && DynamicCollection && DynamicCollection->Active.Num() > 0)
 	{
 		const int32 RootIndex = GetRootIndex();
 		if (RootIndex != INDEX_NONE)
 		{
-			bIsRootBroken = !DynamicCollection->Active[RootIndex];
+			const int32 NumTransforms = DynamicCollection->GetNumTransforms();
+
+			const bool bIsRootBroken = !DynamicCollection->Active[RootIndex];
+			BrokenAndDecayedStates.SetRootIsBroken(bIsRootBroken);
+			// we mark it both broken and decayed
+			BrokenAndDecayedStates.SetIsBroken(RootIndex);
+			BrokenAndDecayedStates.SetHasDecayed(RootIndex);
+
+			if (BrokenAndDecayedStates.GetIsRootBroken())
+			{
+				const FGeometryCollectionDynamicStateFacade DynamicStateFacade(*DynamicCollection);
+				if (DynamicStateFacade.IsValid())
+				{
+					for (int32 TransformIdx = 0; TransformIdx < NumTransforms; ++TransformIdx)
+					{
+						if (BrokenAndDecayedStates.GetIsBroken(TransformIdx) || BrokenAndDecayedStates.GetHasDecayed(TransformIdx))
+						{
+							// already decayed nothing to do anymore 
+							continue;
+						}
+
+						const Chaos::FPBDRigidParticle* Particle = PhysicsProxy->GetExternalParticles()[TransformIdx].Get();
+						if (!Particle)
+						{
+							// when no particle we mark it as a decayed one
+							BrokenAndDecayedStates.SetHasDecayed(TransformIdx);
+							continue;
+						}
+
+						const bool bInternalClusterParentBrokenOff = DynamicStateFacade.HasDynamicInternalClusterParent(TransformIdx) && !DynamicStateFacade.HasClusterUnionParent(TransformIdx);
+						const bool bSelfBrokenOff = DynamicStateFacade.HasBrokenOff(TransformIdx) && !DynamicStateFacade.HasInternalClusterParent(TransformIdx);
+						if (bSelfBrokenOff || bInternalClusterParentBrokenOff)
+						{
+							BrokenAndDecayedStates.SetIsBroken(TransformIdx);
+						}
+					}
+				}
+			}
+		}
+	}
+}
+
+bool UGeometryCollectionComponent::FBrokenAndDecayedStates::GetIsBroken(int32 TransformIndex) const
+{
+	return bIsRootBroken ? IsBroken[TransformIndex] : false;
+}
+bool UGeometryCollectionComponent::FBrokenAndDecayedStates::GetHasDecayed(int32 TransformIndex) const
+{
+	return bIsRootBroken ? HasDecayed[TransformIndex] : false;
+}
+
+void UGeometryCollectionComponent::FBrokenAndDecayedStates::Reset(int32 InNumTransforms)
+{
+	NumTransforms = InNumTransforms;
+	bIsRootBroken = false;
+	NumDecaying = 0;
+	IsBroken.Empty();
+	HasDecayed.Empty();
+}
+
+void UGeometryCollectionComponent::FBrokenAndDecayedStates::SetRootIsBroken(bool bIsBroken)
+{
+	if (bIsBroken != bIsRootBroken)
+	{
+		bIsRootBroken = bIsBroken;
+		NumDecaying = 0;
+		if (bIsBroken)
+		{
+			IsBroken.Init(false, NumTransforms);
+			HasDecayed.Init(false, NumTransforms);
+		}
+		else
+		{
+			IsBroken.Empty();
+			HasDecayed.Empty();
+		}
+	}
+}
+
+void UGeometryCollectionComponent::FBrokenAndDecayedStates::SetIsBroken(int32 TransformIndex)
+{
+	if (bIsRootBroken && !IsBroken[TransformIndex])
+	{
+		IsBroken[TransformIndex] = true;
+		NumDecaying++;
+	}
+}
+
+void UGeometryCollectionComponent::FBrokenAndDecayedStates::SetHasDecayed(int32 TransformIndex)
+{
+	if (bIsRootBroken && !HasDecayed[TransformIndex])
+	{
+		HasDecayed[TransformIndex] = true;
+		if (IsBroken[TransformIndex])
+		{
+			NumDecaying--;
+		}
+	}
+}
+
+void UGeometryCollectionComponent::FBrokenAndDecayedStates::SetHasDecayedRecursive(int32 TransformIndex, const TArray<TSet<int32>>& Children)
+{
+	if (bIsRootBroken && ensure(Children.Num() == HasDecayed.Num()))
+	{
+		if (!HasDecayed[TransformIndex])
+		{
+			HasDecayed[TransformIndex] = true;
+			if (IsBroken[TransformIndex])
+			{
+				NumDecaying--;
+			}
+			for (const int32 ChildIndex : Children[TransformIndex])
+			{
+				SetHasDecayedRecursive(ChildIndex, Children);
+			}
 		}
 	}
 }
 
+bool UGeometryCollectionComponent::FBrokenAndDecayedStates::HasAnyDecaying() const
+{
+	return (NumDecaying > 0);
+}
+
+bool UGeometryCollectionComponent::FBrokenAndDecayedStates::HasFullyDecayed() const
+{
+	if (HasAnyDecaying())
+	{
+		return false;
+	}
+	return (HasDecayed.CountSetBits() == HasDecayed.Num());
+}
+
 void UGeometryCollectionComponent::OnTransformsDirty()
 {
 	ComponentSpaceTransforms.MarkDirty();
@@ -3942,26 +4035,29 @@ void UGeometryCollectionComponent::OnPostPhysicsSync()
 
 	// Once the GC is broken, removal feature will need the tick to properly update the timers 
 	// even if the physics does not get any updates
-	UpdateIsRootBroken();
-	if (IsRootBroken())
+	UpdateBrokenAndDecayedStates();
+	if (BrokenAndDecayedStates.GetIsRootBroken())
 	{
-		if (!PrimaryComponentTick.IsTickFunctionEnabled())
-		{ 
-			if (GeometryCollectionEmitRootBreakingEvent)
+		if (GeometryCollectionEmitRootBreakingEvent)
+		{
+			if (OnRootBreakEvent.IsBound())
 			{
-				if (OnRootBreakEvent.IsBound())
-				{
-					FChaosBreakEvent Event;
-					Event.Index = GetRootIndex();
-					Event.Component = this;
-					OnRootBreakEvent.Broadcast(Event);
-				}
+				FChaosBreakEvent Event;
+				Event.Index = GetRootIndex();
+				Event.Component = this;
+				OnRootBreakEvent.Broadcast(Event);
+				OnRootBreakEvent.Clear(); // one shot
 			}
+		}
 
-			PrimaryComponentTick.SetTickFunctionEnable(true);
+		if (BrokenAndDecayedStates.HasAnyDecaying())
+		{
+			if (!PrimaryComponentTick.IsTickFunctionEnabled())
+			{
+				PrimaryComponentTick.SetTickFunctionEnable(true);
+			}
 		}
 
-		// only update removal if the root is broken
 		UpdateRemovalIfNeeded();
 
 		CheckFullyDecayed();
@@ -3991,7 +4087,7 @@ void UGeometryCollectionComponent::MoveComponentToRootTransform()
 
 			const int32 RootIndex = GetRootIndex();
 			const bool bIsRootActive = DynamicStateFacade.IsDynamicOrSleeping(RootIndex);
-			const bool bHasDynamicOPrClusterUnionParent = DynamicStateFacade.HasDynamicInternalClusterParent(RootIndex) || DynamicStateFacade.HasClusterUnionParent(RootIndex);
+			const bool bHasDynamicOPrClusterUnionParent = DynamicStateFacade.HasDynamicInternalClusterParent(RootIndex) || !DynamicStateFacade.HasClusterUnionParent(RootIndex);
 			if (bIsRootActive || bHasDynamicOPrClusterUnionParent)
 			{
 				const FTransform3f& OriginalComponentSpaceRootTransformOffset = AssetCollection->Transform[RootIndex];
@@ -5929,12 +6025,14 @@ void UGeometryCollectionComponent::UpdateDecay(int32 TransformIdx, float Updated
 				{
 					ContextInOut.ToCrumble.AddUnique(InternalClusterItemindex);
 					Decay = 0.0f;
+					// this is an internal cluster so we do not update the BrokenAndDecayedStates, the children will update their own state 
 				}
 			}
 			else
 			{
 				ContextInOut.ToCrumble.AddUnique(FGeometryCollectionItemIndex::CreateTransformItemIndex(TransformIdx));
 				Decay = 0.0f;
+				BrokenAndDecayedStates.SetHasDecayed(TransformIdx);
 			}
 		}
 		else if (Decay >= 1.0f)
@@ -5942,6 +6040,10 @@ void UGeometryCollectionComponent::UpdateDecay(int32 TransformIdx, float Updated
 			// Disable the particle if it has decayed the requisite time
 			Decay = 1.0f;
 			ContextInOut.ToDisable.Add(TransformIdx);
+			if (RestCollection && RestCollection->GetGeometryCollection())
+			{
+				BrokenAndDecayedStates.SetHasDecayedRecursive(TransformIdx, RestCollection->GetGeometryCollection()->Children.GetConstArray());
+			}
 		}
 
 		// push back Decay in the attribute
diff --git a/Engine/Source/Runtime/Experimental/GeometryCollectionEngine/Public/GeometryCollection/GeometryCollectionComponent.h b/Engine/Source/Runtime/Experimental/GeometryCollectionEngine/Public/GeometryCollection/GeometryCollectionComponent.h
--- a/Engine/Source/Runtime/Experimental/GeometryCollectionEngine/Public/GeometryCollection/GeometryCollectionComponent.h
+++ b/Engine/Source/Runtime/Experimental/GeometryCollectionEngine/Public/GeometryCollection/GeometryCollectionComponent.h
@@ -789,7 +789,7 @@ public:
 
 	/** return true if the root cluster is not longer active at runtime */
 	UFUNCTION(BlueprintCallable, Category = "ChaosPhysics")
-	GEOMETRYCOLLECTIONENGINE_API bool IsRootBroken() const { return bIsRootBroken; }
+	GEOMETRYCOLLECTIONENGINE_API bool IsRootBroken() const { return BrokenAndDecayedStates.GetIsRootBroken(); }
 
 	/** 
 	* Get the initial rest transforms in component (local) space  space, 
@@ -1724,18 +1724,42 @@ private:
 	/** True if GeometryCollection transforms have changed from previous tick. */
 	bool bIsMoving;
 
-	/** 
-	* root bone broken state 
-	* this is updated post physics sync and use as is for quick lookup 
-	*/
-	bool bIsRootBroken;
-
-	void UpdateIsRootBroken();
-
 	bool bUpdateCustomRenderer;
 
 	bool bUpdateCustomRendererOnPostPhysicsSync;
 
+private:
+	struct FBrokenAndDecayedStates
+	{
+	public:
+		void Reset(int32 NumTransforms);
+
+		bool GetIsRootBroken() const { return bIsRootBroken; }
+		bool GetIsBroken(int32 TransformIndex) const;
+		bool GetHasDecayed(int32 TransformIndex) const;
+
+		void SetRootIsBroken(bool bIsBroken);
+		void SetIsBroken(int32 TransformIndex);
+		void SetHasDecayed(int32 TransformIndex);
+		void SetHasDecayedRecursive(int32 TransformIndex, const TArray<TSet<int32>>& Children);
+
+		// return true if any broken piece is not yet decayed 
+		bool HasAnyDecaying() const;
+
+		bool HasFullyDecayed() const;
+
+	private:
+		int32 NumTransforms = 0;
+		bool bIsRootBroken = false;
+		TBitArray<> IsBroken;
+		TBitArray<> HasDecayed;
+		int32 NumDecaying = 0;
+	};
+
+	FBrokenAndDecayedStates BrokenAndDecayedStates;
+
+	void UpdateBrokenAndDecayedStates();
+
 private:
 
 	enum class ENetAwakeningMode
