[Chaos] Removing the need to create a temp TArray when calling FPhysicsObjectExternalInterface::LockRead with a single object.

This reduces memory allocations/deallocations done by that array, which can add up when a lot of calls are made.

diff --git a/Engine/Source/Runtime/Engine/Private/PhysicsEngine/PhysicsObjectExternalInterface.cpp b/Engine/Source/Runtime/Engine/Private/PhysicsEngine/PhysicsObjectExternalInterface.cpp
--- a/Engine/Source/Runtime/Engine/Private/PhysicsEngine/PhysicsObjectExternalInterface.cpp
+++ b/Engine/Source/Runtime/Engine/Private/PhysicsEngine/PhysicsObjectExternalInterface.cpp
@@ -11,7 +11,7 @@ FLockedReadPhysicsObjectExternalInterface FPhysicsObjectExternalInterface::LockR
 
 FLockedReadPhysicsObjectExternalInterface FPhysicsObjectExternalInterface::LockRead(Chaos::FConstPhysicsObjectHandle InObject)
 {
-	return LockRead(TArray<Chaos::FConstPhysicsObjectHandle>{ InObject });
+	return LockRead(GetScene(InObject));
 }
 
 FLockedReadPhysicsObjectExternalInterface FPhysicsObjectExternalInterface::LockRead(TArrayView<const Chaos::FConstPhysicsObjectHandle> InObjects)
diff --git a/Engine/Source/Runtime/Experimental/Chaos/Private/Chaos/PhysicsObjectInterface.cpp b/Engine/Source/Runtime/Experimental/Chaos/Private/Chaos/PhysicsObjectInterface.cpp
--- a/Engine/Source/Runtime/Experimental/Chaos/Private/Chaos/PhysicsObjectInterface.cpp
+++ b/Engine/Source/Runtime/Experimental/Chaos/Private/Chaos/PhysicsObjectInterface.cpp
@@ -1189,11 +1189,7 @@ namespace Chaos
 				continue;
 			}
 
-			FPBDRigidsSolver* Solver = nullptr;
-			if (const IPhysicsProxyBase* Proxy = Object->PhysicsProxy())
-			{
-				Solver = Proxy->GetSolver<FPBDRigidsSolver>();
-			}
+			FPBDRigidsSolver* Solver = GetSolver(Object);
 
 			if (!Solver)
 			{
@@ -1211,6 +1207,23 @@ namespace Chaos
 		return RetSolver;
 	}
 
+	FPBDRigidsSolver* FPhysicsObjectInterface::GetSolver(const FConstPhysicsObjectHandle InObject)
+	{
+		if (!InObject)
+		{
+			return nullptr;
+		}
+
+		FPBDRigidsSolver* Solver = nullptr;
+		if (const IPhysicsProxyBase* Proxy = InObject->PhysicsProxy())
+		{
+			Solver = Proxy->GetSolver<FPBDRigidsSolver>();
+		}
+
+		return Solver;
+	}
+
+
 	IPhysicsProxyBase* FPhysicsObjectInterface::GetProxy(TArrayView<const FConstPhysicsObjectHandle> InObjects)
 	{
 		IPhysicsProxyBase* RetProxy = nullptr;
diff --git a/Engine/Source/Runtime/Experimental/Chaos/Public/Chaos/PhysicsObjectInterface.h b/Engine/Source/Runtime/Experimental/Chaos/Public/Chaos/PhysicsObjectInterface.h
--- a/Engine/Source/Runtime/Experimental/Chaos/Public/Chaos/PhysicsObjectInterface.h
+++ b/Engine/Source/Runtime/Experimental/Chaos/Public/Chaos/PhysicsObjectInterface.h
@@ -195,6 +195,7 @@ namespace Chaos
 		static CHAOS_API int32 GetId(const FConstPhysicsObjectHandle Object);
 
 		static CHAOS_API FPBDRigidsSolver* GetSolver(TArrayView<const FConstPhysicsObjectHandle> InObjects);
+		static CHAOS_API FPBDRigidsSolver* GetSolver(const FConstPhysicsObjectHandle InObject);
 		static CHAOS_API IPhysicsProxyBase* GetProxy(TArrayView<const FConstPhysicsObjectHandle> InObjects);
 
 	protected:
diff --git a/Engine/Source/Runtime/PhysicsCore/Private/PhysicsObjectPhysicsCoreInterface.cpp b/Engine/Source/Runtime/PhysicsCore/Private/PhysicsObjectPhysicsCoreInterface.cpp
--- a/Engine/Source/Runtime/PhysicsCore/Private/PhysicsObjectPhysicsCoreInterface.cpp
+++ b/Engine/Source/Runtime/PhysicsCore/Private/PhysicsObjectPhysicsCoreInterface.cpp
@@ -14,3 +14,13 @@ FChaosScene* PhysicsObjectPhysicsCoreInterface::GetScene(TArrayView<const Chaos:
 	}
 	return reinterpret_cast<FChaosScene*>(Solver->PhysSceneHack);
 }
+
+FChaosScene* PhysicsObjectPhysicsCoreInterface::GetScene(const Chaos::FConstPhysicsObjectHandle InObject)
+{
+	Chaos::FPBDRigidsSolver* Solver = Chaos::FPhysicsObjectInterface::GetSolver(InObject);
+	if (!Solver)
+	{
+		return nullptr;
+	}
+	return reinterpret_cast<FChaosScene*>(Solver->PhysSceneHack);
+}
diff --git a/Engine/Source/Runtime/PhysicsCore/Public/PhysicsObjectPhysicsCoreInterface.h b/Engine/Source/Runtime/PhysicsCore/Public/PhysicsObjectPhysicsCoreInterface.h
--- a/Engine/Source/Runtime/PhysicsCore/Public/PhysicsObjectPhysicsCoreInterface.h
+++ b/Engine/Source/Runtime/PhysicsCore/Public/PhysicsObjectPhysicsCoreInterface.h
@@ -11,4 +11,5 @@ class PhysicsObjectPhysicsCoreInterface
 {
 public:
 	static PHYSICSCORE_API FChaosScene* GetScene(TArrayView<const Chaos::FConstPhysicsObjectHandle> InObjects);
+	static PHYSICSCORE_API FChaosScene* GetScene(const Chaos::FConstPhysicsObjectHandle InObject);
 };
