Fix warning V591: Non-void function should return a value.

diff --git a/Engine/Source/Runtime/Experimental/Chaos/Public/Chaos/Properties.h b/Engine/Source/Runtime/Experimental/Chaos/Public/Chaos/Properties.h
--- a/Engine/Source/Runtime/Experimental/Chaos/Public/Chaos/Properties.h
+++ b/Engine/Source/Runtime/Experimental/Chaos/Public/Chaos/Properties.h
@@ -84,7 +84,7 @@ private:
 	}
 
 	//we don't support this because it could lead to bugs with values not being properly written to remote
-	TChaosProperty<T,PropName>& operator=(const TChaosProperty<T,PropName>& Rhs) {}
+	TChaosProperty<T,PropName>& operator=(const TChaosProperty<T,PropName>& Rhs) = delete;
 };
 
 template <typename T, EChaosProperty PropName>
@@ -158,7 +158,7 @@ private:
 	}
 
 	//we don't support this because it could lead to bugs with values not being properly written to remote
-	TShapeProperty<T,PropName>& operator=(const TShapeProperty<T,PropName>& Rhs) {}
+	TShapeProperty<T, PropName>& operator=(const TShapeProperty<T, PropName>& Rhs) = delete;
 };
 
 template <typename T,EShapeProperty PropName>
diff --git a/Engine/Source/Runtime/Experimental/Chaos/Public/PhysicsProxy/JointConstraintProxy.h b/Engine/Source/Runtime/Experimental/Chaos/Public/PhysicsProxy/JointConstraintProxy.h
--- a/Engine/Source/Runtime/Experimental/Chaos/Public/PhysicsProxy/JointConstraintProxy.h
+++ b/Engine/Source/Runtime/Experimental/Chaos/Public/PhysicsProxy/JointConstraintProxy.h
@@ -114,7 +114,7 @@ class TThreadedJointConstraintPhysicsProxyBase : protected FJointConstraintPhysi
 public:
 
 #define CHAOS_INNER_JOINT_PROPERTY(OuterProp, FuncName, Inner, InnerType)\
-	const InnerType& Get##FuncName() const { ReadRef([](const auto& Data){ return Data.Inner;}); }\
+	const InnerType& Get##FuncName() const { return ReadRef([](const auto& Data){ return Data.Inner;}); }\
 	void Set##FuncName(const InnerType& Val) { Write([&Val](auto& Data){ Data.Inner = Val;}); }\
 
 #include "Chaos/JointProperties.inl"
