Add component keep alive option on ISMPool. This keeps ISM components alive when the last instance is removed. This avoids costly scene proxy creation for components that regularly switch between 0 and N instances.
The behavior is controlled by r.ISMPool.ComponentKeepAlive.
The old free list behavior is still supported and is now controlled by r.ISMPool.ComponentRecycle.
Both cvars can be safely enabled/disabled at runtime for testing.
Note that keep alive will effectively override recycling if both are enabled, since no components are released to the recycling free list.

diff --git a/Engine/Source/Runtime/Experimental/GeometryCollectionEngine/Private/GeometryCollection/GeometryCollectionISMPoolComponent.cpp b/Engine/Source/Runtime/Experimental/GeometryCollectionEngine/Private/GeometryCollection/GeometryCollectionISMPoolComponent.cpp
--- a/Engine/Source/Runtime/Experimental/GeometryCollectionEngine/Private/GeometryCollection/GeometryCollectionISMPoolComponent.cpp
+++ b/Engine/Source/Runtime/Experimental/GeometryCollectionEngine/Private/GeometryCollection/GeometryCollectionISMPoolComponent.cpp
@@ -2,27 +2,40 @@
 
 #include "GeometryCollection/GeometryCollectionISMPoolComponent.h"
 
-#include "Components/HierarchicalInstancedStaticMeshComponent.h"
+#include "ChaosLog.h"
+#include "Components/InstancedStaticMeshComponent.h"
 #include "Engine/CollisionProfile.h"
 #include "Engine/StaticMesh.h"
-#include "GeometryCollection/GeometryCollectionComponent.h"
 
 #include UE_INLINE_GENERATED_CPP_BY_NAME(GeometryCollectionISMPoolComponent)
 
-// Use the FreeLists to enable recycling of ISM components.
-// Might want to add a Tick function to keep the pool at a stable value over time.
-// We would always want some spare components for fast allocation, but want to clean up when numbers get too high.
-static bool GUseComponentFreeList = true;
-FAutoConsoleVariableRef CVarISMPoolUseComponentFreeList(
-	TEXT("r.ISMPool.UseComponentFreeList"),
-	GUseComponentFreeList,
-	TEXT("Recycle ISM components in the ISMPool."));
-
+// Don't release ISM components when they empty, but keep them (and their scene proxy) alive.
+// This can remove the high cost associated with repeated registration, scene proxy creation and mesh draw command creation.
+static bool GComponentKeepAlive = true;
+FAutoConsoleVariableRef CVarISMPoolComponentKeepAlive(
+	TEXT("r.ISMPool.ComponentKeepAlive"),
+	GComponentKeepAlive,
+	TEXT("Keep ISM components alive when all their instances are removed."));
+
+// Use a FreeList to enable recycling of ISM components.
+// ISM components aren't unregistered, but their scene proxy is destroyed.
+// When recycling a component, a new mesh description can be used.
+// This removes the high CPU cost of unregister/register.
+// But there is more CPU cost to recycling a component then to simply keeping it alive because scene proxy creation and mesh draw command caching isn't cheap.
+// The component memory cost is kept bounded when compared to keeping components alive.
+static bool GComponentRecycle = true;
+FAutoConsoleVariableRef CVarISMPoolComponentRecycle(
+	TEXT("r.ISMPool.ComponentRecycle"),
+	GComponentRecycle,
+	TEXT("Recycle ISM components to a free list for reuse when all their instances are removed."));
+
+// Target free list size when recycling ISM components.
+// We try to maintain a pool of free components for fast allocation, but want to clean up when numbers get too high.
 static int32 GComponentFreeListTargetSize = 50;
 FAutoConsoleVariableRef CVarISMPoolComponentFreeListTargetSize(
 	TEXT("r.ISMPool.ComponentFreeListTargetSize"),
 	GComponentFreeListTargetSize,
-	TEXT("Target size for number of ISM components in the ISMPool."));
+	TEXT("Target size for number of ISM components in the recycling free list."));
 
 
 FGeometryCollectionMeshGroup::FMeshId FGeometryCollectionMeshGroup::AddMesh(const FGeometryCollectionStaticMeshInstance& MeshInstance, int32 InstanceCount, const FGeometryCollectionMeshInfo& ISMInstanceInfo)
@@ -53,7 +66,7 @@ void FGeometryCollectionMeshGroup::RemoveAllMeshes(FGeometryCollectionISMPool& I
 {
 	for (const FGeometryCollectionMeshInfo& MeshInfo: MeshInfos)
 	{
-		ISMPool.RemoveISM(MeshInfo);
+		ISMPool.RemoveInstancesFromISM(MeshInfo);
 	}
 	MeshInfos.Empty();
 }
@@ -73,7 +86,7 @@ void FGeometryCollectionISM::CreateISM(AActor* InOwningActor)
 	ISMComponent->RegisterComponent();
 }
 
-void FGeometryCollectionISM::InitISM(const FGeometryCollectionStaticMeshInstance& InMeshInstance)
+void FGeometryCollectionISM::InitISM(const FGeometryCollectionStaticMeshInstance& InMeshInstance, bool bKeepAlive)
 {
 	MeshInstance = InMeshInstance;
 	check(MeshInstance.StaticMesh);
@@ -164,11 +177,18 @@ FInstanceGroups::FInstanceGroupId FGeometryCollectionISM::AddInstanceGroup(int32
 	return InstanceGroupIndex;
 }
 
-FGeometryCollectionISMPool::FISMIndex FGeometryCollectionISMPool::AddISM(UGeometryCollectionISMPoolComponent* OwningComponent, const FGeometryCollectionStaticMeshInstance& MeshInstance)
+FGeometryCollectionISMPool::FGeometryCollectionISMPool()
+	: bCachedKeepAlive(GComponentKeepAlive)
+	, bCachedRecycle(GComponentRecycle)
+{
+}
+
+FGeometryCollectionISMPool::FISMIndex FGeometryCollectionISMPool::GetOrAddISM(UGeometryCollectionISMPoolComponent* OwningComponent, const FGeometryCollectionStaticMeshInstance& MeshInstance, bool& bOutISMCreated)
 {
 	FISMIndex* ISMIndexPtr = MeshToISMIndex.Find(MeshInstance);
 	if (ISMIndexPtr != nullptr)
 	{
+		bOutISMCreated = false;
 		return *ISMIndexPtr;
 	}
 
@@ -191,16 +211,18 @@ FGeometryCollectionISMPool::FISMIndex FGeometryCollectionISMPool::AddISM(UGeomet
 		ISMs[ISMIndex].CreateISM(OwningComponent->GetOwner());
 	}
 	
-	ISMs[ISMIndex].InitISM(MeshInstance);
+	ISMs[ISMIndex].InitISM(MeshInstance, bCachedKeepAlive);
 	
+	bOutISMCreated = true;
 	MeshToISMIndex.Add(MeshInstance, ISMIndex);
 	return ISMIndex;
 }
 
-FGeometryCollectionMeshInfo FGeometryCollectionISMPool::AddISM(UGeometryCollectionISMPoolComponent* OwningComponent, const FGeometryCollectionStaticMeshInstance& MeshInstance, int32 InstanceCount, TArrayView<const float> CustomDataFloats)
+FGeometryCollectionMeshInfo FGeometryCollectionISMPool::AddInstancesToISM(UGeometryCollectionISMPoolComponent* OwningComponent, const FGeometryCollectionStaticMeshInstance& MeshInstance, int32 InstanceCount, TArrayView<const float> CustomDataFloats)
 {
+	bool bISMCreated = false;
 	FGeometryCollectionMeshInfo Info;
-	Info.ISMIndex = AddISM(OwningComponent, MeshInstance);
+	Info.ISMIndex = GetOrAddISM(OwningComponent, MeshInstance, bISMCreated);
 	Info.InstanceGroupIndex = ISMs[Info.ISMIndex].AddInstanceGroup(InstanceCount, CustomDataFloats);
 	return Info;
 }
@@ -281,7 +303,7 @@ void FGeometryCollectionISMPool::BatchUpdateInstanceCustomData(FGeometryCollecti
 	}
 }
 
-void FGeometryCollectionISMPool::RemoveISM(const FGeometryCollectionMeshInfo& MeshInfo)
+void FGeometryCollectionISMPool::RemoveInstancesFromISM(const FGeometryCollectionMeshInfo& MeshInfo)
 {
 	if (ISMs.IsValidIndex(MeshInfo.ISMIndex))
 	{
@@ -309,28 +331,56 @@ void FGeometryCollectionISMPool::RemoveISM(const FGeometryCollectionMeshInfo& Me
 	
 		if (ISM.InstanceGroups.IsEmpty())
 		{
+			ensure(ISM.ISMComponent->PerInstanceSMData.Num() == 0);
+
 			// No live instances, so take opportunity to reset indexing.
 			ISM.InstanceGroups.Reset();
 			ISM.InstanceIds.Reset();
+
+			RemoveISM(MeshInfo.ISMIndex, bCachedKeepAlive, bCachedRecycle);
+
+			if (!bCachedKeepAlive)
+			{
+				MeshToISMIndex.Remove(ISM.MeshInstance);
+			}
 		}
+	}
+}
 
-		if (GUseComponentFreeList && ISM.InstanceGroups.IsEmpty())
-		{
-			// Remove component and push this ISM slot to the free list.
-			ensure(ISM.ISMComponent->PerInstanceSMData.Num() == 0);
-			MeshToISMIndex.Remove(ISM.MeshInstance);
-			FreeListISM.Add(MeshInfo.ISMIndex);
+void FGeometryCollectionISMPool::RemoveISM(FISMIndex ISMIndex, bool bKeepAlive, bool bRecycle)
+{
+	FGeometryCollectionISM& ISM = ISMs[ISMIndex];
+	ensure(ISM.InstanceGroups.IsEmpty());
+	ensure(ISM.InstanceIds.IsEmpty());
 
+	if (bKeepAlive)
+	{
+		// Nothing to do.
+	}
+	else if (bRecycle)
+	{
+		// Recycle to the free list.
 #if WITH_EDITOR
-			ISM.ISMComponent->Rename(nullptr);
+		ISM.ISMComponent->Rename(nullptr);
 #endif
-		}
+		FreeListISM.Add(ISMIndex);
+	}
+	else
+	{
+		// Completely unregister and destroy the component and mark the ISM slot as free.
+		ISM.ISMComponent->UnregisterComponent();
+		ISM.ISMComponent->DestroyComponent();
+		ISM.ISMComponent->GetOwner()->RemoveInstanceComponent(ISM.ISMComponent);
+		ISM.ISMComponent = nullptr;
+		
+		FreeList.Add(ISMIndex);
 	}
 }
 
 void FGeometryCollectionISMPool::Clear()
 {
 	MeshToISMIndex.Reset();
+	PrellocationQueue.Reset();
 	FreeList.Reset();
 	FreeListISM.Reset();
 	if (ISMs.Num() > 0)
@@ -348,21 +398,83 @@ void FGeometryCollectionISMPool::Clear()
 	}
 }
 
-void FGeometryCollectionISMPool::GarbageCollect()
+void FGeometryCollectionISMPool::RequestPreallocateMeshInstance(const FGeometryCollectionStaticMeshInstance& MeshInstance)
 {
-	// Release one component per call until we reach minimum pool size.
-	if (FreeListISM.Num() > GComponentFreeListTargetSize)
+	// Preallocation only makes sense when we are keeping empty components alive.
+	if (bCachedKeepAlive)
 	{
-		const int32 ISMIndex = FreeListISM.Last();
-		FreeListISM.RemoveAt(FreeListISM.Num() - 1);
+		uint32 KeyHash = GetTypeHash(MeshInstance);
+		if (MeshToISMIndex.FindByHash(KeyHash, MeshInstance) == nullptr)
+		{
+			PrellocationQueue.AddByHash(KeyHash, MeshInstance);
+		}
+	}
+}
 
-		UInstancedStaticMeshComponent* ISM = ISMs[ISMIndex].ISMComponent;
-		ISM->UnregisterComponent();
-		ISM->DestroyComponent();
-		ISM->GetOwner()->RemoveInstanceComponent(ISM);
-		ISMs[ISMIndex].ISMComponent = nullptr;
+void FGeometryCollectionISMPool::ProcessPreallocationRequests(UGeometryCollectionISMPoolComponent* OwningComponent, int32 MaxPreallocations)
+{
+	int32 NumAdded = 0;
+	for (TSet<FGeometryCollectionStaticMeshInstance>::TIterator It(PrellocationQueue); It; ++It)
+	{
+		bool bISMCreated = false;
+		GetOrAddISM(OwningComponent, *It, bISMCreated);
+		It.RemoveCurrent();
 
-		FreeList.Add(ISMIndex);
+		if (bISMCreated)
+		{
+			if (++NumAdded >= MaxPreallocations)
+			{
+				break;
+			}
+		}
+	}
+}
+
+void FGeometryCollectionISMPool::Tick(UGeometryCollectionISMPoolComponent* OwningComponent)
+{
+	// Recache component lifecycle state from cvar.
+	const bool bRemovedKeepAlive = bCachedKeepAlive && !GComponentKeepAlive;
+	const bool bRemovedReycle = bCachedRecycle && !GComponentRecycle;
+	bCachedKeepAlive = GComponentKeepAlive;
+	bCachedRecycle = GComponentRecycle;
+
+	// If we disabled keep alive behavior since last update then deal with the zombie components.
+	if (bRemovedKeepAlive)
+	{
+		for (int32 ISMIndex = 0; ISMIndex < ISMs.Num(); ++ISMIndex)
+		{
+			FGeometryCollectionISM& ISM = ISMs[ISMIndex];
+			if (ISM.ISMComponent && ISM.InstanceGroups.IsEmpty())
+			{
+				// Actually release the ISM.
+				RemoveISM(ISMIndex, false, bCachedRecycle);
+				MeshToISMIndex.Remove(ISM.MeshInstance);
+			}
+		}
+	}
+
+	// Process preallocation queue.
+	if (!bCachedKeepAlive)
+	{
+		PrellocationQueue.Reset();
+	}
+	else if (!PrellocationQueue.IsEmpty())
+	{
+		// Preallocate components per tick until the queue is empty.
+		const int32 PreallocateCountPerTick = 2;
+		ProcessPreallocationRequests(OwningComponent, PreallocateCountPerTick);
+	}
+
+	if (FreeListISM.Num() > 0)
+	{
+		// Release components per tick until we reach minimum pool size.
+		const int32 RemoveCountPerTick = 1;
+		const int32 FreeListTargetSize = bRemovedReycle ? 0 : FMath::Max(FMath::Max(FreeListISM.Num() - RemoveCountPerTick, GComponentFreeListTargetSize), 0);
+		while (FreeListISM.Num() > FreeListTargetSize)
+		{
+			const int32 ISMIndex = FreeListISM.Pop(false);
+			RemoveISM(ISMIndex, false, false);
+		}
 	}
 }
 
@@ -379,7 +491,7 @@ UGeometryCollectionISMPoolComponent::UGeometryCollectionISMPoolComponent(const F
 void UGeometryCollectionISMPoolComponent::TickComponent(float DeltaTime, enum ELevelTick TickType, FActorComponentTickFunction* ThisTickFunction)
 {
 	Super::TickComponent(DeltaTime, TickType, ThisTickFunction);
-	Pool.GarbageCollect();
+	Pool.Tick(this);
 }
 
 UGeometryCollectionISMPoolComponent::FMeshGroupId  UGeometryCollectionISMPoolComponent::CreateMeshGroup(bool bAllowPerInstanceRemoval)
@@ -403,7 +515,7 @@ UGeometryCollectionISMPoolComponent::FMeshId UGeometryCollectionISMPoolComponent
 {
 	if (FGeometryCollectionMeshGroup* MeshGroup = MeshGroups.Find(MeshGroupId))
 	{
-		const FGeometryCollectionMeshInfo ISMInstanceInfo = Pool.AddISM(this, MeshInstance, InstanceCount, CustomDataFloats);
+		const FGeometryCollectionMeshInfo ISMInstanceInfo = Pool.AddInstancesToISM(this, MeshInstance, InstanceCount, CustomDataFloats);
 		return MeshGroup->AddMesh(MeshInstance, InstanceCount, ISMInstanceInfo);
 	}
 	UE_LOG(LogChaos, Warning, TEXT("UGeometryCollectionISMPoolComponent : Trying to add a mesh to a mesh group (%d) that does not exists"), MeshGroupId);
@@ -438,12 +550,7 @@ bool UGeometryCollectionISMPoolComponent::BatchUpdateInstanceCustomData(FMeshGro
 
 void UGeometryCollectionISMPoolComponent::PreallocateMeshInstance(const FGeometryCollectionStaticMeshInstance& MeshInstance)
 {
-	// If we are recycling components with a free list then we don't expect to have zero instance components.
-	// So don't do preallocation of components either in that case.
-	if (!GUseComponentFreeList)
-{
-		Pool.AddISM(this, MeshInstance);
-	}
+	Pool.RequestPreallocateMeshInstance(MeshInstance);
 }
 
 void UGeometryCollectionISMPoolComponent::GetResourceSizeEx(FResourceSizeEx& CumulativeResourceSize)
diff --git a/Engine/Source/Runtime/Experimental/GeometryCollectionEngine/Public/GeometryCollection/GeometryCollectionISMPoolComponent.h b/Engine/Source/Runtime/Experimental/GeometryCollectionEngine/Public/GeometryCollection/GeometryCollectionISMPoolComponent.h
--- a/Engine/Source/Runtime/Experimental/GeometryCollectionEngine/Public/GeometryCollection/GeometryCollectionISMPoolComponent.h
+++ b/Engine/Source/Runtime/Experimental/GeometryCollectionEngine/Public/GeometryCollection/GeometryCollectionISMPoolComponent.h
@@ -251,7 +251,7 @@ struct FGeometryCollectionISM
 	/** Create the ISMComponent according to settings on the mesh instance. */
 	void CreateISM(AActor* InOwningActor);
 	/** Initialize the ISMComponent according to settings on the mesh instance. */
-	void InitISM(const FGeometryCollectionStaticMeshInstance& InMeshInstance);
+	void InitISM(const FGeometryCollectionStaticMeshInstance& InMeshInstance, bool bKeepAlive);
 	/** Add a group to the ISM. Returns the group index. */
 	FInstanceGroups::FInstanceGroupId AddInstanceGroup(int32 InstanceCount, TArrayView<const float> CustomDataFloats);
 
@@ -270,29 +270,47 @@ struct FGeometryCollectionISMPool
 {
 	using FISMIndex = int32;
 
-	/** Add ISM contents. */
-	FISMIndex AddISM(UGeometryCollectionISMPoolComponent* OwningComponent, const FGeometryCollectionStaticMeshInstance& MeshInstance);
-	FGeometryCollectionMeshInfo AddISM(UGeometryCollectionISMPoolComponent* OwningComponent, const FGeometryCollectionStaticMeshInstance& MeshInstance, int32 InstanceCount, TArrayView<const float> CustomDataFloats);
-	/** Remove ISM contents. */
-	void RemoveISM(const FGeometryCollectionMeshInfo& MeshInfo);
+	FGeometryCollectionISMPool();
+
+	/** Find or add an ISM and return an ISM index handle. */
+	FISMIndex GetOrAddISM(UGeometryCollectionISMPoolComponent* OwningComponent, const FGeometryCollectionStaticMeshInstance& MeshInstance, bool& bOutISMCreated);
+	/** Remove an ISM. */
+	void RemoveISM(FISMIndex ISMIndex, bool bKeepAlive, bool bRecycle);
+	/** Add instances to ISM and return a mesh info handle. */
+	FGeometryCollectionMeshInfo AddInstancesToISM(UGeometryCollectionISMPoolComponent* OwningComponent, const FGeometryCollectionStaticMeshInstance& MeshInstance, int32 InstanceCount, TArrayView<const float> CustomDataFloats);
+	/** Remove instances from an ISM. */
+	void RemoveInstancesFromISM(const FGeometryCollectionMeshInfo& MeshInfo);
 	/** Update ISM contents. */
 	bool BatchUpdateInstancesTransforms(FGeometryCollectionMeshInfo& MeshInfo, int32 StartInstanceIndex, TArrayView<const FTransform> NewInstancesTransforms, bool bWorldSpace, bool bMarkRenderStateDirty, bool bTeleport, bool bAllowPerInstanceRemoval);
 	void BatchUpdateInstanceCustomData(FGeometryCollectionMeshInfo const& MeshInfo, int32 CustomFloatIndex, float CustomFloatValue);
 
 	/** Clear all ISM components and associated data. */
 	void Clear();
-	/** Garbage collect free lists. */
-	void GarbageCollect();
+
+	/** Tick maintenance of free list and preallocation. */
+	void Tick(UGeometryCollectionISMPoolComponent* OwningComponent);
+
+	/** Add an ISM description to the preallocation queue. */
+	void RequestPreallocateMeshInstance(const FGeometryCollectionStaticMeshInstance& MeshInstances);
+	/** Process the preallocation queue. Processing is timesliced so that only some of the queue will be processed in every call. */
+	void ProcessPreallocationRequests(UGeometryCollectionISMPoolComponent* OwningComponent, int32 MaxPreallocations);
 
 	/** Array of ISM objects. */
 	TArray<FGeometryCollectionISM> ISMs;
 	/** Mapping from mesh description to ISMs array slot. */
 	TMap<FGeometryCollectionStaticMeshInstance, FISMIndex> MeshToISMIndex;
 	
+	/** Set of ISM descriptions that we would like to preallocate. */
+	TSet<FGeometryCollectionStaticMeshInstance> PrellocationQueue;
+
 	/** Free list of indices in ISMs that are empty. */
 	TArray<int32> FreeList;
 	/** Free list of indices in ISMs that have registered ISM components. */
 	TArray<int32> FreeListISM;
+
+	// Cached state of lifecycle cvars from the last Tick()
+	bool bCachedKeepAlive = false;
+	bool bCachedRecycle = false;
 };
 
 
