Prevent a crash when retrying to resolve a distributed endpoint for IAS

#rnx

- Before we retry (by calling IssueEndpointRequests) we unbind the callback which can cause it to be deleted, invalidating it's captures, this means the this pointer we use to call ::IssueEndpointRequests with is trash.
- We need to create a local copy of the this pointer before unbinding the callback to make sure that it is still safe to use.

diff --git a/Engine/Source/Runtime/Experimental/IoStoreOnDemand/Private/DistributionEndpoints.cpp b/Engine/Source/Runtime/Experimental/IoStoreOnDemand/Private/DistributionEndpoints.cpp
--- a/Engine/Source/Runtime/Experimental/IoStoreOnDemand/Private/DistributionEndpoints.cpp
+++ b/Engine/Source/Runtime/Experimental/IoStoreOnDemand/Private/DistributionEndpoints.cpp
@@ -146,12 +146,14 @@ void FDistributionEndpoints::IssueEndpointRequests()
 					{
 						if (++ResolveRequest.RetryCount <= MaxAttempts)
 						{
-							Request->OnProcessRequestComplete().Unbind();
-							return IssueEndpointRequests();
+							FDistributionEndpoints* CachedPtr = this;
+							Request->OnProcessRequestComplete().Unbind(); // <- Invalidates the lambdas captures
+
+							return CachedPtr->IssueEndpointRequests();
 						}
 					}
 
-					CompleteEndpointRequest(ResolveRequest, Response);
+					this->CompleteEndpointRequest(ResolveRequest, Response);
 				});
 
 			ResolveRequest.HttpRequest = HttpRequest;
