- IAS could stop processing requests if it hit its fail timeout
- New work was added in reverse order

IAS has a mechanism to fail requests if they're taking to long. This is to prevent live locks in very rare scenarios where network traffic comes to a complete halt. In the discovered repro case it involved using an external tool to silently drop traffic. This mechanism was incorrectly dropping an array such that there was nothing left to process requests still pending.

While fleshing out the tests I noticed that work that is added slightly out of order. This is mostly harmless due to the way IAS disaptches requests, but for best results in high-latency scenarios the mis-ordering was corrected by reversing a list.

Note that this change includes my test code changes. These are not compiled in Test or Shipping.

#rnx
#tests Paul Chipchase found a repro and confirmed the suggested fix resolves the issue. I improved my existing tests to cover this scenario and then implemented the fix. Win64/test locker, soak tested my library of replays overnight. Mega-compile across all platforms and variants.

diff --git a/Engine/Source/Runtime/Experimental/IoStoreOnDemand/Private/CoreHttp/Http.cpp b/Engine/Source/Runtime/Experimental/IoStoreOnDemand/Private/CoreHttp/Http.cpp
--- a/Engine/Source/Runtime/Experimental/IoStoreOnDemand/Private/CoreHttp/Http.cpp
+++ b/Engine/Source/Runtime/Experimental/IoStoreOnDemand/Private/CoreHttp/Http.cpp
@@ -3053,7 +3053,6 @@ void FHostGroup::Tick(FTickState& State)
 			Group.Fail(State, Reason);
 		}
 
-		SocketGroups.SetNum(0, false);
 		return;
 	}
 
@@ -3221,6 +3220,16 @@ void FEventLoop::FImpl::ReceiveWork()
 		Swap(Activity, Pending);
 	}
 
+	// Pending is in the reverse of the order that requests were made
+	FActivity* Reverse = nullptr;
+	for (FActivity* Next; Activity != nullptr; Activity = Next)
+	{
+		Next = Activity->Next;
+		Activity->Next = Reverse;
+		Reverse = Activity;
+	}
+	Activity = Reverse;
+
 	// Group activities by their host.
 	for (FActivity* Next; Activity != nullptr; Activity = Next)
 	{
@@ -3760,6 +3769,7 @@ IOSTOREONDEMAND_API void IasHttpTest(const ANSICHAR* TestHost="localhost")
 		uint32 Index = Status.GetIndex();
 
 		FResponse& Response = Status.GetResponse();
+		Content[Index].Dest = FIoBuffer();
 		Response.SetDestination(&(Content[Index].Dest));
 	};
 
@@ -3849,25 +3859,31 @@ IOSTOREONDEMAND_API void IasHttpTest(const ANSICHAR* TestHost="localhost")
 	for (int32 i = 0; i < 14; ++i)
 	{
 		bool bExpectFailTimeout = !!(i & 1);
-		FEventLoop Loop2;
-		Loop2.Send(
-			Loop2.Get(BuildUrl("/data?stall", 9494)),
-			[bExpectFailTimeout, Dest=FIoBuffer()] (const FTicketStatus& Status) mutable
+		auto Sink = [bExpectFailTimeout, Dest=FIoBuffer()] (const FTicketStatus& Status) mutable
+		{
+			if (Status.GetId() == FTicketStatus::EId::Response)
 			{
-				if (Status.GetId() == FTicketStatus::EId::Response)
-				{
-					FResponse& Response = Status.GetResponse();
-					Response.SetDestination(&Dest);
-					return;
-				}
+				FResponse& Response = Status.GetResponse();
+				Response.SetDestination(&Dest);
+				return;
+			}
 
-				check(Status.GetId() == FTicketStatus::EId::Error);
+			check(Status.GetId() == FTicketStatus::EId::Error);
 
-				const char* Reason = Status.GetErrorReason();
-				bool IsFailTimeout = (FCStringAnsi::Strstr(Reason, "FailTimeout") != nullptr);
-				check(IsFailTimeout == bExpectFailTimeout);
-			}
-		);
+			const char* Reason = Status.GetErrorReason();
+			bool IsFailTimeout = (FCStringAnsi::Strstr(Reason, "FailTimeout") != nullptr);
+			check(IsFailTimeout == bExpectFailTimeout);
+		};
+
+		FConnectionPool::FParams Params;
+		Params.SetHostFromUrl(BuildUrl("", 9494));
+		FConnectionPool Pool(Params);
+
+		FEventLoop Loop2;
+		Loop2.Send(Loop2.Get("/data?stall", Pool), Sink);
+		Loop2.Send(Loop2.Get("/data", Pool), HashSink);
+		Loop2.Send(Loop2.Get("/data", Pool), HashSink);
+		Loop2.Send(Loop2.Get("/data", Pool), HashSink);
 
 		int32 PollTimeoutMs = -1;
 		if (bExpectFailTimeout)
@@ -3880,6 +3896,9 @@ IOSTOREONDEMAND_API void IasHttpTest(const ANSICHAR* TestHost="localhost")
 			}
 		}
 		while (Loop2.Tick(PollTimeoutMs));
+
+		Loop2.Send(Loop2.Get("/data/23", Pool), NoErrorSink);
+		while (Loop2.Tick(PollTimeoutMs));
 	}
 
 	// no connect
