Adding the ability for Geometry Collections to specify the use of the collision geometry from the source Static Meshes on the game thread instead of the combined GC collision, as the GC collision can be more coarse for improved physics simulation performance.

The cvar p.GeometryCollection.ForceOverrideGCCollisionSetupForTraces can be set to 1 to force GC collision on everything if anything goes sideways.

#tests PIE

diff --git a/Engine/Source/Runtime/Engine/Public/Physics/PhysicsInterfaceTypes.h b/Engine/Source/Runtime/Engine/Public/Physics/PhysicsInterfaceTypes.h
--- a/Engine/Source/Runtime/Engine/Public/Physics/PhysicsInterfaceTypes.h
+++ b/Engine/Source/Runtime/Engine/Public/Physics/PhysicsInterfaceTypes.h
@@ -41,6 +41,25 @@ struct FPhysicalMaterialMaskParams
 
 struct FGeometryAddParams
 {
+	FGeometryAddParams() = default;
+
+	FGeometryAddParams(const FGeometryAddParams& Other)
+	: bDoubleSided(Other.bDoubleSided)
+	, CollisionData(Other.CollisionData)
+	, CollisionTraceType(Other.CollisionTraceType)
+	, Scale(Other.Scale)
+	, SimpleMaterial(Other.SimpleMaterial)
+	, ComplexMaterials(Other.ComplexMaterials)
+	, ComplexMaterialMasks(Other.ComplexMaterialMasks)
+	, LocalTransform(Other.LocalTransform)
+	, WorldTransform(Other.WorldTransform)
+	, Geometry(Other.Geometry)
+	PRAGMA_DISABLE_DEPRECATION_WARNINGS
+	, ChaosTriMeshes(Other.ChaosTriMeshes)
+	PRAGMA_ENABLE_DEPRECATION_WARNINGS
+	, TriMeshGeometries(Other.TriMeshGeometries)
+	{}
+
 	bool bDoubleSided;
 	FBodyCollisionData CollisionData;
 	ECollisionTraceFlag CollisionTraceType;
diff --git a/Engine/Source/Runtime/Experimental/Chaos/Private/PhysicsProxy/GeometryCollectionPhysicsProxy.cpp b/Engine/Source/Runtime/Experimental/Chaos/Private/PhysicsProxy/GeometryCollectionPhysicsProxy.cpp
--- a/Engine/Source/Runtime/Experimental/Chaos/Private/PhysicsProxy/GeometryCollectionPhysicsProxy.cpp
+++ b/Engine/Source/Runtime/Experimental/Chaos/Private/PhysicsProxy/GeometryCollectionPhysicsProxy.cpp
@@ -38,6 +38,8 @@
 #include "GeometryCollection/Facades/CollectionAnchoringFacade.h"
 #include "GeometryCollection/Facades/CollectionConnectionGraphFacade.h"
 
+#include "Physics/Experimental/ChaosInterfaceUtils.h"
+
 #ifndef TODO_REIMPLEMENT_INIT_COMMANDS
 #define TODO_REIMPLEMENT_INIT_COMMANDS 0
 #endif
@@ -137,6 +139,27 @@ FAutoConsoleVariableRef CVarGeometryCollectionRotationUpdateTolerance(
 	GeometryCollectionRotationUpdateTolerance,
 	TEXT("Tolerance to detect if particle rotation has changed has changed when syncing PT to GT"));
 
+enum EOverrideGCCollisionSetupForTraces
+{
+	GCCSFT_Property   = -1,  // Default: do what property says
+	GCCSFT_ForceSM    =  0,  // Force the use of SM collision
+	GCCSFT_ForceGC    =  1,  // Force the use of GC collision
+};
+
+#if !UE_BUILD_SHIPPING
+
+int32 ForceOverrideGCCollisionSetupForTraces = GCCSFT_Property;
+FAutoConsoleVariableRef CVarForceOverrideGCCollisionSetupForTraces(
+	TEXT("p.GeometryCollection.ForceOverrideGCCollisionSetupForTraces"),
+	ForceOverrideGCCollisionSetupForTraces,
+	TEXT("Force the usage of a specific type of collision for traces on the game thread when creating new GC physics representations (-1: use the value of the property | 0: force to use SM collision | 1: force to use GC collision) [def: -1]"));
+
+#else
+
+constexpr int32 ForceOverrideGCCollisionSetupForTraces = GCCSFT_Property;
+
+#endif // !UE_BUILD_SHIPPING
+
 
 DEFINE_LOG_CATEGORY_STATIC(UGCC_LOG, Error, All);
 
@@ -584,12 +607,31 @@ void FGeometryCollectionPhysicsProxy::Initialize(Chaos::FPBDRigidsEvolutionBase
 				P->SetUserData(Parameters.UserData);
 				P->SetProxy(this);
 
-				Chaos::FImplicitObjectPtr ImplicitGeometry = Implicits[Index];
-				if (ImplicitGeometry && !Scale.Equals(FVector::OneVector))
+				if (Index == Parameters.InitialRootIndex && bUseStaticMeshCollisionForTraces && (OptionalTraceCollisionOverrideGeoms.Num() > 0))
+				{
+					Chaos::FShapesArray Shapes;
+					TArray<Chaos::FImplicitObjectPtr> Geoms;
+					for (const FTraceCollisionOverrideData& OverrideData : OptionalTraceCollisionOverrideGeoms)
+					{
+						// Copying Params so we can replace the local transform
+						FGeometryAddParams AddParams = OverrideData.Params;
+						AddParams.LocalTransform = AddParams.LocalTransform * MassToLocal[Index].Inverse();
+
+						OverrideData.CreateGeometryCallback(AddParams, Geoms, Shapes);
+					}
+
+					Chaos::FImplicitObjectPtr ImplicitGeometry = MakeImplicitObjectPtr<Chaos::FImplicitObjectUnion>(MoveTemp(Geoms));
+					P->SetGeometry(ImplicitGeometry);
+				}
+				else
 				{
-					ImplicitGeometry = ImplicitGeometry->CopyGeometryWithScale(Scale);
+					Chaos::FImplicitObjectPtr ImplicitGeometry = Implicits[Index];
+					if (ImplicitGeometry && !Scale.Equals(FVector::OneVector))
+					{
+						ImplicitGeometry = ImplicitGeometry->CopyGeometryWithScale(Scale);
+					}
+					P->SetGeometry(ImplicitGeometry);
 				}
-				P->SetGeometry(ImplicitGeometry);
 
 				if (DynamicCollectionAnchoringFacade.IsAnchored(Index))
 				{
@@ -2925,6 +2967,59 @@ void FGeometryCollectionPhysicsProxy::SetWorldTransform_External(const FTransfor
 	SetProxyDirty_External();
 }
 
+DECLARE_CYCLE_STAT(TEXT("FGeometryCollectionPhysicsProxy::SetUseStaticMeshCollisionForTraces_External"), STAT_SetUseStaticMeshCollisionForTraces_External, STATGROUP_Chaos);
+void FGeometryCollectionPhysicsProxy::SetUseStaticMeshCollisionForTraces_External(bool bInUseStaticMeshCollisionForTraces)
+{
+	SCOPE_CYCLE_COUNTER(STAT_SetUseStaticMeshCollisionForTraces_External);
+
+	check(IsInGameThread());
+
+	// Expensive, so only do it when we need to
+	const bool bUseSMCollisionForTraces = ((bInUseStaticMeshCollisionForTraces && ForceOverrideGCCollisionSetupForTraces == GCCSFT_Property) || (ForceOverrideGCCollisionSetupForTraces == GCCSFT_ForceSM));
+	if (bUseStaticMeshCollisionForTraces != bUseSMCollisionForTraces && (OptionalTraceCollisionOverrideGeoms.Num() > 0))
+	{
+		const int32 Index = Parameters.InitialRootIndex;
+		if (GTParticles.Num() > Index)
+		{
+			if (FParticle* P = GTParticles[Index].Get())
+			{
+				if (bUseSMCollisionForTraces)
+				{
+					const TManagedArray<FTransform>& MassToLocal = Parameters.RestCollection->GetAttribute<FTransform>("MassToLocal", FGeometryCollection::TransformGroup);
+
+					Chaos::FShapesArray Shapes;
+					TArray<Chaos::FImplicitObjectPtr> Geoms;
+					for (const FTraceCollisionOverrideData& OverrideData : OptionalTraceCollisionOverrideGeoms)
+					{
+						// Copying Params so we can replace the local transform
+						FGeometryAddParams AddParams = OverrideData.Params;
+						AddParams.LocalTransform = AddParams.LocalTransform * MassToLocal[Index].Inverse();
+
+						OverrideData.CreateGeometryCallback(AddParams, Geoms, Shapes);
+					}
+
+					Chaos::FImplicitObjectPtr ImplicitGeometry = MakeImplicitObjectPtr<Chaos::FImplicitObjectUnion>(MoveTemp(Geoms));
+					P->SetGeometry(ImplicitGeometry);
+				}
+				else
+				{
+					const FVector Scale = Parameters.WorldTransform.GetScale3D();
+					TManagedArray<Chaos::FImplicitObjectPtr>& Implicits = GameThreadCollection.ModifyAttribute<Chaos::FImplicitObjectPtr>(FGeometryDynamicCollection::ImplicitsAttribute, FTransformCollection::TransformGroup);
+					Chaos::FImplicitObjectPtr ImplicitGeometry = Implicits[Index];
+					if (ImplicitGeometry && !Scale.Equals(FVector::OneVector))
+					{
+						ImplicitGeometry = ImplicitGeometry->CopyGeometryWithScale(Scale);
+					}
+					P->SetGeometry(ImplicitGeometry);
+				}
+
+				// We only want to change the value of our boolean if we actually changed the collision, so the collision state matches
+				bUseStaticMeshCollisionForTraces = bUseSMCollisionForTraces;
+			}
+		}
+	}
+}
+
 void FGeometryCollectionPhysicsProxy::SetDamageThresholds_External(const TArray<float>& DamageThresholds)
 {
 	check(IsInGameThread());
@@ -5367,6 +5462,11 @@ Chaos::FPhysicsObjectHandle FGeometryCollectionPhysicsProxy::GetPhysicsObjectByI
 	return PhysicsObjects[Index].Get();
 }
 
+void FGeometryCollectionPhysicsProxy::RegisterNewTraceCollisionOverrideData(const FGeometryAddParams& InParams, FTraceCollisionCreateGeometryCallback InCreateGeometryCallback)
+{
+	OptionalTraceCollisionOverrideGeoms.Emplace(InParams, InCreateGeometryCallback);
+}
+
 FGeometryCollectionPhysicsProxy::FParticle* FGeometryCollectionPhysicsProxy::GetParticleByIndex_External(int32 Index)
 {
 	return const_cast<FGeometryCollectionPhysicsProxy::FParticle*>(const_cast<const FGeometryCollectionPhysicsProxy*>(this)->GetParticleByIndex_External(Index));
diff --git a/Engine/Source/Runtime/Experimental/Chaos/Public/PhysicsProxy/GeometryCollectionPhysicsProxy.h b/Engine/Source/Runtime/Experimental/Chaos/Public/PhysicsProxy/GeometryCollectionPhysicsProxy.h
--- a/Engine/Source/Runtime/Experimental/Chaos/Public/PhysicsProxy/GeometryCollectionPhysicsProxy.h
+++ b/Engine/Source/Runtime/Experimental/Chaos/Public/PhysicsProxy/GeometryCollectionPhysicsProxy.h
@@ -18,6 +18,7 @@
 #include "PBDRigidsSolver.h"
 #include "Chaos/Defines.h"
 #include "Chaos/GeometryParticlesfwd.h"
+#include "Physics/PhysicsInterfaceTypes.h"
 
 namespace Chaos
 {
@@ -272,6 +273,10 @@ public:
 	// set the world transform ( this needs to be called on the game thread ) 
 	CHAOS_API void SetWorldTransform_External(const FTransform& WorldTransform);
 
+	// todo(chaos): Remove this and move to a cook time approach of the SM data based on the GC property
+	// Set whether the GC should be using collision from the Static Mesh or the GC itself for game thread traces ( this needs to be called on the game thread )
+	CHAOS_API void SetUseStaticMeshCollisionForTraces_External(bool bInUseStaticMeshCollisionForTraces);
+
 	const TArray<FClusterHandle*> GetParticles() const
 	{
 		return SolverParticleHandles;
@@ -308,6 +313,12 @@ public:
 	CHAOS_API FParticleHandle* GetParticleByIndex_Internal(int32 Index);
 	CHAOS_API const FParticleHandle* GetParticleByIndex_Internal(int32 Index) const;
 
+	FParticle* GetInitialRootParticle_External() { return GetParticleByIndex_External(Parameters.InitialRootIndex); }
+	const FParticle* GetInitialRootParticle_External() const { return GetParticleByIndex_External(Parameters.InitialRootIndex); }
+
+	FParticleHandle* GetInitialRootParticle_Internal() { return GetParticleByIndex_Internal(Parameters.InitialRootIndex); }
+	const FParticleHandle* GetInitialRootParticle_Internal() const { return GetParticleByIndex_Internal(Parameters.InitialRootIndex); }
+
 	/**
 	*  * Get all the geometry collection particle handles based on the processing resolution
 	 */
@@ -483,6 +494,11 @@ public:
 	CHAOS_API TArray<Chaos::FPhysicsObjectHandle> GetAllPhysicsObjectIncludingNulls() const;
 	CHAOS_API Chaos::FPhysicsObjectHandle GetPhysicsObjectByIndex(int32 Index) const;
 	int32 GetNumParticles() const { return NumParticles; }
+
+	// todo(chaos): Remove this and move to a cook time approach of the SM data based on the GC property
+	using FTraceCollisionCreateGeometryCallback = TFunction<void(const FGeometryAddParams& InParams, TArray<Chaos::FImplicitObjectPtr>& OutGeoms, Chaos::FShapesArray& OutShapes)>;
+	CHAOS_API void RegisterNewTraceCollisionOverrideData(const FGeometryAddParams& InParams, FTraceCollisionCreateGeometryCallback InCreateGeometryCallback);
+
 protected:
 
 	CHAOS_API float ComputeMaterialBasedDamageThreshold_Internal(int32 TransformIndex) const;
@@ -603,6 +619,25 @@ private:
 	TArray<FBox> ValidGeometryBoundingBoxes;
 	TArray<int32> ValidGeometryTransformIndices;
 
+	// todo(chaos): Remove this and move to a cook time approach of the SM data based on the GC property
+	// These are created off of the static meshes in the event we want to use SM collision on the game thread
+	struct FTraceCollisionOverrideData
+	{
+		FTraceCollisionOverrideData() = default;
+
+		FTraceCollisionOverrideData(const FGeometryAddParams& InParams, FTraceCollisionCreateGeometryCallback InCallback)
+		: Params(InParams)
+		, CreateGeometryCallback(InCallback)
+		{}
+
+		FGeometryAddParams Params;
+		FTraceCollisionCreateGeometryCallback CreateGeometryCallback;
+	};
+	TArray<FTraceCollisionOverrideData> OptionalTraceCollisionOverrideGeoms;
+	
+	// Use SetUseStaticMeshCollisionForTraces_External intead of setting this directly, so that the collision can be recreated if needed
+	bool bUseStaticMeshCollisionForTraces = false;
+
 #ifdef TODO_REIMPLEMENT_RIGID_CACHING
 	TFunction<void(void)> ResetAnimationCacheCallback;
 	TFunction<void(const TArrayView<FTransform> &)> UpdateTransformsCallback;
diff --git a/Engine/Source/Runtime/Experimental/GeometryCollectionEngine/Private/GeometryCollection/GeometryCollectionComponent.cpp b/Engine/Source/Runtime/Experimental/GeometryCollectionEngine/Private/GeometryCollection/GeometryCollectionComponent.cpp
--- a/Engine/Source/Runtime/Experimental/GeometryCollectionEngine/Private/GeometryCollection/GeometryCollectionComponent.cpp
+++ b/Engine/Source/Runtime/Experimental/GeometryCollectionEngine/Private/GeometryCollection/GeometryCollectionComponent.cpp
@@ -37,6 +37,7 @@
 #include "Net/UnrealNetwork.h"
 #include "PhysicalMaterials/PhysicalMaterial.h"
 #include "Physics/Experimental/PhysScene_Chaos.h"
+#include "Physics/Experimental/ChaosInterfaceUtils.h"
 #include "Physics/PhysicsFiltering.h"
 #include "PhysicsEngine/PhysicsObjectExternalInterface.h"
 #include "PhysicsField/PhysicsFieldComponent.h"
@@ -3812,6 +3813,43 @@ void UGeometryCollectionComponent::RegisterAndInitializePhysicsProxy()
 		}
 	}
 
+	// Always set up the overrides in case the boolean flips later to allow SM collision on GT
+	// todo(chaos): Remove this and move to a cook time approach of the SM data based on the GC property
+	if (RestCollection != nullptr)
+	{
+		for (const TObjectPtr<UStaticMesh>& ProxyMesh : RestCollection->RootProxyData.ProxyMeshes)
+		{
+			if (const UStaticMesh* StaticMesh = ProxyMesh.Get())
+			{
+				if (UBodySetup* BodySetup = ProxyMesh->GetBodySetup())
+				{
+					FBodyCollisionData BodyCollisionData;
+					BodyInstance.BuildBodyFilterData(BodyCollisionData.CollisionFilterData);
+					FBodyInstance::BuildBodyCollisionFlags(BodyCollisionData.CollisionFlags, BodyInstance.GetCollisionEnabled(), BodySetup->GetCollisionTraceFlag() == CTF_UseComplexAsSimple);
+
+					FGeometryAddParams AddParams;
+					AddParams.bDoubleSided = BodySetup->bDoubleSidedGeometry;
+					AddParams.CollisionData = BodyCollisionData;
+					AddParams.CollisionTraceType = BodySetup->GetCollisionTraceFlag();
+					AddParams.Scale = BodyInstance.Scale3D;
+					AddParams.SimpleMaterial = BodyInstance.GetSimplePhysicalMaterial();
+					AddParams.LocalTransform = FTransform::Identity;
+					AddParams.WorldTransform = PhysicsProxy->GetSimParameters().WorldTransform;
+					AddParams.Geometry = &BodySetup->AggGeom;
+
+					PhysicsProxy->RegisterNewTraceCollisionOverrideData(
+						AddParams,
+						[](const FGeometryAddParams& InParams, TArray<Chaos::FImplicitObjectPtr>& OutGeoms, Chaos::FShapesArray& OutShapes) {
+							ChaosInterface::CreateGeometry(InParams, OutGeoms, OutShapes);
+						}
+					);
+				}
+			}
+		}
+	}
+
+	PhysicsProxy->SetUseStaticMeshCollisionForTraces_External(bUseStaticMeshCollisionForTraces);
+
 	FPhysScene_Chaos* Scene = GetInnerChaosScene();
 	Scene->AddObject(this, PhysicsProxy);
 
@@ -5976,6 +6014,18 @@ TArray<UStaticMeshComponent*> UGeometryCollectionComponent::CreateProxyComponent
 	return Components;
 }
 
+void UGeometryCollectionComponent::SetUseStaticMeshCollisionForTraces(const bool bInUseStaticMeshCollisionForTraces)
+{
+	if (bUseStaticMeshCollisionForTraces != bInUseStaticMeshCollisionForTraces)
+	{
+		if (PhysicsProxy != nullptr)
+		{
+			PhysicsProxy->SetUseStaticMeshCollisionForTraces_External(bInUseStaticMeshCollisionForTraces);
+		}
+
+		bUseStaticMeshCollisionForTraces = bInUseStaticMeshCollisionForTraces;
+	}
+}
 struct FGeometryCollectionDecayContext
 {
 	FGeometryCollectionDecayContext(FGeometryCollectionPhysicsProxy& PhysicsProxyIn, FGeometryCollectionDecayDynamicFacade& DecayFacadeIn)
diff --git a/Engine/Source/Runtime/Experimental/GeometryCollectionEngine/Public/GeometryCollection/GeometryCollectionComponent.h b/Engine/Source/Runtime/Experimental/GeometryCollectionEngine/Public/GeometryCollection/GeometryCollectionComponent.h
--- a/Engine/Source/Runtime/Experimental/GeometryCollectionEngine/Public/GeometryCollection/GeometryCollectionComponent.h
+++ b/Engine/Source/Runtime/Experimental/GeometryCollectionEngine/Public/GeometryCollection/GeometryCollectionComponent.h
@@ -1282,6 +1282,14 @@ public:
 
 	GEOMETRYCOLLECTIONENGINE_API void SetUpdateNavigationInTick(const bool bUpdateInTick) { bUpdateNavigationInTick = bUpdateInTick; }
 
+	// todo(chaos): Remove this and move to a cook time approach of the SM data based on the GC property
+	UFUNCTION(BlueprintPure, BlueprintInternalUseOnly, Category = "Physics")
+	GEOMETRYCOLLECTIONENGINE_API bool GetUseStaticMeshCollisionForTraces() const { return bUseStaticMeshCollisionForTraces; }
+
+	// todo(chaos): Remove this and move to a cook time approach of the SM data based on the GC property
+	UFUNCTION(BlueprintCallable, BlueprintInternalUseOnly, Category = "Physics")
+	GEOMETRYCOLLECTIONENGINE_API void SetUseStaticMeshCollisionForTraces(bool bInUseStaticMeshCollisionForTraces);
+
 	/** Get any custom renderer. Returns nullptr if none is set. */
 	GEOMETRYCOLLECTIONENGINE_API IGeometryCollectionExternalRenderInterface* GetCustomRenderer() { return CustomRenderer.GetInterface(); }
 	/** Enable or disable root proxy for custom rendering - this can be set at runtime */
@@ -1675,6 +1683,10 @@ private:
 	TArray<int32> EmbeddedInstanceIndex;
 #endif
 
+	// todo(chaos): Remove the ability to change this at runtime, as we'll want to use this at cook time instead
+	UPROPERTY(EditAnywhere, BlueprintGetter="GetUseStaticMeshCollisionForTraces", BlueprintSetter="SetUseStaticMeshCollisionForTraces", Category = "Physics")
+	bool bUseStaticMeshCollisionForTraces  = false;
+
 	GEOMETRYCOLLECTIONENGINE_API bool IsEmbeddedGeometryValid() const;
 	GEOMETRYCOLLECTIONENGINE_API void ClearEmbeddedGeometry();
 
