Added flag to control whether ISM cast far shadow in FGeometryCollectionISM.

diff --git a/Engine/Source/Runtime/Experimental/GeometryCollectionEngine/Private/GeometryCollection/GeometryCollectionISMPoolComponent.cpp b/Engine/Source/Runtime/Experimental/GeometryCollectionEngine/Private/GeometryCollection/GeometryCollectionISMPoolComponent.cpp
--- a/Engine/Source/Runtime/Experimental/GeometryCollectionEngine/Private/GeometryCollection/GeometryCollectionISMPoolComponent.cpp
+++ b/Engine/Source/Runtime/Experimental/GeometryCollectionEngine/Private/GeometryCollection/GeometryCollectionISMPoolComponent.cpp
@@ -130,6 +130,7 @@ void FGeometryCollectionISM::InitISM(const FGeometryCollectionStaticMeshInstance
 	ISMComponent->SetCastShadow((MeshInstance.Desc.Flags & FISMComponentDescription::AffectShadow) != 0);
 	ISMComponent->bAffectDynamicIndirectLighting = (MeshInstance.Desc.Flags & FISMComponentDescription::AffectDynamicIndirectLighting) != 0;
 	ISMComponent->bAffectDistanceFieldLighting = (MeshInstance.Desc.Flags & FISMComponentDescription::AffectDistanceFieldLighting) != 0;
+	ISMComponent->bCastFarShadow = (MeshInstance.Desc.Flags & FISMComponentDescription::AffectFarShadow) != 0;
 	ISMComponent->bWorldPositionOffsetWritesVelocity = (MeshInstance.Desc.Flags & FISMComponentDescription::WorldPositionOffsetWritesVelocity) != 0;
 	ISMComponent->bEvaluateWorldPositionOffset = (MeshInstance.Desc.Flags & FISMComponentDescription::EvaluateWorldPositionOffset) != 0;
 	ISMComponent->bUseGpuLodSelection = (MeshInstance.Desc.Flags & FISMComponentDescription::GpuLodSelection) != 0;
diff --git a/Engine/Source/Runtime/Experimental/GeometryCollectionEngine/Public/GeometryCollection/GeometryCollectionISMPoolComponent.h b/Engine/Source/Runtime/Experimental/GeometryCollectionEngine/Public/GeometryCollection/GeometryCollectionISMPoolComponent.h
--- a/Engine/Source/Runtime/Experimental/GeometryCollectionEngine/Public/GeometryCollection/GeometryCollectionISMPoolComponent.h
+++ b/Engine/Source/Runtime/Experimental/GeometryCollectionEngine/Public/GeometryCollection/GeometryCollectionISMPoolComponent.h
@@ -111,7 +111,8 @@ struct FISMComponentDescription
 		AffectShadow = 1 << 7,
 		AffectDistanceFieldLighting = 1 << 8,
 		AffectDynamicIndirectLighting = 1 << 9,
-		DistanceCullPrimitive = 1 << 10,
+		AffectFarShadow = 1 << 10,
+		DistanceCullPrimitive = 1 << 11,
 	};
 
 	uint32 Flags = WorldPositionOffsetWritesVelocity|EvaluateWorldPositionOffset|AffectShadow;
