Iris
* Now printing the full replication state of the CDO when a client detects it has a protocol from the server.
* Client now tells the server when it detects a protocol mismatch
* Server will then also print the full replication of the CDO of the object

#rnx

diff --git a/Engine/Source/Runtime/Engine/Private/DataChannel.cpp b/Engine/Source/Runtime/Engine/Private/DataChannel.cpp
--- a/Engine/Source/Runtime/Engine/Private/DataChannel.cpp
+++ b/Engine/Source/Runtime/Engine/Private/DataChannel.cpp
@@ -11,6 +11,7 @@
 #include "Engine/Level.h"
 #include "Engine/World.h"
 #include "GameFramework/WorldSettings.h"
+#include "Iris/ReplicationSystem/ReplicationSystem.h"
 #include "Misc/MemStack.h"
 #include "Misc/ScopeExit.h"
 #include "Net/Core/Trace/Private/NetTraceInternal.h"
@@ -1638,6 +1639,7 @@ IMPLEMENT_CONTROL_CHANNEL_MESSAGE(BeaconWelcome);
 IMPLEMENT_CONTROL_CHANNEL_MESSAGE(BeaconJoin);
 IMPLEMENT_CONTROL_CHANNEL_MESSAGE(BeaconAssignGUID);
 IMPLEMENT_CONTROL_CHANNEL_MESSAGE(BeaconNetGUIDAck);
+IMPLEMENT_CONTROL_CHANNEL_MESSAGE(IrisProtocolMismatch);
 
 void UControlChannel::Init( UNetConnection* InConnection, int32 InChannelIndex, EChannelCreateFlags CreateFlags )
 {
@@ -1845,6 +1847,19 @@ void UControlChannel::ReceivedBunch( FInBunch& Bunch )
 				UE::Net::FNetPing::HandleNetPingControlMessage(Connection, NetPingMessageType, MessageStr);
 			}
 		}
+		else if (MessageType == NMT_IrisProtocolMismatch)
+		{
+			uint64 NetRefHandleId = 0;
+			if (FNetControlMessage<NMT_IrisProtocolMismatch>::Receive(Bunch, NetRefHandleId))
+			{
+#if UE_WITH_IRIS
+				if (UReplicationSystem* IrisRepSystem = Connection->Driver->GetReplicationSystem())
+				{
+					IrisRepSystem->ReportProtocolMismatch(NetRefHandleId, Connection->GetConnectionId());
+				}
+#endif
+			}
+		}
 		else if (Connection->Driver->Notify != nullptr)
 		{
 			// Process control message on client/server connection
@@ -1922,6 +1937,9 @@ void UControlChannel::ReceivedBunch( FInBunch& Bunch )
 				case NMT_BeaconNetGUIDAck:
 					FNetControlMessage<NMT_BeaconNetGUIDAck>::Discard(Bunch);
 					break;
+				case NMT_IrisProtocolMismatch:
+					FNetControlMessage<NMT_IrisProtocolMismatch>::Discard(Bunch);
+					break;
 				default:
 					// if this fails, a case is missing above for an implemented message type
 					// or the connection is being sent potentially malformed packets
diff --git a/Engine/Source/Runtime/Engine/Private/Net/Iris/ReplicationSystem/ActorReplicationBridge.cpp b/Engine/Source/Runtime/Engine/Private/Net/Iris/ReplicationSystem/ActorReplicationBridge.cpp
--- a/Engine/Source/Runtime/Engine/Private/Net/Iris/ReplicationSystem/ActorReplicationBridge.cpp
+++ b/Engine/Source/Runtime/Engine/Private/Net/Iris/ReplicationSystem/ActorReplicationBridge.cpp
@@ -31,6 +31,7 @@
 #include "Engine/World.h"
 #include "GameFramework/Actor.h"
 #include "Net/DataBunch.h"
+#include "Net/DataChannel.h"
 #include "Net/Core/Misc/NetSubObjectRegistry.h"
 #include "Net/NetSubObjectRegistryGetter.h"
 #include "Templates/Casts.h"
@@ -1151,6 +1152,18 @@ void UActorReplicationBridge::WakeUpObjectInstantiatedFromRemote(AActor* Actor)
 	}
 }
 
+void UActorReplicationBridge::OnProtocolMismatchDetected(FNetRefHandle ObjectHandle)
+{
+	Super::OnProtocolMismatchDetected(ObjectHandle);
+
+	// As a client tell the server we could not bind this specific NetRefHandle
+	if (NetDriver && NetDriver->ServerConnection)
+	{
+		uint64 RawHandleId = ObjectHandle.GetId();
+		FNetControlMessage<NMT_IrisProtocolMismatch>::Send(NetDriver->ServerConnection, RawHandleId);
+	}
+}
+
 #else //!UE_WITH_IRIS
 
 UActorReplicationBridge::UActorReplicationBridge() = default;
diff --git a/Engine/Source/Runtime/Engine/Public/Net/DataChannel.h b/Engine/Source/Runtime/Engine/Public/Net/DataChannel.h
--- a/Engine/Source/Runtime/Engine/Public/Net/DataChannel.h
+++ b/Engine/Source/Runtime/Engine/Public/Net/DataChannel.h
@@ -194,3 +194,4 @@ DEFINE_CONTROL_CHANNEL_MESSAGE(BeaconWelcome, 25); // server tells client they'r
 DEFINE_CONTROL_CHANNEL_MESSAGE(BeaconJoin, 26, FString, FUniqueNetIdRepl);  // server tries to create beacon type requested by client, sends NetGUID for actor sync
 DEFINE_CONTROL_CHANNEL_MESSAGE(BeaconAssignGUID, 27, FNetworkGUID); // client assigns NetGUID from server to beacon actor, sends NetGUIDAck
 DEFINE_CONTROL_CHANNEL_MESSAGE(BeaconNetGUIDAck, 28, FString); // server received NetGUIDAck from client, connection established successfully
+DEFINE_CONTROL_CHANNEL_MESSAGE(IrisProtocolMismatch, 29, uint64); // client has a different protocol hash from the server
diff --git a/Engine/Source/Runtime/Engine/Public/Net/Iris/ReplicationSystem/ActorReplicationBridge.h b/Engine/Source/Runtime/Engine/Public/Net/Iris/ReplicationSystem/ActorReplicationBridge.h
--- a/Engine/Source/Runtime/Engine/Public/Net/Iris/ReplicationSystem/ActorReplicationBridge.h
+++ b/Engine/Source/Runtime/Engine/Public/Net/Iris/ReplicationSystem/ActorReplicationBridge.h
@@ -98,6 +98,7 @@ protected:
 	virtual bool RemapPathForPIE(uint32 ConnectionId, FString& Path, bool bReading) const override;
 	virtual bool ObjectLevelHasFinishedLoading(UObject* Object) const override;
 	virtual bool IsAllowedToDestroyInstance(const UObject* Instance) const override;
+	virtual void OnProtocolMismatchDetected(FNetRefHandle ObjectHandle) override;
 
 	virtual float GetPollFrequencyOfRootObject(const UObject* ReplicatedObject) const override;
 
diff --git a/Engine/Source/Runtime/Experimental/Iris/Core/Private/Iris/ReplicationSystem/ObjectReplicationBridge.cpp b/Engine/Source/Runtime/Experimental/Iris/Core/Private/Iris/ReplicationSystem/ObjectReplicationBridge.cpp
--- a/Engine/Source/Runtime/Experimental/Iris/Core/Private/Iris/ReplicationSystem/ObjectReplicationBridge.cpp
+++ b/Engine/Source/Runtime/Experimental/Iris/Core/Private/Iris/ReplicationSystem/ObjectReplicationBridge.cpp
@@ -2,7 +2,9 @@
 
 #include "Iris/ReplicationSystem/ObjectReplicationBridge.h"
 
+#include "Misc/ScopeExit.h"
 #include "HAL/IConsoleManager.h"
+
 #include "Iris/IrisConfigInternal.h"
 
 #include "Iris/Core/IrisLog.h"
@@ -37,6 +39,15 @@
 #include "Iris/Serialization/InternalNetSerializationContext.h"
 #include "Iris/Serialization/NetBitStreamUtil.h"
 
+namespace UE::Net::Private::ObjectBridgeDebugging
+{
+	extern void RemoteProtocolMismatchDetected(UReplicationSystem* , uint32, const FReplicationFragments&, const UObject*, const UObject*);
+}
+
+#ifndef UE_IRIS_VALIDATE_PROTOCOLS
+#	define UE_IRIS_VALIDATE_PROTOCOLS !UE_BUILD_SHIPPING
+#endif
+
 DEFINE_LOG_CATEGORY(LogIrisFilterConfig)
 
 #define UE_LOG_OBJECTREPLICATIONBRIDGE(Category, Format, ...)  UE_LOG(LogIrisBridge, Category, TEXT("ObjectReplicationBridge(%u)::") Format, GetReplicationSystem()->GetId(), ##__VA_ARGS__)
@@ -278,16 +289,17 @@ UE::Net::FNetRefHandle UObjectReplicationBridge::BeginReplication(UObject* Insta
 
 	FReplicationInstanceProtocolPtr InstanceProtocol(ProtocolManager->CreateInstanceProtocol(RegisteredFragments));
 		
-	FReplicationProtocolIdentifier ProtocolIdentifier = ProtocolManager->CalculateProtocolIdentifier(RegisteredFragments);
+	FReplicationProtocolIdentifier ProtocolIdentifier = FReplicationProtocolManager::CalculateProtocolIdentifier(RegisteredFragments);
 	const FReplicationProtocol* ReplicationProtocol = ProtocolManager->GetReplicationProtocol(ProtocolIdentifier, ArchetypeOrCDOUsedAsKey);
 	if (!ReplicationProtocol)
 	{
-		ReplicationProtocol = ProtocolManager->CreateReplicationProtocol(ArchetypeOrCDOUsedAsKey, ProtocolIdentifier, RegisteredFragments, *Instance->GetClass()->GetName(), false);
+		constexpr bool bIgnoreProtocolValidation = false;
+		ReplicationProtocol = ProtocolManager->CreateReplicationProtocol(ArchetypeOrCDOUsedAsKey, ProtocolIdentifier, RegisteredFragments, *Instance->GetClass()->GetName(), bIgnoreProtocolValidation);
 	}
-#if !(UE_BUILD_SHIPPING || UE_BUILD_TEST)
+#if UE_IRIS_VALIDATE_PROTOCOLS
 	else
 	{
-		const bool bIsValidProtocol = ProtocolManager->ValidateReplicationProtocol(ReplicationProtocol, RegisteredFragments);
+		const bool bIsValidProtocol = FReplicationProtocolManager::ValidateReplicationProtocol(ReplicationProtocol, RegisteredFragments);
 		if (!bIsValidProtocol)
 		{
 			UE_LOG_OBJECTREPLICATIONBRIDGE(Error, TEXT("BeginReplication Found invalid protocol ProtocolId:0x%" UINT64_x_FMT " for Object named %s"), ReplicationProtocol->ProtocolIdentifier, *Instance->GetName());
@@ -638,12 +650,6 @@ FReplicationBridgeCreateNetRefHandleResult UObjectReplicationBridge::CreateNetRe
 	FFragmentRegistrationContext FragmentRegistrationContext(GetReplicationStateDescriptorRegistry(), EReplicationFragmentTraits::CanReceive);
 	FReplicationProtocolManager* ProtocolManager = GetReplicationProtocolManager();
 
-#if (UE_BUILD_SHIPPING || UE_BUILD_TEST)
-	const bool bVerifyExistingProtocol = false;
-#else
-	const bool bVerifyExistingProtocol = true;
-#endif
-
 	FReplicationBridgeCreateNetRefHandleResult CreateResult;
 
 	// Currently we need to always instantiate remote objects, moving forward we want to make this optional so that can be deferred until it is time to apply received state data.
@@ -652,7 +658,7 @@ FReplicationBridgeCreateNetRefHandleResult UObjectReplicationBridge::CreateNetRe
 	UObject* InstancePtr = InstantiateResult.Object;
 	if (!InstancePtr)
 	{
-		ensureAlwaysMsgf(bSuppressCreateInstanceFailedEnsure, TEXT("Failed to instantiate Handle: %s"), *WantedNetHandle.ToString());
+		ensureMsgf(bSuppressCreateInstanceFailedEnsure, TEXT("Failed to instantiate Handle: %s"), *WantedNetHandle.ToString());
 		return CreateResult;
 	}
 
@@ -674,17 +680,31 @@ FReplicationBridgeCreateNetRefHandleResult UObjectReplicationBridge::CreateNetRe
 	const FReplicationProtocol* ReplicationProtocol = ProtocolManager->GetReplicationProtocol(ReceivedProtocolId, ArchetypeOrCDOUsedAsKey);
 	if (!ReplicationProtocol)
 	{
-		ReplicationProtocol = ProtocolManager->CreateReplicationProtocol(ArchetypeOrCDOUsedAsKey, ReceivedProtocolId, RegisteredFragments, *(InstancePtr->GetClass()->GetName()), true);
+		constexpr bool bVerifyProtocol = true;
+		ReplicationProtocol = ProtocolManager->CreateReplicationProtocol(ArchetypeOrCDOUsedAsKey, ReceivedProtocolId, RegisteredFragments, *(InstancePtr->GetClass()->GetName()), bVerifyProtocol);
 	}
 	else
 	{
-		if (!ensureAlways(!bVerifyExistingProtocol || ProtocolManager->ValidateReplicationProtocol(ReplicationProtocol, RegisteredFragments)))
+		constexpr bool bDoNotLogErrors = false; // Don't log errors because it would spam for every individual object of the same class. 
+		bool bIsValid = FReplicationProtocolManager::ValidateReplicationProtocol(ReplicationProtocol, RegisteredFragments, bDoNotLogErrors);
+		if (!bIsValid)
 		{
-			return CreateResult;
+			ReplicationProtocol = nullptr;
 		}
 	}
 
-	if (ReplicationProtocol)
+	if (!ReplicationProtocol)
+	{
+		UE_LOG(LogIris, Error, TEXT("Protocol mismatch prevents binding %s to instanced object %s (CDO: %s)."), *WantedNetHandle.ToString(), *GetNameSafe(InstancePtr), *GetNameSafe(ArchetypeOrCDOUsedAsKey));
+
+		if (UE_LOG_ACTIVE(LogIris, Error))
+		{
+			UE::Net::Private::ObjectBridgeDebugging::RemoteProtocolMismatchDetected(ReplicationSystem, Context.ConnectionId, RegisteredFragments, ArchetypeOrCDOUsedAsKey, InstancePtr);
+		}
+
+		OnProtocolMismatchDetected(WantedNetHandle);
+	}
+	else
 	{
 		// Create NetHandle
 		FNetRefHandle Handle = InternalCreateNetObjectFromRemote(WantedNetHandle, ReplicationProtocol);
@@ -697,7 +717,7 @@ FReplicationBridgeCreateNetRefHandleResult UObjectReplicationBridge::CreateNetRe
 			(void)InstanceProtocol.Release();
 
 			// Now it is safe to issue OnActorChannelOpen callback
-			ensureAlwaysMsgf(OnInstantiatedFromRemote(InstancePtr, Header.Get(), Context.ConnectionId), TEXT("Failed to invoke OnInstantiatedFromRemote for Instance named %s %s"), *InstancePtr->GetName(), *Handle.ToString());
+			ensureMsgf(OnInstantiatedFromRemote(InstancePtr, Header.Get(), Context.ConnectionId), TEXT("Failed to invoke OnInstantiatedFromRemote for Instance named %s %s"), *InstancePtr->GetName(), *Handle.ToString());
 		}
 	}
 
@@ -1604,7 +1624,7 @@ void UObjectReplicationBridge::FindClassesInPollPeriodOverrides()
 
 void UObjectReplicationBridge::SetShouldUseDefaultSpatialFilterFunction(TFunction<bool(const UClass*)> InShouldUseDefaultSpatialFilterFunction)
 {
-	if (!ensureAlwaysMsgf((bool)InShouldUseDefaultSpatialFilterFunction, TEXT("%s"), TEXT("A valid function must be provided for SetShouldUseDefaultSpatialFilterFunction.")))
+	if (!ensureMsgf((bool)InShouldUseDefaultSpatialFilterFunction, TEXT("%s"), TEXT("A valid function must be provided for SetShouldUseDefaultSpatialFilterFunction.")))
 	{
 		return;
 	}
@@ -1614,7 +1634,7 @@ void UObjectReplicationBridge::SetShouldUseDefaultSpatialFilterFunction(TFunctio
 
 void UObjectReplicationBridge::SetShouldSubclassUseSameFilterFunction(TFunction<bool(const UClass* Class, const UClass* Subclass)> InShouldSubclassUseSameFilterFunction)
 {
-	if (!ensureAlwaysMsgf((bool)InShouldSubclassUseSameFilterFunction, TEXT("%s"), TEXT("A valid function must be provided for SetShouldSubclassUseSameFilterFunction.")))
+	if (!ensureMsgf((bool)InShouldSubclassUseSameFilterFunction, TEXT("%s"), TEXT("A valid function must be provided for SetShouldSubclassUseSameFilterFunction.")))
 	{
 		return;
 	}
@@ -1729,3 +1749,63 @@ void UObjectReplicationBridge::SetPollFrequency(FNetRefHandle RefHandle, float P
 		PollFrequencyLimiter->SetPollWithObject(RootObjectIndex, SubObjectIndex);
 	}
 }
+
+void UObjectReplicationBridge::OnProtocolMismatchReported(FNetRefHandle RefHandle, uint32 ConnectionId)
+{
+	using namespace UE::Net;
+	using namespace UE::Net::Private;
+
+	// Ensure at the end so the log contains the most relevant information already
+	ON_SCOPE_EXIT
+	{
+		ensure(false);
+	};
+	
+
+	const FInternalNetRefIndex ObjectInternalIndex = NetRefHandleManager->GetInternalIndex(RefHandle);
+	if (ObjectInternalIndex == FNetRefHandleManager::InvalidInternalIndex)
+	{
+		UE_LOG(LogIris, Warning, TEXT("OnProtocolMismatchReported from Connection:%u for %s. But object has no InternalIndex."), ConnectionId, *RefHandle.ToString());
+		return;
+	}
+
+	UObject* ObjInstance = NetRefHandleManager->GetReplicatedObjectInstance(ObjectInternalIndex);
+	UObject* ObjArchetype = ObjInstance ? ObjInstance->GetArchetype() : nullptr;
+
+	UE_LOG(LogIris, Error, TEXT("OnProtocolMismatchReported from client:%u when instancing %s. CDO:%s ReplicatedObject:%s"), ConnectionId, *RefHandle.ToString(), *GetNameSafe(ObjArchetype), *GetNameSafe(ObjInstance));
+
+	//$IRIS TODO: Tell ActorBridge so he can choose to disconnect the client if the actor was critical.
+
+	if (UE_LOG_ACTIVE(LogIris, Error))
+	{
+		const FNetRefHandleManager::FReplicatedObjectData& ObjectData = NetRefHandleManager->GetReplicatedObjectDataNoCheck(ObjectInternalIndex);
+		const FReplicationInstanceProtocol* InstanceProtocol = ObjectData.InstanceProtocol;
+		if (!InstanceProtocol)
+		{
+			UE_LOG(LogIris, Warning, TEXT("OnProtocolMismatchReported from Connection:%u for %s. But object %s has no InstanceProtocol."), ConnectionId, *RefHandle.ToString(), *GetNameSafe(ObjInstance));
+			return;
+		}
+
+		const FReplicationProtocol* Protocol = ObjectData.Protocol;
+		if (!Protocol)
+		{
+			UE_LOG(LogIris, Warning, TEXT("OnProtocolMismatchReported from Connection:%u for %s. But object %s has no Protocol."), ConnectionId, *RefHandle.ToString(), *GetNameSafe(ObjInstance));
+			return;
+		}
+
+		check(Protocol->ReplicationStateCount == InstanceProtocol->FragmentCount);
+	
+		// Build the list of fragments of this object
+		FReplicationFragments Fragments;
+		for (uint16 FragmentIndex=0; FragmentIndex < InstanceProtocol->FragmentCount; ++FragmentIndex)
+		{
+			FReplicationFragmentInfo FragmentInfo;
+			FragmentInfo.Fragment = InstanceProtocol->Fragments[FragmentIndex];
+			FragmentInfo.Descriptor = Protocol->ReplicationStateDescriptors[FragmentIndex];
+
+			Fragments.Emplace(MoveTemp(FragmentInfo));
+		}
+
+		UE::Net::Private::ObjectBridgeDebugging::RemoteProtocolMismatchDetected(ReplicationSystem, 0 /*TODO: Local ConnectionId*/, Fragments, ObjInstance, ObjArchetype);
+	}
+}
diff --git a/Engine/Source/Runtime/Experimental/Iris/Core/Private/Iris/ReplicationSystem/ObjectReplicationBridgeDebugging.cpp b/Engine/Source/Runtime/Experimental/Iris/Core/Private/Iris/ReplicationSystem/ObjectReplicationBridgeDebugging.cpp
--- a/Engine/Source/Runtime/Experimental/Iris/Core/Private/Iris/ReplicationSystem/ObjectReplicationBridgeDebugging.cpp
+++ b/Engine/Source/Runtime/Experimental/Iris/Core/Private/Iris/ReplicationSystem/ObjectReplicationBridgeDebugging.cpp
@@ -10,12 +10,55 @@
 #include "Iris/ReplicationSystem/ReplicationSystem.h"
 #include "Iris/ReplicationSystem/ReplicationSystemTypes.h"
 #include "Iris/ReplicationSystem/ReplicationSystemInternal.h"
+#include "Iris/ReplicationSystem/ReplicationOperations.h"
 
 #include "Net/Core/NetBitArrayPrinter.h"
 
 namespace UE::Net::Private::ObjectBridgeDebugging
 {
 
+void PrintDefaultNetObjectState(UReplicationSystem* ReplicationSystem, uint32 ConnectionId, const FReplicationFragments& RegisteredFragments, FStringBuilderBase& StringBuilder)
+{
+	FReplicationSystemInternal* ReplicationSystemInternal = ReplicationSystem->GetReplicationSystemInternal();
+
+	// In order to be able to output object references we need the TokenStoreState, for the server we just use the local one but if we are a client we must use the remote token store state
+	FReplicationConnections& Connections = ReplicationSystemInternal->GetConnections();
+	FNetTokenStoreState* TokenStoreState = ReplicationSystem->IsServer() ? ReplicationSystemInternal->GetNetTokenStore().GetLocalNetTokenStoreState() : &Connections.GetRemoteNetTokenStoreState(ConnectionId);
+
+	// Setup Context
+	FInternalNetSerializationContext InternalContext;
+	FInternalNetSerializationContext::FInitParameters InternalContextInitParams;
+	InternalContextInitParams.ReplicationSystem = ReplicationSystem;
+	InternalContextInitParams.PackageMap = ReplicationSystemInternal->GetIrisObjectReferencePackageMap();
+	InternalContextInitParams.ObjectResolveContext.RemoteNetTokenStoreState = TokenStoreState;
+	InternalContextInitParams.ObjectResolveContext.ConnectionId = ConnectionId;
+	InternalContext.Init(InternalContextInitParams);
+
+	FNetSerializationContext NetSerializationContext;
+	NetSerializationContext.SetInternalContext(&InternalContext);
+	NetSerializationContext.SetLocalConnectionId(ConnectionId);
+
+	FReplicationInstanceOperations::OutputInternalDefaultStateToString(NetSerializationContext, StringBuilder, RegisteredFragments);
+}
+
+void RemoteProtocolMismatchDetected(UReplicationSystem* ReplicationSystem, uint32 ConnectionId, const FReplicationFragments& RegisteredFragments, const UObject* ArchetypeOrCDOKey, const UObject* InstancePtr)
+{
+	if (UE_LOG_ACTIVE(LogIris, Error))
+	{
+		static TMap<FObjectKey, bool> ArchetypesAlreadyPrinted;
+
+		// Only print the CDO state once
+		if (ArchetypesAlreadyPrinted.Find(FObjectKey(ArchetypeOrCDOKey)) == nullptr)
+		{
+			ArchetypesAlreadyPrinted.Add(FObjectKey(ArchetypeOrCDOKey), true);
+
+			TStringBuilder<4096> StringBuilder;
+			PrintDefaultNetObjectState(ReplicationSystem, ConnectionId, RegisteredFragments, StringBuilder);
+			UE_LOG(LogIris, Error, TEXT("Printing replication state of CDO %s used for %s:\n%s"), *GetNameSafe(ArchetypeOrCDOKey), *GetNameSafe(InstancePtr), StringBuilder.ToString());
+		}
+	}
+}
+
 UReplicationSystem* FindReplicationSystem(const TArray<FString>& Args)
 {
 	uint32 RepSystemId = 0;
diff --git a/Engine/Source/Runtime/Experimental/Iris/Core/Private/Iris/ReplicationSystem/ReplicationBridge.cpp b/Engine/Source/Runtime/Experimental/Iris/Core/Private/Iris/ReplicationSystem/ReplicationBridge.cpp
--- a/Engine/Source/Runtime/Experimental/Iris/Core/Private/Iris/ReplicationSystem/ReplicationBridge.cpp
+++ b/Engine/Source/Runtime/Experimental/Iris/Core/Private/Iris/ReplicationSystem/ReplicationBridge.cpp
@@ -292,6 +292,7 @@ UReplicationBridge::~UReplicationBridge()
 void UReplicationBridge::Initialize(UReplicationSystem* InReplicationSystem)
 {
 	using namespace UE::Net;
+	using namespace UE::Net::Private;
 
 	Private::FReplicationSystemInternal* ReplicationSystemInternal = InReplicationSystem->GetReplicationSystemInternal();
 
@@ -304,8 +305,9 @@ void UReplicationBridge::Initialize(UReplicationSystem* InReplicationSystem)
 
 	// Create destruction info protocol
 	{
+		constexpr bool bIgnoreProtocolValidation = false;
 		const FReplicationFragments RegisteredFragments;
-		DestructionInfoProtocol = ReplicationProtocolManager->CreateReplicationProtocol(ReplicationProtocolManager->CalculateProtocolIdentifier(RegisteredFragments), RegisteredFragments, TEXT("InternalDestructionInfo"), false);
+		DestructionInfoProtocol = ReplicationProtocolManager->CreateReplicationProtocol(FReplicationProtocolManager::CalculateProtocolIdentifier(RegisteredFragments), RegisteredFragments, TEXT("InternalDestructionInfo"), bIgnoreProtocolValidation);
 	}
 }
 
diff --git a/Engine/Source/Runtime/Experimental/Iris/Core/Private/Iris/ReplicationSystem/ReplicationOperations.cpp b/Engine/Source/Runtime/Experimental/Iris/Core/Private/Iris/ReplicationSystem/ReplicationOperations.cpp
--- a/Engine/Source/Runtime/Experimental/Iris/Core/Private/Iris/ReplicationSystem/ReplicationOperations.cpp
+++ b/Engine/Source/Runtime/Experimental/Iris/Core/Private/Iris/ReplicationSystem/ReplicationOperations.cpp
@@ -794,22 +794,24 @@ void FReplicationInstanceOperations::OutputInternalStateToString(FNetSerializati
 	}
 }
 
-void FReplicationInstanceOperations::OutputInternalDefaultStateToString(FNetSerializationContext& NetSerializationContext, FStringBuilderBase& StringBuilder, const FReplicationInstanceProtocol* InstanceProtocol, const FReplicationProtocol* Protocol)
+void FReplicationInstanceOperations::OutputInternalDefaultStateToString(FNetSerializationContext& NetSerializationContext, FStringBuilderBase& StringBuilder, const FReplicationFragments& Fragments)
 {
 	// Iterate over fragments
-	const FReplicationStateDescriptor** ReplicationStateDescriptors = Protocol->ReplicationStateDescriptors;
-	FReplicationFragment* const* ReplicationFragments = InstanceProtocol->Fragments;
-
-	for (uint32 StateIt = 0, StateEndIt = Protocol->ReplicationStateCount; StateIt != StateEndIt; ++StateIt)
+	for (int32 Index=0; Index < Fragments.Num(); ++Index)
 	{
-		const FReplicationStateDescriptor* CurrentDescriptor = ReplicationStateDescriptors[StateIt];
-		const FReplicationFragment* CurrentFragment = ReplicationFragments[StateIt];
+		const FReplicationFragmentInfo& FragmentInfo = Fragments[Index];
+
+		const FReplicationStateDescriptor* CurrentDescriptor = FragmentInfo.Descriptor;
+		const FReplicationFragment* CurrentFragment = FragmentInfo.Fragment;
 
 		FReplicationStateApplyContext ReplicationStateToStringContext;
 		ReplicationStateToStringContext.NetSerializationContext = &NetSerializationContext;
 		ReplicationStateToStringContext.Descriptor = CurrentDescriptor;
 		ReplicationStateToStringContext.bIsInit = NetSerializationContext.IsInitState();
 
+		StringBuilder.Appendf(TEXT("[%d/%d] Fragment: %s DescriptorId: 0x%" UINT64_x_FMT " DefaultStateHash: 0x%" UINT64_x_FMT "\n"), 
+			Index+1, Fragments.Num(), ToCStr(CurrentDescriptor->DebugName), CurrentDescriptor->DescriptorIdentifier.Value, CurrentDescriptor->DescriptorIdentifier.DefaultStateHash);
+
 		// Dequantize state data
 		if (!EnumHasAnyFlags(CurrentFragment->GetTraits(), EReplicationFragmentTraits::HasPersistentTargetStateBuffer))
 		{
diff --git a/Engine/Source/Runtime/Experimental/Iris/Core/Private/Iris/ReplicationSystem/ReplicationProtocolManager.cpp b/Engine/Source/Runtime/Experimental/Iris/Core/Private/Iris/ReplicationSystem/ReplicationProtocolManager.cpp
--- a/Engine/Source/Runtime/Experimental/Iris/Core/Private/Iris/ReplicationSystem/ReplicationProtocolManager.cpp
+++ b/Engine/Source/Runtime/Experimental/Iris/Core/Private/Iris/ReplicationSystem/ReplicationProtocolManager.cpp
@@ -132,8 +132,7 @@ void FReplicationProtocolManager::DestroyInstanceProtocol(FReplicationInstancePr
 	FMemory::Free(InstanceProtocol);
 }
 
-FReplicationProtocolIdentifier 
-FReplicationProtocolManager::CalculateProtocolIdentifier(const FReplicationFragments& Fragments) const
+FReplicationProtocolIdentifier FReplicationProtocolManager::CalculateProtocolIdentifier(const FReplicationFragments& Fragments)
 {
 	TArray<uint64, TInlineAllocator<32>> IdBuffer;
 	for (const FReplicationFragmentInfo& Info : Fragments)
@@ -149,37 +148,42 @@ FReplicationProtocolManager::CalculateProtocolIdentifier(const FReplicationFragm
 	return ProtocolIdentifier;
 }
 
-bool FReplicationProtocolManager::ValidateReplicationProtocol(const FReplicationProtocol* Protocol, const FReplicationFragments& Fragments) const
+bool FReplicationProtocolManager::ValidateReplicationProtocol(const FReplicationProtocol* Protocol, const FReplicationFragments& Fragments, bool bLogFragmentErrors)
 {
 	bool bResult = true;
 
 	const uint32 FragmentCount = Fragments.Num();
-	if (!ensureAlwaysMsgf(Protocol->ReplicationStateCount == FragmentCount, TEXT("Protocol %s:ReplicationStateCount %u != FragmentCount:%u"), ToCStr(Protocol->DebugName), Protocol->ReplicationStateCount, FragmentCount))
+	if (ensureMsgf(Protocol->ReplicationStateCount == FragmentCount, TEXT("Protocol %s:ReplicationStateCount %u != FragmentCount:%u"), ToCStr(Protocol->DebugName), Protocol->ReplicationStateCount, FragmentCount))
 	{
-		return false;
-	}
-
-	// Validate individual fragments
-	const FReplicationFragmentInfo* FragmentsData = Fragments.GetData();
-	for (uint32 It = 0; It < FragmentCount; ++It)
-	{
-		const FReplicationFragmentInfo& Info = FragmentsData[It];
-		if (!ensure(Info.Descriptor == Protocol->ReplicationStateDescriptors[It]))
+		// Validate individual fragments
+		const FReplicationFragmentInfo* FragmentsData = Fragments.GetData();
+		for (uint32 It = 0; It < FragmentCount; ++It)
 		{
-			check(Info.Descriptor->DescriptorIdentifier == Protocol->ReplicationStateDescriptors[It]->DescriptorIdentifier);
+			const FReplicationFragmentInfo& Info = FragmentsData[It];
+			const bool bIsSameDescriptor = Info.Descriptor == Protocol->ReplicationStateDescriptors[It];
+	
+			if (!bIsSameDescriptor)
+			{
+				UE_LOG_PROTOCOLMANAGER_WARNING(
+					TEXT("FReplicationProtocolManager::ValidateReplicationProtocol for %s Descriptor Pointer mismatch %p != %p for index %u/%u named %s != %s identifier 0x%" UINT64_x_FMT " != 0x%" UINT64_x_FMT),
+					ToCStr(Protocol->DebugName), Info.Descriptor, Protocol->ReplicationStateDescriptors[It], It, FragmentCount, ToCStr(Info.Descriptor->DebugName), ToCStr(Protocol->ReplicationStateDescriptors[It]->DebugName), 
+					Info.Descriptor->DescriptorIdentifier.Value, Protocol->ReplicationStateDescriptors[It]->DescriptorIdentifier.Value);
 
-			if (Info.Descriptor->DescriptorIdentifier != Protocol->ReplicationStateDescriptors[It]->DescriptorIdentifier)
+				ensure(bIsSameDescriptor);
+				
+				bResult = false;
+			}
+			else if (Info.Descriptor->DescriptorIdentifier != Protocol->ReplicationStateDescriptors[It]->DescriptorIdentifier)
 			{
+				
 				UE_LOG_PROTOCOLMANAGER_WARNING(
-					TEXT("FReplicationProtocolManager::ValidateReplicationProtocol for %s DescriptorIdentfier mismatch for index %u/%u named %s != %s identifier 0x%" UINT64_x_FMT " != 0x%" UINT64_x_FMT), 
+					TEXT("FReplicationProtocolManager::ValidateReplicationProtocol for %s DescriptorIdentfier mismatch for index %u/%u named %s != %s identifier 0x%" UINT64_x_FMT " != 0x%" UINT64_x_FMT),
 					ToCStr(Protocol->DebugName), It, FragmentCount, ToCStr(Info.Descriptor->DebugName), ToCStr(Protocol->ReplicationStateDescriptors[It]->DebugName), Info.Descriptor->DescriptorIdentifier.Value, Protocol->ReplicationStateDescriptors[It]->DescriptorIdentifier.Value);
-			}
 
-			UE_LOG_PROTOCOLMANAGER_WARNING(
-				TEXT("FReplicationProtocolManager::ValidateReplicationProtocol for %s Descriptor Pointer mismatch %p != %p for index %u/%u named %s != %s identifier 0x%" UINT64_x_FMT " != 0x%" UINT64_x_FMT), 
-				ToCStr(Protocol->DebugName), Info.Descriptor, Protocol->ReplicationStateDescriptors[It], It, FragmentCount, ToCStr(Info.Descriptor->DebugName), ToCStr(Protocol->ReplicationStateDescriptors[It]->DebugName), Info.Descriptor->DescriptorIdentifier.Value, Protocol->ReplicationStateDescriptors[It]->DescriptorIdentifier.Value);
-
-			bResult = false;
+				ensure(Info.Descriptor->DescriptorIdentifier == Protocol->ReplicationStateDescriptors[It]->DescriptorIdentifier);
+				
+				bResult = false;
+			}
 		}
 	}
 
@@ -200,7 +204,7 @@ const FReplicationProtocol* FReplicationProtocolManager::GetReplicationProtocol(
 	return nullptr;
 }
 
-static void FragmentListToString(FStringBuilderBase& StringBuilder, const FReplicationFragments& Fragments)
+void FReplicationProtocolManager::FragmentListToString(FStringBuilderBase& StringBuilder, const FReplicationFragments& Fragments)
 {
 	StringBuilder << TEXT("Fragments:\n");
 	const FReplicationFragmentInfo* FragmentsData = Fragments.GetData();
@@ -220,16 +224,16 @@ const FReplicationProtocol* FReplicationProtocolManager::CreateReplicationProtoc
 		const FReplicationProtocolIdentifier NewProtocolId = CalculateProtocolIdentifier(Fragments);
 		if (NewProtocolId != ProtocolId)
 		{
-			UE_LOG_PROTOCOLMANAGER_WARNING(TEXT("FReplicationProtocolManager::CreateReplicationProtocol Id mismatch when creating protocol named %s with ProtocolId:0x%" UINT64_x_FMT), DebugName, ProtocolId);
-#if UE_NET_ENABLE_PROTOCOLMANAGER_LOG
-			if (bIrisLogReplicationProtocols)
-			{
-				TStringBuilder<4096> StringBuilder;
-				FragmentListToString(StringBuilder, Fragments);
-				LOG_SCOPE_VERBOSITY_OVERRIDE(LogIris, ELogVerbosity::Log);
-				UE_LOG_PROTOCOLMANAGER(Log, TEXT("%s"), StringBuilder.ToString());
-			}
-#endif
+			UE_LOG(LogIris, Warning, TEXT("FReplicationProtocolManager::CreateReplicationProtocol Id mismatch when creating protocol named %s with ProtocolId:0x%" UINT64_x_FMT), DebugName, ProtocolId);
+ #if UE_NET_ENABLE_PROTOCOLMANAGER_LOG
+ 			if (UE_LOG_ACTIVE(LogIris, Warning))
+ 			{
+ 				TStringBuilder<4096> StringBuilder;
+ 				FragmentListToString(StringBuilder, Fragments);
+
+ 				UE_LOG(LogIris, Warning, TEXT("%s"), StringBuilder.ToString());
+ 			}
+ #endif
 			return nullptr;
 		}
 	}
diff --git a/Engine/Source/Runtime/Experimental/Iris/Core/Private/Iris/ReplicationSystem/ReplicationProtocolManager.h b/Engine/Source/Runtime/Experimental/Iris/Core/Private/Iris/ReplicationSystem/ReplicationProtocolManager.h
--- a/Engine/Source/Runtime/Experimental/Iris/Core/Private/Iris/ReplicationSystem/ReplicationProtocolManager.h
+++ b/Engine/Source/Runtime/Experimental/Iris/Core/Private/Iris/ReplicationSystem/ReplicationProtocolManager.h
@@ -16,9 +16,6 @@ class FReplicationProtocolManager
 public:
 	~FReplicationProtocolManager();
 
-	/* Calculate protocol Identifier from registered fragment data */
-	IRISCORE_API FReplicationProtocolIdentifier CalculateProtocolIdentifier(const FReplicationFragments& Fragments) const;
-
 	/* Create protocol from registered fragment data with provided Id, verification is optional */
 	IRISCORE_API const FReplicationProtocol* CreateReplicationProtocol(const UObject* ArchetypeOrCDOUsedAsKey, const FReplicationProtocolIdentifier ProtocolId, const FReplicationFragments& Fragments, const TCHAR* DebugName, bool bVerifyId = false);
 
@@ -35,8 +32,6 @@ public:
 	template<typename T>
 	void ForEachProtocol(FReplicationProtocolIdentifier ProtocolId, T&& Functor) const;
 	
-	/* Validate that a existing protocol matches the FragmentList of an instance, returns true of it is a match */
-	IRISCORE_API bool ValidateReplicationProtocol(const FReplicationProtocol*, const FReplicationFragments& Fragments) const;
 
 	/* Destroy existing replication protocol */
 	IRISCORE_API void DestroyReplicationProtocol(const FReplicationProtocol* ReplicationProtocol);
@@ -50,6 +45,17 @@ public:
 
 	IRISCORE_API void InvalidateDescriptor(const FReplicationStateDescriptor* InvalidatedReplicationStateDescriptor);
 
+public:
+
+	/** Calculate protocol Identifier from registered fragment data */
+	IRISCORE_API static FReplicationProtocolIdentifier CalculateProtocolIdentifier(const FReplicationFragments& Fragments);
+
+	/** Validate that a existing protocol matches the FragmentList of an instance, returns true of it is a match */
+	IRISCORE_API static bool ValidateReplicationProtocol(const FReplicationProtocol*, const FReplicationFragments& Fragments, bool bLogFragmentErrors=true);
+
+	/** Print the names of the fragments and their hash */
+	IRISCORE_API static void FragmentListToString(FStringBuilderBase& StringBuilder, const FReplicationFragments& Fragments);
+
 private:
 	void InternalDestroyReplicationProtocol(const FReplicationProtocol* Protocol);
 	void InternalDeferDestroyReplicationProtocol(const FReplicationProtocol* Protocol);
diff --git a/Engine/Source/Runtime/Experimental/Iris/Core/Private/Iris/ReplicationSystem/ReplicationSystem.cpp b/Engine/Source/Runtime/Experimental/Iris/Core/Private/Iris/ReplicationSystem/ReplicationSystem.cpp
--- a/Engine/Source/Runtime/Experimental/Iris/Core/Private/Iris/ReplicationSystem/ReplicationSystem.cpp
+++ b/Engine/Source/Runtime/Experimental/Iris/Core/Private/Iris/ReplicationSystem/ReplicationSystem.cpp
@@ -1609,6 +1609,17 @@ float UReplicationSystem::GetCullDistanceSqrOverride(FNetRefHandle Handle, float
 	return Impl->ReplicationSystemInternal.GetNetCullDistanceOverrides().GetCullDistanceSqr(ObjectInternalIndex, DefaultValue);
 }
 
+void UReplicationSystem::ReportProtocolMismatch(uint64 NetRefHandleId, uint32 ConnectionId)
+{
+	using namespace UE::Net::Private;
+	const FNetRefHandle NetRefHandle = FNetRefHandleManager::MakeNetRefHandle(NetRefHandleId, GetId()/*DON'T SUBMIT REVIEW NOTE: This ok or should it be incomplete?*/ );
+
+	Impl->ReplicationSystemInternal.GetReplicationBridge()->OnProtocolMismatchReported(NetRefHandle, ConnectionId);
+}
+
+
+#pragma region ReplicationSystemFactory
+
 namespace UE::Net
 {
 
@@ -1705,4 +1716,6 @@ TArrayView<UReplicationSystem*> FReplicationSystemFactory::GetAllReplicationSyst
 	return MakeArrayView(ReplicationSystems, MaxReplicationSystemId + 1);
 }
 
-}
+} //end namespace UE::Net
+
+#pragma endregion
\ No newline at end of file
diff --git a/Engine/Source/Runtime/Experimental/Iris/Core/Public/Iris/ReplicationSystem/ObjectReplicationBridge.h b/Engine/Source/Runtime/Experimental/Iris/Core/Public/Iris/ReplicationSystem/ObjectReplicationBridge.h
--- a/Engine/Source/Runtime/Experimental/Iris/Core/Public/Iris/ReplicationSystem/ObjectReplicationBridge.h
+++ b/Engine/Source/Runtime/Experimental/Iris/Core/Public/Iris/ReplicationSystem/ObjectReplicationBridge.h
@@ -196,6 +196,7 @@ protected:
 	IRISCORE_API virtual void PostApplyInitialState(FNetRefHandle Handle) override;
 	IRISCORE_API virtual void DetachInstanceFromRemote(FNetRefHandle Handle, EReplicationBridgeDestroyInstanceReason DestroyReason, EReplicationBridgeDestroyInstanceFlags DestroyFlags) override;
 	IRISCORE_API virtual void DetachInstance(FNetRefHandle Handle) override;
+	IRISCORE_API virtual void OnProtocolMismatchReported(FNetRefHandle RefHandle, uint32 ConnectionId) override;
 
 protected:
 	/**
@@ -234,6 +235,9 @@ protected:
 	/** Destroy or tear-off the game instance on request from remote. */
 	virtual void DestroyInstanceFromRemote(const FDestroyInstanceParams& Params) {}
 
+	/** Called when we found a divergence between the local and remote protocols when trying to instantiate a remote replicated object. */
+	virtual void OnProtocolMismatchDetected(FNetRefHandle ObjectHandle) {}
+
 protected:
 	/** Lookup the UObject associated with the provided Handle. This function will not try to resolve the reference. */
 	IRISCORE_API UObject* GetObjectFromReferenceHandle(FNetRefHandle RefHandle) const;
diff --git a/Engine/Source/Runtime/Experimental/Iris/Core/Public/Iris/ReplicationSystem/ReplicationBridge.h b/Engine/Source/Runtime/Experimental/Iris/Core/Public/Iris/ReplicationSystem/ReplicationBridge.h
--- a/Engine/Source/Runtime/Experimental/Iris/Core/Public/Iris/ReplicationSystem/ReplicationBridge.h
+++ b/Engine/Source/Runtime/Experimental/Iris/Core/Public/Iris/ReplicationSystem/ReplicationBridge.h
@@ -256,6 +256,9 @@ protected:
 	/** Called when destruction info is received to determine whether the instance may be destroyed. */
 	IRISCORE_API virtual bool IsAllowedToDestroyInstance(const UObject* Instance) const;
 
+	/** Called when a remote connection detected a protocol mismatch when trying to instantiate the NetRefHandle replicated object. */
+	IRISCORE_API virtual void OnProtocolMismatchReported(FNetRefHandle RefHandle, uint32 ConnectionId) {}
+
 private:
 
 	// Internal operations invoked by ReplicationSystem/ReplicationWriter
diff --git a/Engine/Source/Runtime/Experimental/Iris/Core/Public/Iris/ReplicationSystem/ReplicationFragment.h b/Engine/Source/Runtime/Experimental/Iris/Core/Public/Iris/ReplicationSystem/ReplicationFragment.h
--- a/Engine/Source/Runtime/Experimental/Iris/Core/Public/Iris/ReplicationSystem/ReplicationFragment.h
+++ b/Engine/Source/Runtime/Experimental/Iris/Core/Public/Iris/ReplicationSystem/ReplicationFragment.h
@@ -221,9 +221,9 @@ ENUM_CLASS_FLAGS(EFragmentRegistrationFlags);
 */
 struct FReplicationFragmentInfo
 {
-	const FReplicationStateDescriptor* Descriptor;
-	void* SrcReplicationStateBuffer;
-	FReplicationFragment* Fragment;
+	const FReplicationStateDescriptor* Descriptor = nullptr;
+	void* SrcReplicationStateBuffer = nullptr;
+	FReplicationFragment* Fragment = nullptr;
 };
 typedef TArray<FReplicationFragmentInfo, TInlineAllocator<32>> FReplicationFragments;
 
diff --git a/Engine/Source/Runtime/Experimental/Iris/Core/Public/Iris/ReplicationSystem/ReplicationOperations.h b/Engine/Source/Runtime/Experimental/Iris/Core/Public/Iris/ReplicationSystem/ReplicationOperations.h
--- a/Engine/Source/Runtime/Experimental/Iris/Core/Public/Iris/ReplicationSystem/ReplicationOperations.h
+++ b/Engine/Source/Runtime/Experimental/Iris/Core/Public/Iris/ReplicationSystem/ReplicationOperations.h
@@ -113,7 +113,7 @@ struct FReplicationInstanceOperations
 	static IRISCORE_API void OutputInternalStateToString(FNetSerializationContext& Context, FStringBuilderBase& StringBuilder, const uint32* ChangeMaskData, const uint8* SrcInternalObjectStateBuffer, const FReplicationInstanceProtocol* InstanceProtocol, const FReplicationProtocol* Protocol);
 
 	/** Dequantize the default state for a replicated object with a given protocol and output the state to string. */
-	static IRISCORE_API void OutputInternalDefaultStateToString(FNetSerializationContext& Context, FStringBuilderBase& StringBuilder, const FReplicationInstanceProtocol* InstanceProtocol, const FReplicationProtocol* Protocol);
+	static IRISCORE_API void OutputInternalDefaultStateToString(FNetSerializationContext& NetSerializationContext, FStringBuilderBase& StringBuilder, const FReplicationFragments& Fragments);
 };
 
 struct FReplicationProtocolOperations
diff --git a/Engine/Source/Runtime/Experimental/Iris/Core/Public/Iris/ReplicationSystem/ReplicationSystem.h b/Engine/Source/Runtime/Experimental/Iris/Core/Public/Iris/ReplicationSystem/ReplicationSystem.h
--- a/Engine/Source/Runtime/Experimental/Iris/Core/Public/Iris/ReplicationSystem/ReplicationSystem.h
+++ b/Engine/Source/Runtime/Experimental/Iris/Core/Public/Iris/ReplicationSystem/ReplicationSystem.h
@@ -584,6 +584,9 @@ public:
 	/** Returns elapsed time in seconds since ReplicatonSystem was created */
 	double GetElapsedTime() const { return ElapsedTime; }
 
+	/** Called when a connection finds a protocol divergence when instantiating a replicated object. */
+	IRISCORE_API void ReportProtocolMismatch(uint64 NetRefHandleId, uint32 ConnectionId);
+
 public:
 	// For internal use and not exported.
 	UE::Net::Private::FReplicationSystemInternal* GetReplicationSystemInternal();
