Fix usage of auto not allowed by coding guidelines

#rnx

diff --git a/Engine/Source/Runtime/Experimental/Iris/Core/Private/Iris/ReplicationState/ReplicationStateDescriptorBuilder.cpp b/Engine/Source/Runtime/Experimental/Iris/Core/Private/Iris/ReplicationState/ReplicationStateDescriptorBuilder.cpp
--- a/Engine/Source/Runtime/Experimental/Iris/Core/Private/Iris/ReplicationState/ReplicationStateDescriptorBuilder.cpp
+++ b/Engine/Source/Runtime/Experimental/Iris/Core/Private/Iris/ReplicationState/ReplicationStateDescriptorBuilder.cpp
@@ -1812,7 +1812,6 @@ FPropertyReplicationStateDescriptorBuilder::Build(const FString& StateName, FRep
 				// Store the default state checksum for later use when verifying protocol
 				Descriptor->DescriptorIdentifier.DefaultStateHash = DefaultStateHash;
 			}
-		
 		}
 
 		return TRefCountPtr<const FReplicationStateDescriptor>(Descriptor);
@@ -2482,7 +2481,7 @@ TRefCountPtr<const FReplicationStateDescriptor> FReplicationStateDescriptorBuild
 	BuildParameters.bIsInitState = false;
 	BuildParameters.bAllMembersAreReplicated = bAllMembersAreReplicated;
 
-	auto Descriptor = Builder.Build(InStruct->GetName(), Parameters.DescriptorRegistry, BuildParameters);
+	TRefCountPtr<const FReplicationStateDescriptor> Descriptor = Builder.Build(InStruct->GetName(), Parameters.DescriptorRegistry, BuildParameters);
 
 	if (Parameters.DescriptorRegistry && Descriptor.IsValid())
 	{
@@ -2582,7 +2581,7 @@ SIZE_T FReplicationStateDescriptorBuilder::CreateDescriptorsForClass(FResult& Cr
 	// Check registry first to see if we already have created descriptors for this class
 	if (Parameters.DescriptorRegistry)
 	{
-		if (auto Result = Parameters.DescriptorRegistry->Find(ObjectClassOrArchetypeUsedAsKey, InObjectClass))
+		if (const FReplicationStateDescriptorRegistry::FDescriptors* Result = Parameters.DescriptorRegistry->Find(ObjectClassOrArchetypeUsedAsKey, InObjectClass))
 		{
 			CreatedDescriptors.Insert(*Result, 0);
 			return CreatedDescriptors.Num();
diff --git a/Engine/Source/Runtime/Experimental/Iris/Core/Private/Iris/ReplicationSystem/ReplicationFragmentUtil.cpp b/Engine/Source/Runtime/Experimental/Iris/Core/Private/Iris/ReplicationSystem/ReplicationFragmentUtil.cpp
--- a/Engine/Source/Runtime/Experimental/Iris/Core/Private/Iris/ReplicationSystem/ReplicationFragmentUtil.cpp
+++ b/Engine/Source/Runtime/Experimental/Iris/Core/Private/Iris/ReplicationSystem/ReplicationFragmentUtil.cpp
@@ -45,7 +45,7 @@ uint32 FReplicationFragmentUtil::CreateAndRegisterFragmentsForObject(UObject* Ob
 
 	uint32 NumCreatedReplicationFragments = 0U;
 	// create fragments and let the instance protocol own them.
-	for (auto& Desc : Result)
+	for (TRefCountPtr<const FReplicationStateDescriptor>& Desc : Result)
 	{
 		if (bRegisterFunctionsOnly && (Desc->FunctionCount == 0))
 		{
