Chaos - enable sphere approximation for collision between two one-way interaction objects by default

diff --git a/Engine/Source/Runtime/Experimental/Chaos/Private/Chaos/PBDRigidsEvolutionGBF.cpp b/Engine/Source/Runtime/Experimental/Chaos/Private/Chaos/PBDRigidsEvolutionGBF.cpp
--- a/Engine/Source/Runtime/Experimental/Chaos/Private/Chaos/PBDRigidsEvolutionGBF.cpp
+++ b/Engine/Source/Runtime/Experimental/Chaos/Private/Chaos/PBDRigidsEvolutionGBF.cpp
@@ -125,7 +125,7 @@ namespace Chaos
 
 		// Determines what happens when two one-way particles collide
 		// See EOneWayInteractionPairCollisionMode
-		int32 ChaosOneWayInteractionPairCollisionMode = (int32)EOneWayInteractionPairCollisionMode::IgnoreCollision;
+		int32 ChaosOneWayInteractionPairCollisionMode = (int32)EOneWayInteractionPairCollisionMode::SphereCollision;
 		FAutoConsoleVariableRef CVarChaosIgnoreOneWayPairCollisions(TEXT("p.Chaos.Solver.OneWayPairCollisionMode"), ChaosOneWayInteractionPairCollisionMode, TEXT("How to treat collisions between two one-way interaction particles. See EOneWayInteractionPairCollisionMode (0: Ignore collisions; 1: Collide as normal; 2: Collide as spheres)"));
 
 
