Incremental leaf object indices for simplified convexes to avoid collision key conflicts

diff --git a/Engine/Source/Runtime/Experimental/Chaos/Private/Chaos/ConvexOptimizer.cpp b/Engine/Source/Runtime/Experimental/Chaos/Private/Chaos/ConvexOptimizer.cpp
--- a/Engine/Source/Runtime/Experimental/Chaos/Private/Chaos/ConvexOptimizer.cpp
+++ b/Engine/Source/Runtime/Experimental/Chaos/Private/Chaos/ConvexOptimizer.cpp
@@ -55,11 +55,11 @@ void FConvexOptimizer::VisitCollisionObjects(const FImplicitHierarchyVisitor& Vi
 		const TArray<Private::FImplicitBVHObject>& ImplicitObjects = BVH.IsValid() ? BVH->GetObjects() : CollisionObjects->ImplicitObjects;
 		
 		int32 ObjectIndex = 1;
-		int32 LeafObjectIndex = 0;
 		for(const Private::FImplicitBVHObject& CollisionObject : ImplicitObjects)
 		{
+			int32 LocalLeafObjectIndex = CollisionObject.GetObjectIndex();
 			CollisionObject.GetGeometry()->VisitLeafObjectsImpl(CollisionObject.GetTransform(),
-				CollisionObject.GetRootObjectIndex(), ObjectIndex, LeafObjectIndex, VisitorFunc);
+				CollisionObject.GetRootObjectIndex(), ObjectIndex, LocalLeafObjectIndex, VisitorFunc);
 		}
 	}	
 }
@@ -69,23 +69,24 @@ void FConvexOptimizer::VisitOverlappingObjects(const FAABB3& LocalBounds, const
 	if(CVars::bChaosConvexSimplifyUnion)
 	{
 		int32 ObjectIndex = 1;
-		int32 LeafObjectIndex = 0;
 		if(BVH.IsValid())
 		{
 			BVH->VisitAllIntersections(LocalBounds,
-			[this, &ObjectIndex, &VisitorFunc, &LocalBounds, &LeafObjectIndex](const FImplicitObject* Implicit, const FRigidTransform3f& RelativeTransformf,
-				const FAABB3f& RelativeBoundsf, const int32 RootObjectIndex, const int32)
+			[this, &ObjectIndex, &VisitorFunc, &LocalBounds](const FImplicitObject* Implicit, const FRigidTransform3f& RelativeTransformf,
+				const FAABB3f& RelativeBoundsf, const int32 RootObjectIndex, const int32 LeafObjectIndex)
 			{
+				int32 LocalLeafObjectIndex = LeafObjectIndex;
 				Implicit->VisitOverlappingLeafObjectsImpl(LocalBounds, FRigidTransform3(RelativeTransformf),
-					RootObjectIndex, ObjectIndex, LeafObjectIndex, VisitorFunc);
+					RootObjectIndex, ObjectIndex, LocalLeafObjectIndex, VisitorFunc);
 			});
 		}
 		else
 		{
 			for(const Private::FImplicitBVHObject& CollisionObject : CollisionObjects->ImplicitObjects)
 			{
+				int32 LocalLeafObjectIndex = CollisionObject.GetObjectIndex();
 				CollisionObject.GetGeometry()->VisitOverlappingLeafObjectsImpl(LocalBounds, CollisionObject.GetTransform(),
-					CollisionObject.GetRootObjectIndex(), ObjectIndex, LeafObjectIndex, VisitorFunc);
+					CollisionObject.GetRootObjectIndex(), ObjectIndex, LocalLeafObjectIndex, VisitorFunc);
 			}
 		}
 	}
@@ -198,7 +199,7 @@ FORCEINLINE void ResizeCachedTriboxes(FConvexOptimizer::FTriboxNodes& RootTribox
 
 DECLARE_CYCLE_STAT(TEXT("Collisions::BuildConvexTriboxes"), STAT_BuildConvexTriboxes, STATGROUP_ChaosCollision);
 FORCEINLINE void BuildConvexTriboxes(const Chaos::FImplicitObjectUnionPtr& UnionGeometry, const FShapesArray& UnionShapes,
-	TUniquePtr<Private::FCollisionObjects>& CollisionObjects, FConvexOptimizer::FTriboxNodes& RootTriboxes)
+	TUniquePtr<Private::FCollisionObjects>& CollisionObjects, FConvexOptimizer::FTriboxNodes& RootTriboxes, int32& NextConvexId)
 {
 	SCOPE_CYCLE_COUNTER(STAT_BuildConvexTriboxes);
 
@@ -219,7 +220,7 @@ FORCEINLINE void BuildConvexTriboxes(const Chaos::FImplicitObjectUnionPtr& Union
 				FTribox LocalTribox;
 				int32 LeafObjectIndex = 0, ObjectIndex = 0;
 				RootObject->VisitLeafObjectsImpl(FRigidTransform3::Identity,RootObjectIndex, ObjectIndex, LeafObjectIndex,
-					[&CollisionObjects, &LocalTribox, &bHasRootTribox, &ShapeIndex](
+					[&CollisionObjects, &LocalTribox, &bHasRootTribox, &ShapeIndex, &NextConvexId](
 					const FImplicitObject* ImplicitObject, const FRigidTransform3& RelativeTransform,
 					const int32 RootIndex, const int32, const int32)
 				{
@@ -258,7 +259,7 @@ FORCEINLINE void BuildConvexTriboxes(const Chaos::FImplicitObjectUnionPtr& Union
 					{
 						// Add non convex implicits to the list of collision objects
 						Private::FImplicitBVH::CollectLeafObject(ImplicitObject,
-							RelativeTransform, ShapeIndex, CollisionObjects->ImplicitObjects);	  
+							RelativeTransform, ShapeIndex, CollisionObjects->ImplicitObjects, NextConvexId++);
 					}
 				});
 				if(!bHasRootTribox && LocalTribox.HasDatas())
@@ -279,22 +280,27 @@ FORCEINLINE void BuildConvexTriboxes(const Chaos::FImplicitObjectUnionPtr& Union
 
 void FConvexOptimizer::BuildSingleConvex(const Chaos::FImplicitObjectUnionPtr& UnionGeometry, const FShapesArray& UnionShapes)
 {
-	BuildConvexTriboxes(UnionGeometry, UnionShapes, CollisionObjects, RootTriboxes);
+	BuildConvexTriboxes(UnionGeometry, UnionShapes, CollisionObjects, RootTriboxes, NextConvexId);
 	
-	FTribox MainTribox;
+	MainTribox.NodeTribox = FTribox(); 
 	for(auto& RootTribox : RootTriboxes)
 	{
 		if(RootTribox.Value.NodeTribox.HasDatas())
 		{
-			MainTribox += RootTribox.Value.NodeTribox;
+			MainTribox.NodeTribox += RootTribox.Value.NodeTribox;
 		}
 	}
-
-	if(MainTribox.HasDatas())
+	if(MainTribox.NodeTribox.HasDatas())
 	{
+		if(!MainTribox.TriboxConvex || (MainTribox.TriboxConvex && (MainTribox.TriboxConvex->AsA<FConvex>()->GetVolume() != MainTribox.NodeTribox.ComputeVolume())))
+		{
+			MainTribox.TriboxConvex = MainTribox.NodeTribox.MakeConvex();
+			MainTribox.ConvexId = NextConvexId++;
+		}
+
 		// Build the tribox and add it to the list of optimized convexes
-		SimplifiedConvexes.Add(MainTribox.MakeConvex());
-		Private::FImplicitBVH::CollectLeafObject(SimplifiedConvexes.Last(), FRigidTransform3::Identity, 0, CollisionObjects->ImplicitObjects);	
+		SimplifiedConvexes.Add(MainTribox.TriboxConvex);
+		Private::FImplicitBVH::CollectLeafObject(SimplifiedConvexes.Last(), FRigidTransform3::Identity, 0, CollisionObjects->ImplicitObjects, MainTribox.ConvexId);
 	}
 }
 
@@ -315,7 +321,7 @@ void FConvexOptimizer::BuildUnionConnectivity(const Chaos::FImplicitObjectUnionP
 		{
 			if (auto* RootTribox = RootTriboxes.Find(UnionGeometry->GetObjects()[RootIndex]))
 			{
-				Private::FImplicitBVH::CollectLeafObject(UnionGeometry->GetObjects()[RootIndex], FRigidTransform3::Identity, RootIndex, ImplicitObjects);
+				Private::FImplicitBVH::CollectLeafObject(UnionGeometry->GetObjects()[RootIndex], FRigidTransform3::Identity, RootIndex, ImplicitObjects, ImplicitObjects.Num());
 			}
 		}
 		LocalBVH = FImplicitBVH::TryMakeFromLeaves(MoveTemp(ImplicitObjects), CVars::ChaosUnionBVHMinShapes, CVars::ChaosUnionBVHMaxDepth);
@@ -437,7 +443,7 @@ void FConvexOptimizer::MergeConnectedShapes(const Chaos::FImplicitObjectUnionPtr
 	
 void FConvexOptimizer::BuildMultipleConvex(const Chaos::FImplicitObjectUnionPtr& UnionGeometry, const FShapesArray& UnionShapes, const bool bEnableMerging)
 {
-	BuildConvexTriboxes(UnionGeometry, UnionShapes, CollisionObjects, RootTriboxes);
+	BuildConvexTriboxes(UnionGeometry, UnionShapes, CollisionObjects, RootTriboxes, NextConvexId);
 
 	if(!bEnableMerging)
 	{
@@ -446,12 +452,13 @@ void FConvexOptimizer::BuildMultipleConvex(const Chaos::FImplicitObjectUnionPtr&
 			if(RootTribox.Value.NodeTribox.HasDatas())
 			{
 				// Build the Tribox and add it to the list of optimized convexes
-				if(!RootTribox.Value.TriboxConvex)
+				if(!RootTribox.Value.TriboxConvex || (RootTribox.Value.TriboxConvex && (RootTribox.Value.TriboxConvex->AsA<FConvex>()->GetVolume() != RootTribox.Value.NodeTribox.ComputeVolume())))
 				{
 					RootTribox.Value.TriboxConvex = RootTribox.Value.NodeTribox.MakeConvex();
+					RootTribox.Value.ConvexId = NextConvexId++;
 				}
 				SimplifiedConvexes.Add(RootTribox.Value.TriboxConvex);
-				Private::FImplicitBVH::CollectLeafObject(SimplifiedConvexes.Last(), FRigidTransform3::Identity, RootTribox.Value.ShapeIndex, CollisionObjects->ImplicitObjects);
+				Private::FImplicitBVH::CollectLeafObject(SimplifiedConvexes.Last(), FRigidTransform3::Identity, RootTribox.Value.ShapeIndex, CollisionObjects->ImplicitObjects, RootTribox.Value.ConvexId);
 			}
 		}
 	}
@@ -475,9 +482,10 @@ void FConvexOptimizer::BuildMultipleConvex(const Chaos::FImplicitObjectUnionPtr&
             	if(!RootTribox->TriboxConvex || (RootTribox->TriboxConvex && (RootTribox->TriboxConvex->AsA<FConvex>()->GetVolume() != MergedNodes[NodeIndex].NodeTribox.ComputeVolume())))
             	{
             		RootTribox->TriboxConvex = MergedNodes[NodeIndex].NodeTribox.MakeConvex();
+					RootTribox->ConvexId = NextConvexId++;
 				}
             	SimplifiedConvexes.Add(RootTribox->TriboxConvex);
-            	Private::FImplicitBVH::CollectLeafObject(SimplifiedConvexes.Last(), FRigidTransform3::Identity, MergedNodes[NodeIndex].ShapeIndex, CollisionObjects->ImplicitObjects);
+            	Private::FImplicitBVH::CollectLeafObject(SimplifiedConvexes.Last(), FRigidTransform3::Identity, MergedNodes[NodeIndex].ShapeIndex, CollisionObjects->ImplicitObjects, RootTribox->ConvexId);
             }
         }
 	}
diff --git a/Engine/Source/Runtime/Experimental/Chaos/Private/Chaos/ImplicitObjectBVH.cpp b/Engine/Source/Runtime/Experimental/Chaos/Private/Chaos/ImplicitObjectBVH.cpp
--- a/Engine/Source/Runtime/Experimental/Chaos/Private/Chaos/ImplicitObjectBVH.cpp
+++ b/Engine/Source/Runtime/Experimental/Chaos/Private/Chaos/ImplicitObjectBVH.cpp
@@ -56,20 +56,19 @@ namespace Chaos
 		}
 		
 		void FImplicitBVH::CollectLeafObject(const FImplicitObject* Object, 
-			const FRigidTransform3& ParentTransform, const int32 RootObjectIndex, TArray<FImplicitBVHObject>& LeafObjects)
+			const FRigidTransform3& ParentTransform, const int32 RootObjectIndex, TArray<FImplicitBVHObject>& LeafObjects, const int32 LeafObjectIndex)
 		{
 			// @todo(chaos): clean this up (SetFromRawLowLevel). We know all the objects we visit are children of a UniquePtr because we own it
 			TSerializablePtr<FImplicitObject> SerializableObject;
 			SerializableObject.SetFromRawLowLevel(Object);
 
-			const int32 ObjectIndex = LeafObjects.Num();
 			LeafObjects.Emplace(
 				SerializableObject,
 				ParentTransform.GetTranslation(),
 				ParentTransform.GetRotation(),
 				Object->CalculateTransformedBounds(ParentTransform),
 				RootObjectIndex,
-				ObjectIndex);
+				LeafObjectIndex);
 		}
 
 		FImplicitBVH::FObjects FImplicitBVH::CollectLeafObjects(const TArrayView<const Chaos::FImplicitObjectPtr>& InRootObjects)
@@ -85,7 +84,7 @@ namespace Chaos
 				[RootObjectIndex, &Objects](const FImplicitObject* Object, const FRigidTransform3& ParentTransform,
 					const int32 UnusedRootObjectIndex, const int32 UnusedObjectIndex, const int32 UnusedLeafObjectIndex)
 				{
-					CollectLeafObject(Object, ParentTransform, RootObjectIndex, Objects);
+					CollectLeafObject(Object, ParentTransform, RootObjectIndex, Objects, Objects.Num());
 				});
 			}
 
diff --git a/Engine/Source/Runtime/Experimental/Chaos/Public/Chaos/ConvexOptimizer.h b/Engine/Source/Runtime/Experimental/Chaos/Public/Chaos/ConvexOptimizer.h
--- a/Engine/Source/Runtime/Experimental/Chaos/Public/Chaos/ConvexOptimizer.h
+++ b/Engine/Source/Runtime/Experimental/Chaos/Public/Chaos/ConvexOptimizer.h
@@ -46,6 +46,9 @@ namespace Private
 
 			// Valid flag to check if the edge has already been processed during merging
 			FValidEdges bValidEdges = {};
+
+			// Convex leaf offset that was used when the convex has been built
+			int32 ConvexId = 0;
 		};
 
 		using FTriboxNodes = TMap<FImplicitObject*, FTriboxNode>;
@@ -105,6 +108,12 @@ namespace Private
 
 		// BVH used to accelerate the collisions queries
 		TUniquePtr<Private::FImplicitBVH> BVH;
+
+		// Main tribox built from all the convexes
+		FTriboxNode MainTribox;
+
+		// Leaf offset to have a unique id for the midphase
+		int32 NextConvexId = 0;
 	};
 
 	// Visit all the collision objects if they exist / otherwise forward it to the RootHierarchy
diff --git a/Engine/Source/Runtime/Experimental/Chaos/Public/Chaos/ImplicitObjectBVH.h b/Engine/Source/Runtime/Experimental/Chaos/Public/Chaos/ImplicitObjectBVH.h
--- a/Engine/Source/Runtime/Experimental/Chaos/Public/Chaos/ImplicitObjectBVH.h
+++ b/Engine/Source/Runtime/Experimental/Chaos/Public/Chaos/ImplicitObjectBVH.h
@@ -45,6 +45,7 @@ namespace Chaos
 			// multiple times in the hierarchy in a union of transformed objects, each will have a different ObjectIndex.
 			// This is the ObjectIndex assigned when visiting the hierarchy via FImplicitObject::VisitHierachy and other 
 			// visit methods and can be used to index arrays initialized via those visitors.
+			// @todo(chaos): rename to GetObjectId()
 			int32 GetObjectIndex() const { return ObjectIndex; }
 
 			// The index of our most distant ancestor. I.e., the index in the root Union. This is used to map
@@ -113,7 +114,7 @@ namespace Chaos
 			static int32 CountLeafObjects(const TArrayView<const Chaos::FImplicitObjectPtr>& InRootObjects);
 			static FObjects CollectLeafObjects(const TArrayView<const Chaos::FImplicitObjectPtr>& InRootObjects);
 			static void CollectLeafObject(const FImplicitObject* Object, const FRigidTransform3& ParentTransform, const int32 RootObjectIndex,
-				TArray<FImplicitBVHObject>& LeafObjects);
+				TArray<FImplicitBVHObject>& LeafObjects, const int32 LeafObjectIndex);
 
 			// Create a BVH around a set of ImplicitObjects. Usually these are the immediate child elements of an FImplcitObjectUnion
 			// TryMake will then recurse into the geometry hierachy and add all descendents to the BVH. Will return null if the 
