Server crash on GC update RepData accessing uninitialized world.

diff --git a/Engine/Source/Runtime/Experimental/GeometryCollectionEngine/Private/GeometryCollection/GeometryCollectionComponent.cpp b/Engine/Source/Runtime/Experimental/GeometryCollectionEngine/Private/GeometryCollection/GeometryCollectionComponent.cpp
--- a/Engine/Source/Runtime/Experimental/GeometryCollectionEngine/Private/GeometryCollection/GeometryCollectionComponent.cpp
+++ b/Engine/Source/Runtime/Experimental/GeometryCollectionEngine/Private/GeometryCollection/GeometryCollectionComponent.cpp
@@ -3850,45 +3850,48 @@ void UGeometryCollectionComponent::RegisterAndInitializePhysicsProxy()
 
 	PhysicsProxy->SetUseStaticMeshCollisionForTraces_External(bUseStaticMeshCollisionForTraces);
 
-	FPhysScene_Chaos* Scene = GetInnerChaosScene();
-	Scene->AddObject(this, PhysicsProxy);
-
-	// If we're replicating we need some extra setup - check netmode as we don't need this for standalone runtime where we aren't going to network the component
-	// IMPORTANT this need to happen after the object is registered so this will guarantee that the particles are properly created by the time the callback below gets called
-	if (GetIsReplicated() && PhysicsProxy->GetReplicationMode() == FGeometryCollectionPhysicsProxy::EReplicationMode::Client)
+	FPhysScene_Chaos* PhysicsScene = GetInnerChaosScene();
+	if (ensure(PhysicsScene))
 	{
-		// Client side : geometry collection children of parents below the rep level need to be infinitely strong so that client cannot break it 
-		if (Chaos::FPhysicsSolver* CurrSolver = GetSolver(*this))
+		PhysicsScene->AddObject(this, PhysicsProxy);
+
+		// If we're replicating we need some extra setup - check netmode as we don't need this for standalone runtime where we aren't going to network the component
+		// IMPORTANT this need to happen after the object is registered so this will guarantee that the particles are properly created by the time the callback below gets called
+		if (GetIsReplicated() && PhysicsProxy->GetReplicationMode() == FGeometryCollectionPhysicsProxy::EReplicationMode::Client)
 		{
-			CurrSolver->EnqueueCommandImmediate([Proxy = PhysicsProxy, AbandonAfterLevel = ReplicationAbandonAfterLevel, EnableAbandonAfterLevel = bEnableAbandonAfterLevel]()
-				{
-					// As we're not in control we make it so our simulated proxy cannot break clusters
-					// We have to set the strain to a high value but be below the max for the data type
-					// so releasing on authority demand works
-					for (Chaos::FPBDRigidClusteredParticleHandle* ParticleHandle : Proxy->GetParticles())
+			// Client side : geometry collection children of parents below the rep level need to be infinitely strong so that client cannot break it 
+			if (Chaos::FPhysicsSolver* CurrSolver = GetSolver(*this))
+			{
+				CurrSolver->EnqueueCommandImmediate([Proxy = PhysicsProxy, AbandonAfterLevel = ReplicationAbandonAfterLevel, EnableAbandonAfterLevel = bEnableAbandonAfterLevel]()
 					{
-						if (ParticleHandle)
+						// As we're not in control we make it so our simulated proxy cannot break clusters
+						// We have to set the strain to a high value but be below the max for the data type
+						// so releasing on authority demand works
+						for (Chaos::FPBDRigidClusteredParticleHandle* ParticleHandle : Proxy->GetParticles())
 						{
-							const int32 Level = EnableAbandonAfterLevel ? ComputeParticleLevel(ParticleHandle) : -1;
-							if (Level <= AbandonAfterLevel)	//we only replicate up until level X, but it means we should replicate the breaking event of level X+1 (but not X+1's positions)
+							if (ParticleHandle)
 							{
-								ParticleHandle->SetUnbreakable(true);
+								const int32 Level = EnableAbandonAfterLevel ? ComputeParticleLevel(ParticleHandle) : -1;
+								if (Level <= AbandonAfterLevel)	//we only replicate up until level X, but it means we should replicate the breaking event of level X+1 (but not X+1's positions)
+								{
+									ParticleHandle->SetUnbreakable(true);
+								}
 							}
 						}
-					}
-				});
+					});
+			}
 		}
-	}
 
-	LoadCollisionProfiles();
+		LoadCollisionProfiles();
 
-	// We need to add the geometry collection into the external acceleration structure so that it's immediately available for queries instead of waiting for the sync from the physics thread (which could take awhile).
-	// Just adding the root particle should be sufficient since that'll be the only particle we'd expect any collisions with right after initialization.
-	if (Chaos::FPhysicsObjectHandle RootObject = GetPhysicsObjectByName(NAME_None))
-	{
-		TArrayView<Chaos::FPhysicsObjectHandle> Handles{ &RootObject, 1 };
-		FLockedWritePhysicsObjectExternalInterface Interface = FPhysicsObjectExternalInterface::LockWrite(Handles);
-		Interface->AddToSpatialAcceleration(Handles, Scene->GetSpacialAcceleration());
+		// We need to add the geometry collection into the external acceleration structure so that it's immediately available for queries instead of waiting for the sync from the physics thread (which could take awhile).
+		// Just adding the root particle should be sufficient since that'll be the only particle we'd expect any collisions with right after initialization.
+		if (Chaos::FPhysicsObjectHandle RootObject = GetPhysicsObjectByName(NAME_None))
+		{
+			TArrayView<Chaos::FPhysicsObjectHandle> Handles{ &RootObject, 1 };
+			FLockedWritePhysicsObjectExternalInterface Interface = FPhysicsObjectExternalInterface::LockWrite(Handles);
+			Interface->AddToSpatialAcceleration(Handles, PhysicsScene->GetSpacialAcceleration());
+		}
 	}
 
 	RegisterForEvents();
@@ -4515,8 +4518,13 @@ void UGeometryCollectionComponent::OnDestroyPhysicsState()
 
 	if(PhysicsProxy)
 	{
-		FPhysScene_Chaos* Scene = GetInnerChaosScene();
-		Scene->RemoveObject(PhysicsProxy);
+		// Clear physics sync callback
+		PhysicsProxy->SetPostPhysicsSyncCallback([]() {}); 
+
+		if (FPhysScene* PhysScene = GetInnerChaosScene())
+		{
+			PhysScene->RemoveObject(PhysicsProxy);
+		}
 		InitializationState = ESimulationInitializationState::Unintialized;
 
 		// clear the clusters to rep as the information hold by it is now invalid
@@ -5561,13 +5569,18 @@ FPhysScene_Chaos* UGeometryCollectionComponent::GetInnerChaosScene() const
 	}
 	else
 	{
-		if (ensure(GetOwner()) && ensure(GetOwner()->GetWorld()))
+		if (GetOwner() && GetOwner()->GetWorld())
 		{
 			return GetOwner()->GetWorld()->GetPhysicsScene();
 		}
 
-		check(GWorld);
-		return GWorld->GetPhysicsScene();
+		if (GWorld)
+		{
+			return GWorld->GetPhysicsScene();
+		}
+
+		UE_LOG(LogPhysics, Error, TEXT("Failed to find valid UWorld in UGeometryCollectionComponent::GetInnerChaosScene"));
+		return nullptr;
 	}
 }
 
