Iris
- Moved NetStats.h/NetStats.cpp to private

#rnx

diff --git a/Engine/Source/Runtime/Experimental/Iris/Core/Private/Iris/Stats/NetStats.cpp b/Engine/Source/Runtime/Experimental/Iris/Core/Private/Iris/Stats/NetStats.cpp
new file mode 100644
--- /dev/null
+++ b/Engine/Source/Runtime/Experimental/Iris/Core/Private/Iris/Stats/NetStats.cpp
@@ -0,0 +1,80 @@
+// Copyright Epic Games, Inc. All Rights Reserved.
+
+#include "Iris/Stats/NetStats.h"
+#include "ProfilingDebugging/CsvProfiler.h"
+#include "Misc/ScopeLock.h"
+
+namespace UE::Net
+{
+
+// Enable Iris category by default on servers
+CSV_DEFINE_CATEGORY(Iris, WITH_SERVER_CODE);
+
+void FNetSendStats::Accumulate(const FNetSendStats& Other)
+{
+	FScopeLock Lock(&CS);
+
+	Stats.ScheduledForReplicationRootObjectCount += Other.Stats.ScheduledForReplicationRootObjectCount;
+	Stats.ReplicatedRootObjectCount += Other.Stats.ReplicatedRootObjectCount;
+	Stats.ReplicatedObjectCount += Other.Stats.ReplicatedObjectCount;
+	Stats.ReplicatedDestructionInfoCount += Other.Stats.ReplicatedDestructionInfoCount;
+	Stats.DeltaCompressedObjectCount += Other.Stats.DeltaCompressedObjectCount;
+	Stats.ReplicationWasteObjectCount += Other.Stats.ReplicationWasteObjectCount;
+	Stats.ReplicatedObjectStatesMaskedOut += Other.Stats.ReplicatedObjectStatesMaskedOut;
+	Stats.ActiveHugeObjectCount += Other.Stats.ActiveHugeObjectCount;
+	Stats.HugeObjectsWaitingForAckCount += Other.Stats.HugeObjectsWaitingForAckCount;
+	Stats.HugeObjectsStallingCount += Other.Stats.HugeObjectsStallingCount;
+
+	Stats.ReplicationWasteTimeInSeconds += Other.Stats.ReplicationWasteTimeInSeconds;
+	Stats.HugeObjectWaitingForAckTimeInSeconds += Other.Stats.HugeObjectWaitingForAckTimeInSeconds;
+	Stats.HugeObjectStallingTimeInSeconds += Other.Stats.HugeObjectStallingTimeInSeconds;
+}
+
+void FNetSendStats::Reset()
+{
+	FScopeLock Lock(&CS);
+	Stats = FStats();
+}
+
+void FNetSendStats::ReportCsvStats()
+{
+#if CSV_PROFILER
+	FScopeLock Lock(&CS);
+
+	// Calculate connection averages for some stats
+	if (Stats.ReplicatingConnectionCount > 0)
+	{
+		const float ConnectionCountFloat = float(Stats.ReplicatingConnectionCount);
+		CSV_CUSTOM_STAT(Iris, AvgScheduledForReplicationRootObjectCount, Stats.ScheduledForReplicationRootObjectCount/ConnectionCountFloat, ECsvCustomStatOp::Set);
+		CSV_CUSTOM_STAT(Iris, AvgReplicatedRootObjectCount, Stats.ReplicatedRootObjectCount/ConnectionCountFloat, ECsvCustomStatOp::Set);
+		CSV_CUSTOM_STAT(Iris, AvgReplicatedObjectCount, Stats.ReplicatedObjectCount/ConnectionCountFloat, ECsvCustomStatOp::Set);
+		CSV_CUSTOM_STAT(Iris, AvgReplicatedDestructionInfoCount, Stats.ReplicatedDestructionInfoCount/ConnectionCountFloat, ECsvCustomStatOp::Set);
+		CSV_CUSTOM_STAT(Iris, AvgReplicationWasteObjectCount, Stats.ReplicationWasteObjectCount/ConnectionCountFloat, ECsvCustomStatOp::Set);
+		CSV_CUSTOM_STAT(Iris, AvgReplicatedObjectStatesMaskedOut, Stats.ReplicatedObjectStatesMaskedOut/ConnectionCountFloat, ECsvCustomStatOp::Set);
+		CSV_CUSTOM_STAT(Iris, AvgDeltaCompressedObjectCount, Stats.DeltaCompressedObjectCount/ConnectionCountFloat, ECsvCustomStatOp::Set);
+	}
+	else
+	{
+		CSV_CUSTOM_STAT(Iris, AvgScheduledForReplicationRootObjectCount, 0, ECsvCustomStatOp::Set);
+		CSV_CUSTOM_STAT(Iris, AvgReplicatedRootObjectCount, 0, ECsvCustomStatOp::Set);
+		CSV_CUSTOM_STAT(Iris, AvgReplicatedObjectCount, 0, ECsvCustomStatOp::Set);
+		CSV_CUSTOM_STAT(Iris, AvgReplicatedDestructionInfoCount, 0, ECsvCustomStatOp::Set);
+		CSV_CUSTOM_STAT(Iris, AvgReplicatedHugeObjectCount, 0, ECsvCustomStatOp::Set);
+		CSV_CUSTOM_STAT(Iris, AvgReplicationWasteObjectCount, 0, ECsvCustomStatOp::Set);
+		CSV_CUSTOM_STAT(Iris, AvgDeltaCompressedObjectCount, 0, ECsvCustomStatOp::Set);
+	}
+
+	// Huge object counts
+	CSV_CUSTOM_STAT(Iris, ActiveHugeObjectCount, Stats.ActiveHugeObjectCount, ECsvCustomStatOp::Set);
+	CSV_CUSTOM_STAT(Iris, HugeObjectsWaitingForAckCount, Stats.HugeObjectsWaitingForAckCount, ECsvCustomStatOp::Set);
+	CSV_CUSTOM_STAT(Iris, HugeObjectsStallingCount, Stats.HugeObjectsStallingCount, ECsvCustomStatOp::Set);
+
+	CSV_CUSTOM_STAT(Iris, ReplicatingConnectionCount, Stats.ReplicatingConnectionCount, ECsvCustomStatOp::Set);
+
+	CSV_CUSTOM_STAT(Iris, ReplicationWasteTimeMilliseconds, Stats.ReplicationWasteTimeInSeconds*1000.0, ECsvCustomStatOp::Set);
+	CSV_CUSTOM_STAT(Iris, HugeObjectWaitingForAckTimeInSeconds, Stats.HugeObjectWaitingForAckTimeInSeconds, ECsvCustomStatOp::Set);
+	CSV_CUSTOM_STAT(Iris, HugeObjectStallingTimeInSeconds, Stats.HugeObjectStallingTimeInSeconds, ECsvCustomStatOp::Set);
+#endif
+}
+
+}
diff --git a/Engine/Source/Runtime/Experimental/Iris/Core/Private/Iris/Stats/NetStats.h b/Engine/Source/Runtime/Experimental/Iris/Core/Private/Iris/Stats/NetStats.h
new file mode 100644
--- /dev/null
+++ b/Engine/Source/Runtime/Experimental/Iris/Core/Private/Iris/Stats/NetStats.h
@@ -0,0 +1,159 @@
+// Copyright Epic Games, Inc. All Rights Reserved.
+
+#pragma once
+
+#include "HAL/Platform.h"
+#include "HAL/CriticalSection.h"
+
+namespace UE::Net
+{
+
+/**
+ * Send stats for Iris replication reported to the CSV profiler. Mostly of interest on the server side due to the server authoritative network model.
+ * Its intended use is to do thread local tracking to an instance and then use the Accumulate function for thread safe updating of the ReplicationSystem owned instance.
+ * The ReplicationSystem owned instance is the one reporting to the CSV profiler.
+ */
+class FNetSendStats
+{
+public:
+	FNetSendStats() = default;
+	FNetSendStats(const FNetSendStats&) = delete;
+	FNetSendStats& operator=(const FNetSendStats&) = delete;
+
+	/** Set number of objects scheduled for replication. */
+	void SetNumberOfRootObjectsScheduledForReplication(uint32 Count);
+
+	/** Add number of replicated root objects. */
+	void AddNumberOfReplicatedRootObjects(uint32 Count);
+
+	/** Add number of replicated objects, including subobjects. */
+	void AddNumberOfReplicatedObjects(uint32 Count);
+
+	/** Add number of replicated destruction infos. */
+	void AddNumberOfReplicatedDestructionInfos(uint32 Count);
+
+	/** Add number of replicated objects, including subobjects, using delta compression. */
+	void AddNumberOfDeltaCompressedReplicatedObjects(uint32 Count);
+
+	/** Add number of replicated object states masked out such that no state is replicated for the object. The object may still replicate attachments. */
+	void AddNumberOfReplicatedObjectStatesMaskedOut(uint32 Count);
+
+	/** Get the number of replicated root objects. */
+	uint32 GetNumberOfReplicatedRootObjects() const;
+
+	/** Get the number of replicated objects, including subobjects. */
+	uint32 GetNumberOfReplicatedObjects() const;
+
+	/** Set the number of huge objects in sending or waiting to be acked. */
+	void SetNumberOfActiveHugeObjects(uint32 Count);
+
+	/** Add time in seconds waiting for completely sent huge object to be acked. */
+	void AddHugeObjectWaitingTime(double Seconds);
+
+	/** Add time in seconds waiting to be able to continue sending huge object. */
+	void AddHugeObjectStallTime(double Seconds);
+
+	/** Add time spent on trying to replicate an object in the packet but that ended up not fitting. */
+	void AddReplicationWasteTime(double Seconds);
+
+	/** Add stats from another instance. */
+	IRISCORE_API void Accumulate(const FNetSendStats& Stats);
+
+	/** Reset stats. */
+	IRISCORE_API void Reset();
+
+	/** Report the stats to the CSV profiler. Does nothing if CSV profiler support is compiled out. */
+	IRISCORE_API void ReportCsvStats();
+
+	/** Set number of replicating connections. */
+	void SetNumberOfReplicatingConnections(uint32 Count);
+
+private:
+	// Helper struct to facilitate reset of stats.
+	struct FStats
+	{
+		double HugeObjectWaitingForAckTimeInSeconds = 0;
+		double HugeObjectStallingTimeInSeconds = 0;
+		double ReplicationWasteTimeInSeconds = 0;
+
+		int32 ScheduledForReplicationRootObjectCount = 0;
+		int32 ReplicatedRootObjectCount = 0;
+		int32 ReplicatedObjectCount = 0;
+		int32 ReplicatedDestructionInfoCount = 0;
+		int32 DeltaCompressedObjectCount = 0;
+		int32 ReplicationWasteObjectCount = 0;
+		int32 ReplicatedObjectStatesMaskedOut = 0;
+		int32 ActiveHugeObjectCount = 0;
+		int32 HugeObjectsWaitingForAckCount = 0;
+		int32 HugeObjectsStallingCount = 0;
+		int32 ReplicatingConnectionCount = 0;
+	};
+
+	FCriticalSection CS;
+	FStats Stats;
+};
+
+inline void FNetSendStats::SetNumberOfRootObjectsScheduledForReplication(uint32 Count)
+{
+	Stats.ScheduledForReplicationRootObjectCount = Count;
+}
+
+inline void FNetSendStats::AddNumberOfReplicatedRootObjects(uint32 Count)
+{
+	Stats.ReplicatedRootObjectCount += Count;
+}
+
+inline void FNetSendStats::AddNumberOfReplicatedObjects(uint32 Count)
+{
+	Stats.ReplicatedObjectCount += Count;
+}
+
+inline void FNetSendStats::AddNumberOfDeltaCompressedReplicatedObjects(uint32 Count)
+{
+	Stats.DeltaCompressedObjectCount += Count;
+}
+
+inline void FNetSendStats::AddNumberOfReplicatedObjectStatesMaskedOut(uint32 Count)
+{
+	Stats.ReplicatedObjectStatesMaskedOut += Count;
+}
+
+inline void FNetSendStats::AddNumberOfReplicatedDestructionInfos(uint32 Count)
+{
+	Stats.ReplicatedDestructionInfoCount += Count;
+}
+
+inline uint32 FNetSendStats::GetNumberOfReplicatedObjects() const
+{
+	return Stats.ReplicatedObjectCount;
+}
+
+inline void FNetSendStats::SetNumberOfActiveHugeObjects(uint32 Count)
+{
+	Stats.ActiveHugeObjectCount = static_cast<int32>(Count);
+}
+
+inline void FNetSendStats::AddHugeObjectWaitingTime(double Seconds)
+{
+	++Stats.HugeObjectsWaitingForAckCount;
+	Stats.HugeObjectWaitingForAckTimeInSeconds += Seconds;
+}
+
+inline void FNetSendStats::AddHugeObjectStallTime(double Seconds)
+{
+	++Stats.HugeObjectsStallingCount;
+	Stats.HugeObjectStallingTimeInSeconds += Seconds;
+}
+
+inline void FNetSendStats::AddReplicationWasteTime(double Seconds)
+{
+	++Stats.ReplicationWasteObjectCount;
+	Stats.ReplicationWasteTimeInSeconds += Seconds;
+}
+
+inline void FNetSendStats::SetNumberOfReplicatingConnections(uint32 Count)
+{
+	Stats.ReplicatingConnectionCount = Count;
+}
+
+}
diff --git a/Engine/Source/Runtime/Experimental/Iris/Core/Public/Iris/Stats/NetStats.cpp b/Engine/Source/Runtime/Experimental/Iris/Core/Public/Iris/Stats/NetStats.cpp
deleted file mode 100644
diff --git a/Engine/Source/Runtime/Experimental/Iris/Core/Public/Iris/Stats/NetStats.h b/Engine/Source/Runtime/Experimental/Iris/Core/Public/Iris/Stats/NetStats.h
deleted file mode 100644
