Fixed a bug that was causing IAS to assert if a service url was provided via the cmdline or OnDemand.ini file

#rnx

- FOnDemandIoBackend::InitializePrimaryEndpoint was assuming that it would always need to select and test and endpoint, which is true if the list of endpoints was downloading from an online server, but if we only have a single url provided by the cmdline or OnDemand.ini file then the endpoint url will already be set up.
- We no longer assert if the endpoint is already set, we accept the provided endpoint and only try to find one if required.
- We also need to call ::FlushDeferredTocs for both code paths (but not if there is an error) as when running with a single url, if the initial attempt to download the .iochunktoc fails we will pass it on as a deferredtoc, so we need to make sure that it gets processed.

diff --git a/Engine/Source/Runtime/Experimental/IoStoreOnDemand/Private/OnDemandIoDispatcherBackend.cpp b/Engine/Source/Runtime/Experimental/IoStoreOnDemand/Private/OnDemandIoDispatcherBackend.cpp
--- a/Engine/Source/Runtime/Experimental/IoStoreOnDemand/Private/OnDemandIoDispatcherBackend.cpp
+++ b/Engine/Source/Runtime/Experimental/IoStoreOnDemand/Private/OnDemandIoDispatcherBackend.cpp
@@ -2025,27 +2025,30 @@ void FOnDemandIoBackend::InitializePrimaryEndpoint()
 {
 	TRACE_CPUPROFILER_EVENT_SCOPE(IasBackend::InitializePrimaryEndpoint);
 
-	UE_LOG(LogIas, Log, TEXT("Attempting to find the primary endpoint..."));
-
-	check(!AvailableEps.HasCurrent());
-
-	TConstArrayView<FString> Urls = AvailableEps.Urls;
-	const int32 MaxUrls = FMath::Min(GIasMaxEndpointTestCountAtStartup, AvailableEps.Urls.Num());
-	const FString TestPath = GetEndpointTestPath();
-	if (int32 Idx = LatencyTest(Urls.Left(MaxUrls), TestPath, bStopRequested); Idx != INDEX_NONE)
+	// We only need to run this code if we do not have an currently selected, most likely because
+	// a list of endpoints was polled from online rather than a url being provided by the cmdline/config system.
+	if (!AvailableEps.HasCurrent())
 	{
-		AvailableEps.Current = Idx;
-		UE_LOG(LogIas, Log, TEXT("Using endpoint '%s'"), *AvailableEps.GetCurrent());
-		
-		// We need to call this after a valid endpoint has been selected in case we need to download
-		// the .iochunktoc file.
-		FlushDeferredTocs(EFlushMode::All);
-	}
-	else
-	{
-		UE_LOG(LogIas, Error, TEXT("Unable to connect to any valid endpoint"));
-		BackendStatus.SetHttpError(true);
+		UE_LOG(LogIas, Log, TEXT("Attempting to find the primary endpoint..."));
+
+		TConstArrayView<FString> Urls = AvailableEps.Urls;
+		const int32 MaxUrls = FMath::Min(GIasMaxEndpointTestCountAtStartup, AvailableEps.Urls.Num());
+		const FString TestPath = GetEndpointTestPath();
+		if (int32 Idx = LatencyTest(Urls.Left(MaxUrls), TestPath, bStopRequested); Idx != INDEX_NONE)
+		{
+			AvailableEps.Current = Idx;
+			UE_LOG(LogIas, Log, TEXT("Using endpoint '%s'"), *AvailableEps.GetCurrent());
+		}
+		else
+		{
+			UE_LOG(LogIas, Error, TEXT("Unable to connect to any valid endpoint"));
+			BackendStatus.SetHttpError(true);
+			return;
+		}
 	}
+
+	// Now we have a valid endpoint, we need to flush any pending .iochunktoc download requests
+	FlushDeferredTocs(EFlushMode::All);
 }
 
 void FOnDemandIoBackend::FlushDeferredTocs(EFlushMode FlushMode)
