Added tamper tests to IAS

#rnx

diff --git a/Engine/Source/Runtime/Experimental/IoStoreOnDemand/Private/CoreHttp/Http.cpp b/Engine/Source/Runtime/Experimental/IoStoreOnDemand/Private/CoreHttp/Http.cpp
--- a/Engine/Source/Runtime/Experimental/IoStoreOnDemand/Private/CoreHttp/Http.cpp
+++ b/Engine/Source/Runtime/Experimental/IoStoreOnDemand/Private/CoreHttp/Http.cpp
@@ -4038,6 +4038,22 @@ IOSTOREONDEMAND_API void IasHttpTest(const ANSICHAR* TestHost="localhost")
 		WaitForLoopIdle();
 	}
 
+	// tamper
+	for (int32 i = 1; i <= 100; ++i)
+	{
+		TAnsiStringBuilder<32> TamperUrl;
+		TamperUrl << "/data?tamper=" << i;
+		FAnsiStringView Url = BuildUrl(TamperUrl.ToString(), 9494);
+
+		for (int j = 0; j < 48; ++j)
+		{
+			FRequest Request = Loop.Request("GET", Url);
+			Loop.Send(MoveTemp(Request), NullSink);
+		}
+
+		WaitForLoopIdle();
+	}
+
 	LoopStop = true;
 	LoopTask.Wait();
 
diff --git a/Engine/Source/Runtime/Experimental/IoStoreOnDemand/Private/CoreHttp/TestServer.py b/Engine/Source/Runtime/Experimental/IoStoreOnDemand/Private/CoreHttp/TestServer.py
--- a/Engine/Source/Runtime/Experimental/IoStoreOnDemand/Private/CoreHttp/TestServer.py
+++ b/Engine/Source/Runtime/Experimental/IoStoreOnDemand/Private/CoreHttp/TestServer.py
@@ -1,5 +1,6 @@
 import io
 import os
+import re
 import base64
 import random
 import socket
@@ -15,19 +16,28 @@ async def proxy_impl(loop, client):
     while not msg.endswith(b"\r\n\r\n"):
         msg += await loop.sock_recv(client, 2048)
 
+    mode = "="
     min_trunc = 0
     disconnect = False
     stall = False
+    tamper = 0
 
     line = next(iter(io.BytesIO(msg)))
     if b"?disconnect HTTP" in line:
+        mode = "D"
         disconnect = True
 
     if b"?stall HTTP" in line:
+        mode = "S"
         disconnect = True
         stall = True
         min_trunc = len(msg) + 16
 
+    if m := re.search(b"\?tamper=(\d+) HTTP", line):
+        mode = "T"
+        tamper = int(m.group(1))
+        tamper = tamper / 100.0
+
     send_time = 1.0 + (random.random() * 0.5)
     with socket.socket(socket.AF_INET, socket.SOCK_STREAM) as httpd:
         httpd.connect(("127.0.0.1", 9493))
@@ -45,18 +55,30 @@ async def proxy_impl(loop, client):
             print("trun:", data_size, "->", trunc)
             data = data[:trunc]
 
-        while data:
-            percent = 0.02 + (random.random() * 0.08)
-            send_size = max(int(data_size * percent), 1)
-            print("snd:%5d tot:%5d slp:%5f %%:%5f" % (send_size, data_size, send_time * percent, percent), end="\r")
-            await loop.sock_sendall(client, data[:send_size])
-            data = data[send_size:]
-            await asyncio.sleep(send_time * percent)
+        if tamper:
+            data = bytearray(data)
+            for i in range(len(data)):
+                c = data[i] if random.random() > tamper else (int(random.random() * 0x4567) & 0xff)
+                data[i] = c
+
+        try:
+            while data:
+                percent = 0.02 + (random.random() * 0.08)
+                send_size = max(int(data_size * percent), 1)
+                print("%s snd:%5d tot:%5d slp:%5f %%:%5f" % (mode, send_size, data_size, send_time * percent, percent), end="\r")
+                await loop.sock_sendall(client, data[:send_size])
+                data = data[send_size:]
+                if not tamper:
+                    await asyncio.sleep(send_time * percent)
+        except ConnectionError:
+            pass
         print("")
 
+
         if stall:
             print("...stalling", end="")
             await asyncio.sleep(2)
+            print("")
 
     client.close()
 
@@ -69,10 +91,7 @@ async def proxy_loop():
         sock.setblocking(True)
         while True:
             client, address = await loop.sock_accept(sock)
-            try:
-                loop.create_task(proxy_impl(loop, client))
-            except (ConnectionResetError, ConnectionAbortedError):
-                pass
+            loop.create_task(proxy_impl(loop, client))
 
 #-------------------------------------------------------------------------------
 def proxy():
