ChaosSofts : Add Gauss Seidel and Newton/CG for cloth constraints. Add corotated codimensional model. Bug fix for rest length in anisotropic bias elements constraints. Add more functions to matrix class.

diff --git a/Engine/Plugins/Experimental/ChaosCloth/Source/Chaos/Private/ChaosCloth/ChaosClothConstraints.cpp b/Engine/Plugins/Experimental/ChaosCloth/Source/Chaos/Private/ChaosCloth/ChaosClothConstraints.cpp
--- a/Engine/Plugins/Experimental/ChaosCloth/Source/Chaos/Private/ChaosCloth/ChaosClothConstraints.cpp
+++ b/Engine/Plugins/Experimental/ChaosCloth/Source/Chaos/Private/ChaosCloth/ChaosClothConstraints.cpp
@@ -21,8 +21,69 @@
 #include "Chaos/PBDTriangleMeshIntersections.h"
 #include "Chaos/PBDEvolution.h"
 #include "Chaos/CollectionPropertyFacade.h"
+#include "Chaos/Deformable/GaussSeidelMasterConstraint.h"
+#include "Chaos/Deformable/GaussSeidelCorotatedCodimensionalConstraints.h"
 
 namespace Chaos {
+#if !UE_BUILD_SHIPPING
+
+	bool bEnableGS = false;
+
+    static FAutoConsoleVariableRef CVarClothbEnableGS(TEXT("p.Chaos.Cloth.EnableGaussSeidel"), bEnableGS, TEXT("Use Gauss Seidel constraints instead of XPBD [def: false]"));
+
+	bool bDisplayResidual = false;
+
+	static FAutoConsoleVariableRef CVarClothbDisplayResidual(TEXT("p.Chaos.Cloth.DisplayResidual"), bDisplayResidual, TEXT("Diplay residual norms for the first 10 iters [def: false]"));
+
+	int32 MaxResidualIters = 10;
+
+	static FAutoConsoleVariableRef CVarClothMaxResidualIters(TEXT("p.Chaos.Cloth.MaxResidualIters"), MaxResidualIters, TEXT("Max number of iterations to diaplay residuals [def: 10]"));
+
+	bool bWriteResidual2File = false;
+
+	static FAutoConsoleVariableRef CVarClothbWriteResidual2File(TEXT("p.Chaos.Cloth.WriteResidual2File"), bWriteResidual2File, TEXT("Write residual to file [def: false]"));
+
+	bool bReplaceBiasElementsWithCorotatedCod = false;
+
+	static FAutoConsoleVariableRef CVarClothReplaceBiasElementsWithCorotatedCod(TEXT("p.Chaos.Cloth.ReplaceBiasElementsWithCorotatedCod"), bReplaceBiasElementsWithCorotatedCod, TEXT("Replace existing aniso bias element constraint with gauss seidel corotated codimensional [def: false]"));
+
+	Softs::FSolverReal YoungsModulus = 10000.f;
+
+	static FAutoConsoleVariableRef CVarClothYoungsModulus(TEXT("p.Chaos.Cloth.YoungsModulus"), YoungsModulus, TEXT("Youngs modulus [def: 1e4]"));
+
+	bool bClothDoQuasistatics = false;
+
+	static FAutoConsoleVariableRef CVarClothDoQuasistatics(TEXT("p.Chaos.Cloth.DoQuasistatics"), bClothDoQuasistatics, TEXT("Do cloth quasistatics [def: false]"));
+
+	bool bEnableCG = false;
+
+	static FAutoConsoleVariableRef CVarClothEnableCG(TEXT("p.Chaos.Cloth.EnableCG"), bEnableCG, TEXT("Use conjugate gradient instead of nonlinear gauss seidel [def: false]"));
+
+	bool bMakeSandwich = false;
+
+	static FAutoConsoleVariableRef CVarClothMakeSandwich(TEXT("p.Chaos.Cloth.MakeSandwich"), bMakeSandwich, TEXT("MakeSandwich solver of xpbd - Newton/GS [def: false]"));
+
+	int32 UpperBreadIters = 3;
+
+	static FAutoConsoleVariableRef CVarClothUpperBreadIters(TEXT("p.Chaos.Cloth.UpperBreadIters"), UpperBreadIters, TEXT("Upper Bread Iters for the sandwich solver [def: 3]"));
+
+	int32 MiddleBreadIters = 3;
+
+	static FAutoConsoleVariableRef CVarClothMiddleBreadIters(TEXT("p.Chaos.Cloth.MiddleBreadIters"), MiddleBreadIters, TEXT("Middle Bread Iters for the sandwich solver [def: 3]"));
+
+	bool bWriteFinalResiduals = false;
+
+	static FAutoConsoleVariableRef CVarClothWriteFinalResiduals(TEXT("p.Chaos.Cloth.WriteFinalResiduals"), bWriteFinalResiduals, TEXT("Write final residuals at each timestep to a file [def: false]"));
+
+	bool bUseSOR = true;
+
+	static FAutoConsoleVariableRef CVarClothUseSOR(TEXT("p.Chaos.Cloth.UseSOR"), bUseSOR, TEXT("Use SOR acceleration for Gauss Seidel [def: true]"));
+
+	Softs::FSolverReal SOROmega = 1.7f;
+
+	static FAutoConsoleVariableRef CVarClothSOROmega(TEXT("p.Chaos.Cloth.SOROmega"), SOROmega, TEXT("SOR omega coefficient for acceleration [def: 1.7]"));
+
+#endif
 
 FClothConstraints::FClothConstraints()
 	: Evolution(nullptr)
@@ -146,6 +207,311 @@ void FClothConstraints::AddRules(
 		Tethers, MeshScale, bEnabled);
 }
 
+//only counts the number of rules that GS/CG currently cupports
+void FClothConstraints::GetGSNumRules()
+{
+#if !UE_BUILD_SHIPPING
+	NumConstraintInits = 1;
+	NumConstraintRules = 0;
+	if (XStretchBiasConstraints)
+	{
+		NumConstraintInits += 1;
+	}
+	if (XBendingElementConstraints)
+	{
+		NumConstraintInits += 1;
+	}
+	if (XEdgeConstraints)
+	{
+		NumConstraintInits += 1;
+	}
+	if (XAnisoBendingElementConstraints)
+	{
+		NumConstraintInits += 1;
+	}
+
+	if (bMakeSandwich)
+	{
+		NumConstraintRules = 2;
+	}
+	else
+	{
+		NumConstraintRules = 1;
+	}
+#endif
+}
+
+
+void FClothConstraints::CreateGSRules()
+{
+#if !UE_BUILD_SHIPPING
+	check(Evolution);
+	check(ConstraintInitOffset == INDEX_NONE)
+	GetGSNumRules();
+
+	ConstraintInitOffset = Evolution->AddConstraintInitRange(NumConstraintInits, false);
+
+	check(ConstraintRuleOffset == INDEX_NONE)
+
+	ConstraintRuleOffset = Evolution->AddConstraintRuleRange(NumConstraintRules, false);
+
+	check(PostprocessingConstraintRuleOffset == INDEX_NONE)
+
+	int32 PostProcessingConstraintRule = 0;
+	if (bWriteFinalResiduals)
+	{
+		PostProcessingConstraintRule = 1;
+	}
+
+	if (PostProcessingConstraintRule)
+	{
+		PostprocessingConstraintRuleOffset = Evolution->AddConstraintPostprocessingsRange(PostProcessingConstraintRule, false);
+	}
+
+	TFunction<void(Softs::FSolverParticles&, const Softs::FSolverReal)>* const ConstraintInits = Evolution->ConstraintInits().GetData() + ConstraintInitOffset;
+	TFunction<void(Softs::FSolverParticles&, const Softs::FSolverReal)>* const ConstraintRules = Evolution->ConstraintRules().GetData() + ConstraintRuleOffset;
+	TFunction<void(Softs::FSolverParticles&, const Softs::FSolverReal)>* const PostprocessingRules = Evolution->ConstraintPostprocessings().GetData() + PostprocessingConstraintRuleOffset;
+
+	int32 ConstraintInitIndex = 0;
+	int32 ConstraintRuleIndex = 0;
+
+	int32 ConstraintPostProcessingIndex = 0;
+
+	GSMasterConstraint = MakeShared<Chaos::Softs::FGaussSeidelMasterConstraint<Softs::FSolverReal, Softs::FSolverParticles>>(Evolution->Particles(), bClothDoQuasistatics, bUseSOR, SOROmega, 100);
+
+	if (bClothDoQuasistatics)
+	{
+		Evolution->SetQuasistatics(bClothDoQuasistatics);
+		const Softs::FSolverVec3& GravityTerm = Evolution->GetGravity(0);
+		GSMasterConstraint->ExternalForce[0] = 0.f;
+		GSMasterConstraint->ExternalForce[1] = 0.f;
+		GSMasterConstraint->ExternalForce[2] = GravityTerm[2];
+	}
+
+	ConstraintInits[ConstraintInitIndex++] =
+		[this](Softs::FSolverParticles& InParticles, const Softs::FSolverReal Dt)
+	{
+		this->GSMasterConstraint->Init(Dt, InParticles);
+	};
+
+	if (XStretchBiasConstraints)
+	{
+		ConstraintInits[ConstraintInitIndex++] =
+			[this](Softs::FSolverParticles& /*Particles*/, const Softs::FSolverReal Dt)
+		{
+			XStretchBiasConstraints->Init();
+			XStretchBiasConstraints->ApplyProperties(Dt, Evolution->GetIterations());
+		};
+	}
+
+	if (XBendingElementConstraints)
+	{
+		ConstraintInits[ConstraintInitIndex++] =
+			[this](Softs::FSolverParticles& InParticles, const Softs::FSolverReal Dt)
+		{
+			this->XBendingElementConstraints->Init(InParticles);
+			this->XBendingElementConstraints->ApplyProperties(Dt, Evolution->GetIterations());
+		};
+	}
+
+	if (XAnisoBendingElementConstraints)
+	{
+		ConstraintInits[ConstraintInitIndex++] =
+			[this](Softs::FSolverParticles& InParticles, const Softs::FSolverReal Dt)
+		{
+			this->XAnisoBendingElementConstraints->Init(InParticles);
+			this->XAnisoBendingElementConstraints->ApplyProperties(Dt, Evolution->GetIterations());
+		};
+	}
+
+	if (XEdgeConstraints)
+	{
+		ConstraintInits[ConstraintInitIndex++] =
+			[this](Softs::FSolverParticles& InParticles, const Softs::FSolverReal Dt)
+		{
+			this->XEdgeConstraints->ApplyProperties(Dt, Evolution->GetIterations());
+		};
+	}
+
+	if (bMakeSandwich)
+	{
+		if (XStretchBiasConstraints)
+		{
+			ConstraintRules[ConstraintRuleIndex++] =
+				[this, PreIterXPBDNums = UpperBreadIters, bbDisplayResidual = bDisplayResidual](Softs::FSolverParticles& Particles, const Softs::FSolverReal Dt)
+			{
+				for (int32 i = 0; i < PreIterXPBDNums; i++)
+				{
+					XStretchBiasConstraints->Apply(Particles, Dt);
+					if (this->XBendingElementConstraints)
+					{
+						this->XBendingElementConstraints->Apply(Particles, Dt);
+					}
+					if (bbDisplayResidual)
+					{
+						TArray<Chaos::Softs::FSolverVec3> Residual = this->GSMasterConstraint->ComputeNewtonResiduals(Particles, Dt);
+					}
+				}
+			};
+		}
+
+
+		//middle layer of sandwich:
+		if (bEnableCG)
+		{
+			ConstraintRules[ConstraintRuleIndex++] =
+				[this, GSIterNums = MiddleBreadIters](Softs::FSolverParticles& InParticles, const Softs::FSolverReal Dt)
+			{
+				for (int32 i = 0; i < GSIterNums; i++)
+				{
+					this->GSMasterConstraint->ApplyCG(InParticles, Dt);
+				}
+			};
+		}
+		else
+		{
+			ConstraintRules[ConstraintRuleIndex++] =
+				[this, GSIterNums = MiddleBreadIters](Softs::FSolverParticles& InParticles, const Softs::FSolverReal Dt)
+			{
+				for (int32 i = 0; i < GSIterNums; i++)
+				{
+					this->GSMasterConstraint->Apply(InParticles, Dt, MaxResidualIters, bWriteResidual2File);
+				}
+			};
+		}
+
+	}
+	else
+	{
+		if (bEnableCG)
+		{
+			ConstraintRules[ConstraintRuleIndex++] =
+				[this](Softs::FSolverParticles& InParticles, const Softs::FSolverReal Dt)
+			{
+				this->GSMasterConstraint->ApplyCG(InParticles, Dt);
+			};
+		}
+		else
+		{
+			ConstraintRules[ConstraintRuleIndex++] =
+				[this](Softs::FSolverParticles& InParticles, const Softs::FSolverReal Dt)
+			{
+				this->GSMasterConstraint->Apply(InParticles, Dt, MaxResidualIters, bWriteResidual2File);
+			};
+		}
+	}
+
+
+
+	if (XStretchBiasConstraints)
+	{
+		if (bReplaceBiasElementsWithCorotatedCod)
+		{
+			TArray<TArray<int32>> IncidentElements, IncidentElementsLocal;
+
+			const TArray<TVec3<int32>> CodimensionalMesh = XStretchBiasConstraints->GetConstraints();
+
+			GSCorotatedCodimensionalConstraint = MakeShared<Chaos::Softs::FGaussSeidelCorotatedCodimensionalConstraints<Softs::FSolverReal, Softs::FSolverParticles>>(Evolution->Particles(), CodimensionalMesh, false, YoungsModulus);
+
+			GSMasterConstraint->AddStaticConstraints(XStretchBiasConstraints->GetConstraintsArray(), IncidentElements, IncidentElementsLocal);
+			const int32 StaticIndex = GSMasterConstraint->AddStaticConstraintResidualAndHessianRange(1);
+
+			GSMasterConstraint->StaticConstraintResidualAndHessian()[StaticIndex] = [this](const Softs::FSolverParticles& Particles, const int32 ElementIndex, const int32 ElementIndexLocal, const Softs::FSolverReal Dt, TVec3<Softs::FSolverReal>& ParticleResidual, Chaos::PMatrix<Softs::FSolverReal, 3, 3>& ParticleHessian)
+			{
+				this->GSCorotatedCodimensionalConstraint->AddHyperelasticResidualAndHessian(Particles, ElementIndex, ElementIndexLocal, Dt, ParticleResidual, ParticleHessian);
+			};
+		}
+		else
+		{
+			TArray<TArray<int32>> IncidentElements, IncidentElementsLocal;
+			GSMasterConstraint->AddStaticConstraints(XStretchBiasConstraints->GetConstraintsArray(), IncidentElements, IncidentElementsLocal);
+			const int32 StaticIndex = GSMasterConstraint->AddStaticConstraintResidualAndHessianRange(1);
+
+			GSMasterConstraint->StaticConstraintResidualAndHessian()[StaticIndex] = [this](const Softs::FSolverParticles& Particles, const int32 ElementIndex, const int32 ElementIndexLocal, const Softs::FSolverReal Dt, TVec3<Softs::FSolverReal>& ParticleResidual, Chaos::PMatrix<Softs::FSolverReal, 3, 3>& ParticleHessian)
+			{
+				this->XStretchBiasConstraints->AddStretchBiasElementResidualAndHessian(Particles, ElementIndex, ElementIndexLocal, Dt, ParticleResidual, ParticleHessian);
+			};
+
+			if (bEnableCG)
+			{
+				const int32 ForceDifferentialRange = GSMasterConstraint->AddAddInternalForceDifferentialsRange(1);
+
+				GSMasterConstraint->InternalForceDifferentials()[ForceDifferentialRange] = [this](const Softs::FSolverParticles& Particles, const TArray<Softs::FSolverVec3>& Deltax, TArray<Softs::FSolverVec3>& ndf)
+				{
+					this->XStretchBiasConstraints->AddInternalForceDifferential(Particles, Deltax, ndf);
+				};
+			}
+
+		}
+
+		XStretchBiasConstraints->InitializeDmInvAndMeasures(Evolution->Particles());
+
+	}
+
+	if (XAnisoBendingElementConstraints)
+	{
+		TArray<TArray<int32>> IncidentElements, IncidentElementsLocal;
+		GSMasterConstraint->AddStaticConstraints(XAnisoBendingElementConstraints->GetConstraintsArray(), IncidentElements, IncidentElementsLocal);
+		const int32 StaticIndex = GSMasterConstraint->AddStaticConstraintResidualAndHessianRange(1);
+
+		GSMasterConstraint->StaticConstraintResidualAndHessian()[StaticIndex] = [this](const Softs::FSolverParticles& Particles, const int32 ElementIndex, const int32 ElementIndexLocal, const Softs::FSolverReal Dt, TVec3<Softs::FSolverReal>& ParticleResidual, Chaos::PMatrix<Softs::FSolverReal, 3, 3>& ParticleHessian)
+		{
+			this->XAnisoBendingElementConstraints->AddAnisotropicBendingResidualAndHessian(Particles, ElementIndex, ElementIndexLocal, Dt, ParticleResidual, ParticleHessian);
+		};
+
+		if (bEnableCG)
+		{
+			const int32 ForceDifferentialRange = GSMasterConstraint->AddAddInternalForceDifferentialsRange(1);
+
+			GSMasterConstraint->InternalForceDifferentials()[ForceDifferentialRange] = [this](const Softs::FSolverParticles& Particles, const TArray<Softs::FSolverVec3>& Deltax, TArray<Softs::FSolverVec3>& ndf)
+			{
+				this->XAnisoBendingElementConstraints->AddInternalForceDifferential(Particles, Deltax, ndf);
+			};
+		}
+	}
+
+	if (XBendingElementConstraints)
+	{
+
+		TArray<TArray<int32>> IncidentElements, IncidentElementsLocal;
+		GSMasterConstraint->AddStaticConstraints(XBendingElementConstraints->GetConstraintsArray(), IncidentElements, IncidentElementsLocal);
+		const int32 StaticIndex = GSMasterConstraint->AddStaticConstraintResidualAndHessianRange(1);
+
+		GSMasterConstraint->StaticConstraintResidualAndHessian()[StaticIndex] = [this](const Softs::FSolverParticles& Particles, const int32 ElementIndex, const int32 ElementIndexLocal, const Softs::FSolverReal Dt, TVec3<Softs::FSolverReal>& ParticleResidual, Chaos::PMatrix<Softs::FSolverReal, 3, 3>& ParticleHessian)
+		{
+			this->XBendingElementConstraints->AddBendingResidualAndHessian(Particles, ElementIndex, ElementIndexLocal, Dt, ParticleResidual, ParticleHessian);
+		};
+
+
+		if (bEnableCG)
+		{
+			const int32 ForceDifferentialRange = GSMasterConstraint->AddAddInternalForceDifferentialsRange(1);
+
+			GSMasterConstraint->InternalForceDifferentials()[ForceDifferentialRange] = [this](const Softs::FSolverParticles& Particles, const TArray<Softs::FSolverVec3>& Deltax, TArray<Softs::FSolverVec3>& ndf)
+			{
+				this->XBendingElementConstraints->AddInternalForceDifferential(Particles, Deltax, ndf);
+			};
+		}
+	}
+
+	GSMasterConstraint->InitStaticColor(Evolution->Particles());
+
+	if (bWriteFinalResiduals)
+	{
+		FString file = FPaths::ProjectDir();
+		file.Append(TEXT("/DebugOutput/NewtonResidual.txt"));
+		FFileHelper::SaveStringToFile(FString(TEXT("Newton Norm\r\n")), *file);
+		PostprocessingRules[ConstraintPostProcessingIndex++] =
+			[this](Softs::FSolverParticles& InParticles, const Softs::FSolverReal Dt)
+		{
+			this->GSMasterConstraint->ComputeNewtonResiduals(InParticles, Dt, true, nullptr);
+		};
+	}
+
+#endif
+
+}
+
 void FClothConstraints::AddRules(
 	const Softs::FCollectionPropertyConstFacade& ConfigProperties,
 	const FTriangleMesh& TriangleMesh,
@@ -179,9 +545,23 @@ void FClothConstraints::AddRules(
 	// Animation Drive Constraints
 	CreateAnimDriveConstraints(ConfigProperties, WeightMaps);
 
+
 	// Commit rules to solver
 PRAGMA_DISABLE_DEPRECATION_WARNINGS
-	CreateRules();  // TODO: Move CreateRules to private
+
+#if !UE_BUILD_SHIPPING
+	if (!bEnableGS)
+	{
+		CreateRules();  // TODO: Move CreateRules to private
+	} 
+	else
+	{
+		CreateGSRules();
+	}
+#else
+	CreateRules();
+#endif
+
 PRAGMA_ENABLE_DEPRECATION_WARNINGS
 
 	// Enable or disable constraints as requested
@@ -527,11 +907,22 @@ void FClothConstraints::CreateRules()
 {
 	check(Evolution);
 	check(ConstraintInitOffset == INDEX_NONE)
+#if !UE_BUILD_SHIPPING	
+	if (bDisplayResidual)
+	{
+		NumConstraintRules++;
+	}
+	if (bWriteFinalResiduals || bDisplayResidual)
+	{
+		NumConstraintInits++;
+	}
+#endif
 	if (NumConstraintInits)
 	{
 		ConstraintInitOffset = Evolution->AddConstraintInitRange(NumConstraintInits, false);
 	}
 	check(ConstraintRuleOffset == INDEX_NONE)
+
 	if (NumConstraintRules)
 	{
 		ConstraintRuleOffset = Evolution->AddConstraintRuleRange(NumConstraintRules, false);
@@ -542,22 +933,106 @@ void FClothConstraints::CreateRules()
 		PostCollisionConstraintRuleOffset = Evolution->AddPostCollisionConstraintRuleRange(NumPostCollisionConstraintRules, false);
 	}
 	check(PostprocessingConstraintRuleOffset == INDEX_NONE);
+#if !UE_BUILD_SHIPPING	
+	if (bWriteFinalResiduals)
+	{
+		NumPostprocessingConstraintRules += 1;
+	}
+#endif
 	if (NumPostprocessingConstraintRules)
 	{
 		PostprocessingConstraintRuleOffset = Evolution->AddConstraintPostprocessingsRange(NumPostprocessingConstraintRules, false);
 	}
 
+	check(PostprocessingConstraintRuleOffset == INDEX_NONE)
+
+
 	TFunction<void(Softs::FSolverParticles&, const Softs::FSolverReal)>* const ConstraintInits = Evolution->ConstraintInits().GetData() + ConstraintInitOffset;
 	TFunction<void(Softs::FSolverParticles&, const Softs::FSolverReal)>* const ConstraintRules = Evolution->ConstraintRules().GetData() + ConstraintRuleOffset;
 	TFunction<void(Softs::FSolverParticles&, const Softs::FSolverReal)>* const PostCollisionConstraintRules = Evolution->PostCollisionConstraintRules().GetData() + PostCollisionConstraintRuleOffset;
-	TFunction<void(Softs::FSolverParticles&, const Softs::FSolverReal)>* const
-		PostprocessingConstraintRules = Evolution->ConstraintPostprocessings().GetData() +
-		PostprocessingConstraintRuleOffset;
+	TFunction<void(Softs::FSolverParticles&, const Softs::FSolverReal)>* const PostprocessingConstraintRules = Evolution->ConstraintPostprocessings().GetData() + PostprocessingConstraintRuleOffset;
 
 	int32 ConstraintInitIndex = 0;
 	int32 ConstraintRuleIndex = 0;
 	int32 PostCollisionConstraintRuleIndex = 0;
 	int32 PostprocessingConstraintRuleIndex = 0;
+	#if !UE_BUILD_SHIPPING	
+	if (bDisplayResidual || bWriteFinalResiduals)
+	{
+		GSMasterConstraint = MakeShared<Chaos::Softs::FGaussSeidelMasterConstraint<Softs::FSolverReal, Softs::FSolverParticles>>(Evolution->Particles(), false, false, (Softs::FSolverReal)1.2, 100);
+
+		ConstraintInits[ConstraintInitIndex++] =
+			[this](Softs::FSolverParticles& InParticles, const Softs::FSolverReal Dt)
+		{
+			this->GSMasterConstraint->Init(Dt, InParticles);
+		};
+		if (bDisplayResidual)
+		{
+
+			ConstraintRules[ConstraintRuleIndex++] =
+				[this](Softs::FSolverParticles& InParticles, const Softs::FSolverReal Dt)
+			{
+				if (this->GSMasterConstraint->DebugResidual && this->GSMasterConstraint->PassedIters < MaxResidualIters)
+				{
+					UE_LOG(LogTemp, Warning, TEXT("Current Iteration As From Evolution: %d"), this->Evolution->GetIterations());
+					this->GSMasterConstraint->ComputeNewtonResiduals(InParticles, Dt, bWriteResidual2File);
+					this->GSMasterConstraint->PassedIters += 1;
+				}
+			};
+		}
+
+		if (XStretchBiasConstraints)
+		{
+			TArray<TArray<int32>> IncidentElements, IncidentElementsLocal;
+			GSMasterConstraint->AddStaticConstraints(XStretchBiasConstraints->GetConstraintsArray(), IncidentElements, IncidentElementsLocal);
+			int32 StaticIndex = GSMasterConstraint->AddStaticConstraintResidualAndHessianRange(1);
+
+			GSMasterConstraint->StaticConstraintResidualAndHessian()[StaticIndex] = [this](const Softs::FSolverParticles& Particles, const int32 ElementIndex, const int32 ElementIndexLocal, const Softs::FSolverReal Dt, TVec3<Softs::FSolverReal>& ParticleResidual, Chaos::PMatrix<Softs::FSolverReal, 3, 3>& ParticleHessian)
+			{
+				this->XStretchBiasConstraints->AddStretchBiasElementResidualAndHessian(Particles, ElementIndex, ElementIndexLocal, Dt, ParticleResidual, ParticleHessian);
+			};
+		}
+
+		if (XAnisoBendingElementConstraints)
+		{
+			TArray<TArray<int32>> IncidentElements, IncidentElementsLocal;
+			GSMasterConstraint->AddStaticConstraints(XAnisoBendingElementConstraints->GetConstraintsArray(), IncidentElements, IncidentElementsLocal);
+			int32 StaticIndex = GSMasterConstraint->AddStaticConstraintResidualAndHessianRange(1);
+
+			GSMasterConstraint->StaticConstraintResidualAndHessian()[StaticIndex] = [this](const Softs::FSolverParticles& Particles, const int32 ElementIndex, const int32 ElementIndexLocal, const Softs::FSolverReal Dt, TVec3<Softs::FSolverReal>& ParticleResidual, Chaos::PMatrix<Softs::FSolverReal, 3, 3>& ParticleHessian)
+			{
+				this->XAnisoBendingElementConstraints->AddAnisotropicBendingResidualAndHessian(Particles, ElementIndex, ElementIndexLocal, Dt, ParticleResidual, ParticleHessian);
+			};
+		}
+
+
+		if (XBendingElementConstraints)
+		{
+			TArray<TArray<int32>> IncidentElements, IncidentElementsLocal;
+			GSMasterConstraint->AddStaticConstraints(XBendingElementConstraints->GetConstraintsArray(), IncidentElements, IncidentElementsLocal);
+			int32 StaticIndex = GSMasterConstraint->AddStaticConstraintResidualAndHessianRange(1);
+
+			GSMasterConstraint->StaticConstraintResidualAndHessian()[StaticIndex] = [this](const Softs::FSolverParticles& Particles, const int32 ElementIndex, const int32 ElementIndexLocal, const Softs::FSolverReal Dt, TVec3<Softs::FSolverReal>& ParticleResidual, Chaos::PMatrix<Softs::FSolverReal, 3, 3>& ParticleHessian)
+			{
+				this->XBendingElementConstraints->AddBendingResidualAndHessian(Particles, ElementIndex, ElementIndexLocal, Dt, ParticleResidual, ParticleHessian);
+			};
+		}
+
+		GSMasterConstraint->InitStaticColor(Evolution->Particles());
+
+		if (bWriteFinalResiduals)
+		{
+			FString file = FPaths::ProjectDir();
+			file.Append(TEXT("/DebugOutput/NewtonResidual.txt"));
+			FFileHelper::SaveStringToFile(FString(TEXT("Newton Norm\r\n")), *file);
+			PostprocessingConstraintRules[PostprocessingConstraintRuleIndex++] =
+				[this](Softs::FSolverParticles& InParticles, const Softs::FSolverReal Dt)
+			{
+				this->GSMasterConstraint->ComputeNewtonResiduals(InParticles, Dt, true, nullptr);
+			};
+		}
+	}
+#endif
 
 	if (XStretchBiasConstraints)
 	{
diff --git a/Engine/Plugins/Experimental/ChaosCloth/Source/Chaos/Public/ChaosCloth/ChaosClothConstraints.h b/Engine/Plugins/Experimental/ChaosCloth/Source/Chaos/Public/ChaosCloth/ChaosClothConstraints.h
--- a/Engine/Plugins/Experimental/ChaosCloth/Source/Chaos/Public/ChaosCloth/ChaosClothConstraints.h
+++ b/Engine/Plugins/Experimental/ChaosCloth/Source/Chaos/Public/ChaosCloth/ChaosClothConstraints.h
@@ -7,6 +7,8 @@
 #include "Chaos/PBDSoftsEvolutionFwd.h"
 #include "Chaos/PBDLongRangeConstraints.h"
 #include "Chaos/PBDCollisionSpringConstraintsBase.h"
+#include "Chaos/Deformable/GaussSeidelMasterConstraint.h"
+#include "Chaos/Deformable/GaussSeidelCorotatedCodimensionalConstraints.h"
 
 namespace Chaos
 {
@@ -251,6 +253,8 @@ PRAGMA_ENABLE_DEPRECATION_WARNINGS
 		void CreateAnimDriveConstraints(
 			const Softs::FCollectionPropertyConstFacade& ConfigProperties,
 			const TMap<FString, TConstArrayView<FRealSingle>>& WeightMaps);
+		void CreateGSRules();
+		void GetGSNumRules();
 
 		//~ Begin deprecated constraints
 		TSharedPtr<Softs::FPBDSpringConstraints> EdgeConstraints_Deprecated;
@@ -283,6 +287,8 @@ PRAGMA_ENABLE_DEPRECATION_WARNINGS
 		TSharedPtr<Softs::FPBDCollisionSpringConstraints> SelfCollisionConstraints;
 		TSharedPtr<Softs::FPBDTriangleMeshIntersections> SelfIntersectionConstraints;
 		TSharedPtr<Softs::FPBDSelfCollisionSphereConstraints> SelfCollisionSphereConstraints;
+		TSharedPtr<Softs::FGaussSeidelMasterConstraint<Softs::FSolverReal, Softs::FSolverParticles>> GSMasterConstraint;
+		TSharedPtr<Softs::FGaussSeidelCorotatedCodimensionalConstraints<Softs::FSolverReal, Softs::FSolverParticles>> GSCorotatedCodimensionalConstraint;
 		
 		Softs::FPBDEvolution* Evolution;
 		const TArray<Softs::FSolverVec3>* AnimationPositions;
diff --git a/Engine/Source/Runtime/Experimental/Chaos/Private/Chaos/Deformable/GaussSeidelMasterConstraint.cpp b/Engine/Source/Runtime/Experimental/Chaos/Private/Chaos/Deformable/GaussSeidelMasterConstraint.cpp
--- a/Engine/Source/Runtime/Experimental/Chaos/Private/Chaos/Deformable/GaussSeidelMasterConstraint.cpp
+++ b/Engine/Source/Runtime/Experimental/Chaos/Private/Chaos/Deformable/GaussSeidelMasterConstraint.cpp
@@ -1,8 +1,17 @@
 // Copyright Epic Games, Inc. All Rights Reserved.
 #include "Chaos/Deformable/GaussSeidelMasterConstraint.h"
+#include "Chaos/Math/Krylov.h"
 
 namespace Chaos::Softs
 {
+	int32 MaxItCG = 50;
+
+	FAutoConsoleVariableRef CVarClothMaxItCG(TEXT("p.Chaos.Cloth.MaxItCG"), MaxItCG, TEXT("Max iter for CG [def: 50]"));
+
+	FSolverReal CGTol = 1e-4f;
+
+	FAutoConsoleVariableRef CVarClothCGTol(TEXT("p.Chaos.Cloth.CGTol"), MaxItCG, TEXT("CG Tolerance [def: 1e-4]"));
+
 	template <typename T, typename ParticleType>
 	void FGaussSeidelMasterConstraint<T, ParticleType>::AddStaticConstraints(const TArray<TArray<int32>>& ExtraConstraints, TArray<TArray<int32>>& ExtraIncidentElements, TArray<TArray<int32>>& ExtraIncidentElementsLocal)
 	{	
@@ -69,13 +78,152 @@ namespace Chaos::Softs
 	
 	}
 
+	template <typename T, typename ParticleType>
+	TArray<TVec3<T>> FGaussSeidelMasterConstraint<T, ParticleType>::ComputeNewtonResiduals(const ParticleType& Particles, const T Dt, const bool Write2File, TArray<PMatrix<T, 3, 3>>* AllParticleHessian)
+	{
+		TRACE_CPUPROFILER_EVENT_SCOPE(STAT_ChaosGaussSeidelComputeNewtonResidual);
+
+		TArray<TVec3<T>> NewtonResidual;
+		NewtonResidual.Init(TVec3<T>((T)0.), xtilde.Num());
+
+		if (AllParticleHessian)
+		{
+			AllParticleHessian->Init(PMatrix<T, 3, 3>((T)0.), Particles.Size());
+		}
+
+		for (int32 k = 0; k < ParticlesPerColor.Num(); k++)
+		{
+			PhysicsParallelFor(ParticlesPerColor[k].Num(), [this, &AllParticleHessian, Dt, k, &NewtonResidual, &Particles](const int32 j)
+				{
+					const int32 p = ParticlesPerColor[k][j];
+
+					if (Particles.InvM(p) != (T)0.)
+					{
+						int32 ConstraintIndex = 0;
+						Chaos::PMatrix<T, 3, 3> ParticleHessian((T)0., (T)0., (T)0.);
+
+						this->ComputeInitialResidualAndHessian(Particles, p, Dt, NewtonResidual[p], ParticleHessian);
+
+						for (int32 i = 0; i < StaticIncidentElements[p].Num(); i++)
+						{
+							while (StaticIncidentElements[p][i] >= StaticIncidentElementsOffsets[ConstraintIndex + 1] && ConstraintIndex < StaticIncidentElementsOffsets.Num() - 1)
+							{
+								ConstraintIndex += 1;
+							}
+
+							this->AddStaticConstraintResidualAndHessian[ConstraintIndex](Particles, StaticIncidentElements[p][i] - StaticIncidentElementsOffsets[ConstraintIndex], StaticIncidentElementsLocal[p][i], Dt, NewtonResidual[p], ParticleHessian);
+						}
+
+						ConstraintIndex = 0;
+
+						for (int32 i = 0; i < DynamicIncidentElements[p].Num(); i++)
+						{
+							while (DynamicIncidentElements[p][i] >= DynamicIncidentElementsOffsets[ConstraintIndex + 1] && ConstraintIndex < DynamicIncidentElementsOffsets.Num() - 1)
+							{
+								ConstraintIndex += 1;
+							}
+
+							this->AddDynamicConstraintResidualAndHessian[ConstraintIndex](Particles, DynamicIncidentElements[p][i] - DynamicIncidentElementsOffsets[ConstraintIndex], DynamicIncidentElementsLocal[p][i], Dt, NewtonResidual[p], ParticleHessian);
+						}
+
+						if (AllParticleHessian)
+						{
+							(*AllParticleHessian)[p] = ParticleHessian;
+						}
+					}
+			}, ParticlesPerColor[k].Num() < 1000);
+		}
+
+		T NewtonNorm = (T)0.;
+		for (int32 p = 0; p < NewtonResidual.Num(); p++)
+		{
+			NewtonNorm += NewtonResidual[p].SizeSquared();
+		}
+
+		NewtonNorm = FMath::Sqrt(NewtonNorm);
+
+		UE_LOG(LogTemp, Warning, TEXT("Current Iteration: %d"), CurrentIt);
+
+		UE_LOG(LogTemp, Warning, TEXT("Newton Residual is %f"), NewtonNorm);
+
+		PassedIters++;
+
+
+		if (Write2File)
+		{
+			FString file = FPaths::ProjectDir();
+			file.Append(TEXT("/DebugOutput/NewtonResidual.txt"));
+			if (PassedIters == 0)
+			{
+				FFileHelper::SaveStringToFile(FString(TEXT("Newton Norm\r\n")), *file);
+			}
+			FFileHelper::SaveStringToFile((FString::SanitizeFloat(NewtonNorm) + FString(TEXT(",\r\n"))), *file, FFileHelper::EEncodingOptions::AutoDetect, &IFileManager::Get(), EFileWrite::FILEWRITE_Append);
+		}
+
+		return NewtonResidual;
+
+	}
+
+	template <typename T, typename ParticleType>
+	void FGaussSeidelMasterConstraint<T, ParticleType>::ApplyCG(ParticleType& Particles, const T Dt)
+	{
+		TRACE_CPUPROFILER_EVENT_SCOPE(STAT_ChaosGaussSeidelApplyCG);
+		TFunction<void(TArray<FSolverVec3>&)> ProjectBCs = [&Particles](TArray<FSolverVec3>& y)
+		{
+			for (int32 i = 0; i < y.Num(); i++)
+			{
+				if (Particles.InvM(i) == 0.f)
+				{
+					y[i] = FSolverVec3(0.f);
+				}
+			}
+		};
+
+		int32 max_it_cg = MaxItCG;
+		FSolverReal cg_tol = CGTol;
+		int32 MaxNewtonIt = 1;
+
+		for (int32 It = 0; It < MaxNewtonIt; It++)
+		{
+			const TArray<FSolverVec3> Residual = ComputeNewtonResiduals(Particles, Dt, false, nullptr);
+
+			auto multiply = [this, &ProjectBCs, &Dt, &Particles](TArray<FSolverVec3>& y, const TArray<FSolverVec3>& x) 
+			{
+				TArray<FSolverVec3> XProj = x;
+				ProjectBCs(XProj);
+
+				y.Init(FSolverVec3((FSolverReal)0.), XProj.Num());
+
+				for (int32 i = 0; i < this->AddInternalForceDifferentials.Num(); i++)
+				{
+					this->AddInternalForceDifferentials[i](Particles, XProj, y);
+				}
+
+				PhysicsParallelFor(y.Num(), 
+					[&y, &Particles, &XProj, Dt](const int32 i) 
+					{
+						y[i] = Particles.M(i) * XProj[i] + Dt * Dt * y[i];
+					});
+
+				ProjectBCs(y);
+			};
+
+			TArray<FSolverVec3> Deltax;
+			Deltax.Init(FSolverVec3(0.f), Particles.Size());
+
+			Chaos::LanczosCG<FSolverReal>(multiply, Deltax, Residual, MaxItCG, CGTol, use_list.Get()); 
+
+			PhysicsParallelFor(Particles.Size(), [&Particles, &Deltax](const int32 i) 
+				{
+					Particles.P(i) -= Deltax[i];
+				});
+		}
+	}
 }
 
 
-template CHAOS_API void Chaos::Softs::FGaussSeidelMasterConstraint<Chaos::FRealSingle, Chaos::Softs::FSolverParticles>::AddStaticConstraints(const TArray<TArray<int32>>& ExtraConstraints, TArray<TArray<int32>>& ExtraIncidentElements, TArray<TArray<int32>>& ExtraIncidentElementsLocal);
+template class Chaos::Softs::FGaussSeidelMasterConstraint<Chaos::Softs::FSolverReal, Chaos::Softs::FSolverParticles>;
 
 template CHAOS_API void Chaos::Softs::FGaussSeidelMasterConstraint<Chaos::FRealDouble, Chaos::TDynamicParticles<Chaos::FRealDouble, 3>>::AddStaticConstraints(const TArray<TArray<int32>>& ExtraConstraints, TArray<TArray<int32>>& ExtraIncidentElements, TArray<TArray<int32>>& ExtraIncidentElementsLocal);
 
-template CHAOS_API void Chaos::Softs::FGaussSeidelMasterConstraint<Chaos::FRealSingle, Chaos::Softs::FSolverParticles>::AddDynamicConstraints(const TArray<TArray<int32>>& ExtraConstraints, TArray<TArray<int32>>& ExtraIncidentElements, TArray<TArray<int32>>& ExtraIncidentElementsLocal, bool CheckIncidentElements);
-
 template CHAOS_API void Chaos::Softs::FGaussSeidelMasterConstraint<Chaos::FRealDouble, Chaos::TDynamicParticles<Chaos::FRealDouble, 3>>::AddDynamicConstraints(const TArray<TArray<int32>>& ExtraConstraints, TArray<TArray<int32>>& ExtraIncidentElements, TArray<TArray<int32>>& ExtraIncidentElementsLocal, bool CheckIncidentElements);
diff --git a/Engine/Source/Runtime/Experimental/Chaos/Private/Chaos/XPBDAnisotropicBendingConstraints.cpp b/Engine/Source/Runtime/Experimental/Chaos/Private/Chaos/XPBDAnisotropicBendingConstraints.cpp
--- a/Engine/Source/Runtime/Experimental/Chaos/Private/Chaos/XPBDAnisotropicBendingConstraints.cpp
+++ b/Engine/Source/Runtime/Experimental/Chaos/Private/Chaos/XPBDAnisotropicBendingConstraints.cpp
@@ -471,7 +471,7 @@ void FXPBDAnisotropicBendingConstraints::ApplyHelper(FSolverParticles& Particles
 	const FSolverVec3& WarpWeftBiasBaseMultiplier = WarpWeftBiasBaseMultipliers[ConstraintIndex];
 
 	const FSolverReal BiphasicStiffnessValue = IsBuckled[ConstraintIndex] ?
-		WarpWeftBiasBaseMultiplier.Dot(ExpStiffnessValues) : WarpWeftBiasBaseMultiplier.Dot(ExpBucklingStiffnessValues);
+		WarpWeftBiasBaseMultiplier.Dot(ExpBucklingStiffnessValues): WarpWeftBiasBaseMultiplier.Dot(ExpStiffnessValues);
 
 	if (BiphasicStiffnessValue < MinStiffness || (Particles.InvM(i1) == (FSolverReal)0. && Particles.InvM(i2) == (FSolverReal)0. && Particles.InvM(i3) == (FSolverReal)0. && Particles.InvM(i4) == (FSolverReal)0.))
 	{
@@ -753,4 +753,148 @@ void FXPBDAnisotropicBendingConstraints::Apply(FSolverParticles& Particles, cons
 	}
 }
 
+void FXPBDAnisotropicBendingConstraints::ComputeGradTheta(const FSolverVec3& X0, const FSolverVec3& X1, const FSolverVec3& X2, const FSolverVec3& X3, const int32 Index, FSolverVec3& dThetadx, FSolverReal& Theta) 
+{
+	const FSolverVec3 E21 = X2 - X1;
+	const FSolverReal Norme = E21.Length();
+	const FSolverVec3 E10 = X1 - X0;
+	const FSolverVec3 E20 = X2 - X0;
+	const FSolverVec3 N0 = FSolverVec3::CrossProduct(E20, E10);
+	const FSolverReal SquaredNorm0 = N0.SquaredLength();
+
+	const FSolverVec3 E23 = X2 - X3;
+	const FSolverVec3 E13 = X1 - X3;
+	const FSolverVec3 N1 = FSolverVec3::CrossProduct(E13, E23); 
+	const FSolverReal SquaredNorm1 = N1.SquaredLength();
+
+	const FSolverVec3 N0CrossN1 = FSolverVec3::CrossProduct(N0, N1);
+	Theta = FMath::Atan2(FSolverVec3::DotProduct(N0CrossN1, E21) / Norme, N0.Dot(N1)); 
+
+	switch (Index)
+		{
+			case 0:
+				dThetadx = -(Norme / SquaredNorm0) * N0;
+				break;
+			case 1:
+				dThetadx = ((E20.Dot(E21)) / (Norme * SquaredNorm1)) * N1 + ((E23.Dot(E21)) / (Norme * SquaredNorm0)) * N0;
+				break;
+			case 2:
+				dThetadx = -((E10.Dot(E21)) / (Norme * SquaredNorm1)) * N1 - ((E13.Dot(E21)) / (Norme * SquaredNorm0)) * N0;
+				break;
+			case 3:
+				dThetadx = -(Norme / SquaredNorm1) * N1;
+				break;
+			default:
+				break;
+		}
+}
+
+void FXPBDAnisotropicBendingConstraints::AddAnisotropicBendingResidualAndHessian(const FSolverParticles& Particles, const int32 ConstraintIndex, const int32 ConstraintIndexLocal, const FSolverReal Dt, TVec3<FSolverReal>& ParticleResidual, Chaos::PMatrix<FSolverReal, 3, 3>& ParticleHessian)
+{
+	const FSolverVec3 ExpStiffnessValue((FSolverReal)StiffnessWeft, (FSolverReal)Stiffness, (FSolverReal)StiffnessBias);
+	const FSolverVec3 ExpBucklingValue((FSolverReal)BucklingStiffnessWeft, (FSolverReal)BucklingStiffness, (FSolverReal)BucklingStiffnessBias);
+	const FSolverReal DampingRatioValue = (FSolverReal)DampingRatio;
+	const FSolverVec3& WarpWeftBiasBaseMultiplier = WarpWeftBiasBaseMultipliers[ConstraintIndex];
+
+	const TVec4<int32>& Constraint = Constraints[ConstraintIndex];
+	const int32 i1 = Constraint[0];
+	const int32 i2 = Constraint[1];
+	const int32 i3 = Constraint[2];
+	const int32 i4 = Constraint[3];
+	const FSolverVec3& P1 = Particles.P(Constraint[0]);
+	const FSolverVec3& P2 = Particles.P(Constraint[1]);
+	const FSolverVec3& P3 = Particles.P(Constraint[2]);
+	const FSolverVec3& P4 = Particles.P(Constraint[3]);
+
+
+	FSolverVec3 DThetaDx(0.f);
+	FSolverReal Theta = 0.f;
+
+	//TStaticArray<int32, 4> LocalIndexMap { 2, 1, 0, 3 };{ 2, 1, 0, 3 };
+	constexpr int32 LocalIndexMap[] = { 2, 1, 0, 3 };
+	const int32 ActualConstraintIndexLocal = LocalIndexMap[ConstraintIndexLocal];
+
+	ComputeGradTheta(P3, P2, P1, P4, ActualConstraintIndexLocal, DThetaDx, Theta);
+
+	const FSolverReal CAngle = Theta - RestAngles[ConstraintIndex];
+
+	const FSolverReal BiphasicStiffnessValue = IsBuckled[ConstraintIndex] ?
+		WarpWeftBiasBaseMultiplier.Dot(ExpBucklingValue) : WarpWeftBiasBaseMultiplier.Dot(ExpStiffnessValue);
+
+	ParticleResidual -= Dt * Dt * BiphasicStiffnessValue * CAngle * DThetaDx;
+
+	for (int32 Alpha = 0; Alpha < 3; Alpha++)
+	{
+		ParticleHessian.SetRow(Alpha, ParticleHessian.GetRow(Alpha) + Dt * Dt * BiphasicStiffnessValue * DThetaDx[Alpha] * DThetaDx);
+	}
+}
+
+
+void FXPBDAnisotropicBendingConstraints::AddInternalForceDifferential(const FSolverParticles& InParticles, const TArray<TVector<FSolverReal, 3>>& DeltaParticles, TArray<TVector<FSolverReal, 3>>& ndf)
+{
+	const FSolverVec3 ExpStiffnessValue((FSolverReal)StiffnessWeft, (FSolverReal)Stiffness, (FSolverReal)StiffnessBias);
+	const FSolverVec3 ExpBucklingValue((FSolverReal)BucklingStiffnessWeft, (FSolverReal)BucklingStiffness, (FSolverReal)BucklingStiffnessBias);
+	const FSolverReal DampingRatioValue = (FSolverReal)DampingRatio;
+
+	int32 ParticleStart = 0;
+	int32 ParticleNum = InParticles.Size();
+	ensure(ndf.Num() == InParticles.Size());
+
+	const int32 ConstraintColorNum = ConstraintsPerColorStartIndex.Num() - 1;
+
+	for (int32 ConstraintColorIndex = 0; ConstraintColorIndex < ConstraintColorNum; ++ConstraintColorIndex)
+	{
+		const int32 ColorStart = ConstraintsPerColorStartIndex[ConstraintColorIndex];
+		const int32 ColorSize = ConstraintsPerColorStartIndex[ConstraintColorIndex + 1] - ColorStart;
+		PhysicsParallelFor(ColorSize, [&](const int32 Index)
+			{
+				const int32 ConstraintIndex = ColorStart + Index;
+				const TVec4<int32>& Constraint = Constraints[ConstraintIndex];
+				const int32 i1 = Constraint[0];
+				const int32 i2 = Constraint[1];
+				const int32 i3 = Constraint[2];
+				const int32 i4 = Constraint[3];
+
+				const FSolverVec3& P1 = InParticles.P(Constraint[0]);
+				const FSolverVec3& P2 = InParticles.P(Constraint[1]);
+				const FSolverVec3& P3 = InParticles.P(Constraint[2]);
+				const FSolverVec3& P4 = InParticles.P(Constraint[3]);
+				const FSolverVec3& WarpWeftBiasBaseMultiplier = WarpWeftBiasBaseMultipliers[ConstraintIndex];
+				FSolverReal Theta = 0.f;
+
+				TArray<FSolverVec3> Alldthetadx;
+				Alldthetadx.Init(FSolverVec3(0.f), 4);
+
+				constexpr int32 LocalIndexMap[] = { 2, 1, 0, 3 };
+				for (int32 i = 0; i < 4; i++)
+				{
+					int32 ActualConstraintIndexLocal = LocalIndexMap[i];
+					ComputeGradTheta(P3, P2, P1, P4, ActualConstraintIndexLocal, Alldthetadx[i], Theta);
+				}
+
+				const FSolverReal BiphasicStiffnessValue = IsBuckled[ConstraintIndex] ?
+					WarpWeftBiasBaseMultiplier.Dot(ExpBucklingValue) : WarpWeftBiasBaseMultiplier.Dot(ExpStiffnessValue);
+
+				FSolverReal DeltaC = 0.f;
+
+				for (int32 i = 0; i < 4; i++)
+				{
+					for (int32 j = 0; j < 3; j++)
+					{
+						DeltaC += Alldthetadx[i][j] * DeltaParticles[Constraint[i]][j];
+					}
+				}
+
+				for (int32 i = 0; i < 4; i++)
+				{
+					ndf[Constraint[i]] += BiphasicStiffnessValue * DeltaC * Alldthetadx[i];
+				}
+
+		});
+	}
+}
+
+
+
+
 }  // End namespace Chaos::Softs
diff --git a/Engine/Source/Runtime/Experimental/Chaos/Private/Chaos/XPBDBendingConstraints.cpp b/Engine/Source/Runtime/Experimental/Chaos/Private/Chaos/XPBDBendingConstraints.cpp
--- a/Engine/Source/Runtime/Experimental/Chaos/Private/Chaos/XPBDBendingConstraints.cpp
+++ b/Engine/Source/Runtime/Experimental/Chaos/Private/Chaos/XPBDBendingConstraints.cpp
@@ -371,4 +371,203 @@ void FXPBDBendingConstraints::Apply(FSolverParticles& Particles, const FSolverRe
 	}
 }
 
+FSolverReal FXPBDBendingConstraints::ComputeTotalEnergy(const FSolverParticles& InParticles, const FSolverReal ExplicitStiffness)
+{
+	FSolverReal StiffnessValue = (FSolverReal)Stiffness;
+	if (ExplicitStiffness > 0.f)
+	{
+		StiffnessValue = ExplicitStiffness;
+	}
+	const FSolverReal DampingRatioValue = (FSolverReal)DampingRatio;
+
+	auto SafeRecip = [](const FSolverReal Len, const FSolverReal Fallback)
+	{
+		if (Len > UE_SMALL_NUMBER)
+		{
+			return 1.f / Len;
+		}
+		return Fallback;
+	};
+
+	FSolverReal Energy = 0.f;
+
+	for (int32 ConstraintIndex = 0; ConstraintIndex < Constraints.Num(); ConstraintIndex++)
+	{
+		const TVec4<int32>& Constraint = Constraints[ConstraintIndex];
+		const int32 i1 = Constraint[0];
+		const int32 i2 = Constraint[1];
+		const int32 i3 = Constraint[2];
+		const int32 i4 = Constraint[3];
+		//const FSolverVec3& WarpWeftBiasBaseMultiplier = WarpWeftBiasBaseMultipliers[ConstraintIndex];
+		const FSolverVec3& P1 = InParticles.P(Constraint[0]);
+		const FSolverVec3& P2 = InParticles.P(Constraint[1]);
+		const FSolverVec3& P3 = InParticles.P(Constraint[2]);
+		const FSolverVec3& P4 = InParticles.P(Constraint[3]);
+
+		//FSolverVec3 g0(0.f), g1(0.f);
+
+		const FSolverVec3 y1 = P2 - P1;
+		const FSolverVec3 y2 = P3 - P1;
+		const FSolverVec3 y3 = P4 - P1;
+		const FSolverVec3 t1 = y1.GetSafeNormal();
+		const FSolverVec3 t2 = y2.GetSafeNormal();
+		const FSolverVec3 t3 = y3.GetSafeNormal();
+
+		const FSolverVec3 z0 = y2 - FSolverVec3::DotProduct(y2, t1) * t1;
+		const FSolverVec3 z1 = y3 - FSolverVec3::DotProduct(y3, t1) * t1;
+
+		FSolverVec3 z0Normalized = z0.GetSafeNormal();
+		FSolverVec3 z1Normalized = z1.GetSafeNormal();
+
+		FSolverVec3 tb = FSolverVec3::CrossProduct(z0, z1).GetSafeNormal();
+
+		FSolverReal S = FSolverVec3::DotProduct(tb, t1);
+
+		FSolverReal Theta = FMath::Acos(FSolverVec3::DotProduct(z0Normalized, z1Normalized));
+
+		Theta = S * (PI - Theta);
+
+		FSolverReal CAngle = Theta - RestAngles[ConstraintIndex];
+
+
+		Energy += CAngle * CAngle * StiffnessValue / 2.f;
+
+	}
+
+	return Energy;
+
+}
+
+
+void ComputeGradTheta(const FSolverVec3& X0, const FSolverVec3& X1, const FSolverVec3& X2, const FSolverVec3& X3, const int32 Index, FSolverVec3& dThetadx, FSolverReal& Theta) 
+{
+	const FSolverVec3 E21 = X2 - X1;
+	const FSolverReal Norme = E21.Length();
+	const FSolverVec3 E10 = X1 - X0;
+	const FSolverVec3 E20 = X2 - X0;
+	const FSolverVec3 N0 = FSolverVec3::CrossProduct(E20, E10);
+	const FSolverReal SquaredNorm0 = N0.SquaredLength();
+
+	const FSolverVec3 E23 = X2 - X3;
+	const FSolverVec3 E13 = X1 - X3;
+	const FSolverVec3 N1 = FSolverVec3::CrossProduct(E13, E23); 
+	const FSolverReal SquaredNorm1 = N1.SquaredLength();
+
+	const FSolverVec3 N0CrossN1 = FSolverVec3::CrossProduct(N0, N1);
+	Theta = FMath::Atan2(FSolverVec3::DotProduct(N0CrossN1, E21) / Norme, N0.Dot(N1)); 
+
+	switch (Index)
+		{
+			case 0:
+				dThetadx = -(Norme / SquaredNorm0) * N0;
+				break;
+			case 1:
+				dThetadx = ((E20.Dot(E21)) / (Norme * SquaredNorm1)) * N1 + ((E23.Dot(E21)) / (Norme * SquaredNorm0)) * N0;
+				break;
+			case 2:
+				dThetadx = -((E10.Dot(E21)) / (Norme * SquaredNorm1)) * N1 - ((E13.Dot(E21)) / (Norme * SquaredNorm0)) * N0;
+				break;
+			case 3:
+				dThetadx = -(Norme / SquaredNorm1) * N1;
+				break;
+			default:
+				break;
+		}
+}
+
+
+void FXPBDBendingConstraints::AddBendingResidualAndHessian(const FSolverParticles& Particles, const int32 ConstraintIndex, const int32 ConstraintIndexLocal, const FSolverReal Dt, TVec3<FSolverReal>& ParticleResidual, Chaos::PMatrix<FSolverReal, 3, 3>& ParticleHessian)
+{
+	FSolverReal ExpStiffnessValue = (FSolverReal)Stiffness;
+	FSolverReal ExpBucklingValue = (FSolverReal)BucklingStiffness;
+	const FSolverReal DampingRatioValue = (FSolverReal)DampingRatio;
+
+	const TVec4<int32>& Constraint = Constraints[ConstraintIndex];
+	const int32 i1 = Constraint[0];
+	const int32 i2 = Constraint[1];
+	const int32 i3 = Constraint[2];
+	const int32 i4 = Constraint[3];
+	const FSolverVec3& P1 = Particles.P(Constraint[0]);
+	const FSolverVec3& P2 = Particles.P(Constraint[1]);
+	const FSolverVec3& P3 = Particles.P(Constraint[2]);
+	const FSolverVec3& P4 = Particles.P(Constraint[3]);
+
+	FSolverVec3 DThetaDx(0.f);
+	FSolverReal Theta = 0.f;
+
+	constexpr int32 LocalIndexMap[] = { 2, 1, 0, 3 };
+	const int32 ActualConstraintIndexLocal = LocalIndexMap[ConstraintIndexLocal];
+
+	ComputeGradTheta(P3, P2, P1, P4, ActualConstraintIndexLocal, DThetaDx, Theta);
+
+	const FSolverReal CAngle = Theta - RestAngles[ConstraintIndex];
+
+	const FSolverReal BiphasicStiffnessValue = IsBuckled[ConstraintIndex] ? ExpBucklingValue : ExpStiffnessValue;
+
+	ParticleResidual -= Dt * Dt * BiphasicStiffnessValue * CAngle * DThetaDx;
+
+	for (int32 Alpha = 0; Alpha < 3; Alpha++)
+	{
+		ParticleHessian.SetRow(Alpha, ParticleHessian.GetRow(Alpha) + Dt * Dt * BiphasicStiffnessValue * DThetaDx[Alpha] * DThetaDx);
+	}
+}
+
+
+void FXPBDBendingConstraints::AddInternalForceDifferential(const FSolverParticles& InParticles, const TArray<TVector<FSolverReal, 3>>& DeltaParticles, TArray<TVector<FSolverReal, 3>>& ndf)
+{
+	FSolverReal ExpStiffnessValue = (FSolverReal)Stiffness;
+	FSolverReal ExpBucklingValue = (FSolverReal)BucklingStiffness;
+	const FSolverReal DampingRatioValue = (FSolverReal)DampingRatio;
+
+
+	int32 ParticleStart = 0;
+	int32 ParticleNum = InParticles.Size();
+	ensure(ndf.Num() == InParticles.Size());
+
+
+	for (int32 ConstraintIndex = 0; ConstraintIndex < Constraints.Num(); ConstraintIndex++)
+	{
+		const TVec4<int32>& Constraint = Constraints[ConstraintIndex];
+		const int32 i1 = Constraint[0];
+		const int32 i2 = Constraint[1];
+		const int32 i3 = Constraint[2];
+		const int32 i4 = Constraint[3];
+
+		const FSolverVec3& P1 = InParticles.P(Constraint[0]);
+		const FSolverVec3& P2 = InParticles.P(Constraint[1]);
+		const FSolverVec3& P3 = InParticles.P(Constraint[2]);
+		const FSolverVec3& P4 = InParticles.P(Constraint[3]);
+
+		FSolverReal Theta = 0.f;
+
+		TArray<FSolverVec3> Alldthetadx;
+		Alldthetadx.Init(FSolverVec3(0.f), 4);
+
+		TArray<int32> LocalIndexMap = { 2, 1, 0, 3 };
+		for (int32 i = 0; i < 4; i++)
+		{
+			int32 ActualConstraintIndexLocal = LocalIndexMap[i];
+			ComputeGradTheta(P3, P2, P1, P4, ActualConstraintIndexLocal, Alldthetadx[i], Theta);
+		}
+
+		FSolverReal BiphasicStiffnessValue = IsBuckled[ConstraintIndex] ? ExpBucklingValue : ExpStiffnessValue;
+
+		FSolverReal DeltaC = 0.f;
+
+		for (int32 i = 0; i < 4; i++)
+		{
+			for (int32 j = 0; j < 3; j++)
+			{
+				DeltaC += Alldthetadx[i][j] * DeltaParticles[Constraint[i]][j];
+			}
+		}
+
+		for (int32 i = 0; i < 4; i++)
+		{
+			ndf[Constraint[i]] += BiphasicStiffnessValue * DeltaC * Alldthetadx[i];
+		}
+
+	}
+}
+
 }  // End namespace Chaos::Softs
diff --git a/Engine/Source/Runtime/Experimental/Chaos/Private/Chaos/XPBDStretchBiasElementConstraints.cpp b/Engine/Source/Runtime/Experimental/Chaos/Private/Chaos/XPBDStretchBiasElementConstraints.cpp
--- a/Engine/Source/Runtime/Experimental/Chaos/Private/Chaos/XPBDStretchBiasElementConstraints.cpp
+++ b/Engine/Source/Runtime/Experimental/Chaos/Private/Chaos/XPBDStretchBiasElementConstraints.cpp
@@ -17,13 +17,17 @@ bool bChaos_XPBDStretchBiasElement_ISPC_Enabled = true;
 FAutoConsoleVariableRef CVarChaosXPBDStretchBiasISPCEnabled(TEXT("p.Chaos.XPBDStretchBias.ISPC"), bChaos_XPBDStretchBiasElement_ISPC_Enabled, TEXT("Whether to use ISPC optimizations in XPBD Stretch Bias constraints"));
 static_assert(sizeof(ispc::FVector4f) == sizeof(Chaos::Softs::FPAndInvM), "sizeof(ispc::FVector4f) != sizeof(Chaos::Softs::FPAndInvM");
 static_assert(sizeof(ispc::FIntVector) == sizeof(Chaos::TVec3<int32>), "sizeof(ispc::FIntVector) != sizeof(Chaos::TVec3<int32>");
-static_assert(sizeof(ispc::FVector2f) == sizeof(Chaos::Softs::FSolverVec2), "sizeof(ispc::FVector2f) != sizeof(Chaos::Softs::FSolverVec2");
+//static_assert(sizeof(ispc::FVector2f) == sizeof(Chaos::Softs::FSolverVec2), "sizeof(ispc::FVector2f) != sizeof(Chaos::Softs::FSolverVec2");
 static_assert(sizeof(ispc::FVector4f) == sizeof(Chaos::Softs::FSolverMatrix22), "sizeof(ispc::FVector2f) != sizeof(Chaos::Softs::FSolverMatrix22");
 static_assert(sizeof(ispc::FVector3f) == sizeof(Chaos::Softs::FSolverVec3), "sizeof(ispc::FVector3f) != sizeof(Chaos::Softs::FSolverVec3");
 #endif
 
 namespace Chaos::Softs {
 
+	Softs::FSolverReal StiffnessPaddingRatio = 0.f;
+
+	FAutoConsoleVariableRef CVarClothStiffnessPaddingRatio(TEXT("p.Chaos.Cloth.StiffnessPaddingRatio"), StiffnessPaddingRatio, TEXT("stiffness padding for gauss seidel [def: 1]"));
+
 // @todo(chaos): the parallel threshold (or decision to run parallel) should probably be owned by the solver and passed to the constraint container
 static int32 Chaos_XPBDStretchBias_ParallelConstraintCount = 100;
 
@@ -217,12 +221,28 @@ void FXPBDStretchBiasElementConstraints::InitConstraintsAndRestData(const FSolve
 			const FSolverVec3& X1 = InParticles.X(Constraint[1]);
 			const FSolverVec3& X2 = InParticles.X(Constraint[2]);
 			const FSolverVec3 X01 = X1 - X0;
-			const FSolverVec3 X02 = X2 - X0;
-
+			const FSolverVec3 X02 = X2 - X0; 
 			const FSolverVec3 DXDu = X01 * DeltaUVInv.M[0] + X02 * DeltaUVInv.M[1];
 			const FSolverVec3 DXDv = X01 * DeltaUVInv.M[2] + X02 * DeltaUVInv.M[3];
-			const FSolverVec2& RestStretch = RestStretchLengths.Add_GetRef(FSolverVec2(DXDu.Length(), DXDv.Length()));
 
+			auto SafeRecip = [](const FSolverReal Len, const FSolverReal Fallback)
+			{
+				if (Len > UE_SMALL_NUMBER)
+				{
+					return 1.f / Len;
+				}
+				return Fallback;
+			};
+			const FSolverReal OneOverDXDuLen = SafeRecip(DXDu.Length(), 0.f);
+			const FSolverReal OneOverDXDvLen = SafeRecip(DXDv.Length(), 0.f);
+
+			const FSolverVec3 DXDuNormalized = DXDu * OneOverDXDuLen;
+			const FSolverVec3 DXDvNormalized = DXDv * OneOverDXDvLen;
+
+			const FSolverReal ThetaRestStretch = FSolverVec3::DotProduct(DXDuNormalized, DXDvNormalized);
+
+			const FSolverVec3& RestStretch = RestStretchLengths.Add_GetRef(FSolverVec3(DXDu.Length(), DXDv.Length(), ThetaRestStretch));
+			
 			const FSolverReal StiffnessScaleU = RestStretch[0] < UE_SMALL_NUMBER ? 1.f : 1.f / RestStretch[0];
 			const FSolverReal StiffnessScaleV = RestStretch[1] < UE_SMALL_NUMBER ? 1.f : 1.f / RestStretch[1];
 			const FSolverReal BiasScale = FMath::Sqrt((FSolverReal).5 * FSolverVec3::CrossProduct(X01, X02).Length()); // Sqrt of triangle area
@@ -236,7 +256,7 @@ void FXPBDStretchBiasElementConstraints::InitConstraintsAndRestData(const FSolve
 			const FSolverReal DeltaUVDet = DeltaUUnnormalized[0] * DeltaVUnnormalized[1] - DeltaVUnnormalized[0] * DeltaUUnnormalized[1];
 			const FSolverMatrix22& DeltaUVInv = FMath::Abs(DeltaUVDet) < UE_SMALL_NUMBER ? DeltaUVInverse.Add_GetRef(FSolverMatrix22(1.f, 0.f, 0.f, 1.f)) :
 				DeltaUVInverse.Add_GetRef(FSolverMatrix22(DeltaVUnnormalized[1] / DeltaUVDet, -DeltaVUnnormalized[0] / DeltaUVDet, -DeltaUUnnormalized[1] / DeltaUVDet, DeltaUUnnormalized[0] / DeltaUVDet));
-			RestStretchLengths.Add(FSolverVec2(1.f));
+			RestStretchLengths.Add(FSolverVec3(1.f, 1.f, 0.f));
 
 			const FSolverReal DeltaULen = DeltaUUnnormalized.Length();
 			const FSolverReal DeltaVLen = DeltaVUnnormalized.Length();
@@ -262,7 +282,7 @@ void FXPBDStretchBiasElementConstraints::InitColor(const FSolverParticles& InPar
 
 		// Reorder constraints based on color so each array in ConstraintsPerColor contains contiguous elements.
 		TArray<TVec3<int32>> ReorderedConstraints;
-		TArray<FSolverVec2> ReorderedRestStretchLengths;
+		TArray<FSolverVec3> ReorderedRestStretchLengths;
 		TArray<FSolverMatrix22> ReorderedDeltaUVInverse;
 		TArray<FSolverVec3> ReorderedStiffnessScales;
 		TArray<int32> OrigToReorderedIndices; // used to reorder stiffness indices
@@ -478,7 +498,7 @@ void FXPBDStretchBiasElementConstraints::Apply(FSolverParticles& Particles, cons
 							(ispc::FVector4f*)Particles.GetPAndInvM().GetData(),
 							(const ispc::FVector3f*)Particles.X().GetData(),
 							(ispc::FIntVector*)&Constraints.GetData()[ColorStart],
-							(ispc::FVector2f*)&RestStretchLengths.GetData()[ColorStart],
+							(ispc::FVector3f*)&RestStretchLengths.GetData()[ColorStart],
 							(ispc::FVector4f*)&DeltaUVInverse.GetData()[ColorStart],
 							(ispc::FVector3f*)&StiffnessScales.GetData()[ColorStart],
 							(ispc::FVector3f*)&Lambdas.GetData()[ColorStart],
@@ -501,7 +521,7 @@ void FXPBDStretchBiasElementConstraints::Apply(FSolverParticles& Particles, cons
 						ispc::ApplyXPBDStretchBiasConstraints(
 							(ispc::FVector4f*)Particles.GetPAndInvM().GetData(),
 							(ispc::FIntVector*)&Constraints.GetData()[ColorStart],
-							(ispc::FVector2f*)&RestStretchLengths.GetData()[ColorStart],
+							(ispc::FVector3f*)&RestStretchLengths.GetData()[ColorStart],
 							(ispc::FVector4f*)&DeltaUVInverse.GetData()[ColorStart],
 							(ispc::FVector3f*)&StiffnessScales.GetData()[ColorStart],
 							(ispc::FVector3f*)&Lambdas.GetData()[ColorStart],
@@ -545,7 +565,7 @@ void FXPBDStretchBiasElementConstraints::Apply(FSolverParticles& Particles, cons
 							(ispc::FVector4f*)Particles.GetPAndInvM().GetData(),
 							(const ispc::FVector3f*)Particles.X().GetData(),
 							(ispc::FIntVector*)&Constraints.GetData()[ColorStart],
-							(ispc::FVector2f*)&RestStretchLengths.GetData()[ColorStart],
+							(ispc::FVector3f*)&RestStretchLengths.GetData()[ColorStart],
 							(ispc::FVector4f*)&DeltaUVInverse.GetData()[ColorStart],
 							(ispc::FVector3f*)&StiffnessScales.GetData()[ColorStart],
 							(ispc::FVector3f*)&Lambdas.GetData()[ColorStart],
@@ -582,7 +602,7 @@ void FXPBDStretchBiasElementConstraints::Apply(FSolverParticles& Particles, cons
 						ispc::ApplyXPBDStretchBiasConstraintsWithMaps(
 							(ispc::FVector4f*)Particles.GetPAndInvM().GetData(),
 							(ispc::FIntVector*)&Constraints.GetData()[ColorStart],
-							(ispc::FVector2f*)&RestStretchLengths.GetData()[ColorStart],
+							(ispc::FVector3f*)&RestStretchLengths.GetData()[ColorStart],
 							(ispc::FVector4f*)&DeltaUVInverse.GetData()[ColorStart],
 							(ispc::FVector3f*)&StiffnessScales.GetData()[ColorStart],
 							(ispc::FVector3f*)&Lambdas.GetData()[ColorStart],
@@ -686,6 +706,7 @@ static FSolverReal CalcDLambda(const FSolverReal StiffnessValue, const FSolverRe
 	return (C - Alpha * Lambda + DampingTerm) / Denom;
 }
 
+
 void FXPBDStretchBiasElementConstraints::ApplyHelper(FSolverParticles& Particles, const FSolverReal Dt, const int32 ConstraintIndex, const FSolverVec3& ExpStiffnessValue, const FSolverReal DampingRatioValue, const FSolverReal WarpScaleValue, const FSolverReal WeftScaleValue) const
 {
 	const TVec3<int32>& Constraint = Constraints[ConstraintIndex];
@@ -721,8 +742,7 @@ void FXPBDStretchBiasElementConstraints::ApplyHelper(FSolverParticles& Particles
 	// constraints
 	const FSolverReal Cu = (DXDuLength - RestStretchLengths[ConstraintIndex][0]*WeftScaleValue); // stretch in weft direction
 	const FSolverReal Cv = (DXDvLength - RestStretchLengths[ConstraintIndex][1]*WarpScaleValue); // stretch in warp direction
-	const FSolverReal Cs = FSolverVec3::DotProduct(DXDuNormalized, DXDvNormalized);
-
+	const FSolverReal Cs = FSolverVec3::DotProduct(DXDuNormalized, DXDvNormalized) - RestStretchLengths[ConstraintIndex][2];
 
 	const FSolverMatrix22& DeltaUVInv = DeltaUVInverse[ConstraintIndex];
 
@@ -762,4 +782,372 @@ void FXPBDStretchBiasElementConstraints::ApplyHelper(FSolverParticles& Particles
 	Lambda += DLambda;
 }
 
+
+void FXPBDStretchBiasElementConstraints::InitializeDmInvAndMeasures(const FSolverParticles& Particles)
+{
+	Measure.SetNum(Constraints.Num());
+	DmInverse.Init(PMatrix<FSolverReal, 2, 2>(0.f, 0.f, 0.f), Constraints.Num());
+	DmArray.Init(PMatrix<FSolverReal, 2, 2>(0.f, 0.f, 0.f),Constraints.Num());
+	RestDmArray.Init(PMatrix<FSolverReal, 3, 2>(0.f),Constraints.Num());
+	for (int32 e = 0; e < Constraints.Num(); e++)
+	{
+		const TVec3<FSolverReal> X1X0 = Particles.X(Constraints[e][1]) - Particles.X(Constraints[e][0]);
+		const TVec3<FSolverReal> X2X0 = Particles.X(Constraints[e][2]) - Particles.X(Constraints[e][0]);
+		PMatrix<FSolverReal, 3, 2> RestDm(0.f);
+		PMatrix<FSolverReal, 2, 2> Dm(0.f, 0.f, 0.f);
+		Dm.M[0] = X1X0.Size();
+		Dm.M[2] = X1X0.Dot(X2X0) / Dm.M[0];
+		Dm.M[3] = Chaos::TVector<FSolverReal, 3>::CrossProduct(X1X0, X2X0).Size() / Dm.M[0];
+		Measure[e] = Chaos::TVector<FSolverReal, 3>::CrossProduct(X1X0, X2X0).Size() / 2.f;
+		ensureMsgf(Measure[e] > 0.f, TEXT("Degenerate triangle detected"));
+
+		PMatrix<FSolverReal, 2, 2> DmInv = Dm.Inverse();
+		DmInverse[e] = DmInv;
+		DmArray[e] = Dm;
+		RestDmArray[e] = Chaos::PMatrix<FSolverReal, 3, 2>(X1X0, X2X0);
+	}
+
+	bDmInitialized = true;
+}
+
+void FXPBDStretchBiasElementConstraints::AddStretchBiasElementResidualAndHessian(const FSolverParticles& Particles, const int32 ConstraintIndex, const int32 ConstraintIndexLocal, const FSolverReal Dt, TVec3<FSolverReal>& ParticleResidual, Chaos::PMatrix<FSolverReal, 3, 3>& ParticleHessian)
+{
+	const bool StiffnessWarpHasWeightMap = StiffnessWarp.HasWeightMap();
+	const bool StiffnessWeftHasWeightMap = StiffnessWeft.HasWeightMap();
+	const bool StiffnessBiasHasWeightMap = StiffnessBias.HasWeightMap();
+	const bool DampingHasWeightMap = DampingRatio.HasWeightMap();
+	const bool WarpScaleHasWeightMap = WarpScale.HasWeightMap();
+	const bool WeftScaleHasWeightMap = WeftScale.HasWeightMap();
+
+	const FSolverReal StiffnessWarpNoMap = (FSolverReal)StiffnessWarp;
+	const FSolverReal StiffnessWeftNoMap = (FSolverReal)StiffnessWeft;
+	const FSolverReal StiffnessBiasNoMap = (FSolverReal)StiffnessBias;
+	const FSolverReal DampingNoMap = (FSolverReal)DampingRatio;
+	const FSolverReal WarpScaleNoMap = (FSolverReal)WarpScale;
+	const FSolverReal WeftScaleNoMap = (FSolverReal)WeftScale;
+	const FSolverVec3 ExpStiffnessValue(StiffnessWeftNoMap, StiffnessWarpNoMap, StiffnessBiasNoMap);
+
+	const FSolverMatrix22& DeltaUVInv = DeltaUVInverse[ConstraintIndex];
+
+	const TVec3<int32>& Constraint = Constraints[ConstraintIndex];
+	const int32 i0 = Constraint[0];
+	const int32 i1 = Constraint[1];
+	const int32 i2 = Constraint[2];
+	if (Particles.InvM(i0) == (FSolverReal)0. && Particles.InvM(i1) == (FSolverReal)0. && Particles.InvM(i2) == (FSolverReal)0.)
+	{
+		return;
+	}
+
+	FSolverVec3 DXDu, DXDv;
+	CalculateUVStretch(ConstraintIndex, Particles.P(i0), Particles.P(i1), Particles.P(i2), DXDu, DXDv);
+
+	const FSolverReal DXDuLength = DXDu.Length();
+	const FSolverReal DXDvLength = DXDv.Length();
+	auto SafeRecip = [](const FSolverReal Len, const FSolverReal Fallback)
+	{
+		if (Len > UE_SMALL_NUMBER)
+		{
+			return 1.f / Len;
+		}
+		return Fallback;
+	};
+	const FSolverReal OneOverDXDuLen = SafeRecip(DXDuLength, 0.f);
+	const FSolverReal OneOverDXDvLen = SafeRecip(DXDvLength, 0.f);
+	const FSolverReal OneOverDXDuLenCube = OneOverDXDuLen * OneOverDXDuLen * OneOverDXDuLen;
+	const FSolverReal OneOverDXDvLenCube = OneOverDXDvLen * OneOverDXDvLen * OneOverDXDvLen;
+
+	const FSolverVec3 DXDuNormalized = DXDu * OneOverDXDuLen;
+	const FSolverVec3 DXDvNormalized = DXDv * OneOverDXDvLen;
+
+	const FSolverVec3 FinalStiffnesses = StiffnessScales[ConstraintIndex] * ExpStiffnessValue;
+
+	// constraints
+	const FSolverReal Cu = (DXDuLength - RestStretchLengths[ConstraintIndex][0] * WeftScaleNoMap); // stretch in weft direction
+	const FSolverReal Cv = (DXDvLength - RestStretchLengths[ConstraintIndex][1] * WarpScaleNoMap); // stretch in warp direction
+	const FSolverReal Cs = FSolverVec3::DotProduct(DXDuNormalized, DXDvNormalized) - RestStretchLengths[ConstraintIndex][2];
+
+	//original uv constraint coefficients:
+	const FSolverReal Coeffu =  Cu * OneOverDXDuLen * FinalStiffnesses[0] * Dt * Dt;
+	const FSolverReal Coeffv =  Cv * OneOverDXDvLen * FinalStiffnesses[1] * Dt * Dt;
+	const FSolverReal Coeffs =  Cs * FinalStiffnesses[2] * Dt * Dt;
+
+	const FSolverVec3 GradCu1 = DXDuNormalized * DeltaUVInv.M[0];
+	const FSolverVec3 GradCu2 = DXDuNormalized * DeltaUVInv.M[1];
+	const FSolverVec3 GradCu0 = -GradCu1 - GradCu2;
+	const FSolverVec3 GradCv1 = DXDvNormalized * DeltaUVInv.M[2];
+	const FSolverVec3 GradCv2 = DXDvNormalized * DeltaUVInv.M[3];
+	const FSolverVec3 GradCv0 = -GradCv1 - GradCv2;
+
+	const FSolverReal StiffnessScaling = StiffnessPaddingRatio;
+
+	const PMatrix<FSolverReal, 3, 2> DmRest = RestDmArray[ConstraintIndex];
+	const FSolverMatrix22 DmMat = DmArray[ConstraintIndex], DmInvMat = DmInverse[ConstraintIndex];
+	const PMatrix<FSolverReal, 3, 2> RestUVMatrix = DmRest * DeltaUVInv;
+	const FSolverReal RestColumnInnProduct = RestUVMatrix.M[0] * RestUVMatrix.M[3] + RestUVMatrix.M[1] * RestUVMatrix.M[4] + RestUVMatrix.M[2] * RestUVMatrix.M[5];
+	
+
+	const PMatrix<FSolverReal, 2, 2> ActualUVMatrix = DmMat * DeltaUVInv;
+
+	const FSolverReal ColumnInnProduct = ActualUVMatrix.M[0] * ActualUVMatrix.M[2] + ActualUVMatrix.M[1] * ActualUVMatrix.M[3];
+
+	const FSolverVec2 TGSLu = DmMat.TransformPosition(FSolverVec2(DeltaUVInv.M[0], DeltaUVInv.M[1]));
+	const FSolverVec2 TGSLv = DmMat.TransformPosition(FSolverVec2(DeltaUVInv.M[2], DeltaUVInv.M[3]));
+
+	const FSolverReal TGSLuNormSquared = TGSLu.Length() * TGSLu.Length();
+	const FSolverReal TGSLvNormSquared = TGSLv.Length() * TGSLv.Length();
+
+	FSolverReal StiffnessPadding = FinalStiffnesses[0] + FinalStiffnesses[1] + FinalStiffnesses[2];
+	StiffnessPadding *= StiffnessScaling;
+
+	FSolverReal PaddingCoeff = 0.f;
+	if (bDmInitialized)
+	{
+		PaddingCoeff = Dt* Dt* Measure[ConstraintIndex];
+	}
+
+	FSolverVec3 GradCs((FSolverReal)0.), GradCu((FSolverReal)0.), GradCv((FSolverReal)0.);
+	//The following is the simplified version of the gradient:
+	switch (ConstraintIndexLocal)
+	{
+		case 0:
+			GradCs = -(DXDvNormalized * OneOverDXDuLen * DeltaUVInv.M[1] + DXDuNormalized * OneOverDXDvLen * DeltaUVInv.M[3]) - (DXDvNormalized * OneOverDXDuLen * DeltaUVInv.M[0] + DXDuNormalized * OneOverDXDvLen * DeltaUVInv.M[2]);
+			GradCu = -DXDuNormalized * DeltaUVInv.M[0] - DXDuNormalized * DeltaUVInv.M[1];
+			GradCv = -DXDvNormalized * DeltaUVInv.M[2] - DXDvNormalized * DeltaUVInv.M[3];
+			break;
+		case 1:
+			GradCs = DXDvNormalized * OneOverDXDuLen * DeltaUVInv.M[0] + DXDuNormalized * OneOverDXDvLen * DeltaUVInv.M[2];
+			GradCu = DXDuNormalized * DeltaUVInv.M[0];
+			GradCv = DXDvNormalized * DeltaUVInv.M[2];
+			break;
+		case 2:
+			GradCs = DXDvNormalized * OneOverDXDuLen * DeltaUVInv.M[1] + DXDuNormalized * OneOverDXDvLen * DeltaUVInv.M[3];
+			GradCu = DXDuNormalized * DeltaUVInv.M[1];
+			GradCv = DXDvNormalized * DeltaUVInv.M[3];
+			break;
+		default:
+			break;
+	}
+
+	FSolverVec3 ResContribution((FSolverReal)0.);
+
+	//residual part:
+	if (ConstraintIndexLocal > 0)
+	{
+		for (int32 Alpha = 0; Alpha < 3; Alpha++)
+		{
+			//original uv constraints:
+			ResContribution[Alpha] += Coeffu * DXDu[Alpha] * DeltaUVInv.M[ConstraintIndexLocal - 1];
+			ResContribution[Alpha] += Coeffv * DXDv[Alpha] * DeltaUVInv.M[ConstraintIndexLocal + 1];
+
+			//the new ones:
+			ResContribution[Alpha] += Coeffs * GradCs[Alpha];
+		}
+	}
+	else
+	{
+		for (int32 Alpha = 0; Alpha < 3; Alpha++)
+		{
+			ResContribution[Alpha] += Coeffs * GradCs[Alpha];
+			//original uv constraint residuals:
+			for (int32 s = 0; s < 2; s++)
+			{
+				ResContribution[Alpha] -= Coeffu * DXDu[Alpha] * DeltaUVInv.M[s];
+				ResContribution[Alpha] -= Coeffv * DXDv[Alpha] * DeltaUVInv.M[s + 2];
+			}
+		}
+	}
+
+	PMatrix<FSolverReal, 3, 3> HessianContribution((FSolverReal)0., (FSolverReal)0., (FSolverReal)0.);
+
+	const FSolverReal RankOneuCoeff = FinalStiffnesses[0] * Dt * Dt * OneOverDXDuLen * OneOverDXDuLen;
+	const FSolverReal RankOnevCoeff = FinalStiffnesses[1] * Dt * Dt * OneOverDXDvLen * OneOverDXDvLen;
+	const FSolverReal HessianCoeffs = FinalStiffnesses[2] * Dt * Dt;
+	const FSolverReal HessianCoeffu = FinalStiffnesses[0] * Dt * Dt;
+	const FSolverReal HessianCoeffv = FinalStiffnesses[1] * Dt * Dt;
+
+	//hessian part:
+	if (ConstraintIndexLocal > 0)
+	{
+		for (int32 Alpha = 0; Alpha < 3; Alpha++)
+		{
+			for (int32 Beta = 0; Beta < 3; Beta++)
+			{
+				//constraint u:
+				HessianContribution.SetAt(Alpha, Beta, HessianContribution.GetAt(Alpha, Beta) + RankOneuCoeff * DeltaUVInv.M[ConstraintIndexLocal - 1] * DeltaUVInv.M[ConstraintIndexLocal - 1] * DXDu[Alpha] * DXDu[Beta]);
+				
+				//constraint v:
+				HessianContribution.SetAt(Alpha, Beta, HessianContribution.GetAt(Alpha, Beta) + RankOnevCoeff * DeltaUVInv.M[ConstraintIndexLocal + 1] * DeltaUVInv.M[ConstraintIndexLocal + 1] * DXDv[Alpha] * DXDv[Beta]);
+
+			}
+			//constraint s:
+			HessianContribution.SetRow(Alpha, HessianContribution.GetRow(Alpha) + HessianCoeffs * GradCs[Alpha] * GradCs);
+		}
+		//code for stiffness padding:
+		FSolverReal DmInvsum = 0.f;
+		for (int32 nu = 0; nu < 2; nu++)
+		{
+			DmInvsum += DmInvMat.GetAt(ConstraintIndexLocal - 1, nu) * DmInvMat.GetAt(ConstraintIndexLocal - 1, nu);
+		}
+		for (int32 alpha = 0; alpha < 3; alpha++)
+		{
+			HessianContribution.SetAt(alpha, alpha, HessianContribution.GetAt(alpha, alpha) + PaddingCoeff * StiffnessPadding * DmInvsum);
+		}
+	}
+	else
+	{
+		for (int32 Alpha = 0; Alpha < 3; Alpha++)
+		{
+			for (int32 Beta = 0; Beta < 3; Beta++)
+			{
+				for (int32 s = 0; s < 2; s++)
+				{
+					for (int32 n = 0; n < 2; n++)
+					{
+						//constraint u (definitely could be simplified by adding u[s] first):
+						HessianContribution.SetAt(Alpha, Beta, HessianContribution.GetAt(Alpha, Beta) + RankOneuCoeff * DeltaUVInv.M[s] * DeltaUVInv.M[n] * DXDu[Alpha] * DXDu[Beta]);
+						
+						//constraint v (definitely could be simplified by adding v[s] first):
+						HessianContribution.SetAt(Alpha, Beta, HessianContribution.GetAt(Alpha, Beta) + RankOnevCoeff * DeltaUVInv.M[s+2] * DeltaUVInv.M[n+2] * DXDv[Alpha] * DXDv[Beta]);
+					}
+				}
+			}
+			HessianContribution.SetRow(Alpha, HessianContribution.GetRow(Alpha) + HessianCoeffs * GradCs[Alpha] * GradCs);
+		}
+		//code for stiffness padding:
+		FSolverReal DmInvsum = 0.f;
+		for (int32 nu = 0; nu < 2; nu++) 
+		{
+			FSolverReal localDmsum = 0.f;
+			for (int32 k = 0; k < 2; k++) 
+			{
+				localDmsum += DmInvMat.M[nu * 2 + k];
+			}
+			DmInvsum += localDmsum * localDmsum;
+		}
+		for (int32 alpha = 0; alpha < 3; alpha++) 
+		{
+			HessianContribution.SetAt(alpha, alpha, HessianContribution.GetAt(alpha, alpha) + PaddingCoeff * StiffnessPadding * DmInvsum);
+		}
+
+	}
+
+	ParticleResidual += ResContribution;
+	ParticleHessian += HessianContribution;
+}
+
+
+void FXPBDStretchBiasElementConstraints::AddInternalForceDifferential(const FSolverParticles& InParticles, const TArray<TVector<FSolverReal, 3>>& DeltaParticles, TArray<TVector<FSolverReal, 3>>& ndf)
+{
+	TRACE_CPUPROFILER_EVENT_SCOPE(STAT_ChaosCGAnisoBiasElementAddInternalForceDifferential);
+	const FSolverReal StiffnessWarpNoMap = (FSolverReal)StiffnessWarp;
+	const FSolverReal StiffnessWeftNoMap = (FSolverReal)StiffnessWeft;
+	const FSolverReal StiffnessBiasNoMap = (FSolverReal)StiffnessBias;
+	const FSolverReal DampingNoMap = (FSolverReal)DampingRatio;
+	const FSolverReal WarpScaleNoMap = (FSolverReal)WarpScale;
+	const FSolverReal WeftScaleNoMap = (FSolverReal)WeftScale;
+	const FSolverVec3 ExpStiffnessValue(StiffnessWeftNoMap, StiffnessWarpNoMap, StiffnessBiasNoMap);
+
+	int32 ParticleStart = 0; 
+	int32 ParticleNum = InParticles.Size();
+
+	check(ndf.Num() == InParticles.Size());
+
+	const int32 ConstraintColorNum = ConstraintsPerColorStartIndex.Num() - 1;
+
+	for (int32 ConstraintColorIndex = 0; ConstraintColorIndex < ConstraintColorNum; ++ConstraintColorIndex)
+	{
+		const int32 ColorStart = ConstraintsPerColorStartIndex[ConstraintColorIndex];
+		const int32 ColorSize = ConstraintsPerColorStartIndex[ConstraintColorIndex + 1] - ColorStart;
+		PhysicsParallelFor(ColorSize, [&](const int32 Index)
+			{
+				const int32 ConstraintIndex = ColorStart + Index;
+				const TVec3<int32>& Constraint = Constraints[ConstraintIndex];
+				const int32 i0 = Constraint[0];
+				const int32 i1 = Constraint[1];
+				const int32 i2 = Constraint[2];
+
+				if (InParticles.InvM(i0) == (FSolverReal)0. && InParticles.InvM(i1) == (FSolverReal)0. && InParticles.InvM(i2) == (FSolverReal)0.)
+				{
+					return;
+				}
+
+				FSolverVec3 DXDu, DXDv;
+				CalculateUVStretch(ConstraintIndex, InParticles.P(i0), InParticles.P(i1), InParticles.P(i2), DXDu, DXDv);
+
+				const FSolverReal DXDuLength = DXDu.Length();
+				const FSolverReal DXDvLength = DXDv.Length();
+				auto SafeRecip = [](const FSolverReal Len, const FSolverReal Fallback)
+				{
+					if (Len > UE_SMALL_NUMBER)
+					{
+						return 1.f / Len;
+					}
+					return Fallback;
+				};
+				const FSolverReal OneOverDXDuLen = SafeRecip(DXDuLength, 0.f);
+				const FSolverReal OneOverDXDvLen = SafeRecip(DXDvLength, 0.f);
+
+				const FSolverVec3 DXDuNormalized = DXDu * OneOverDXDuLen;
+				const FSolverVec3 DXDvNormalized = DXDv * OneOverDXDvLen;
+
+				// constraints
+				const FSolverReal Cu = (DXDuLength - RestStretchLengths[ConstraintIndex][0] * WeftScaleNoMap); // stretch in weft direction
+				const FSolverReal Cv = (DXDvLength - RestStretchLengths[ConstraintIndex][1] * WarpScaleNoMap); // stretch in warp direction
+				const FSolverReal Cs = FSolverVec3::DotProduct(DXDuNormalized, DXDvNormalized) - RestStretchLengths[ConstraintIndex][2];
+
+
+				const FSolverMatrix22& DeltaUVInv = DeltaUVInverse[ConstraintIndex];
+
+				const FSolverVec3 GradCu1 = DXDuNormalized * DeltaUVInv.M[0];
+				const FSolverVec3 GradCu2 = DXDuNormalized * DeltaUVInv.M[1];
+				const FSolverVec3 GradCu0 = -GradCu1 - GradCu2;
+				const FSolverVec3 GradCv1 = DXDvNormalized * DeltaUVInv.M[2];
+				const FSolverVec3 GradCv2 = DXDvNormalized * DeltaUVInv.M[3];
+				const FSolverVec3 GradCv0 = -GradCv1 - GradCv2;
+				const FSolverVec3 GradCs1 = (DXDvNormalized * OneOverDXDuLen * DeltaUVInv.M[0] + DXDuNormalized * OneOverDXDvLen * DeltaUVInv.M[2]);
+				const FSolverVec3 GradCs2 = (DXDvNormalized * OneOverDXDuLen * DeltaUVInv.M[1] + DXDuNormalized * OneOverDXDvLen * DeltaUVInv.M[3]);
+				const FSolverVec3 GradCs0 = -GradCs1 - GradCs2;
+
+				FSolverReal DeltaCu = 0.f, DeltaCv = 0.f, DeltaCs = 0.f;
+
+				const FSolverVec3 DeltaP01 = DeltaParticles[i1 - ParticleStart] - DeltaParticles[i0 - ParticleStart];
+				const FSolverVec3 DeltaP02 = DeltaParticles[i2 - ParticleStart] - DeltaParticles[i0 - ParticleStart];
+				for (int32 i = 0; i < 3; i++)
+				{
+					DeltaCu += DXDuNormalized[i] * DeltaP01[i] * DeltaUVInv.M[0];
+					DeltaCu += DXDuNormalized[i] * DeltaP02[i] * DeltaUVInv.M[1];
+					DeltaCv += DXDvNormalized[i] * DeltaP01[i] * DeltaUVInv.M[2];
+					DeltaCv += DXDvNormalized[i] * DeltaP02[i] * DeltaUVInv.M[3];
+					DeltaCs += GradCs1[i] * DeltaP01[i];
+					DeltaCs += GradCs2[i] * DeltaP02[i];
+				}
+
+				const FSolverVec3 FinalStiffnesses = StiffnessScales[ConstraintIndex] * ExpStiffnessValue;
+
+				FSolverVec3 DeltaDXDu, DeltaDXDv;
+				CalculateUVStretch(ConstraintIndex, DeltaParticles[i0 - ParticleStart], DeltaParticles[i1 - ParticleStart], DeltaParticles[i2 - ParticleStart], DeltaDXDu, DeltaDXDv);
+				const FSolverVec3 DeltaDXDuOverDxDuNorm = DeltaDXDu * OneOverDXDuLen;
+				const FSolverVec3 DeltaDXDvOverDxDvNorm = DeltaDXDv * OneOverDXDvLen;
+				FSolverVec3 DeltaGradCu0 = DeltaDXDuOverDxDuNorm * DeltaUVInv.M[0];
+				FSolverVec3 DeltaGradCu1 = DeltaDXDuOverDxDuNorm * DeltaUVInv.M[1];
+				FSolverVec3 DeltaGradCu2 = -DeltaGradCu0 - DeltaGradCu1;
+				FSolverVec3 DeltaGradCv0 = DeltaDXDvOverDxDvNorm * DeltaUVInv.M[2];
+				FSolverVec3 DeltaGradCv1 = DeltaDXDvOverDxDvNorm * DeltaUVInv.M[3];
+				FSolverVec3 DeltaGradCv2 = -DeltaGradCv0 - DeltaGradCv1;
+
+
+				ndf[i0] += FinalStiffnesses[0] * DeltaCu * GradCu0 + FinalStiffnesses[0] * Cu * DeltaGradCu0;
+				ndf[i0] += FinalStiffnesses[1] * DeltaCv * GradCv0 + FinalStiffnesses[0] * Cu * DeltaGradCu1;
+				ndf[i0] += FinalStiffnesses[2] * DeltaCs * GradCs0 + FinalStiffnesses[0] * Cu * DeltaGradCu2;
+				ndf[i1] += FinalStiffnesses[0] * DeltaCu * GradCu1 + FinalStiffnesses[1] * Cv * DeltaGradCv0;
+				ndf[i1] += FinalStiffnesses[1] * DeltaCv * GradCv1 + FinalStiffnesses[1] * Cv * DeltaGradCv1;
+				ndf[i1] += FinalStiffnesses[2] * DeltaCs * GradCs1 + FinalStiffnesses[1] * Cv * DeltaGradCv2;
+				ndf[i2] += FinalStiffnesses[0] * DeltaCu * GradCu2;
+				ndf[i2] += FinalStiffnesses[1] * DeltaCv * GradCv2;
+				ndf[i2] += FinalStiffnesses[2] * DeltaCs * GradCs2;
+			});
+	}
+}
+
+
 }  // End namespace Chaos::Softs
diff --git a/Engine/Source/Runtime/Experimental/Chaos/Private/Chaos/XPBDStretchBiasElementConstraints.ispc b/Engine/Source/Runtime/Experimental/Chaos/Private/Chaos/XPBDStretchBiasElementConstraints.ispc
--- a/Engine/Source/Runtime/Experimental/Chaos/Private/Chaos/XPBDStretchBiasElementConstraints.ispc
+++ b/Engine/Source/Runtime/Experimental/Chaos/Private/Chaos/XPBDStretchBiasElementConstraints.ispc
@@ -11,7 +11,7 @@ static inline varying float SafeDivide(const varying float Numerator, const vary
 }
 
 static void CalcConstraintsAndGrads(const FVector3f& P0, const FVector3f& P1, const FVector3f& P2, const FVector4f& DeltaUVInv,
-	const FVector2f& RestStretchLengths, const float WarpScale, const float WeftScale,
+	const FVector3f& RestStretchLengths, const float WarpScale, const float WeftScale,
 	float& Cu, float& Cv, float& Cs, 
 	FVector3f& GradCu1, FVector3f& GradCu2, FVector3f& GradCv1, FVector3f& GradCv2, FVector3f& GradCs1, FVector3f& GradCs2)
 {
@@ -31,7 +31,7 @@ static void CalcConstraintsAndGrads(const FVector3f& P0, const FVector3f& P1, co
 
 	Cu = DXDuLength - RestStretchLengths.V[0] * WeftScale;
 	Cv = DXDvLength - RestStretchLengths.V[1] * WarpScale;
-	Cs = VectorDot(DXDuNormalized, DXDvNormalized);
+	Cs = VectorDot(DXDuNormalized, DXDvNormalized) - RestStretchLengths.V[2];
 
 	GradCu1 = DXDuNormalized * DeltaUVInv.V[0];
 	GradCu2 = DXDuNormalized * DeltaUVInv.V[1];
@@ -71,7 +71,7 @@ static inline float CalcDLambda(const float StiffnessValue, const float Dt, cons
 
 export void ApplyXPBDStretchBiasConstraints(uniform FVector4f PandInvM[],
 									const uniform FIntVector Constraints[],
-									const uniform FVector2f RestStretchLengths[],
+									const uniform FVector3f RestStretchLengths[],
 									const uniform FVector4f DeltaUVInverses[],
 									const uniform FVector3f StiffnessScales[],
 									uniform FVector3f Lambdas[],
@@ -103,7 +103,7 @@ export void ApplyXPBDStretchBiasConstraints(uniform FVector4f PandInvM[],
 		UnzipPandInvM(PandInvM2, P2, M2);
 
 		const varying FVector4f DeltaUVInverse = VectorLoad(&DeltaUVInverses[extract(i,0)]);
-		const varying FVector2f RestStretchLength = VectorLoad(&RestStretchLengths[extract(i,0)]);
+		const varying FVector3f RestStretchLength = VectorLoad(&RestStretchLengths[extract(i,0)]);
 
 		varying float Cu, Cv, Cs;
 		varying FVector3f GradCu1, GradCu2, GradCv1, GradCv2, GradCs1, GradCs2;
@@ -143,7 +143,7 @@ export void ApplyXPBDStretchBiasConstraints(uniform FVector4f PandInvM[],
 export void ApplyXPBDStretchBiasConstraintsWithDamping(uniform FVector4f PandInvM[],
 									const uniform FVector3f X[],
 									const uniform FIntVector Constraints[],
-									const uniform FVector2f RestStretchLengths[],
+									const uniform FVector3f RestStretchLengths[],
 									const uniform FVector4f DeltaUVInverses[],
 									const uniform FVector3f StiffnessScales[],
 									uniform FVector3f Lambdas[],
@@ -176,7 +176,7 @@ export void ApplyXPBDStretchBiasConstraintsWithDamping(uniform FVector4f PandInv
 		UnzipPandInvM(PandInvM2, P2, M2);
 
 		const varying FVector4f DeltaUVInverse = VectorLoad(&DeltaUVInverses[extract(i,0)]);
-		const varying FVector2f RestStretchLength = VectorLoad(&RestStretchLengths[extract(i,0)]);
+		const varying FVector3f RestStretchLength = VectorLoad(&RestStretchLengths[extract(i,0)]);
 
 		varying float Cu, Cv, Cs;
 		varying FVector3f GradCu1, GradCu2, GradCv1, GradCv2, GradCs1, GradCs2;
@@ -228,7 +228,7 @@ export void ApplyXPBDStretchBiasConstraintsWithDamping(uniform FVector4f PandInv
 
 export void ApplyXPBDStretchBiasConstraintsWithMaps(uniform FVector4f PandInvM[],
 											const uniform FIntVector Constraints[],
-											const uniform FVector2f RestStretchLengths[],
+											const uniform FVector3f RestStretchLengths[],
 											const uniform FVector4f DeltaUVInverses[],
 											const uniform FVector3f StiffnessScales[],
 											uniform FVector3f Lambdas[],
@@ -277,7 +277,7 @@ export void ApplyXPBDStretchBiasConstraintsWithMaps(uniform FVector4f PandInvM[]
 		UnzipPandInvM(PandInvM2, P2, M2);
 
 		const varying FVector4f DeltaUVInverse = VectorLoad(&DeltaUVInverses[extract(i,0)]);
-		const varying FVector2f RestStretchLength = VectorLoad(&RestStretchLengths[extract(i,0)]);
+		const varying FVector3f RestStretchLength = VectorLoad(&RestStretchLengths[extract(i,0)]);
 		
 		#pragma ignore warning(perf)
 		const varying float WarpScale = WarpScaleTable[WarpScaleIndex];
@@ -330,7 +330,7 @@ export void ApplyXPBDStretchBiasConstraintsWithMaps(uniform FVector4f PandInvM[]
 export void ApplyXPBDStretchBiasConstraintsWithDampingAndMaps(uniform FVector4f PandInvM[],
 											const uniform FVector3f X[],
 											const uniform FIntVector Constraints[],
-											const uniform FVector2f RestStretchLengths[],
+											const uniform FVector3f RestStretchLengths[],
 											const uniform FVector4f DeltaUVInverses[],
 											const uniform FVector3f StiffnessScales[],
 											uniform FVector3f Lambdas[],
@@ -383,7 +383,7 @@ export void ApplyXPBDStretchBiasConstraintsWithDampingAndMaps(uniform FVector4f
 		UnzipPandInvM(PandInvM2, P2, M2);
 
 		const varying FVector4f DeltaUVInverse = VectorLoad(&DeltaUVInverses[extract(i,0)]);
-		const varying FVector2f RestStretchLength = VectorLoad(&RestStretchLengths[extract(i,0)]);		
+		const varying FVector3f RestStretchLength = VectorLoad(&RestStretchLengths[extract(i,0)]);		
 		
 		#pragma ignore warning(perf)
 		const varying float WarpScale = WarpScaleTable[WarpScaleIndex];
diff --git a/Engine/Source/Runtime/Experimental/Chaos/Public/Chaos/Deformable/GaussSeidelCorotatedCodimensionalConstraints.h b/Engine/Source/Runtime/Experimental/Chaos/Public/Chaos/Deformable/GaussSeidelCorotatedCodimensionalConstraints.h
new file mode 100644
--- /dev/null
+++ b/Engine/Source/Runtime/Experimental/Chaos/Public/Chaos/Deformable/GaussSeidelCorotatedCodimensionalConstraints.h
@@ -0,0 +1,362 @@
+// Copyright Epic Games, Inc. All Rights Reserved.
+#pragma once
+
+#include "Chaos/PBDSpringConstraintsBase.h"
+#include "ChaosStats.h"
+#include "Chaos/ImplicitQRSVD.h"
+#include "Chaos/GraphColoring.h"
+#include "Chaos/Framework/Parallel.h"
+#include "Chaos/Deformable/ChaosDeformableSolverTypes.h"
+#include "Chaos/Vector.h"
+
+
+namespace Chaos::Softs
+{
+	template <typename T, typename ParticleType>
+	class FGaussSeidelCorotatedCodimensionalConstraints
+	{
+
+	public:
+		//this one only accepts tetmesh input and mesh
+		FGaussSeidelCorotatedCodimensionalConstraints(
+			const ParticleType& InParticles,
+			const TArray<TVector<int32, 3>>& InMesh,
+			const bool bRecordMetricIn = true,
+			const T& EMesh = (T)10.0,
+			const T& NuMesh = (T).3
+		)
+			: bRecordMetric(bRecordMetricIn), MeshConstraints(InMesh)
+		{
+			Measure.Init((T)0., MeshConstraints.Num());
+			Lambda = EMesh * NuMesh / (((T)1. + NuMesh) * ((T)1. - (T)2. * NuMesh));
+			Mu = EMesh / ((T)2. * ((T)1. + NuMesh));
+
+			MuElementArray.Init(Mu, MeshConstraints.Num());
+			LambdaElementArray.Init(Lambda, MeshConstraints.Num());
+
+			InitializeCodimensionData(InParticles);
+		}
+
+		virtual ~FGaussSeidelCorotatedCodimensionalConstraints() {}
+
+		PMatrix<T, 3, 3> DsInit(const int32 E, const ParticleType& InParticles)  
+		{
+			PMatrix<T, 3, 3> Result((T)0.);
+			for (int32 i = 0; i < 3; i++) 
+			{
+				for (int32 c = 0; c < 3; c++) 
+				{
+					Result.SetAt(c, i, InParticles.X(MeshConstraints[E][i + 1])[c] - InParticles.X(MeshConstraints[E][0])[c]);
+				}
+			}
+			return Result;
+		}
+
+		PMatrix<T, 3, 2> Ds(const int32 E, const ParticleType& InParticles) const 
+		{
+			const TVec3<T> P1P0 = InParticles.P(MeshConstraints[E][1]) - InParticles.P(MeshConstraints[E][0]);
+			const TVec3<T> P2P0 = InParticles.P(MeshConstraints[E][2]) - InParticles.P(MeshConstraints[E][0]);
+
+			return PMatrix<T, 3, 2>(
+				P1P0[0], P1P0[1], P1P0[2],
+				P2P0[0], P2P0[1], P2P0[2]);
+		}
+
+		PMatrix<T, 3, 2> F(const int32 E, const ParticleType& InParticles) const 
+		{
+			return Ds(E, InParticles) * DmInverse[E];
+		}
+
+		TArray<TArray<int32>> GetConstraintsArray() const
+		{
+			TArray<TArray<int32>> Constraints;
+			Constraints.SetNum(MeshConstraints.Num());
+			for (int32 i = 0; i < MeshConstraints.Num(); i++)
+			{
+				Constraints[i].SetNum(3);
+				for (int32 j = 0; j < 3; j++)
+				{
+					Constraints[i][j] = MeshConstraints[i][j];
+				}
+			}
+			return Constraints;
+		}
+
+		void ComputeNodalMass(const T InDensity, const int32 NumParticles, TArray<T>& NodalMass) const
+		{
+			NodalMass.Init((T)0., NumParticles);
+			for (int32 e = 0; e < Measure.Num(); e++)
+			{
+				const T TotalMass = InDensity * Measure[e];
+				for (int32 j = 0; j < 3; j++)
+				{
+					NodalMass[MeshConstraints[e][j]] += TotalMass / (T)3.;
+				}
+			}
+		}
+
+
+
+	protected:
+
+		void InitializeCodimensionData(const ParticleType& Particles)
+		{
+			Measure.Init((T)0., MeshConstraints.Num());
+			DmInverse.Init(PMatrix<FSolverReal, 2, 2>(0.f, 0.f, 0.f), MeshConstraints.Num());
+			for (int32 e = 0; e < MeshConstraints.Num(); e++)
+			{
+				const TVec3<T> X1X0 = Particles.X(MeshConstraints[e][1]) - Particles.X(MeshConstraints[e][0]);
+				const TVec3<T> X2X0 = Particles.X(MeshConstraints[e][2]) - Particles.X(MeshConstraints[e][0]);
+				PMatrix<T, 2, 2> Dm((T)0., (T)0., (T)0.);
+				Dm.M[0] = X1X0.Size();
+				Dm.M[2] = X1X0.Dot(X2X0) / Dm.M[0];
+				Dm.M[3] = Chaos::TVector<T, 3>::CrossProduct(X1X0, X2X0).Size() / Dm.M[0];
+				Measure[e] = Chaos::TVector<T, 3>::CrossProduct(X1X0, X2X0).Size() / (T)2.;
+				ensureMsgf(Measure[e] > (T)0., TEXT("Degenerate triangle detected"));
+
+				PMatrix<T, 2, 2> DmInv = Dm.Inverse();
+				DmInverse[e] = DmInv;
+			}
+		}
+
+		static T SafeRecip(const T Len, const T Fallback)
+		{
+			if (Len > (T)UE_SMALL_NUMBER)
+			{
+				return (T)1. / Len;
+			}
+			return Fallback;
+		}
+
+		static inline void SymSchur2D(const T Aqq, const T App, const T Apq, T& c, T& s)
+		{
+			// A is an n x n matrix
+			// (p, q) is an index pair satisfying 1 <= p < q <= n
+			//
+			// This function computes a (c, s) pair such that
+			//
+			// B = [ c, s]^T [a_pp, a_pq] [ c, s]
+			//     [-s, c]   [a_pq, a_qq] [-s, c]
+			//
+			// is diagonal
+
+			if (Apq != 0) 
+			{
+				T Tau = T(0.5) * (Aqq - App) / Apq;
+				T t;
+				if (Tau >= 0)
+					t = T(1) / (Tau + FMath::Sqrt((T)1. + Tau * Tau));
+				else
+					t = T(-1) / (-Tau + FMath::Sqrt((T)1. + Tau * Tau));
+				c = T(1) / FMath::Sqrt((T)1. + t * t);
+				s = t * c;
+			}
+			else 
+			{
+				c = (T)1.;
+				s = (T)0.;
+			}
+		}
+
+
+		static inline void Jacobi(const PMatrix<T, 2, 2>& B, PMatrix<T, 2, 2>& D, PMatrix<T, 2, 2>& V)
+		{
+			T c, s;
+			SymSchur2D(B.M[3], B.M[0], B.M[1], c, s);
+
+			V.M[0] = c;
+			V.M[1] = -s;
+			V.M[2] = s;
+			V.M[3] = c;
+
+			D = V.GetTransposed() * B * V;
+		}
+
+		static PMatrix<T, 3, 2> ComputeR(const PMatrix<T, 3, 2>& Fe) 
+		{
+			PMatrix<T, 2, 2> FTF = ComputeFTF(Fe), D((T)0., (T)0., (T)0.), U((T)0., (T)0., (T)0.);
+			Jacobi(FTF, D, U);
+			PMatrix<T, 2, 2> SDiag(FMath::Sqrt(D.M[0]), (T)0., (T)0., FMath::Sqrt(D.M[3]));
+			PMatrix<T, 2, 2> S = U * SDiag * U.GetTransposed(), SInv = S.Inverse();;
+			return Fe * S.Inverse();
+		}
+
+		static PMatrix<T, 3, 2> ComputeRSimple(const PMatrix<T, 3, 2>& InputMatrix) 
+		{
+			TVec3<T> Col1(InputMatrix.M[0], InputMatrix.M[1], InputMatrix.M[2]), Col2(InputMatrix.M[3], InputMatrix.M[4], InputMatrix.M[5]), a(T(0)), b(T(0));
+			if (FMath::Abs(Col1[0]) < (T)UE_SMALL_NUMBER && FMath::Abs(Col1[1]) < (T)UE_SMALL_NUMBER)
+			{
+				a[0] = T(1);
+			}
+			else
+			{
+				a[0] = -Col1[1];
+				a[1] = Col1[0];
+				const T OneOveraNorm = SafeRecip(a.Length(), 0.f);
+				a *= OneOveraNorm;
+			}
+			b = TVec3<T>::CrossProduct(a, Col2);
+			const T OneOverbNorm = SafeRecip(b.Length(), 0.f);
+			b *= OneOverbNorm;
+			return Chaos::PMatrix<T, 3, 2>(b[0], b[1], b[2], a[0], a[1], a[2]);
+		}
+
+		//Does Gram Schmidt on a 3*2 matrix and returns an orthogonal one
+		static Chaos::PMatrix<T, 3, 2> GramSchmidt(const Chaos::PMatrix<T, 3, 2>& InputMatrix)
+		{
+			TVec3<T> Col1(InputMatrix.M[0], InputMatrix.M[1], InputMatrix.M[2]), Col2(InputMatrix.M[3], InputMatrix.M[4], InputMatrix.M[5]);
+			const T Col1Norm = Col1.Length();
+			const T OneOverCol1Norm = SafeRecip(Col1Norm, 0.f);
+			TVec3<T> Col1Normalized = Col1 * OneOverCol1Norm;
+			TVec3<T> Col2Orthogonal = Col2 - TVec3<T>::DotProduct(Col1Normalized, Col2) * Col1Normalized;
+			const T Col2OrthogonalNorm = Col2Orthogonal.Length();
+			const T OneOverCol2OrthogonalNorm = SafeRecip(Col2OrthogonalNorm, 0.f);
+			Col2Orthogonal *= OneOverCol2OrthogonalNorm;
+			return Chaos::PMatrix<T, 3, 2>(Col1Normalized[0], Col1Normalized[1], Col1Normalized[2], Col2Orthogonal[0], Col2Orthogonal[1], Col2Orthogonal[2]);
+		}
+
+		static PMatrix<T, 2, 2> ComputeFTF(const PMatrix<T, 3, 2>& InputMatrix)
+		{
+			return PMatrix<T, 2, 2>(InputMatrix.M[0] * InputMatrix.M[0] + InputMatrix.M[1] * InputMatrix.M[1] + InputMatrix.M[2] * InputMatrix.M[2],
+				InputMatrix.M[0] * InputMatrix.M[3] + InputMatrix.M[1] * InputMatrix.M[4] + InputMatrix.M[2] * InputMatrix.M[5],
+				InputMatrix.M[3] * InputMatrix.M[3] + InputMatrix.M[4] * InputMatrix.M[4] + InputMatrix.M[5] * InputMatrix.M[5]);
+		}
+
+		//Computes det(FTF)^{1/2} * F * (FTF)^-1
+		static void dJdF32(const PMatrix<T, 3, 2>& F, PMatrix<T, 3, 2>& dJ)
+		{
+			const PMatrix<T, 2, 2> FTF = ComputeFTF(F);
+			const T J2 = FTF.Determinant();
+
+			const PMatrix<T, 2, 2> J2invT(FTF.M[3], -FTF.M[2], -FTF.M[1], FTF.M[0]);
+
+			if (J2 > T(0))
+			{
+				dJ = (T(1) / FMath::Sqrt(J2)) *  (F * J2invT);
+			}
+			else
+			{
+				dJ = PMatrix<T, 3, 2>(TVec3<T>((T)0.), TVec3<T>((T)0.));
+			}
+		}
+
+	public:
+		void AddHyperelasticResidualAndHessian(const ParticleType& Particles, const int32 ElementIndex, const int32 ElementIndexLocal, const T Dt, TVec3<T>& ParticleResidual, Chaos::PMatrix<T, 3, 3>& ParticleHessian)
+		{
+			const Chaos::PMatrix<T, 2, 2> DmInvT = DmInverse[ElementIndex].GetTransposed(), DmInv = DmInverse[ElementIndex]; 
+			const Chaos::PMatrix<T, 3, 2> Fe = F(ElementIndex, Particles);
+
+			const PMatrix<T, 3, 2> Re = ComputeR(Fe);
+			const PMatrix<T, 2, 2> FTF = ComputeFTF(Fe);
+			const T J = FMath::Sqrt(FTF.Determinant());
+
+			PMatrix<T, 3, 2> Pe(TVec3<T>((T)0.), TVec3<T>((T)0.)), ForceTerm(TVec3<T>((T)0.), TVec3<T>((T)0.));
+			Pe = T(2) * MuElementArray[ElementIndex] * (Fe - Re) + LambdaElementArray[ElementIndex] * (J - T(1)) * J * Fe * ComputeFTF(Fe).Inverse(); 
+
+			ForceTerm = -Measure[ElementIndex] * Pe * DmInverse[ElementIndex].GetTransposed();
+
+			Chaos::TVector<T, 3> Dx((T)0.);
+			if (ElementIndexLocal > 0)
+			{
+				for (int32 c = 0; c < 3; c++)
+				{
+					Dx[c] += ForceTerm.M[ElementIndexLocal * 3 - 3 + c];
+				}
+			}
+			else
+			{
+				for (int32 c = 0; c < 3; c++)
+				{
+					for (int32 h = 0; h < 2; h++)
+					{
+						Dx[c] -= ForceTerm.M[h * 3 + c];
+					}
+				}
+			}
+
+			Dx *= Dt * Dt;
+
+			ParticleResidual -= Dx;
+
+			PMatrix<T, 3, 2> dJ(TVec3<T>((T)0.), TVec3<T>((T)0.));
+
+			dJdF32(Fe, dJ);
+
+			T Coeff = Dt * Dt * Measure[ElementIndex];
+
+			if (ElementIndexLocal == 0) 
+			{
+				T DmInvsum = T(0);
+				for (int32 nu = 0; nu < 2; nu++) 
+				{
+					T localDmsum = T(0);
+					for (int32 k = 0; k < 2; k++) 
+					{
+						localDmsum += DmInv.M[nu * 2 + k];
+					}
+					DmInvsum += localDmsum * localDmsum;
+				}
+				for (int32 alpha = 0; alpha < 3; alpha++) 
+				{
+					ParticleHessian.SetAt(alpha, alpha, ParticleHessian.GetAt(alpha, alpha) + Coeff * T(2) * MuElementArray[ElementIndex] * DmInvsum);
+				}
+
+				const PMatrix<T, 3, 2> dJDmInvT = dJ * DmInvT; 
+				Chaos::TVector<T, 3> L((T)0.);
+				for (int32 alpha = 0; alpha < 3; alpha++) 
+				{
+					for (int32 k = 0; k < 2; k++) 
+					{
+						L[alpha] += dJDmInvT.M[3 * k + alpha]; 
+					}
+				}
+				for (int32 alpha = 0; alpha < 3; alpha++)
+				{
+					ParticleHessian.SetRow(alpha, ParticleHessian.GetRow(alpha) + Coeff * LambdaElementArray[ElementIndex] * L[alpha] * L);
+				}
+			}
+			else 
+			{
+				T DmInvsum = T(0);
+				for (int32 nu = 0; nu < 2; nu++)
+				{
+					DmInvsum += DmInv.GetAt(ElementIndexLocal - 1, nu) * DmInv.GetAt(ElementIndexLocal - 1, nu);
+				}
+				for (int32 alpha = 0; alpha < 3; alpha++)
+				{
+					ParticleHessian.SetAt(alpha, alpha, ParticleHessian.GetAt(alpha, alpha) + Dt * Dt * Measure[ElementIndex] * T(2) * MuElementArray[ElementIndex] * DmInvsum);
+				}
+
+				const PMatrix<T, 3, 2> dJDmInvT = dJ * DmInvT; 
+
+				Chaos::TVector<T, 3> L((T)0.);
+				for (int32 alpha = 0; alpha < 3; alpha++) 
+				{
+					L[alpha] = dJDmInvT.M[ElementIndexLocal* 3 - 3 + alpha];
+				}
+				for (int32 alpha = 0; alpha < 3; alpha++)
+				{
+					ParticleHessian.SetRow(alpha, ParticleHessian.GetRow(alpha) + Dt * Dt * Measure[ElementIndex] * LambdaElementArray[ElementIndex] * L[alpha] * L);
+				}
+			}
+		}
+
+
+
+	protected:
+		mutable TArray<FSolverMatrix22> DmInverse;
+
+		//material constants calculated from E:
+		T Mu;
+		T Lambda;
+		TArray<T> MuElementArray;
+		TArray<T> LambdaElementArray;
+		TArray<T> AlphaJArray;
+		bool bRecordMetric;
+
+		TArray<TVector<int32, 3>> MeshConstraints;
+		mutable TArray<T> Measure;
+	};
+
+}  // End namespace Chaos::Softs
+
diff --git a/Engine/Source/Runtime/Experimental/Chaos/Public/Chaos/Deformable/GaussSeidelMasterConstraint.h b/Engine/Source/Runtime/Experimental/Chaos/Public/Chaos/Deformable/GaussSeidelMasterConstraint.h
--- a/Engine/Source/Runtime/Experimental/Chaos/Public/Chaos/Deformable/GaussSeidelMasterConstraint.h
+++ b/Engine/Source/Runtime/Experimental/Chaos/Public/Chaos/Deformable/GaussSeidelMasterConstraint.h
@@ -62,6 +62,8 @@ namespace Chaos::Softs
 		TArray<TFunction<void(const ParticleType&, const int32, const int32, const T, TVec3<T>&, Chaos::PMatrix<T, 3, 3>&)>>& DynamicConstraintResidualAndHessian() { return AddDynamicConstraintResidualAndHessian; }
 		const TArray<TFunction<void(const int32, const T, Chaos::PMatrix<T, 3, 3>&)>>& PerNodeHessian() const { return AddPerNodeHessian; }
 		TArray<TFunction<void(const int32, const T, Chaos::PMatrix<T, 3, 3>&)>>& PerNodeHessian() { return AddPerNodeHessian; }
+		const TArray<TFunction<void(const ParticleType&, const TArray<TVec3<T>>&, TArray<TVec3<T>>&)>>& InternalForceDifferentials() const { return AddInternalForceDifferentials; }
+		TArray<TFunction<void(const ParticleType&, const TArray<TVec3<T>>&, TArray<TVec3<T>>&)>>& InternalForceDifferentials() { return AddInternalForceDifferentials; }
 
 		int32 AddStaticConstraintResidualAndHessianRange(int32 NumConstraints)
 		{
@@ -84,14 +86,27 @@ namespace Chaos::Softs
 			return CurrentSize;
 		}
 
-		void AddStaticConstraints(const TArray<TArray<int32>>& ExtraConstraints, TArray<TArray<int32>>& ExtraIncidentElements, TArray<TArray<int32>>& ExtraIncidentElementsLocal);
+		int32 AddAddInternalForceDifferentialsRange(int32 NumConstraints)
+		{
+			int32 CurrentSize = AddInternalForceDifferentials.Num();
+			AddInternalForceDifferentials.AddDefaulted(NumConstraints);
+			return CurrentSize;
+		}
 
-		void AddDynamicConstraints(const TArray<TArray<int32>>& ExtraConstraints, TArray<TArray<int32>>& ExtraIncidentElements, TArray<TArray<int32>>& ExtraIncidentElementsLocal, bool CheckIncidentElements = false);
+		CHAOS_API void AddStaticConstraints(const TArray<TArray<int32>>& ExtraConstraints, TArray<TArray<int32>>& ExtraIncidentElements, TArray<TArray<int32>>& ExtraIncidentElementsLocal);
 
-		void Apply(ParticleType& Particles, const T Dt)
-		{
-			PERF_SCOPE(STAT_ChaosGSMasterConstraint_Apply);
+		CHAOS_API void AddDynamicConstraints(const TArray<TArray<int32>>& ExtraConstraints, TArray<TArray<int32>>& ExtraIncidentElements, TArray<TArray<int32>>& ExtraIncidentElementsLocal, bool CheckIncidentElements = false);
 
+		void Apply(ParticleType& Particles, const T Dt, const int32 MaxWriteIters = 10, const bool Write2File = false)
+		{
+ 			PERF_SCOPE(STAT_ChaosGSMasterConstraint_Apply);
+			
+			if (DebugResidual && PassedIters < MaxWriteIters)
+			{
+				ComputeNewtonResiduals(Particles, Dt, Write2File);
+			}
+			
+			
 			for (int32 i = 0; i < ParticlesPerColor.Num(); i++)
 			{
 				int32 NumBatch = ParticlesPerColor[i].Num() / CorotatedParams.XPBDCorotatedBatchSize;
@@ -99,7 +114,7 @@ namespace Chaos::Softs
 				{
 					NumBatch ++;
 				}
-
+			
 				PhysicsParallelFor(NumBatch, [&](const int32 BatchIndex)
 					{
 						for (int32 BatchSubIndex = 0; BatchSubIndex < CorotatedParams.XPBDCorotatedBatchSize; BatchSubIndex++) {
@@ -111,14 +126,15 @@ namespace Chaos::Softs
 							}
 						}
 					}, NumBatch < CorotatedParams.XPBDCorotatedBatchThreshold);
-
+			
 			}
-
+			
+						
 			if (bDoAcceleration)
 			{
 				AccelerationTechnique(Particles);
 			}
-
+			
 			CurrentIt ++;
 		}
 
@@ -139,6 +155,8 @@ namespace Chaos::Softs
 
 		void Init(const T Dt, const ParticleType& Particles)
 		{
+			Resize((int32)Particles.Size());
+
 			PERF_SCOPE(STAT_ChaosGSMasterConstraint_Init);
 			DynamicConstraints.SetNum(0);
 			for (int32 p = 0; p < DynamicIncidentElements.Num(); p++)
@@ -149,9 +167,10 @@ namespace Chaos::Softs
 			DynamicIncidentElementsOffsets.SetNum(0);
 			if (!bDoQuasistatics)
 			{
-				for (int32 i = 0; i < (int32)Particles.Size(); i++)
+				for (int32 i = 0; i < xtilde.Num(); i++)
 				{
-					xtilde[i] = Particles.X(i) + Dt * Particles.V(i);
+					//xtilde[i] = Particles.X(i) + Dt * Particles.V(i);
+					xtilde[i] = Particles.P(i);
 				}
 			}
 			CurrentIt = 0;
@@ -185,16 +204,19 @@ namespace Chaos::Softs
 			return false;
 		}
 
+		CHAOS_API TArray<TVec3<T>> ComputeNewtonResiduals(const ParticleType& Particles, const T Dt, const bool Write2File = false, TArray<PMatrix<T, 3, 3>>* AllParticleHessian = nullptr);
 		
+		CHAOS_API void ApplyCG(ParticleType& Particles, const T Dt);
+
 	private:
 
 		void ApplySingleParticle(const int32 p, const T Dt, ParticleType& Particles)
 		{
 			int32 ConstraintIndex = 0;
 			Chaos::TVector<T, 3> ParticleResidual((T)0.);
-			Chaos::PMatrix<T, 3, 3> ParticleHessian = Chaos::PMatrix<T, 3, 3>::Zero;
+			Chaos::PMatrix<T, 3, 3> ParticleHessian((T)0., (T)0., (T)0.);
 
-			ComputeInitialResidualAndHessian(Particles, p, ParticleResidual, ParticleHessian);
+			ComputeInitialResidualAndHessian(Particles, p, Dt, ParticleResidual, ParticleHessian);
 
 			for (int32 i = 0; i < StaticIncidentElements[p].Num(); i++)
 			{
@@ -204,7 +226,6 @@ namespace Chaos::Softs
 				}
 
 				AddStaticConstraintResidualAndHessian[ConstraintIndex](Particles, StaticIncidentElements[p][i] - StaticIncidentElementsOffsets[ConstraintIndex], StaticIncidentElementsLocal[p][i], Dt, ParticleResidual, ParticleHessian);
-
 			}
 
 			ConstraintIndex = 0;
@@ -217,7 +238,6 @@ namespace Chaos::Softs
 				}
 
 				AddDynamicConstraintResidualAndHessian[ConstraintIndex](Particles, DynamicIncidentElements[p][i] - DynamicIncidentElementsOffsets[ConstraintIndex], DynamicIncidentElementsLocal[p][i], Dt, ParticleResidual, ParticleHessian);
-
 			}
 
 			for (int32 i = 0; i < AddPerNodeHessian.Num(); i++)
@@ -227,18 +247,19 @@ namespace Chaos::Softs
 
 
 			T HessianDet = ParticleHessian.Determinant();
-			if (HessianDet > UE_SMALL_NUMBER)
+			if (HessianDet > UE_SMALL_NUMBER && ParticleResidual.Size() > UE_SMALL_NUMBER)
 			{
 				Chaos::PMatrix<T, 3, 3> HessianInv = ParticleHessian.SymmetricCofactorMatrix();
 				HessianInv *= T(1) / HessianDet;
 				Chaos::TVector<T, 3> Dx = HessianInv.GetTransposed() * (-ParticleResidual);
+
 				Particles.P(p) += Dx;
 			}
 		}
 		
 		void InitializeLambdas()
 		{
-			ComputeInitialResidualAndHessian = [this](const ParticleType& Particles, const int32 p, TVec3<T>& ParticleResidual, Chaos::PMatrix<T, 3, 3>& ParticleHessian)
+			ComputeInitialResidualAndHessian = [this](const ParticleType& Particles, const int32 p, const T Dt, TVec3<T>& ParticleResidual, Chaos::PMatrix<T, 3, 3>& ParticleHessian)
 			{
 				if (!this->bDoQuasistatics) {
 					for (int32 alpha = 0; alpha < 3; alpha++) {
@@ -248,6 +269,12 @@ namespace Chaos::Softs
 						ParticleHessian.SetAt(alpha, alpha, Particles.M(p));
 					}
 				}
+				else
+				{
+					for (int32 alpha = 0; alpha < 3; alpha++) {
+						ParticleResidual[alpha] = -Dt * Dt * ExternalForce[alpha] * Particles.M(p);
+					}
+				}
 			};
 
 			AccelerationTechnique = [this](ParticleType& Particles)
@@ -265,6 +292,7 @@ namespace Chaos::Softs
 			};
 		}
 
+
 		//Constraints storage:
 		TArray<TArray<int32>> StaticConstraints;
 		TArray<TArray<int32>> StaticIncidentElements;
@@ -274,7 +302,7 @@ namespace Chaos::Softs
 		TArray<TArray<int32>> DynamicIncidentElementsLocal;
 
 		//Lambds for spcifying residual/hessian computations:
-		TFunction<void(const ParticleType&, const int32, TVec3<T>&, Chaos::PMatrix<T, 3, 3>&)> ComputeInitialResidualAndHessian;
+		TFunction<void(const ParticleType&, const int32, const T, TVec3<T>&, Chaos::PMatrix<T, 3, 3>&)> ComputeInitialResidualAndHessian;
 		TArray<TFunction<void(const ParticleType&, const int32, const int32, const T, TVec3<T>&, Chaos::PMatrix<T, 3, 3>&)>> AddStaticConstraintResidualAndHessian;
 		TArray<TFunction<void(const ParticleType&, const int32, const int32, const T, TVec3<T>&, Chaos::PMatrix<T, 3, 3>&)>> AddDynamicConstraintResidualAndHessian;
 		TArray<TFunction<void(const int32, const T, Chaos::PMatrix<T, 3, 3>&)>> AddPerNodeHessian;
@@ -299,11 +327,28 @@ namespace Chaos::Softs
 		int32 CurrentIt = 0;
 		bool bDoAcceleration = true;
 		T OmegaSOR = T(1.6);
-		int32 SORStart = 1;
+		int32 SORStart = 1; //1 is smallest
 
 		int32 ParallelMax = 1000;
 
 		FDeformableXPBDCorotatedParams CorotatedParams;
 
+		//Newton solver variables:
+		TArray<TFunction<void(const ParticleType&, const TArray<TVec3<T>>&, TArray<TVec3<T>>&)>> AddInternalForceDifferentials;
+		TUniquePtr<TArray<int32>> use_list;
+
+		int32 NumTotalParticles;
+		TArray<TVec3<T>> ReorderedPs;
+
+		public:
+		bool DebugResidual = false;
+		bool IsFirstFrame = true;
+		int32 PassedIters = 0;
+
+		TVec3<T> ExternalForce = TVec3<T>((T)0.);
+
 	};
+
+	//template class FGaussSeidelMasterConstraint<FSolverReal, FSolverParticles>;
 }
+
diff --git a/Engine/Source/Runtime/Experimental/Chaos/Public/Chaos/Math/Krylov.h b/Engine/Source/Runtime/Experimental/Chaos/Public/Chaos/Math/Krylov.h
--- a/Engine/Source/Runtime/Experimental/Chaos/Public/Chaos/Math/Krylov.h
+++ b/Engine/Source/Runtime/Experimental/Chaos/Public/Chaos/Math/Krylov.h
@@ -69,7 +69,7 @@ template <class T, typename Func, int32 d=3>
 void LanczosCG(
 	Func multiplyA, 
 	TArray<TVector<T, d>>& x,
-	TArray<TVector<T, d>>& b,
+	const TArray<TVector<T, d>>& b,
 	const int max_it, 
 	const T res = 1e-4, 
 	const TArray<int32>* use_list = nullptr)
diff --git a/Engine/Source/Runtime/Experimental/Chaos/Public/Chaos/PBDEvolution.h b/Engine/Source/Runtime/Experimental/Chaos/Public/Chaos/PBDEvolution.h
--- a/Engine/Source/Runtime/Experimental/Chaos/Public/Chaos/PBDEvolution.h
+++ b/Engine/Source/Runtime/Experimental/Chaos/Public/Chaos/PBDEvolution.h
@@ -140,6 +140,8 @@ class FPBDEvolution : public TArrayCollection
 	const FSolverVec3& GetGravity(const uint32 GroupId = 0) const { check(GroupId < TArrayCollection::Size()); return MGroupGravityAccelerations[GroupId]; }
 	void SetGravity(const FSolverVec3& Acceleration, const uint32 GroupId = 0) { check(GroupId < TArrayCollection::Size()); MGroupGravityAccelerations[GroupId] = Acceleration; }
 
+	void SetQuasistatics(const bool bDoQuasistaticsIn) { bDoQuasistatics = bDoQuasistaticsIn; }
+
 	FVelocityAndPressureField& GetVelocityAndPressureField(const uint32 GroupId = 0) { check(GroupId < TArrayCollection::Size()); return MGroupVelocityAndPressureFields[GroupId]; }
 	const FVelocityAndPressureField& GetVelocityAndPressureField(const uint32 GroupId = 0) const { check(GroupId < TArrayCollection::Size()); return MGroupVelocityAndPressureFields[GroupId]; }
 
diff --git a/Engine/Source/Runtime/Experimental/Chaos/Public/Chaos/XPBDAnisotropicBendingConstraints.h b/Engine/Source/Runtime/Experimental/Chaos/Public/Chaos/XPBDAnisotropicBendingConstraints.h
--- a/Engine/Source/Runtime/Experimental/Chaos/Public/Chaos/XPBDAnisotropicBendingConstraints.h
+++ b/Engine/Source/Runtime/Experimental/Chaos/Public/Chaos/XPBDAnisotropicBendingConstraints.h
@@ -96,6 +96,26 @@ public:
 	const TArray<int32>& GetConstraintsPerColorStartIndex() const { return ConstraintsPerColorStartIndex; }
 	const TArray<FSolverVec3>& GetWarpWeftBiasBaseMultipliers() const { return WarpWeftBiasBaseMultipliers; }
 
+	CHAOS_API void AddAnisotropicBendingResidualAndHessian(const FSolverParticles& Particles, const int32 ConstraintIndex, const int32 ConstraintIndexLocal, const FSolverReal Dt, TVec3<FSolverReal>& ParticleResidual, Chaos::PMatrix<FSolverReal, 3, 3>& ParticleHessian);
+
+	TArray<TArray<int32>> GetConstraintsArray() const
+	{
+		TArray<TArray<int32>> ConstraintsArray;
+		ConstraintsArray.SetNum(Constraints.Num());
+		for (int32 i = 0; i < Constraints.Num(); i++)
+		{
+			ConstraintsArray[i].SetNum(4);
+			for (int32 j = 0; j < 4; j++)
+			{
+				ConstraintsArray[i][j] = Constraints[i][j];
+			}
+		}
+		return ConstraintsArray;
+	}
+
+	CHAOS_API void AddInternalForceDifferential(const FSolverParticles& InParticles, const TArray<TVector<FSolverReal, 3>>& DeltaParticles, TArray<TVector<FSolverReal, 3>>& ndf);
+
+
 private:
 	CHAOS_API void InitColor(const FSolverParticles& InParticles);
 	CHAOS_API void ApplyHelper(FSolverParticles& Particles, const FSolverReal Dt, const int32 ConstraintIndex, const FSolverVec3& ExpStiffnessValues, 
@@ -136,6 +156,7 @@ private:
 			return FSolverVec2(GetWeightedFloatXPBDAnisoRestAngle(PropertyCollection, 0.f));
 		}
 	}
+	void ComputeGradTheta(const FSolverVec3& X0, const FSolverVec3& X1, const FSolverVec3& X2, const FSolverVec3& X3, const int32 Index, FSolverVec3& dThetadx, FSolverReal& Theta); 
 
 	using Base::Constraints;
 	using Base::ParticleOffset;
diff --git a/Engine/Source/Runtime/Experimental/Chaos/Public/Chaos/XPBDBendingConstraints.h b/Engine/Source/Runtime/Experimental/Chaos/Public/Chaos/XPBDBendingConstraints.h
--- a/Engine/Source/Runtime/Experimental/Chaos/Public/Chaos/XPBDBendingConstraints.h
+++ b/Engine/Source/Runtime/Experimental/Chaos/Public/Chaos/XPBDBendingConstraints.h
@@ -227,6 +227,26 @@ public:
 
 	const TArray<int32>& GetConstraintsPerColorStartIndex() const { return ConstraintsPerColorStartIndex; }
 
+	CHAOS_API void AddBendingResidualAndHessian(const FSolverParticles& Particles, const int32 ConstraintIndex, const int32 ConstraintIndexLocal, const FSolverReal Dt, TVec3<FSolverReal>& ParticleResidual, Chaos::PMatrix<FSolverReal, 3, 3>& ParticleHessian);
+
+	TArray<TArray<int32>> GetConstraintsArray() const
+	{
+		TArray<TArray<int32>> ConstraintsArray;
+		ConstraintsArray.SetNum(Constraints.Num());
+		for (int32 i = 0; i < Constraints.Num(); i++)
+		{
+			ConstraintsArray[i].SetNum(4);
+			for (int32 j = 0; j < 4; j++)
+			{
+				ConstraintsArray[i][j] = Constraints[i][j];
+			}
+		}
+		return ConstraintsArray;
+	}
+
+	CHAOS_API void AddInternalForceDifferential(const FSolverParticles& InParticles, const TArray<TVector<FSolverReal, 3>>& DeltaParticles, TArray<TVector<FSolverReal, 3>>& ndf);
+
+	CHAOS_API FSolverReal ComputeTotalEnergy(const FSolverParticles& InParticles, const FSolverReal ExplicitStiffness = -1.f);
 private:
 	CHAOS_API void InitColor(const FSolverParticles& InParticles);
 	CHAOS_API void ApplyHelper(FSolverParticles& Particles, const FSolverReal Dt, const int32 ConstraintIndex, const FSolverReal ExpStiffnessValue, const FSolverReal ExpBucklingValue, const FSolverReal DampingRatioValue) const;
diff --git a/Engine/Source/Runtime/Experimental/Chaos/Public/Chaos/XPBDStretchBiasElementConstraints.h b/Engine/Source/Runtime/Experimental/Chaos/Public/Chaos/XPBDStretchBiasElementConstraints.h
--- a/Engine/Source/Runtime/Experimental/Chaos/Public/Chaos/XPBDStretchBiasElementConstraints.h
+++ b/Engine/Source/Runtime/Experimental/Chaos/Public/Chaos/XPBDStretchBiasElementConstraints.h
@@ -100,7 +100,22 @@ public:
 	CHAOS_API void CalculateUVStretch(const int32 ConstraintIndex, const FSolverVec3& P0, const FSolverVec3& P1, const FSolverVec3& P2, FSolverVec3& DXDu, FSolverVec3& DXDv) const;
 	
 	const TArray<TVec3<int32>>& GetConstraints() const { return Constraints; }
-	const TArray<FSolverVec2> GetRestStretchLengths() const { return RestStretchLengths; }
+
+	TArray<TArray<int32>> GetConstraintsArray() const
+	{
+		TArray<TArray<int32>> ConstraintsArray;
+		ConstraintsArray.SetNum(Constraints.Num());
+		for (int32 i = 0; i < Constraints.Num(); i++)
+		{
+			ConstraintsArray[i].SetNum(3);
+			for (int32 j = 0; j < 3; j++)
+			{
+				ConstraintsArray[i][j] = Constraints[i][j];
+			}
+		}
+		return ConstraintsArray;
+	}
+	const TArray<FSolverVec3> GetRestStretchLengths() const { return RestStretchLengths; }
 	FSolverVec2 GetWarpWeftScale(const int32 ConstraintIndex) const 
 	{ 
 		return FSolverVec2(WarpScale.HasWeightMap() ? WarpScale[ConstraintIndex] : (FSolverReal)WarpScale,
@@ -108,6 +123,12 @@ public:
 	}
 	const TArray<int32>& GetConstraintsPerColorStartIndex() const { return ConstraintsPerColorStartIndex; }
 
+	CHAOS_API void AddStretchBiasElementResidualAndHessian(const FSolverParticles& Particles, const int32 ConstraintIndex, const int32 ConstraintIndexLocal, const FSolverReal Dt, TVec3<FSolverReal>& ParticleResidual, Chaos::PMatrix<FSolverReal, 3, 3>& ParticleHessian);
+
+	CHAOS_API void InitializeDmInvAndMeasures(const FSolverParticles& Particles);
+
+	CHAOS_API void AddInternalForceDifferential(const FSolverParticles& InParticles, const TArray<TVector<FSolverReal, 3>>& DeltaParticles, TArray<TVector<FSolverReal, 3>>& ndf);
+
 private:
 	CHAOS_API void InitConstraintsAndRestData(const FSolverParticles& InParticles, const FTriangleMesh& TriangleMesh,
 		const TArray<TVec3<FVec2f>>& FaceVertexUVs, const bool bUse3dRestLengths, const bool bTrimKinematicConstraints);
@@ -127,9 +148,16 @@ private:
 	TArray<int32> ConstraintsPerColorStartIndex; // Constraints are ordered so each batch is contiguous. This is ColorNum + 1 length so it can be used as start and end.
 
 	TArray<FSolverMatrix22> DeltaUVInverse; // Used to convert from DeltaX to dX/dU and dX/dV
-	TArray<FSolverVec2> RestStretchLengths;
+	TArray<FSolverVec3> RestStretchLengths;
 	TArray<FSolverVec3> StiffnessScales; // Used to make everything resolution independent.
 
+	TArray<FSolverReal> Measure;
+	TArray<FSolverMatrix22> DmInverse;
+	TArray<FSolverMatrix22> DmArray;
+	bool bDmInitialized = false;
+
+	TArray<PMatrix<FSolverReal, 3, 2>> RestDmArray;
+
 	UE_CHAOS_DECLARE_PROPERTYCOLLECTION_NAME(XPBDAnisoStretchUse3dRestLengths, bool);
 	UE_CHAOS_DECLARE_PROPERTYCOLLECTION_NAME(XPBDAnisoStretchStiffnessWarp, float);
 	UE_CHAOS_DECLARE_PROPERTYCOLLECTION_NAME(XPBDAnisoStretchStiffnessWeft, float);
diff --git a/Engine/Source/Runtime/Experimental/ChaosCore/Public/Chaos/Matrix.h b/Engine/Source/Runtime/Experimental/ChaosCore/Public/Chaos/Matrix.h
--- a/Engine/Source/Runtime/Experimental/ChaosCore/Public/Chaos/Matrix.h
+++ b/Engine/Source/Runtime/Experimental/ChaosCore/Public/Chaos/Matrix.h
@@ -85,6 +85,16 @@ namespace Chaos
 			M[5] = x21;
 		}
 
+		PMatrix(const FRealSingle x00)
+		{
+			M[0] = x00;
+			M[1] = x00;
+			M[2] = x00;
+			M[3] = x00;
+			M[4] = x00;
+			M[5] = x00;
+		}
+
 		TVector<FRealSingle, 3> operator*(const TVector<FRealSingle, 2>& Other)
 		{
 			return TVector<FRealSingle, 3>(
@@ -92,6 +102,41 @@ namespace Chaos
 				M[1] * Other[0] + M[4] * Other[1],
 				M[2] * Other[0] + M[5] * Other[1]);
 		}
+
+		PMatrix<FRealSingle, 3,2> operator-(const PMatrix<FRealSingle, 3, 2>& Other) const
+		{
+			return PMatrix<FRealSingle, 3, 2>(
+				M[0] - Other.M[0],
+				M[1] - Other.M[1],
+				M[2] - Other.M[2],
+				M[3] - Other.M[3],
+				M[4] - Other.M[4],
+				M[5] - Other.M[5]);
+		}
+
+		PMatrix<FRealSingle, 3, 2> operator+(const PMatrix<FRealSingle, 3, 2>& Other) const
+		{
+			return PMatrix<FRealSingle, 3, 2>(
+				M[0] + Other.M[0],
+				M[1] + Other.M[1],
+				M[2] + Other.M[2],
+				M[3] + Other.M[3],
+				M[4] + Other.M[4],
+				M[5] + Other.M[5]);
+		}
+
+		friend PMatrix<FRealSingle, 3, 2> operator*(const FRealSingle OtherF, const PMatrix<FRealSingle, 3, 2>& OtherM)
+		{
+			return PMatrix<FRealSingle, 3, 2>(
+				OtherF * OtherM.M[0],
+				OtherF * OtherM.M[1],
+				OtherF * OtherM.M[2],
+				OtherF * OtherM.M[3],
+				OtherF * OtherM.M[4],
+				OtherF * OtherM.M[5]);
+		}
+
+		
 	};
 
 	template<>
@@ -141,6 +186,11 @@ namespace Chaos
 			    -OneOverDeterminant * M[2],
 			    OneOverDeterminant * M[0]);
 		}
+
+		PMatrix<FReal, 2, 2> GetTransposed() const
+		{
+			return PMatrix<FReal, 2, 2>(M[0], M[2], M[1], M[3]);
+		}
 	};
 
 	template<>
@@ -190,6 +240,42 @@ namespace Chaos
 				-OneOverDeterminant * M[2],
 				OneOverDeterminant * M[0]);
 		}
+
+		PMatrix<FRealSingle, 2, 2> GetTransposed() const
+		{
+			return PMatrix<FRealSingle, 2, 2>(M[0], M[2], M[1], M[3]);
+		}
+
+		FRealSingle Determinant() const
+		{
+			return M[0] * M[3] - M[1] * M[2];
+		}
+
+		FORCEINLINE FRealSingle GetAt(int32 RowIndex, int32 ColIndex) const
+		{
+			return M[ColIndex * 2 + RowIndex];
+		}
+
+		//This is not column major so no need to invert matrix order when multiply them
+		PMatrix<FRealSingle, 2, 2> operator*(const PMatrix<FRealSingle, 2, 2>& Other) const
+		{
+			return PMatrix<FRealSingle, 2, 2>(
+				Other.M[0] * M[0] + Other.M[1] * M[2],
+				M[1] * Other.M[0] + M[3] * Other.M[1],
+				M[0] * Other.M[2] + M[2] * Other.M[3],
+				M[1] * Other.M[2] + M[3] * Other.M[3]);
+		}
+
+		friend PMatrix<FRealSingle, 3, 2> operator*(const PMatrix<FRealSingle, 3, 2>& First, const PMatrix<FRealSingle, 2, 2>& Other)
+		{
+			return PMatrix<FRealSingle, 3, 2>(
+				First.M[0] * Other.M[0] + First.M[3] * Other.M[1],
+				First.M[1] * Other.M[0] + First.M[4] * Other.M[1],
+				First.M[2] * Other.M[0] + First.M[5] * Other.M[1],
+				First.M[0] * Other.M[2] + First.M[3] * Other.M[3],
+				First.M[1] * Other.M[2] + First.M[4] * Other.M[3],
+				First.M[2] * Other.M[2] + First.M[5] * Other.M[3]);
+		}
 	};
 
 	template<>
