Added invalidation to cluster union geometry duplication so we only attempt to copy the updated geometry if it actually changed
Added shallow copy of the root cluster instead of deep copy of all geometry during buffering

diff --git a/Engine/Source/Runtime/Experimental/Chaos/Private/Chaos/ClusterUnionManager.cpp b/Engine/Source/Runtime/Experimental/Chaos/Private/Chaos/ClusterUnionManager.cpp
--- a/Engine/Source/Runtime/Experimental/Chaos/Private/Chaos/ClusterUnionManager.cpp
+++ b/Engine/Source/Runtime/Experimental/Chaos/Private/Chaos/ClusterUnionManager.cpp
@@ -258,11 +258,13 @@ namespace Chaos
 
 	void FClusterUnion::ClearAllPendingGeometryOperations()
 	{
+		bGeometryModified = true;
 		PendingGeometryOperations.Empty();
 	}
 
 	void FClusterUnion::ClearPendingGeometryOperations(EClusterUnionGeometryOperation Op)
 	{
+		bGeometryModified = true;
 		PendingGeometryOperations.Remove(Op);
 	}
 
diff --git a/Engine/Source/Runtime/Experimental/Chaos/Private/PhysicsProxy/ClusterUnionPhysicsProxy.cpp b/Engine/Source/Runtime/Experimental/Chaos/Private/PhysicsProxy/ClusterUnionPhysicsProxy.cpp
--- a/Engine/Source/Runtime/Experimental/Chaos/Private/PhysicsProxy/ClusterUnionPhysicsProxy.cpp
+++ b/Engine/Source/Runtime/Experimental/Chaos/Private/PhysicsProxy/ClusterUnionPhysicsProxy.cpp
@@ -12,6 +12,9 @@
 
 namespace Chaos
 {
+	static int32 GShallowCopyClusterUnionGeometryOnUpdate = 1;
+	FAutoConsoleVariableRef CVar_ShallowCopyClusterUnionGeometryOnUpdate(TEXT("p.ShallowCopyOnClusterUnionUpdate"), GShallowCopyClusterUnionGeometryOnUpdate, TEXT("If 1, shallow copy the root union geometry of a cluster union when its geometry updates, otherwise deep copy the geometry hierarchy"));
+
 	namespace
 	{
 		FPBDRigidsEvolutionGBF* GetEvolution(FClusterUnionPhysicsProxy* Proxy)
@@ -585,10 +588,25 @@ namespace Chaos
 				}
 			}
 
-			// TODO: Probably only want to do this when something about the geometry changes.
-			if (!ClusterUnion->ChildParticles.IsEmpty())
+			if (ClusterUnion->bGeometryModified && !ClusterUnion->ChildParticles.IsEmpty())
 			{
-				BufferData.Geometry = ClusterUnion->Geometry->DeepCopyGeometry();
+				if(FImplicitObjectUnion* AsUnion = ClusterUnion->Geometry->AsA<FImplicitObjectUnion>();
+				   GShallowCopyClusterUnionGeometryOnUpdate && AsUnion)
+				{
+					// Shallow copy the root union for the GT update
+					// This ensures the GT has a snapshot of the union as it is now for this results data. The PT is free to
+					// continue modifying its geometry without potentially disrupting GT reads. Internally we don't want
+					// to duplicate all the geometry - we essentially just need the list of objects in the root union to
+					// point to the correct internal geometries, but have a separate list on GT and PT.
+					TArray<FImplicitObjectPtr> ObjectsCopy = AsUnion->GetObjects();
+					BufferData.Geometry = MakeImplicitObjectPtr<FImplicitObjectUnion>(MoveTemp(ObjectsCopy));
+				}
+				else
+				{
+					// Fallback to deep copy geometry hierarchy
+					BufferData.Geometry = ClusterUnion->Geometry->DeepCopyGeometry();
+				}
+				ClusterUnion->bGeometryModified = false;
 			}
 		}
 	}
diff --git a/Engine/Source/Runtime/Experimental/Chaos/Public/Chaos/ClusterUnionManager.h b/Engine/Source/Runtime/Experimental/Chaos/Public/Chaos/ClusterUnionManager.h
--- a/Engine/Source/Runtime/Experimental/Chaos/Public/Chaos/ClusterUnionManager.h
+++ b/Engine/Source/Runtime/Experimental/Chaos/Public/Chaos/ClusterUnionManager.h
@@ -107,16 +107,19 @@ namespace Chaos
 		FClusterUnionCreationParameters ClusterUnionParameters;
 
 		// Whether or not the position/rotation needs to be computed the first time a particle is added.
-		bool bNeedsXRInitialization = true;
+		bool bNeedsXRInitialization : 1 = true;
 
 		// Whether or not the anchor got set by the GT and we shouldn't try to recompute its value.
-		bool bAnchorLock = false;
+		bool bAnchorLock : 1 = false;
 
 		// Whether or not we need to check for connectivity for this cluster union.
-		bool bCheckConnectivity = false;
+		bool bCheckConnectivity : 1 = false;
 
 		// Whether or not to generate connectivity edges (used to differentiate between server graphs vs client graphs).
-		bool bGenerateConnectivityEdges = true;
+		bool bGenerateConnectivityEdges : 1 = true;
+
+		// Whether the geometry for the union has change such that we will need to push a copy to the game thread
+		bool bGeometryModified : 1 = true;
 
 		// Pending particles that need to be added into the connectivity graph.
 		TArray<TPair<FPBDRigidParticleHandle*, EClusterUnionConnectivityOperation>> PendingConnectivityOperations;
