Don't create render state for geometry collections using custom renderers

diff --git a/Engine/Source/Runtime/Experimental/GeometryCollectionEngine/Private/GeometryCollection/GeometryCollectionComponent.cpp b/Engine/Source/Runtime/Experimental/GeometryCollectionEngine/Private/GeometryCollection/GeometryCollectionComponent.cpp
--- a/Engine/Source/Runtime/Experimental/GeometryCollectionEngine/Private/GeometryCollection/GeometryCollectionComponent.cpp
+++ b/Engine/Source/Runtime/Experimental/GeometryCollectionEngine/Private/GeometryCollection/GeometryCollectionComponent.cpp
@@ -583,7 +583,6 @@ UGeometryCollectionComponent::UGeometryCollectionComponent(const FObjectInitiali
 	, ReplicationAbandonAfterLevel(0)
 	, ReplicationMaxPositionAndVelocityCorrectionLevel(100)
 	, bInitializedRemovalDynamicAttribute(false)
-	, bRenderStateDirty(true)
 	, bEnableBoneSelection(false)
 	, ViewLevel(-1)
 	, NavmeshInvalidationTimeSliceIndex(0)
@@ -1016,6 +1015,11 @@ void UGeometryCollectionComponent::UpdateCachedBounds()
 	UpdateBounds();
 }
 
+bool UGeometryCollectionComponent::ShouldCreateRenderState() const
+{
+	return !CanUseCustomRenderer();
+}
+
 void UGeometryCollectionComponent::CreateRenderState_Concurrent(FRegisterComponentContext* Context)
 {
 	Super::CreateRenderState_Concurrent(Context);
@@ -4180,15 +4184,17 @@ void UGeometryCollectionComponent::UpdateRenderSystemsIfNeeded(bool bDynamicColl
 			RefreshCustomRenderer();
 		}
 
-		if (SceneProxy && SceneProxy->IsNaniteMesh())
+		if (SceneProxy != nullptr)
 		{
-			FNaniteGeometryCollectionSceneProxy* NaniteProxy = static_cast<FNaniteGeometryCollectionSceneProxy*>(SceneProxy);
-			NaniteProxy->FlushGPUSceneUpdate_GameThread();
-		}
+			if (SceneProxy->IsNaniteMesh())
+			{
+				FNaniteGeometryCollectionSceneProxy* NaniteProxy = static_cast<FNaniteGeometryCollectionSceneProxy*>(SceneProxy);
+				NaniteProxy->FlushGPUSceneUpdate_GameThread();
+			}
 
-		MarkRenderTransformDirty();
-		MarkRenderDynamicDataDirty();
-		bRenderStateDirty = false;
+			MarkRenderTransformDirty();
+			MarkRenderDynamicDataDirty();
+		}
 	}
 }
 
diff --git a/Engine/Source/Runtime/Experimental/GeometryCollectionEngine/Public/GeometryCollection/GeometryCollectionComponent.h b/Engine/Source/Runtime/Experimental/GeometryCollectionEngine/Public/GeometryCollection/GeometryCollectionComponent.h
--- a/Engine/Source/Runtime/Experimental/GeometryCollectionEngine/Public/GeometryCollection/GeometryCollectionComponent.h
+++ b/Engine/Source/Runtime/Experimental/GeometryCollectionEngine/Public/GeometryCollection/GeometryCollectionComponent.h
@@ -583,9 +583,10 @@ public:
 	//~ End UObject Interface.
 
 	//~ Begin UActorComponent Interface.
+	GEOMETRYCOLLECTIONENGINE_API virtual bool ShouldCreateRenderState() const override;
 	GEOMETRYCOLLECTIONENGINE_API virtual void CreateRenderState_Concurrent(FRegisterComponentContext* Context) override;
 	GEOMETRYCOLLECTIONENGINE_API virtual void SendRenderDynamicData_Concurrent() override;
-	FORCEINLINE void SetRenderStateDirty() { bRenderStateDirty = true; }
+	FORCEINLINE void SetRenderStateDirty() { /* Deprecated. */ }
 	GEOMETRYCOLLECTIONENGINE_API virtual void SetCollisionObjectType(ECollisionChannel Channel) override;
 	GEOMETRYCOLLECTIONENGINE_API virtual void OnActorEnableCollisionChanged() override;
 	GEOMETRYCOLLECTIONENGINE_API virtual void BeginPlay() override;
@@ -1566,7 +1567,6 @@ private:
 	void ComputeCurrentGlobalsMatrices(TArray<FTransform3f>& OutTransforms) const;
 
 	bool bInitializedRemovalDynamicAttribute;
-	bool bRenderStateDirty;
 	bool bEnableBoneSelection;
 	int ViewLevel;
 
