Iris
* Server will now immediately force disconnect a client when it detects a protocol mismatch for critical classes

#rnx

diff --git a/Engine/Config/BaseEngine.ini b/Engine/Config/BaseEngine.ini
--- a/Engine/Config/BaseEngine.ini
+++ b/Engine/Config/BaseEngine.ini
@@ -1293,11 +1293,15 @@ DefaultSpatialFilterName=Spatial
 ; PlayerState is by default bAlwaysRelevant but the default player state prioritizer handles this.
 +PrioritizerConfigs=(ClassName=/Script/Engine.PlayerState, PrioritizerName=PlayerState, bForceEnableOnAllInstances=true)
 
-;DeltaCompressionConfigs=ClearArray
+; DeltaCompressionConfigs
 ; Classes for which deltacompression is enabled
 +DeltaCompressionConfigs=(ClassName=/Script/Engine.Pawn))
 +DeltaCompressionConfigs=(ClassName=/Script/Engine.PlayerState))
 
+; CriticalClassConfigs
+; Classes which force a disconnection when a protocol mismatch is detected
++CriticalClassConfigs=(ClassName=/Script/Engine.PlayerController)
+
 [/Script/IrisCore.ReplicationStateDescriptorConfig]
 ; Declarate structs that are vetted to work using reflection based struct serialization even though there exists a custom NetSerialize function for the struct
 +SupportsStructNetSerializerList=(StructName=GameplayCueParameters)
diff --git a/Engine/Source/Runtime/Engine/Private/Net/Iris/ReplicationSystem/ActorReplicationBridge.cpp b/Engine/Source/Runtime/Engine/Private/Net/Iris/ReplicationSystem/ActorReplicationBridge.cpp
--- a/Engine/Source/Runtime/Engine/Private/Net/Iris/ReplicationSystem/ActorReplicationBridge.cpp
+++ b/Engine/Source/Runtime/Engine/Private/Net/Iris/ReplicationSystem/ActorReplicationBridge.cpp
@@ -32,6 +32,7 @@
 #include "GameFramework/Actor.h"
 #include "Net/DataBunch.h"
 #include "Net/DataChannel.h"
+#include "Net/Core/Connection/NetCloseResult.h"
 #include "Net/Core/Misc/NetSubObjectRegistry.h"
 #include "Net/NetSubObjectRegistryGetter.h"
 #include "Templates/Casts.h"
@@ -1164,6 +1165,41 @@ void UActorReplicationBridge::OnProtocolMismatchDetected(FNetRefHandle ObjectHan
 	}
 }
 
+void UActorReplicationBridge::OnProtocolMismatchReported(FNetRefHandle RefHandle, uint32 ConnectionId)
+{
+	Super::OnProtocolMismatchReported(RefHandle, ConnectionId);
+
+	// If we are the server force the client to disconnect since not replicating a critical class will prevent the game from working.
+	if (NetDriver && NetDriver->IsServer())
+	{
+		const UObject* ReplicatedObject = GetReplicatedObject(RefHandle);
+
+		// If the object instance doesn't exist anymore, pass a null class anyway in case the config wants to disconnect on ALL class types.
+		const UClass* ObjectClass = ReplicatedObject ? ReplicatedObject->GetClass() : nullptr;
+
+		if (IsClassCritical(ObjectClass))
+		{
+			if (UNetConnection* ClientConnection = NetDriver->GetConnectionById(ConnectionId))
+			{
+				FString ErrorMsg = FString::Printf(TEXT("Protocol mismatch: %s:%s. Class: %s"), *RefHandle.ToString(), *GetNameSafe(ReplicatedObject), *GetNameSafe(ObjectClass));
+				UE_LOG(LogIrisBridge, Error, TEXT("%s: Closing connection due to: %s"), ToCStr(ClientConnection->Describe()), ToCStr(ErrorMsg));
+				{
+					UE::Net::FNetCloseResult CloseReason = ENetCloseResult::IrisProtocolMismatch;
+					ClientConnection->SendCloseReason(MoveTemp(CloseReason));
+				}
+
+				FNetControlMessage<NMT_Failure>::Send(ClientConnection, ErrorMsg);
+				ClientConnection->FlushNet(true);
+
+				{
+					UE::Net::FNetCloseResult CloseReason = ENetCloseResult::IrisProtocolMismatch;
+					ClientConnection->Close(MoveTemp(CloseReason));
+				}
+			}
+		}
+	}
+}
+
 #else //!UE_WITH_IRIS
 
 UActorReplicationBridge::UActorReplicationBridge() = default;
diff --git a/Engine/Source/Runtime/Engine/Public/Net/Iris/ReplicationSystem/ActorReplicationBridge.h b/Engine/Source/Runtime/Engine/Public/Net/Iris/ReplicationSystem/ActorReplicationBridge.h
--- a/Engine/Source/Runtime/Engine/Public/Net/Iris/ReplicationSystem/ActorReplicationBridge.h
+++ b/Engine/Source/Runtime/Engine/Public/Net/Iris/ReplicationSystem/ActorReplicationBridge.h
@@ -99,6 +99,7 @@ protected:
 	virtual bool ObjectLevelHasFinishedLoading(UObject* Object) const override;
 	virtual bool IsAllowedToDestroyInstance(const UObject* Instance) const override;
 	virtual void OnProtocolMismatchDetected(FNetRefHandle ObjectHandle) override;
+	virtual void OnProtocolMismatchReported(FNetRefHandle RefHandle, uint32 ConnectionId) override;
 
 	virtual float GetPollFrequencyOfRootObject(const UObject* ReplicatedObject) const override;
 
diff --git a/Engine/Source/Runtime/Experimental/Iris/Core/Private/Iris/ReplicationSystem/ObjectReplicationBridge.cpp b/Engine/Source/Runtime/Experimental/Iris/Core/Private/Iris/ReplicationSystem/ObjectReplicationBridge.cpp
--- a/Engine/Source/Runtime/Experimental/Iris/Core/Private/Iris/ReplicationSystem/ObjectReplicationBridge.cpp
+++ b/Engine/Source/Runtime/Experimental/Iris/Core/Private/Iris/ReplicationSystem/ObjectReplicationBridge.cpp
@@ -1373,6 +1373,28 @@ bool UObjectReplicationBridge::ShouldClassBeDeltaCompressed(const UClass* Class)
 	return false;
 }
 
+bool UObjectReplicationBridge::IsClassCritical(const UClass* Class)
+{
+	const UObjectReplicationBridgeConfig* BridgeConfig = UObjectReplicationBridgeConfig::GetConfig();
+	if (BridgeConfig->AreAllClassesCritical())
+	{
+		return true;
+	}
+
+	if (ClassesFlaggedCritical.Num() > 0)
+	{
+		for (; Class != nullptr; Class = Class->GetSuperClass())
+		{
+			if (bool* bIsClassCritical = ClassesFlaggedCritical.Find(GetConfigClassPathName(Class)))
+			{
+				return *bIsClassCritical;
+			}
+		}
+	}
+
+	return false;
+}
+
 void UObjectReplicationBridge::SetClassDynamicFilterConfig(FName ClassPathName, const UE::Net::FNetObjectFilterHandle FilterHandle)
 {
 	if (ClassPathName.IsNone())
@@ -1527,6 +1549,7 @@ void UObjectReplicationBridge::LoadConfig()
 	ClassesWithDynamicFilter.Empty();
 	ClassesWithPrioritizer.Empty();
 	ClassesWithDeltaCompression.Empty();
+	ClassesFlaggedCritical.Empty();
 
 	// Reset PathNameCache
 	ConfigClassPathNameCache.Empty();
@@ -1601,6 +1624,20 @@ void UObjectReplicationBridge::LoadConfig()
 			ClassesWithDeltaCompression.Add(DCConfig.ClassName, DCConfig.bEnableDeltaCompression);
 		}
 	}
+
+	// Critical classes
+	if (!BridgeConfig->AreAllClassesCritical())
+	{
+		for (const FObjectReplicatedBridgeCriticalClassConfig& CriticalClassConfig : BridgeConfig->GetCriticalClassConfigs())
+		{
+			if (!ensure(!ForbiddenNamesArray.Contains(CriticalClassConfig.ClassName)))
+			{
+				continue;
+			}
+
+			ClassesFlaggedCritical.Add(CriticalClassConfig.ClassName, CriticalClassConfig.bDisconnectOnProtocolMismatch);
+		}
+	}
 }
 
 void UObjectReplicationBridge::FindClassesInPollPeriodOverrides()
@@ -1755,10 +1792,10 @@ void UObjectReplicationBridge::OnProtocolMismatchReported(FNetRefHandle RefHandl
 	using namespace UE::Net;
 	using namespace UE::Net::Private;
 
-	// Ensure at the end so the log contains the most relevant information already
+	// Ensure at the end so the log contains all the relevant information
 	ON_SCOPE_EXIT
 	{
-		ensure(false);
+		ensureMsgf(false, TEXT("Protocol mismatch detected. Compare the CDO state in the server and client logs to find the source of the issue."));
 	};
 	
 
@@ -1806,6 +1843,6 @@ void UObjectReplicationBridge::OnProtocolMismatchReported(FNetRefHandle RefHandl
 			Fragments.Emplace(MoveTemp(FragmentInfo));
 		}
 
-		UE::Net::Private::ObjectBridgeDebugging::RemoteProtocolMismatchDetected(ReplicationSystem, 0 /*TODO: Local ConnectionId*/, Fragments, ObjInstance, ObjArchetype);
+		UE::Net::Private::ObjectBridgeDebugging::RemoteProtocolMismatchDetected(ReplicationSystem, 0 /*TODO: Local ConnectionId*/, Fragments, ObjArchetype, ObjInstance);
 	}
 }
diff --git a/Engine/Source/Runtime/Experimental/Iris/Core/Private/Iris/ReplicationSystem/ObjectReplicationBridgeConfig.cpp b/Engine/Source/Runtime/Experimental/Iris/Core/Private/Iris/ReplicationSystem/ObjectReplicationBridgeConfig.cpp
--- a/Engine/Source/Runtime/Experimental/Iris/Core/Private/Iris/ReplicationSystem/ObjectReplicationBridgeConfig.cpp
+++ b/Engine/Source/Runtime/Experimental/Iris/Core/Private/Iris/ReplicationSystem/ObjectReplicationBridgeConfig.cpp
@@ -32,3 +32,8 @@ TConstArrayView<FObjectReplicationBridgeDeltaCompressionConfig> UObjectReplicati
 {
 	return MakeArrayView(DeltaCompressionConfigs);
 }
+
+TConstArrayView<FObjectReplicatedBridgeCriticalClassConfig> UObjectReplicationBridgeConfig::GetCriticalClassConfigs() const
+{
+	return MakeArrayView(CriticalClassConfigs);
+}
diff --git a/Engine/Source/Runtime/Experimental/Iris/Core/Public/Iris/ReplicationSystem/ObjectReplicationBridge.h b/Engine/Source/Runtime/Experimental/Iris/Core/Public/Iris/ReplicationSystem/ObjectReplicationBridge.h
--- a/Engine/Source/Runtime/Experimental/Iris/Core/Public/Iris/ReplicationSystem/ObjectReplicationBridge.h
+++ b/Engine/Source/Runtime/Experimental/Iris/Core/Public/Iris/ReplicationSystem/ObjectReplicationBridge.h
@@ -290,6 +290,9 @@ protected:
 	 */
 	IRISCORE_API bool GetClassPollFrequency(const UClass* Class, float& OutPollPeriod) const;
 
+	/** Returns true if the class is considered critical and we force a disconnection if a protocol mismatch prevents instances of this class from replicating. */
+	IRISCORE_API bool IsClassCritical(const UClass* Class);
+
 	/** Current max tick rate set by the engine */
 	float GetMaxTickRate() const { return MaxTickRate; }
 
@@ -390,6 +393,9 @@ private:
 	// Delta compression
 	TMap<FName, bool> ClassesWithDeltaCompression;
 
+	// Classes that may force a disconnection when a protocol mismatch is detected.
+	TMap<FName, bool> ClassesFlaggedCritical;
+
 	// Array of dormant objects that has requested a flush
 	TArray<FNetRefHandle> DormantHandlesPendingFlush;
 
diff --git a/Engine/Source/Runtime/Experimental/Iris/Core/Public/Iris/ReplicationSystem/ObjectReplicationBridgeConfig.h b/Engine/Source/Runtime/Experimental/Iris/Core/Public/Iris/ReplicationSystem/ObjectReplicationBridgeConfig.h
--- a/Engine/Source/Runtime/Experimental/Iris/Core/Public/Iris/ReplicationSystem/ObjectReplicationBridgeConfig.h
+++ b/Engine/Source/Runtime/Experimental/Iris/Core/Public/Iris/ReplicationSystem/ObjectReplicationBridgeConfig.h
@@ -80,6 +80,20 @@ struct FObjectReplicationBridgeDeltaCompressionConfig
 	bool bEnableDeltaCompression = true;
 };
 
+USTRUCT()
+struct FObjectReplicatedBridgeCriticalClassConfig
+{
+	GENERATED_BODY()
+
+	/** Instances of this class or its subclasses will force a client disconnection when it detects a protocol mismatch.*/
+	UPROPERTY()
+	FName ClassName;
+
+	/** When true we force the client to disconnect when a protocol mismatch prevents it from instantiating replicated objects of this class. */
+	UPROPERTY()
+	bool bDisconnectOnProtocolMismatch = true;
+};
+
 UCLASS(transient, config=Engine)
 class UObjectReplicationBridgeConfig : public UObject
 {
@@ -93,9 +107,14 @@ public:
 	IRISCORE_API TConstArrayView<FObjectReplicationBridgeFilterConfig> GetFilterConfigs() const;
 	IRISCORE_API TConstArrayView<FObjectReplicationBridgePrioritizerConfig> GetPrioritizerConfigs() const;
 	IRISCORE_API TConstArrayView<FObjectReplicationBridgeDeltaCompressionConfig> GetDeltaCompressionConfigs() const;
+	IRISCORE_API TConstArrayView<FObjectReplicatedBridgeCriticalClassConfig> GetCriticalClassConfigs() const;
+
 	FName GetDefaultSpatialFilterName() const;
 	FName GetRequiredNetDriverChannelClassName() const;
 
+	/** When true any class with a protocol mismatch will force a disconnection. */
+	bool AreAllClassesCritical() const { return bAllClassesCritical; }
+
 protected:
 	UObjectReplicationBridgeConfig();
 
@@ -128,6 +147,14 @@ private:
 	UPROPERTY(Config)
 	TArray<FObjectReplicationBridgeDeltaCompressionConfig> DeltaCompressionConfigs;
 
+	/** Classes that are considered critical and will force a disconnection when a protocol mismatch is detected. */
+	UPROPERTY(Config)
+	TArray<FObjectReplicatedBridgeCriticalClassConfig> CriticalClassConfigs;
+
+	/** Set this to true if you want any class with a protocol mismatch to force a disconnection. */
+	UPROPERTY(Config)
+	bool bAllClassesCritical = false;
+
 	/**
 	 * The name of the filter to apply objects that can have spatial filtering applied.
 	 */
@@ -139,6 +166,9 @@ private:
 	 */
 	UPROPERTY(Config)
 	FName RequiredNetDriverChannelClassName;
+
+	UPROPERTY(Config)
+	TArray<FName> CriticalActorClasses;
 };
 
 inline FName UObjectReplicationBridgeConfig::GetDefaultSpatialFilterName() const
diff --git a/Engine/Source/Runtime/Net/Core/Public/Net/Core/Connection/NetCloseResult.h b/Engine/Source/Runtime/Net/Core/Public/Net/Core/Connection/NetCloseResult.h
--- a/Engine/Source/Runtime/Net/Core/Public/Net/Core/Connection/NetCloseResult.h
+++ b/Engine/Source/Runtime/Net/Core/Public/Net/Core/Connection/NetCloseResult.h
@@ -371,6 +371,12 @@ enum class ENetCloseResult : uint32
 	/** Join failure, unexpected control message */
 	BeaconSpawnUnexpectedError,
 
+	
+	/** Iris Error result types */
+
+	/** Protocol mismatch in Iris preventing a critical object instantiation */
+	IrisProtocolMismatch,
+
 
 	/** Fault Handler Error result types */
 
