ChaosSofts : Fix CIS Issue

diff --git a/Engine/Source/Runtime/Experimental/Chaos/Public/Chaos/Deformable/GaussSeidelCorotatedCodimensionalConstraints.h b/Engine/Source/Runtime/Experimental/Chaos/Public/Chaos/Deformable/GaussSeidelCorotatedCodimensionalConstraints.h
--- a/Engine/Source/Runtime/Experimental/Chaos/Public/Chaos/Deformable/GaussSeidelCorotatedCodimensionalConstraints.h
+++ b/Engine/Source/Runtime/Experimental/Chaos/Public/Chaos/Deformable/GaussSeidelCorotatedCodimensionalConstraints.h
@@ -54,6 +54,10 @@ namespace Chaos::Softs
 
 		PMatrix<T, 3, 2> Ds(const int32 E, const ParticleType& InParticles) const 
 		{
+			check(E > INDEX_NONE && E < MeshConstraints.Num());
+			check(MeshConstraints[E][0] < (int32)InParticles.Size() && MeshConstraints[E][0] > INDEX_NONE);
+			check(MeshConstraints[E][1] < (int32)InParticles.Size() && MeshConstraints[E][1] > INDEX_NONE);
+			check(MeshConstraints[E][2] < (int32)InParticles.Size() && MeshConstraints[E][2] > INDEX_NONE);
 			const TVec3<T> P1P0 = InParticles.P(MeshConstraints[E][1]) - InParticles.P(MeshConstraints[E][0]);
 			const TVec3<T> P2P0 = InParticles.P(MeshConstraints[E][2]) - InParticles.P(MeshConstraints[E][0]);
 
@@ -64,6 +68,8 @@ namespace Chaos::Softs
 
 		PMatrix<T, 3, 2> F(const int32 E, const ParticleType& InParticles) const 
 		{
+			check(E > INDEX_NONE && E < MeshConstraints.Num());
+			check(E < DmInverse.Num());
 			return Ds(E, InParticles) * DmInverse[E];
 		}
 
@@ -105,6 +111,9 @@ namespace Chaos::Softs
 			DmInverse.Init(PMatrix<FSolverReal, 2, 2>(0.f, 0.f, 0.f), MeshConstraints.Num());
 			for (int32 e = 0; e < MeshConstraints.Num(); e++)
 			{
+				check(MeshConstraints[e][0] < (int32)Particles.Size() && MeshConstraints[e][0] > INDEX_NONE);
+				check(MeshConstraints[e][1] < (int32)Particles.Size() && MeshConstraints[e][1] > INDEX_NONE);
+				check(MeshConstraints[e][2] < (int32)Particles.Size() && MeshConstraints[e][2] > INDEX_NONE);
 				const TVec3<T> X1X0 = Particles.X(MeshConstraints[e][1]) - Particles.X(MeshConstraints[e][0]);
 				const TVec3<T> X2X0 = Particles.X(MeshConstraints[e][2]) - Particles.X(MeshConstraints[e][0]);
 				PMatrix<T, 2, 2> Dm((T)0., (T)0., (T)0.);
@@ -243,6 +252,11 @@ namespace Chaos::Softs
 	public:
 		void AddHyperelasticResidualAndHessian(const ParticleType& Particles, const int32 ElementIndex, const int32 ElementIndexLocal, const T Dt, TVec3<T>& ParticleResidual, Chaos::PMatrix<T, 3, 3>& ParticleHessian)
 		{
+			check(ElementIndex < DmInverse.Num() && ElementIndex > INDEX_NONE);
+			check(ElementIndex < MuElementArray.Num());
+			check(ElementIndex < Measure.Num());
+			check(ElementIndex < LambdaElementArray.Num());
+			check(ElementIndexLocal < 4 && ElementIndexLocal > INDEX_NONE);
 			const Chaos::PMatrix<T, 2, 2> DmInvT = DmInverse[ElementIndex].GetTransposed(), DmInv = DmInverse[ElementIndex]; 
 			const Chaos::PMatrix<T, 3, 2> Fe = F(ElementIndex, Particles);
 
