Removing the usage of Engine headers inside of Chaos

#tests PIE, built HeadlessChaos

diff --git a/Engine/Source/Runtime/Experimental/Chaos/Private/PhysicsProxy/GeometryCollectionPhysicsProxy.cpp b/Engine/Source/Runtime/Experimental/Chaos/Private/PhysicsProxy/GeometryCollectionPhysicsProxy.cpp
--- a/Engine/Source/Runtime/Experimental/Chaos/Private/PhysicsProxy/GeometryCollectionPhysicsProxy.cpp
+++ b/Engine/Source/Runtime/Experimental/Chaos/Private/PhysicsProxy/GeometryCollectionPhysicsProxy.cpp
@@ -38,8 +38,6 @@
 #include "GeometryCollection/Facades/CollectionAnchoringFacade.h"
 #include "GeometryCollection/Facades/CollectionConnectionGraphFacade.h"
 
-#include "Physics/Experimental/ChaosInterfaceUtils.h"
-
 #ifndef TODO_REIMPLEMENT_INIT_COMMANDS
 #define TODO_REIMPLEMENT_INIT_COMMANDS 0
 #endif
@@ -607,18 +605,12 @@ void FGeometryCollectionPhysicsProxy::Initialize(Chaos::FPBDRigidsEvolutionBase
 				P->SetUserData(Parameters.UserData);
 				P->SetProxy(this);
 
-				if (Index == Parameters.InitialRootIndex && bUseStaticMeshCollisionForTraces && (OptionalTraceCollisionOverrideGeoms.Num() > 0))
+				if (Index == Parameters.InitialRootIndex && bUseStaticMeshCollisionForTraces && CreateTraceCollisionGeometryCallback != nullptr)
 				{
-					Chaos::FShapesArray Shapes;
+					const FTransform ToLocal = MassToLocal[Index].Inverse();
 					TArray<Chaos::FImplicitObjectPtr> Geoms;
-					for (const FTraceCollisionOverrideData& OverrideData : OptionalTraceCollisionOverrideGeoms)
-					{
-						// Copying Params so we can replace the local transform
-						FGeometryAddParams AddParams = OverrideData.Params;
-						AddParams.LocalTransform = AddParams.LocalTransform * MassToLocal[Index].Inverse();
-
-						OverrideData.CreateGeometryCallback(AddParams, Geoms, Shapes);
-					}
+					Chaos::FShapesArray Shapes;
+					CreateTraceCollisionGeometryCallback(ToLocal, Geoms, Shapes);
 
 					Chaos::FImplicitObjectPtr ImplicitGeometry = MakeImplicitObjectPtr<Chaos::FImplicitObjectUnion>(MoveTemp(Geoms));
 					P->SetGeometry(ImplicitGeometry);
@@ -2976,7 +2968,7 @@ void FGeometryCollectionPhysicsProxy::SetUseStaticMeshCollisionForTraces_Externa
 
 	// Expensive, so only do it when we need to
 	const bool bUseSMCollisionForTraces = ((bInUseStaticMeshCollisionForTraces && ForceOverrideGCCollisionSetupForTraces == GCCSFT_Property) || (ForceOverrideGCCollisionSetupForTraces == GCCSFT_ForceSM));
-	if (bUseStaticMeshCollisionForTraces != bUseSMCollisionForTraces && (OptionalTraceCollisionOverrideGeoms.Num() > 0))
+	if (bUseStaticMeshCollisionForTraces != bUseSMCollisionForTraces && CreateTraceCollisionGeometryCallback != nullptr)
 	{
 		const int32 Index = Parameters.InitialRootIndex;
 		if (GTParticles.Num() > Index)
@@ -2987,16 +2979,10 @@ void FGeometryCollectionPhysicsProxy::SetUseStaticMeshCollisionForTraces_Externa
 				{
 					const TManagedArray<FTransform>& MassToLocal = Parameters.RestCollection->GetAttribute<FTransform>("MassToLocal", FGeometryCollection::TransformGroup);
 
-					Chaos::FShapesArray Shapes;
+					const FTransform ToLocal = MassToLocal[Index].Inverse();
 					TArray<Chaos::FImplicitObjectPtr> Geoms;
-					for (const FTraceCollisionOverrideData& OverrideData : OptionalTraceCollisionOverrideGeoms)
-					{
-						// Copying Params so we can replace the local transform
-						FGeometryAddParams AddParams = OverrideData.Params;
-						AddParams.LocalTransform = AddParams.LocalTransform * MassToLocal[Index].Inverse();
-
-						OverrideData.CreateGeometryCallback(AddParams, Geoms, Shapes);
-					}
+					Chaos::FShapesArray Shapes;
+					CreateTraceCollisionGeometryCallback(ToLocal, Geoms, Shapes);
 
 					Chaos::FImplicitObjectPtr ImplicitGeometry = MakeImplicitObjectPtr<Chaos::FImplicitObjectUnion>(MoveTemp(Geoms));
 					P->SetGeometry(ImplicitGeometry);
@@ -5462,11 +5448,6 @@ Chaos::FPhysicsObjectHandle FGeometryCollectionPhysicsProxy::GetPhysicsObjectByI
 	return PhysicsObjects[Index].Get();
 }
 
-void FGeometryCollectionPhysicsProxy::RegisterNewTraceCollisionOverrideData(const FGeometryAddParams& InParams, FTraceCollisionCreateGeometryCallback InCreateGeometryCallback)
-{
-	OptionalTraceCollisionOverrideGeoms.Emplace(InParams, InCreateGeometryCallback);
-}
-
 FGeometryCollectionPhysicsProxy::FParticle* FGeometryCollectionPhysicsProxy::GetParticleByIndex_External(int32 Index)
 {
 	return const_cast<FGeometryCollectionPhysicsProxy::FParticle*>(const_cast<const FGeometryCollectionPhysicsProxy*>(this)->GetParticleByIndex_External(Index));
diff --git a/Engine/Source/Runtime/Experimental/Chaos/Public/PhysicsProxy/GeometryCollectionPhysicsProxy.h b/Engine/Source/Runtime/Experimental/Chaos/Public/PhysicsProxy/GeometryCollectionPhysicsProxy.h
--- a/Engine/Source/Runtime/Experimental/Chaos/Public/PhysicsProxy/GeometryCollectionPhysicsProxy.h
+++ b/Engine/Source/Runtime/Experimental/Chaos/Public/PhysicsProxy/GeometryCollectionPhysicsProxy.h
@@ -18,7 +18,6 @@
 #include "PBDRigidsSolver.h"
 #include "Chaos/Defines.h"
 #include "Chaos/GeometryParticlesfwd.h"
-#include "Physics/PhysicsInterfaceTypes.h"
 
 namespace Chaos
 {
@@ -496,8 +495,8 @@ public:
 	int32 GetNumParticles() const { return NumParticles; }
 
 	// todo(chaos): Remove this and move to a cook time approach of the SM data based on the GC property
-	using FTraceCollisionCreateGeometryCallback = TFunction<void(const FGeometryAddParams& InParams, TArray<Chaos::FImplicitObjectPtr>& OutGeoms, Chaos::FShapesArray& OutShapes)>;
-	CHAOS_API void RegisterNewTraceCollisionOverrideData(const FGeometryAddParams& InParams, FTraceCollisionCreateGeometryCallback InCreateGeometryCallback);
+	using FCreateTraceCollisionGeometryCallback = TFunction<void(const FTransform& InToLocal, TArray<Chaos::FImplicitObjectPtr>& OutGeoms, Chaos::FShapesArray& OutShapes)>;
+	void SetCreateTraceCollisionGeometryCallback(FCreateTraceCollisionGeometryCallback InCreateGeometryCallback) { CreateTraceCollisionGeometryCallback = InCreateGeometryCallback; }
 
 protected:
 
@@ -620,22 +619,9 @@ private:
 	TArray<int32> ValidGeometryTransformIndices;
 
 	// todo(chaos): Remove this and move to a cook time approach of the SM data based on the GC property
-	// These are created off of the static meshes in the event we want to use SM collision on the game thread
-	struct FTraceCollisionOverrideData
-	{
-		FTraceCollisionOverrideData() = default;
-
-		FTraceCollisionOverrideData(const FGeometryAddParams& InParams, FTraceCollisionCreateGeometryCallback InCallback)
-		: Params(InParams)
-		, CreateGeometryCallback(InCallback)
-		{}
-
-		FGeometryAddParams Params;
-		FTraceCollisionCreateGeometryCallback CreateGeometryCallback;
-	};
-	TArray<FTraceCollisionOverrideData> OptionalTraceCollisionOverrideGeoms;
+	FCreateTraceCollisionGeometryCallback CreateTraceCollisionGeometryCallback;
 	
-	// Use SetUseStaticMeshCollisionForTraces_External intead of setting this directly, so that the collision can be recreated if needed
+	// Use SetUseStaticMeshCollisionForTraces_External instead of setting this directly, so that the collision can be recreated if needed
 	bool bUseStaticMeshCollisionForTraces = false;
 
 #ifdef TODO_REIMPLEMENT_RIGID_CACHING
diff --git a/Engine/Source/Runtime/Experimental/GeometryCollectionEngine/Private/GeometryCollection/GeometryCollectionComponent.cpp b/Engine/Source/Runtime/Experimental/GeometryCollectionEngine/Private/GeometryCollection/GeometryCollectionComponent.cpp
--- a/Engine/Source/Runtime/Experimental/GeometryCollectionEngine/Private/GeometryCollection/GeometryCollectionComponent.cpp
+++ b/Engine/Source/Runtime/Experimental/GeometryCollectionEngine/Private/GeometryCollection/GeometryCollectionComponent.cpp
@@ -3816,7 +3816,30 @@ void UGeometryCollectionComponent::RegisterAndInitializePhysicsProxy()
 		}
 	}
 
-	// Always set up the overrides in case the boolean flips later to allow SM collision on GT
+	auto CreateTraceCollisionFromStaticMeshGeometry = [this](const FTransform& InToLocal, TArray<Chaos::FImplicitObjectPtr>& OutGeoms, Chaos::FShapesArray& OutShapes) {
+		for (const TObjectPtr<UStaticMesh>& ProxyMesh : RestCollection->RootProxyData.ProxyMeshes)
+		{
+			// If these were null we wouldn't have set this callback in the first place
+			const UStaticMesh* StaticMesh = ProxyMesh.Get();
+			UBodySetup* BodySetup = ProxyMesh->GetBodySetup();
+
+			FBodyCollisionData BodyCollisionData;
+			BodyInstance.BuildBodyFilterData(BodyCollisionData.CollisionFilterData);
+			FBodyInstance::BuildBodyCollisionFlags(BodyCollisionData.CollisionFlags, BodyInstance.GetCollisionEnabled(), BodySetup->GetCollisionTraceFlag() == CTF_UseComplexAsSimple);
+
+			FGeometryAddParams AddParams;
+			AddParams.bDoubleSided = BodySetup->bDoubleSidedGeometry;
+			AddParams.CollisionData = BodyCollisionData;
+			AddParams.CollisionTraceType = BodySetup->GetCollisionTraceFlag();
+			AddParams.Scale = BodyInstance.Scale3D;
+			AddParams.SimpleMaterial = BodyInstance.GetSimplePhysicalMaterial();
+			AddParams.LocalTransform = InToLocal;
+			AddParams.WorldTransform = PhysicsProxy->GetSimParameters().WorldTransform;
+			AddParams.Geometry = &BodySetup->AggGeom;
+			ChaosInterface::CreateGeometry(AddParams, OutGeoms, OutShapes);
+		}
+	};
+
 	// todo(chaos): Remove this and move to a cook time approach of the SM data based on the GC property
 	if (RestCollection != nullptr)
 	{
@@ -3826,33 +3849,14 @@ void UGeometryCollectionComponent::RegisterAndInitializePhysicsProxy()
 			{
 				if (UBodySetup* BodySetup = ProxyMesh->GetBodySetup())
 				{
-					FBodyCollisionData BodyCollisionData;
-					BodyInstance.BuildBodyFilterData(BodyCollisionData.CollisionFilterData);
-					FBodyInstance::BuildBodyCollisionFlags(BodyCollisionData.CollisionFlags, BodyInstance.GetCollisionEnabled(), BodySetup->GetCollisionTraceFlag() == CTF_UseComplexAsSimple);
-
-					FGeometryAddParams AddParams;
-					AddParams.bDoubleSided = BodySetup->bDoubleSidedGeometry;
-					AddParams.CollisionData = BodyCollisionData;
-					AddParams.CollisionTraceType = BodySetup->GetCollisionTraceFlag();
-					AddParams.Scale = BodyInstance.Scale3D;
-					AddParams.SimpleMaterial = BodyInstance.GetSimplePhysicalMaterial();
-					AddParams.LocalTransform = FTransform::Identity;
-					AddParams.WorldTransform = PhysicsProxy->GetSimParameters().WorldTransform;
-					AddParams.Geometry = &BodySetup->AggGeom;
-
-					PhysicsProxy->RegisterNewTraceCollisionOverrideData(
-						AddParams,
-						[](const FGeometryAddParams& InParams, TArray<Chaos::FImplicitObjectPtr>& OutGeoms, Chaos::FShapesArray& OutShapes) {
-							ChaosInterface::CreateGeometry(InParams, OutGeoms, OutShapes);
-						}
-					);
+					// We have at least one valid mesh, so set the callback in case we are asked to use SM collision for traces
+					PhysicsProxy->SetCreateTraceCollisionGeometryCallback(CreateTraceCollisionFromStaticMeshGeometry);
+					break;
 				}
 			}
 		}
 	}
 
-	PhysicsProxy->SetUseStaticMeshCollisionForTraces_External(bUseStaticMeshCollisionForTraces);
-
 	FPhysScene_Chaos* PhysicsScene = GetInnerChaosScene();
 	if (ensure(PhysicsScene))
 	{
