Particle data stored in single precision.
104 bytes have been remove by particle, using single precision for most of data, but position.

diff --git a/Engine/Plugins/Runtime/NetworkPrediction/Source/NetworkPrediction/Private/NetworkPredictionPhysics.cpp b/Engine/Plugins/Runtime/NetworkPrediction/Source/NetworkPrediction/Private/NetworkPredictionPhysics.cpp
--- a/Engine/Plugins/Runtime/NetworkPrediction/Source/NetworkPrediction/Private/NetworkPredictionPhysics.cpp
+++ b/Engine/Plugins/Runtime/NetworkPrediction/Source/NetworkPrediction/Private/NetworkPredictionPhysics.cpp
@@ -43,7 +43,9 @@ void FNetworkPredictionPhysicsState::NetSend(const FNetSerializeParams& P, FBody
 	{
 		npEnsure(Body_External.CanTreatAsKinematic());
 		P.Ar << const_cast<Chaos::FVec3&>(Body_External.X());
-		P.Ar << const_cast<Chaos::FRotation3&>(Body_External.R());
+		Chaos::FRotation3 Rotation = Body_External.R();
+		P.Ar << Rotation;
+		Body_External.SetR(Rotation);
 		Chaos::FVec3 V = Body_External.V();
 		Chaos::FVec3 W = Body_External.W();
 		P.Ar << V;
@@ -54,7 +56,9 @@ void FNetworkPredictionPhysicsState::NetSend(const FNetSerializeParams& P, FBody
 		bool bSuccess = true;
 
 		SerializePackedVector<100, 30>(const_cast<Chaos::FVec3&>(Body_External.X()), P.Ar);
-		const_cast<Chaos::FRotation3&>(Body_External.R()).NetSerialize(P.Ar, nullptr, bSuccess);
+		Chaos::FRotation3 Rotation = Body_External.R();
+		Rotation.NetSerialize(P.Ar, nullptr, bSuccess);
+		Body_External.SetR(Rotation);
 
 		npEnsure(Body_External.CanTreatAsKinematic());
 		Chaos::FVec3 V = Body_External.V();
diff --git a/Engine/Source/Runtime/Core/Public/UObject/FortniteReleaseBranchCustomObjectVersion.h b/Engine/Source/Runtime/Core/Public/UObject/FortniteReleaseBranchCustomObjectVersion.h
--- a/Engine/Source/Runtime/Core/Public/UObject/FortniteReleaseBranchCustomObjectVersion.h
+++ b/Engine/Source/Runtime/Core/Public/UObject/FortniteReleaseBranchCustomObjectVersion.h
@@ -45,6 +45,9 @@ struct FFortniteReleaseBranchCustomObjectVersion
 		// Enable SkipOnlyEditorOnly style cooking of NaniteOverrideMaterial
 		NaniteMaterialOverrideUsesEditorOnly,
 
+		// Store game thread particles data in single precision
+		SinglePrecisonParticleData,
+
 		// -----<new versions can be added above this line>-------------------------------------------------
 		VersionPlusOne,
 		LatestVersion = VersionPlusOne - 1
diff --git a/Engine/Source/Runtime/Experimental/Chaos/Public/Chaos/Particle/ParticleUtilities.h b/Engine/Source/Runtime/Experimental/Chaos/Public/Chaos/Particle/ParticleUtilities.h
--- a/Engine/Source/Runtime/Experimental/Chaos/Public/Chaos/Particle/ParticleUtilities.h
+++ b/Engine/Source/Runtime/Experimental/Chaos/Public/Chaos/Particle/ParticleUtilities.h
@@ -24,9 +24,10 @@ namespace Chaos
 		template<typename T_PARTICLEHANDLE>
 		static inline const FVec3& GetPosition(T_PARTICLEHANDLE Particle) { return Particle->X(); }
 		template<typename T_PARTICLEHANDLE>
-		static inline FRotation3& Rotation(T_PARTICLEHANDLE Particle) { return Particle->R(); }
+		UE_DEPRECATED(5.4, "Use Particle->SetR() instead.")
+		static inline FRotation3 Rotation(T_PARTICLEHANDLE Particle) { return Particle->R(); }
 		template<typename T_PARTICLEHANDLE>
-		static inline const FRotation3& GetRotation(T_PARTICLEHANDLE Particle) { return Particle->R(); }
+		static inline const FRotation3 GetRotation(T_PARTICLEHANDLE Particle) { return Particle->R(); }
 
 		static inline FVec3& Position(TPBDRigidParticles<FReal, 3>& Particles, const int32 Index) { return Particles.X(Index); }
 		static inline FRotation3& Rotation(TPBDRigidParticles<FReal, 3>& Particles, const int32 Index) { return Particles.R(Index); }
diff --git a/Engine/Source/Runtime/Experimental/Chaos/Public/Chaos/ParticleDirtyFlags.h b/Engine/Source/Runtime/Experimental/Chaos/Public/Chaos/ParticleDirtyFlags.h
--- a/Engine/Source/Runtime/Experimental/Chaos/Public/Chaos/ParticleDirtyFlags.h
+++ b/Engine/Source/Runtime/Experimental/Chaos/Public/Chaos/ParticleDirtyFlags.h
@@ -62,20 +62,32 @@ class FParticlePositionRotation
 public:
 	void Serialize(FChaosArchive& Ar)
 	{
-		Ar << MX << MR;
+		Ar.UsingCustomVersion(FFortniteReleaseBranchCustomObjectVersion::GUID);
+
+		if (Ar.CustomVer(FFortniteReleaseBranchCustomObjectVersion::GUID) >= FFortniteReleaseBranchCustomObjectVersion::SinglePrecisonParticleData)
+		{
+			Ar << MX << MR;
+		}
+		else
+		{
+			FRotation3 DoublePrecisionRotation = FRotation3(MR);
+			Ar << MX << DoublePrecisionRotation;
+			MR = FRotation3f(MR);
+		}
+		
 	}
 
 	template <typename TOther>
 	void CopyFrom(const TOther& Other)
 	{
 		MX = Other.X();
-		MR = Other.R();
+		MR = FRotation3f(Other.R());
 	}
 
 	template <typename TOther>
 	bool IsEqual(const TOther& Other) const
 	{
-		return MX == Other.X() && MR == Other.R();
+		return MX == Other.X() && MR == FRotation3f(Other.R());
 	}
 
 	bool operator==(const FParticlePositionRotation& Other) const
@@ -86,12 +98,12 @@ public:
 	const FVec3& X() const { return MX; }
 	void SetX(const FVec3& InX){ MX = InX; }
 
-	const FRotation3& R() const { return MR; }
-	void SetR(const FRotation3& InR) { MR = InR; }
+	const FRotation3 R() const { return FRotation3(MR); }
+	void SetR(const FRotation3& InR) { MR = FRotation3f(InR); }
 	
 private:
 	FVec3 MX;
-	FRotation3 MR;
+	FRotation3f MR;
 
 };
 
@@ -106,7 +118,21 @@ class FParticleVelocities
 public:
 	void Serialize(FChaosArchive& Ar)
 	{
-		Ar << MV << MW;
+		Ar.UsingCustomVersion(FFortniteReleaseBranchCustomObjectVersion::GUID);
+
+		if (Ar.CustomVer(FFortniteReleaseBranchCustomObjectVersion::GUID) >= FFortniteReleaseBranchCustomObjectVersion::SinglePrecisonParticleData)
+		{
+			Ar << MV << MW;
+		}
+		else
+		{
+			FVec3 MVDouble(MV);
+			FVec3 MWDouble(MW);
+			Ar << MVDouble << MWDouble;
+			MV = FVec3f(MVDouble);
+			MW = FVec3f(MWDouble);
+		}
+
 	}
 
 	template <typename TOther>
@@ -119,7 +145,7 @@ public:
 	template <typename TOther>
 	bool IsEqual(const TOther& Other) const
 	{
-		return MV == Other.V() && MW == Other.W();
+		return MV == FVec3f(Other.V()) && MW == FVec3f(Other.W());
 	}
 
 	bool operator==(const FParticleVelocities& Other) const
@@ -127,15 +153,15 @@ public:
 		return IsEqual(Other);
 	}
 
-	const FVec3& V() const { return MV; }
-	void SetV(const FVec3& V) { MV = V; }
+	const FVec3 V() const { return FVec3(MV); }
+	void SetV(const FVec3& V) { MV = FVec3f(V); }
 
-	const FVec3& W() const { return MW; }
-	void SetW(const FVec3& W){ MW = W; }
+	const FVec3 W() const { return FVec3(MW); }
+	void SetW(const FVec3& W){ MW = FVec3f(W); }
 
 private:
-	FVec3 MV;
-	FVec3 MW;
+	FVec3f MV;
+	FVec3f MW;
 };
 
 inline FChaosArchive& operator<<(FChaosArchive& Ar,FParticleVelocities& Data)
@@ -149,10 +175,30 @@ class FParticleDynamics
 public:
 	void Serialize(FChaosArchive& Ar)
 	{
-		Ar << MAcceleration;
-		Ar << MAngularAcceleration;
-		Ar << MLinearImpulseVelocity;
-		Ar << MAngularImpulseVelocity;	
+		Ar.UsingCustomVersion(FFortniteReleaseBranchCustomObjectVersion::GUID);
+
+		if (Ar.CustomVer(FFortniteReleaseBranchCustomObjectVersion::GUID) >= FFortniteReleaseBranchCustomObjectVersion::SinglePrecisonParticleData)
+		{
+			Ar << MAcceleration;
+			Ar << MAngularAcceleration;
+			Ar << MLinearImpulseVelocity;
+			Ar << MAngularImpulseVelocity;
+		}
+		else
+		{
+			FVec3 AccelerationDouble(MAcceleration);
+			FVec3 AngularAccelerationDouble(MAngularAcceleration);
+			FVec3 LinearImpulseVelocityDouble(MLinearImpulseVelocity);
+			FVec3 AngularImpulseVelocityDouble(MAngularImpulseVelocity);
+			Ar << AccelerationDouble;
+			Ar << AngularAccelerationDouble;
+			Ar << LinearImpulseVelocityDouble;
+			Ar << AngularImpulseVelocityDouble;
+			MAcceleration = FVec3f(AccelerationDouble);
+			MAngularAcceleration = FVec3f(AngularAccelerationDouble);
+			MLinearImpulseVelocity = FVec3f(LinearImpulseVelocityDouble);
+			MAngularImpulseVelocity = FVec3f(AngularImpulseVelocityDouble);
+		}
 	}
 
 	template <typename TOther>
@@ -178,17 +224,17 @@ public:
 		return IsEqual(Other);
 	}
 
-	const FVec3& Acceleration() const { return MAcceleration; }
-	void SetAcceleration(const FVec3& Acceleration){ MAcceleration = Acceleration; }
+	FVec3 Acceleration() const { return FVec3(MAcceleration); }
+	void SetAcceleration(const FVec3& Acceleration){ MAcceleration = FVec3f(Acceleration); }
 
-	const FVec3& AngularAcceleration() const { return MAngularAcceleration; }
-	void SetAngularAcceleration(const FVec3& AngularAcceleration){ MAngularAcceleration = AngularAcceleration; }
+	FVec3 AngularAcceleration() const { return FVec3(MAngularAcceleration); }
+	void SetAngularAcceleration(const FVec3& AngularAcceleration){ MAngularAcceleration = FVec3f(AngularAcceleration); }
 
-	const FVec3& LinearImpulseVelocity() const { return MLinearImpulseVelocity; }
-	void SetLinearImpulseVelocity(const FVec3& LinearImpulseVelocity){ MLinearImpulseVelocity = LinearImpulseVelocity; }
+	FVec3 LinearImpulseVelocity() const { return FVec3(MLinearImpulseVelocity); }
+	void SetLinearImpulseVelocity(const FVec3& LinearImpulseVelocity){ MLinearImpulseVelocity = FVec3f(LinearImpulseVelocity); }
 
-	const FVec3& AngularImpulseVelocity() const { return MAngularImpulseVelocity; }
-	void SetAngularImpulseVelocity(const FVec3& AngularImpulseVelocity){ MAngularImpulseVelocity = AngularImpulseVelocity; }
+	FVec3 AngularImpulseVelocity() const { return FVec3(MAngularImpulseVelocity); }
+	void SetAngularImpulseVelocity(const FVec3& AngularImpulseVelocity){ MAngularImpulseVelocity = FVec3f(AngularImpulseVelocity); }
 
 	static FParticleDynamics ZeroValue()
 	{
@@ -202,10 +248,10 @@ public:
 	}
 
 private:
-	FVec3 MAcceleration;
-	FVec3 MAngularAcceleration;
-	FVec3 MLinearImpulseVelocity;
-	FVec3 MAngularImpulseVelocity;
+	FVec3f MAcceleration;
+	FVec3f MAngularAcceleration;
+	FVec3f MLinearImpulseVelocity;
+	FVec3f MAngularImpulseVelocity;
 
 };
 
@@ -324,13 +370,29 @@ public:
 		Ar.UsingCustomVersion(FUE5MainStreamObjectVersion::GUID);
 		Ar.UsingCustomVersion(FExternalPhysicsCustomObjectVersion::GUID);
 		Ar.UsingCustomVersion(FPhysicsObjectVersion::GUID);
+		Ar.UsingCustomVersion(FFortniteReleaseBranchCustomObjectVersion::GUID);
 
-		// Flags moved into a bitmask
-		const bool bAddControlFlags = (Ar.CustomVer(FUE5MainStreamObjectVersion::GUID) >= FUE5MainStreamObjectVersion::AddRigidParticleControlFlags);
+		const bool bSinglePrecision = Ar.CustomVer(FFortniteReleaseBranchCustomObjectVersion::GUID) >= FFortniteReleaseBranchCustomObjectVersion::SinglePrecisonParticleData;
 
-		Ar << MLinearEtherDrag;
-		Ar << MAngularEtherDrag;
+		if (bSinglePrecision)
+		{
+			Ar << MLinearEtherDrag;
+			Ar << MAngularEtherDrag;
+		}
+		else
+		{
+			FReal LinearEtherDragDouble = FReal(MLinearEtherDrag);
+			FReal AngularEtherDragDouble = FReal(MAngularEtherDrag);
+			Ar << LinearEtherDragDouble;
+			Ar << AngularEtherDragDouble;
+			MLinearEtherDrag = FRealSingle(LinearEtherDragDouble);
+			MAngularEtherDrag = FRealSingle(AngularEtherDragDouble);
+		}
+		
 		Ar << MObjectState;
+
+		// Flags moved into a bitmask
+		const bool bAddControlFlags = (Ar.CustomVer(FUE5MainStreamObjectVersion::GUID) >= FUE5MainStreamObjectVersion::AddRigidParticleControlFlags);
 		if (!bAddControlFlags && Ar.IsLoading())
 		{
 			bool bGravityEnabled;
@@ -376,8 +438,20 @@ public:
 		const bool bAddChaosMaxLinearAngularSpeedUE5 = (Ar.CustomVer(FUE5ReleaseStreamObjectVersion::GUID) >= FUE5ReleaseStreamObjectVersion::AddChaosMaxLinearAngularSpeed);
 		if (bAddChaosMaxLinearAngularSpeedUE4 || bAddChaosMaxLinearAngularSpeedUE5)
 		{
-			Ar << MMaxLinearSpeedSq;
-			Ar << MMaxAngularSpeedSq;
+			if (bSinglePrecision)
+			{
+				Ar << MMaxLinearSpeedSq;
+				Ar << MMaxAngularSpeedSq;
+			}
+			else
+			{
+				FReal MaxLinearSpeedSqDouble = FReal(MMaxLinearSpeedSq);
+				FReal MaxAngularSpeedSqDouble = FReal(MMaxAngularSpeedSq);
+				Ar << MaxLinearSpeedSqDouble;
+				Ar << MaxAngularSpeedSqDouble;
+				MMaxLinearSpeedSq = FRealSingle(MaxLinearSpeedSqDouble);
+				MMaxAngularSpeedSq = FRealSingle(MaxAngularSpeedSqDouble);
+			}
 		}
 
 		// @todo(chaos): add this
@@ -413,10 +487,10 @@ public:
 	bool IsEqual(const TOther& Other) const
 	{
 		return ObjectState() == Other.ObjectState()
-			&& LinearEtherDrag() == Other.LinearEtherDrag()
-			&& AngularEtherDrag() == Other.AngularEtherDrag()
-			&& MaxLinearSpeedSq() == Other.MaxLinearSpeedSq()
-			&& MaxAngularSpeedSq() == Other.MaxAngularSpeedSq()
+			&& MLinearEtherDrag == FRealSingle(Other.LinearEtherDrag())
+			&& MAngularEtherDrag == FRealSingle(Other.AngularEtherDrag())
+			&& MMaxLinearSpeedSq == FRealSingle(Other.MaxLinearSpeedSq())
+			&& MMaxAngularSpeedSq == FRealSingle(Other.MaxAngularSpeedSq())
 			&& InitialOverlapDepenetrationVelocity() == Other.InitialOverlapDepenetrationVelocity()
 			&& CollisionGroup() == Other.CollisionGroup()
 			&& SleepType() == Other.SleepType()
@@ -430,17 +504,17 @@ public:
 		return IsEqual(Other);
 	}
 
-	FReal LinearEtherDrag() const { return MLinearEtherDrag; }
-	void SetLinearEtherDrag(FReal InLinearEtherDrag) { MLinearEtherDrag = InLinearEtherDrag; }
+	FReal LinearEtherDrag() const { return FReal(MLinearEtherDrag); }
+	void SetLinearEtherDrag(FReal InLinearEtherDrag) { MLinearEtherDrag = FRealSingle(InLinearEtherDrag); }
 
-	FReal AngularEtherDrag() const { return MAngularEtherDrag; }
-	void SetAngularEtherDrag(FReal InAngularEtherDrag) { MAngularEtherDrag = InAngularEtherDrag; }
+	FReal AngularEtherDrag() const { return FReal(MAngularEtherDrag); }
+	void SetAngularEtherDrag(FReal InAngularEtherDrag) { MAngularEtherDrag = FRealSingle(InAngularEtherDrag); }
 
-	FReal MaxLinearSpeedSq() const { return MMaxLinearSpeedSq; }
-	void SetMaxLinearSpeedSq(FReal InMaxLinearSpeed) { MMaxLinearSpeedSq = InMaxLinearSpeed; }
+	FReal MaxLinearSpeedSq() const { return FReal(MMaxLinearSpeedSq); }
+	void SetMaxLinearSpeedSq(FReal InMaxLinearSpeed) { MMaxLinearSpeedSq = FRealSingle(InMaxLinearSpeed); }
 
-	FReal MaxAngularSpeedSq() const { return MMaxAngularSpeedSq; }
-	void SetMaxAngularSpeedSq(FReal InMaxAngularSpeed) { MMaxAngularSpeedSq = InMaxAngularSpeed; }
+	FReal MaxAngularSpeedSq() const { return FReal(MMaxAngularSpeedSq); }
+	void SetMaxAngularSpeedSq(FReal InMaxAngularSpeed) { MMaxAngularSpeedSq = FRealSingle(InMaxAngularSpeed); }
 
 	FRealSingle InitialOverlapDepenetrationVelocity() const { return MInitialOverlapDepenetrationVelocity; }
 	void SetInitialOverlapDepenetrationVelocity(FRealSingle InVel) { MInitialOverlapDepenetrationVelocity = InVel; }
@@ -484,10 +558,10 @@ private:
 	//NOTE: MObjectState is the only sim-writable data in this struct
 	//If you add any more, make sure to update SyncSimWritablePropsFromSim
 	//Or consider breaking it (and object state) out of this struct entirely
-	FReal MLinearEtherDrag;
-	FReal MAngularEtherDrag;
-	FReal MMaxLinearSpeedSq;
-	FReal MMaxAngularSpeedSq;
+	FRealSingle MLinearEtherDrag;
+	FRealSingle MAngularEtherDrag;
+	FRealSingle MMaxLinearSpeedSq;
+	FRealSingle MMaxAngularSpeedSq;
 	FRealSingle MInitialOverlapDepenetrationVelocity = 0;
 	int32 MCollisionGroup;
 
diff --git a/Engine/Source/Runtime/Experimental/Chaos/Public/Chaos/ParticleHandle.h b/Engine/Source/Runtime/Experimental/Chaos/Public/Chaos/ParticleHandle.h
--- a/Engine/Source/Runtime/Experimental/Chaos/Public/Chaos/ParticleHandle.h
+++ b/Engine/Source/Runtime/Experimental/Chaos/Public/Chaos/ParticleHandle.h
@@ -2565,7 +2565,7 @@ public:
 		MNonFrequentData.Modify(bInvalidate, MDirtyFlags, Proxy, [ParticleID](auto& Data) { Data.SetParticleID(ParticleID); });
 	}
 
-	const TRotation<T, d>& R() const { return MXR.Read().R(); }
+	const TRotation<T, d> R() const { return MXR.Read().R(); }
 	void SetR(const TRotation<T, d>& InR, bool bInvalidate = true);
 
 	const FParticlePositionRotation& XR() const { return MXR.Read(); }
@@ -3007,10 +3007,10 @@ public:
 		//Ar << MKinematicTarget; // TODO
 	}
 
-	const TVector<T, d>& V() const { return MVelocities.Read().V(); }
+	const TVector<T, d> V() const { return MVelocities.Read().V(); }
 	void SetV(const TVector<T, d>& InV, bool bInvalidate = true);
 
-	const TVector<T, d>& W() const { return MVelocities.Read().W(); }
+	const TVector<T, d> W() const { return MVelocities.Read().W(); }
 	void SetW(const TVector<T, d>& InW, bool bInvalidate = true);
 
 	const FKinematicTarget KinematicTarget() const {
@@ -3186,7 +3186,7 @@ public:
 		this->MInitialized = InInitialized;
 	}
 
-	const TVector<T, d>& Acceleration() const { return MDynamics.Read().Acceleration(); }
+	const TVector<T, d> Acceleration() const { return MDynamics.Read().Acceleration(); }
 	void SetAcceleration(const FVec3& Acceleration, bool bInvalidate = true)
 	{ 
 		MDynamics.Modify(bInvalidate, MDirtyFlags, Proxy, [&Acceleration](auto& Data) { Data.SetAcceleration(Acceleration); });
@@ -3215,7 +3215,7 @@ public:
 		}
 	}
 
-	const TVector<T, d>& AngularAcceleration() const { return MDynamics.Read().AngularAcceleration(); }
+	const TVector<T, d> AngularAcceleration() const { return MDynamics.Read().AngularAcceleration(); }
 	void SetAngularAcceleration(const TVector<T, d>& InTorque, bool bInvalidate = true)
 	{
 		MDynamics.Modify(bInvalidate, MDirtyFlags, Proxy, [&InTorque](auto& Data) { Data.SetAngularAcceleration(InTorque);});
@@ -3227,13 +3227,13 @@ public:
 		MDynamics.Modify(bInvalidate, MDirtyFlags, Proxy, [](auto& Data) { Data.SetAngularAcceleration(FVec3(0)); });
 	}
 
-	const TVector<T, d>& LinearImpulseVelocity() const { return MDynamics.Read().LinearImpulseVelocity(); }
+	const TVector<T, d> LinearImpulseVelocity() const { return MDynamics.Read().LinearImpulseVelocity(); }
 	void SetLinearImpulseVelocity(const TVector<T, d>& InLinearImpulseVelocity, bool bInvalidate = true)
 	{
 		MDynamics.Modify(bInvalidate,MDirtyFlags,Proxy,[&InLinearImpulseVelocity](auto& Data){ Data.SetLinearImpulseVelocity(InLinearImpulseVelocity);});
 	}
 
-	const TVector<T, d>& AngularImpulseVelocity() const { return MDynamics.Read().AngularImpulseVelocity(); }
+	const TVector<T, d> AngularImpulseVelocity() const { return MDynamics.Read().AngularImpulseVelocity(); }
 	void SetAngularImpulseVelocity(const TVector<T, d>& InAngularImpulseVelocity, bool bInvalidate = true)
 	{
 		MDynamics.Modify(bInvalidate,MDirtyFlags,Proxy,[&InAngularImpulseVelocity](auto& Data){ Data.SetAngularImpulseVelocity(InAngularImpulseVelocity);});
diff --git a/Engine/Source/Runtime/Experimental/Chaos/Public/PhysicsProxy/SingleParticlePhysicsProxy.h b/Engine/Source/Runtime/Experimental/Chaos/Public/PhysicsProxy/SingleParticlePhysicsProxy.h
--- a/Engine/Source/Runtime/Experimental/Chaos/Public/PhysicsProxy/SingleParticlePhysicsProxy.h
+++ b/Engine/Source/Runtime/Experimental/Chaos/Public/PhysicsProxy/SingleParticlePhysicsProxy.h
@@ -228,7 +228,7 @@ public:
 	FUniqueIdx UniqueIdx() const { return Read([](auto* Particle) { return Particle->UniqueIdx(); }); }
 	void SetUniqueIdx(const FUniqueIdx UniqueIdx, bool bInvalidate = true) { Write([UniqueIdx, bInvalidate](auto* Particle) { Particle->SetUniqueIdx(UniqueIdx, bInvalidate); }); }
 
-	const FRotation3& R() const { return ReadRef([](auto* Particle) -> const auto& { return Particle->R(); }); }
+	FRotation3 R() const { return Read([](auto* Particle) -> const auto { return Particle->R(); }); }
 
 protected:
 	void SetRBase(const FRotation3& InR, bool bInvalidate = true){ Write([&InR, bInvalidate, this](auto* Particle)
