Chaos cloth: Make collision particles a new stripped down particle SOA that only has the fields we need.

Previously, we were using TKinematicGeometryParticlesImp which have a ton of fields we don't actually use. Create a new SimpleGeometryParticles SOA that GeometryParticlesImp and our SolverCollisionParticles can both derive from.

diff --git a/Engine/Plugins/Experimental/ChaosCloth/Source/Chaos/Private/ChaosCloth/ChaosClothingSimulationSolver.cpp b/Engine/Plugins/Experimental/ChaosCloth/Source/Chaos/Private/ChaosCloth/ChaosClothingSimulationSolver.cpp
--- a/Engine/Plugins/Experimental/ChaosCloth/Source/Chaos/Private/ChaosCloth/ChaosClothingSimulationSolver.cpp
+++ b/Engine/Plugins/Experimental/ChaosCloth/Source/Chaos/Private/ChaosCloth/ChaosClothingSimulationSolver.cpp
@@ -114,7 +114,7 @@ FClothingSimulationSolver::FClothingSimulationSolver()
 	SetConfig(nullptr); // This will generate a local default config so we have something to use if there are no cloth assets..
 
 	Softs::FSolverParticles LocalParticles;
-	Softs::FSolverRigidParticles RigidParticles;
+	Softs::FSolverCollisionParticles RigidParticles;
 	Evolution.Reset(
 		new Softs::FPBDEvolution(
 			MoveTemp(LocalParticles),
@@ -153,7 +153,7 @@ FClothingSimulationSolver::FClothingSimulationSolver()
 		});
 
 	Evolution->SetCollisionKinematicUpdateFunction(
-		[this](Softs::FSolverRigidParticles& ParticlesInput, const Softs::FSolverReal Dt, const Softs::FSolverReal LocalTime, const int32 Index)
+		[this](Softs::FSolverCollisionParticles& ParticlesInput, const Softs::FSolverReal Dt, const Softs::FSolverReal LocalTime, const int32 Index)
 		{
 			checkSlow(Dt > SMALL_NUMBER && DeltaTime > SMALL_NUMBER);
 			const Softs::FSolverReal Alpha = (LocalTime - Time) / DeltaTime;
@@ -1037,11 +1037,11 @@ void FClothingSimulationSolver::ApplyPreSimulationTransforms()
 		TRACE_CPUPROFILER_EVENT_SCOPE(FClothingSimulationSolver_CollisionPreSimulationTransforms);
 		SCOPE_CYCLE_COUNTER(STAT_ChaosClothCollisionPreSimulationTransforms);
 
-		const TPBDActiveView<Softs::FSolverRigidParticles>& CollisionParticlesActiveView = Evolution->CollisionParticlesActiveView();
+		const TPBDActiveView<Softs::FSolverCollisionParticles>& CollisionParticlesActiveView = Evolution->CollisionParticlesActiveView();
 		const TArray<uint32>& CollisionParticleGroupIds = Evolution->CollisionParticleGroupIds();
 
 		CollisionParticlesActiveView.SequentialFor(  // There's unlikely to ever have enough collision particles for a parallel for
-			[this, &CollisionParticleGroupIds, &DeltaLocalSpaceLocation](Softs::FSolverRigidParticles& CollisionParticles, int32 Index)
+			[this, &CollisionParticleGroupIds, &DeltaLocalSpaceLocation](Softs::FSolverCollisionParticles& CollisionParticles, int32 Index)
 			{
 				const Softs::FSolverRigidTransform3& GroupSpaceTransform = PreSimulationTransforms[CollisionParticleGroupIds[Index]];
 
diff --git a/Engine/Source/Runtime/Experimental/Chaos/Private/Chaos/Deformable/ChaosDeformableSolver.cpp b/Engine/Source/Runtime/Experimental/Chaos/Private/Chaos/Deformable/ChaosDeformableSolver.cpp
--- a/Engine/Source/Runtime/Experimental/Chaos/Private/Chaos/Deformable/ChaosDeformableSolver.cpp
+++ b/Engine/Source/Runtime/Experimental/Chaos/Private/Chaos/Deformable/ChaosDeformableSolver.cpp
@@ -124,7 +124,7 @@ namespace Chaos::Softs
 		Property = InProps;
 		MObjects = TArrayCollectionArray<const UObject*>();
 		FSolverParticles LocalParticlesDummy;
-		FSolverRigidParticles RigidParticles;
+		FSolverCollisionParticles RigidParticles;
 		Evolution.Reset(new FPBDEvolution(MoveTemp(LocalParticlesDummy), MoveTemp(RigidParticles), {},
 			Property.NumSolverIterations, (FSolverReal)0.,
 			/*SelfCollisionsThickness = */(FSolverReal)0.,
@@ -1677,7 +1677,7 @@ namespace Chaos::Softs
 		if (Evolution && GDeformableDebugParams.bDoDrawRigidCollisionGeometry)
 		{
 			Evolution->CollisionParticlesActiveView().RangeFor(
-				[this, ToFVec3, ToFVector, ToFQuat](FSolverRigidParticles& CollisionParticles, int32 CollisionOffset, int32 CollisionRange)
+				[this, ToFVec3, ToFVector, ToFQuat](FSolverCollisionParticles& CollisionParticles, int32 CollisionOffset, int32 CollisionRange)
 				{
 					for (int32 Index = CollisionOffset; Index < CollisionRange; Index++)
 					{
diff --git a/Engine/Source/Runtime/Experimental/Chaos/Private/Chaos/GeometryParticles.cpp b/Engine/Source/Runtime/Experimental/Chaos/Private/Chaos/GeometryParticles.cpp
--- a/Engine/Source/Runtime/Experimental/Chaos/Private/Chaos/GeometryParticles.cpp
+++ b/Engine/Source/Runtime/Experimental/Chaos/Private/Chaos/GeometryParticles.cpp
@@ -59,7 +59,7 @@ namespace Chaos
 	template <typename T, int d, EGeometryParticlesSimType SimType>
 	void TGeometryParticlesImp<T, d, SimType>::UpdateShapesArray(const int32 Index)
 	{
-		UpdateShapesArrayFromGeometry(MShapesArray[Index], MGeometry[Index], FRigidTransform3(X(Index), R(Index)));
+		UpdateShapesArrayFromGeometry(MShapesArray[Index], GetGeometry(Index), FRigidTransform3(X(Index), R(Index)));
 	}
 
 	template <typename T, int d, EGeometryParticlesSimType SimType>
diff --git a/Engine/Source/Runtime/Experimental/Chaos/Private/Chaos/NewtonEvolution.cpp b/Engine/Source/Runtime/Experimental/Chaos/Private/Chaos/NewtonEvolution.cpp
--- a/Engine/Source/Runtime/Experimental/Chaos/Private/Chaos/NewtonEvolution.cpp
+++ b/Engine/Source/Runtime/Experimental/Chaos/Private/Chaos/NewtonEvolution.cpp
@@ -75,7 +75,7 @@ void FNewtonEvolution::ResetGroups()
 
 FNewtonEvolution::FNewtonEvolution(
 	FSolverParticles&& InParticles,
-	FSolverRigidParticles&& InGeometryParticles,
+	FSolverCollisionParticles&& InGeometryParticles,
 	TArray<TVec3<int32>>&& CollisionTriangles,
 	const TArray<TVector<int32, 4>>& InMesh,
 	TArray<TArray<TVector<int32, 2>>>&& InIncidentElements,
@@ -776,7 +776,7 @@ void FNewtonEvolution::WriteOutputLog(const int32 Frame)
 				//SCOPE_CYCLE_COUNTER(STAT_ChaosNewtonCollisionKinematicUpdate);
 
 				MCollisionParticlesActiveView.SequentialFor(
-					[this, Dt](FSolverRigidParticles& CollisionParticles, int32 Index)
+					[this, Dt](FSolverCollisionParticles& CollisionParticles, int32 Index)
 					{
 						// Store active collision particle frames prior to the kinematic update for CCD collisions
 						MCollisionTransforms[Index] = FSolverRigidTransform3(CollisionParticles.X(Index), CollisionParticles.R(Index));
diff --git a/Engine/Source/Runtime/Experimental/Chaos/Private/Chaos/PBDEvolution.cpp b/Engine/Source/Runtime/Experimental/Chaos/Private/Chaos/PBDEvolution.cpp
--- a/Engine/Source/Runtime/Experimental/Chaos/Private/Chaos/PBDEvolution.cpp
+++ b/Engine/Source/Runtime/Experimental/Chaos/Private/Chaos/PBDEvolution.cpp
@@ -74,7 +74,7 @@ void FPBDEvolution::ResetGroups()
 
 FPBDEvolution::FPBDEvolution(
 	FSolverParticles&& InParticles,
-	FSolverRigidParticles&& InGeometryParticles,
+	FSolverCollisionParticles&& InGeometryParticles,
 	TArray<TVec3<int32>>&& CollisionTriangles,
 	int32 NumIterations, FSolverReal CollisionThickness,
 	FSolverReal SelfCollisionThickness,
@@ -550,7 +550,7 @@ void FPBDEvolution::AdvanceOneTimeStep(const FSolverReal Dt)
 			SCOPE_CYCLE_COUNTER(STAT_ChaosPBDCollisionKinematicUpdate);
 
 			MCollisionParticlesActiveView.SequentialFor(
-				[this, Dt](FSolverRigidParticles& CollisionParticles, int32 Index)
+				[this, Dt](FSolverCollisionParticles& CollisionParticles, int32 Index)
 				{
 					// Store active collision particle frames prior to the kinematic update for CCD collisions
 					MCollisionTransforms[Index] = FSolverRigidTransform3(CollisionParticles.X(Index), CollisionParticles.R(Index));
diff --git a/Engine/Source/Runtime/Experimental/Chaos/Private/Chaos/PerParticlePBDCollisionConstraint.cpp b/Engine/Source/Runtime/Experimental/Chaos/Private/Chaos/PerParticlePBDCollisionConstraint.cpp
--- a/Engine/Source/Runtime/Experimental/Chaos/Private/Chaos/PerParticlePBDCollisionConstraint.cpp
+++ b/Engine/Source/Runtime/Experimental/Chaos/Private/Chaos/PerParticlePBDCollisionConstraint.cpp
@@ -213,7 +213,7 @@ void FPerParticlePBDCollisionConstraint::ApplyHelperISPC(FSolverParticles& InPar
 
 #if INTEL_ISPC
 		MCollisionParticlesActiveView.RangeFor(
-			[this, &InParticles, Dt, BatchBegin, BatchEnd, DynamicGroupId, PerGroupFriction, PerGroupThickness](FSolverRigidParticles& CollisionParticles, int32 CollisionOffset, int32 CollisionRange)
+			[this, &InParticles, Dt, BatchBegin, BatchEnd, DynamicGroupId, PerGroupFriction, PerGroupThickness](FSolverCollisionParticles& CollisionParticles, int32 CollisionOffset, int32 CollisionRange)
 			{
 				ispc::ApplyPerParticleCollisionFastFriction(
 					(ispc::FVector4f*)InParticles.GetPAndInvM().GetData(),
@@ -249,7 +249,7 @@ void FPerParticlePBDCollisionConstraint::ApplyHelperISPC(FSolverParticles& InPar
 
 #if INTEL_ISPC
 		MCollisionParticlesActiveView.RangeFor(
-			[this, &InParticles, Dt, BatchBegin, BatchEnd, DynamicGroupId, PerGroupThickness](FSolverRigidParticles& CollisionParticles, int32 CollisionOffset, int32 CollisionRange)
+			[this, &InParticles, Dt, BatchBegin, BatchEnd, DynamicGroupId, PerGroupThickness](FSolverCollisionParticles& CollisionParticles, int32 CollisionOffset, int32 CollisionRange)
 			{
 				ispc::ApplyPerParticleCollisionNoFriction(
 					(ispc::FVector4f*)InParticles.GetPAndInvM().GetData(),
diff --git a/Engine/Source/Runtime/Experimental/Chaos/Public/Chaos/GeometryParticles.h b/Engine/Source/Runtime/Experimental/Chaos/Public/Chaos/GeometryParticles.h
--- a/Engine/Source/Runtime/Experimental/Chaos/Public/Chaos/GeometryParticles.h
+++ b/Engine/Source/Runtime/Experimental/Chaos/Public/Chaos/GeometryParticles.h
@@ -1,10 +1,7 @@
 // Copyright Epic Games, Inc. All Rights Reserved.
 #pragma once
 
-#include "Chaos/ArrayCollectionArray.h"
-#include "Chaos/ImplicitObject.h"
-#include "Chaos/Particles.h"
-#include "Chaos/Rotation.h"
+#include "Chaos/SimpleGeometryParticles.h"
 #include "Chaos/ParticleHandleFwd.h"
 #include "Chaos/GeometryParticlesfwd.h"
 #include "Chaos/CollisionFilterData.h"
@@ -152,17 +149,21 @@ namespace Chaos
 	};
 	
 	template<class T, int d, EGeometryParticlesSimType SimType>
-	class TGeometryParticlesImp : public TParticles<T, d>
+	class TGeometryParticlesImp : public TSimpleGeometryParticles<T, d>
 	{
 	public:
 
 		using TArrayCollection::Size;
 		using TParticles<T,d>::X;
+		using TSimpleGeometryParticles<T, d>::R;
+		using TSimpleGeometryParticles<T, d>::GetGeometry;
+		using TSimpleGeometryParticles<T, d>::SetGeometry;
+		using TSimpleGeometryParticles<T, d>::GetAllGeometry;
 
 		CHAOS_API static TGeometryParticlesImp<T, d, SimType>* SerializationFactory(FChaosArchive& Ar, TGeometryParticlesImp < T, d, SimType>* Particles);
 		
 		TGeometryParticlesImp()
-		    : TParticles<T, d>()
+		    : TSimpleGeometryParticles<T, d>()
 			, MContainerListMask(EGeometryParticleListMask::None)
 		{
 			MParticleType = EParticleType::Static;
@@ -170,11 +171,9 @@ namespace Chaos
 		}
 		TGeometryParticlesImp(const TGeometryParticlesImp<T, d, SimType>& Other) = delete;
 		TGeometryParticlesImp(TGeometryParticlesImp<T, d, SimType>&& Other)
-		    : TParticles<T, d>(MoveTemp(Other))
+		    : TSimpleGeometryParticles<T, d>(MoveTemp(Other))
 			, MContainerListMask(Other.MContainerListMask)
 			, MUniqueIdx(MoveTemp(Other.MUniqueIdx))
-			, MR(MoveTemp(Other.MR))
-			, MGeometry(MoveTemp(Other.MGeometry))
 			, MGeometryParticleHandle(MoveTemp(Other.MGeometryParticleHandle))
 			, MGeometryParticle(MoveTemp(Other.MGeometryParticle))
 			, MPhysicsProxy(MoveTemp(Other.MPhysicsProxy))
@@ -208,7 +207,7 @@ namespace Chaos
 		static constexpr bool IsRigidBodySim() { return SimType == EGeometryParticlesSimType::RigidBodySim; }
 
 		TGeometryParticlesImp(TParticles<T, d>&& Other)
-		    : TParticles<T, d>(MoveTemp(Other))
+		    : TSimpleGeometryParticles<T, d>(MoveTemp(Other))
 			, MContainerListMask(EGeometryParticleListMask::None)
 		{
 			MParticleType = EParticleType::Static;
@@ -218,16 +217,12 @@ namespace Chaos
 		virtual ~TGeometryParticlesImp()
 		{}
 
-		FORCEINLINE const TRotation<T, d>& R(const int32 Index) const { return MR[Index]; }
-		FORCEINLINE TRotation<T, d>& R(const int32 Index) { return MR[Index]; }
-
 		FUniqueIdx UniqueIdx(const int32 Index) const { return MUniqueIdx[Index]; }
 		FUniqueIdx& UniqueIdx(const int32 Index) { return MUniqueIdx[Index]; }
 
 		ESyncState& SyncState(const int32 Index) { return MSyncState[Index].State; }
 		ESyncState SyncState(const int32 Index) const { return MSyncState[Index].State; }
 
-		const FImplicitObjectPtr& GetGeometry(const int32 Index) const { return MGeometry[Index]; }
 
 		UE_DEPRECATED(5.4, "Please use GetGeometry instead")
 		TSerializablePtr<FImplicitObject> Geometry(const int32 Index) const { check(false); return TSerializablePtr<FImplicitObject>(); }
@@ -261,11 +256,6 @@ namespace Chaos
 		FParticleID& ParticleID(const int32 Idx) { return MParticleIDs[Idx]; }
 #endif
 		
-		void SetGeometry(const int32 Index, const FImplicitObjectPtr& InGeometry)
-		{
-			SetGeometryImpl(Index, InGeometry);
-		}
-
 		UE_DEPRECATED(5.4, "Please use SetGeometry with FImplicitObjectPtr instead")
 		void SetDynamicGeometry(const int32 Index, TUniquePtr<FImplicitObject>&& InUnique) { check(false); }
 
@@ -276,8 +266,6 @@ namespace Chaos
 		void RegisterArrays()
 		{
 			TArrayCollection::AddArray(&MUniqueIdx);
-			TArrayCollection::AddArray(&MR);
-			TArrayCollection::AddArray(&MGeometry);
 #if CHAOS_DETERMINISTIC
 			TArrayCollection::AddArray(&MParticleIDs);
 #endif
@@ -310,14 +298,14 @@ namespace Chaos
 			}
 		}
 
-		void SetGeometryImpl(const int32 Index, const FImplicitObjectPtr& InGeometry)
+		virtual void SetGeometryImpl(const int32 Index, const FImplicitObjectPtr& InGeometry) override
 		{
 			// We hit these checks if there's a call to modify geometry without first clearing constraints 
 			// on the particle (e.g., PBDRigidsEvolutionGBF::InvalidateParticle)
 			check(MParticleCollisions[Index].Num() == 0);
 			//check(MParticleConstraints[Index].Num() == 0);
 
-			MGeometry[Index] = InGeometry;
+			TSimpleGeometryParticles<T, d>::SetGeometryImpl(Index, InGeometry);
 
 			UpdateShapesArray(Index);
 
@@ -479,8 +467,6 @@ namespace Chaos
 			MWorldSpaceInflatedBounds[Index].GrowByVector(DeltaX);
 		}
 
-		const TArray<FImplicitObjectPtr>& GetAllGeometry() const { return MGeometry; }
-
 		typedef FGeometryParticleHandle THandleType;
 		FORCEINLINE THandleType* Handle(int32 Index) const { return const_cast<THandleType*>(MGeometryParticleHandle[Index].Get()); }
 
@@ -574,33 +560,10 @@ public:
 			return FString::Printf(TEXT("%s, MUniqueIdx:%d MR:%s, MGeometry:%s"), *BaseString, UniqueIdx(index).Idx, *R(index).ToString(), (GetGeometry(index) ? *(GetGeometry(index)->ToString()) : TEXT("none")));
 		}
 
-		virtual void Serialize(FChaosArchive& Ar)
+		virtual void Serialize(FChaosArchive& Ar) override
 		{
 			LLM_SCOPE(ELLMTag::ChaosParticles);
-			TParticles<T, d>::Serialize(Ar);
-			
-			Ar.UsingCustomVersion(FFortniteValkyrieBranchObjectVersion::GUID);
-			if (Ar.CustomVer(FFortniteValkyrieBranchObjectVersion::GUID) < FFortniteValkyrieBranchObjectVersion::RefCountedOImplicitObjects)
-			{
-				TArrayCollectionArray<TSerializablePtr<FImplicitObject>> LGeometry;
-				TArrayCollectionArray<TUniquePtr<Chaos::FImplicitObject>> LDynamicGeometry;
-				Ar << LGeometry << LDynamicGeometry;
-
-				if(Ar.IsLoading())
-				{
-					MGeometry.SetNumUninitialized(LGeometry.Num());
-					uint32 ImplicitIndex = 0;
-					for(const TSerializablePtr<FImplicitObject>& ImplicitObjectPtr : LGeometry)
-					{
-						MGeometry[ImplicitIndex++] = ImplicitObjectPtr->CopyGeometry();
-					}
-				}
-			}
-			else
-			{
-				Ar << MGeometry;
-			}
-			Ar << MR;
+			TSimpleGeometryParticles<T, d>::Serialize(Ar);
 			
 			Ar.UsingCustomVersion(FPhysicsObjectVersion::GUID);
 			if (Ar.CustomVer(FPhysicsObjectVersion::GUID) >= FPhysicsObjectVersion::PerShapeData)
@@ -631,12 +594,12 @@ public:
 			else
 			{
 				//just assume all bounds come from geometry (technically wrong for pbd rigids with only sample points, but backwards compat is not that important right now)
-				for (int32 Idx = 0; Idx < MGeometry.Num(); ++Idx)
+				for (int32 Idx = 0; Idx < GetAllGeometry().Num(); ++Idx)
 				{
-					MHasBounds[Idx] = MGeometry[Idx] && MGeometry[Idx]->HasBoundingBox();
+					MHasBounds[Idx] = GetGeometry(Idx) && GetGeometry(Idx)->HasBoundingBox();
 					if (MHasBounds[Idx])
 					{
-						MLocalBounds[Idx] = TAABB<T, d>(MGeometry[Idx]->BoundingBox());
+						MLocalBounds[Idx] = TAABB<T, d>(GetGeometry(Idx)->BoundingBox());
 						//ignore velocity too, really just trying to get something reasonable)
 						UpdateWorldSpaceState(Idx, FRigidTransform3(X(Idx), R(Idx)), FVec3(0));
 					}
@@ -645,7 +608,7 @@ public:
 
 			if (Ar.CustomVer(FExternalPhysicsCustomObjectVersion::GUID) < FExternalPhysicsCustomObjectVersion::SpatialIdxSerialized)
 			{
-				MSpatialIdx.AddZeroed(MGeometry.Num());
+				MSpatialIdx.AddZeroed(GetAllGeometry().Num());
 			}
 			else
 			{
@@ -663,7 +626,6 @@ public:
 		FORCEINLINE EGeometryParticleListMask GetContainerListMask() const { return MContainerListMask; }
 		void SetContainerListMask(const EGeometryParticleListMask InMask) { MContainerListMask = InMask; }
 
-		FORCEINLINE TArray<TRotation<T, d>>& AllR() { return MR; }
 		FORCEINLINE TArray<TAABB<T, d>>& AllLocalBounds() { return MLocalBounds; }
 		FORCEINLINE TArray<TAABB<T, d>>& AllWorldSpaceInflatedBounds() { return MWorldSpaceInflatedBounds; }
 		FORCEINLINE TArray<bool>& AllHasBounds() { return MHasBounds; }
@@ -674,10 +636,6 @@ public:
 
 	private:
 		TArrayCollectionArray<FUniqueIdx> MUniqueIdx;
-		TArrayCollectionArray<TRotation<T, d>> MR;
-		// MGeometry contains raw ptrs to every entry in both MSharedGeometry and MDynamicGeometry.
-		// It may also contain raw ptrs to geometry which is managed outside of Chaos.
-		TArrayCollectionArray<FImplicitObjectPtr> MGeometry;
 		TArrayCollectionArray<TSerializablePtr<FGeometryParticleHandle>> MGeometryParticleHandle;
 		TArrayCollectionArray<FGeometryParticle*> MGeometryParticle;
 		TArrayCollectionArray<IPhysicsProxyBase*> MPhysicsProxy;
diff --git a/Engine/Source/Runtime/Experimental/Chaos/Public/Chaos/NewtonEvolution.h b/Engine/Source/Runtime/Experimental/Chaos/Public/Chaos/NewtonEvolution.h
--- a/Engine/Source/Runtime/Experimental/Chaos/Public/Chaos/NewtonEvolution.h
+++ b/Engine/Source/Runtime/Experimental/Chaos/Public/Chaos/NewtonEvolution.h
@@ -8,7 +8,10 @@
 #include "Chaos/PBDSoftsSolverParticles.h"
 #include "Chaos/NewtonElasticFEM.h"
 #include "Chaos/NewtonCorotatedCache.h"
+#if UE_ENABLE_INCLUDE_ORDER_DEPRECATED_IN_5_4
 #include "Chaos/KinematicGeometryParticles.h"
+#endif
+#include "Chaos/SoftsSolverCollisionParticles.h"
 #include "Chaos/VelocityField.h"
 
 namespace Chaos::Softs
@@ -25,7 +28,7 @@ public:
 	// TODO: Tidy up this constructor (and update Headless Chaos)
 	CHAOS_API FNewtonEvolution(
 		FSolverParticles&& InParticles,
-		FSolverRigidParticles&& InGeometryParticles,
+		FSolverCollisionParticles&& InGeometryParticles,
 		TArray<TVec3<int32>>&& CollisionTriangles,
 		const TArray<TVector<int32, 4>>& InMesh,
 		TArray<TArray<TVector<int32, 2>>>&& InIncidentElements,
@@ -83,10 +86,10 @@ public:
 	int32 GetCollisionParticleRangeSize(int32 Offset) const { return MCollisionParticlesActiveView.GetRangeSize(Offset); }
 
 	//// collision particles accessors
-	const FSolverRigidParticles& CollisionParticles() const { return MCollisionParticles; }
-	FSolverRigidParticles& CollisionParticles() { return MCollisionParticles; }
+	const FSolverCollisionParticles& CollisionParticles() const { return MCollisionParticles; }
+	FSolverCollisionParticles& CollisionParticles() { return MCollisionParticles; }
 	const TArray<uint32>& CollisionParticleGroupIds() const { return MCollisionParticleGroupIds; }
-	const TPBDActiveView<FSolverRigidParticles>& CollisionParticlesActiveView() { return MCollisionParticlesActiveView; }
+	const TPBDActiveView<FSolverCollisionParticles>& CollisionParticlesActiveView() { return MCollisionParticlesActiveView; }
 
 
 	// Reset all constraint init and rule functions.
@@ -124,7 +127,7 @@ public:
 	TArray<TFunction<void(FSolverParticles&, const FSolverReal)>>& PostCollisionConstraintRules() { return MPostCollisionConstraintRules; }
 
 	void SetKinematicUpdateFunction(TFunction<void(FSolverParticles&, const FSolverReal, const FSolverReal, const int32)> KinematicUpdate) { MKinematicUpdate = KinematicUpdate; }
-	void SetCollisionKinematicUpdateFunction(TFunction<void(FSolverRigidParticles&, const FSolverReal, const FSolverReal, const int32)> KinematicUpdate) { MCollisionKinematicUpdate = KinematicUpdate; }
+	void SetCollisionKinematicUpdateFunction(TFunction<void(FSolverCollisionParticles&, const FSolverReal, const FSolverReal, const int32)> KinematicUpdate) { MCollisionKinematicUpdate = KinematicUpdate; }
 
 	TFunction<void(FSolverParticles&, const FSolverReal, const int32)>& GetForceFunction(const uint32 GroupId = 0) { return MGroupForceRules[GroupId]; }
 	const TFunction<void(FSolverParticles&, const FSolverReal, const int32)>& GetForceFunction(const uint32 GroupId = 0) const { return MGroupForceRules[GroupId]; }
@@ -195,8 +198,8 @@ private:
 private:
 	FSolverParticles MParticles;
 	TPBDActiveView<FSolverParticles> MParticlesActiveView;
-	FSolverRigidParticles MCollisionParticles;
-	TPBDActiveView<FSolverRigidParticles> MCollisionParticlesActiveView;
+	FSolverCollisionParticles MCollisionParticles;
+	TPBDActiveView<FSolverCollisionParticles> MCollisionParticlesActiveView;
 
 	TArrayCollectionArray<FSolverRigidTransform3> MCollisionTransforms;  // Used for CCD to store the initial state before the kinematic update
 	TArrayCollectionArray<bool> MCollided;
@@ -222,7 +225,7 @@ private:
 	TPBDActiveView<TArray<TFunction<void(FSolverParticles&, const FSolverReal)>>> MPostCollisionConstraintRulesActiveView;
 
 	TFunction<void(FSolverParticles&, const FSolverReal, const FSolverReal, const int32)> MKinematicUpdate;
-	TFunction<void(FSolverRigidParticles&, const FSolverReal, const FSolverReal, const int32)> MCollisionKinematicUpdate;
+	TFunction<void(FSolverCollisionParticles&, const FSolverReal, const FSolverReal, const int32)> MCollisionKinematicUpdate;
 	//Newton specific lambdas:
 	TFunction<void(const bool)> UpdatePositionBasedState;
 	TFunction<void(const PMatrix<FSolverReal, 3, 3>&, PMatrix<FSolverReal, 3, 3>&, const int32)> PStress;
diff --git a/Engine/Source/Runtime/Experimental/Chaos/Public/Chaos/PBDEvolution.h b/Engine/Source/Runtime/Experimental/Chaos/Public/Chaos/PBDEvolution.h
--- a/Engine/Source/Runtime/Experimental/Chaos/Public/Chaos/PBDEvolution.h
+++ b/Engine/Source/Runtime/Experimental/Chaos/Public/Chaos/PBDEvolution.h
@@ -6,7 +6,10 @@
 #include "Chaos/PBDActiveView.h"
 #include "Chaos/PBDSoftsEvolutionFwd.h"
 #include "Chaos/PBDSoftsSolverParticles.h"
+#if UE_ENABLE_INCLUDE_ORDER_DEPRECATED_IN_5_4
 #include "Chaos/KinematicGeometryParticles.h"
+#endif
+#include "Chaos/SoftsSolverCollisionParticles.h"
 #include "Chaos/VelocityField.h"
 
 namespace Chaos::Softs
@@ -18,7 +21,7 @@ class FPBDEvolution : public TArrayCollection
 	// TODO: Tidy up this constructor (and update Headless Chaos)
 	CHAOS_API FPBDEvolution(
 		FSolverParticles&& InParticles,
-		FSolverRigidParticles&& InGeometryParticles,
+		FSolverCollisionParticles&& InGeometryParticles,
 		TArray<TVec3<int32>>&& CollisionTriangles,
 		int32 NumIterations = 1,
 		FSolverReal CollisionThickness = (FSolverReal)0.,
@@ -83,12 +86,12 @@ class FPBDEvolution : public TArrayCollection
 	int32 GetCollisionParticleRangeSize(int32 Offset) const { return MCollisionParticlesActiveView.GetRangeSize(Offset); }
 
 	// Collision particles accessors
-	const FSolverRigidParticles& CollisionParticles() const { return MCollisionParticles; }
-	FSolverRigidParticles& CollisionParticles() { return MCollisionParticles; }
+	const FSolverCollisionParticles& CollisionParticles() const { return MCollisionParticles; }
+	FSolverCollisionParticles& CollisionParticles() { return MCollisionParticles; }
 	TArray<uint32>& CollisionParticleGroupIds() { return MCollisionParticleGroupIds; }
 	const TArray<uint32>& CollisionParticleGroupIds() const { return MCollisionParticleGroupIds; }
-	TPBDActiveView<FSolverRigidParticles>& CollisionParticlesActiveView() { return MCollisionParticlesActiveView; }
-	const TPBDActiveView<FSolverRigidParticles>& CollisionParticlesActiveView() const { return MCollisionParticlesActiveView; }
+	TPBDActiveView<FSolverCollisionParticles>& CollisionParticlesActiveView() { return MCollisionParticlesActiveView; }
+	const TPBDActiveView<FSolverCollisionParticles>& CollisionParticlesActiveView() const { return MCollisionParticlesActiveView; }
 
 	// Reset all constraint init and rule functions.
 	void ResetConstraintRules() 
@@ -132,7 +135,7 @@ class FPBDEvolution : public TArrayCollection
 	TArray<TFunction<void(FSolverParticles&, const FSolverReal)>>& ConstraintPostprocessings() { return MConstraintPostprocessings; }
 	
 	void SetKinematicUpdateFunction(TFunction<void(FSolverParticles&, const FSolverReal, const FSolverReal, const int32)> KinematicUpdate) { MKinematicUpdate = KinematicUpdate; }
-	void SetCollisionKinematicUpdateFunction(TFunction<void(FSolverRigidParticles&, const FSolverReal, const FSolverReal, const int32)> KinematicUpdate) { MCollisionKinematicUpdate = KinematicUpdate; }
+	void SetCollisionKinematicUpdateFunction(TFunction<void(FSolverCollisionParticles&, const FSolverReal, const FSolverReal, const int32)> KinematicUpdate) { MCollisionKinematicUpdate = KinematicUpdate; }
 
 	TFunction<void(FSolverParticles&, const FSolverReal, const int32)>& GetForceFunction(const uint32 GroupId = 0) { return MGroupForceRules[GroupId]; }
 	const TFunction<void(FSolverParticles&, const FSolverReal, const int32)>& GetForceFunction(const uint32 GroupId = 0) const { return MGroupForceRules[GroupId]; }
@@ -201,8 +204,8 @@ class FPBDEvolution : public TArrayCollection
 private:
 	FSolverParticles MParticles;
 	TPBDActiveView<FSolverParticles> MParticlesActiveView;
-	FSolverRigidParticles MCollisionParticles;
-	TPBDActiveView<FSolverRigidParticles> MCollisionParticlesActiveView;
+	FSolverCollisionParticles MCollisionParticles;
+	TPBDActiveView<FSolverCollisionParticles> MCollisionParticlesActiveView;
 	TArray<FParticleVievToken> RemovedCollisionIndices;
 
 	TArrayCollectionArray<FSolverRigidTransform3> MCollisionTransforms;  // Used for CCD to store the initial state before the kinematic update
@@ -235,7 +238,7 @@ private:
 	TPBDActiveView<TArray<TFunction<void(FSolverParticles&, const FSolverReal)>>> MConstraintPostprocessingsActiveView;
 
 	TFunction<void(FSolverParticles&, const FSolverReal, const FSolverReal, const int32)> MKinematicUpdate;
-	TFunction<void(FSolverRigidParticles&, const FSolverReal, const FSolverReal, const int32)> MCollisionKinematicUpdate;
+	TFunction<void(FSolverCollisionParticles&, const FSolverReal, const FSolverReal, const int32)> MCollisionKinematicUpdate;
 
 	int32 MNumIterations;
 	FSolverVec3 MGravity;
diff --git a/Engine/Source/Runtime/Experimental/Chaos/Public/Chaos/PBDSoftsEvolutionFwd.h b/Engine/Source/Runtime/Experimental/Chaos/Public/Chaos/PBDSoftsEvolutionFwd.h
--- a/Engine/Source/Runtime/Experimental/Chaos/Public/Chaos/PBDSoftsEvolutionFwd.h
+++ b/Engine/Source/Runtime/Experimental/Chaos/Public/Chaos/PBDSoftsEvolutionFwd.h
@@ -14,7 +14,6 @@ namespace Chaos
 	template <typename T, int d> class TVector;
 	template<class T, int m, int n> class PMatrix;
 	template <class T, int d> class TPBDParticles;
-	template <class T, int d, EGeometryParticlesSimType SimType> class TKinematicGeometryParticlesImp;
 	template <class T, int d> class TRotation;
 	template <class T, int d> class TRigidTransform;
 	template <class T, int d> class TAABB;
@@ -33,7 +32,8 @@ namespace Chaos
 		using FSolverMatrix22 = PMatrix<FSolverReal, 2, 2>;
 		using FSolverMatrix33 = PMatrix<FSolverReal, 3, 3>;
 		using FSolverMatrix44 = PMatrix<FSolverReal, 4, 4>;
-		using FSolverRigidParticles = TKinematicGeometryParticlesImp<FSolverReal, 3, EGeometryParticlesSimType::Other>;
+		class FSolverCollisionParticles;
+		using FSolverRigidParticles UE_DEPRECATED(5.4, "Alias FSolverRigidParticles has been renamed FSolverCollisionParticles and refers to a different class.") = FSolverCollisionParticles;
 		using FSolverRotation3 = TRotation<FSolverReal, 3>;
 		using FSolverRigidTransform3 = TRigidTransform<FSolverReal, 3>;
 		using FSolverTransform3 = UE::Math::TTransform<FSolverReal>;
diff --git a/Engine/Source/Runtime/Experimental/Chaos/Public/Chaos/PerParticlePBDCCDCollisionConstraint.h b/Engine/Source/Runtime/Experimental/Chaos/Public/Chaos/PerParticlePBDCCDCollisionConstraint.h
--- a/Engine/Source/Runtime/Experimental/Chaos/Public/Chaos/PerParticlePBDCCDCollisionConstraint.h
+++ b/Engine/Source/Runtime/Experimental/Chaos/Public/Chaos/PerParticlePBDCCDCollisionConstraint.h
@@ -3,7 +3,10 @@
 
 // HEADER_UNIT_SKIP - Internal
 
+#if UE_ENABLE_INCLUDE_ORDER_DEPRECATED_IN_5_4
 #include "Chaos/GeometryParticlesfwd.h"
+#endif
+#include "Chaos/SoftsSolverCollisionParticles.h"
 #include "Chaos/Transform.h"
 #include "Chaos/PBDActiveView.h"
 #include "Chaos/PBDSoftsSolverParticles.h"
@@ -18,7 +21,7 @@ class FPerParticlePBDCCDCollisionConstraint final
 {
 public:
 	FPerParticlePBDCCDCollisionConstraint(
-		const TPBDActiveView<FSolverRigidParticles>& InCollisionParticlesActiveView,
+		const TPBDActiveView<FSolverCollisionParticles>& InCollisionParticlesActiveView,
 		TArray<FSolverRigidTransform3>& InCollisionTransforms,
 		TArray<bool>& InCollided,
 		TArray<FSolverVec3>& InContacts,
@@ -75,7 +78,7 @@ private:
 				return;  // Continue
 			}
 
-			CollisionParticlesActiveView.SequentialFor([this, &Particles, &Index, DynamicGroupId, Thickness, Friction, Dt](FSolverRigidParticles& CollisionParticles, int32 CollisionIndex)
+			CollisionParticlesActiveView.SequentialFor([this, &Particles, &Index, DynamicGroupId, Thickness, Friction, Dt](FSolverCollisionParticles& CollisionParticles, int32 CollisionIndex)
 			{
 				const uint32 KinematicGroupId = KinematicGroupIds[CollisionIndex];  // Collision group Id
 
@@ -155,7 +158,7 @@ private:
 
 private:
 	// TODO(mlentine): Need a bb hierarchy
-	const TPBDActiveView<FSolverRigidParticles>& CollisionParticlesActiveView;
+	const TPBDActiveView<FSolverCollisionParticles>& CollisionParticlesActiveView;
 	const TArray<FSolverRigidTransform3>& CollisionTransforms;
 	TArray<bool>& Collided;
 	TArray<FSolverVec3>& Contacts;
diff --git a/Engine/Source/Runtime/Experimental/Chaos/Public/Chaos/PerParticlePBDCollisionConstraint.h b/Engine/Source/Runtime/Experimental/Chaos/Public/Chaos/PerParticlePBDCollisionConstraint.h
--- a/Engine/Source/Runtime/Experimental/Chaos/Public/Chaos/PerParticlePBDCollisionConstraint.h
+++ b/Engine/Source/Runtime/Experimental/Chaos/Public/Chaos/PerParticlePBDCollisionConstraint.h
@@ -5,7 +5,10 @@
 #include "Chaos/PBDActiveView.h"
 #include "Chaos/PBDSoftsEvolutionFwd.h"
 #include "Chaos/PBDSoftsSolverParticles.h"
+#if UE_ENABLE_INCLUDE_ORDER_DEPRECATED_IN_5_4
 #include "Chaos/KinematicGeometryParticles.h"
+#endif
+#include "Chaos/SoftsSolverCollisionParticles.h"
 #include "Chaos/WeightedLatticeImplicitObject.h"
 #include "Chaos/Levelset.h"
 #include "HAL/PlatformMath.h"
@@ -33,14 +36,14 @@ class FPerParticlePBDCollisionConstraint final
 	};
 
 public:
-	FPerParticlePBDCollisionConstraint(const TPBDActiveView<FSolverRigidParticles>& InParticlesActiveView, TArray<bool>& Collided, TArray<uint32>& DynamicGroupIds, TArray<uint32>& KinematicGroupIds, const TArray<FSolverReal>& PerGroupThickness, const TArray<FSolverReal>& PerGroupFriction)
+	FPerParticlePBDCollisionConstraint(const TPBDActiveView<FSolverCollisionParticles>& InParticlesActiveView, TArray<bool>& Collided, TArray<uint32>& DynamicGroupIds, TArray<uint32>& KinematicGroupIds, const TArray<FSolverReal>& PerGroupThickness, const TArray<FSolverReal>& PerGroupFriction)
 	: bFastPositionBasedFriction(true)
 	, MCollisionParticlesActiveView(InParticlesActiveView)
 	, MCollided(Collided), MDynamicGroupIds(DynamicGroupIds)
 	, MKinematicGroupIds(KinematicGroupIds), MPerGroupThickness(PerGroupThickness)
 	, MPerGroupFriction(PerGroupFriction) {}
 
-	FPerParticlePBDCollisionConstraint(const TPBDActiveView<FSolverRigidParticles>& InParticlesActiveView, TArray<bool>& Collided,
+	FPerParticlePBDCollisionConstraint(const TPBDActiveView<FSolverCollisionParticles>& InParticlesActiveView, TArray<bool>& Collided,
 		TArray<FSolverVec3>& InContacts,
 		TArray<FSolverVec3>& InNormals, 
 		TArray<FSolverReal>& InPhis,
@@ -118,7 +121,7 @@ private:
 					return;  // Continue
 				}
 
-				MCollisionParticlesActiveView.SequentialFor([this, &Particles, &Dt, &Index, DynamicGroupId, PerGroupFriction, PerGroupThickness](FSolverRigidParticles& CollisionParticles, int32 i)
+				MCollisionParticlesActiveView.SequentialFor([this, &Particles, &Dt, &Index, DynamicGroupId, PerGroupFriction, PerGroupThickness](FSolverCollisionParticles& CollisionParticles, int32 i)
 				{
 					const uint32 KinematicGroupId = MKinematicGroupIds[i];  // Collision group Id
 
@@ -209,7 +212,7 @@ private:
 					return;  // Continue
 				}
 
-				MCollisionParticlesActiveView.SequentialFor([this, &Particles, &Dt, &Index, DynamicGroupId, PerGroupFriction, PerGroupThickness](FSolverRigidParticles& CollisionParticles, int32 i)
+				MCollisionParticlesActiveView.SequentialFor([this, &Particles, &Dt, &Index, DynamicGroupId, PerGroupFriction, PerGroupThickness](FSolverCollisionParticles& CollisionParticles, int32 i)
 				{
 					const uint32 KinematicGroupId = MKinematicGroupIds[i];  // Collision group Id
 
@@ -251,7 +254,7 @@ private:
 private:
 	bool bFastPositionBasedFriction;
 	// TODO(mlentine): Need a bb hierarchy
-	const TPBDActiveView<FSolverRigidParticles>& MCollisionParticlesActiveView;
+	const TPBDActiveView<FSolverCollisionParticles>& MCollisionParticlesActiveView;
 	TArray<bool>& MCollided;
 	TArray<FSolverVec3>* const Contacts = nullptr;
 	TArray<FSolverVec3>* const Normals = nullptr;
diff --git a/Engine/Source/Runtime/Experimental/Chaos/Public/Chaos/SimpleGeometryParticles.h b/Engine/Source/Runtime/Experimental/Chaos/Public/Chaos/SimpleGeometryParticles.h
new file mode 100644
--- /dev/null
+++ b/Engine/Source/Runtime/Experimental/Chaos/Public/Chaos/SimpleGeometryParticles.h
@@ -0,0 +1,108 @@
+// Copyright Epic Games, Inc. All Rights Reserved.
+#pragma once
+
+#include "Chaos/ArrayCollectionArray.h"
+#include "Chaos/ImplicitObject.h"
+#include "Chaos/Particles.h"
+#include "Chaos/Rotation.h"
+#include "UObject/FortniteValkyrieBranchObjectVersion.h"
+
+
+namespace Chaos
+{
+	template<class T, int d>
+	class TSimpleGeometryParticles : public TParticles<T, d>
+	{
+	public:
+
+		using TArrayCollection::Size;
+		using TParticles<T,d>::X;
+		
+		TSimpleGeometryParticles()
+		    : TParticles<T, d>()
+		{
+			TArrayCollection::AddArray(&MR);
+			TArrayCollection::AddArray(&MGeometry);
+		}
+		TSimpleGeometryParticles(const TSimpleGeometryParticles<T, d>& Other) = delete;
+		TSimpleGeometryParticles(TSimpleGeometryParticles<T, d>&& Other)
+		    : TParticles<T, d>(MoveTemp(Other))
+			, MR(MoveTemp(Other.MR))
+			, MGeometry(MoveTemp(Other.MGeometry))
+		{
+			TArrayCollection::AddArray(&MR);
+			TArrayCollection::AddArray(&MGeometry);
+		}
+		TSimpleGeometryParticles& operator=(const TSimpleGeometryParticles<T, d>& Other) = delete;
+		TSimpleGeometryParticles& operator=(TSimpleGeometryParticles<T, d>&& Other) = delete;
+
+		TSimpleGeometryParticles(TParticles<T, d>&& Other)
+		    : TParticles<T, d>(MoveTemp(Other))
+		{
+			TArrayCollection::AddArray(&MR);
+			TArrayCollection::AddArray(&MGeometry);
+		}
+
+		virtual ~TSimpleGeometryParticles() override
+		{}
+
+		FORCEINLINE const TRotation<T, d>& R(const int32 Index) const { return MR[Index]; }
+		FORCEINLINE TRotation<T, d>& R(const int32 Index) { return MR[Index]; }
+		FORCEINLINE const FImplicitObjectPtr& GetGeometry(const int32 Index) const { return MGeometry[Index]; }
+		void SetGeometry(const int32 Index, const FImplicitObjectPtr& InGeometry)
+		{
+			SetGeometryImpl(Index, InGeometry);
+		}
+
+		FORCEINLINE const TArray<FImplicitObjectPtr>& GetAllGeometry() const { return MGeometry; }
+		FORCEINLINE TArray<TRotation<T, d>>& AllR() { return MR; }
+
+		virtual void Serialize(FChaosArchive& Ar)
+		{
+			TParticles<T, d>::Serialize(Ar);
+			
+			Ar.UsingCustomVersion(FFortniteValkyrieBranchObjectVersion::GUID);
+			if (Ar.CustomVer(FFortniteValkyrieBranchObjectVersion::GUID) < FFortniteValkyrieBranchObjectVersion::RefCountedOImplicitObjects)
+			{
+				TArrayCollectionArray<TSerializablePtr<FImplicitObject>> LGeometry;
+				TArrayCollectionArray<TUniquePtr<Chaos::FImplicitObject>> LDynamicGeometry;
+				Ar << LGeometry << LDynamicGeometry;
+
+				if(Ar.IsLoading())
+				{
+					MGeometry.SetNumUninitialized(LGeometry.Num());
+					uint32 ImplicitIndex = 0;
+					for(const TSerializablePtr<FImplicitObject>& ImplicitObjectPtr : LGeometry)
+					{
+						MGeometry[ImplicitIndex++] = ImplicitObjectPtr->CopyGeometry();
+					}
+				}
+			}
+			else
+			{
+				Ar << MGeometry;
+			}
+			Ar << MR;
+		}
+
+	protected:
+		virtual void SetGeometryImpl(const int32 Index, const FImplicitObjectPtr& InGeometry)
+		{
+			MGeometry[Index] = InGeometry;
+		}
+
+	private:
+		TArrayCollectionArray<TRotation<T, d>> MR;
+		// MGeometry contains raw ptrs to every entry in both MSharedGeometry and MDynamicGeometry.
+		// It may also contain raw ptrs to geometry which is managed outside of Chaos.
+		TArrayCollectionArray<FImplicitObjectPtr> MGeometry;
+	};
+
+	template <typename T, int d>
+	FChaosArchive& operator<<(FChaosArchive& Ar, TSimpleGeometryParticles<T, d>& Particles)
+	{
+		Particles.Serialize(Ar);
+		return Ar;
+	}
+}
+
diff --git a/Engine/Source/Runtime/Experimental/Chaos/Public/Chaos/SoftsSolverCollisionParticles.h b/Engine/Source/Runtime/Experimental/Chaos/Public/Chaos/SoftsSolverCollisionParticles.h
new file mode 100644
--- /dev/null
+++ b/Engine/Source/Runtime/Experimental/Chaos/Public/Chaos/SoftsSolverCollisionParticles.h
@@ -0,0 +1,49 @@
+// Copyright Epic Games, Inc. All Rights Reserved.
+#pragma once
+
+#include "SimpleGeometryParticles.h"
+#include "PBDSoftsEvolutionFwd.h"
+
+namespace Chaos::Softs
+{
+class FSolverCollisionParticles : public TSimpleGeometryParticles<FSolverReal, 3>
+{
+  public:
+	FSolverCollisionParticles()
+	    : TSimpleGeometryParticles<FSolverReal, 3>()
+	{
+		TArrayCollection::AddArray(&MV);
+		TArrayCollection::AddArray(&MW);
+	}
+	FSolverCollisionParticles(const FSolverCollisionParticles& Other) = delete;
+	FSolverCollisionParticles(FSolverCollisionParticles&& Other)
+	    : TSimpleGeometryParticles<FSolverReal, 3>(MoveTemp(Other)), MV(MoveTemp(Other.MV)), MW(MoveTemp(Other.MW))
+	{
+		TArrayCollection::AddArray(&MV);
+		TArrayCollection::AddArray(&MW);
+	}
+	virtual ~FSolverCollisionParticles() override {};
+	FSolverCollisionParticles& operator=(const FSolverCollisionParticles& Other) = delete;
+	FSolverCollisionParticles& operator=(FSolverCollisionParticles&& Other) = delete;
+
+	const FSolverVec3& V(const int32 Index) const { return MV[Index]; }
+	FSolverVec3& V(const int32 Index) { return MV[Index]; }
+
+	const FSolverVec3& W(const int32 Index) const { return MW[Index]; }
+	FSolverVec3& W(const int32 Index) { return MW[Index]; }
+
+	virtual void Serialize(FChaosArchive& Ar) override
+	{
+		TSimpleGeometryParticles<FSolverReal, 3>::Serialize(Ar);
+		Ar << MV << MW;
+	}
+
+	FORCEINLINE TArray<FSolverVec3>& AllV() { return MV; }
+	FORCEINLINE TArray<FSolverVec3>& AllW() { return MW; }
+
+  private:
+	TArrayCollectionArray<FSolverVec3> MV;
+	TArrayCollectionArray<FSolverVec3> MW;
+};
+}
+
