expose GetGeometry on the physics object interface

diff --git a/Engine/Source/Runtime/Experimental/Chaos/Private/Chaos/PhysicsObjectInterface.cpp b/Engine/Source/Runtime/Experimental/Chaos/Private/Chaos/PhysicsObjectInterface.cpp
--- a/Engine/Source/Runtime/Experimental/Chaos/Private/Chaos/PhysicsObjectInterface.cpp
+++ b/Engine/Source/Runtime/Experimental/Chaos/Private/Chaos/PhysicsObjectInterface.cpp
@@ -479,6 +479,22 @@ namespace Chaos
 		return AllShapes;
 	}
 
+	template<EThreadContext Id>
+	FImplicitObjectRef FReadPhysicsObjectInterface<Id>::GetGeometry(const FConstPhysicsObjectHandle Handle)
+	{
+		if (!Handle)
+		{
+			return nullptr;
+		}
+
+		if (TThreadParticle<Id>* Particle = Handle->GetParticle<Id>())
+		{
+			return Particle->GetGeometry();
+		}
+
+		return nullptr;
+	}
+
 	template<EThreadContext Id>
 	void FReadPhysicsObjectInterface<Id>::VisitEveryShape(TArrayView<const FConstPhysicsObjectHandle> InObjects, TFunctionRef<bool(const FConstPhysicsObjectHandle, TThreadShapeInstance<Id>*)> Lambda)
 	{
diff --git a/Engine/Source/Runtime/Experimental/Chaos/Public/Chaos/PhysicsObjectInterface.h b/Engine/Source/Runtime/Experimental/Chaos/Public/Chaos/PhysicsObjectInterface.h
--- a/Engine/Source/Runtime/Experimental/Chaos/Public/Chaos/PhysicsObjectInterface.h
+++ b/Engine/Source/Runtime/Experimental/Chaos/Public/Chaos/PhysicsObjectInterface.h
@@ -3,6 +3,7 @@
 #pragma once
 
 #include "Chaos/ClusterCreationParameters.h"
+#include "Chaos/ImplicitFwd.h"
 #include "Chaos/ParticleHandle.h"
 #include "Chaos/PhysicsObject.h"
 #include "Chaos/ShapeInstanceFwd.h"
@@ -82,6 +83,7 @@ namespace Chaos
 		CHAOS_API TArray<FPerShapeData*> GetAllShapes(TArrayView<const FConstPhysicsObjectHandle> InObjects);
 
 		CHAOS_API TArray<TThreadShapeInstance<Id>*> GetAllThreadShapes(TArrayView<const FConstPhysicsObjectHandle> InObjects);
+		CHAOS_API FImplicitObjectRef GetGeometry(const FConstPhysicsObjectHandle Handle);
 
 		// Returns true if a shape is found and we can stop iterating.
 		CHAOS_API void VisitEveryShape(TArrayView<const FConstPhysicsObjectHandle> InObjects, TFunctionRef<bool(const FConstPhysicsObjectHandle, TThreadShapeInstance<Id>*)> Lambda);
