Add partial destruction decay system for moving buildings
When the building is moving , check at regular interval if any partially destroyed geometry collection has been broken for certain amount of time since the building is moving
If so, then release it
- Cache the dynamic state in the cluster union component to avoid calling IsSimulatingPhysics()
- Added a virtual function called whenever the physics state of the cluster union is set

diff --git a/Engine/Source/Runtime/Engine/Classes/PhysicsEngine/ClusterUnionComponent.h b/Engine/Source/Runtime/Engine/Classes/PhysicsEngine/ClusterUnionComponent.h
--- a/Engine/Source/Runtime/Engine/Classes/PhysicsEngine/ClusterUnionComponent.h
+++ b/Engine/Source/Runtime/Engine/Classes/PhysicsEngine/ClusterUnionComponent.h
@@ -310,6 +310,8 @@ protected:
 	Chaos::FClusterUnionPhysicsProxy* GetPhysicsProxy() const { return PhysicsProxy; }
 	Chaos::FClusterUnionPhysicsProxy* GetPhysicsProxy() { return PhysicsProxy; }
 
+	ENGINE_API virtual void SetRigidState(Chaos::EObjectStateType ObjectState);
+
 	// We need to keep track of the mapping of primitive components to physics objects.
 	// This way we know the right physics objects to pass when removing the component (because
 	// it's possible to get a different list of physics objects when we get to removal). A
diff --git a/Engine/Source/Runtime/Engine/Private/PhysicsEngine/ClusterUnionComponent.cpp b/Engine/Source/Runtime/Engine/Private/PhysicsEngine/ClusterUnionComponent.cpp
--- a/Engine/Source/Runtime/Engine/Private/PhysicsEngine/ClusterUnionComponent.cpp
+++ b/Engine/Source/Runtime/Engine/Private/PhysicsEngine/ClusterUnionComponent.cpp
@@ -1174,7 +1174,7 @@ void UClusterUnionComponent::OnRep_RigidState()
 	}
 
 	PhysicsProxy->SetIsAnchored_External(ReplicatedRigidState.bIsAnchored);
-	PhysicsProxy->SetObjectState_External(static_cast<Chaos::EObjectStateType>(ReplicatedRigidState.ObjectState));
+	SetRigidState(static_cast<Chaos::EObjectStateType>(ReplicatedRigidState.ObjectState));
 }
 
 void UClusterUnionComponent::GetLifetimeReplicatedProps(TArray<FLifetimeProperty>& OutLifetimeProps) const
@@ -1220,13 +1220,17 @@ void UClusterUnionComponent::ForceSetChildToParent(UPrimitiveComponent* InCompon
 	}
 }
 
-void UClusterUnionComponent::SetSimulatePhysics(bool bSimulate)
+void UClusterUnionComponent::SetRigidState(Chaos::EObjectStateType ObjectState)
 {
-	if (!PhysicsProxy)
+	if (PhysicsProxy)
 	{
-		return;
+		PhysicsProxy->SetObjectState_External(ObjectState);
 	}
-	PhysicsProxy->SetObjectState_External(bSimulate ? Chaos::EObjectStateType::Dynamic : Chaos::EObjectStateType::Kinematic);
+}
+
+void UClusterUnionComponent::SetSimulatePhysics(bool bSimulate)
+{
+	SetRigidState(bSimulate ? Chaos::EObjectStateType::Dynamic : Chaos::EObjectStateType::Kinematic);
 }
 
 void UClusterUnionComponent::WakeAllRigidBodies()
@@ -1239,7 +1243,7 @@ void UClusterUnionComponent::WakeAllRigidBodies()
 	Chaos::EObjectStateType State = PhysicsProxy->GetObjectState_External();
 	if (State == Chaos::EObjectStateType::Sleeping)
 	{
-		PhysicsProxy->SetObjectState_External(Chaos::EObjectStateType::Dynamic);
+		SetRigidState(Chaos::EObjectStateType::Dynamic);
 	}
 }
 
diff --git a/Engine/Source/Runtime/Experimental/GeometryCollectionEngine/Private/GeometryCollection/GeometryCollectionComponent.cpp b/Engine/Source/Runtime/Experimental/GeometryCollectionEngine/Private/GeometryCollection/GeometryCollectionComponent.cpp
--- a/Engine/Source/Runtime/Experimental/GeometryCollectionEngine/Private/GeometryCollection/GeometryCollectionComponent.cpp
+++ b/Engine/Source/Runtime/Experimental/GeometryCollectionEngine/Private/GeometryCollection/GeometryCollectionComponent.cpp
@@ -3962,6 +3962,21 @@ bool UGeometryCollectionComponent::FBrokenAndDecayedStates::GetHasDecayed(int32
 	return bIsRootBroken ? HasDecayed[TransformIndex] : false;
 }
 
+double UGeometryCollectionComponent::FBrokenAndDecayedStates::GetRootBrokenEventTimeInMs() const
+{ 
+	return RootBrokenEventTimeInMs;
+}
+
+double UGeometryCollectionComponent::FBrokenAndDecayedStates::GetRootBrokenElapsedTimeInMs() const
+{
+	if (bIsRootBroken)
+	{
+		const double CurrentTimeInMs = FPlatformTime::ToMilliseconds64(FPlatformTime::Cycles64());
+		return (CurrentTimeInMs - RootBrokenEventTimeInMs);
+	}
+	return 0;
+}
+
 void UGeometryCollectionComponent::FBrokenAndDecayedStates::Reset(int32 InNumTransforms)
 {
 	NumTransforms = InNumTransforms;
@@ -3969,6 +3984,7 @@ void UGeometryCollectionComponent::FBrokenAndDecayedStates::Reset(int32 InNumTra
 	NumDecaying = 0;
 	IsBroken.Empty();
 	HasDecayed.Empty();
+	RootBrokenEventTimeInMs = 0;
 }
 
 void UGeometryCollectionComponent::FBrokenAndDecayedStates::SetRootIsBroken(bool bIsBroken)
@@ -3981,11 +3997,13 @@ void UGeometryCollectionComponent::FBrokenAndDecayedStates::SetRootIsBroken(bool
 		{
 			IsBroken.Init(false, NumTransforms);
 			HasDecayed.Init(false, NumTransforms);
+			RootBrokenEventTimeInMs = FPlatformTime::ToMilliseconds64(FPlatformTime::Cycles64());
 		}
 		else
 		{
 			IsBroken.Empty();
 			HasDecayed.Empty();
+			RootBrokenEventTimeInMs = 0;
 		}
 	}
 }
@@ -4055,6 +4073,7 @@ void UGeometryCollectionComponent::OnPostPhysicsSync()
 	SCOPE_CYCLE_COUNTER(STAT_GCPostPhysicsSync);
 
 	// dirty the transform if the collection is
+	bool bHasRootMoved = false;
 	if (DynamicCollection && DynamicCollection->IsDirty())
 	{
 		// Can't be a const reference - we need to make a copy or else the next RequestRootTransform will change the value under our nose.
@@ -4062,13 +4081,10 @@ void UGeometryCollectionComponent::OnPostPhysicsSync()
 
 		OnTransformsDirty();
 		
-		if (OnRootMovedEvent.IsBound())
+		if (OnRootMovedEvent.IsBound() || OnRootMovedNativeEvent.IsBound())
 		{
 			const FTransform3f& NewRootTransform = ComponentSpaceTransforms.RequestRootTransform();
-			if (!PreviousRootTransform.EqualsNoScale(NewRootTransform))
-			{
-				OnRootMovedEvent.Broadcast();
-			}
+			bHasRootMoved = !PreviousRootTransform.EqualsNoScale(NewRootTransform);
 		}
 	}
 
@@ -4118,6 +4134,20 @@ void UGeometryCollectionComponent::OnPostPhysicsSync()
 	const bool bDynamicDataIsDirty = (DynamicCollection && DynamicCollection->IsDirty() && HasVisibleGeometry());
 	UpdateRenderSystemsIfNeeded(bDynamicDataIsDirty);
 	UpdateNavigationDataIfNeeded(bDynamicDataIsDirty);
+
+	// at this is called at the end as thos events may change some internal values of the component
+	// and interfere with transforms required for naviagtion or rendering 
+	if (bHasRootMoved)
+	{
+		if (OnRootMovedEvent.IsBound())
+		{
+			OnRootMovedEvent.Broadcast();
+		}
+		if (OnRootMovedNativeEvent.IsBound())
+		{
+			OnRootMovedNativeEvent.Broadcast(this);
+		}
+	}
 }
 
 void UGeometryCollectionComponent::MoveComponentToRootTransform()
@@ -5711,6 +5741,7 @@ const FTransform3f& UGeometryCollectionComponent::FComponentSpaceTransforms::Req
 			bIsRootDirty = false;
 		}
 
+		ensureAlways(!Transforms[RootIndex].ContainsNaN());
 		return Transforms[RootIndex];
 	}
 
diff --git a/Engine/Source/Runtime/Experimental/GeometryCollectionEngine/Public/GeometryCollection/GeometryCollectionComponent.h b/Engine/Source/Runtime/Experimental/GeometryCollectionEngine/Public/GeometryCollection/GeometryCollectionComponent.h
--- a/Engine/Source/Runtime/Experimental/GeometryCollectionEngine/Public/GeometryCollection/GeometryCollectionComponent.h
+++ b/Engine/Source/Runtime/Experimental/GeometryCollectionEngine/Public/GeometryCollection/GeometryCollectionComponent.h
@@ -58,7 +58,6 @@ enum class EObjectStateTypeEnum : uint8;
 namespace Chaos { enum class EObjectStateType: int8; }
 template<class InElementType> class TManagedArray;
 
-
 DECLARE_DYNAMIC_MULTICAST_DELEGATE_OneParam(FOnChaosBreakEvent, const FChaosBreakEvent&, BreakEvent);
 
 DECLARE_DYNAMIC_MULTICAST_DELEGATE_OneParam(FOnChaosRemovalEvent, const FChaosRemovalEvent&, RemovalEvent);
@@ -68,6 +67,7 @@ DECLARE_DYNAMIC_MULTICAST_DELEGATE_OneParam(FOnChaosCrumblingEvent, const FChaos
 DECLARE_DYNAMIC_MULTICAST_DELEGATE(FOnGeometryCollectionFullyDecayedEvent);
 
 DECLARE_DYNAMIC_MULTICAST_DELEGATE(FOnGeometryCollectionRootMovedEvent);
+DECLARE_MULTICAST_DELEGATE_OneParam(FOnGeometryCollectionRootMovedNativeEvent, UGeometryCollectionComponent*);
 
 namespace GeometryCollection
 {
@@ -1227,6 +1227,7 @@ public:
 
 	FOnGeometryCollectionFullyDecayedEvent OnFullyDecayedEvent;
 	FOnGeometryCollectionRootMovedEvent OnRootMovedEvent;
+	FOnGeometryCollectionRootMovedNativeEvent OnRootMovedNativeEvent;
 
 	GEOMETRYCOLLECTIONENGINE_API void DispatchBreakEvent(const FChaosBreakEvent& Event);
 
@@ -1319,6 +1320,8 @@ public:
 
 	GEOMETRYCOLLECTIONENGINE_API bool ShouldUpdateComponentTransformToRootBone() const { return bUpdateComponentTransformToRootBone; }
 
+	GEOMETRYCOLLECTIONENGINE_API double GetRootBrokenElapsedTimeInMs() const { return BrokenAndDecayedStates.GetRootBrokenElapsedTimeInMs(); }
+
 public:
 	UPROPERTY(BlueprintAssignable, Category = "Collision")
 	FOnChaosPhysicsCollision OnChaosPhysicsCollision;
@@ -1761,6 +1764,8 @@ private:
 		bool GetIsRootBroken() const { return bIsRootBroken; }
 		bool GetIsBroken(int32 TransformIndex) const;
 		bool GetHasDecayed(int32 TransformIndex) const;
+		double GetRootBrokenEventTimeInMs() const;
+		double GetRootBrokenElapsedTimeInMs() const;
 
 		void SetRootIsBroken(bool bIsBroken);
 		void SetIsBroken(int32 TransformIndex);
@@ -1775,6 +1780,7 @@ private:
 	private:
 		int32 NumTransforms = 0;
 		bool bIsRootBroken = false;
+		double RootBrokenEventTimeInMs = 0;
 		TBitArray<> IsBroken;
 		TBitArray<> HasDecayed;
 		int32 NumDecaying = 0;
