IAS - list on demand TOC file(s)

A new command for listing on demand TOC files from an S3 compatible endpoint.
Usage: UnrealPak.exe -ListTocs -Region="us-east-1" -BucketPath=<some/bucket/path/> -CredentialsFile="s3_credentials.txt"

- write json file with -Json="Path/To/File.json".
- dump a specific toc: -TocKey=<IoHash>
- dump a specific build version: -BuildVersion=<Version>
- dump a specific build version and platform: -BuildVersion=<Version> -TargetPlatform=<Platform>
- dump chunk object keys with -ChunkKeys=<path/To/File.json>
- fetch a specific toc from a public CDN: -TocUrl="http://some.cdn.net/bd22076ca12bfc6feb982ffb064d18f28156606c.iochunktoc"

#rnx

diff --git a/Engine/Source/Developer/IoStoreUtilities/Private/IoStoreUtilities.cpp b/Engine/Source/Developer/IoStoreUtilities/Private/IoStoreUtilities.cpp
--- a/Engine/Source/Developer/IoStoreUtilities/Private/IoStoreUtilities.cpp
+++ b/Engine/Source/Developer/IoStoreUtilities/Private/IoStoreUtilities.cpp
@@ -9355,16 +9355,17 @@ bool UploadIoStoreContainerFiles(const TCHAR* ContainerPathOrWildcard)
 
 bool DownloadIoStoreContainerFiles(const TCHAR* TocPath)
 {
+	using namespace UE::IO::IAS;
 	check(TocPath);
 
-	TIoStatusOr<UE::IO::IAS::FIoStoreDownloadParams> Params = UE::IO::IAS::FIoStoreDownloadParams::Parse(FCommandLine::Get());
+	TIoStatusOr<FIoStoreDownloadParams> Params = FIoStoreDownloadParams::Parse(FCommandLine::Get());
 	if (Params.IsOk() == false)
 	{
 		UE_LOG(LogIoStore, Error, TEXT("Failed to download container file(s), reason '%s'"), *Params.Status().ToString());
 		return false;
 	}
 
-	FIoStatus Status = UE::IO::IAS::DownloadContainerFiles(Params.ConsumeValueOrDie(), TocPath);
+	FIoStatus Status = DownloadContainerFiles(Params.ConsumeValueOrDie(), TocPath);
 	if (Status.IsOk() == false)
 	{
 		UE_LOG(LogIoStore, Error, TEXT("Failed to download container file(s), reason '%s'"), *Status.ToString());
@@ -9373,6 +9374,26 @@ bool DownloadIoStoreContainerFiles(const TCHAR* TocPath)
 	return Status.IsOk();
 }
 
+bool ListOnDemandTocs()
+{
+	using namespace UE::IO::IAS;
+
+	TIoStatusOr<FIoStoreListTocsParams> Params = FIoStoreListTocsParams::Parse(FCommandLine::Get());
+	if (Params.IsOk() == false)
+	{
+		UE_LOG(LogIoStore, Error, TEXT("Failed to list TOC file(s), reason '%s'"), *Params.Status().ToString());
+		return false;
+	}
+
+	FIoStatus Status = ListTocs(Params.ConsumeValueOrDie());
+	if (Status.IsOk() == false)
+	{
+		UE_LOG(LogIoStore, Error, TEXT("Failed to list TOC file(s), reason '%s'"), *Status.ToString());
+	}
+
+	return Status.IsOk();
+}
+
 bool PrimeEndPoint(FStringView IoStoreOnDemandIniPath)
 {
 	FIoStatus Status = UE::IO::IAS::PrimeEndPoint(IoStoreOnDemandIniPath);
diff --git a/Engine/Source/Developer/IoStoreUtilities/Public/IoStoreUtilities.h b/Engine/Source/Developer/IoStoreUtilities/Public/IoStoreUtilities.h
--- a/Engine/Source/Developer/IoStoreUtilities/Public/IoStoreUtilities.h
+++ b/Engine/Source/Developer/IoStoreUtilities/Public/IoStoreUtilities.h
@@ -57,6 +57,8 @@ IOSTOREUTILITIES_API bool UploadIoStoreContainerFiles(const TCHAR* ContainerPath
 
 IOSTOREUTILITIES_API bool DownloadIoStoreContainerFiles(const TCHAR* TocPath);
 
+IOSTOREUTILITIES_API bool ListOnDemandTocs();
+
 /**
  * Attempts to connect to the endpoint in the provided ini file, download all
  * available tocs and then download all chunks in those tocs. This will
diff --git a/Engine/Source/Developer/PakFileUtilities/Private/PakFileUtilities.cpp b/Engine/Source/Developer/PakFileUtilities/Private/PakFileUtilities.cpp
--- a/Engine/Source/Developer/PakFileUtilities/Private/PakFileUtilities.cpp
+++ b/Engine/Source/Developer/PakFileUtilities/Private/PakFileUtilities.cpp
@@ -5473,6 +5473,11 @@ bool ExecuteUnrealPak(const TCHAR* CmdLine)
 			{
 				return PrimeEndPoint(*IoStoreArg);
 			}
+
+			if (FParse::Param(CmdLine, TEXT("ListTocs")))
+			{
+				return ListOnDemandTocs();
+			}
 		}
 	}
 
diff --git a/Engine/Source/Runtime/Experimental/IoStoreOnDemand/Private/IoStoreOnDemand.cpp b/Engine/Source/Runtime/Experimental/IoStoreOnDemand/Private/IoStoreOnDemand.cpp
--- a/Engine/Source/Runtime/Experimental/IoStoreOnDemand/Private/IoStoreOnDemand.cpp
+++ b/Engine/Source/Runtime/Experimental/IoStoreOnDemand/Private/IoStoreOnDemand.cpp
@@ -33,8 +33,14 @@
 #if (PLATFORM_DESKTOP && (IS_PROGRAM || WITH_EDITOR))
 #include "Algo/Sort.h"
 #include "Async/Async.h"
+#include "Async/Mutex.h"
+#include "Async/ParallelFor.h"
+#include "Async/UniqueLock.h"
 #include "HAL/PlatformFileManager.h"
 #include "IO/IoStore.h"
+#include "Misc/EnumClassFlags.h"
+#include "Serialization/JsonSerializer.h"
+#include "Serialization/JsonWriter.h"
 #include "S3/S3Client.h"
 #endif // (PLATFORM_DESKTOP && (IS_PROGRAM || WITH_EDITOR))
 
@@ -630,6 +636,116 @@ TIoStatusOr<FOnDemandToc> LoadTocFromUrl(const FString& ServiceUrl, const FStrin
 
 ////////////////////////////////////////////////////////////////////////////////
 #if (PLATFORM_DESKTOP && (IS_PROGRAM || WITH_EDITOR))
+
+using FJsonWriter = TSharedPtr<TJsonWriter<TCHAR, TPrettyJsonPrintPolicy<TCHAR>>>; 
+using FJsonWriterFactory = TJsonWriterFactory<TCHAR, TPrettyJsonPrintPolicy<TCHAR>>;
+
+////////////////////////////////////////////////////////////////////////////////
+static void GetChunkObjectKeys(const FOnDemandToc& Toc, FString Prefix, TArray<FString>& OutKeys)
+{
+	if (Prefix.EndsWith(TEXT("/")))
+	{
+		Prefix = Prefix.Left(Prefix.Len() - 1);
+	}
+
+	TStringBuilder<256> Sb;
+	for (const FOnDemandTocContainerEntry& Container : Toc.Containers)
+	{
+		for (const FOnDemandTocEntry& Entry : Container.Entries)
+		{
+			Sb.Reset();
+			const FString Hash = LexToString(Entry.Hash);
+			Sb << Prefix << TEXT("/") << Toc.Header.ChunksDirectory << TEXT("/") << Hash.Left(2) << TEXT("/") << Hash << TEXT(".iochunk");
+			OutKeys.Add(FString(Sb.ToString()).ToLower());
+		}
+	}
+}
+
+////////////////////////////////////////////////////////////////////////////////
+enum class EOnDemandTocJsonOptions
+{
+	Header			= (1 << 0),
+	TocEntries		= (1 << 1),
+	BlockSizes		= (1 << 2),
+	BlockHashes		= (1 << 3),
+	All				= Header | TocEntries | BlockSizes | BlockHashes
+};
+ENUM_CLASS_FLAGS(EOnDemandTocJsonOptions);
+
+static void ToJson(FJsonWriter JsonWriter, const FOnDemandToc& Toc, EOnDemandTocJsonOptions Options)
+{
+	JsonWriter->WriteObjectStart(TEXT("Header"));
+	{
+		JsonWriter->WriteValue(TEXT("Magic"), Toc.Header.Magic);
+		JsonWriter->WriteValue(TEXT("Version"), uint64(Toc.Header.Version));
+		JsonWriter->WriteValue(TEXT("BlockSize"), uint64(Toc.Header.BlockSize));
+		JsonWriter->WriteValue(TEXT("CompressonFormat"), Toc.Header.CompressionFormat);
+		JsonWriter->WriteValue(TEXT("ChunksDirectory"), Toc.Header.ChunksDirectory);
+	}
+	JsonWriter->WriteObjectEnd();
+	JsonWriter->WriteObjectStart(TEXT("Meta"));
+	{
+		const FDateTime Dt = FDateTime::FromUnixTimestamp(Toc.Meta.EpochTimestamp);
+		JsonWriter->WriteValue(TEXT("DateTime"), Dt.ToString()), 
+		JsonWriter->WriteValue(TEXT("BuildVersion"), Toc.Meta.BuildVersion);
+		JsonWriter->WriteValue(TEXT("TargetPlatform"), Toc.Meta.TargetPlatform);
+	}
+	JsonWriter->WriteObjectEnd();
+	JsonWriter->WriteArrayStart(TEXT("Containers"));
+	{
+		for (const FOnDemandTocContainerEntry& Container : Toc.Containers)
+		{
+			JsonWriter->WriteObjectStart();
+			JsonWriter->WriteValue(TEXT("Name"), Container.ContainerName);
+			JsonWriter->WriteValue(TEXT("EncryptionKeyGuid"), Container.EncryptionKeyGuid);
+
+			if (EnumHasAnyFlags(Options, EOnDemandTocJsonOptions::TocEntries))
+			{
+				JsonWriter->WriteArrayStart("Entries");
+				for (const FOnDemandTocEntry& Entry : Container.Entries)
+				{
+					JsonWriter->WriteObjectStart();
+					JsonWriter->WriteValue(TEXT("Hash"), LexToString(Entry.Hash));
+					JsonWriter->WriteValue(TEXT("ChunkId"), LexToString(Entry.ChunkId));
+					JsonWriter->WriteValue(TEXT("RawSize"), Entry.RawSize);
+					JsonWriter->WriteValue(TEXT("EncodedSize"), Entry.EncodedSize);
+					JsonWriter->WriteValue(TEXT("BlockOffset"), int32(Entry.BlockOffset));
+					JsonWriter->WriteValue(TEXT("BlockCount"), int32(Entry.BlockCount));
+					JsonWriter->WriteObjectEnd();
+				}
+				JsonWriter->WriteArrayEnd();
+			}
+
+			if (EnumHasAnyFlags(Options, EOnDemandTocJsonOptions::BlockSizes))
+			{
+				JsonWriter->WriteArrayStart("Blocks");
+				{
+					for (uint32 BlockSize : Container.BlockSizes)
+					{
+						JsonWriter->WriteValue(int32(BlockSize));
+					}
+				}
+				JsonWriter->WriteArrayEnd();
+			}
+
+			if (EnumHasAnyFlags(Options, EOnDemandTocJsonOptions::BlockHashes))
+			{
+				JsonWriter->WriteArrayStart("BlockHashes");
+				{
+					for (const FIoBlockHash& BlockHash : Container.BlockHashes)
+					{
+						JsonWriter->WriteValue(int32(BlockHash));
+					}
+				}
+				JsonWriter->WriteArrayEnd();
+			}
+			JsonWriter->WriteObjectEnd();
+		}
+	}
+	JsonWriter->WriteArrayEnd();
+}
+
+////////////////////////////////////////////////////////////////////////////////
 class FS3UploadQueue
 {
 public:
@@ -1533,6 +1649,437 @@ FIoStatus DownloadContainerFiles(const FIoStoreDownloadParams& DownloadParams, c
 	return FIoStatus::Ok;
 }
 
+////////////////////////////////////////////////////////////////////////////////
+TIoStatusOr<FIoStoreListTocsParams> FIoStoreListTocsParams::Parse(const TCHAR* CommandLine)
+{
+	FIoStoreListTocsParams Params;
+
+	// Convenience argument to specifiy both bucket and bucket prefix
+	// -BucketPath="mybucket/some/data/path" is equal to -Bucket="bucket" -BucketPrefix="some/data/path
+	FString BucketPath;
+	if (FParse::Value(CommandLine, TEXT("-BucketPath="), BucketPath))
+	{
+		FStringView PathView = BucketPath;
+		int32 Idx = INDEX_NONE;
+		if (PathView.FindChar(TCHAR('/'), Idx))
+		{
+			Params.Bucket = PathView.Left(Idx);
+			FStringView BucketPrefix = PathView.RightChop(Idx + 1);
+			if (BucketPrefix.EndsWith(TEXT("/")))
+			{
+				BucketPrefix.LeftChopInline(1);
+			}
+			Params.BucketPrefix = BucketPrefix;
+		}
+		else
+		{
+			Params.Bucket = MoveTemp(BucketPath);
+		}
+	}
+
+	if (Params.Bucket.IsEmpty())
+	{
+		if (!FParse::Value(CommandLine, TEXT("Bucket="), Params.Bucket))
+		{
+			FIoStatus(EIoErrorCode::InvalidParameter, TEXT("Invalid bucket name"));
+		}
+	}
+
+	FParse::Value(CommandLine, TEXT("BucketPrefix="), Params.BucketPrefix);
+	FParse::Value(CommandLine, TEXT("ServiceUrl="), Params.ServiceUrl);
+	FParse::Value(CommandLine, TEXT("Region="), Params.Region);
+	FParse::Value(CommandLine, TEXT("AccessKey="), Params.AccessKey);
+	FParse::Value(CommandLine, TEXT("SecretKey="), Params.SecretKey);
+	FParse::Value(CommandLine, TEXT("SessionToken="), Params.SessionToken);
+	FParse::Value(CommandLine, TEXT("CredentialsFile="), Params.CredentialsFile);
+	FParse::Value(CommandLine, TEXT("TocKey="), Params.TocKey);
+	FParse::Value(CommandLine, TEXT("BuildVersion="), Params.BuildVersion);
+	FParse::Value(CommandLine, TEXT("TargetPlatform="), Params.TargetPlatform);
+	FParse::Value(CommandLine, TEXT("ChunkKeys="), Params.ChunkKeys);
+	// JSON serialization options
+	Params.bTocEntries = FParse::Param(CommandLine, TEXT("TocEntries"));
+	Params.bBlockSizes = FParse::Param(CommandLine, TEXT("BlockSizes"));
+	Params.bBlockHashes = FParse::Param(CommandLine, TEXT("BlockHashes"));
+
+	if (!FParse::Value(CommandLine, TEXT("CredentialsFileKeyName="), Params.CredentialsFileKeyName))
+	{
+		Params.CredentialsFileKeyName = TEXT("default");
+	}
+
+	FParse::Value(CommandLine, TEXT("Json="), Params.OutFile);
+
+
+	if (FParse::Value(CommandLine, TEXT("TocUrl="), Params.TocUrl))
+	{
+		FStringView UrlView(Params.TocUrl);
+		if (UrlView.StartsWith(TEXTVIEW("http://")) && UrlView.EndsWith(TEXTVIEW(".iochunktoc")))
+		{
+			int32 Delim = INDEX_NONE;
+			if (UrlView.RightChop(7).FindChar(TEXT('/'), Delim))
+			{
+				Params.ServiceUrl = FString(UrlView.Left(7 +  Delim));
+				Params.TocKey = UrlView.RightChop(Params.ServiceUrl.Len() + 1);
+			}
+		}
+	}
+
+	if (!Params.AccessKey.IsEmpty() && Params.SecretKey.IsEmpty())
+	{
+		return FIoStatus(EIoErrorCode::InvalidParameter, TEXT("Invalid secret key"));
+	}
+	else if (Params.AccessKey.IsEmpty() && !Params.SecretKey.IsEmpty())
+	{
+		return FIoStatus(EIoErrorCode::InvalidParameter, TEXT("Invalid access key"));
+	}
+
+	if (!Params.CredentialsFile.IsEmpty() && Params.CredentialsFileKeyName.IsEmpty())
+	{
+		return FIoStatus(EIoErrorCode::InvalidParameter, TEXT("Invalid credential file key name"));
+	}
+
+	if (Params.TocUrl.IsEmpty() && Params.ServiceUrl.IsEmpty() && Params.Region.IsEmpty())
+	{
+		return FIoStatus(EIoErrorCode::InvalidParameter, TEXT("Service URL or AWS region needs to be specified"));
+	}
+
+	return Params;
+}
+
+FIoStatus ListTocs(const FIoStoreListTocsParams& Params)
+{
+	struct FChunkStats
+	{
+		TSet<FIoHash> Chunks;
+		uint64 TotalChunkSize = 0;
+
+		void Add(const FIoHash& Hash, uint64 ChunkSize)
+		{
+			bool bExist = false;
+			Chunks.Add(Hash, &bExist);
+			if (!bExist)
+			{
+				TotalChunkSize += ChunkSize;
+			}
+		}
+	};
+	TMap<FString, FChunkStats> Stats;
+
+	struct FTocDescription
+	{
+		FOnDemandToc Toc;
+		FDateTime DateTime;
+		FString Key;
+		uint64 Size = 0;
+		uint64 ChunkCount = 0;
+		uint64 TotalChunkSize = 0;
+	};
+
+	const bool bFilteredQuery = !Params.TocKey.IsEmpty() || !Params.BuildVersion.IsEmpty();
+
+	FS3ClientConfig Config;
+	Config.ServiceUrl = Params.ServiceUrl;
+	Config.Region = Params.Region;
+	
+	FS3ClientCredentials Credentials;
+	if (Params.CredentialsFile.IsEmpty() == false)
+	{
+		UE_LOG(LogIas, Display, TEXT("Loading credentials file '%s'"), *Params.CredentialsFile);
+		FS3CredentialsProfileStore CredentialsStore = FS3CredentialsProfileStore::FromFile(Params.CredentialsFile);
+		if (CredentialsStore.TryGetCredentials(Params.CredentialsFileKeyName, Credentials) == false)
+		{
+			return FIoStatus(EIoErrorCode::InvalidParameter, TEXT("Failed to find valid credentials in credentials file"));
+		}
+		else
+		{
+			UE_LOG(LogIas, Display, TEXT("Found credentials for '%s'"), *Params.CredentialsFileKeyName);
+		}
+	}
+	else
+	{
+		Credentials = FS3ClientCredentials(Params.AccessKey, Params.SecretKey, Params.SessionToken);
+	}
+
+	FS3Client Client(Config, Credentials);
+	FMutex TocMutex;
+	TArray<FTocDescription> Tocs;
+
+	if (Params.TocUrl.IsEmpty())
+	{
+		FStringView PrefixView = Params.BucketPrefix;
+		if (PrefixView.StartsWith(TEXT("/")))
+		{
+			PrefixView.RightChopInline(1);
+		}
+
+		TStringBuilder<256> Path;
+		Path << PrefixView << TEXT("/");
+
+		UE_LOG(LogIas, Display, TEXT("Fetching TOC's from '%s/%s/%s'"), *Client.GetConfig().ServiceUrl, *Params.Bucket, Path.ToString());
+		FS3ListObjectResponse Response = Client.ListObjects(FS3ListObjectsRequest
+		{
+			Params.Bucket,
+			Path.ToString(),
+			TEXT('/')
+		});
+
+		if (Response.Objects.IsEmpty())
+		{
+			UE_LOG(LogIas, Display, TEXT("Not TOC's found at '%s/%s/%s' (StatusCode: %u)"), *Client.GetConfig().ServiceUrl, *Params.Bucket, Path.ToString(), Response.StatusCode);
+			return FIoStatus(EIoErrorCode::NotFound);
+		}
+
+		ParallelFor(Response.Objects.Num(), [&Params, &Client, &Response, &Tocs, &TocMutex](int32 Index)
+		{
+			const FS3Object& Obj = Response.Objects[Index];
+			if (Obj.Key.EndsWith(TEXT("iochunktoc")) == false)
+			{
+				return;
+			}
+
+			if (!Params.TocKey.IsEmpty() && !Params.TocKey.Equals(FPaths::GetBaseFilename(Obj.Key), ESearchCase::IgnoreCase))
+			{
+				return;
+			}
+
+			UE_LOG(LogIas, Display, TEXT("Fetching TOC '%s/%s/%s'"), *Client.GetConfig().ServiceUrl, *Params.Bucket, *Obj.Key);
+			FS3GetObjectResponse TocResponse = Client.GetObject(FS3GetObjectRequest
+			{
+				Params.Bucket,
+				Obj.Key
+			});
+
+			if (TocResponse.IsOk() == false)
+			{
+				UE_LOG(LogIas, Warning, TEXT("Failed to fetch TOC '%s/%s/%s'"), *Client.GetConfig().ServiceUrl, *Params.Bucket, *Obj.Key);
+				return;
+			}
+
+			FOnDemandToc Toc;
+			FMemoryReaderView Ar(TocResponse.Body.GetView());
+			Ar << Toc;
+
+			if (Toc.Header.Magic != FOnDemandTocHeader::ExpectedMagic)
+			{
+				UE_LOG(LogIas, Warning, TEXT("Failed to serialize TOC '%s/%s/%s'. Header magic mimsatch"), *Client.GetConfig().ServiceUrl, *Params.Bucket, *Obj.Key);
+				return;
+			}
+
+			if (!Params.BuildVersion.IsEmpty())
+			{
+				if (!Params.BuildVersion.Equals(Toc.Meta.BuildVersion, ESearchCase::IgnoreCase))
+				{
+					return;
+				}
+				if (!Params.TargetPlatform.IsEmpty() && !Params.TargetPlatform.Equals(Toc.Meta.TargetPlatform, ESearchCase::IgnoreCase))
+				{
+					return;
+				}
+			}
+
+			FDateTime DateTime = FDateTime::FromUnixTimestamp(Toc.Meta.EpochTimestamp);
+			FTocDescription Desc = FTocDescription 
+			{
+				.Toc = MoveTemp(Toc),
+				.DateTime = DateTime,
+				.Key = Obj.Key,
+				.Size = Obj.Size
+			};
+
+			{
+				TUniqueLock Lock(TocMutex);
+				Tocs.Add(MoveTemp(Desc));
+			}
+		});
+	}
+	else
+	{
+		UE_LOG(LogIas, Display, TEXT("Fetching TOC '%s/%s'"), *Client.GetConfig().ServiceUrl, *Params.TocKey);
+		FS3GetObjectResponse TocResponse = Client.GetObject(FS3GetObjectRequest
+		{
+			TEXT("/"),
+			Params.TocKey
+		});
+
+		if (TocResponse.IsOk())
+		{
+			FOnDemandToc Toc;
+			FMemoryReaderView Ar(TocResponse.Body.GetView());
+			Ar << Toc;
+
+			if (Toc.Header.Magic == FOnDemandTocHeader::ExpectedMagic)
+			{
+				FDateTime DateTime = FDateTime::FromUnixTimestamp(Toc.Meta.EpochTimestamp);
+				Tocs.Add(FTocDescription 
+				{
+					.Toc = MoveTemp(Toc),
+					.DateTime = DateTime,
+					.Key = Params.TocKey,
+					.Size = TocResponse.Body.GetSize()
+				});
+			}
+			else
+			{
+				UE_LOG(LogIas, Warning, TEXT("Failed to serialize TOC '%s/%s'. Header magic mimsatch"), *Client.GetConfig().ServiceUrl, *Params.TocKey);
+			}
+		}
+	}
+
+	if (Tocs.IsEmpty())
+	{
+		return FIoStatus(EIoErrorCode::NotFound);
+	}
+
+	for (FTocDescription& Desc : Tocs)
+	{
+		FChunkStats& BuildStats = Stats.FindOrAdd(Desc.Toc.Meta.BuildVersion);
+		FChunkStats& BucketStats = Stats.FindOrAdd(Params.Bucket);
+		for (const FOnDemandTocContainerEntry& Container : Desc.Toc.Containers)
+		{
+			for (const FOnDemandTocEntry& Entry : Container.Entries)
+			{
+				Desc.ChunkCount++;
+				Desc.TotalChunkSize += Entry.EncodedSize;
+				BucketStats.Add(Entry.Hash, Entry.EncodedSize);
+				BuildStats.Add(Entry.Hash, Entry.EncodedSize);
+			}
+		}
+	}
+
+	Tocs.Sort([](const auto& LHS, const auto& RHS) { return LHS.DateTime > RHS.DateTime; });
+
+	if (Params.OutFile.IsEmpty())
+	{
+		int32 Counter = 1;
+		TStringBuilder<256> Url;
+
+		for (const FTocDescription& Desc : Tocs)
+		{
+			Url.Reset();
+			Url << Client.GetConfig().ServiceUrl << TEXT("/") << Params.Bucket << TEXT("/") << Desc.Key;
+
+			UE_LOG(LogIas, Display, TEXT(""));
+			UE_LOG(LogIas, Display, TEXT("%d) %s"), Counter++, *Desc.Key);
+			UE_LOG(LogIas, Display, TEXT("%-20s: %s"),TEXT("Date"), *Desc.DateTime.ToString());
+			UE_LOG(LogIas, Display, TEXT("%-20s: %s"), TEXT("BuildVersion"), *Desc.Toc.Meta.BuildVersion);
+			UE_LOG(LogIas, Display, TEXT("%-20s: %s"), TEXT("TargetPlatform"), *Desc.Toc.Meta.TargetPlatform);
+			UE_LOG(LogIas, Display, TEXT("%-20s: %.2lf KiB"), TEXT("TocSize"), double(Desc.Size) / 1024.0);
+			UE_LOG(LogIas, Display, TEXT("%-20s: %llu"), TEXT("ChunkCount"), Desc.ChunkCount);
+			UE_LOG(LogIas, Display, TEXT("%-20s: %.2lf MiB"), TEXT("TotalChunkSize"), double(Desc.TotalChunkSize) / 1024.0 / 1024.0);
+			UE_LOG(LogIas, Display, TEXT("%-20s: %s"), TEXT("Url"), Url.ToString());
+		}
+		UE_LOG(LogIas, Display, TEXT(""));
+		
+		if (bFilteredQuery == false)
+		{
+			TArray<FString> Keys;
+			Stats.GetKeys(Keys);
+			Keys.Sort();
+
+			UE_LOG(LogIas, Display, TEXT("%-80s %15s %15s"), TEXT("BuildVersion"), TEXT("Chunk(s)"), TEXT("MiB"));
+			UE_LOG(LogIas, Display, TEXT("-------------------------------------------------------------------------------------------------------------------"));
+			for (const FString& Key : Keys)
+			{
+				if (Key != Params.Bucket)
+				{
+					const FChunkStats& ChunkStats = Stats.FindChecked(Key);
+					UE_LOG(LogIas, Display, TEXT("%-80s %15llu %15.2f"), *Key, ChunkStats.Chunks.Num(), double(ChunkStats.TotalChunkSize) / 1024.0 / 1024.0);
+				}
+			}
+			UE_LOG(LogIas, Display, TEXT("-------------------------------------------------------------------------------------------------------------------"));
+			FChunkStats& BucketStats = Stats.FindOrAdd(Params.Bucket);
+			UE_LOG(LogIas, Display, TEXT("%-80s %15llu %15.2f"), *Params.Bucket, BucketStats.Chunks.Num(), double(BucketStats.TotalChunkSize) / 1024.0 / 1024.0);
+			UE_LOG(LogIas, Display, TEXT(""));
+		}
+	}
+	else
+	{
+		FString Json;
+		FJsonWriter JsonWriter = FJsonWriterFactory::Create(&Json);
+		JsonWriter->WriteObjectStart();
+		JsonWriter->WriteValue(TEXT("ServiceUrl"), Client.GetConfig().ServiceUrl);
+		JsonWriter->WriteValue(TEXT("Bucket"), Params.Bucket);
+		JsonWriter->WriteValue(TEXT("BucketPrefix"), Params.BucketPrefix);
+
+		JsonWriter->WriteArrayStart(TEXT("Tocs"));
+		for (const FTocDescription& Desc : Tocs)
+		{
+			JsonWriter->WriteObjectStart();
+			JsonWriter->WriteValue(TEXT("Key"), Desc.Key);
+			JsonWriter->WriteValue(TEXT("ChunkCount"), Desc.ChunkCount);
+			JsonWriter->WriteValue(TEXT("TotalChunkSize"), Desc.TotalChunkSize);
+
+			EOnDemandTocJsonOptions JsonOptions = EOnDemandTocJsonOptions::Header;
+			if (Params.bTocEntries)
+			{
+				JsonOptions |= EOnDemandTocJsonOptions::TocEntries;
+			}
+			if (Params.bBlockSizes)
+			{
+				JsonOptions |= EOnDemandTocJsonOptions::BlockSizes;
+			}
+			if (Params.bBlockHashes)
+			{
+				JsonOptions |= EOnDemandTocJsonOptions::BlockHashes;
+			}
+			ToJson(JsonWriter, Desc.Toc, JsonOptions);
+			JsonWriter->WriteObjectEnd();
+		}
+		JsonWriter->WriteArrayEnd();
+
+		JsonWriter->WriteObjectEnd();
+		JsonWriter->Close();
+
+		UE_LOG(LogIas, Display, TEXT("Saving file '%s'"), *Params.OutFile);
+		if (!FFileHelper::SaveStringToFile(Json, *Params.OutFile))
+		{
+			return FIoStatus(EIoErrorCode::WriteError, TEXTVIEW("Failed writing JSON file")); 
+		}
+	}
+
+	if (Params.ChunkKeys.IsEmpty() == false)
+	{
+		FString Json;
+		FJsonWriter JsonWriter = FJsonWriterFactory::Create(&Json);
+		
+		JsonWriter->WriteObjectStart();
+		JsonWriter->WriteValue(TEXT("ServiceUrl"), Client.GetConfig().ServiceUrl);
+		JsonWriter->WriteValue(TEXT("Bucket"), Params.Bucket);
+		JsonWriter->WriteValue(TEXT("BucketPrefix"), Params.BucketPrefix);
+
+		JsonWriter->WriteArrayStart(TEXT("Tocs"));
+		for (const FTocDescription& Desc : Tocs)
+		{
+			TArray<FString> ObjKeys;
+			GetChunkObjectKeys(Desc.Toc, Params.BucketPrefix, ObjKeys);
+
+			JsonWriter->WriteObjectStart();
+			JsonWriter->WriteValue(TEXT("Key"), Desc.Key);
+			JsonWriter->WriteValue(TEXT("BuildVersion"), Desc.Toc.Meta.BuildVersion);
+			JsonWriter->WriteValue(TEXT("TargetPlatform"), Desc.Toc.Meta.TargetPlatform);
+			JsonWriter->WriteArrayStart(TEXT("ChunkKeys"));
+			for (const FString& Key : ObjKeys)
+			{
+				JsonWriter->WriteValue(Key);
+			}
+			JsonWriter->WriteArrayEnd();
+			JsonWriter->WriteObjectEnd();
+		}
+		JsonWriter->WriteArrayEnd();
+
+		JsonWriter->WriteObjectEnd();
+		JsonWriter->Close();
+
+		UE_LOG(LogIas, Display, TEXT("Saving chunk key(s) '%s'"), *Params.ChunkKeys);
+		if (!FFileHelper::SaveStringToFile(Json, *Params.ChunkKeys))
+		{
+			return FIoStatus(EIoErrorCode::WriteError, TEXTVIEW("Failed writing JSON file")); 
+		}
+	}
+
+	return FIoStatus::Ok;
+}
+
+////////////////////////////////////////////////////////////////////////////////
 FIoStatus PrimeEndPoint(FStringView IoStoreOnDemandIniPath)
 {
 	return UE::IO::IAS::PrimeEndpointInternal(FString(IoStoreOnDemandIniPath));
diff --git a/Engine/Source/Runtime/Experimental/IoStoreOnDemand/Public/IO/IoStoreOnDemand.h b/Engine/Source/Runtime/Experimental/IoStoreOnDemand/Public/IO/IoStoreOnDemand.h
--- a/Engine/Source/Runtime/Experimental/IoStoreOnDemand/Public/IO/IoStoreOnDemand.h
+++ b/Engine/Source/Runtime/Experimental/IoStoreOnDemand/Public/IO/IoStoreOnDemand.h
@@ -211,6 +211,64 @@ struct FIoStoreDownloadParams
 
 UE_API FIoStatus DownloadContainerFiles(const FIoStoreDownloadParams& DownloadParams, const FString& TocPath);
 
+/**
+ * Parameters for listing uploaded TOC file(s) from an S3 compatible endpoint.
+ *
+ * Example usage:
+ *
+ * 1) Print available TOC's from a local server to standard out.
+ * UnrealPak.exe -ListTocs -ServiceUrl="http://10.24.101.92:9000" -Bucket=<bucketname> -BucketPrefix=<some/path/to/data> -AccessKey=<accesskey> -SecretKey=<secretkey>
+ *
+ * 2) Print available TOC's from AWS S3.
+ * UnrealPak.exe -Region="us-east-1" -BucketPath=<mybucket/some/data/path/> -CredentialsFile=<path/to/credentials.txt> -BuildVersion=<version> -Json=<path/to/file.json>
+ * 
+ * 3) Serialize all TOC's matching a specific build version to JSON:
+ * UnrealPak.exe -Region="us-east-1" -BucketPath=<mybucket/some/data/path/> -CredentialsFile=<path/to/credentials.txt> -BuildVersion=<version> -Json=<path/to/file.json>
+ *
+ * 4) Serialize all chunk object key(s) to JSON.
+ * UnrealPak.exe -Region="us-east-1" -BucketPath=<mybucket/some/data/path/> -CredentialsFile=<path/to/credentials.txt> -BuildVersion=<version> -ChunkKeys=<path/to/file.json>
+ *
+ * 5) Fetch a TOC from a public CDN.
+ * UnrealPak.exe -ListTocs -TocUrl=<http://some.public.endpoint.net/path/to/1a32076ca12bfc6feb982ffb064d18f28156606c.iochunktoc>
+ *
+ * Parameters: -TocEntries, -BlockSizes and -BlockHashes controls what to include when serializing TOC's to JSON.
+ *
+ * Credentials file example:
+ *
+ * [default]
+ * aws_access_key_id="<key>"
+ * aws_secret_access_key="<key>
+ * aws_session_token="<token>"
+ *
+ * Note: All values must be surounded with "".
+ */
+struct FIoStoreListTocsParams
+{
+	FString OutFile;
+	FString ServiceUrl;
+	FString Bucket;
+	FString BucketPrefix;
+	FString Region; 
+	FString AccessKey;
+	FString SecretKey;
+	FString SessionToken;
+	FString CredentialsFile;
+	FString CredentialsFileKeyName;
+	FString TocUrl;
+	FString TocKey;
+	FString BuildVersion;
+	FString TargetPlatform;
+	FString ChunkKeys;
+	bool bTocEntries = false;
+	bool bBlockSizes = false;
+	bool bBlockHashes = false;
+
+	UE_API static TIoStatusOr<FIoStoreListTocsParams> Parse(const TCHAR* CommandLine);
+};
+
+UE_API FIoStatus ListTocs(const FIoStoreListTocsParams& Params);
+
+////////////////////////////////////////////////////////////////////////////////
 UE_API FIoStatus PrimeEndPoint(FStringView IoStoreOnDemandIniPath);
 #endif // (IS_PROGRAM || WITH_EDITOR)
 
