IAS - config value for controlling minimum size for partial chunk request

diff --git a/Engine/Source/Runtime/Experimental/IoStoreOnDemand/Private/OnDemandIoDispatcherBackend.cpp b/Engine/Source/Runtime/Experimental/IoStoreOnDemand/Private/OnDemandIoDispatcherBackend.cpp
--- a/Engine/Source/Runtime/Experimental/IoStoreOnDemand/Private/OnDemandIoDispatcherBackend.cpp
+++ b/Engine/Source/Runtime/Experimental/IoStoreOnDemand/Private/OnDemandIoDispatcherBackend.cpp
@@ -213,6 +213,13 @@ static FAutoConsoleVariableRef CVar_IasAsyncTocGeneration(
 	ECVF_ReadOnly
 );
 
+static int32 GIasHttpRangeRequestMinSizeKiB = 0;
+static FAutoConsoleVariableRef CVar_IasHttpRangeRequestMinSizeKiB(
+	TEXT("ias.HttpRangeRequestMinSizeKiB"),
+	GIasHttpRangeRequestMinSizeKiB,
+	TEXT("Minimum chunk size for partial chunk request(s)")
+);
+
 #if !UE_BUILD_SHIPPING
 static FAutoConsoleCommand CVar_IasAbandonCache(
 	TEXT("Ias.AbandonCache"),
@@ -675,14 +682,22 @@ struct FChunkRequestParams
 {
 	static FChunkRequestParams Create(const FIoOffsetAndLength& OffsetLength, FOnDemandIoStore::FChunkInfo ChunkInfo)
 	{
-		const uint64 RawSize = FMath::Min<uint64>(OffsetLength.GetLength(), ChunkInfo.Entry->RawSize);
-		
-		const FIoOffsetAndLength ChunkRange = FIoChunkEncoding::GetChunkRange(
-			ChunkInfo.Entry->RawSize,
-			ChunkInfo.Container->BlockSize,
-			ChunkInfo.GetBlocks(),
-			OffsetLength.GetOffset(),
-			RawSize).ConsumeValueOrDie();
+		FIoOffsetAndLength ChunkRange;
+		if (ChunkInfo.Entry->EncodedSize <= (uint64(GIasHttpRangeRequestMinSizeKiB) << 10))
+		{
+			ChunkRange = FIoOffsetAndLength(0, ChunkInfo.Entry->EncodedSize);
+		}
+		else
+		{
+			const uint64 RawSize = FMath::Min<uint64>(OffsetLength.GetLength(), ChunkInfo.Entry->RawSize);
+
+			ChunkRange = FIoChunkEncoding::GetChunkRange(
+				ChunkInfo.Entry->RawSize,
+				ChunkInfo.Container->BlockSize,
+				ChunkInfo.GetBlocks(),
+				OffsetLength.GetOffset(),
+				RawSize).ConsumeValueOrDie();
+		}
 
 		return FChunkRequestParams{GetChunkKey(ChunkInfo.Entry->Hash, ChunkRange), ChunkRange, ChunkInfo};
 	}
@@ -840,7 +855,8 @@ static void LogIoResult(
 	uint64 DurationMs,
 	uint64 UncompressedSize,
 	uint64 UncompressedOffset,
-	uint64 CompressedOffset,
+	const FIoOffsetAndLength& ChunkRange,
+	uint64 ChunkSize,
 	int32 Priority,
 	bool bCached)
 {
@@ -853,15 +869,38 @@ static void LogIoResult(
 		return bCached ? TEXT("io-cache") : TEXT("io-http ");
 	}();
 
-	UE_LOG(LogIas, VeryVerbose, TEXT("%s: %5" UINT64_FMT "ms %5" UINT64_FMT "KiB[%7" UINT64_FMT "] % s: % s | %" UINT64_FMT "(%d)"),
+	auto PrioToString = [](int32 Prio) -> const TCHAR*
+	{
+		if (Prio < IoDispatcherPriority_Low)
+		{
+			return TEXT("Min");
+		}
+		if (Prio < IoDispatcherPriority_Medium)
+		{
+			return TEXT("Low");
+		}
+		if (Prio < IoDispatcherPriority_High)
+		{
+			return TEXT("Medium");
+		}
+		if (Prio < IoDispatcherPriority_Max)
+		{
+			return TEXT("High");
+		}
+
+		return TEXT("Max");
+	};
+
+	UE_LOG(LogIas, VeryVerbose, TEXT("%s: %5" UINT64_FMT "ms %5" UINT64_FMT "KiB[%7" UINT64_FMT "] % s: % s | Range: %" UINT64_FMT "-%" UINT64_FMT "/%" UINT64_FMT " (%.2f%%) | Prio: %s"),
 		Prefix,
 		DurationMs,
 		UncompressedSize >> 10,
 		UncompressedOffset,
 		*LexToString(ChunkId),
 		*LexToString(UrlHash),
-		CompressedOffset,
-		Priority);
+		ChunkRange.GetOffset(), (ChunkRange.GetOffset() + ChunkRange.GetLength() - 1), ChunkSize,
+		100.0f * (float(ChunkRange.GetLength()) / float(ChunkSize)),
+		PrioToString(Priority));
 };
 
 /** Utility for finding all available on demand utoc files currently in valid pak directories */
@@ -1721,7 +1760,8 @@ void FOnDemandIoBackend::CompleteRequest(FChunkRequest* ChunkRequest)
 			Stats.OnIoRequestComplete(Request->GetBuffer().GetSize(), DurationMs);
 			LogIoResult(Request->ChunkId, ChunkRequest->Params.GetUrlHash(), DurationMs,
 				Request->GetBuffer().DataSize(), Request->Options.GetOffset(),
-				ChunkRequest->Params.ChunkRange.GetOffset(), ChunkRequest->Priority, ChunkRequest->bCached);
+				ChunkRequest->Params.ChunkRange, ChunkRequest->Params.ChunkInfo.Entry->EncodedSize,
+				ChunkRequest->Priority, ChunkRequest->bCached);
 				
 		}
 		else
@@ -1732,7 +1772,8 @@ void FOnDemandIoBackend::CompleteRequest(FChunkRequest* ChunkRequest)
 			Stats.OnIoRequestError();
 			LogIoResult(Request->ChunkId, ChunkRequest->Params.GetUrlHash(), DurationMs,
 				0, Request->Options.GetOffset(),
-				ChunkRequest->Params.ChunkRange.GetOffset(), ChunkRequest->Priority, ChunkRequest->bCached);
+				ChunkRequest->Params.ChunkRange, ChunkRequest->Params.ChunkInfo.Entry->EncodedSize,
+				ChunkRequest->Priority, ChunkRequest->bCached);
 		}
 
 		CompletedRequests.Enqueue(Request);
