Remove ensure about Nans iwhen getting the root transform until the proxy race condition is properly fixed

diff --git a/Engine/Source/Runtime/Experimental/GeometryCollectionEngine/Private/GeometryCollection/GeometryCollectionComponent.cpp b/Engine/Source/Runtime/Experimental/GeometryCollectionEngine/Private/GeometryCollection/GeometryCollectionComponent.cpp
--- a/Engine/Source/Runtime/Experimental/GeometryCollectionEngine/Private/GeometryCollection/GeometryCollectionComponent.cpp
+++ b/Engine/Source/Runtime/Experimental/GeometryCollectionEngine/Private/GeometryCollection/GeometryCollectionComponent.cpp
@@ -5741,7 +5741,6 @@ const FTransform3f& UGeometryCollectionComponent::FComponentSpaceTransforms::Req
 			bIsRootDirty = false;
 		}
 
-		ensureAlways(!Transforms[RootIndex].ContainsNaN());
 		return Transforms[RootIndex];
 	}
 
