Fix missing IAS statistic.

#rnx

- If the connection to the backend has been interrupted and we have marked it as being in an error state we will not actively push anymore http requests and existing pending requests will be ignored.
- In one code path these pending requests will have already been marked as "pending" in the statistics so we need to signal to the statistics system that the request is no longer happening otherwise the pending stats become inaccurate.
-The two options at the moment are to mark the request as encountering a http error or mark it as cancelled. Neither case accurately describes what actually happened and we will extend the statistics system to better reflect this in the future but for now we will simply mark them as cancelled so as not to inflate the http error count when gathering analytics.

diff --git a/Engine/Source/Runtime/Experimental/IoStoreOnDemand/Private/OnDemandIoDispatcherBackend.cpp b/Engine/Source/Runtime/Experimental/IoStoreOnDemand/Private/OnDemandIoDispatcherBackend.cpp
--- a/Engine/Source/Runtime/Experimental/IoStoreOnDemand/Private/OnDemandIoDispatcherBackend.cpp
+++ b/Engine/Source/Runtime/Experimental/IoStoreOnDemand/Private/OnDemandIoDispatcherBackend.cpp
@@ -2295,6 +2295,12 @@ void FOnDemandIoBackend::ProcessHttpRequests(FHttpClient& HttpClient, FBitWindow
 					{
 						CompleteRequest(ChunkRequest);
 					});
+
+					// Technically this request is being skipped because of a pre-existing error. It is not
+					// an error itself and it is not being canceled by higher level code. However we do not
+					// currently have a statistic for that and we have to call one of the existing types in
+					// order to correctly reduce the pending count.
+					Stats.OnHttpCancel();
 				}
 				else
 				{
