Iris
- Fixed crash in ReplicationReader.cpp due to destroying index 0.

diff --git a/Engine/Source/Runtime/Experimental/Iris/Core/Private/Iris/ReplicationSystem/ReplicationReader.cpp b/Engine/Source/Runtime/Experimental/Iris/Core/Private/Iris/ReplicationSystem/ReplicationReader.cpp
--- a/Engine/Source/Runtime/Experimental/Iris/Core/Private/Iris/ReplicationSystem/ReplicationReader.cpp
+++ b/Engine/Source/Runtime/Experimental/Iris/Core/Private/Iris/ReplicationSystem/ReplicationReader.cpp
@@ -222,7 +222,7 @@ void FReplicationReader::Init(const FReplicationParameters& InParameters)
 	}
 
 	// reserve index 0
-	StartReplication(0);
+	StartReplication(ObjectIndexForOOBAttachment);
 }
 
 void FReplicationReader::Deinit()
@@ -392,6 +392,10 @@ void FReplicationReader::CleanupObjectData(FReplicatedObjectInfo& ObjectInfo)
 
 void FReplicationReader::EndReplication(uint32 InternalIndex, bool bTearOff, bool bDestroyInstance)
 {
+	if (!ensure(InternalIndex != ObjectIndexForOOBAttachment))
+	{
+		return;
+	}
 	if (FReplicatedObjectInfo* ObjectInfo = ReplicatedObjects.Find(InternalIndex))
 	{
 		const FNetRefHandleManager::FReplicatedObjectData& Data = NetRefHandleManager->GetReplicatedObjectDataNoCheck(InternalIndex);
@@ -772,6 +776,7 @@ void FReplicationReader::ReadObjectInBatch(FNetSerializationContext& Context, FN
 	//UE_LOG_REPLICATIONREADER(TEXT("FReplicationReader::Read Object with %s InitialState: %u"), *IncompleteHandle.ToString(), bIsInitialState ? 1u : 0u);
 
 	bool bHasErrors = false;
+	bool bIsReplicatedDestroyForInvalidObject = false;
 
 	// Read creation data
 	Context.SetIsInitState(bIsInitialState);
@@ -851,10 +856,17 @@ void FReplicationReader::ReadObjectInBatch(FNetSerializationContext& Context, FN
 			// If we get back an invalid internal index then either the object has been deleted or there's bitstream corruption.
 			InternalIndex = NetRefHandleManager->GetInternalIndex(IncompleteHandle);
 
-			// If this is a subobject that is being destroyed this was no error as we send destroy info for unconfirmed objects
-			if ((ReplicatedDestroyHeaderFlags & ReplicatedDestroyHeaderFlags_EndReplication) == 0U)
+			if (InternalIndex == FNetRefHandleManager::InvalidInternalIndex)
 			{
-				bHasErrors = InternalIndex == FNetRefHandleManager::InvalidInternalIndex;
+				if ((ReplicatedDestroyHeaderFlags & ReplicatedDestroyHeaderFlags_EndReplication) == 0U)
+				{
+					bHasErrors = true;
+				}
+				else
+				{
+					// If this is a subobject that is being destroyed this was no error as we send destroy info for unconfirmed object
+					bIsReplicatedDestroyForInvalidObject = true;
+				}
 			}
 		}
 	}
@@ -880,10 +892,10 @@ void FReplicationReader::ReadObjectInBatch(FNetSerializationContext& Context, FN
 
 		if (bHasState)
 		{
-			bHasErrors = IsObjectIndexForOOBAttachment(InternalIndex);
+			bHasErrors = IsObjectIndexForOOBAttachment(InternalIndex) || bIsReplicatedDestroyForInvalidObject;
 			if (bHasErrors)
 			{
-				UE_LOG_REPLICATIONREADER_WARNING(TEXT("FReplicationReader::ReadObject Bitstream corrupted. Getting state when only expecting RPCs."));
+				UE_LOG_REPLICATIONREADER_WARNING(TEXT("FReplicationReader::ReadObject Bitstream corrupted. Getting state when not expecting state data."));
 				Context.SetError(GNetError_BitStreamError);
 				goto ErrorHandling;
 			}
@@ -954,9 +966,15 @@ void FReplicationReader::ReadObjectInBatch(FNetSerializationContext& Context, FN
 		{
 			if (IsObjectIndexForOOBAttachment(InternalIndex))
 			{
-				const bool bIsHugeObject = Reader.ReadBool();
-				AttachmentType = (bIsHugeObject ? ENetObjectAttachmentType::HugeObject : ENetObjectAttachmentType::OutOfBand);
-				bHasErrors = bHasErrors || (!Parameters.bAllowReceivingAttachmentsFromRemoteObjectsNotInScope && AttachmentType == ENetObjectAttachmentType::OutOfBand);
+				bHasErrors = bHasErrors || bIsReplicatedDestroyForInvalidObject;			
+
+				if (!bHasErrors)
+				{
+					const bool bIsHugeObject = Reader.ReadBool();
+					AttachmentType = (bIsHugeObject ? ENetObjectAttachmentType::HugeObject : ENetObjectAttachmentType::OutOfBand);
+					bHasErrors = bHasErrors || (!Parameters.bAllowReceivingAttachmentsFromRemoteObjectsNotInScope && AttachmentType == ENetObjectAttachmentType::OutOfBand);
+				}
+
 				if (bHasErrors)
 				{
 					Context.SetError(GNetError_InvalidNetHandle);
@@ -973,8 +991,9 @@ void FReplicationReader::ReadObjectInBatch(FNetSerializationContext& Context, FN
 			goto ErrorHandling;
 		}
 
-		// Fill in ReadObjectInfo, we skip HugeObjects as they are not added to the dispatch list until they are fully assembled
-		if (AttachmentType != ENetObjectAttachmentType::HugeObject)
+		// Fill in ReadObjectInfo, we must skip objects that has not been created and HugeObjects as they are not added to the dispatch list until they are fully assembled
+		const bool bShouldCommitPendingDispatchObjectInfo = (AttachmentType != ENetObjectAttachmentType::HugeObject) && !bIsReplicatedDestroyForInvalidObject;
+		if (bShouldCommitPendingDispatchObjectInfo)
 		{
 			Info.InternalIndex = InternalIndex;
 			Info.bIsInitialState = bIsInitialState ? 1U : 0U;
