IAS cache correctly handles failed data writes.

If a data write fails then we do not want to update our journal or add any items to maps, so instead we drop the "phrase" (a collection of cache items written together) and close it off. DiskJournal takes ownership of the phrase, but will simply discard it as

Note that this change includes test code. The edits to FDiskCache::ClosePhrase() are the actual fix.

#rnx
#tests created a disk 64MiB disk and wrote cache data larger than it to force write errors. Uploaded a save data with a nearly full cache on platform S and browsed the locker until new code path was hit.

diff --git a/Engine/Source/Runtime/Experimental/IoStoreOnDemand/Private/JournaledCache.cpp b/Engine/Source/Runtime/Experimental/IoStoreOnDemand/Private/JournaledCache.cpp
--- a/Engine/Source/Runtime/Experimental/IoStoreOnDemand/Private/JournaledCache.cpp
+++ b/Engine/Source/Runtime/Experimental/IoStoreOnDemand/Private/JournaledCache.cpp
@@ -650,11 +650,19 @@ void FDiskCache::ClosePhrase(FDiskPhrase&& Phrase)
 		DataCursor = 0;
 	}
 
+	bool bWriteOk;
 	{
 		TRACE_CPUPROFILER_EVENT_SCOPE(IasCache::DataWrite);
 		const uint8* Buffer = Phrase.GetPhraseData();
 		DataHandle->Seek(DataCursor);
-		DataHandle->Write(Buffer, WriteSize);
+		bWriteOk = DataHandle->Write(Buffer, WriteSize);
+	}
+
+	if (!bWriteOk)
+	{
+		Phrase.Drop();
+		Journal.ClosePhrase(MoveTemp(Phrase), 0);
+		return;
 	}
 
 	{
@@ -2166,13 +2174,21 @@ static uint64 KeyGen(const FIoBuffer& Data)
 ////////////////////////////////////////////////////////////////////////////////
 struct FSupport
 {
-	FSupport()
+	FSupport(const TCHAR* InCacheDir=nullptr)
 	{
 		for (uint32 i = 0; i < WorkingSize; i += 8)
 		{
 			*(uint64*)(Working + i) = Mix();
 		}
 
+		TestDir = (InCacheDir != nullptr) ? FString(InCacheDir) : FPaths::ProjectPersistentDownloadDir();
+		TestDir /= TEXT("ias_cache_test");
+
+		CleanFs();
+	}
+
+	~FSupport()
+	{
 		CleanFs();
 	}
 
@@ -2201,7 +2217,7 @@ struct FSupport
 	const uint64		WorkingSize = 1_Mi;
 	TUniquePtr<uint8[]> WorkingScope = TUniquePtr<uint8[]>(new uint8[WorkingSize]);
 	uint8*				Working = WorkingScope.Get();
-	FString				TestDir = FPaths::ProjectPersistentDownloadDir() / TEXT("ias_cache_test");
+	FString				TestDir;
 };
 
 ////////////////////////////////////////////////////////////////////////////////
@@ -2265,6 +2281,37 @@ static void MemCacheTests(FSupport& Support)
 	Peeled.Reset();
 }
 
+////////////////////////////////////////////////////////////////////////////////
+static void BigCache(FSupport& Support)
+{
+	using namespace JournaledCache;
+
+	FCache::FConfig Config;
+	Config.Path = Support.TestDir / "big_cache";
+	Config.MemoryQuota = uint32(2_Mi);
+	Config.DiskQuota = 512_Mi;
+	Config.JournalQuota = uint32(32_Ki);
+	Config.DropCache = false;
+	FCache Cache(MoveTemp(Config));
+
+	uint32 FlushPeriod = 3;
+	for (uint32 Countdown = 1171; Countdown-- != 0;)
+	{
+		for (uint64 Num = (Support.Mix() % 26) + 1; Num-- != 0;)
+		{
+			uint64 Size = Support.Mix() & ((128_Ki) - 1);
+			FIoBuffer Data = Support.DummyData(Size);
+			Cache.Put(KeyGen(Data), Data);
+		}
+
+		Cache.WriteMemToDisk(uint32(768_Ki));
+		if ((Countdown % FlushPeriod) == 0)
+		{
+			Cache.Flush();
+		}
+	}
+}
+
 ////////////////////////////////////////////////////////////////////////////////
 static void CacheTests(FSupport& Support)
 {
@@ -2468,12 +2515,13 @@ static void MiscTests(FSupport& Support)
 }
 
 ////////////////////////////////////////////////////////////////////////////////
-IOSTOREONDEMAND_API void Tests()
+IOSTOREONDEMAND_API void Tests(const TCHAR* CacheDir=nullptr)
 {
-	FSupport Support;
+	FSupport Support(CacheDir);
 	MiscTests(Support);
 	MemCacheTests(Support);
 	CacheTests(Support);
+	BigCache(Support);
 }
 
 } // namespace IasJournaledFileCacheTest
