Iris
- Added per type stats to report time/waste/etc in csv perfreports

diff --git a/Engine/Source/Runtime/Experimental/Iris/Core/Private/Iris/ReplicationSystem/ObjectReplicationBridge.cpp b/Engine/Source/Runtime/Experimental/Iris/Core/Private/Iris/ReplicationSystem/ObjectReplicationBridge.cpp
--- a/Engine/Source/Runtime/Experimental/Iris/Core/Private/Iris/ReplicationSystem/ObjectReplicationBridge.cpp
+++ b/Engine/Source/Runtime/Experimental/Iris/Core/Private/Iris/ReplicationSystem/ObjectReplicationBridge.cpp
@@ -293,8 +293,8 @@ UE::Net::FNetRefHandle UObjectReplicationBridge::BeginReplication(UObject* Insta
 	const FReplicationProtocol* ReplicationProtocol = ProtocolManager->GetReplicationProtocol(ProtocolIdentifier, ArchetypeOrCDOUsedAsKey);
 	if (!ReplicationProtocol)
 	{
-		constexpr bool bIgnoreProtocolValidation = false;
-		ReplicationProtocol = ProtocolManager->CreateReplicationProtocol(ArchetypeOrCDOUsedAsKey, ProtocolIdentifier, RegisteredFragments, *Instance->GetClass()->GetName(), bIgnoreProtocolValidation);
+		FCreateReplicationProtocolParameters CreateProtocolParams {.ArchetypeOrCDOUsedAsKey = ArchetypeOrCDOUsedAsKey, .TypeStatsIndex = GetTypeStatsIndex(Instance->GetClass()) };
+		ReplicationProtocol = ProtocolManager->CreateReplicationProtocol(ProtocolIdentifier, RegisteredFragments, *Instance->GetClass()->GetName(), CreateProtocolParams);
 	}
 #if UE_IRIS_VALIDATE_PROTOCOLS
 	else
@@ -680,8 +680,8 @@ FReplicationBridgeCreateNetRefHandleResult UObjectReplicationBridge::CreateNetRe
 	const FReplicationProtocol* ReplicationProtocol = ProtocolManager->GetReplicationProtocol(ReceivedProtocolId, ArchetypeOrCDOUsedAsKey);
 	if (!ReplicationProtocol)
 	{
-		constexpr bool bVerifyProtocol = true;
-		ReplicationProtocol = ProtocolManager->CreateReplicationProtocol(ArchetypeOrCDOUsedAsKey, ReceivedProtocolId, RegisteredFragments, *(InstancePtr->GetClass()->GetName()), bVerifyProtocol);
+		FCreateReplicationProtocolParameters CreateProtocolParams {.ArchetypeOrCDOUsedAsKey = ArchetypeOrCDOUsedAsKey, .bValidateProtocolId = true};
+		ReplicationProtocol = ProtocolManager->CreateReplicationProtocol(ReceivedProtocolId, RegisteredFragments, *(InstancePtr->GetClass()->GetName()), CreateProtocolParams);
 	}
 	else
 	{
@@ -1395,6 +1395,42 @@ bool UObjectReplicationBridge::IsClassCritical(const UClass* Class)
 	return false;
 }
 
+int32 UObjectReplicationBridge::GetTypeStatsIndex(const UClass* Class)
+{
+	using namespace UE::Net;
+	using namespace UE::Net::Private;
+	
+	if (ClassesWithTypeStats.Num() > 0)
+	{
+		FNetTypeStats& TypeStats = GetReplicationSystem()->GetReplicationSystemInternal()->GetNetTypeStats();
+		for (; Class != nullptr; Class = Class->GetSuperClass())
+		{
+			if (FName* TypeStatsName = ClassesWithTypeStats.Find(GetConfigClassPathName(Class)))
+			{
+				return TypeStats.GetOrCreateTypeStats(*TypeStatsName);
+			}
+		}
+	}
+
+	return FNetTypeStats::DefaultTypeStatsIndex;
+}
+
+void UObjectReplicationBridge::SetClassTypeStatsConfig(FName ClassPathName, FName TypeStatsName)
+{
+	if (ClassPathName.IsNone())
+	{
+		return;
+	}
+
+	ClassesWithTypeStats.Add(ClassPathName, TypeStatsName);
+
+}
+
+void UObjectReplicationBridge::SetClassTypeStatsConfig(const FString& ClassPathName, const FString& TypeStatsName)
+{
+	SetClassTypeStatsConfig(FName(ClassPathName), FName(TypeStatsName));
+}
+
 void UObjectReplicationBridge::SetClassDynamicFilterConfig(FName ClassPathName, const UE::Net::FNetObjectFilterHandle FilterHandle)
 {
 	if (ClassPathName.IsNone())
@@ -1638,6 +1674,21 @@ void UObjectReplicationBridge::LoadConfig()
 			ClassesFlaggedCritical.Add(CriticalClassConfig.ClassName, CriticalClassConfig.bDisconnectOnProtocolMismatch);
 		}
 	}
+	// Load TypeStats settings
+	{
+		for (const FObjectReplicationBridgeTypeStatsConfig& TypeStatsConfig : BridgeConfig->GetTypeStatsConfigs())
+		{
+#if !UE_NET_IRIS_VERBOSE_CSV_STATS
+			// Skip all non shipping TypeStats
+			if (!TypeStatsConfig.bIncludeInMinimalCSVStats)
+			{
+				continue;
+			}
+#endif			
+			ClassesWithTypeStats.Add(TypeStatsConfig.ClassName, TypeStatsConfig.TypeStatsName);
+
+		}
+	}
 }
 
 void UObjectReplicationBridge::FindClassesInPollPeriodOverrides()
diff --git a/Engine/Source/Runtime/Experimental/Iris/Core/Private/Iris/ReplicationSystem/ObjectReplicationBridgeConfig.cpp b/Engine/Source/Runtime/Experimental/Iris/Core/Private/Iris/ReplicationSystem/ObjectReplicationBridgeConfig.cpp
--- a/Engine/Source/Runtime/Experimental/Iris/Core/Private/Iris/ReplicationSystem/ObjectReplicationBridgeConfig.cpp
+++ b/Engine/Source/Runtime/Experimental/Iris/Core/Private/Iris/ReplicationSystem/ObjectReplicationBridgeConfig.cpp
@@ -37,3 +37,8 @@ TConstArrayView<FObjectReplicatedBridgeCriticalClassConfig> UObjectReplicationBr
 {
 	return MakeArrayView(CriticalClassConfigs);
 }
+
+TConstArrayView<FObjectReplicationBridgeTypeStatsConfig> UObjectReplicationBridgeConfig::GetTypeStatsConfigs() const
+{
+	return MakeArrayView(TypeStatsConfigs);
+}
diff --git a/Engine/Source/Runtime/Experimental/Iris/Core/Private/Iris/ReplicationSystem/Polling/ObjectPoller.cpp b/Engine/Source/Runtime/Experimental/Iris/Core/Private/Iris/ReplicationSystem/Polling/ObjectPoller.cpp
--- a/Engine/Source/Runtime/Experimental/Iris/Core/Private/Iris/ReplicationSystem/Polling/ObjectPoller.cpp
+++ b/Engine/Source/Runtime/Experimental/Iris/Core/Private/Iris/ReplicationSystem/Polling/ObjectPoller.cpp
@@ -11,6 +11,7 @@
 #include "Iris/ReplicationSystem/NetRefHandle.h"
 
 #include "Iris/Core/IrisProfiler.h"
+#include "Iris/Stats/NetStatsContext.h"
 
 #include "Net/Core/Trace/NetDebugName.h"
 
@@ -21,6 +22,7 @@ FObjectPoller::FObjectPoller(const FInitParams& InitParams)
 	: ObjectReplicationBridge(InitParams.ObjectReplicationBridge)
 	, ReplicationSystemInternal(InitParams.ReplicationSystemInternal)
 	, LocalNetRefHandleManager(ReplicationSystemInternal->GetNetRefHandleManager())
+	, NetStatsContext(nullptr)
 	, ReplicatedInstances(LocalNetRefHandleManager.GetReplicatedInstances())
 	, AccumulatedDirtyObjects(ReplicationSystemInternal->GetDirtyNetObjectTracker().GetAccumulatedDirtyNetObjects())
 	, DirtyObjectsToCopy(LocalNetRefHandleManager.GetDirtyObjectsToCopy())
@@ -33,11 +35,14 @@ FObjectPoller::FObjectPoller(const FInitParams& InitParams)
 void FObjectPoller::PreUpdatePass(const FNetBitArrayView& ObjectsConsideredForPolling)
 {
 	IRIS_PROFILER_SCOPE_VERBOSE(PreUpdatePass);
+	NetStatsContext = ReplicationSystemInternal->GetNetTypeStats().GetNetStatsContext();
 
 	ObjectsConsideredForPolling.ForAllSetBits([this](FInternalNetRefIndex Objectindex)
 	{
 		CallPreUpdate(Objectindex);
 	});
+
+	NetStatsContext = nullptr;
 }
 
 void FObjectPoller::CallPreUpdate(FInternalNetRefIndex ObjectIndex)
@@ -53,8 +58,12 @@ void FObjectPoller::CallPreUpdate(FInternalNetRefIndex ObjectIndex)
 	// Call per-instance PreUpdate function
 	if (ObjectReplicationBridge->PreUpdateInstanceFunction && EnumHasAnyFlags(ObjectData.InstanceProtocol->InstanceTraits, EReplicationInstanceProtocolTraits::NeedsPreSendUpdate))
 	{
+		UE_NET_IRIS_STATS_TIMER(Timer, NetStatsContext);
+	
 		ObjectReplicationBridge->PreUpdateInstanceFunction(ObjectData.RefHandle, ReplicatedInstances[ObjectIndex], ObjectReplicationBridge);
 		++PollStats.PreUpdatedObjectCount;
+		
+		UE_NET_IRIS_STATS_ADD_TIME_AND_COUNT_FOR_OBJECT(Timer, PreUpdate, ObjectIndex);
 	}
 }
 
@@ -63,6 +72,8 @@ void FObjectPoller::PollObjects(const FNetBitArrayView& ObjectsConsideredForPoll
 	FDirtyObjectsAccessor DirtyObjectsAccessor(ReplicationSystemInternal->GetDirtyNetObjectTracker());
 	DirtyObjectsThisFrame = DirtyObjectsAccessor.GetDirtyNetObjects();
 
+	NetStatsContext = ReplicationSystemInternal->GetNetTypeStats().GetNetStatsContext();
+
 	if (IsIrisPushModelEnabled())
 	{
 		IRIS_PROFILER_SCOPE_VERBOSE(PollPushBased);
@@ -79,12 +90,16 @@ void FObjectPoller::PollObjects(const FNetBitArrayView& ObjectsConsideredForPoll
 			ForcePollObject(Objectindex);
 		});
 	}
+
+	NetStatsContext = nullptr;
 }
 
 void FObjectPoller::PollSingleObject(FNetRefHandle Handle)
 {
 	if (uint32 InternalObjectIndex = LocalNetRefHandleManager.GetInternalIndex(Handle))
 	{
+		NetStatsContext = ReplicationSystemInternal->GetNetTypeStats().GetNetStatsContext();
+
 		CallPreUpdate(InternalObjectIndex);
 
 		FDirtyObjectsAccessor DirtyObjectsAccessor(ReplicationSystemInternal->GetDirtyNetObjectTracker());
@@ -98,6 +113,7 @@ void FObjectPoller::PollSingleObject(FNetRefHandle Handle)
 
 		// Clear ref to locked dirty bit array
 		DirtyObjectsThisFrame = FNetBitArrayView();
+		NetStatsContext = nullptr;
 	}
 }
 
@@ -118,6 +134,7 @@ void FObjectPoller::ForcePollObject(FInternalNetRefIndex ObjectIndex)
 	if (EnumHasAnyFlags(ObjectData.InstanceProtocol->InstanceTraits, EReplicationInstanceProtocolTraits::NeedsPoll))
 	{
 		IRIS_PROFILER_SCOPE_VERBOSE(Poll);
+		UE_NET_IRIS_STATS_TIMER(Timer, NetStatsContext);
 
 		const bool bIsGCAffectedObject = GarbageCollectionAffectedObjects.GetBit(ObjectIndex);
 		GarbageCollectionAffectedObjects.ClearBit(ObjectIndex);
@@ -130,9 +147,15 @@ void FObjectPoller::ForcePollObject(FInternalNetRefIndex ObjectIndex)
 		const bool bPollFoundDirty = FReplicationInstanceOperations::PollAndRefreshCachedPropertyData(ObjectData.InstanceProtocol, PollOptions);
 		if (bWasAlreadyDirty || bPollFoundDirty)
 		{
+			UE_NET_IRIS_STATS_ADD_TIME_AND_COUNT_FOR_OBJECT(Timer, Poll, ObjectIndex);
+
 			DirtyObjectsToCopy.SetBit(ObjectIndex);
 			DirtyObjectsThisFrame.SetBit(ObjectIndex);
 		}
+		else
+		{
+			UE_NET_IRIS_STATS_ADD_TIME_AND_COUNT_FOR_OBJECT_AS_WASTE(Timer, Poll, ObjectIndex);
+		}
 		++PollStats.PolledObjectCount;
 	}
 	else
@@ -175,6 +198,7 @@ void FObjectPoller::PushModelPollObject(FInternalNetRefIndex ObjectIndex)
 	}
 
 	IRIS_PROFILER_SCOPE_VERBOSE(PollPushBased);
+	UE_NET_IRIS_STATS_TIMER(Timer, NetStatsContext);
 
 	// Does the object need to poll all states once.
 	const bool bWantsFullPoll = ObjectData.bWantsFullPoll;
@@ -226,9 +250,15 @@ void FObjectPoller::PushModelPollObject(FInternalNetRefIndex ObjectIndex)
 
 	if (bIsDirtyObject || bPollFoundDirty)
 	{
+		UE_NET_IRIS_STATS_ADD_TIME_AND_COUNT_FOR_OBJECT(Timer, Poll, ObjectIndex);
+
 		DirtyObjectsToCopy.SetBit(ObjectIndex);
 		DirtyObjectsThisFrame.SetBit(ObjectIndex);
 	}
+	else
+	{
+		UE_NET_IRIS_STATS_ADD_TIME_AND_COUNT_FOR_OBJECT_AS_WASTE(Timer, Poll, ObjectIndex);
+	}
 }
 
 } // end namespace UE::Net::Private
\ No newline at end of file
diff --git a/Engine/Source/Runtime/Experimental/Iris/Core/Private/Iris/ReplicationSystem/Polling/ObjectPoller.h b/Engine/Source/Runtime/Experimental/Iris/Core/Private/Iris/ReplicationSystem/Polling/ObjectPoller.h
--- a/Engine/Source/Runtime/Experimental/Iris/Core/Private/Iris/ReplicationSystem/Polling/ObjectPoller.h
+++ b/Engine/Source/Runtime/Experimental/Iris/Core/Private/Iris/ReplicationSystem/Polling/ObjectPoller.h
@@ -17,6 +17,7 @@ namespace UE::Net
 
 		class FReplicationSystemInternal;
 		class FNetRefHandleManager;
+		class FNetStatsContext;
 	}
 }
 
@@ -77,6 +78,7 @@ private:
 	FReplicationSystemInternal* ReplicationSystemInternal;
 
 	FNetRefHandleManager& LocalNetRefHandleManager;
+	FNetStatsContext* NetStatsContext = nullptr;
 	const TArray<UObject*>& ReplicatedInstances;
 
 	const FNetBitArrayView AccumulatedDirtyObjects;
diff --git a/Engine/Source/Runtime/Experimental/Iris/Core/Private/Iris/ReplicationSystem/ReplicationBridge.cpp b/Engine/Source/Runtime/Experimental/Iris/Core/Private/Iris/ReplicationSystem/ReplicationBridge.cpp
--- a/Engine/Source/Runtime/Experimental/Iris/Core/Private/Iris/ReplicationSystem/ReplicationBridge.cpp
+++ b/Engine/Source/Runtime/Experimental/Iris/Core/Private/Iris/ReplicationSystem/ReplicationBridge.cpp
@@ -305,9 +305,9 @@ void UReplicationBridge::Initialize(UReplicationSystem* InReplicationSystem)
 
 	// Create destruction info protocol
 	{
-		constexpr bool bIgnoreProtocolValidation = false;
 		const FReplicationFragments RegisteredFragments;
-		DestructionInfoProtocol = ReplicationProtocolManager->CreateReplicationProtocol(FReplicationProtocolManager::CalculateProtocolIdentifier(RegisteredFragments), RegisteredFragments, TEXT("InternalDestructionInfo"), bIgnoreProtocolValidation);
+		FCreateReplicationProtocolParameters CreateProtocolParams {.bValidateProtocolId = false, .TypeStatsIndex =  GetReplicationSystem()->GetReplicationSystemInternal()->GetNetTypeStats().GetOrCreateTypeStats(FName("DestructionInfo"))};
+		DestructionInfoProtocol = ReplicationProtocolManager->CreateReplicationProtocol(FReplicationProtocolManager::CalculateProtocolIdentifier(RegisteredFragments), RegisteredFragments, TEXT("InternalDestructionInfo"), CreateProtocolParams);
 	}
 }
 
@@ -718,6 +718,7 @@ void UReplicationBridge::InternalFlushStateData(FNetRefHandle Handle)
 	FNetSerializationContext SerializationContext;
 	FInternalNetSerializationContext InternalContext(ReplicationSystem);
 	SerializationContext.SetInternalContext(&InternalContext);
+	SerializationContext.SetNetStatsContext(ReplicationSystem->GetReplicationSystemInternal()->GetNetTypeStats().GetNetStatsContext());
 
 	InternalFlushStateData(SerializationContext, ChangeMaskCache, ChangeMaskWriter, InternalObjectIndex);
 
@@ -773,6 +774,7 @@ void UReplicationBridge::InternalTearOff(FNetRefHandle Handle)
 		FNetSerializationContext SerializationContext;
 		FInternalNetSerializationContext InternalContext(ReplicationSystem);
 		SerializationContext.SetInternalContext(&InternalContext);
+		SerializationContext.SetNetStatsContext(ReplicationSystem->GetReplicationSystemInternal()->GetNetTypeStats().GetNetStatsContext());
 
 		if (ObjectData.InstanceProtocol && EnumHasAnyFlags(ObjectData.InstanceProtocol->InstanceTraits, EReplicationInstanceProtocolTraits::NeedsPoll | EReplicationInstanceProtocolTraits::NeedsPreSendUpdate))
 		{
diff --git a/Engine/Source/Runtime/Experimental/Iris/Core/Private/Iris/ReplicationSystem/ReplicationOperationsInternal.cpp b/Engine/Source/Runtime/Experimental/Iris/Core/Private/Iris/ReplicationSystem/ReplicationOperationsInternal.cpp
--- a/Engine/Source/Runtime/Experimental/Iris/Core/Private/Iris/ReplicationSystem/ReplicationOperationsInternal.cpp
+++ b/Engine/Source/Runtime/Experimental/Iris/Core/Private/Iris/ReplicationSystem/ReplicationOperationsInternal.cpp
@@ -28,6 +28,7 @@
 #include "HAL/IConsoleManager.h"
 #include "Iris/Core/IrisProfiler.h"
 #include "Net/Core/Trace/NetDebugName.h"
+#include "Iris/Stats/NetStatsContext.h"
 
 namespace UE::Net::Private
 {
@@ -78,6 +79,7 @@ uint32 FReplicationInstanceOperationsInternal::CopyObjectStateData(FNetBitStream
 		if (Object.InstanceProtocol && Object.Protocol->InternalTotalSize > 0U)
 		{
 			IRIS_PROFILER_PROTOCOL_NAME(Object.Protocol->DebugName->Name);
+			UE_NET_IRIS_STATS_TIMER(Timer, SerializationContext.GetNetStatsContext());
 
 			// if the object was scopable prev frame we can do partial copy
 			bool bShouldPropagateChangedStates = Object.bShouldPropagateChangedStates;
@@ -132,6 +134,8 @@ uint32 FReplicationInstanceOperationsInternal::CopyObjectStateData(FNetBitStream
 				Cache.PopLastEntry();
 			}
 
+			UE_NET_IRIS_STATS_ADD_TIME_AND_COUNT_FOR_OBJECT(Timer, Copy, InternalIndex);
+
 			return 1U;
 		}
 		else if (const uint32 SubObjectOwnerIndex = Object.SubObjectRootIndex)
diff --git a/Engine/Source/Runtime/Experimental/Iris/Core/Private/Iris/ReplicationSystem/ReplicationProtocol.h b/Engine/Source/Runtime/Experimental/Iris/Core/Private/Iris/ReplicationSystem/ReplicationProtocol.h
--- a/Engine/Source/Runtime/Experimental/Iris/Core/Private/Iris/ReplicationSystem/ReplicationProtocol.h
+++ b/Engine/Source/Runtime/Experimental/Iris/Core/Private/Iris/ReplicationSystem/ReplicationProtocol.h
@@ -99,6 +99,9 @@ struct FReplicationProtocol
 	FReplicationProtocolIdentifier ProtocolIdentifier;
 	const FNetDebugName* DebugName;
 
+	// TypeStatsIndex assigned to this protocol
+	int32 TypeStatsIndex;
+
 	// TODO: Cache parts of the descriptors in the protocol to avoid having to pull in the descriptor itself for operations iterating over the protocol
 
 	// Selected traits from the ReplicationStateDescriptors that might be handy to cache in the protocol.
diff --git a/Engine/Source/Runtime/Experimental/Iris/Core/Private/Iris/ReplicationSystem/ReplicationProtocolManager.cpp b/Engine/Source/Runtime/Experimental/Iris/Core/Private/Iris/ReplicationSystem/ReplicationProtocolManager.cpp
--- a/Engine/Source/Runtime/Experimental/Iris/Core/Private/Iris/ReplicationSystem/ReplicationProtocolManager.cpp
+++ b/Engine/Source/Runtime/Experimental/Iris/Core/Private/Iris/ReplicationSystem/ReplicationProtocolManager.cpp
@@ -13,6 +13,7 @@
 #include "HAL/PlatformString.h"
 #include "Containers/StringFwd.h"
 #include "Logging/LogScopedVerbosityOverride.h"
+#include "Iris/Stats/NetStats.h"
 
 #ifndef UE_NET_ENABLE_PROTOCOLMANAGER_LOG
 // Default is to enable protocol logs in non-shipping
@@ -217,12 +218,12 @@ void FReplicationProtocolManager::FragmentListToString(FStringBuilderBase& Strin
 	}
 }
 
-const FReplicationProtocol* FReplicationProtocolManager::CreateReplicationProtocol(const UObject* ArchetypeOrCDOUsedAsKey, const FReplicationProtocolIdentifier ProtocolId, const FReplicationFragments& Fragments, const TCHAR* DebugName, bool bVerifyId)
+const FReplicationProtocol* FReplicationProtocolManager::CreateReplicationProtocol(const FReplicationProtocolIdentifier ProtocolId, const FReplicationFragments& Fragments, const TCHAR* DebugName, const FCreateReplicationProtocolParameters& Params)
 {
-	if (bVerifyId)
+	if (Params.bValidateProtocolId)
 	{
 		const FReplicationProtocolIdentifier NewProtocolId = CalculateProtocolIdentifier(Fragments);
-		if (NewProtocolId != ProtocolId)
+		if (!ensureMsgf(NewProtocolId == ProtocolId, TEXT("FReplicationProtocolManager::CreateReplicationProtocol Id mismatch when creating protocol named %s with in ProtocolId:0x%" UINT64_x_FMT "Calculated ProtocolId:0x$%" UINT64_x_FMT), DebugName, ProtocolId, NewProtocolId))
 		{
 			UE_LOG(LogIris, Warning, TEXT("FReplicationProtocolManager::CreateReplicationProtocol Id mismatch when creating protocol named %s with ProtocolId:0x%" UINT64_x_FMT), DebugName, ProtocolId);
  #if UE_NET_ENABLE_PROTOCOLMANAGER_LOG
@@ -238,7 +239,7 @@ const FReplicationProtocol* FReplicationProtocolManager::CreateReplicationProtoc
 		}
 	}
 
-	check(GetReplicationProtocol(ProtocolId, ArchetypeOrCDOUsedAsKey) == nullptr);
+	check(GetReplicationProtocol(ProtocolId, Params.ArchetypeOrCDOUsedAsKey) == nullptr);
 
 	// Create the protocol
 	const uint32 FragmentCount = Fragments.Num();
@@ -366,25 +367,30 @@ const FReplicationProtocol* FReplicationProtocolManager::CreateReplicationProtoc
 	Protocol->ProtocolIdentifier = ProtocolId;
 	Protocol->ProtocolTraits = ProtocolTraits;
 	Protocol->DebugName = CreatePersistentNetDebugName(DebugName);
+	Protocol->TypeStatsIndex = Params.TypeStatsIndex >= 0 ? Params.TypeStatsIndex : FNetTypeStats::DefaultTypeStatsIndex;
 	Protocol->RefCount = 0;
 
 	// Register protocol
 	FRegisteredProtocolInfo Info;
-	Info.ArchetypeOrCDOUsedAsKey = ArchetypeOrCDOUsedAsKey;
+	Info.ArchetypeOrCDOUsedAsKey = Params.ArchetypeOrCDOUsedAsKey;
 	Info.Protocol = Protocol;
 
 	RegisteredProtocols.AddUnique(ProtocolId, Info);
 	ProtocolToInfoMap.Add(Protocol, Info);
 
-	UE_LOG_PROTOCOLMANAGER(Verbose, TEXT("FReplicationProtocolManager::CreateReplicationProtocol Created new protocol %s with ProtocolId:0x%" UINT64_x_FMT), ToCStr(Protocol->DebugName), ProtocolId);
 #if UE_NET_ENABLE_PROTOCOLMANAGER_LOG
 	if (bIrisLogReplicationProtocols)
 	{
 		TStringBuilder<4096> StringBuilder;
 		FragmentListToString(StringBuilder, Fragments);
 		LOG_SCOPE_VERBOSITY_OVERRIDE(LogIris, ELogVerbosity::Log);
+		UE_LOG_PROTOCOLMANAGER(Log, TEXT("FReplicationProtocolManager::CreateReplicationProtocol Created new protocol %s with ProtocolId:0x%" UINT64_x_FMT), ToCStr(Protocol->DebugName), ProtocolId);
 		UE_LOG_PROTOCOLMANAGER(Log, TEXT("%s"), StringBuilder.ToString());
 	}
+	else
+	{
+		UE_LOG_PROTOCOLMANAGER(Verbose, TEXT("FReplicationProtocolManager::CreateReplicationProtocol Created new protocol %s with ProtocolId:0x%" UINT64_x_FMT), ToCStr(Protocol->DebugName), ProtocolId);	
+	}
 #endif
 
 	return Protocol;
diff --git a/Engine/Source/Runtime/Experimental/Iris/Core/Private/Iris/ReplicationSystem/ReplicationProtocolManager.h b/Engine/Source/Runtime/Experimental/Iris/Core/Private/Iris/ReplicationSystem/ReplicationProtocolManager.h
--- a/Engine/Source/Runtime/Experimental/Iris/Core/Private/Iris/ReplicationSystem/ReplicationProtocolManager.h
+++ b/Engine/Source/Runtime/Experimental/Iris/Core/Private/Iris/ReplicationSystem/ReplicationProtocolManager.h
@@ -11,19 +11,20 @@
 namespace UE::Net::Private
 {
 
+struct FCreateReplicationProtocolParameters
+{
+	const UObject* ArchetypeOrCDOUsedAsKey = nullptr;
+	bool bValidateProtocolId = false;
+	int32 TypeStatsIndex = INDEX_NONE;
+};
+
 class FReplicationProtocolManager
 {
 public:
 	~FReplicationProtocolManager();
 
 	/* Create protocol from registered fragment data with provided Id, verification is optional */
-	IRISCORE_API const FReplicationProtocol* CreateReplicationProtocol(const UObject* ArchetypeOrCDOUsedAsKey, const FReplicationProtocolIdentifier ProtocolId, const FReplicationFragments& Fragments, const TCHAR* DebugName, bool bVerifyId = false);
-
-	/* Create protocol from registered fragment data with provided Id, verification is optional */
-	const FReplicationProtocol* CreateReplicationProtocol(const FReplicationProtocolIdentifier ProtocolId, const FReplicationFragments& Fragments, const TCHAR* DebugName, bool bVerifyId = false)
-	{
-		return CreateReplicationProtocol(nullptr, ProtocolId, Fragments, DebugName, bVerifyId);
-	}
+	IRISCORE_API const FReplicationProtocol* CreateReplicationProtocol(const FReplicationProtocolIdentifier ProtocolId, const FReplicationFragments& Fragments, const TCHAR* DebugName, const FCreateReplicationProtocolParameters& Params = FCreateReplicationProtocolParameters());
 
 	/* Get an existing replication protocol */
 	IRISCORE_API const FReplicationProtocol* GetReplicationProtocol(FReplicationProtocolIdentifier ProtocolId, const UObject* ArchetypeOrCDOUsedAsKey = nullptr);
@@ -32,7 +33,6 @@ public:
 	template<typename T>
 	void ForEachProtocol(FReplicationProtocolIdentifier ProtocolId, T&& Functor) const;
 	
-
 	/* Destroy existing replication protocol */
 	IRISCORE_API void DestroyReplicationProtocol(const FReplicationProtocol* ReplicationProtocol);
 
diff --git a/Engine/Source/Runtime/Experimental/Iris/Core/Private/Iris/ReplicationSystem/ReplicationSystem.cpp b/Engine/Source/Runtime/Experimental/Iris/Core/Private/Iris/ReplicationSystem/ReplicationSystem.cpp
--- a/Engine/Source/Runtime/Experimental/Iris/Core/Private/Iris/ReplicationSystem/ReplicationSystem.cpp
+++ b/Engine/Source/Runtime/Experimental/Iris/Core/Private/Iris/ReplicationSystem/ReplicationSystem.cpp
@@ -234,6 +234,13 @@ public:
 		}
 
 		ConnectionsPendingPostTickDispatchSend.Init(ReplicationSystemInternal.GetConnections().GetMaxConnectionCount());
+
+		FNetTypeStats& NetStats = ReplicationSystemInternal.GetNetTypeStats();
+		{
+			FNetTypeStats::FInitParams InitParams;
+			InitParams.NetRefHandleManager = &NetRefHandleManager;
+			NetStats.Init(InitParams);
+		}
 	}
 
 	void Deinit()
@@ -404,6 +411,7 @@ public:
 		FInternalNetSerializationContext InternalContext(ReplicationSystem);
 
 		SerializationContext.SetInternalContext(&InternalContext);
+		SerializationContext.SetNetStatsContext(ReplicationSystemInternal.GetNetTypeStats().GetNetStatsContext());
 
 		// Copy the state data of objects that were dirty this frame.
 		FNetBitArrayView DirtyObjectsToCopy = NetRefHandleManager.GetDirtyObjectsToCopy();
@@ -779,8 +787,11 @@ void UReplicationSystem::PostSendUpdate()
 
 #if UE_NET_IRIS_CSV_STATS && CSV_PROFILER
 		{
-			UE::Net::FNetSendStats& SendStats = Impl->ReplicationSystemInternal.GetSendStats();
+			FNetSendStats& SendStats = Impl->ReplicationSystemInternal.GetSendStats();
 			SendStats.ReportCsvStats();
+
+			FNetTypeStats& TypeStats = Impl->ReplicationSystemInternal.GetNetTypeStats();
+			TypeStats.ReportCSVStats();
 		}
 #endif
 
diff --git a/Engine/Source/Runtime/Experimental/Iris/Core/Private/Iris/ReplicationSystem/ReplicationSystemInternal.h b/Engine/Source/Runtime/Experimental/Iris/Core/Private/Iris/ReplicationSystem/ReplicationSystemInternal.h
--- a/Engine/Source/Runtime/Experimental/Iris/Core/Private/Iris/ReplicationSystem/ReplicationSystemInternal.h
+++ b/Engine/Source/Runtime/Experimental/Iris/Core/Private/Iris/ReplicationSystem/ReplicationSystemInternal.h
@@ -98,6 +98,8 @@ public:
 	FDeltaCompressionBaselineManager& GetDeltaCompressionBaselineManager() { return DeltaCompressionBaselineManager; }
 	FDeltaCompressionBaselineInvalidationTracker& GetDeltaCompressionBaselineInvalidationTracker() { return DeltaCompressionBaselineInvalidationTracker; }
 
+	FNetTypeStats& GetNetTypeStats() { return TypeStats; }
+
 	FNetSendStats& GetSendStats()
 	{ 
 		return SendStats;
@@ -135,6 +137,7 @@ private:
 	FDeltaCompressionBaselineManager DeltaCompressionBaselineManager;
 	FDeltaCompressionBaselineInvalidationTracker DeltaCompressionBaselineInvalidationTracker;
 	FNetSendStats SendStats;
+	FNetTypeStats TypeStats;
 	FForwardNetRPCCallMulticastDelegate ForwardNetRPCCallMulticastDelegate;
 	uint32 Id;
 
diff --git a/Engine/Source/Runtime/Experimental/Iris/Core/Private/Iris/ReplicationSystem/ReplicationWriter.cpp b/Engine/Source/Runtime/Experimental/Iris/Core/Private/Iris/ReplicationSystem/ReplicationWriter.cpp
--- a/Engine/Source/Runtime/Experimental/Iris/Core/Private/Iris/ReplicationSystem/ReplicationWriter.cpp
+++ b/Engine/Source/Runtime/Experimental/Iris/Core/Private/Iris/ReplicationSystem/ReplicationWriter.cpp
@@ -24,7 +24,7 @@
 #include "Iris/Serialization/NetBitStreamUtil.h"
 #include "Iris/Serialization/NetSerializer.h"
 #include "Iris/Serialization/NetExportContext.h"
-#include "Iris/Stats/NetStats.h"
+#include "Iris/Stats/NetStatsContext.h"
 #include "HAL/IConsoleManager.h"
 #include "ProfilingDebugging/CsvProfiler.h"
 #include <algorithm>
@@ -69,24 +69,6 @@ static FAutoConsoleVariableRef CvarValidateObjectsWithDirtyChanges(TEXT("net.Iri
 
 static const FName NetError_ObjectStateTooLarge("Object state is too large to be split.");
 
-/** Helper class for timing various operations. */
-class FIrisStatsTimer
-{
-public:
-	FIrisStatsTimer()
-	: StartCycle(FPlatformTime::Cycles64()) 
-	{
-	}
-
-	double GetSeconds() const
-	{
-		return FPlatformTime::ToSeconds64(FPlatformTime::Cycles64() - StartCycle);
-	}
-
-private:
-	uint64 StartCycle;
-};
-
 const TCHAR* FReplicationWriter::LexToString(const EReplicatedObjectState State)
 {
 	static const TCHAR* Names[] = {
@@ -383,6 +365,7 @@ void FReplicationWriter::Init(const FReplicationParameters& InParameters)
 	const FNetBlobManager* NetBlobManager = &ReplicationSystemInternal->GetNetBlobManager();
 	PartialNetObjectAttachmentHandler = NetBlobManager->GetPartialNetObjectAttachmentHandler();
 	NetObjectBlobHandler = NetBlobManager->GetNetObjectBlobHandler();
+	NetTypeStats = &ReplicationSystemInternal->GetNetTypeStats();
 
 	// Init book keeping
 	ReplicatedObjects.SetNumZeroed(Parameters.MaxActiveReplicatedObjectCount);
@@ -1857,7 +1840,7 @@ FReplicationWriter::EWriteObjectStatus FReplicationWriter::WriteObjectAndSubObje
 	UE_NET_TRACE_OBJECT_SCOPE(NetRefHandleForTraceScope, Writer, Context.GetTraceCollector(), ENetTraceVerbosity::Trace);
 #endif
 
-	// We only need to write batch 
+	// We only need to write batch info for root objects
 	const bool bWriteBatchInfo = !Info.IsSubObject;
 	uint32 InitialStateHeaderPos = 0U;
 	const uint32 NumBitsUsedForBatchSize = Parameters.NumBitsUsedForBatchSize;
@@ -2060,6 +2043,8 @@ FReplicationWriter::EWriteObjectStatus FReplicationWriter::WriteObjectAndSubObje
 				}
 				// Serialize initial state data for this object using delta compression against default state
 				FReplicationProtocolOperations::SerializeInitialStateWithMask(Context, Info.GetChangeMaskStoragePointer(), ReplicatedObjectStateBuffer, ObjectData.Protocol);
+
+				UE_NET_IRIS_STATS_ADD_BITS_WRITTEN_AND_COUNT_FOR_OBJECT(Context.GetNetStatsContext(), Writer.GetPosBits() - ObjectRollbackScope.GetStartPos(), WriteCreationInfo, InternalIndex);
 			}
 			else
 			{
@@ -2154,6 +2139,7 @@ FReplicationWriter::EWriteObjectStatus FReplicationWriter::WriteObjectAndSubObje
 
 		if (Writer.IsOverflown())
 		{
+			UE_NET_IRIS_STATS_ADD_BITS_WRITTEN_FOR_OBJECT_AS_WASTE(Context.GetNetStatsContext(), Writer.GetPosBits() - ObjectRollbackScope.GetStartPos(), Write, InternalIndex);
 			return EWriteObjectStatus::BitStreamOverflow;
 		}
 
@@ -2209,12 +2195,11 @@ FReplicationWriter::EWriteObjectStatus FReplicationWriter::WriteObjectAndSubObje
 	CreatedBaselineIndex = FDeltaCompressionBaselineManager::InvalidBaselineIndex;
 
 	// Write dirty sub objects
-	bool bWroteSubObjectData = false;
+	const uint32 SubObjectStartPos = Writer.GetPosBits();
+	uint32 SubObjectsWrittenBits = 0U;
 	if (Info.HasDirtySubObjects && !Info.IsSubObject)
 	{
 		bool bHasDirtySubObjects = false;
-
-		const uint32 SubObjectStartPos = Writer.GetPosBits();
 		
 		FReplicationConditionals::FSubObjectsToReplicateArray SubObjectsToReplicate;
 		ReplicationConditionals->GetSubObjectsToReplicate(Parameters.ConnectionId, InternalIndex, SubObjectsToReplicate);		
@@ -2230,6 +2215,8 @@ FReplicationWriter::EWriteObjectStatus FReplicationWriter::WriteObjectAndSubObje
 			EWriteObjectStatus SubObjectWriteStatus = WriteObjectAndSubObjects(Context, SubObjectInternalIndex, WriteObjectFlags, OutBatchInfo);
 			if (!IsWriteObjectSuccess(SubObjectWriteStatus))
 			{
+				// SubObject will rollback on fail (and report its own waste) but we as we will rollback successfully written subobjects it is better to at least report it with the owner.
+				UE_NET_IRIS_STATS_ADD_BITS_WRITTEN_FOR_OBJECT_AS_WASTE(Context.GetNetStatsContext(), Writer.GetPosBits() - ObjectRollbackScope.GetStartPos(), Write, InternalIndex);
 				return SubObjectWriteStatus;
 			}
 
@@ -2241,7 +2228,7 @@ FReplicationWriter::EWriteObjectStatus FReplicationWriter::WriteObjectAndSubObje
 			}
 		}
 
-		bWroteSubObjectData = Writer.GetPosBits() != SubObjectStartPos;
+		SubObjectsWrittenBits = Writer.GetPosBits() - SubObjectStartPos;
 
 		// Update parent batch info
 		{
@@ -2259,12 +2246,14 @@ FReplicationWriter::EWriteObjectStatus FReplicationWriter::WriteObjectAndSubObje
 		const uint32 WrittenBitsInBatch = (Writer.GetPosBits() - InitialStateHeaderPos) - NumBitsUsedForBatchSize;
 		
 		const bool bWroteData = (ParentBatchEntry.bSentState || ParentBatchEntry.bSentAttachments || bSentTearOff || Info.SubObjectPendingDestroy);
-		if (bWroteData || bWroteSubObjectData)
+		if (bWroteData || (SubObjectsWrittenBits != 0U))
 		{
 			const FObjectReferenceCache::EWriteExportsResult WriteExportResult = ObjectReferenceCache->WritePendingExports(Context);
 
 			if (WriteExportResult == FObjectReferenceCache::EWriteExportsResult::BitStreamOverflow)
 			{
+				// If we fail to write exports, we fail the entire object
+				UE_NET_IRIS_STATS_ADD_BITS_WRITTEN_FOR_OBJECT_AS_WASTE(Context.GetNetStatsContext(), Writer.GetPosBits() - ObjectRollbackScope.GetStartPos(), Write, InternalIndex);
 				return EWriteObjectStatus::BitStreamOverflow;	
 			}
 
@@ -2280,6 +2269,8 @@ FReplicationWriter::EWriteObjectStatus FReplicationWriter::WriteObjectAndSubObje
 			}
 
 			ParentBatchEntry.bSentBatchData = 1U;
+
+			UE_NET_IRIS_STATS_ADD_BITS_WRITTEN_FOR_OBJECT(Context.GetNetStatsContext(), (Writer.GetPosBits() - ObjectRollbackScope.GetStartPos()) - SubObjectsWrittenBits, Write, InternalIndex);
 		}
 		// If we did not write any data we rollback any written headers and report a success
 		else
@@ -2296,6 +2287,8 @@ FReplicationWriter::EWriteObjectStatus FReplicationWriter::WriteObjectInBatch(FN
 {
 	IRIS_PROFILER_SCOPE(FReplicationWriter_WriteObjectInBatch);
 
+	UE_NET_IRIS_STATS_TIMER(Timer, Context.GetNetStatsContext());
+
 	// Reset pending exports
 	FNetExportContext* ExportContext = Context.GetExportContext();
 	if (ExportContext)
@@ -2307,9 +2300,12 @@ FReplicationWriter::EWriteObjectStatus FReplicationWriter::WriteObjectInBatch(FN
 	const EWriteObjectStatus WriteObjectStatus = WriteObjectAndSubObjects(Context, InternalIndex, WriteObjectFlags, OutBatchInfo);
 	if (!IsWriteObjectSuccess(WriteObjectStatus))
 	{
+		UE_NET_IRIS_STATS_ADD_TIME_AND_COUNT_FOR_OBJECT_AS_WASTE(Timer, Write, InternalIndex);
 		return WriteObjectStatus;
 	}
 
+	UE_NET_IRIS_STATS_ADD_TIME_AND_COUNT_FOR_OBJECT(Timer, Write, InternalIndex);
+
 	// Include dependent objects as separate batch, (for hugeobjects they will be included as they are written to a separate bitstream)
 	{
 		const uint32 OldBatchInfoParentInternalIndex = OutBatchInfo.ParentInternalIndex;
@@ -2759,16 +2755,7 @@ uint32 FReplicationWriter::WriteObjects(FNetSerializationContext& Context)
 	{
 		if (!this->WriteContext.ObjectsWrittenThisPacket.GetBit(InternalIndex) && this->CanSendObject(InternalIndex))
 		{
-#if UE_NET_IRIS_CSV_STATS && CSV_PROFILER
-			FIrisStatsTimer Timer;
-#endif
 			const int32 Result = this->WriteObjectBatch(Context, InternalIndex, WriteObjectFlag_State | WriteObjectFlag_Attachments);
-#if UE_NET_IRIS_CSV_STATS && CSV_PROFILER
-			if (Result <= 0)
-			{
-				this->WriteContext.Stats.AddReplicationWasteTime(Timer.GetSeconds());
-			}
-#endif
 			if (Result >= 0)
 			{
 				WrittenObjectCount += Result;
@@ -3180,7 +3167,8 @@ UDataStream::EWriteResult FReplicationWriter::Write(FNetSerializationContext& Co
 	// Setup internal context
 	FInternalNetSerializationContext InternalContext(Parameters.ReplicationSystem);
 	Context.SetLocalConnectionId(Parameters.ConnectionId);
-	Context.SetInternalContext(&InternalContext);	
+	Context.SetInternalContext(&InternalContext);
+	Context.SetNetStatsContext(NetTypeStats->GetNetStatsContext());
 
 	// Give some info for the case when we consider splitting a huge object.
 	WriteBitStreamInfo.ReplicationStartPos = Writer.GetPosBits();
diff --git a/Engine/Source/Runtime/Experimental/Iris/Core/Private/Iris/ReplicationSystem/ReplicationWriter.h b/Engine/Source/Runtime/Experimental/Iris/Core/Private/Iris/ReplicationSystem/ReplicationWriter.h
--- a/Engine/Source/Runtime/Experimental/Iris/Core/Private/Iris/ReplicationSystem/ReplicationWriter.h
+++ b/Engine/Source/Runtime/Experimental/Iris/Core/Private/Iris/ReplicationSystem/ReplicationWriter.h
@@ -566,6 +566,7 @@ private:
 	const UPartialNetObjectAttachmentHandler* PartialNetObjectAttachmentHandler = nullptr;
 	const UNetObjectBlobHandler* NetObjectBlobHandler = nullptr;
 	FNetExports* NetExports = nullptr;
+	FNetTypeStats* NetTypeStats = nullptr;
 
 	FWriteContext WriteContext;
 	FBitStreamInfo WriteBitStreamInfo;
diff --git a/Engine/Source/Runtime/Experimental/Iris/Core/Private/Iris/Stats/NetStats.cpp b/Engine/Source/Runtime/Experimental/Iris/Core/Private/Iris/Stats/NetStats.cpp
--- a/Engine/Source/Runtime/Experimental/Iris/Core/Private/Iris/Stats/NetStats.cpp
+++ b/Engine/Source/Runtime/Experimental/Iris/Core/Private/Iris/Stats/NetStats.cpp
@@ -1,15 +1,43 @@
 // Copyright Epic Games, Inc. All Rights Reserved.
 
 #include "Iris/Stats/NetStats.h"
-#include "ProfilingDebugging/CsvProfiler.h"
+#include "Iris/Stats/NetStatsContext.h"
+#include "Iris/Core/IrisProfiler.h"
+#include "HAL/IConsoleManager.h"
 #include "Misc/ScopeLock.h"
+#include "ProfilingDebugging/CsvProfiler.h"
 
-namespace UE::Net
+namespace UE::Net::Private
 {
 
+static bool bCVARShouldIncludeSubObjectWithRoot = true;
+static FAutoConsoleVariableRef CShouldIncludeSubObjectWithRoot(
+	TEXT("net.Iris.Stats.ShouldIncludeSubObjectWithRoot"),
+	bCVARShouldIncludeSubObjectWithRoot,
+	TEXT("If enabled SubObjects will reports stats with RootObject, if set to false SubObjects will be treated as separate objects."
+	));
+
 // Enable Iris category by default on servers
 CSV_DEFINE_CATEGORY(Iris, WITH_SERVER_CODE);
 
+// Per type stats
+CSV_DEFINE_CATEGORY(IrisPreUpdateMS, WITH_SERVER_CODE);
+CSV_DEFINE_CATEGORY(IrisPreUpdateCount, WITH_SERVER_CODE);
+CSV_DEFINE_CATEGORY(IrisPollMS, WITH_SERVER_CODE);
+CSV_DEFINE_CATEGORY(IrisPollCount, WITH_SERVER_CODE);
+CSV_DEFINE_CATEGORY(IrisPollWasteMS, WITH_SERVER_CODE);
+CSV_DEFINE_CATEGORY(IrisPollWasteCount, WITH_SERVER_CODE);
+CSV_DEFINE_CATEGORY(IrisCopyMS, WITH_SERVER_CODE);
+CSV_DEFINE_CATEGORY(IrisCopyCount, WITH_SERVER_CODE);
+CSV_DEFINE_CATEGORY(IrisWriteMS, WITH_SERVER_CODE);
+CSV_DEFINE_CATEGORY(IrisWriteCount, WITH_SERVER_CODE);
+CSV_DEFINE_CATEGORY(IrisWriteKBytes, WITH_SERVER_CODE);
+CSV_DEFINE_CATEGORY(IrisWriteWasteMS, WITH_SERVER_CODE);
+CSV_DEFINE_CATEGORY(IrisWriteWasteCount, WITH_SERVER_CODE);
+CSV_DEFINE_CATEGORY(IrisWriteWasteKBytes, WITH_SERVER_CODE);
+CSV_DEFINE_CATEGORY(IrisWriteCreationInfoCount, WITH_SERVER_CODE);
+CSV_DEFINE_CATEGORY(IrisWriteCreationInfoKBytes, WITH_SERVER_CODE);
+
 void FNetSendStats::Accumulate(const FNetSendStats& Other)
 {
 	FScopeLock Lock(&CS);
@@ -19,13 +47,11 @@ void FNetSendStats::Accumulate(const FNetSendStats& Other)
 	Stats.ReplicatedObjectCount += Other.Stats.ReplicatedObjectCount;
 	Stats.ReplicatedDestructionInfoCount += Other.Stats.ReplicatedDestructionInfoCount;
 	Stats.DeltaCompressedObjectCount += Other.Stats.DeltaCompressedObjectCount;
-	Stats.ReplicationWasteObjectCount += Other.Stats.ReplicationWasteObjectCount;
 	Stats.ReplicatedObjectStatesMaskedOut += Other.Stats.ReplicatedObjectStatesMaskedOut;
 	Stats.ActiveHugeObjectCount += Other.Stats.ActiveHugeObjectCount;
 	Stats.HugeObjectsWaitingForAckCount += Other.Stats.HugeObjectsWaitingForAckCount;
 	Stats.HugeObjectsStallingCount += Other.Stats.HugeObjectsStallingCount;
 
-	Stats.ReplicationWasteTimeInSeconds += Other.Stats.ReplicationWasteTimeInSeconds;
 	Stats.HugeObjectWaitingForAckTimeInSeconds += Other.Stats.HugeObjectWaitingForAckTimeInSeconds;
 	Stats.HugeObjectStallingTimeInSeconds += Other.Stats.HugeObjectStallingTimeInSeconds;
 }
@@ -49,7 +75,6 @@ void FNetSendStats::ReportCsvStats()
 		CSV_CUSTOM_STAT(Iris, AvgReplicatedRootObjectCount, Stats.ReplicatedRootObjectCount/ConnectionCountFloat, ECsvCustomStatOp::Set);
 		CSV_CUSTOM_STAT(Iris, AvgReplicatedObjectCount, Stats.ReplicatedObjectCount/ConnectionCountFloat, ECsvCustomStatOp::Set);
 		CSV_CUSTOM_STAT(Iris, AvgReplicatedDestructionInfoCount, Stats.ReplicatedDestructionInfoCount/ConnectionCountFloat, ECsvCustomStatOp::Set);
-		CSV_CUSTOM_STAT(Iris, AvgReplicationWasteObjectCount, Stats.ReplicationWasteObjectCount/ConnectionCountFloat, ECsvCustomStatOp::Set);
 		CSV_CUSTOM_STAT(Iris, AvgReplicatedObjectStatesMaskedOut, Stats.ReplicatedObjectStatesMaskedOut/ConnectionCountFloat, ECsvCustomStatOp::Set);
 		CSV_CUSTOM_STAT(Iris, AvgDeltaCompressedObjectCount, Stats.DeltaCompressedObjectCount/ConnectionCountFloat, ECsvCustomStatOp::Set);
 	}
@@ -60,7 +85,6 @@ void FNetSendStats::ReportCsvStats()
 		CSV_CUSTOM_STAT(Iris, AvgReplicatedObjectCount, 0, ECsvCustomStatOp::Set);
 		CSV_CUSTOM_STAT(Iris, AvgReplicatedDestructionInfoCount, 0, ECsvCustomStatOp::Set);
 		CSV_CUSTOM_STAT(Iris, AvgReplicatedHugeObjectCount, 0, ECsvCustomStatOp::Set);
-		CSV_CUSTOM_STAT(Iris, AvgReplicationWasteObjectCount, 0, ECsvCustomStatOp::Set);
 		CSV_CUSTOM_STAT(Iris, AvgDeltaCompressedObjectCount, 0, ECsvCustomStatOp::Set);
 	}
 
@@ -71,10 +95,145 @@ void FNetSendStats::ReportCsvStats()
 
 	CSV_CUSTOM_STAT(Iris, ReplicatingConnectionCount, Stats.ReplicatingConnectionCount, ECsvCustomStatOp::Set);
 
-	CSV_CUSTOM_STAT(Iris, ReplicationWasteTimeMilliseconds, Stats.ReplicationWasteTimeInSeconds*1000.0, ECsvCustomStatOp::Set);
 	CSV_CUSTOM_STAT(Iris, HugeObjectWaitingForAckTimeInSeconds, Stats.HugeObjectWaitingForAckTimeInSeconds, ECsvCustomStatOp::Set);
 	CSV_CUSTOM_STAT(Iris, HugeObjectStallingTimeInSeconds, Stats.HugeObjectStallingTimeInSeconds, ECsvCustomStatOp::Set);
 #endif
 }
 
+FNetTypeStats::FNetTypeStats()
+{
+	StatsContext = TUniquePtr<FNetStatsContext>(CreateNetStatsContext());
+	// Add default TypeStats
+	GetOrCreateTypeStats(TEXT("Undefined"));
+	GetOrCreateTypeStats(TEXT("OOBChannel"));
+}
+
+FNetTypeStats::~FNetTypeStats()
+{
+}
+
+void FNetTypeStats::Init(FInitParams& InitParams)
+{
+	NetRefHandleManager = InitParams.NetRefHandleManager;
+	ResetStats();
+}
+
+void FNetTypeStats::ResetStats()
+{
+	UpdateContext(*StatsContext);
+}
+
+int32 FNetTypeStats::GetOrCreateTypeStats(FName Name)
+{
+	const int32 ExistingIndex = TypeStatsNames.Find(Name);
+	if (ExistingIndex != INDEX_NONE)
+	{
+		return ExistingIndex;
+	}
+
+	const int32 NewTypeStatsIndex = int32(TypeStatsNames.Num());
+	TypeStatsNames.Add(Name);
+	StatsContext->TypeStatsData.SetNum(TypeStatsNames.Num());
+
+	return NewTypeStatsIndex;
+}
+
+void FNetTypeStats::UpdateContext(FNetStatsContext& Context)
+{
+	Context.ResetStats(TypeStatsNames.Num());
+	Context.NetRefHandleManager = NetRefHandleManager;
+	Context.bShouldIncludeSubObjectWithRoot = bCVARShouldIncludeSubObjectWithRoot;
+};
+
+FNetStatsContext* FNetTypeStats::CreateNetStatsContext()
+{
+	FNetStatsContext* Context = new FNetStatsContext;
+	UpdateContext(*Context);
+	return Context;
+};
+
+void FNetTypeStats::Accumulate(FNetStatsContext& Context)
+{
+	IRIS_PROFILER_SCOPE(FNetTypeStats_Accumulate);
+
+	// Skip default context as that is our target.
+	if (&Context == StatsContext.Get())
+	{
+		return;
+	}
+
+	if (!ensureMsgf(Context.TypeStatsData.Num() <= TypeStatsNames.Num(), TEXT("Invalid Context")))
+	{
+		return;
+	}
+
+	// Accumulate stats
+	const FNetTypeStatsData* Src = Context.TypeStatsData.GetData();
+	FNetTypeStatsData* Dst = StatsContext->TypeStatsData.GetData();
+	for (int32 StatsIndex = 0, EndIndex = FMath::Min(Context.TypeStatsData.Num(), StatsContext->TypeStatsData.Num()); StatsIndex < EndIndex; ++StatsIndex)
+	{
+		Dst[StatsIndex].Accumulate(Src[StatsIndex]);
+	}
+
+	Context.ResetStats(TypeStatsNames.Num());
+}
+
+#define UE_NET_STATS_RECORD_TYPESTATS_TIME(StatsName, ValueName, StatsData) FCsvProfiler::RecordCustomStat(StatsName, CSV_CATEGORY_INDEX(Iris##ValueName##MS), FGenericPlatformTime::ToMilliseconds64(StatsData.Values[FNetTypeStatsData::EStatsIndex::ValueName].Time), ECsvCustomStatOp::Set)
+#define UE_NET_STATS_RECORD_TYPESTATS_COUNT(StatsName, ValueName, StatsData) FCsvProfiler::RecordCustomStat(StatsName, CSV_CATEGORY_INDEX(Iris##ValueName##Count), static_cast<int32>(StatsData.Values[FNetTypeStatsData::EStatsIndex::ValueName].Count) , ECsvCustomStatOp::Set)
+#define UE_NET_STATS_RECORD_TYPESTATS_BITS(StatsName, ValueName, StatsData) FCsvProfiler::RecordCustomStat(StatsName, CSV_CATEGORY_INDEX(Iris##ValueName##KBytes), ((StatsData.Values[FNetTypeStatsData::EStatsIndex::ValueName].Bits + 7U) / 8) / 1000.f , ECsvCustomStatOp::Set)
+
+void FNetTypeStats::ReportCSVStats()
+{
+#if UE_NET_IRIS_CSV_STATS && CSV_PROFILER
+
+	FCsvProfiler* Profiler = FCsvProfiler::Get();
+	bIsEnabled = Profiler->IsCapturing();
+
+	if (bIsEnabled)
+	{
+		IRIS_PROFILER_SCOPE(FNetTypeStats_ReportCSVStats);
+
+		// Report stats for this frame
+		FNetTypeStatsData* TypeStatsDatas = StatsContext->TypeStatsData.GetData();
+		for (int32 StatsIndex = 0; StatsIndex < TypeStatsNames.Num(); ++StatsIndex)
+		{
+			const FName StatsName = TypeStatsNames[StatsIndex];
+			FNetTypeStatsData& TypeStatsData = TypeStatsDatas[StatsIndex];
+
+			// Report, we could do a loop here but we might end up not wanting to report all collected stats.
+			UE_NET_STATS_RECORD_TYPESTATS_TIME(StatsName, PreUpdate, TypeStatsData);
+			UE_NET_STATS_RECORD_TYPESTATS_COUNT(StatsName, PreUpdate, TypeStatsData);
+
+			UE_NET_STATS_RECORD_TYPESTATS_TIME(StatsName, Poll, TypeStatsData);
+			UE_NET_STATS_RECORD_TYPESTATS_COUNT(StatsName, Poll, TypeStatsData);
+
+			UE_NET_STATS_RECORD_TYPESTATS_TIME(StatsName, PollWaste, TypeStatsData);
+			UE_NET_STATS_RECORD_TYPESTATS_COUNT(StatsName, PollWaste, TypeStatsData);
+
+			UE_NET_STATS_RECORD_TYPESTATS_TIME(StatsName, Copy, TypeStatsData);
+			UE_NET_STATS_RECORD_TYPESTATS_COUNT(StatsName, Copy, TypeStatsData);
+
+			UE_NET_STATS_RECORD_TYPESTATS_TIME(StatsName, Write, TypeStatsData);
+			UE_NET_STATS_RECORD_TYPESTATS_COUNT(StatsName, Write, TypeStatsData);
+			UE_NET_STATS_RECORD_TYPESTATS_BITS(StatsName, Write, TypeStatsData);
+
+			UE_NET_STATS_RECORD_TYPESTATS_TIME(StatsName, WriteWaste, TypeStatsData);
+			UE_NET_STATS_RECORD_TYPESTATS_COUNT(StatsName, Write, TypeStatsData);
+			UE_NET_STATS_RECORD_TYPESTATS_BITS(StatsName, WriteWaste, TypeStatsData);
+
+			UE_NET_STATS_RECORD_TYPESTATS_BITS(StatsName, WriteCreationInfo, TypeStatsData);
+			UE_NET_STATS_RECORD_TYPESTATS_COUNT(StatsName, WriteCreationInfo, TypeStatsData);
+
+			// Reset
+			TypeStatsData.Reset();
+		}
+	}
+
+#endif
+}
+
+#undef UE_NET_STATS_RECORD_TYPESTATS_TIME
+#undef UE_NET_STATS_RECORD_TYPESTATS_COUNT
+#undef UE_NET_STATS_RECORD_TYPESTATS_BITS
+
 }
diff --git a/Engine/Source/Runtime/Experimental/Iris/Core/Private/Iris/Stats/NetStats.h b/Engine/Source/Runtime/Experimental/Iris/Core/Private/Iris/Stats/NetStats.h
--- a/Engine/Source/Runtime/Experimental/Iris/Core/Private/Iris/Stats/NetStats.h
+++ b/Engine/Source/Runtime/Experimental/Iris/Core/Private/Iris/Stats/NetStats.h
@@ -4,8 +4,16 @@
 
 #include "HAL/Platform.h"
 #include "HAL/CriticalSection.h"
+#include "Templates/UniquePtr.h"
+#include "Containers/Array.h"
 
-namespace UE::Net
+namespace UE::Net::Private
+{
+	class FNetRefHandleManager;
+	class FNetStatsContext;
+}
+
+namespace UE::Net::Private
 {
 
 /**
@@ -74,14 +82,12 @@ private:
 	{
 		double HugeObjectWaitingForAckTimeInSeconds = 0;
 		double HugeObjectStallingTimeInSeconds = 0;
-		double ReplicationWasteTimeInSeconds = 0;
 
 		int32 ScheduledForReplicationRootObjectCount = 0;
 		int32 ReplicatedRootObjectCount = 0;
 		int32 ReplicatedObjectCount = 0;
 		int32 ReplicatedDestructionInfoCount = 0;
 		int32 DeltaCompressedObjectCount = 0;
-		int32 ReplicationWasteObjectCount = 0;
 		int32 ReplicatedObjectStatesMaskedOut = 0;
 		int32 ActiveHugeObjectCount = 0;
 		int32 HugeObjectsWaitingForAckCount = 0;
@@ -145,15 +151,62 @@ inline void FNetSendStats::AddHugeObjectStallTime(double Seconds)
 	Stats.HugeObjectStallingTimeInSeconds += Seconds;
 }
 
-inline void FNetSendStats::AddReplicationWasteTime(double Seconds)
-{
-	++Stats.ReplicationWasteObjectCount;
-	Stats.ReplicationWasteTimeInSeconds += Seconds;
-}
-
 inline void FNetSendStats::SetNumberOfReplicatingConnections(uint32 Count)
 {
 	Stats.ReplicatingConnectionCount = Count;
 }
 
+/**
+ * Stats defined per object type for Iris replication reported to the CSV profiler. Mostly of interest on the server side due to the server authoritative network model.
+ * Currently we use a single NetStatsContext when collecting the stats, when we go wide we need to extend this to use separate contexts for different threads.
+ */
+class FNetTypeStats
+{
+public:
+
+	struct FInitParams
+	{
+		FNetRefHandleManager* NetRefHandleManager = nullptr;
+	};
+
+	// Preset stats type indices
+	static constexpr int32 DefaultTypeStatsIndex = 0U;
+	static constexpr int32 OOBChannelTypeStatsIndex = 1U;
+
+public:
+	FNetTypeStats();
+	FNetTypeStats(const FNetTypeStats&) = delete;
+	FNetTypeStats& operator=(const FNetTypeStats&) = delete;
+	~FNetTypeStats();
+
+	void Init(FInitParams& InitParams);
+
+	/** Reset stats */
+	void ResetStats();
+		
+	/** Returns the TypeStatIndex associated with the Name or creates a new one if it does not exist */
+	int32 GetOrCreateTypeStats(FName Name);
+
+	/** Get default context if stats is enabled, once we decide to go wide we will need to expose methods to create and refresh StatsContexts per thread */
+	FNetStatsContext* GetNetStatsContext() { return IsEnabled() ?  StatsContext.Get() : nullptr; }
+
+	/** Updated every frame based on the state of the CSVProfiler */
+	bool IsEnabled() const { return bIsEnabled; }
+
+	/** Accumulate stats from context to main context */
+	void Accumulate(FNetStatsContext& Context);
+
+	/** ReportCSVStats and reset context */
+	void ReportCSVStats();
+
+private:
+	FNetStatsContext* CreateNetStatsContext();
+	void UpdateContext(FNetStatsContext& Context);
+
+	TUniquePtr<FNetStatsContext> StatsContext = nullptr;
+	FNetRefHandleManager* NetRefHandleManager = nullptr;
+	TArray<FName> TypeStatsNames;
+	bool bIsEnabled = false;
+};
+
 }
diff --git a/Engine/Source/Runtime/Experimental/Iris/Core/Private/Iris/Stats/NetStatsContext.h b/Engine/Source/Runtime/Experimental/Iris/Core/Private/Iris/Stats/NetStatsContext.h
new file mode 100644
--- /dev/null
+++ b/Engine/Source/Runtime/Experimental/Iris/Core/Private/Iris/Stats/NetStatsContext.h
@@ -0,0 +1,202 @@
+// Copyright Epic Games, Inc. All Rights Reserved.
+
+#pragma once
+
+#include "HAL/Platform.h"
+#include "Iris/Stats/NetStats.h"
+#include "ProfilingDebugging/CsvProfilerConfig.h"
+#include "Iris/ReplicationSystem/NetRefHandleManager.h"
+#include "Iris/ReplicationSystem/ReplicationProtocol.h"
+
+namespace UE::Net::Private
+{
+
+struct FNetTypeStatsData
+{
+	enum EStatsIndex : unsigned
+	{
+		PreUpdate = 0U,
+		Poll,
+		PollWaste,
+		Copy,
+		Write,
+		WriteWaste,
+		WriteCreationInfo,
+		Count
+	};
+
+	struct FStatsValue
+	{
+		void Accumulate(const FStatsValue& Other)
+		{
+			Time += Other.Time;
+			Bits += Other.Bits;
+			Count += Other.Count;
+		}
+		uint64 Time = uint64(0);
+		uint32 Bits = 0U;
+		uint32 Count = 0U;
+	};
+
+	/** Zero out all stats */
+	void Reset()
+	{
+		*this = FNetTypeStatsData();
+	}
+
+	/** Accumulate all stats */
+	void Accumulate(const FNetTypeStatsData& Other)
+	{
+		for (int32 Index = 0; Index < EStatsIndex::Count; ++Index)
+		{
+			Values[Index].Accumulate(Other.Values[Index]);
+		}
+	}
+
+	FStatsValue Values[EStatsIndex::Count];
+};
+
+class FNetStatsContext
+{
+public:
+	FNetStatsContext() = default;
+	FNetStatsContext(const FNetStatsContext&) = delete;
+	FNetStatsContext& operator=(const FNetStatsContext&) = delete;
+
+	void ResetStats(int32 NumTypeStats);
+	
+	// Helper to lookup TypeStatsIndex from object, we currently store the NetTypeStatsIndex in the ReplicationProtocol but nothing prevents us from storing it elsewhere
+	// Depending on config subobjects report their stats with the root, for those cases we do not bump the count of root objects.
+	// returns the FNetTypeStatsData associated with the protocol used by the object references by the InternalIndex
+	static FNetTypeStatsData& GetTypeStatsDataForObject(FNetStatsContext& Context, FInternalNetRefIndex InternalIndex, uint32& OutUpdateCount);
+	
+private:
+	friend class FNetTypeStats;
+
+	const FNetRefHandleManager* NetRefHandleManager = nullptr;
+	TArray<FNetTypeStatsData> TypeStatsData;
+	bool bShouldIncludeSubObjectWithRoot = false;
+};
+
+class FNetStatsTimer
+{
+public:
+	FNetStatsTimer(FNetStatsContext* InStatsContext)
+	: NetStatsContext(InStatsContext)
+	, StartCycle(InStatsContext ? FPlatformTime::Cycles64() : uint64(0)) 
+	{
+	}
+
+	uint64 GetCyclesSinceStart() const { return FPlatformTime::Cycles64() - StartCycle; }
+
+	// If we have a context, this timer is enabled 
+	FNetStatsContext* GetNetStatsContext() const { return NetStatsContext; }
+	
+private:
+	FNetStatsContext* NetStatsContext;
+	uint64 StartCycle;
+};
+
+inline FNetTypeStatsData& FNetStatsContext::GetTypeStatsDataForObject(FNetStatsContext& Context, FInternalNetRefIndex InternalIndex, uint32& OutUpdateCount)
+{
+	const FNetRefHandleManager::FReplicatedObjectData& ObjectData = Context.NetRefHandleManager->GetReplicatedObjectDataNoCheck(InternalIndex);		
+	const bool bTreatAsRoot = !Context.bShouldIncludeSubObjectWithRoot || !ObjectData.IsSubObject();
+	int32 TypeStatsIndex = FNetTypeStats::DefaultTypeStatsIndex;
+	if (bTreatAsRoot)
+	{
+		TypeStatsIndex = ObjectData.Protocol ? ObjectData.Protocol->TypeStatsIndex : FNetTypeStats::OOBChannelTypeStatsIndex;
+	}
+	else
+	{
+		const FNetRefHandleManager::FReplicatedObjectData& RootObjectData = Context.NetRefHandleManager->GetReplicatedObjectDataNoCheck(ObjectData.SubObjectRootIndex);
+		TypeStatsIndex = RootObjectData.Protocol ? RootObjectData.Protocol->TypeStatsIndex : FNetTypeStats::OOBChannelTypeStatsIndex;
+	}
+	
+	OutUpdateCount = bTreatAsRoot ? 1U : 0U;
+
+	return Context.TypeStatsData[TypeStatsIndex];
+}
+
+inline void FNetStatsContext::ResetStats(int32 NumTypeStats)
+{ 
+	TypeStatsData.SetNum(NumTypeStats);
+	for (FNetTypeStatsData& TypeStats : TypeStatsData)
+	{
+		TypeStats.Reset();
+	}
+} 
+
+}
+
+// Wrap usage of stats in macros so we can compile it out
+#if UE_NET_IRIS_CSV_STATS && CSV_PROFILER
+
+#define UE_NET_IRIS_STATS_TIMER(TimerName, NetStatsContext) UE::Net::Private::FNetStatsTimer TimerName(NetStatsContext);
+
+#define UE_NET_IRIS_STATS_ADD_TIME_AND_COUNT_FOR_OBJECT(Timer, StatName, ObjectIndex) \
+	do { \
+		if (UE::Net::Private::FNetStatsContext* LocalNetStatsContext = Timer.GetNetStatsContext()) \
+		{ \
+			uint32 CountIncrement = 0U; \
+			UE::Net::Private::FNetTypeStatsData& StatsData = UE::Net::Private::FNetStatsContext::GetTypeStatsDataForObject(*LocalNetStatsContext, ObjectIndex, CountIncrement); \
+			StatsData.Values[UE::Net::Private::FNetTypeStatsData::EStatsIndex::StatName].Time += Timer.GetCyclesSinceStart(); \
+			StatsData.Values[UE::Net::Private::FNetTypeStatsData::EStatsIndex::StatName].Count += CountIncrement; \
+		} \
+	} while (0)
+
+#define UE_NET_IRIS_STATS_ADD_TIME_AND_COUNT_FOR_OBJECT_AS_WASTE(Timer, StatName, ObjectIndex) \
+	do { \
+		if (UE::Net::Private::FNetStatsContext* LocalNetStatsContext = Timer.GetNetStatsContext()) \
+		{ \
+			const uint64 DeltaTimeForStat = Timer.GetCyclesSinceStart(); \
+			uint32 CountIncrement = 0U; \
+			UE::Net::Private::FNetTypeStatsData& StatsData = UE::Net::Private::FNetStatsContext::GetTypeStatsDataForObject(*LocalNetStatsContext, ObjectIndex, CountIncrement); \
+			StatsData.Values[UE::Net::Private::FNetTypeStatsData::EStatsIndex::StatName].Time += DeltaTimeForStat; \
+			StatsData.Values[UE::Net::Private::FNetTypeStatsData::EStatsIndex::StatName].Count += CountIncrement; \
+			StatsData.Values[UE::Net::Private::FNetTypeStatsData::EStatsIndex::StatName##Waste].Time += DeltaTimeForStat; \
+			StatsData.Values[UE::Net::Private::FNetTypeStatsData::EStatsIndex::StatName##Waste].Count += CountIncrement; \
+		} \
+	} while (0)
+
+#define UE_NET_IRIS_STATS_ADD_BITS_WRITTEN_AND_COUNT_FOR_OBJECT(NetStatsContext, BitCount, StatName, ObjectIndex) \
+	do { \
+		if (NetStatsContext) \
+		{ \
+			uint32 CountIncrement = 0U; \
+			UE::Net::Private::FNetTypeStatsData& StatsData = UE::Net::Private::FNetStatsContext::GetTypeStatsDataForObject(*NetStatsContext, ObjectIndex, CountIncrement); \
+			StatsData.Values[UE::Net::Private::FNetTypeStatsData::EStatsIndex::StatName].Bits += BitCount; \
+			StatsData.Values[UE::Net::Private::FNetTypeStatsData::EStatsIndex::StatName].Count += CountIncrement; \
+		} \
+	} while (0)
+
+#define UE_NET_IRIS_STATS_ADD_BITS_WRITTEN_FOR_OBJECT(NetStatsContext, BitCount, StatName, ObjectIndex) \
+	do { \
+		if (NetStatsContext) \
+		{ \
+			uint32 CountIncrement = 0U; \
+			UE::Net::Private::FNetTypeStatsData& StatsData = UE::Net::Private::FNetStatsContext::GetTypeStatsDataForObject(*NetStatsContext, ObjectIndex, CountIncrement); \
+			StatsData.Values[UE::Net::Private::FNetTypeStatsData::EStatsIndex::StatName].Bits += BitCount; \
+		} \
+	} while (0)
+
+#define UE_NET_IRIS_STATS_ADD_BITS_WRITTEN_FOR_OBJECT_AS_WASTE(NetStatsContext, BitCount, StatName, ObjectIndex) \
+	do { \
+		if (NetStatsContext) \
+		{ \
+			uint32 CountIncrement = 0U; \
+			UE::Net::Private::FNetTypeStatsData& StatsData = UE::Net::Private::FNetStatsContext::GetTypeStatsDataForObject(*NetStatsContext, ObjectIndex, CountIncrement); \
+			StatsData.Values[UE::Net::Private::FNetTypeStatsData::EStatsIndex::StatName].Bits += BitCount; \
+			StatsData.Values[UE::Net::Private::FNetTypeStatsData::EStatsIndex::StatName##Waste].Bits += BitCount; \
+		} \
+	} while (0)
+
+#else
+
+#define UE_NET_IRIS_STATS_TIMER(...)
+#define UE_NET_IRIS_STATS_ADD_TIME_AND_COUNT_FOR_OBJECT(...)
+#define UE_NET_IRIS_STATS_ADD_TIME_AND_COUNT_FOR_OBJECT_AS_WASTE(...)
+#define UE_NET_IRIS_STATS_ADD_BITS_WRITTEN_AND_COUNT_FOR_OBJECT(...)
+#define UE_NET_IRIS_STATS_ADD_BITS_WRITTEN_FOR_OBJECT(...)
+#define UE_NET_IRIS_STATS_ADD_BITS_WRITTEN_FOR_OBJECT_AS_WASTE(...)
+
+#endif
diff --git a/Engine/Source/Runtime/Experimental/Iris/Core/Public/Iris/ReplicationSystem/ObjectReplicationBridge.h b/Engine/Source/Runtime/Experimental/Iris/Core/Public/Iris/ReplicationSystem/ObjectReplicationBridge.h
--- a/Engine/Source/Runtime/Experimental/Iris/Core/Public/Iris/ReplicationSystem/ObjectReplicationBridge.h
+++ b/Engine/Source/Runtime/Experimental/Iris/Core/Public/Iris/ReplicationSystem/ObjectReplicationBridge.h
@@ -172,6 +172,10 @@ public:
 	IRISCORE_API void SetClassDynamicFilterConfig(FName ClassPathName, const UE::Net::FNetObjectFilterHandle FilterHandle);
 	IRISCORE_API void SetClassDynamicFilterConfig(FName ClassPathName, FName FilterName);
 
+	/** Set the TypeStats to use for specified class and any derived classes without explicit config */
+	IRISCORE_API void SetClassTypeStatsConfig(FName ClassPathName, FName TypeStatsName);
+	IRISCORE_API void SetClassTypeStatsConfig(const FString& ClassPathName, const FString& TypeStatsName);
+
 public:
 
 	// Debug functions exposed via console commands
@@ -339,10 +343,12 @@ private:
 	/** Returns true if instances of this class should be delta compressed */
 	bool ShouldClassBeDeltaCompressed(const UClass* Class);
 
+	/** Returns the TypeStatsIndex this class should use */
+	int32 GetTypeStatsIndex(const UClass* Class);
+
 	FInstancePreUpdateFunction PreUpdateInstanceFunction;
 	FInstanceGetWorldObjectInfoFunction GetInstanceWorldObjectInfoFunction;
 	
-
 	FName GetConfigClassPathName(const UClass* Class);
 
 	void InitConditionalPropertyDelegates();
@@ -396,6 +402,9 @@ private:
 	// Classes that may force a disconnection when a protocol mismatch is detected.
 	TMap<FName, bool> ClassesFlaggedCritical;
 
+	// Type stats
+	TMap<FName, FName> ClassesWithTypeStats;
+
 	// Array of dormant objects that has requested a flush
 	TArray<FNetRefHandle> DormantHandlesPendingFlush;
 
diff --git a/Engine/Source/Runtime/Experimental/Iris/Core/Public/Iris/ReplicationSystem/ObjectReplicationBridgeConfig.h b/Engine/Source/Runtime/Experimental/Iris/Core/Public/Iris/ReplicationSystem/ObjectReplicationBridgeConfig.h
--- a/Engine/Source/Runtime/Experimental/Iris/Core/Public/Iris/ReplicationSystem/ObjectReplicationBridgeConfig.h
+++ b/Engine/Source/Runtime/Experimental/Iris/Core/Public/Iris/ReplicationSystem/ObjectReplicationBridgeConfig.h
@@ -94,6 +94,24 @@ struct FObjectReplicatedBridgeCriticalClassConfig
 	bool bDisconnectOnProtocolMismatch = true;
 };
 
+USTRUCT()
+struct FObjectReplicationBridgeTypeStatsConfig
+{
+	GENERATED_BODY()
+
+	/** Instances of this class or derived from this class should use delta compression */
+	UPROPERTY()
+	FName ClassName;
+
+	/** The TypeStatsName this class should use. */
+	UPROPERTY()
+	FName TypeStatsName;
+
+	/** If set to true this type will be reported even in configs with minimal stats reporting */
+	UPROPERTY()
+	bool bIncludeInMinimalCSVStats = false;
+};
+
 UCLASS(transient, config=Engine)
 class UObjectReplicationBridgeConfig : public UObject
 {
@@ -108,6 +126,7 @@ public:
 	IRISCORE_API TConstArrayView<FObjectReplicationBridgePrioritizerConfig> GetPrioritizerConfigs() const;
 	IRISCORE_API TConstArrayView<FObjectReplicationBridgeDeltaCompressionConfig> GetDeltaCompressionConfigs() const;
 	IRISCORE_API TConstArrayView<FObjectReplicatedBridgeCriticalClassConfig> GetCriticalClassConfigs() const;
+	IRISCORE_API TConstArrayView<FObjectReplicationBridgeTypeStatsConfig> GetTypeStatsConfigs() const;
 
 	FName GetDefaultSpatialFilterName() const;
 	FName GetRequiredNetDriverChannelClassName() const;
@@ -155,6 +174,12 @@ private:
 	UPROPERTY(Config)
 	bool bAllClassesCritical = false;
 
+	/**
+	 * Which classes should collect TypeStats. Derived classes will get the same behavior unless overidden
+	 */
+	UPROPERTY(Config)
+	TArray<FObjectReplicationBridgeTypeStatsConfig> TypeStatsConfigs;
+
 	/**
 	 * The name of the filter to apply objects that can have spatial filtering applied.
 	 */
diff --git a/Engine/Source/Runtime/Experimental/Iris/Core/Public/Iris/Serialization/NetBitStreamUtil.h b/Engine/Source/Runtime/Experimental/Iris/Core/Public/Iris/Serialization/NetBitStreamUtil.h
--- a/Engine/Source/Runtime/Experimental/Iris/Core/Public/Iris/Serialization/NetBitStreamUtil.h
+++ b/Engine/Source/Runtime/Experimental/Iris/Core/Public/Iris/Serialization/NetBitStreamUtil.h
@@ -66,6 +66,8 @@ public:
 		Writer.Seek(StartPos);
 	}
 
+	uint32 GetStartPos() const { return StartPos; }
+
 private:
 	FNetBitStreamWriter& Writer;
 	const uint32 StartPos;
diff --git a/Engine/Source/Runtime/Experimental/Iris/Core/Public/Iris/Serialization/NetSerializationContext.h b/Engine/Source/Runtime/Experimental/Iris/Core/Public/Iris/Serialization/NetSerializationContext.h
--- a/Engine/Source/Runtime/Experimental/Iris/Core/Public/Iris/Serialization/NetSerializationContext.h
+++ b/Engine/Source/Runtime/Experimental/Iris/Core/Public/Iris/Serialization/NetSerializationContext.h
@@ -16,6 +16,7 @@ namespace UE::Net
 	{
 		class FInternalNetSerializationContext;
 		class FNetExportContext;
+		class FNetStatsContext;
 	}
 }
 
@@ -75,6 +76,9 @@ public:
 	void SetExportContext(Private::FNetExportContext* InExportContext) { ExportContext = InExportContext; }
 	Private::FNetExportContext* GetExportContext() { return ExportContext; }
 
+	void SetNetStatsContext(Private::FNetStatsContext* InNetStatsContext) { NetStatsContext = InNetStatsContext; }
+	Private::FNetStatsContext* GetNetStatsContext() { return NetStatsContext; }
+
 	void SetIsInitializingDefaultState(bool bInIsInitializingDefaultState) { bIsInitializingDefaultState = bInIsInitializingDefaultState; }
 	bool IsInitializingDefaultState() const { return bIsInitializingDefaultState; }
 
@@ -91,6 +95,7 @@ private:
 	FNetTraceCollector* TraceCollector;
 	Private::FInternalNetSerializationContext* InternalContext;
 	Private::FNetExportContext* ExportContext;
+	Private::FNetStatsContext* NetStatsContext;
 	const FNetBitArrayView* ChangeMask;
 	INetBlobReceiver* NetBlobReceiver;
 	uint32 LocalConnectionId;
@@ -108,6 +113,7 @@ inline FNetSerializationContext::FNetSerializationContext(FNetBitStreamReader* I
 , TraceCollector(nullptr)
 , InternalContext(nullptr)
 , ExportContext(nullptr)
+, NetStatsContext(nullptr)
 , ChangeMask(nullptr)
 , NetBlobReceiver(nullptr)
 , LocalConnectionId(0)
