Iris
- Fixed compilation errors for android

#rnx

diff --git a/Engine/Source/Runtime/Experimental/Iris/Core/Private/Iris/Stats/NetStats.cpp b/Engine/Source/Runtime/Experimental/Iris/Core/Private/Iris/Stats/NetStats.cpp
--- a/Engine/Source/Runtime/Experimental/Iris/Core/Private/Iris/Stats/NetStats.cpp
+++ b/Engine/Source/Runtime/Experimental/Iris/Core/Private/Iris/Stats/NetStats.cpp
@@ -100,9 +100,12 @@ void FNetSendStats::ReportCsvStats()
 #endif
 }
 
+FNetStatsContext::FNetStatsContext() = default;
+FNetStatsContext::~FNetStatsContext() = default;
+
 FNetTypeStats::FNetTypeStats()
 {
-	StatsContext = TUniquePtr<FNetStatsContext>(CreateNetStatsContext());
+	StatsContext = CreateNetStatsContext();
 	// Add default TypeStats
 	GetOrCreateTypeStats(TEXT("Undefined"));
 	GetOrCreateTypeStats(TEXT("OOBChannel"));
@@ -110,6 +113,7 @@ FNetTypeStats::FNetTypeStats()
 
 FNetTypeStats::~FNetTypeStats()
 {
+	delete StatsContext;
 }
 
 void FNetTypeStats::Init(FInitParams& InitParams)
@@ -157,7 +161,7 @@ void FNetTypeStats::Accumulate(FNetStatsContext& Context)
 	IRIS_PROFILER_SCOPE(FNetTypeStats_Accumulate);
 
 	// Skip default context as that is our target.
-	if (&Context == StatsContext.Get())
+	if (&Context == StatsContext)
 	{
 		return;
 	}
diff --git a/Engine/Source/Runtime/Experimental/Iris/Core/Private/Iris/Stats/NetStats.h b/Engine/Source/Runtime/Experimental/Iris/Core/Private/Iris/Stats/NetStats.h
--- a/Engine/Source/Runtime/Experimental/Iris/Core/Private/Iris/Stats/NetStats.h
+++ b/Engine/Source/Runtime/Experimental/Iris/Core/Private/Iris/Stats/NetStats.h
@@ -4,7 +4,6 @@
 
 #include "HAL/Platform.h"
 #include "HAL/CriticalSection.h"
-#include "Templates/UniquePtr.h"
 #include "Containers/Array.h"
 
 namespace UE::Net::Private
@@ -188,7 +187,7 @@ public:
 	int32 GetOrCreateTypeStats(FName Name);
 
 	/** Get default context if stats is enabled, once we decide to go wide we will need to expose methods to create and refresh StatsContexts per thread */
-	FNetStatsContext* GetNetStatsContext() { return IsEnabled() ?  StatsContext.Get() : nullptr; }
+	FNetStatsContext* GetNetStatsContext() { return IsEnabled() ?  StatsContext : nullptr; }
 
 	/** Updated every frame based on the state of the CSVProfiler */
 	bool IsEnabled() const { return bIsEnabled; }
@@ -203,7 +202,7 @@ private:
 	FNetStatsContext* CreateNetStatsContext();
 	void UpdateContext(FNetStatsContext& Context);
 
-	TUniquePtr<FNetStatsContext> StatsContext = nullptr;
+	FNetStatsContext* StatsContext = nullptr;
 	FNetRefHandleManager* NetRefHandleManager = nullptr;
 	TArray<FName> TypeStatsNames;
 	bool bIsEnabled = false;
diff --git a/Engine/Source/Runtime/Experimental/Iris/Core/Private/Iris/Stats/NetStatsContext.h b/Engine/Source/Runtime/Experimental/Iris/Core/Private/Iris/Stats/NetStatsContext.h
--- a/Engine/Source/Runtime/Experimental/Iris/Core/Private/Iris/Stats/NetStatsContext.h
+++ b/Engine/Source/Runtime/Experimental/Iris/Core/Private/Iris/Stats/NetStatsContext.h
@@ -59,7 +59,8 @@ struct FNetTypeStatsData
 class FNetStatsContext
 {
 public:
-	FNetStatsContext() = default;
+	FNetStatsContext();
+	~FNetStatsContext();
 	FNetStatsContext(const FNetStatsContext&) = delete;
 	FNetStatsContext& operator=(const FNetStatsContext&) = delete;
 
