ChaosSofts : Fix CIS Issue

diff --git a/Engine/Source/Runtime/Experimental/Chaos/Public/Chaos/Deformable/GaussSeidelCorotatedCodimensionalConstraints.h b/Engine/Source/Runtime/Experimental/Chaos/Public/Chaos/Deformable/GaussSeidelCorotatedCodimensionalConstraints.h
--- a/Engine/Source/Runtime/Experimental/Chaos/Public/Chaos/Deformable/GaussSeidelCorotatedCodimensionalConstraints.h
+++ b/Engine/Source/Runtime/Experimental/Chaos/Public/Chaos/Deformable/GaussSeidelCorotatedCodimensionalConstraints.h
@@ -54,23 +54,39 @@ namespace Chaos::Softs
 
 		PMatrix<T, 3, 2> Ds(const int32 E, const ParticleType& InParticles) const 
 		{
-			check(E > INDEX_NONE && E < MeshConstraints.Num());
-			check(MeshConstraints[E][0] < (int32)InParticles.Size() && MeshConstraints[E][0] > INDEX_NONE);
-			check(MeshConstraints[E][1] < (int32)InParticles.Size() && MeshConstraints[E][1] > INDEX_NONE);
-			check(MeshConstraints[E][2] < (int32)InParticles.Size() && MeshConstraints[E][2] > INDEX_NONE);
-			const TVec3<T> P1P0 = InParticles.P(MeshConstraints[E][1]) - InParticles.P(MeshConstraints[E][0]);
-			const TVec3<T> P2P0 = InParticles.P(MeshConstraints[E][2]) - InParticles.P(MeshConstraints[E][0]);
-
-			return PMatrix<T, 3, 2>(
-				P1P0[0], P1P0[1], P1P0[2],
-				P2P0[0], P2P0[1], P2P0[2]);
+			if (INDEX_NONE < E && E< MeshConstraints.Num()
+				&& INDEX_NONE < MeshConstraints[E][0] && MeshConstraints[E][0] < (int32)InParticles.Size()
+				&& INDEX_NONE < MeshConstraints[E][1] && MeshConstraints[E][1] < (int32)InParticles.Size()
+				&& INDEX_NONE < MeshConstraints[E][2] && MeshConstraints[E][2] < (int32)InParticles.Size())
+			{
+				const TVec3<T> P1P0 = InParticles.P(MeshConstraints[E][1]) - InParticles.P(MeshConstraints[E][0]);
+				const TVec3<T> P2P0 = InParticles.P(MeshConstraints[E][2]) - InParticles.P(MeshConstraints[E][0]);
+
+				return PMatrix<T, 3, 2>(
+					P1P0[0], P1P0[1], P1P0[2],
+					P2P0[0], P2P0[1], P2P0[2]);
+			}
+			else
+			{
+				return PMatrix<T, 3, 2>(
+					(T)0., (T)0, (T)0,
+					(T)0, (T)0, (T)0);
+			}
+		
 		}
 
 		PMatrix<T, 3, 2> F(const int32 E, const ParticleType& InParticles) const 
 		{
-			check(E > INDEX_NONE && E < MeshConstraints.Num());
-			check(E < DmInverse.Num());
-			return Ds(E, InParticles) * DmInverse[E];
+			if (INDEX_NONE < E && E < DmInverse.Num())
+			{
+				return Ds(E, InParticles) * DmInverse[E];
+			}
+			else
+			{
+				return PMatrix<T, 3, 2>(
+					(T)0., (T)0, (T)0,
+					(T)0, (T)0, (T)0);
+			}
 		}
 
 		TArray<TArray<int32>> GetConstraintsArray() const
@@ -114,17 +130,22 @@ namespace Chaos::Softs
 				check(MeshConstraints[e][0] < (int32)Particles.Size() && MeshConstraints[e][0] > INDEX_NONE);
 				check(MeshConstraints[e][1] < (int32)Particles.Size() && MeshConstraints[e][1] > INDEX_NONE);
 				check(MeshConstraints[e][2] < (int32)Particles.Size() && MeshConstraints[e][2] > INDEX_NONE);
-				const TVec3<T> X1X0 = Particles.X(MeshConstraints[e][1]) - Particles.X(MeshConstraints[e][0]);
-				const TVec3<T> X2X0 = Particles.X(MeshConstraints[e][2]) - Particles.X(MeshConstraints[e][0]);
-				PMatrix<T, 2, 2> Dm((T)0., (T)0., (T)0.);
-				Dm.M[0] = X1X0.Size();
-				Dm.M[2] = X1X0.Dot(X2X0) / Dm.M[0];
-				Dm.M[3] = Chaos::TVector<T, 3>::CrossProduct(X1X0, X2X0).Size() / Dm.M[0];
-				Measure[e] = Chaos::TVector<T, 3>::CrossProduct(X1X0, X2X0).Size() / (T)2.;
-				ensureMsgf(Measure[e] > (T)0., TEXT("Degenerate triangle detected"));
-
-				PMatrix<T, 2, 2> DmInv = Dm.Inverse();
-				DmInverse[e] = DmInv;
+				if (MeshConstraints[e][0] < (int32)Particles.Size() && MeshConstraints[e][0] > INDEX_NONE
+					&& MeshConstraints[e][1] < (int32)Particles.Size() && MeshConstraints[e][1] > INDEX_NONE
+					&& MeshConstraints[e][2] < (int32)Particles.Size() && MeshConstraints[e][2] > INDEX_NONE)
+				{
+					const TVec3<T> X1X0 = Particles.X(MeshConstraints[e][1]) - Particles.X(MeshConstraints[e][0]);
+					const TVec3<T> X2X0 = Particles.X(MeshConstraints[e][2]) - Particles.X(MeshConstraints[e][0]);
+					PMatrix<T, 2, 2> Dm((T)0., (T)0., (T)0.);
+					Dm.M[0] = X1X0.Size();
+					Dm.M[2] = X1X0.Dot(X2X0) / Dm.M[0];
+					Dm.M[3] = Chaos::TVector<T, 3>::CrossProduct(X1X0, X2X0).Size() / Dm.M[0];
+					Measure[e] = Chaos::TVector<T, 3>::CrossProduct(X1X0, X2X0).Size() / (T)2.;
+					ensureMsgf(Measure[e] > (T)0., TEXT("Degenerate triangle detected"));
+
+					PMatrix<T, 2, 2> DmInv = Dm.Inverse();
+					DmInverse[e] = DmInv;
+				}
 			}
 		}
 
@@ -257,100 +278,109 @@ namespace Chaos::Softs
 			check(ElementIndex < Measure.Num());
 			check(ElementIndex < LambdaElementArray.Num());
 			check(ElementIndexLocal < 4 && ElementIndexLocal > INDEX_NONE);
-			const Chaos::PMatrix<T, 2, 2> DmInvT = DmInverse[ElementIndex].GetTransposed(), DmInv = DmInverse[ElementIndex]; 
-			const Chaos::PMatrix<T, 3, 2> Fe = F(ElementIndex, Particles);
+			if (ElementIndexLocal < 4 
+				&& ElementIndexLocal > INDEX_NONE
+				&& ElementIndex < Measure.Num()
+				&& ElementIndex < MuElementArray.Num()
+				&& ElementIndex < LambdaElementArray.Num()
+				&& ElementIndex < DmInverse.Num() 
+				&& ElementIndex > INDEX_NONE)
+			{
+				const Chaos::PMatrix<T, 2, 2> DmInvT = DmInverse[ElementIndex].GetTransposed(), DmInv = DmInverse[ElementIndex]; 
+				const Chaos::PMatrix<T, 3, 2> Fe = F(ElementIndex, Particles);
 
-			const PMatrix<T, 3, 2> Re = ComputeR(Fe);
-			const PMatrix<T, 2, 2> FTF = ComputeFTF(Fe);
-			const T J = FMath::Sqrt(FTF.Determinant());
+				const PMatrix<T, 3, 2> Re = ComputeR(Fe);
+				const PMatrix<T, 2, 2> FTF = ComputeFTF(Fe);
+				const T J = FMath::Sqrt(FTF.Determinant());
 
-			PMatrix<T, 3, 2> Pe(TVec3<T>((T)0.), TVec3<T>((T)0.)), ForceTerm(TVec3<T>((T)0.), TVec3<T>((T)0.));
-			Pe = T(2) * MuElementArray[ElementIndex] * (Fe - Re) + LambdaElementArray[ElementIndex] * (J - T(1)) * J * Fe * ComputeFTF(Fe).Inverse(); 
+				PMatrix<T, 3, 2> Pe(TVec3<T>((T)0.), TVec3<T>((T)0.)), ForceTerm(TVec3<T>((T)0.), TVec3<T>((T)0.));
+				Pe = T(2) * MuElementArray[ElementIndex] * (Fe - Re) + LambdaElementArray[ElementIndex] * (J - T(1)) * J * Fe * ComputeFTF(Fe).Inverse(); 
 
-			ForceTerm = -Measure[ElementIndex] * Pe * DmInverse[ElementIndex].GetTransposed();
+				ForceTerm = -Measure[ElementIndex] * Pe * DmInverse[ElementIndex].GetTransposed();
 
-			Chaos::TVector<T, 3> Dx((T)0.);
-			if (ElementIndexLocal > 0)
-			{
-				for (int32 c = 0; c < 3; c++)
+				Chaos::TVector<T, 3> Dx((T)0.);
+				if (ElementIndexLocal > 0)
 				{
-					Dx[c] += ForceTerm.M[ElementIndexLocal * 3 - 3 + c];
+					for (int32 c = 0; c < 3; c++)
+					{
+						Dx[c] += ForceTerm.M[ElementIndexLocal * 3 - 3 + c];
+					}
 				}
-			}
-			else
-			{
-				for (int32 c = 0; c < 3; c++)
+				else
 				{
-					for (int32 h = 0; h < 2; h++)
+					for (int32 c = 0; c < 3; c++)
 					{
-						Dx[c] -= ForceTerm.M[h * 3 + c];
+						for (int32 h = 0; h < 2; h++)
+						{
+							Dx[c] -= ForceTerm.M[h * 3 + c];
+						}
 					}
 				}
-			}
 
-			Dx *= Dt * Dt;
+				Dx *= Dt * Dt;
 
-			ParticleResidual -= Dx;
+				ParticleResidual -= Dx;
 
-			PMatrix<T, 3, 2> dJ(TVec3<T>((T)0.), TVec3<T>((T)0.));
+				PMatrix<T, 3, 2> dJ(TVec3<T>((T)0.), TVec3<T>((T)0.));
 
-			dJdF32(Fe, dJ);
+				dJdF32(Fe, dJ);
 
-			T Coeff = Dt * Dt * Measure[ElementIndex];
+				T Coeff = Dt * Dt * Measure[ElementIndex];
 
-			if (ElementIndexLocal == 0) 
-			{
-				T DmInvsum = T(0);
-				for (int32 nu = 0; nu < 2; nu++) 
+				if (ElementIndexLocal == 0) 
 				{
-					T localDmsum = T(0);
-					for (int32 k = 0; k < 2; k++) 
+					T DmInvsum = T(0);
+					for (int32 nu = 0; nu < 2; nu++) 
 					{
-						localDmsum += DmInv.M[nu * 2 + k];
+						T localDmsum = T(0);
+						for (int32 k = 0; k < 2; k++) 
+						{
+							localDmsum += DmInv.M[nu * 2 + k];
+						}
+						DmInvsum += localDmsum * localDmsum;
+					}
+					for (int32 alpha = 0; alpha < 3; alpha++) 
+					{
+						ParticleHessian.SetAt(alpha, alpha, ParticleHessian.GetAt(alpha, alpha) + Coeff * T(2) * MuElementArray[ElementIndex] * DmInvsum);
 					}
-					DmInvsum += localDmsum * localDmsum;
-				}
-				for (int32 alpha = 0; alpha < 3; alpha++) 
-				{
-					ParticleHessian.SetAt(alpha, alpha, ParticleHessian.GetAt(alpha, alpha) + Coeff * T(2) * MuElementArray[ElementIndex] * DmInvsum);
-				}
 
-				const PMatrix<T, 3, 2> dJDmInvT = dJ * DmInvT; 
-				Chaos::TVector<T, 3> L((T)0.);
-				for (int32 alpha = 0; alpha < 3; alpha++) 
-				{
-					for (int32 k = 0; k < 2; k++) 
+					const PMatrix<T, 3, 2> dJDmInvT = dJ * DmInvT; 
+					Chaos::TVector<T, 3> L((T)0.);
+					for (int32 alpha = 0; alpha < 3; alpha++) 
 					{
-						L[alpha] += dJDmInvT.M[3 * k + alpha]; 
+						for (int32 k = 0; k < 2; k++) 
+						{
+							L[alpha] += dJDmInvT.M[3 * k + alpha]; 
+						}
+					}
+					for (int32 alpha = 0; alpha < 3; alpha++)
+					{
+						ParticleHessian.SetRow(alpha, ParticleHessian.GetRow(alpha) + Coeff * LambdaElementArray[ElementIndex] * L[alpha] * L);
 					}
 				}
-				for (int32 alpha = 0; alpha < 3; alpha++)
-				{
-					ParticleHessian.SetRow(alpha, ParticleHessian.GetRow(alpha) + Coeff * LambdaElementArray[ElementIndex] * L[alpha] * L);
-				}
-			}
-			else 
-			{
-				T DmInvsum = T(0);
-				for (int32 nu = 0; nu < 2; nu++)
-				{
-					DmInvsum += DmInv.GetAt(ElementIndexLocal - 1, nu) * DmInv.GetAt(ElementIndexLocal - 1, nu);
-				}
-				for (int32 alpha = 0; alpha < 3; alpha++)
+				else 
 				{
-					ParticleHessian.SetAt(alpha, alpha, ParticleHessian.GetAt(alpha, alpha) + Dt * Dt * Measure[ElementIndex] * T(2) * MuElementArray[ElementIndex] * DmInvsum);
-				}
+					T DmInvsum = T(0);
+					for (int32 nu = 0; nu < 2; nu++)
+					{
+						DmInvsum += DmInv.GetAt(ElementIndexLocal - 1, nu) * DmInv.GetAt(ElementIndexLocal - 1, nu);
+					}
+					for (int32 alpha = 0; alpha < 3; alpha++)
+					{
+						ParticleHessian.SetAt(alpha, alpha, ParticleHessian.GetAt(alpha, alpha) + Dt * Dt * Measure[ElementIndex] * T(2) * MuElementArray[ElementIndex] * DmInvsum);
+					}
 
-				const PMatrix<T, 3, 2> dJDmInvT = dJ * DmInvT; 
+					const PMatrix<T, 3, 2> dJDmInvT = dJ * DmInvT; 
 
-				Chaos::TVector<T, 3> L((T)0.);
-				for (int32 alpha = 0; alpha < 3; alpha++) 
-				{
-					L[alpha] = dJDmInvT.M[ElementIndexLocal* 3 - 3 + alpha];
-				}
-				for (int32 alpha = 0; alpha < 3; alpha++)
-				{
-					ParticleHessian.SetRow(alpha, ParticleHessian.GetRow(alpha) + Dt * Dt * Measure[ElementIndex] * LambdaElementArray[ElementIndex] * L[alpha] * L);
+					Chaos::TVector<T, 3> L((T)0.);
+					for (int32 alpha = 0; alpha < 3; alpha++) 
+					{
+						L[alpha] = dJDmInvT.M[ElementIndexLocal* 3 - 3 + alpha];
+					}
+					for (int32 alpha = 0; alpha < 3; alpha++)
+					{
+						ParticleHessian.SetRow(alpha, ParticleHessian.GetRow(alpha) + Dt * Dt * Measure[ElementIndex] * LambdaElementArray[ElementIndex] * L[alpha] * L);
+					}
 				}
 			}
 		}
