Fix race condition caused by simulation parameters being unsafely written on the game thread while being read on the physics thread
- Removed the need for FGeometryCollectioPerFrameData
- Moved Set*_External function to post update on the physics thread via the immediate enqueuing system ( with is lock free )

diff --git a/Engine/Source/Runtime/Experimental/Chaos/Private/PhysicsProxy/GeometryCollectionPhysicsProxy.cpp b/Engine/Source/Runtime/Experimental/Chaos/Private/PhysicsProxy/GeometryCollectionPhysicsProxy.cpp
--- a/Engine/Source/Runtime/Experimental/Chaos/Private/PhysicsProxy/GeometryCollectionPhysicsProxy.cpp
+++ b/Engine/Source/Runtime/Experimental/Chaos/Private/PhysicsProxy/GeometryCollectionPhysicsProxy.cpp
@@ -190,6 +190,14 @@ void FGeometryCollectionResults::Reset()
 // FGeometryCollectionPhysicsProxy helper functions
 //==============================================================================
 
+template<typename TLambda>
+void ExecuteOnPhysicsThread(FGeometryCollectionPhysicsProxy& Proxy, TLambda&& Lambda)
+{
+	if (Chaos::FPhysicsSolver* RBDSolver = Proxy.GetSolver<Chaos::FPhysicsSolver>())
+	{
+		RBDSolver->EnqueueCommandImmediate(Lambda);
+	}
+}
 
 TUniquePtr<Chaos::FTriangleMesh> CreateTriangleMesh(
 	const int32 FaceStart,
@@ -486,13 +494,8 @@ FGeometryCollectionPhysicsProxy::FGeometryCollectionPhysicsProxy(
 	, CollisionParticlesPerObjectFraction(CollisionParticlesPerObjectFractionDefault)
 	, PhysicsThreadCollection(Parameters.RestCollection)
 	, GameThreadCollection(GameThreadCollectionIn)
-	, GameThreadPerFrameData(SimulationParameters)
-	, MaterialOverrideMassScaleMultiplierChange(0)
-	, bIsPhysicsThreadWorldTransformDirty(false)
-	, bIsCollisionFilterDataDirty(false)
-	, bIsDamageThresholdDataDirty(false)
-	, bIsGravityGroupIndexDirty(false)
-	, bIsOneWayInteractionDirty(false)
+	, WorldTransform_External(FTransform::Identity)
+	, bIsGameThreadWorldTransformDirty(false)
 	, CollectorGuid(InCollectorGuid)
 {
 	// We rely on a guarded buffer.
@@ -544,8 +547,8 @@ void FGeometryCollectionPhysicsProxy::Initialize(Chaos::FPBDRigidsEvolutionBase
 		}
 	}
 
-	// compatibility requirement to make sure we at least initialize GameThreadPerFrameData properly
-	GameThreadPerFrameData.SetWorldTransform(Parameters.WorldTransform);
+	// we need to make sure the world transform is kept up to dat eon the game thread 
+	WorldTransform_External = Parameters.WorldTransform;
 
 	//
 	// Collision vertices down sampling validation.  
@@ -2592,51 +2595,63 @@ void FGeometryCollectionPhysicsProxy::SetProxyDirty_External()
 
 void FGeometryCollectionPhysicsProxy::SetEnableDamageFromCollision_External(bool bEnable)
 {
-	check(IsInGameThread());
-	GameThreadPerFrameData.SetEnableStrainOnCollision(bEnable);
-	SetProxyDirty_External();
+	ExecuteOnPhysicsThread(*this,
+		[this, bEnable]()
+		{
+			Parameters.bEnableStrainOnCollision = bEnable;
+		});
 }
 
 void FGeometryCollectionPhysicsProxy::SetNotifyBreakings_External(bool bNotify)
 {
-	check(IsInGameThread());
-	GameThreadPerFrameData.SetNotifyBreakings(bNotify);
-	SetProxyDirty_External();
+	ExecuteOnPhysicsThread(*this,
+		[this, bNotify]()
+		{
+			Parameters.bGenerateBreakingData = bNotify;
+		});
 }
 
 void FGeometryCollectionPhysicsProxy::SetNotifyRemovals_External(bool bNotify)
 {
-	check(IsInGameThread());
-	GameThreadPerFrameData.SetNotifyRemovals(bNotify);
-	SetProxyDirty_External();
+	// nothing to do : there's no removal flag to set on the proxy game thread or physics thread side
+	// todo(chaos) we shoudl probably have one as we may add to the removal array regardless of what the user sets on the component side 
 }
 
 void FGeometryCollectionPhysicsProxy::SetNotifyCrumblings_External(bool bNotify, bool bIncludeChildren)
 {
-	check(IsInGameThread());
-	GameThreadPerFrameData.SetNotifyCrumblings(bNotify, bIncludeChildren);
-	SetProxyDirty_External();
+	ExecuteOnPhysicsThread(*this,
+		[this, bNotify, bIncludeChildren]()
+		{
+			Parameters.bGenerateCrumblingData = bNotify;
+			Parameters.bGenerateCrumblingChildrenData = bIncludeChildren;
+
+		});
 }
 
 void FGeometryCollectionPhysicsProxy::SetNotifyGlobalBreakings_External(bool bNotify)
 {
-	check(IsInGameThread());
-	GameThreadPerFrameData.SetNotifyGlobalBreakings(bNotify);
-	SetProxyDirty_External();
+	ExecuteOnPhysicsThread(*this,
+		[this, bNotify]()
+		{
+			Parameters.bGenerateGlobalBreakingData = bNotify;
+		});
 }
 
 void FGeometryCollectionPhysicsProxy::SetNotifyGlobalRemovals_External(bool bNotify)
 {
-	check(IsInGameThread());
-	GameThreadPerFrameData.SetNotifyGlobalRemovals(bNotify);
-	SetProxyDirty_External();
+	// nothing to do : there's no removal flag to set on the proxy game thread or physics thread side
+	// todo(chaos) we shoudl probably have one as we may add to the removal array regardless of what the user sets on the component side 
 }
 
 void FGeometryCollectionPhysicsProxy::SetNotifyGlobalCrumblings_External(bool bNotify, bool bIncludeChildren)
 {
-	check(IsInGameThread());
-	GameThreadPerFrameData.SetNotifyGlobalCrumblings(bNotify, bIncludeChildren);
-	SetProxyDirty_External();
+	ExecuteOnPhysicsThread(*this,
+		[this, bNotify, bIncludeChildren]()
+		{
+			Parameters.bGenerateGlobalCrumblingData = bNotify;
+			Parameters.bGenerateGlobalCrumblingChildrenData = bIncludeChildren;
+
+		});
 }
 
 int32 FGeometryCollectionPhysicsProxy::CalculateHierarchyLevel(const FGeometryDynamicCollection& DynamicCollection, int32 TransformIndex)
@@ -2954,9 +2969,103 @@ void FGeometryCollectionPhysicsProxy::BufferGameState()
 
 void FGeometryCollectionPhysicsProxy::SetWorldTransform_External(const FTransform& WorldTransform)
 {
-	check(IsInGameThread());
-	GameThreadPerFrameData.SetWorldTransform(WorldTransform);
-	SetProxyDirty_External();
+	// todo : change to compare with previous value
+	const bool bHasTransformChanged = !WorldTransform.Equals(WorldTransform_External);
+	if (bHasTransformChanged)
+	{
+		bIsGameThreadWorldTransformDirty = bHasTransformChanged;
+		WorldTransform_External = WorldTransform;
+
+		ExecuteOnPhysicsThread(*this,
+			[this, WorldTransform]()
+			{
+				SetWorldTransform_Internal(WorldTransform);
+			});
+	}
+}
+
+void FGeometryCollectionPhysicsProxy::SetWorldTransform_Internal(const FTransform& InWorldTransform)
+{
+	using namespace Chaos;
+
+	Parameters.PrevWorldTransform = Parameters.WorldTransform;
+	Parameters.WorldTransform = InWorldTransform;
+
+	TSet<FClusterHandle*> ProcessedInternalClusters;
+	const FTransform& ActorToWorld = Parameters.WorldTransform;
+
+	const int32 NumTransforms = PhysicsThreadCollection.NumElements(FGeometryCollection::TransformGroup);
+	if (Chaos::FPhysicsSolver* RigidSolver = GetSolver<Chaos::FPhysicsSolver>())
+	{
+		FClusterUnionManager& ClusterUnionManager = RigidSolver->GetEvolution()->GetRigidClustering().GetClusterUnionManager();
+
+		// Assume all particles are in the same cluster union.
+		FClusterUnionIndex ClusterUnionIndex = INDEX_NONE;
+		TArray<FPBDRigidParticleHandle*> DeferredClusterUnionParticleUpdates;
+		TArray<FTransform> DeferredClusterUnionChildToParentUpdates;
+
+		const TManagedArray<FTransform>& MassToLocal = Parameters.RestCollection->GetAttribute<FTransform>(MassToLocalAttributeName, FTransformCollection::TransformGroup);
+
+		for (int32 TransformGroupIndex = 0; TransformGroupIndex < NumTransforms; ++TransformGroupIndex)
+		{
+			if (Chaos::FPBDRigidClusteredParticleHandle* Handle = SolverParticleHandles[TransformGroupIndex])
+			{
+				FClusterHandle* KinematicRootHandle = nullptr;
+				FClusterHandle* ParentHandle = Handle->Parent();
+
+				if (!Handle->Disabled() && Handle->ObjectState() == Chaos::EObjectStateType::Kinematic)
+				{
+					KinematicRootHandle = Handle;
+				}
+				else
+				{
+					// is there a internal parent as a kinematic root?
+					if (ParentHandle && ParentHandle->InternalCluster() && !ParentHandle->Disabled() && ParentHandle->ObjectState() == Chaos::EObjectStateType::Kinematic && ParentHandle->PhysicsProxy() == this)
+					{
+						if (!ProcessedInternalClusters.Contains(ParentHandle))
+						{
+							ProcessedInternalClusters.Add(ParentHandle);
+							KinematicRootHandle = ParentHandle;
+						}
+					}
+				}
+
+				if (KinematicRootHandle)
+				{
+					const FTransform RootWorldTransform(KinematicRootHandle->R(), KinematicRootHandle->X());
+					const FTransform RootRelativeTransform = RootWorldTransform.GetRelativeTransform(Parameters.PrevWorldTransform);
+					const FTransform WorldTransform = RootRelativeTransform * ActorToWorld;
+
+					SetClusteredParticleKinematicTarget_Internal(KinematicRootHandle, WorldTransform);
+				}
+				else if (ParentHandle && !ParentHandle->IsDynamic())
+				{
+					if (ClusterUnionIndex == INDEX_NONE)
+					{
+						ClusterUnionIndex = ClusterUnionManager.FindClusterUnionIndexFromParticle(Handle);
+					}
+
+					if (ClusterUnionIndex != INDEX_NONE)
+					{
+						const FTransform ParentWorldTransform{ ParentHandle->R(), ParentHandle->X() };
+						const FTransform NewWorldTransform = MassToLocal[TransformGroupIndex] * FTransform(PhysicsThreadCollection.GetTransform(TransformGroupIndex)) * Parameters.WorldTransform;
+						const FTransform RelativeTransform = NewWorldTransform.GetRelativeTransform(ParentWorldTransform);
+
+						DeferredClusterUnionParticleUpdates.Add(Handle);
+						DeferredClusterUnionChildToParentUpdates.Add(RelativeTransform);
+
+						// Make sure the particle is mark dirty to make sure its proxy will properly update the transforms
+						RigidSolver->GetEvolution()->GetParticles().MarkTransientDirtyParticle(Handle);
+					}
+				}
+			}
+		}
+
+		if (ClusterUnionIndex != INDEX_NONE && !DeferredClusterUnionParticleUpdates.IsEmpty() && !DeferredClusterUnionChildToParentUpdates.IsEmpty())
+		{
+			ClusterUnionManager.UpdateClusterUnionParticlesChildToParent(ClusterUnionIndex, DeferredClusterUnionParticleUpdates, DeferredClusterUnionChildToParentUpdates, false);
+		}
+	}
 }
 
 DECLARE_CYCLE_STAT(TEXT("FGeometryCollectionPhysicsProxy::SetUseStaticMeshCollisionForTraces_External"), STAT_SetUseStaticMeshCollisionForTraces_External, STATGROUP_Chaos);
@@ -2989,7 +3098,7 @@ void FGeometryCollectionPhysicsProxy::SetUseStaticMeshCollisionForTraces_Externa
 				}
 				else
 				{
-					const FVector Scale = Parameters.WorldTransform.GetScale3D();
+					const FVector Scale = WorldTransform_External.GetScale3D();
 					TManagedArray<Chaos::FImplicitObjectPtr>& Implicits = GameThreadCollection.ModifyAttribute<Chaos::FImplicitObjectPtr>(FGeometryDynamicCollection::ImplicitsAttribute, FTransformCollection::TransformGroup);
 					Chaos::FImplicitObjectPtr ImplicitGeometry = Implicits[Index];
 					if (ImplicitGeometry && !Scale.Equals(FVector::OneVector))
@@ -3008,334 +3117,162 @@ void FGeometryCollectionPhysicsProxy::SetUseStaticMeshCollisionForTraces_Externa
 
 void FGeometryCollectionPhysicsProxy::SetDamageThresholds_External(const TArray<float>& DamageThresholds)
 {
-	check(IsInGameThread());
-	GameThreadPerFrameData.SetDamageThresholds(DamageThresholds);
-	SetProxyDirty_External();
+	ExecuteOnPhysicsThread(*this,
+		[this, DamageThresholds]()
+		{
+			SetDamageThresholds_Internal(DamageThresholds);
+		});
 }
 
-void FGeometryCollectionPhysicsProxy::SetDamagePropagationData_External(bool bEnabled, float BreakDamagePropagationFactor, float ShockDamagePropagationFactor)
+void FGeometryCollectionPhysicsProxy::SetDamageThresholds_Internal(const TArray<float>& DamageThresholds)
 {
-	check(IsInGameThread());
-	FGeometryCollectioPerFrameData::FDamagePropagationData Data;
-	Data.bEnabled = bEnabled;
-	Data.BreakDamagePropagationFactor = BreakDamagePropagationFactor;
-	Data.ShockDamagePropagationFactor = ShockDamagePropagationFactor;
-	GameThreadPerFrameData.DamagePropagationData = Data;
-	SetProxyDirty_External();
+	Parameters.DamageThreshold = DamageThresholds;
+
+	UpdateDamageThreshold_Internal();
 }
 
-void FGeometryCollectionPhysicsProxy::SetDamageModel_External(EDamageModelTypeEnum DamageModel)
+void FGeometryCollectionPhysicsProxy::SetDamagePropagationData_External(bool bEnabled, float BreakDamagePropagationFactor, float ShockDamagePropagationFactor)
 {
-	check(IsInGameThread());
-	GameThreadPerFrameData.DamageModel = DamageModel;
-	SetProxyDirty_External();
+	ExecuteOnPhysicsThread(*this,
+		[this, bEnabled, BreakDamagePropagationFactor, ShockDamagePropagationFactor]()
+		{
+			SetDamagePropagationData_Internal(bEnabled, BreakDamagePropagationFactor, ShockDamagePropagationFactor);
+		});
 }
 
-void FGeometryCollectionPhysicsProxy::SetUseMaterialDamageModifiers_External(bool bUseMaterialDamageModifiers)
+void FGeometryCollectionPhysicsProxy::SetDamagePropagationData_Internal(bool bEnabled, float BreakDamagePropagationFactor, float ShockDamagePropagationFactor)
 {
-	check(IsInGameThread());
-	GameThreadPerFrameData.bUseMaterialDamageModifiers = bUseMaterialDamageModifiers;
-	SetProxyDirty_External();
+	Parameters.bUseDamagePropagation = bEnabled;
+	Parameters.BreakDamagePropagationFactor = BreakDamagePropagationFactor;
+	Parameters.ShockDamagePropagationFactor = ShockDamagePropagationFactor;
 }
 
-void FGeometryCollectionPhysicsProxy::SetMaterialOverrideMassScaleMultiplier_External(float InMultiplier)
+void FGeometryCollectionPhysicsProxy::SetDamageModel_External(EDamageModelTypeEnum DamageModel)
 {
-	check(IsInGameThread());
-	GameThreadPerFrameData.MaterialOverrideMassScaleMultiplier = InMultiplier;
-	SetProxyDirty_External();
+	ExecuteOnPhysicsThread(*this,
+		[this, DamageModel]()
+		{
+			SetDamageModel_Internal(DamageModel);
+		});
 }
 
-void FGeometryCollectionPhysicsProxy::SetGravityGroupIndex_External(int32 GravityGroupIndex)
+void FGeometryCollectionPhysicsProxy::SetDamageModel_Internal(EDamageModelTypeEnum DamageModel)
 {
-	check(IsInGameThread());
-	GameThreadPerFrameData.GravityGroupIndex = static_cast<uint8>(GravityGroupIndex); // we actually only support a handful gravity groups 
-	SetProxyDirty_External();
+	Parameters.DamageModel = DamageModel;
+	Parameters.DamageEvaluationModel = GetDamageEvaluationModel(Parameters.DamageModel);
+
+	UpdateDamageThreshold_Internal();
 }
 
-void FGeometryCollectionPhysicsProxy::SetIsOneWayInteraction_External(bool bIsOneWayInteraction)
+void FGeometryCollectionPhysicsProxy::SetUseMaterialDamageModifiers_External(bool bUseMaterialDamageModifiers)
 {
-	check(IsInGameThread());
-	GameThreadPerFrameData.bIsOneWayInteraction = bIsOneWayInteraction;
-	SetProxyDirty_External();
+	ExecuteOnPhysicsThread(*this,
+		[this, bUseMaterialDamageModifiers]()
+		{
+			SetUseMaterialDamageModifiers_Internal(bUseMaterialDamageModifiers);
+		});
 }
 
-void FGeometryCollectionPhysicsProxy::PushStateOnGameThread(Chaos::FPBDRigidsSolver* InSolver)
+void FGeometryCollectionPhysicsProxy::SetUseMaterialDamageModifiers_Internal(bool bUseMaterialDamageModifiers)
 {
-	// rest all the dirty flags
-	bIsPhysicsThreadWorldTransformDirty = false;
-	bIsCollisionFilterDataDirty = false;
-	bIsDamageThresholdDataDirty = false;
-	bIsGravityGroupIndexDirty = false;
-	bIsOneWayInteractionDirty = false;
-	MaterialOverrideMassScaleMultiplierChange = 0;
-
-	// CONTEXT: GAMETHREAD
-	// this is running on GAMETHREAD before the PhysicsThread code runs for this frame
-	bIsPhysicsThreadWorldTransformDirty = GameThreadPerFrameData.GetIsWorldTransformDirty();
-	if (bIsPhysicsThreadWorldTransformDirty)
-	{
-		Parameters.PrevWorldTransform = Parameters.WorldTransform;
-		Parameters.WorldTransform = GameThreadPerFrameData.GetWorldTransform();
-		GameThreadPerFrameData.ResetIsWorldTransformDirty();
-	}
+	Parameters.bUseMaterialDamageModifiers = bUseMaterialDamageModifiers;
 
-	bIsCollisionFilterDataDirty = GameThreadPerFrameData.GetIsCollisionFilterDataDirty();
-	if (bIsCollisionFilterDataDirty)
-	{
-		Parameters.QueryFilterData = GameThreadPerFrameData.GetQueryFilter();
-		Parameters.SimulationFilterData = GameThreadPerFrameData.GetSimFilter();
-		GameThreadPerFrameData.ResetIsCollisionFilterDataDirty();
-	}
-
-	if (GameThreadPerFrameData.GetIsNotificationDataDirty())
-	{
-		Parameters.bGenerateBreakingData = GameThreadPerFrameData.GetNotifyBreakings();
-		Parameters.bGenerateCrumblingData = GameThreadPerFrameData.GetNotifyCrumblings();
-		Parameters.bGenerateCrumblingChildrenData = GameThreadPerFrameData.GetCrumblingEventIncludesChildren();
-		Parameters.bGenerateGlobalBreakingData = GameThreadPerFrameData.GetNotifyGlobalBreakings();
-		Parameters.bGenerateGlobalCrumblingData = GameThreadPerFrameData.GetNotifyGlobalCrumblings();
-		Parameters.bGenerateGlobalCrumblingChildrenData = GameThreadPerFrameData.GetGlobalCrumblingEventIncludesChildren();
-		GameThreadPerFrameData.ResetIsNotificationDataDirty();
-	}
-
-	if (GameThreadPerFrameData.GetIsDamageSettingsDataDirty())
-	{
-		Parameters.bEnableStrainOnCollision = GameThreadPerFrameData.GetEnableStrainOnCollision();
-		GameThreadPerFrameData.ResetIsDamageSettingsDataDirty();
-	}
+	UpdateDamageThreshold_Internal();
+}
 
-	if (GameThreadPerFrameData.GetIsDamageThresholdDataDirty())
-	{
-		Parameters.DamageThreshold = GameThreadPerFrameData.GetDamageThresholds();
-		GameThreadPerFrameData.ResetIsDamageThresholdDataDirty();
-		bIsDamageThresholdDataDirty = true;
-	}
+void FGeometryCollectionPhysicsProxy::SetMaterialOverrideMassScaleMultiplier_External(float InMultiplier)
+{
+	ExecuteOnPhysicsThread(*this,
+		[this, InMultiplier]()
+		{
+			SetMaterialOverrideMassScaleMultiplier_Internal(InMultiplier);
+		});
+}
 
-	if (GameThreadPerFrameData.DamageModel.IsSet())
-	{
-		Parameters.DamageModel = GameThreadPerFrameData.DamageModel.GetValue();
-		Parameters.DamageEvaluationModel = GetDamageEvaluationModel(Parameters.DamageModel);
-		GameThreadPerFrameData.DamageModel.Reset();
-		bIsDamageThresholdDataDirty = true;
-	}
+void FGeometryCollectionPhysicsProxy::SetMaterialOverrideMassScaleMultiplier_Internal(float InMultiplier)
+{
+	const float NewValue = InMultiplier;
+	const float OldValue = Parameters.MaterialOverrideMassScaleMultiplier;
 
-	if (GameThreadPerFrameData.bUseMaterialDamageModifiers.IsSet())
+	// Because we need to send a change in scale , we need to make sure the physics state has been created
+	// otherwise we set the change to 1.0 and let the PT pick up the right Parameters.MaterialOverrideMassScaleMultiplier during InitializeBodiesPT
+	float MaterialOverrideMassScaleMultiplierChange = 1.0;
+	if (bIsInitializedOnPhysicsThread && OldValue > SMALL_NUMBER)
 	{
-		Parameters.bUseMaterialDamageModifiers = GameThreadPerFrameData.bUseMaterialDamageModifiers.GetValue();
-		GameThreadPerFrameData.bUseMaterialDamageModifiers.Reset();
-		bIsDamageThresholdDataDirty = true;
+		MaterialOverrideMassScaleMultiplierChange = NewValue / OldValue;
 	}
+	Parameters.MaterialOverrideMassScaleMultiplier = NewValue;
 
-	if (GameThreadPerFrameData.MaterialOverrideMassScaleMultiplier.IsSet())
+	for (int32 TransformIndex = 0; TransformIndex < SolverParticleHandles.Num(); ++TransformIndex)
 	{
-		const float NewValue = GameThreadPerFrameData.MaterialOverrideMassScaleMultiplier.GetValue();
-
-		// Because we need to send a change in scale , we need to make sure the physics state has been created
-		// otherwise we set the change to 1.0 and let the PT pick up the right Parameters.MaterialOverrideMassScaleMultiplier during InitializeBodiesPT
-		MaterialOverrideMassScaleMultiplierChange = 1.0;
-		if (bIsInitializedOnPhysicsThread)
+		if (Chaos::FPBDRigidClusteredParticleHandle* Handle = SolverParticleHandles[TransformIndex])
 		{
-			MaterialOverrideMassScaleMultiplierChange = NewValue / Parameters.MaterialOverrideMassScaleMultiplier;
-		}
-
-		Parameters.MaterialOverrideMassScaleMultiplier = NewValue;
-		GameThreadPerFrameData.MaterialOverrideMassScaleMultiplier.Reset();
-	}
+			const Chaos::FReal NewM = Handle->M() * MaterialOverrideMassScaleMultiplierChange;
+			const Chaos::FVec3 NewI = Handle->I() * MaterialOverrideMassScaleMultiplierChange;
 
-	if (GameThreadPerFrameData.GravityGroupIndex.IsSet())
-	{
-		Parameters.GravityGroupIndex = GameThreadPerFrameData.GravityGroupIndex.GetValue();
-		GameThreadPerFrameData.GravityGroupIndex.Reset();
-		bIsGravityGroupIndexDirty = true;
-	}
-
-	if (GameThreadPerFrameData.bIsOneWayInteraction.IsSet())
-	{
-		Parameters.bIsOneWayInteraction = GameThreadPerFrameData.bIsOneWayInteraction.GetValue();
-		GameThreadPerFrameData.bIsOneWayInteraction.Reset();
-		bIsOneWayInteractionDirty = true;
+			Handle->SetM(NewM);
+			Handle->SetI(NewI);
+			const Chaos::FReal InvM = (NewM > 0.0f) ? 1.0f / NewM : 0.0f;
+			const Chaos::FVec3 InvI = (NewM > 0.0f) ? Chaos::FVec3(NewI).Reciprocal() : Chaos::FVec3::ZeroVector;
+			Handle->SetInvM(InvM);
+			Handle->SetInvI(InvI);
+		}
 	}
+}
 
-	if (GameThreadPerFrameData.DamagePropagationData.IsSet())
-	{
-		const FGeometryCollectioPerFrameData::FDamagePropagationData& Data = GameThreadPerFrameData.DamagePropagationData.GetValue();
-		Parameters.bUseDamagePropagation = Data.bEnabled;
-		Parameters.BreakDamagePropagationFactor = Data.BreakDamagePropagationFactor;
-		Parameters.ShockDamagePropagationFactor = Data.ShockDamagePropagationFactor;
-		GameThreadPerFrameData.DamagePropagationData.Reset();
-	}
+void FGeometryCollectionPhysicsProxy::SetGravityGroupIndex_External(int32 GravityGroupIndex)
+{
+	ExecuteOnPhysicsThread(*this, 
+		[this, GravityGroupIndex]() 
+		{ 
+			SetGravityGroupIndex_Internal(GravityGroupIndex); 
+		});
 }
 
-void FGeometryCollectionPhysicsProxy::PushToPhysicsState()
+void FGeometryCollectionPhysicsProxy::SetGravityGroupIndex_Internal(int32 GravityGroupIndex)
 {
-	using namespace Chaos;
-	// CONTEXT: PHYSICSTHREAD
+	Parameters.GravityGroupIndex = GravityGroupIndex;
 
-	TArray<FClusterHandle*> ProcessedInternalClusters;
-	if (bIsCollisionFilterDataDirty)
+	for (Chaos::FPBDRigidClusteredParticleHandle* Handle : SolverParticleHandles)
 	{
-		if (Chaos::FPhysicsSolver* RigidSolver = GetSolver<Chaos::FPhysicsSolver>())
+		if (Handle)
 		{
-			FRigidClustering& RigidClustering = RigidSolver->GetEvolution()->GetRigidClustering();
-
-			const int32 NumTransforms = PhysicsThreadCollection.NumElements(FGeometryCollection::TransformGroup);
-			for (int32 TransformIndex = 0; TransformIndex < NumTransforms; ++TransformIndex)
-			{
-				if (Chaos::FPBDRigidClusteredParticleHandle* Handle = SolverParticleHandles[TransformIndex])
-				{
-					// Must update our filters before updating an internal cluster parent
-					const Chaos::FShapesArray& ShapesArray = Handle->ShapesArray();
-					for (const TUniquePtr<Chaos::FPerShapeData>& Shape : ShapesArray)
-					{
-						Shape->SetQueryData(Parameters.QueryFilterData);
-						Shape->SetSimData(Parameters.SimulationFilterData);
-					}
-
-					FClusterHandle* ParentHandle = Handle->Parent();
-					if (ParentHandle && ParentHandle->InternalCluster())
-					{
-						if (!ProcessedInternalClusters.Contains(ParentHandle))
-						{
-							ProcessedInternalClusters.Add(ParentHandle);
-							// Must update our filters before updating an internal cluster parent
-							FRigidClustering::FRigidHandleArray* ChildArray = RigidClustering.GetChildrenMap().Find(ParentHandle);
-							if (ensure(ChildArray))
-							{
-								// If our filter changed, internal cluster parent's filter may be stale.
-								UpdateClusterFilterDataFromChildren(ParentHandle, *ChildArray);
-							}
-						}
-					}
-				}
-			}
+			Handle->SetGravityGroupIndex(Parameters.GravityGroupIndex);
 		}
 	}
+}
 
-	if (bIsPhysicsThreadWorldTransformDirty)
-	{
-		ProcessedInternalClusters.Reset();
-
-		const FTransform& ActorToWorld = Parameters.WorldTransform;
-
-		const int32 NumTransforms = PhysicsThreadCollection.NumElements(FGeometryCollection::TransformGroup);
-		if (Chaos::FPhysicsSolver* RigidSolver = GetSolver<Chaos::FPhysicsSolver>())
+void FGeometryCollectionPhysicsProxy::SetIsOneWayInteraction_External(bool bIsOneWayInteraction)
+{
+	ExecuteOnPhysicsThread(*this, [this, bIsOneWayInteraction]()
 		{
-			FClusterUnionManager& ClusterUnionManager = RigidSolver->GetEvolution()->GetRigidClustering().GetClusterUnionManager();
-
-			// Assume all particles are in the same cluster union.
-			FClusterUnionIndex ClusterUnionIndex = INDEX_NONE;
-			TArray<FPBDRigidParticleHandle*> DeferredClusterUnionParticleUpdates;
-			TArray<FTransform> DeferredClusterUnionChildToParentUpdates;
-
-			const TManagedArray<FTransform>& MassToLocal = Parameters.RestCollection->GetAttribute<FTransform>(MassToLocalAttributeName, FTransformCollection::TransformGroup);
-
-			for (int32 TransformGroupIndex = 0; TransformGroupIndex < NumTransforms; ++TransformGroupIndex)
-			{
-				if (Chaos::FPBDRigidClusteredParticleHandle* Handle = SolverParticleHandles[TransformGroupIndex])
-				{
-					FClusterHandle* KinematicRootHandle = nullptr;
-					FClusterHandle* ParentHandle = Handle->Parent();
-
-					if (!Handle->Disabled() && Handle->ObjectState() == Chaos::EObjectStateType::Kinematic)
-					{
-						KinematicRootHandle = Handle;
-					}
-					else
-					{
-						// is there a internal parent as a kinematic root?
-						if (ParentHandle && ParentHandle->InternalCluster() && !ParentHandle->Disabled() && ParentHandle->ObjectState() == Chaos::EObjectStateType::Kinematic && ParentHandle->PhysicsProxy() == this)
-						{
-							if (!ProcessedInternalClusters.Contains(ParentHandle))
-							{
-								ProcessedInternalClusters.Add(ParentHandle);
-								KinematicRootHandle = ParentHandle;
-							}
-						}
-					}
-
-					if (KinematicRootHandle)
-					{
-						const FTransform RootWorldTransform(KinematicRootHandle->R(), KinematicRootHandle->X());
-						const FTransform RootRelativeTransform = RootWorldTransform.GetRelativeTransform(Parameters.PrevWorldTransform);
-						const FTransform WorldTransform = RootRelativeTransform * ActorToWorld;
-
-						SetClusteredParticleKinematicTarget_Internal(KinematicRootHandle, WorldTransform);
-					}
-					else if (ParentHandle && !ParentHandle->IsDynamic())
-					{
-						if (ClusterUnionIndex == INDEX_NONE)
-						{
-							ClusterUnionIndex = ClusterUnionManager.FindClusterUnionIndexFromParticle(Handle);
-						}
-
-						if (ClusterUnionIndex != INDEX_NONE)
-						{
-							const FTransform ParentWorldTransform{ ParentHandle->R(), ParentHandle->X() };
-							const FTransform NewWorldTransform = MassToLocal[TransformGroupIndex] * FTransform(PhysicsThreadCollection.GetTransform(TransformGroupIndex)) * Parameters.WorldTransform;
-							const FTransform RelativeTransform = NewWorldTransform.GetRelativeTransform(ParentWorldTransform);
-
-							DeferredClusterUnionParticleUpdates.Add(Handle);
-							DeferredClusterUnionChildToParentUpdates.Add(RelativeTransform);
-
-							// Make sure the particle is mark dirty to make sure its proxy will properly update the transforms
-							RigidSolver->GetEvolution()->GetParticles().MarkTransientDirtyParticle(Handle);
-						}
-					}
-				}
-			}
-
-			if (ClusterUnionIndex != INDEX_NONE && !DeferredClusterUnionParticleUpdates.IsEmpty() && !DeferredClusterUnionChildToParentUpdates.IsEmpty())
-			{
-				ClusterUnionManager.UpdateClusterUnionParticlesChildToParent(ClusterUnionIndex, DeferredClusterUnionParticleUpdates, DeferredClusterUnionChildToParentUpdates, false);
-			}
-		}
-	}
+			SetIsOneWayInteraction_Internal(bIsOneWayInteraction);
+		});
+}
 
-	if (bIsDamageThresholdDataDirty)
-	{
-		UpdateDamageThreshold_Internal();
-	}
+void FGeometryCollectionPhysicsProxy::SetIsOneWayInteraction_Internal(bool bInIsOneWayInteraction)
+{
+	Parameters.bIsOneWayInteraction = bInIsOneWayInteraction;
 
-	if (bIsGravityGroupIndexDirty)
+	for (Chaos::FPBDRigidClusteredParticleHandle* Handle : SolverParticleHandles)
 	{
-		for (int32 TransformIndex = 0; TransformIndex < SolverParticleHandles.Num(); ++TransformIndex)
+		if (Handle)
 		{
-			if (Chaos::FPBDRigidClusteredParticleHandle* Handle = SolverParticleHandles[TransformIndex])
-			{
-				Handle->SetGravityGroupIndex(Parameters.GravityGroupIndex);
-			}
+			Handle->SetOneWayInteraction(Parameters.bIsOneWayInteraction);
 		}
 	}
+}
 
-	if (bIsOneWayInteractionDirty)
-	{
-		for (int32 TransformIndex = 0; TransformIndex < SolverParticleHandles.Num(); ++TransformIndex)
-		{
-			if (Chaos::FPBDRigidClusteredParticleHandle* Handle = SolverParticleHandles[TransformIndex])
-			{
-				Handle->SetOneWayInteraction(Parameters.bIsOneWayInteraction);
-			}
-		}
-	}
+void FGeometryCollectionPhysicsProxy::PushStateOnGameThread(Chaos::FPBDRigidsSolver* InSolver)
+{
+	// CONTEXT: ANYTHREAD but spawned form GAMETHREAD in a parallelFor
+	// nothing to do 
+}
 
-	if (MaterialOverrideMassScaleMultiplierChange > 0)
-	{
-		for (int32 TransformIndex = 0; TransformIndex < SolverParticleHandles.Num(); ++TransformIndex)
-		{
-			if (Chaos::FPBDRigidClusteredParticleHandle* Handle = SolverParticleHandles[TransformIndex])
-			{
-				const Chaos::FReal NewM = Handle->M() * MaterialOverrideMassScaleMultiplierChange;
-				const Chaos::FVec3 NewI = Handle->I() * MaterialOverrideMassScaleMultiplierChange;
-
-				Handle->SetM(NewM);
-				Handle->SetI(NewI);
-				const Chaos::FReal InvM = (NewM > 0.0f) ? 1.0f / NewM : 0.0f;
-				const Chaos::FVec3 InvI = (NewM > 0.0f) ? Chaos::FVec3(NewI).Reciprocal() : Chaos::FVec3::ZeroVector;
-				Handle->SetInvM(InvM);
-				Handle->SetInvI(InvI);
-			}
-		}
-	}
+void FGeometryCollectionPhysicsProxy::PushToPhysicsState()
+{
+	using namespace Chaos;
+	// CONTEXT: PHYSICSTHREAD
 }
 
 void FGeometryCollectionPhysicsProxy::SetClusteredParticleKinematicTarget_Internal(Chaos::FPBDRigidClusteredParticleHandle* Handle, const FTransform& NewWorldTransform)
@@ -4021,8 +3958,8 @@ bool FGeometryCollectionPhysicsProxy::PullFromPhysicsState(const Chaos::FDirtyGe
 		{
 			const TBitArray<>& PrevResultsModifiedIndices = PullData.Results().GetModifiedTransformIndices();
 
-			const bool bIsComponentTransformScaled = !Parameters.WorldTransform.GetScale3D().Equals(FVector::OneVector);
-			const FTransform ComponentScaleTransform(FQuat::Identity, FVector::ZeroVector, Parameters.WorldTransform.GetScale3D());
+			const bool bIsComponentTransformScaled = !WorldTransform_External.GetScale3D().Equals(FVector::OneVector);
+			const FTransform ComponentScaleTransform(FQuat::Identity, FVector::ZeroVector, WorldTransform_External.GetScale3D());
 
 			for (int32 TransformIndex = 0; TransformIndex < NumTransforms; TransformIndex++)
 			{
@@ -4047,7 +3984,7 @@ bool FGeometryCollectionPhysicsProxy::PullFromPhysicsState(const Chaos::FDirtyGe
 							const FTransform& WorldTransform = ParticleMassToLocal.Inverse() * FTransform { GTParticle.R(), GTParticle.X() };
 
 							// by default parent is the component's world transform 
-							FTransform ParentWorldTransform = Parameters.WorldTransform;
+							FTransform ParentWorldTransform = WorldTransform_External;
 							if (ParentTransformIndex != INDEX_NONE)
 							{
 								if (const FParticle* GTParentParticle = (ParentTransformIndex != INDEX_NONE) ? GTParticles[ParentTransformIndex].Get() : nullptr)
@@ -4090,11 +4027,11 @@ bool FGeometryCollectionPhysicsProxy::PullFromPhysicsState(const Chaos::FDirtyGe
 		// broken sleeping particles on kinemically driven GCs remained updated and the object feel grounded to the world
 
 		GameThreadCollection.MakeClean();
-		if (bIsCollectionDirty || bIsPhysicsThreadWorldTransformDirty)
+		if (bIsCollectionDirty || bIsGameThreadWorldTransformDirty)
 		{
 			GameThreadCollection.MakeDirty();
 		}
-		bIsPhysicsThreadWorldTransformDirty = false;
+		bIsGameThreadWorldTransformDirty = false;
 	}
 
 	if (PostPhysicsSyncCallback)
@@ -4128,54 +4065,102 @@ void FGeometryCollectionPhysicsProxy::UpdateFilterData_External(const FCollision
 		}
 	}
 
-	GameThreadPerFrameData.SetSimFilter(NewSimFilter);
-	GameThreadPerFrameData.SetQueryFilter(NewQueryFilter);
+	ExecuteOnPhysicsThread(*this, [this, NewSimFilter, NewQueryFilter]()
+		{
+			SetFilterData_Internal(NewSimFilter, NewQueryFilter);
+		});
+}
+
+void FGeometryCollectionPhysicsProxy::SetFilterData_Internal(const FCollisionFilterData& NewSimFilter, const FCollisionFilterData& NewQueryFilter)
+{
+	using namespace Chaos;
 
-	SetProxyDirty_External();
+	Parameters.SimulationFilterData = NewSimFilter;
+	Parameters.QueryFilterData = NewQueryFilter;
+
+	TArray<FClusterHandle*> ProcessedInternalClusters;
+	if (Chaos::FPhysicsSolver* RigidSolver = GetSolver<Chaos::FPhysicsSolver>())
+	{
+		FRigidClustering& RigidClustering = RigidSolver->GetEvolution()->GetRigidClustering();
+
+		const int32 NumTransforms = PhysicsThreadCollection.NumElements(FGeometryCollection::TransformGroup);
+		for (int32 TransformIndex = 0; TransformIndex < NumTransforms; ++TransformIndex)
+		{
+			if (Chaos::FPBDRigidClusteredParticleHandle* Handle = SolverParticleHandles[TransformIndex])
+			{
+				// Must update our filters before updating an internal cluster parent
+				const Chaos::FShapesArray& ShapesArray = Handle->ShapesArray();
+				for (const TUniquePtr<Chaos::FPerShapeData>& Shape : ShapesArray)
+				{
+					Shape->SetQueryData(Parameters.QueryFilterData);
+					Shape->SetSimData(Parameters.SimulationFilterData);
+				}
+
+				FClusterHandle* ParentHandle = Handle->Parent();
+				if (ParentHandle && ParentHandle->InternalCluster())
+				{
+					if (!ProcessedInternalClusters.Contains(ParentHandle))
+					{
+						ProcessedInternalClusters.Add(ParentHandle);
+						// Must update our filters before updating an internal cluster parent
+						FRigidClustering::FRigidHandleArray* ChildArray = RigidClustering.GetChildrenMap().Find(ParentHandle);
+						if (ensure(ChildArray))
+						{
+							// If our filter changed, internal cluster parent's filter may be stale.
+							UpdateClusterFilterDataFromChildren(ParentHandle, *ChildArray);
+						}
+					}
+				}
+			}
+		}
+	}
 }
 
 DECLARE_CYCLE_STAT(TEXT("FGeometryCollectionPhysicsProxy::UpdatePerParticleFilterData_External"), STAT_GeometryCollectionPhysicsProxyUpdatePerParticleFilterData, STATGROUP_Chaos);
 void FGeometryCollectionPhysicsProxy::UpdatePerParticleFilterData_External(const TArray<FParticleCollisionFilterData>& PerParticleData)
 {
 	SCOPE_CYCLE_COUNTER(STAT_GeometryCollectionPhysicsProxyUpdatePerParticleFilterData);
-	check(IsInGameThread());
 
-	// TODO: Go through the PushToPhysicsState flow instead?
 	// TODO: Need to figure out how the per-particle collision filter data works with the global one. The per-particle one should probably replace the global one entirely...
+	ExecuteOnPhysicsThread(*this, [this, PerParticleData]()
+		{
+			SetPerParticleFilterData_Internal(PerParticleData);
+		});
+}
+
+void FGeometryCollectionPhysicsProxy::SetPerParticleFilterData_Internal(const TArray<FParticleCollisionFilterData>& PerParticleData)
+{
 	if (Chaos::FPhysicsSolver* RBDSolver = GetSolver<Chaos::FPhysicsSolver>())
 	{
-		RBDSolver->EnqueueCommandImmediate([this, RBDSolver, PerParticleData]()
-		{
-			Chaos::FWritePhysicsObjectInterface_Internal Interface= Chaos::FPhysicsObjectInternalInterface::GetWrite();
+		Chaos::FWritePhysicsObjectInterface_Internal Interface = Chaos::FPhysicsObjectInternalInterface::GetWrite();
 
-			for (const FParticleCollisionFilterData& Data : PerParticleData)
+		for (const FParticleCollisionFilterData& Data : PerParticleData)
+		{
+			if (Data.bIsValid && Data.ParticleIndex != INDEX_NONE)
 			{
-				if (Data.bIsValid && Data.ParticleIndex != INDEX_NONE)
+				if (Chaos::FPhysicsObjectHandle Object = PhysicsObjects[Data.ParticleIndex].Get())
 				{
-					if (Chaos::FPhysicsObjectHandle Object = PhysicsObjects[Data.ParticleIndex].Get())
-					{
-						TArrayView<Chaos::FPhysicsObjectHandle> ParticleView{ &Object, 1 };
-						Interface.UpdateShapeCollisionFlags(ParticleView, Data.bSimEnabled, Data.bQueryEnabled);
-						Interface.UpdateShapeFilterData(ParticleView, Data.QueryFilter, Data.SimFilter);
-					}
+					TArrayView<Chaos::FPhysicsObjectHandle> ParticleView{ &Object, 1 };
+					Interface.UpdateShapeCollisionFlags(ParticleView, Data.bSimEnabled, Data.bQueryEnabled);
+					Interface.UpdateShapeFilterData(ParticleView, Data.QueryFilter, Data.SimFilter);
+				}
 
-					// todo(chaos): It's not ideal but the geometry collection needs to request the cluster union update its
-					// cached shape data if the particle is in a cluster union. This is because we don't share shape
-					// data between the GC shapes and the cluster union shapes.
-					Chaos::FRigidClustering& Clustering = RBDSolver->GetEvolution()->GetRigidClustering();
-					Chaos::FClusterUnionManager& ClusterUnionManager = Clustering.GetClusterUnionManager();
+				// todo(chaos): It's not ideal but the geometry collection needs to request the cluster union update its
+				// cached shape data if the particle is in a cluster union. This is because we don't share shape
+				// data between the GC shapes and the cluster union shapes.
+				Chaos::FRigidClustering& Clustering = RBDSolver->GetEvolution()->GetRigidClustering();
+				Chaos::FClusterUnionManager& ClusterUnionManager = Clustering.GetClusterUnionManager();
 
-					if (Chaos::FPBDRigidClusteredParticleHandle* Particle = SolverParticleHandles[Data.ParticleIndex])
+				if (Chaos::FPBDRigidClusteredParticleHandle* Particle = SolverParticleHandles[Data.ParticleIndex])
+				{
+					if (Chaos::FClusterUnion* ClusterUnion = ClusterUnionManager.FindClusterUnionFromParticle(Particle))
 					{
-						if (Chaos::FClusterUnion* ClusterUnion = ClusterUnionManager.FindClusterUnionFromParticle(Particle))
-						{
-							ClusterUnion->AddPendingGeometryOperation(Chaos::EClusterUnionGeometryOperation::Refresh, Particle);
-							ClusterUnionManager.RequestDeferredClusterPropertiesUpdate(ClusterUnion->InternalIndex, Chaos::EUpdateClusterUnionPropertiesFlags::IncrementalGenerateGeometry);
-						}
+						ClusterUnion->AddPendingGeometryOperation(Chaos::EClusterUnionGeometryOperation::Refresh, Particle);
+						ClusterUnionManager.RequestDeferredClusterPropertiesUpdate(ClusterUnion->InternalIndex, Chaos::EUpdateClusterUnionPropertiesFlags::IncrementalGenerateGeometry);
 					}
 				}
 			}
-		});
+		}
 	}
 }
 
diff --git a/Engine/Source/Runtime/Experimental/Chaos/Public/GeometryCollectionProxyData.h b/Engine/Source/Runtime/Experimental/Chaos/Public/GeometryCollectionProxyData.h
--- a/Engine/Source/Runtime/Experimental/Chaos/Public/GeometryCollectionProxyData.h
+++ b/Engine/Source/Runtime/Experimental/Chaos/Public/GeometryCollectionProxyData.h
@@ -216,197 +216,6 @@ private:
 	FGeometryDynamicCollection& DynamicCollection;
 };
 
-class FGeometryCollectioPerFrameData
-{
-public:
-	struct FDamagePropagationData
-	{
-		bool bEnabled = false;
-		float BreakDamagePropagationFactor = 0.0f;
-		float ShockDamagePropagationFactor = 0.0f;
-	};
-
-public:
-	FGeometryCollectioPerFrameData()
-		: bIsWorldTransformDirty(false)
-		, bIsCollisionFilterDataDirty(false)
-		, bIsNotificationDataDirty(false) 
-		, bIsDamageSettingsDataDirty(false)
-		, bIsDamageThresholdDataDirty(false)
-		, bNotifyBreakings(false)
-		, bNotifyRemovals(false)
-		, bNotifyCrumblings(false)
-		, bCrumblingEventIncludesChildren(false)
-		, bNotifyGlobalBreakings(false)
-		, bNotifyGlobalRemovals(false)
-		, bNotifyGlobalCrumblings(false)
-		, bGlobalCrumblingEventIncludesChildren(false)
-		, bEnableStrainOnCollision(false)
-	{}
-
-	FGeometryCollectioPerFrameData(const FSimulationParameters& Parameters)
-		: bIsWorldTransformDirty(false)
-		, bIsCollisionFilterDataDirty(false)
-		, bIsNotificationDataDirty(false)
-		, bIsDamageSettingsDataDirty(false)
-		, bIsDamageThresholdDataDirty(false)
-		, bNotifyBreakings(Parameters.bGenerateBreakingData)
-		, bNotifyRemovals(false)
-		, bNotifyCrumblings(Parameters.bGenerateCrumblingData)
-		, bCrumblingEventIncludesChildren(Parameters.bGenerateCrumblingChildrenData)
-		, bNotifyGlobalBreakings(Parameters.bGenerateGlobalBreakingData)
-		, bNotifyGlobalRemovals(false)
-		, bNotifyGlobalCrumblings(Parameters.bGenerateGlobalCrumblingData)
-		, bGlobalCrumblingEventIncludesChildren(Parameters.bGenerateGlobalCrumblingChildrenData)
-		, bEnableStrainOnCollision(false)
-		, DamageThresholds(Parameters.DamageThreshold)
-	{}
-
-	const FTransform& GetWorldTransform() const { return WorldTransform; }
-
-	void SetWorldTransform(const FTransform& InWorldTransform)
-	{
-		if (!WorldTransform.Equals(InWorldTransform))
-		{
-			WorldTransform = InWorldTransform;
-			bIsWorldTransformDirty = true;
-		}
-	}
-
-	bool GetIsWorldTransformDirty() const { return bIsWorldTransformDirty; }
-	void ResetIsWorldTransformDirty() { bIsWorldTransformDirty = false; }
-
-	const FCollisionFilterData& GetSimFilter() const { return SimFilter; }
-	void SetSimFilter(const FCollisionFilterData& NewSimFilter)
-	{
-		SimFilter = NewSimFilter;
-		bIsCollisionFilterDataDirty = true;
-	}
-
-	const FCollisionFilterData& GetQueryFilter() const { return QueryFilter; }
-	void SetQueryFilter(const FCollisionFilterData& NewQueryFilter)
-	{
-		QueryFilter = NewQueryFilter;
-		bIsCollisionFilterDataDirty = true;
-	}
-
-	bool GetIsCollisionFilterDataDirty() const { return bIsCollisionFilterDataDirty; }
-	void ResetIsCollisionFilterDataDirty() { bIsCollisionFilterDataDirty = false; }
-
-	void SetNotifyBreakings(bool bNotify)
-	{
-		bNotifyBreakings = bNotify;
-		bIsNotificationDataDirty = true;
-	}
-	bool GetNotifyBreakings() const { return bNotifyBreakings; }
-
-	void SetNotifyRemovals(bool bNotify)
-	{
-		bNotifyRemovals = bNotify;
-		bIsNotificationDataDirty = true;
-	}
-
-	bool GetNotifyRemovals() const { return bNotifyRemovals; }
-
-	void SetNotifyCrumblings(bool bNotify, bool bIncludeChildren)
-	{
-		bNotifyCrumblings = bNotify;
-		bCrumblingEventIncludesChildren = bIncludeChildren;
-		bIsNotificationDataDirty = true;
-	}
-
-	bool GetNotifyCrumblings() const { return bNotifyCrumblings; }
-	bool GetCrumblingEventIncludesChildren() const { return bCrumblingEventIncludesChildren; }
-
-	void SetNotifyGlobalBreakings(bool bNotify)
-	{
-		bNotifyGlobalBreakings = bNotify;
-		bIsNotificationDataDirty = true;
-	}
-	bool GetNotifyGlobalBreakings() const { return bNotifyGlobalBreakings; }
-
-	void SetNotifyGlobalRemovals(bool bNotify)
-	{
-		bNotifyGlobalRemovals = bNotify;
-		bIsNotificationDataDirty = true;
-	}
-	bool GetNotifyGlobalRemovals() const { return bNotifyGlobalRemovals; }
-
-	void SetNotifyGlobalCrumblings(bool bNotify, bool bIncludeChildren)
-	{
-		bNotifyGlobalCrumblings = bNotify;
-		bGlobalCrumblingEventIncludesChildren = bIncludeChildren;
-		bIsNotificationDataDirty = true;
-	}
-
-	bool GetNotifyGlobalCrumblings() const { return bNotifyGlobalCrumblings; }
-	bool GetGlobalCrumblingEventIncludesChildren() const { return bGlobalCrumblingEventIncludesChildren; }
-
-
-	bool GetIsNotificationDataDirty() const { return bIsNotificationDataDirty; }
-	void ResetIsNotificationDataDirty() { bIsNotificationDataDirty = false; }
-
-	void SetEnableStrainOnCollision(bool bEnable)
-	{
-		bEnableStrainOnCollision = bEnable;
-		bIsDamageSettingsDataDirty = true;
-	}
-
-	bool GetEnableStrainOnCollision() const { return bEnableStrainOnCollision; }
-
-	bool GetIsDamageSettingsDataDirty() const { return bIsDamageSettingsDataDirty; }
-	void ResetIsDamageSettingsDataDirty() { bIsDamageSettingsDataDirty = false; }
-
-	void SetDamageThresholds(const TArray<float>& InDamageThresholds)
-	{
-		DamageThresholds = InDamageThresholds;
-		bIsDamageThresholdDataDirty = true;
-	}
-
-	const TArray<float>& GetDamageThresholds() const { return DamageThresholds; }
-
-	bool GetIsDamageThresholdDataDirty() const { return bIsDamageThresholdDataDirty; }
-	void ResetIsDamageThresholdDataDirty() { bIsDamageThresholdDataDirty = false; }
-
-private:
-	uint16 bIsWorldTransformDirty : 1;
-	uint16 bIsCollisionFilterDataDirty : 1;
-	uint16 bIsNotificationDataDirty: 1;
-	uint16 bIsDamageSettingsDataDirty : 1;
-	uint16 bIsDamageThresholdDataDirty : 1;
-
-	/** updated when bNotificationDataDirty is set */
-	uint16 bNotifyBreakings : 1;
-	uint16 bNotifyRemovals : 1;
-	uint16 bNotifyCrumblings : 1;
-	uint16 bCrumblingEventIncludesChildren : 1;
-	uint16 bNotifyGlobalBreakings : 1;
-	uint16 bNotifyGlobalRemovals : 1;
-	uint16 bNotifyGlobalCrumblings : 1;
-	uint16 bGlobalCrumblingEventIncludesChildren : 1;
-
-	/** updated when bIsDamageSettingsDataDirty is set */
-	uint16 bEnableStrainOnCollision : 1;
-
-	/** updated when bIsDamageThresholdDataDirty is set */
-	TArray<float> DamageThresholds;
-
-	/** updated when bIsWorldTransformDirty is set */
-	FTransform WorldTransform; 
-
-	/** updated when bIsCollisionFilterDataDirty is set */
-	FCollisionFilterData SimFilter;
-	FCollisionFilterData QueryFilter;
-
-public:
-	TOptional<EDamageModelTypeEnum> DamageModel;
-	TOptional<FDamagePropagationData> DamagePropagationData;
-	TOptional<bool> bUseMaterialDamageModifiers;
-	TOptional<float> MaterialOverrideMassScaleMultiplier;
-	TOptional<uint8> GravityGroupIndex;
-	TOptional<bool> bIsOneWayInteraction;
-};
-
 /**
  * Buffer structure for communicating simulation state between game and physics
  * threads.
diff --git a/Engine/Source/Runtime/Experimental/Chaos/Public/PhysicsProxy/GeometryCollectionPhysicsProxy.h b/Engine/Source/Runtime/Experimental/Chaos/Public/PhysicsProxy/GeometryCollectionPhysicsProxy.h
--- a/Engine/Source/Runtime/Experimental/Chaos/Public/PhysicsProxy/GeometryCollectionPhysicsProxy.h
+++ b/Engine/Source/Runtime/Experimental/Chaos/Public/PhysicsProxy/GeometryCollectionPhysicsProxy.h
@@ -553,6 +553,18 @@ protected:
 
 	CHAOS_API void UpdateDamageThreshold_Internal();
 
+
+	CHAOS_API void SetWorldTransform_Internal(const FTransform& WorldTransform);
+	CHAOS_API void SetFilterData_Internal(const FCollisionFilterData& NewSimFilter, const FCollisionFilterData& NewQueryFilter);
+	CHAOS_API void SetPerParticleFilterData_Internal(const TArray<FParticleCollisionFilterData>& PerParticleData);
+	CHAOS_API void SetDamagePropagationData_Internal(bool bEnabled, float BreakDamagePropagationFactor, float ShockDamagePropagationFactor);
+	CHAOS_API void SetDamageThresholds_Internal(const TArray<float>& DamageThresholds);
+	CHAOS_API void SetDamageModel_Internal(EDamageModelTypeEnum DamageModel);
+	CHAOS_API void SetUseMaterialDamageModifiers_Internal(bool bUseMaterialDamageModifiers);
+	CHAOS_API void SetMaterialOverrideMassScaleMultiplier_Internal(float InMultiplier);
+	CHAOS_API void SetGravityGroupIndex_Internal(int32 GravityGroupIndex);
+	CHAOS_API void SetIsOneWayInteraction_Internal(bool bInIsOneWayInteraction);
+
 private:
 
 	static TBitArray<> CalculateClustersToCreateFromChildren(const FGeometryDynamicCollection& DynamicCollection, int32 NumTransforms);
@@ -656,15 +668,9 @@ private:
 	FGeometryDynamicCollection PhysicsThreadCollection;
 	FGeometryDynamicCollection& GameThreadCollection;
 
-	// this data flows from Game thread to physics thread
-	// todo: replace this with proper structure copies between the two threads 
-	FGeometryCollectioPerFrameData GameThreadPerFrameData;
-	float MaterialOverrideMassScaleMultiplierChange;
-	uint8 bIsPhysicsThreadWorldTransformDirty : 1;
-	uint8 bIsCollisionFilterDataDirty: 1;
-	uint8 bIsDamageThresholdDataDirty: 1;
-	uint8 bIsGravityGroupIndexDirty: 1;
-	uint8 bIsOneWayInteractionDirty : 1;
+	// todo : we should probably keep a simulation parameter copy on the game thread instead 
+	FTransform WorldTransform_External;
+	uint8 bIsGameThreadWorldTransformDirty : 1;
 
 	// Currently this is using triple buffers for game-physics and 
 	// physics-game thread communication, but not for any reason other than this 
