Stop geometry collection components from spending time updating their body instance inherited from primitive component when updating transforms. The geometry collection component controls its own physics transforms internally and doesn't require the body instance.

diff --git a/Engine/Source/Runtime/Experimental/GeometryCollectionEngine/Private/GeometryCollection/GeometryCollectionComponent.cpp b/Engine/Source/Runtime/Experimental/GeometryCollectionEngine/Private/GeometryCollection/GeometryCollectionComponent.cpp
--- a/Engine/Source/Runtime/Experimental/GeometryCollectionEngine/Private/GeometryCollection/GeometryCollectionComponent.cpp
+++ b/Engine/Source/Runtime/Experimental/GeometryCollectionEngine/Private/GeometryCollection/GeometryCollectionComponent.cpp
@@ -3167,7 +3167,9 @@ FGeometryCollectionDynamicData* UGeometryCollectionComponent::InitDynamicData(bo
 
 void UGeometryCollectionComponent::OnUpdateTransform(EUpdateTransformFlags UpdateTransformFlags, ETeleportType Teleport)
 {
-	Super::OnUpdateTransform(UpdateTransformFlags, Teleport);
+	// Don't allow the primitive component to update physics as that does processing to set a transform for
+	// UPrimitiveComponent::BodyInstance as it is unecessary for geometry collections. Out update is handled below.
+	Super::OnUpdateTransform(UpdateTransformFlags | EUpdateTransformFlags::SkipPhysicsUpdate, Teleport);
 
 	const bool bSkipPhysicsUpdate = ((UpdateTransformFlags & EUpdateTransformFlags::SkipPhysicsUpdate) == EUpdateTransformFlags::SkipPhysicsUpdate);
 	if (!bSkipPhysicsUpdate && PhysicsProxy)
