Geometry collection component : make sure physics material  can be set while physics state has been created

diff --git a/Engine/Source/Runtime/Experimental/Chaos/Private/PhysicsProxy/GeometryCollectionPhysicsProxy.cpp b/Engine/Source/Runtime/Experimental/Chaos/Private/PhysicsProxy/GeometryCollectionPhysicsProxy.cpp
--- a/Engine/Source/Runtime/Experimental/Chaos/Private/PhysicsProxy/GeometryCollectionPhysicsProxy.cpp
+++ b/Engine/Source/Runtime/Experimental/Chaos/Private/PhysicsProxy/GeometryCollectionPhysicsProxy.cpp
@@ -3263,6 +3263,64 @@ void FGeometryCollectionPhysicsProxy::SetIsOneWayInteraction_Internal(bool bInIs
 	}
 }
 
+void FGeometryCollectionPhysicsProxy::SetPhysicsMaterial_External(const Chaos::FMaterialHandle& MaterialHandle)
+{
+	// update GT particles
+	for (const TUniquePtr<FParticle>& GTParticle : GTParticles)
+	{
+		if (GTParticle)
+		{
+			for (const TUniquePtr<Chaos::FPerShapeData>& Shape : GTParticle->ShapesArray())
+			{
+				Shape->SetMaterial(MaterialHandle);
+			}
+		}
+	}
+
+	// update physics thread
+	ExecuteOnPhysicsThread(*this, [this, MaterialHandle]()
+		{
+			SetPhysicsMaterial_Internal(MaterialHandle);
+		});
+}
+
+void FGeometryCollectionPhysicsProxy::SetPhysicsMaterial_Internal(const Chaos::FMaterialHandle& MaterialHandle)
+{
+	using namespace Chaos;
+	if (Parameters.PhysicalMaterialHandle != MaterialHandle)
+	{
+		Parameters.PhysicalMaterialHandle = MaterialHandle;
+
+		// set materials on the particles
+		FChaosPhysicsMaterial* SolverMaterial = nullptr;
+		if (Chaos::FPhysicsSolver* RigidSolver = GetSolver<Chaos::FPhysicsSolver>())
+		{
+			// #BGTODO - non-updating parameters - remove lin/ang drag arrays and always query material if this stays a material parameter
+			SolverMaterial = RigidSolver->GetSimMaterials().Get(Parameters.PhysicalMaterialHandle.InnerHandle);
+		}
+		if (SolverMaterial)
+		{
+			for (Chaos::FPBDRigidClusteredParticleHandle* Handle : SolverParticleHandles)
+			{
+				if (Handle)
+				{
+					Handle->SetLinearEtherDrag(SolverMaterial->LinearEtherDrag);
+					Handle->SetAngularEtherDrag(SolverMaterial->AngularEtherDrag);
+
+					const Chaos::FShapesArray& Shapes = Handle->ShapesArray();
+					for (const TUniquePtr<Chaos::FPerShapeData>& Shape : Shapes)
+					{
+						Shape->SetMaterial(Parameters.PhysicalMaterialHandle);
+					}
+				}
+			}
+		}
+
+		// adjust damage threshold as they depends on the material DamageThresholdMultiplier or strength
+		UpdateDamageThreshold_Internal();
+	}
+}
+
 void FGeometryCollectionPhysicsProxy::PushStateOnGameThread(Chaos::FPBDRigidsSolver* InSolver)
 {
 	// CONTEXT: ANYTHREAD but spawned form GAMETHREAD in a parallelFor
diff --git a/Engine/Source/Runtime/Experimental/Chaos/Public/PhysicsProxy/GeometryCollectionPhysicsProxy.h b/Engine/Source/Runtime/Experimental/Chaos/Public/PhysicsProxy/GeometryCollectionPhysicsProxy.h
--- a/Engine/Source/Runtime/Experimental/Chaos/Public/PhysicsProxy/GeometryCollectionPhysicsProxy.h
+++ b/Engine/Source/Runtime/Experimental/Chaos/Public/PhysicsProxy/GeometryCollectionPhysicsProxy.h
@@ -474,6 +474,7 @@ public:
 	CHAOS_API void SetMaterialOverrideMassScaleMultiplier_External(float InMultiplier);
 	CHAOS_API void SetGravityGroupIndex_External(int32 GravityGroupIndex);
 	CHAOS_API void SetIsOneWayInteraction_External(bool bInIsOneWayInteraction);
+	CHAOS_API void SetPhysicsMaterial_External(const Chaos::FMaterialHandle& MaterialHandle);
 	/** 
 	 * Traverses the parents of TransformGroupIdx counting number of levels,
 	 * and sets levels array value for TransformGroupIdx and its parents if not yet initialized.
@@ -564,6 +565,7 @@ protected:
 	CHAOS_API void SetMaterialOverrideMassScaleMultiplier_Internal(float InMultiplier);
 	CHAOS_API void SetGravityGroupIndex_Internal(int32 GravityGroupIndex);
 	CHAOS_API void SetIsOneWayInteraction_Internal(bool bInIsOneWayInteraction);
+	CHAOS_API void SetPhysicsMaterial_Internal(const Chaos::FMaterialHandle& MaterialHandle);
 
 private:
 
diff --git a/Engine/Source/Runtime/Experimental/GeometryCollectionEngine/Private/GeometryCollection/GeometryCollectionComponent.cpp b/Engine/Source/Runtime/Experimental/GeometryCollectionEngine/Private/GeometryCollection/GeometryCollectionComponent.cpp
--- a/Engine/Source/Runtime/Experimental/GeometryCollectionEngine/Private/GeometryCollection/GeometryCollectionComponent.cpp
+++ b/Engine/Source/Runtime/Experimental/GeometryCollectionEngine/Private/GeometryCollection/GeometryCollectionComponent.cpp
@@ -6841,10 +6841,13 @@ void UGeometryCollectionComponent::SetPhysMaterialOverride(UPhysicalMaterial* Ne
 {
 	Super::SetPhysMaterialOverride(NewPhysMaterial);
 
-	UPhysicalMaterial* EnginePhysicalMaterial = GetPhysicalMaterial();
-	if (ensure(EnginePhysicalMaterial) && PhysicsProxy && !PhysicsProxy->IsInitializedOnPhysicsThread())
+	if (PhysicsProxy)
 	{
-		PhysicsProxy->GetSimParameters().PhysicalMaterialHandle = EnginePhysicalMaterial->GetPhysicsMaterial();
+		UPhysicalMaterial* EnginePhysicalMaterial = GetPhysicalMaterial();
+		if (ensure(EnginePhysicalMaterial))
+		{
+			PhysicsProxy->SetPhysicsMaterial_External(EnginePhysicalMaterial->GetPhysicsMaterial());
+		}
 	}
 }
 
