Chaos - fix negative scaled collision between convex and mesh shapes

diff --git a/Engine/Source/Runtime/Experimental/Chaos/Private/Chaos/Collision/ConvexTriangleContactPoint.cpp b/Engine/Source/Runtime/Experimental/Chaos/Private/Chaos/Collision/ConvexTriangleContactPoint.cpp
--- a/Engine/Source/Runtime/Experimental/Chaos/Private/Chaos/Collision/ConvexTriangleContactPoint.cpp
+++ b/Engine/Source/Runtime/Experimental/Chaos/Private/Chaos/Collision/ConvexTriangleContactPoint.cpp
@@ -64,7 +64,7 @@ namespace Chaos
 
 		// Now clip against all planes that belong to the convex plane's, edges
 		// Note winding order matters here, and we have to handle negative scales
-		const FReal WindingOrder = Convex.GetWindingOrder();
+		const FReal ConvexWindingOrder = Convex.GetWindingOrder();
 		const int32 ConvexFaceVerticesNum = Convex.NumPlaneVertices(ConvexPlaneIndex);
 		int32 ClippingPlaneCount = ConvexFaceVerticesNum;
 		FVec3 PrevPoint = Convex.GetVertex(Convex.GetPlaneVertex(ConvexPlaneIndex, ClippingPlaneCount - 1));
@@ -74,7 +74,7 @@ namespace Chaos
 
 			// Convex edge clipping plane
 			// NOTE: Plane is not normalized, but the length cancels out in all clip operations
-			const FVec3 ClippingPlaneNormal = WindingOrder * FVec3::CrossProduct(Axis, PrevPoint - CurrentPoint);
+			const FVec3 ClippingPlaneNormal = ConvexWindingOrder * FVec3::CrossProduct(Axis, PrevPoint - CurrentPoint);
 			if (ClippingPlaneNormal.SizeSquared() > UE_SMALL_NUMBER)
 			{
 				ContactPointCount = Collisions::ClipVerticesAgainstPlane(VertexBuffer1.GetData(), VertexBuffer2.GetData(), ContactPointCount, VertexBuffer1.Num(), ClippingPlaneNormal, FVec3::DotProduct(CurrentPoint, ClippingPlaneNormal));
@@ -100,13 +100,14 @@ namespace Chaos
 		check(VertexBuffer1.Num() == VertexBuffer2.Num());
 
 		// Populate the clipped vertices with the convex face vertices
+		const FReal ConvexWindingOrder = Convex.GetWindingOrder();
 		int32 ContactPointCount = 0;
 		const int32 ConvexFaceVerticesNum = Convex.NumPlaneVertices(ConvexPlaneIndex);
 		ContactPointCount = FMath::Min(ConvexFaceVerticesNum, VertexBuffer1.Num()); // Number of face vertices
 		for (int32 VertexIndex = 0; VertexIndex < ContactPointCount; ++VertexIndex)
 		{
-			// Todo Check for Grey code
-			VertexBuffer1[VertexIndex] = Convex.GetVertex(Convex.GetPlaneVertex(ConvexPlaneIndex, VertexIndex));
+			const int32 BufferIndex = (ConvexWindingOrder >= 0) ? VertexIndex : ContactPointCount - VertexIndex - 1;
+			VertexBuffer1[BufferIndex] = Convex.GetVertex(Convex.GetPlaneVertex(ConvexPlaneIndex, VertexIndex));
 		}
 
 		// Now clip against all planes that belong to the reference plane's, edges
@@ -224,15 +225,21 @@ namespace Chaos
 			FMath::Min(TriVertexConvexDMin1, TriVertexConvexDMin2),
 		};
 		
+		FReal ConvexWinding = Convex.GetWindingOrder();
 		FVec3 EdgeEdgeN = FVec3(0);
 		FReal EdgeEdgeDMin = InvalidPhi;
 		int32 ConvexEdgeIndexMin = INDEX_NONE;
 		int32 TriEdgeIndexMin = INDEX_NONE;
-		for (int32 ConvexEdgeIndex = 0; ConvexEdgeIndex < Convex.NumEdges(); ++ConvexEdgeIndex)
+		for (int32 ConvexEdgeLoopIndex = 0; ConvexEdgeLoopIndex < Convex.NumEdges(); ++ConvexEdgeLoopIndex)
 		{
+			// Handle reverse winding for negative scaled convexes. Loop over edges in reverse order, and reverse edge vertex order
+			const int32 ConvexEdgeIndex = (ConvexWinding >= 0) ? ConvexEdgeLoopIndex : Convex.NumEdges() - ConvexEdgeLoopIndex - 1;
+			const int32 ConvexEdgeVIndex0 = (ConvexWinding >= 0) ? 0 : 1;
+			const int32 ConvexEdgeVIndex1 = (ConvexWinding >= 0) ? 1 : 0;
+
 			// Skip convex edges beyond CullDistance of the triangle face
-			const int32 ConvexEdgeVertexIndex0 = Convex.GetEdgeVertex(ConvexEdgeIndex, 0);
-			const int32 ConvexEdgeVertexIndex1 = Convex.GetEdgeVertex(ConvexEdgeIndex, 1);
+			const int32 ConvexEdgeVertexIndex0 = Convex.GetEdgeVertex(ConvexEdgeIndex, ConvexEdgeVIndex0);
+			const int32 ConvexEdgeVertexIndex1 = Convex.GetEdgeVertex(ConvexEdgeIndex, ConvexEdgeVIndex1);
 			const FReal FaceConvexD0 = ConvexVertexDs[ConvexEdgeVertexIndex0];
 			const FReal FaceConvexD1 = ConvexVertexDs[ConvexEdgeVertexIndex1];
 			if ((FaceConvexD0 > CullDistance) && (FaceConvexD1 > CullDistance))
diff --git a/Engine/Source/Runtime/Experimental/Chaos/Private/Chaos/Collision/PBDCollisionContainerSolver.cpp b/Engine/Source/Runtime/Experimental/Chaos/Private/Chaos/Collision/PBDCollisionContainerSolver.cpp
--- a/Engine/Source/Runtime/Experimental/Chaos/Private/Chaos/Collision/PBDCollisionContainerSolver.cpp
+++ b/Engine/Source/Runtime/Experimental/Chaos/Private/Chaos/Collision/PBDCollisionContainerSolver.cpp
@@ -156,7 +156,7 @@ namespace Chaos
 		// Initial Phi for initial-overlap depenetration.
 		// If we have an initial contact, calculate the initial overlap. This will get saved in SetSolverResults
 		FRealSingle WorldContactInitialPhi = 0;
-		if (MaxDepentrationVelocity >= 0)
+		if ((MaxDepentrationVelocity >= 0) && CVars::bChaos_Collision_EnableInitialDepenetration)
 		{
 			if (ManifoldPoint.Flags.bInitialContact)
 			{
