ChaosSofts : Fix CIS Issue

diff --git a/Engine/Source/Runtime/Experimental/Chaos/Public/Chaos/Deformable/GaussSeidelCorotatedCodimensionalConstraints.h b/Engine/Source/Runtime/Experimental/Chaos/Public/Chaos/Deformable/GaussSeidelCorotatedCodimensionalConstraints.h
--- a/Engine/Source/Runtime/Experimental/Chaos/Public/Chaos/Deformable/GaussSeidelCorotatedCodimensionalConstraints.h
+++ b/Engine/Source/Runtime/Experimental/Chaos/Public/Chaos/Deformable/GaussSeidelCorotatedCodimensionalConstraints.h
@@ -277,8 +277,8 @@ namespace Chaos::Softs
 			check(ElementIndex < MuElementArray.Num());
 			check(ElementIndex < Measure.Num());
 			check(ElementIndex < LambdaElementArray.Num());
-			check(ElementIndexLocal < 4 && ElementIndexLocal > INDEX_NONE);
-			if (ElementIndexLocal < 4 
+			check(ElementIndexLocal < 3 && ElementIndexLocal > INDEX_NONE);
+			if (ElementIndexLocal < 3 
 				&& ElementIndexLocal > INDEX_NONE
 				&& ElementIndex < Measure.Num()
 				&& ElementIndex < MuElementArray.Num()
