Deprecated COPY_ON_WRITE_ATTRIBUTE and Indirect TManagedArray, removed CopyOnWriteAttributeList saving 200 bytes per GC.
Use the RestCollection to access those attributes instead.

diff --git a/Engine/Source/Runtime/Experimental/GeometryCollectionEngine/Private/GeometryCollection/GeometryCollectionComponent.cpp b/Engine/Source/Runtime/Experimental/GeometryCollectionEngine/Private/GeometryCollection/GeometryCollectionComponent.cpp
--- a/Engine/Source/Runtime/Experimental/GeometryCollectionEngine/Private/GeometryCollection/GeometryCollectionComponent.cpp
+++ b/Engine/Source/Runtime/Experimental/GeometryCollectionEngine/Private/GeometryCollection/GeometryCollectionComponent.cpp
@@ -436,27 +436,21 @@ FGeomComponentCacheParameters::FGeomComponentCacheParameters()
 #define COPY_ON_WRITE_ATTRIBUTE(Type, Name, Group)										\
 const TManagedArray<Type>& UGeometryCollectionComponent::Get##Name##Array() const 		\
 {																						\
-	return Indirect##Name##Array ?														\
-		*Indirect##Name##Array : RestCollection->GetGeometryCollection()->Name;			\
+	return RestCollection->GetGeometryCollection()->Name;								\
 }																						\
 TManagedArray<Type>& UGeometryCollectionComponent::Get##Name##ArrayCopyOnWrite()		\
 {																						\
-	if(!Indirect##Name##Array)															\
+	static FName StaticName(#Name);														\
+	if (!DynamicCollection->HasAttribute(StaticName, Group))							\
 	{																					\
-		static FName StaticName(#Name);													\
 		DynamicCollection->AddAttribute<Type>(StaticName, Group);						\
 		DynamicCollection->CopyAttribute(												\
-			*RestCollection->GetGeometryCollection(), StaticName, Group);				\
-		Indirect##Name##Array =															\
-			&DynamicCollection->ModifyAttribute<Type>(StaticName, Group);				\
-		CopyOnWriteAttributeList.Add(													\
-			reinterpret_cast<FManagedArrayBase**>(&Indirect##Name##Array));				\
+		    *RestCollection->GetGeometryCollection(), StaticName, Group);				\
 	}																					\
-	return *Indirect##Name##Array;														\
+	return RestCollection->GetGeometryCollection()->Name;								\
 }																						\
 void UGeometryCollectionComponent::Reset##Name##ArrayDynamic()							\
 {																						\
-	Indirect##Name##Array = NULL;														\
 }																						\
 const TManagedArray<Type>& UGeometryCollectionComponent::Get##Name##ArrayRest() const	\
 {																						\
@@ -509,7 +503,6 @@ TManagedArray<int32>& UGeometryCollectionComponent::GetParentArrayCopyOnWrite()
 		DynamicCollection->AddAttribute<int32>(FTransformCollection::ParentAttribute, FTransformCollection::TransformGroup);
 		DynamicCollection->CopyAttribute(*RestCollection->GetGeometryCollection(), FTransformCollection::ParentAttribute, FTransformCollection::TransformGroup);
 		IndirectParentArray = &DynamicCollection->ModifyAttribute<int32>(FTransformCollection::ParentAttribute, FTransformCollection::TransformGroup);
-		CopyOnWriteAttributeList.Add(reinterpret_cast<FManagedArrayBase**>(&IndirectParentArray));
 	}
 	return *IndirectParentArray;
 }
@@ -849,15 +842,15 @@ namespace
 
 		auto GeometryCollectionPtr = Component.GetRestCollection()->GetGeometryCollection();
 		const TManagedArray<FBox>* TransformBoundingBoxes = GeometryCollectionPtr->FindAttribute<FBox>(BoundingBoxAttributeName, FGeometryCollection::TransformGroup);
-		const TManagedArray<FBox>& GeometryBoundingBoxes = Component.GetBoundingBoxArray();
-		const TManagedArray<int32>& TransformToGeometryIndex = Component.GetTransformToGeometryIndexArray();
+		const TManagedArray<FBox>& GeometryBoundingBoxes =  GeometryCollectionPtr->BoundingBox;
+		const TManagedArray<int32>& TransformToGeometryIndex = GeometryCollectionPtr->TransformToGeometryIndex;
 
 		if (TransformBoundingBoxes)
 		{
 			return ComputeBoundsFromTransformBoundingBoxes(TransformToGeometryIndex, *TransformBoundingBoxes, GlobalMatricesArray, LocalToWorldWithScale);
 		}
 
-		const TManagedArray<int32>& TransformIndices = Component.GetTransformIndexArray();
+		const TManagedArray<int32>& TransformIndices = GeometryCollectionPtr->TransformIndex;
 		return ComputeBoundsFromGeometryBoundingBoxes(TransformToGeometryIndex, TransformIndices, GeometryBoundingBoxes, GlobalMatricesArray, LocalToWorldWithScale);
 	}
 }
@@ -1437,7 +1430,7 @@ bool UGeometryCollectionComponent::DoCustomNavigableGeometryExport(FNavigableGeo
 	//gather data needed for indices
 	const TManagedArray<int32>& FaceStartArray = Collection->FaceStart;
 	const TManagedArray<FIntVector>& Indices = Collection->Indices;
-	const TManagedArray<bool>& Visible = GetVisibleArray();
+	const TManagedArray<bool>& Visible = Collection->Visible;
 	const TManagedArray<int32>& MaterialIndex = Collection->MaterialIndex;
 
 	//pre-allocate enough room (assuming all faces are visible)
@@ -1525,7 +1518,7 @@ void UGeometryCollectionComponent::RefreshEmbeddedGeometry()
 		return;
 	}
 
-	const TManagedArray<int32>& ExemplarIndexArray = GetExemplarIndexArray();
+	const TManagedArray<int32>& ExemplarIndexArray = GetRestCollection()->GetGeometryCollection()->ExemplarIndex;
 	const int32 TransformCount = ComponentSpaceTransforms.Num();
 	if (!ensureMsgf(TransformCount == ExemplarIndexArray.Num(), TEXT("GlobalMatrices (Num=%d) cached on GeometryCollectionComponent are not in sync with ExemplarIndexArray (Num=%d) on underlying GeometryCollection; likely missed a dynamic data update"), TransformCount, ExemplarIndexArray.Num()))
 	{
@@ -1710,7 +1703,6 @@ void UGeometryCollectionComponent::GetResourceSizeEx(FResourceSizeEx& Cumulative
 		+ CollisionProfilePerLevel.GetAllocatedSize()
 		+ ComponentSpaceTransforms.GetAllocatedSize()
 		+ EventsPlayed.GetAllocatedSize()
-		+ CopyOnWriteAttributeList.GetAllocatedSize()
 		+ EmbeddedGeometryComponents.GetAllocatedSize()
 		+ (ClustersToRep ? ClustersToRep->GetAllocatedSize() : 0);
 
@@ -3546,10 +3538,6 @@ void UGeometryCollectionComponent::ResetDynamicCollection()
 	if (bCreateDynamicCollection && RestCollection && RestCollection->GetGeometryCollection())
 	{
 		DynamicCollection = MakeUnique<FGeometryDynamicCollection>(RestCollection->GetGeometryCollection().Get());
-		for (const auto DynamicArray : CopyOnWriteAttributeList)
-		{
-			*DynamicArray = nullptr;
-		}
 
 		GetParentArrayCopyOnWrite();
 
@@ -5870,7 +5858,7 @@ void UGeometryCollectionComponent::SelectEmbeddedGeometry()
 		EmbeddedGeometryComponent->ClearInstanceSelection();
 	}
 	
-	const TManagedArray<int32>& ExemplarIndex = GetExemplarIndexArray();
+	const TManagedArray<int32>& ExemplarIndex = RestCollection->GetGeometryCollection()->ExemplarIndex;
 	for (int32 SelectedBone : SelectedBones)
 	{
 		if (EmbeddedGeometryComponents.IsValidIndex(ExemplarIndex[SelectedBone]))
diff --git a/Engine/Source/Runtime/Experimental/GeometryCollectionEngine/Private/GeometryCollection/GeometryCollectionSceneProxy.cpp b/Engine/Source/Runtime/Experimental/GeometryCollectionEngine/Private/GeometryCollection/GeometryCollectionSceneProxy.cpp
--- a/Engine/Source/Runtime/Experimental/GeometryCollectionEngine/Private/GeometryCollection/GeometryCollectionSceneProxy.cpp
+++ b/Engine/Source/Runtime/Experimental/GeometryCollectionEngine/Private/GeometryCollection/GeometryCollectionSceneProxy.cpp
@@ -1040,7 +1040,7 @@ FNaniteGeometryCollectionSceneProxy::FNaniteGeometryCollectionSceneProxy(UGeomet
 	const TSharedPtr<FGeometryCollection, ESPMode::ThreadSafe> Collection = GeometryCollection->GetGeometryCollection();
 	const TManagedArray<int32>& TransformToGeometryIndices = Collection->TransformToGeometryIndex;
 	const TManagedArray<int32>& SimulationType = Collection->SimulationType;
-	const TManagedArray<FGeometryCollectionSection>& SectionsArray = Component->GetSectionsArray();
+	const TManagedArray<FGeometryCollectionSection>& SectionsArray = Collection->Sections;
 
 	MaterialSections.SetNumZeroed(SectionsArray.Num());
 
diff --git a/Engine/Source/Runtime/Experimental/GeometryCollectionEngine/Public/GeometryCollection/GeometryCollectionComponent.h b/Engine/Source/Runtime/Experimental/GeometryCollectionEngine/Public/GeometryCollection/GeometryCollectionComponent.h
--- a/Engine/Source/Runtime/Experimental/GeometryCollectionEngine/Public/GeometryCollection/GeometryCollectionComponent.h
+++ b/Engine/Source/Runtime/Experimental/GeometryCollectionEngine/Public/GeometryCollection/GeometryCollectionComponent.h
@@ -282,12 +282,17 @@ private:
 //This generates pointers to arrays marked private. Macro assumes getters are public
 //todo(ocohen): may want to take in a static name
 #define COPY_ON_WRITE_ATTRIBUTE(Type, Name, Group)											\
+	UE_DEPRECATED(5.4, "Use GetGeometryCollection()->"#Name" instead.")						\
 	GEOMETRYCOLLECTIONENGINE_API const TManagedArray<Type>& Get##Name##Array() const;		\
+	UE_DEPRECATED(5.4, "Use GetGeometryCollection()->"#Name" instead.")						\
 	GEOMETRYCOLLECTIONENGINE_API TManagedArray<Type>& Get##Name##ArrayCopyOnWrite();		\
+	UE_DEPRECATED(5.4, "Use GetGeometryCollection()->"#Name" instead.")						\
 	GEOMETRYCOLLECTIONENGINE_API void Reset##Name##ArrayDynamic();							\
+	UE_DEPRECATED(5.4, "Use GetGeometryCollection()->"#Name" instead.")						\
 	GEOMETRYCOLLECTIONENGINE_API const TManagedArray<Type>& Get##Name##ArrayRest() const;	\
 private:																					\
-	TManagedArray<Type>* Indirect##Name##Array;												\
+	/* Deprecated */																		\
+	/*TManagedArray<Type>* Indirect##Name##Array;*/											\
 public:
 
 /**
@@ -1664,7 +1669,6 @@ private:
 
 	FGeometryCollectionPhysicsProxy* PhysicsProxy;
 	TUniquePtr<FGeometryDynamicCollection> DynamicCollection;
-	TArray<FManagedArrayBase**> CopyOnWriteAttributeList;
 
 	// Temporary dummies to interface with Physx expectations of the SQ syatem
 	friend class FGeometryCollectionSQAccelerator;
