ChaosVD - contact point rendering
Chaos - remove unused debug draw setting

diff --git a/Engine/Plugins/Experimental/ChaosVD/Source/ChaosVD/Private/Visualizers/ChaosVDSolverCollisionDataComponentVisualizer.cpp b/Engine/Plugins/Experimental/ChaosVD/Source/ChaosVD/Private/Visualizers/ChaosVDSolverCollisionDataComponentVisualizer.cpp
--- a/Engine/Plugins/Experimental/ChaosVD/Source/ChaosVD/Private/Visualizers/ChaosVDSolverCollisionDataComponentVisualizer.cpp
+++ b/Engine/Plugins/Experimental/ChaosVD/Source/ChaosVD/Private/Visualizers/ChaosVDSolverCollisionDataComponentVisualizer.cpp
@@ -233,6 +233,7 @@ void FChaosVDSolverCollisionDataComponentVisualizer::DrawMidPhaseData(const UAct
 				FColor EdgeNormalColor = FColor(250, 150, 0);
 				FColor ImpulseColor = FColor(0, 0, 250);
 				FColor PushOutImpulseColor = FColor(0, 250, 250);
+				FColor InitialPhiColor = FColor(189, 195, 199);
 				float ContactLenScale = 1.0f;
 
 				// TODO: Make this a setting when all CVD debug draw options support changing the line thickness
@@ -293,7 +294,13 @@ void FChaosVDSolverCollisionDataComponentVisualizer::DrawMidPhaseData(const UAct
 
 				if (EnumHasAnyFlags(VisualizationFlags, EChaosVDCollisionVisualizationFlags::ContactPoints))
 				{
-					FChaosVDDebugDrawUtils::DrawCircle(PDI, WorldPlaneLocation, DebugDrawSettings.ContactCircleRadius, CircleSegments, DiscColor, LinesThickness, Axes.GetUnitAxis(EAxis::Y), Axes.GetUnitAxis(EAxis::Z), TEXT("Contact : WorldPlaneLocation"), DebugDrawSettings.DepthPriority);
+					FChaosVDDebugDrawUtils::DrawCircle(PDI, WorldPlaneLocation, DebugDrawSettings.ContactCircleRadius, CircleSegments, DiscColor, LinesThickness, Axes.GetUnitAxis(EAxis::Y), Axes.GetUnitAxis(EAxis::Z), TEXT("Contact : Manifold Plane"), DebugDrawSettings.DepthPriority);
+					FChaosVDDebugDrawUtils::DrawCircle(PDI, WorldPointLocation, 0.5f * DebugDrawSettings.ContactCircleRadius, CircleSegments, DiscColor, LinesThickness, Axes.GetUnitAxis(EAxis::Y), Axes.GetUnitAxis(EAxis::Z), TEXT("Contact : Manifold Point"), DebugDrawSettings.DepthPriority);
+					if (ManifoldPoint.InitialPhi != 0)
+					{
+						FChaosVDDebugDrawUtils::DrawCircle(PDI, WorldPlaneLocation + ManifoldPoint.InitialPhi * WorldPlaneNormal, 0.25f * DebugDrawSettings.ContactCircleRadius, CircleSegments, InitialPhiColor, LinesThickness, Axes.GetUnitAxis(EAxis::Y), Axes.GetUnitAxis(EAxis::Z), TEXT("Contact : Manifold Initial Phi"), DebugDrawSettings.DepthPriority);
+
+					}
 				}
 
 				if (EnumHasAnyFlags(VisualizationFlags, EChaosVDCollisionVisualizationFlags::ContactNormal))
@@ -303,17 +310,6 @@ void FChaosVDSolverCollisionDataComponentVisualizer::DrawMidPhaseData(const UAct
 					FChaosVDDebugDrawUtils::DrawArrowVector(PDI, WorldPlaneLocation, WorldPlaneLocation + WorldPlaneNormal * Scale, TEXT("Contact Normal"), NormalColor, DebugDrawSettings.DepthPriority);
 				}
 
-				if (EnumHasAnyFlags(VisualizationFlags, EChaosVDCollisionVisualizationFlags::ContactPoints))
-				{
-					if (ManifoldPoint.ContactPoint.Phi < FLT_MAX)
-					{
-						FChaosVDDebugDrawUtils::DrawCircle(PDI, WorldPlaneLocation - ManifoldPoint.ContactPoint.Phi * WorldPlaneNormal, DebugDrawSettings.ContactPhiCircleRadius, CircleSegments, PhiDiscColor, LinesThickness, Axes.GetUnitAxis(EAxis::Y), Axes.GetUnitAxis(EAxis::Z), TEXT("Contact : WorldPlaneLocation - Phi"), DebugDrawSettings.DepthPriority);
-					}
-
-					// Manifold point
-					FChaosVDDebugDrawUtils::DrawCircle(PDI, WorldPointLocation, DebugDrawSettings.ContactCircleRadius, CircleSegments, DiscColor, LinesThickness, Axes.GetUnitAxis(EAxis::Y), Axes.GetUnitAxis(EAxis::Z), TEXT("Contact : Manifold Point"), DebugDrawSettings.DepthPriority);				
-				}
-
 				if (EnumHasAnyFlags(VisualizationFlags, EChaosVDCollisionVisualizationFlags::AccumulatedImpulse))
 				{
 					if (!Constraint.AccumulatedImpulse.IsNearlyZero())
diff --git a/Engine/Plugins/Experimental/ChaosVD/Source/ChaosVD/Public/ChaosVDEditorSettings.h b/Engine/Plugins/Experimental/ChaosVD/Source/ChaosVD/Public/ChaosVDEditorSettings.h
--- a/Engine/Plugins/Experimental/ChaosVD/Source/ChaosVD/Public/ChaosVDEditorSettings.h
+++ b/Engine/Plugins/Experimental/ChaosVD/Source/ChaosVD/Public/ChaosVDEditorSettings.h
@@ -62,10 +62,6 @@ struct FChaosVDContactDebugDrawSettings
 	/** The scale value to be applied to the normal vector of a contact used to change its size to make it easier to see */
 	UPROPERTY(EditAnywhere, Category=DebugDraw)
 	float ContactNormalScale = 30.0f;
-
-	/** The radius of the debug draw circle used to represent the Phi value (penetration) of a contact point */
-	UPROPERTY(EditAnywhere, Category=DebugDraw)
-	float ContactPhiCircleRadius = 2.0f;
 };
 
 /** Structure holding the settings using to debug draw Particles shape based on their state on the Chaos Visual Debugger */
diff --git a/Engine/Source/Runtime/Engine/Private/PhysicsEngine/ImmediatePhysics/ImmediatePhysicsChaos/ImmediatePhysicsSimulation_Chaos.cpp b/Engine/Source/Runtime/Engine/Private/PhysicsEngine/ImmediatePhysics/ImmediatePhysicsChaos/ImmediatePhysicsSimulation_Chaos.cpp
--- a/Engine/Source/Runtime/Engine/Private/PhysicsEngine/ImmediatePhysics/ImmediatePhysicsChaos/ImmediatePhysicsSimulation_Chaos.cpp
+++ b/Engine/Source/Runtime/Engine/Private/PhysicsEngine/ImmediatePhysics/ImmediatePhysicsChaos/ImmediatePhysicsSimulation_Chaos.cpp
@@ -200,7 +200,6 @@ Chaos::DebugDraw::FChaosDebugDrawSettings ChaosImmPhysDebugDebugDrawSettings(
 	/* BodyAxisLen =				*/ 4.0f,
 	/* ContactLen =					*/ 4.0f,
 	/* ContactWidth =				*/ 2.0f,
-	/* ContactPhiWidth =			*/ 0.0f,
 	/* ContactInfoWidth				*/ 2.0f,
 	/* ContactOwnerWidth =			*/ 0.0f,
 	/* ConstraintAxisLen =			*/ 5.0f,
@@ -231,7 +230,6 @@ FAutoConsoleVariableRef CVarChaosImmPhysBodyAxisLen(TEXT("p.Chaos.ImmPhys.DebugD
 FAutoConsoleVariableRef CVarChaosImmPhysContactLen(TEXT("p.Chaos.ImmPhys.DebugDraw.ContactLen"), ChaosImmPhysDebugDebugDrawSettings.ContactLen, TEXT("ContactLen."));
 FAutoConsoleVariableRef CVarChaosImmPhysContactWidth(TEXT("p.Chaos.ImmPhys.DebugDraw.ContactWidth"), ChaosImmPhysDebugDebugDrawSettings.ContactWidth, TEXT("ContactWidth."));
 FAutoConsoleVariableRef CVarChaosImmPhysContactInfoWidth(TEXT("p.Chaos.ImmPhys.DebugDraw.ContactInfoWidth"), ChaosImmPhysDebugDebugDrawSettings.ContactInfoWidth, TEXT("ContactInfoWidth."));
-FAutoConsoleVariableRef CVarChaosImmPhysContactPhiWidth(TEXT("p.Chaos.ImmPhys.DebugDraw.ContactPhiWidth"), ChaosImmPhysDebugDebugDrawSettings.ContactPhiWidth, TEXT("ContactPhiWidth."));
 FAutoConsoleVariableRef CVarChaosImmPhysContactOwnerWidth(TEXT("p.Chaos.ImmPhys.DebugDraw.ContactOwnerWidth"), ChaosImmPhysDebugDebugDrawSettings.ContactOwnerWidth, TEXT("ContactOwnerWidth."));
 FAutoConsoleVariableRef CVarChaosImmPhysConstraintAxisLen(TEXT("p.Chaos.ImmPhys.DebugDraw.ConstraintAxisLen"), ChaosImmPhysDebugDebugDrawSettings.ConstraintAxisLen, TEXT("ConstraintAxisLen."));
 FAutoConsoleVariableRef CVarChaosImmPhysLineThickness(TEXT("p.Chaos.ImmPhys.DebugDraw.LineThickness"), ChaosImmPhysDebugDebugDrawSettings.LineThickness, TEXT("LineThickness."));
diff --git a/Engine/Source/Runtime/Experimental/Chaos/Private/Chaos/ChaosDebugDraw.cpp b/Engine/Source/Runtime/Experimental/Chaos/Private/Chaos/ChaosDebugDraw.cpp
--- a/Engine/Source/Runtime/Experimental/Chaos/Private/Chaos/ChaosDebugDraw.cpp
+++ b/Engine/Source/Runtime/Experimental/Chaos/Private/Chaos/ChaosDebugDraw.cpp
@@ -136,7 +136,6 @@ namespace Chaos
 			/* BodyAxisLen =				*/ 4.0f,
 			/* ContactLen =					*/ 4.0f,
 			/* ContactWidth =				*/ 2.0f,
-			/* ContactPhiWidth =			*/ 0.0f,
 			/* ContactInfoWidth				*/ 2.0f,
 			/* ContactOwnerWidth =			*/ 0.0f,
 			/* ConstraintAxisLen =			*/ 5.0f,
@@ -1262,11 +1261,17 @@ namespace Chaos
 						FDebugDrawQueue::GetInstance().DrawDebugLine(WorldPointPlaneLocation, WorldPointPlaneLocation + Settings.DrawScale * Settings.ImpulseScale * SpaceTransform.TransformVectorNoScale(FVec3(ManifoldPointResult.NetImpulse)), Color, false, Duration, uint8(Settings.DrawPriority), Settings.LineThickness);
 					}
 
-					// Manifold plane and normal
+					// Manifold point and normal
 					if (Settings.ContactWidth > 0)
 					{
 						FColor C0 = DiscColor;
 						FDebugDrawQueue::GetInstance().DrawDebugCircle(WorldPlaneLocation, Settings.DrawScale * Settings.ContactWidth, 12, C0, false, Duration, uint8(Settings.DrawPriority), Settings.LineThickness, Axes.GetUnitAxis(EAxis::Y), Axes.GetUnitAxis(EAxis::Z), false);
+						FDebugDrawQueue::GetInstance().DrawDebugCircle(WorldPointLocation, 0.5f * Settings.DrawScale * Settings.ContactWidth, 12, C0, false, Duration, uint8(Settings.DrawPriority), Settings.LineThickness, Axes.GetUnitAxis(EAxis::Y), Axes.GetUnitAxis(EAxis::Z), false);
+						if (ManifoldPoint.InitialPhi != 0)
+						{
+							FColor C3 = FColor::Silver;
+							FDebugDrawQueue::GetInstance().DrawDebugCircle(WorldPlaneLocation + ManifoldPoint.InitialPhi * WorldPlaneNormal, 0.25f * Settings.DrawScale * Settings.ContactWidth, 12, C3, false, Duration, uint8(Settings.DrawPriority), Settings.LineThickness, Axes.GetUnitAxis(EAxis::Y), Axes.GetUnitAxis(EAxis::Z), false);
+						}
 					}
 					if (Settings.ContactLen > 0)
 					{
@@ -1274,20 +1279,6 @@ namespace Chaos
 						FColor C1 = NormalColor;
 						FDebugDrawQueue::GetInstance().DrawDebugLine(WorldPlaneLocation, WorldPlaneLocation + Settings.DrawScale * Settings.ContactLen * ContactLenScale * WorldPlaneNormal, C1, false, Duration, uint8(Settings.DrawPriority), Settings.LineThickness);
 					}
-					if (Settings.ContactPhiWidth > 0 && (ManifoldPoint.ContactPoint.Phi < FLT_MAX))
-					{
-						FColor C2 = PhiDiscColor;
-						FDebugDrawQueue::GetInstance().DrawDebugCircle(WorldPlaneLocation - ManifoldPoint.ContactPoint.Phi * WorldPlaneNormal, Settings.DrawScale * Settings.ContactPhiWidth, 12, C2, false, Duration, uint8(Settings.DrawPriority), Settings.LineThickness, Axes.GetUnitAxis(EAxis::Y), Axes.GetUnitAxis(EAxis::Z), false);
-					}
-
-					// Manifold point
-					FDebugDrawQueue::GetInstance().DrawDebugCircle(WorldPointLocation, 0.5f * Settings.DrawScale * Settings.ContactWidth, 12, DiscColor, false, Duration, uint8(Settings.DrawPriority), Settings.LineThickness, Axes.GetUnitAxis(EAxis::Y), Axes.GetUnitAxis(EAxis::Z), false);
-
-					if (ManifoldPoint.InitialPhi != 0)
-					{
-						FColor C3 = FColor::Silver;
-						FDebugDrawQueue::GetInstance().DrawDebugCircle(WorldPlaneLocation + ManifoldPoint.InitialPhi * WorldPlaneNormal, 0.25f * Settings.DrawScale * Settings.ContactWidth, 12, C3, false, Duration, uint8(Settings.DrawPriority), Settings.LineThickness, Axes.GetUnitAxis(EAxis::Y), Axes.GetUnitAxis(EAxis::Z), false);
-					}
 
 					// Whether restored
 					if (Settings.ContactInfoWidth > 0)
@@ -1329,11 +1320,11 @@ namespace Chaos
 				}
 
 				// Show the sphere approximation if enabled
-				if ((CVars::ChaosOneWayInteractionPairCollisionMode == (int32)EOneWayInteractionPairCollisionMode::SphereCollision) && Particle0->OneWayInteraction() && Particle1->OneWayInteraction())
-				{
-					DrawShape(Contact.GetShapeWorldTransform0(), Contact.GetImplicit0(), nullptr, FColor::Green, 0, &Settings);
-					DrawShape(Contact.GetShapeWorldTransform1(), Contact.GetImplicit1(), nullptr, FColor::Green, 0, &Settings);
-				}
+				//if ((CVars::ChaosOneWayInteractionPairCollisionMode == (int32)EOneWayInteractionPairCollisionMode::SphereCollision) && Particle0->OneWayInteraction() && Particle1->OneWayInteraction())
+				//{
+				//	DrawShape(Contact.GetShapeWorldTransform0(), Contact.GetImplicit0(), nullptr, FColor::Green, 0, &Settings);
+				//	DrawShape(Contact.GetShapeWorldTransform1(), Contact.GetImplicit1(), nullptr, FColor::Green, 0, &Settings);
+				//}
 			}
 			if (Settings.ContactOwnerWidth > 0)
 			{
diff --git a/Engine/Source/Runtime/Experimental/Chaos/Private/ChaosVisualDebugger/ChaosVDDataWrapperUtils.cpp b/Engine/Source/Runtime/Experimental/Chaos/Private/ChaosVisualDebugger/ChaosVDDataWrapperUtils.cpp
--- a/Engine/Source/Runtime/Experimental/Chaos/Private/ChaosVisualDebugger/ChaosVDDataWrapperUtils.cpp
+++ b/Engine/Source/Runtime/Experimental/Chaos/Private/ChaosVisualDebugger/ChaosVDDataWrapperUtils.cpp
@@ -18,6 +18,7 @@ void FChaosVDDataWrapperUtils::CopyManifoldPointsToDataWrapper(const Chaos::FMan
 	OutCopyTo.bWasReplaced = InCopyFrom.Flags.bWasReplaced;
 	OutCopyTo.bHasStaticFrictionAnchor = InCopyFrom.Flags.bHasStaticFrictionAnchor;
 	OutCopyTo.TargetPhi = InCopyFrom.TargetPhi;
+	OutCopyTo.InitialPhi = InCopyFrom.InitialPhi;
 
 	Algo::Transform(InCopyFrom.ShapeAnchorPoints, OutCopyTo.ShapeAnchorPoints, &FChaosVDDataWrapperUtils::ConvertToFVector);
 	Algo::Transform(InCopyFrom.InitialShapeContactPoints, OutCopyTo.InitialShapeContactPoints, &FChaosVDDataWrapperUtils::ConvertToFVector);
@@ -105,6 +106,8 @@ FChaosVDConstraint FChaosVDDataWrapperUtils::BuildConstraintDataWrapperFromConst
 	WrappedConstraintData.ClosestManifoldPointIndex = InConstraint.ClosestManifoldPointIndex;
 	WrappedConstraintData.ExpectedNumManifoldPoints = InConstraint.ExpectedNumManifoldPoints;
 	WrappedConstraintData.Stiffness = InConstraint.Stiffness;
+	WrappedConstraintData.MinInitialPhi = InConstraint.MinInitialPhi;
+	WrappedConstraintData.InitialOverlapDepenetrationVelocity = InConstraint.InitialOverlapDepenetrationVelocity;
 	WrappedConstraintData.CCDTimeOfImpact = InConstraint.CCDTimeOfImpact;
 	WrappedConstraintData.CCDEnablePenetration = InConstraint.CCDEnablePenetration;
 	WrappedConstraintData.CCDTargetPenetration = InConstraint.CCDTargetPenetration;
diff --git a/Engine/Source/Runtime/Experimental/Chaos/Private/PBDRigidsSolver.cpp b/Engine/Source/Runtime/Experimental/Chaos/Private/PBDRigidsSolver.cpp
--- a/Engine/Source/Runtime/Experimental/Chaos/Private/PBDRigidsSolver.cpp
+++ b/Engine/Source/Runtime/Experimental/Chaos/Private/PBDRigidsSolver.cpp
@@ -182,7 +182,6 @@ namespace Chaos
 			/* BodyAxisLen =				*/ 30.0f,
 			/* ContactLen =					*/ 30.0f,
 			/* ContactWidth =				*/ 6.0f,
-			/* ContactPhiWidth =			*/ 0.0f,
 			/* ContactInfoWidth				*/ 6.0f,
 			/* ContactOwnerWidth =			*/ 0.0f,
 			/* ConstraintAxisLen =			*/ 30.0f,
@@ -239,7 +238,6 @@ namespace Chaos
 		FAutoConsoleVariableRef CVarChaosSolverContactLen(TEXT("p.Chaos.Solver.DebugDraw.ContactLen"), ChaosSolverDebugDebugDrawSettings.ContactLen, TEXT("ContactLen."));
 		FAutoConsoleVariableRef CVarChaosSolverContactWidth(TEXT("p.Chaos.Solver.DebugDraw.ContactWidth"), ChaosSolverDebugDebugDrawSettings.ContactWidth, TEXT("ContactWidth."));
 		FAutoConsoleVariableRef CVarChaosSolverContactInfoWidth(TEXT("p.Chaos.Solver.DebugDraw.ContactInfoWidth"), ChaosSolverDebugDebugDrawSettings.ContactInfoWidth, TEXT("ContactInfoWidth."));
-		FAutoConsoleVariableRef CVarChaosSolverContactPhiWidth(TEXT("p.Chaos.Solver.DebugDraw.ContactPhiWidth"), ChaosSolverDebugDebugDrawSettings.ContactPhiWidth, TEXT("ContactPhiWidth."));
 		FAutoConsoleVariableRef CVarChaosSolverContactOwnerWidth(TEXT("p.Chaos.Solver.DebugDraw.ContactOwnerWidth"), ChaosSolverDebugDebugDrawSettings.ContactOwnerWidth, TEXT("ContactOwnerWidth."));
 		FAutoConsoleVariableRef CVarChaosSolverConstraintAxisLen(TEXT("p.Chaos.Solver.DebugDraw.ConstraintAxisLen"), ChaosSolverDebugDebugDrawSettings.ConstraintAxisLen, TEXT("ConstraintAxisLen."));
 		FAutoConsoleVariableRef CVarChaosSolverLineThickness(TEXT("p.Chaos.Solver.DebugDraw.LineThickness"), ChaosSolverDebugDebugDrawSettings.LineThickness, TEXT("LineThickness."));
diff --git a/Engine/Source/Runtime/Experimental/Chaos/Public/Chaos/ChaosDebugDraw.h b/Engine/Source/Runtime/Experimental/Chaos/Public/Chaos/ChaosDebugDraw.h
--- a/Engine/Source/Runtime/Experimental/Chaos/Public/Chaos/ChaosDebugDraw.h
+++ b/Engine/Source/Runtime/Experimental/Chaos/Public/Chaos/ChaosDebugDraw.h
@@ -72,7 +72,6 @@ namespace Chaos
 				FRealSingle InBodyAxisLen,
 				FRealSingle InContactLen,
 				FRealSingle InContactWidth,
-				FRealSingle InContactPhiWidth,
 				FRealSingle InContactInfoWidth,
 				FRealSingle InContactOwnerWidth,
 				FRealSingle InConstraintAxisLen,
@@ -101,7 +100,6 @@ namespace Chaos
 				, BodyAxisLen(InBodyAxisLen)
 				, ContactLen(InContactLen)
 				, ContactWidth(InContactWidth)
-				, ContactPhiWidth(InContactPhiWidth)
 				, ContactInfoWidth(InContactInfoWidth)
 				, ContactOwnerWidth(InContactOwnerWidth)
 				, ConstraintAxisLen(InConstraintAxisLen)
@@ -132,7 +130,6 @@ namespace Chaos
 			FRealSingle BodyAxisLen;
 			FRealSingle ContactLen;
 			FRealSingle ContactWidth;
-			FRealSingle ContactPhiWidth;
 			FRealSingle ContactInfoWidth;
 			FRealSingle ContactOwnerWidth;
 			FRealSingle ConstraintAxisLen;
diff --git a/Engine/Source/Runtime/Experimental/ChaosVisualDebugger/Private/DataWrappers/ChaosVDCollisionDataWrappers.cpp b/Engine/Source/Runtime/Experimental/ChaosVisualDebugger/Private/DataWrappers/ChaosVDCollisionDataWrappers.cpp
--- a/Engine/Source/Runtime/Experimental/ChaosVisualDebugger/Private/DataWrappers/ChaosVDCollisionDataWrappers.cpp
+++ b/Engine/Source/Runtime/Experimental/ChaosVisualDebugger/Private/DataWrappers/ChaosVDCollisionDataWrappers.cpp
@@ -27,6 +27,7 @@ bool FChaosVDManifoldPoint::Serialize(FArchive& Ar)
 	Ar << NetImpulse;
 
 	Ar << TargetPhi;
+	Ar << InitialPhi;
 	Ar << ShapeAnchorPoints;
 	Ar << InitialShapeContactPoints;
 	Ar << ContactPoint;
@@ -82,6 +83,8 @@ bool FChaosVDConstraint::Serialize(FArchive& Ar)
 	Ar << LastShapeWorldPositionDelta;
 	Ar << LastShapeWorldRotationDelta;
 	Ar << Stiffness;
+	Ar << MinInitialPhi;
+	Ar << InitialOverlapDepenetrationVelocity;
 	Ar << CCDTimeOfImpact;
 	Ar << CCDEnablePenetration;
 	Ar << CCDTargetPenetration;
diff --git a/Engine/Source/Runtime/Experimental/ChaosVisualDebugger/Public/DataWrappers/ChaosVDCollisionDataWrappers.h b/Engine/Source/Runtime/Experimental/ChaosVisualDebugger/Public/DataWrappers/ChaosVDCollisionDataWrappers.h
--- a/Engine/Source/Runtime/Experimental/ChaosVisualDebugger/Public/DataWrappers/ChaosVDCollisionDataWrappers.h
+++ b/Engine/Source/Runtime/Experimental/ChaosVisualDebugger/Public/DataWrappers/ChaosVDCollisionDataWrappers.h
@@ -120,6 +120,8 @@ struct CHAOSVDRUNTIME_API FChaosVDManifoldPoint
 
 	UPROPERTY(VisibleAnywhere, Category=ContactData)
 	float TargetPhi = 0.f;
+	UPROPERTY(VisibleAnywhere, Category=ContactData)
+	float InitialPhi = 0.f;
 	UPROPERTY()
 	TArray<FVector> ShapeAnchorPoints;
 	UPROPERTY()
@@ -282,6 +284,12 @@ struct CHAOSVDRUNTIME_API FChaosVDConstraint
 	
 	UPROPERTY(VisibleAnywhere, Category=ConstraintData)
 	float Stiffness = 0.f;
+
+	UPROPERTY(VisibleAnywhere, Category = ConstraintData)
+	float MinInitialPhi = 0.f;
+
+	UPROPERTY(VisibleAnywhere, Category = ConstraintData)
+	float InitialOverlapDepenetrationVelocity = 0.f;
 	
 	UPROPERTY(VisibleAnywhere, Category=ConstraintData)
 	float CCDTimeOfImpact = 0.f;
