Chaos - fix joint debug draw when bodies are sleeping

diff --git a/Engine/Source/Runtime/Experimental/Chaos/Private/Chaos/ChaosDebugDraw.cpp b/Engine/Source/Runtime/Experimental/Chaos/Private/Chaos/ChaosDebugDraw.cpp
--- a/Engine/Source/Runtime/Experimental/Chaos/Private/Chaos/ChaosDebugDraw.cpp
+++ b/Engine/Source/Runtime/Experimental/Chaos/Private/Chaos/ChaosDebugDraw.cpp
@@ -1497,82 +1497,6 @@ namespace Chaos
 			}
 		}
 
-		void DrawJointConstraintImpl(const FRigidTransform3& SpaceTransform, const FVec3& InPa, const FVec3& InCa, const FVec3& InXa, const FMatrix33& Ra, const FVec3& InPb, const FVec3& InCb, const FVec3& InXb, const FMatrix33& Rb, Chaos::FRealSingle ColorScale, const FChaosDebugDrawJointFeatures& FeatureMask, const FChaosDebugDrawSettings& Settings)
-		{
-			using namespace Chaos::DebugDraw;
-			FColor R = (ColorScale * FColor::Red).ToFColor(false);
-			FColor G = (ColorScale * FColor::Green).ToFColor(false);
-			FColor B = (ColorScale * FColor::Blue).ToFColor(false);
-			FColor C = (ColorScale * FColor::Cyan).ToFColor(false);
-			FColor M = (ColorScale * FColor::Magenta).ToFColor(false);
-			FColor Y = (ColorScale * FColor::Yellow).ToFColor(false);
-			FVec3 Pa = SpaceTransform.TransformPosition(InPa);
-			FVec3 Pb = SpaceTransform.TransformPosition(InPb);
-			FVec3 Ca = SpaceTransform.TransformPosition(InCa);
-			FVec3 Cb = SpaceTransform.TransformPosition(InCb);
-			FVec3 Xa = SpaceTransform.TransformPosition(InXa);
-			FVec3 Xb = SpaceTransform.TransformPosition(InXb);
-
-			if (FeatureMask.bActorConnector)
-			{
-				const FRealSingle ConnectorThickness = 1.5f * Settings.LineThickness;
-				const FReal CoMSize = Settings.DrawScale * Settings.JointComSize;
-				// Leave a gap around the actor position so we can see where the center is
-				FVec3 Sa = Pa;
-				const FReal Lena = (Xa - Pa).Size();
-				if (Lena > UE_KINDA_SMALL_NUMBER)
-				{
-					Sa = FMath::Lerp(Pa, Xa, FMath::Clamp<FReal>(CoMSize / Lena, 0., 1.));
-				}
-				FVec3 Sb = Pb;
-				const FReal Lenb = (Xb - Pb).Size();
-				if (Lenb > UE_KINDA_SMALL_NUMBER)
-				{
-					Sb = FMath::Lerp(Pb, Xb, FMath::Clamp<FReal>(CoMSize / Lena, 0., 1.));
-				}
-				FDebugDrawQueue::GetInstance().DrawDebugLine(Pa, Sa, FColor::White, false, UE_KINDA_SMALL_NUMBER, uint8(Settings.DrawPriority), ConnectorThickness);
-				FDebugDrawQueue::GetInstance().DrawDebugLine(Pb, Sb, FColor::White, false, UE_KINDA_SMALL_NUMBER, uint8(Settings.DrawPriority), ConnectorThickness);
-				FDebugDrawQueue::GetInstance().DrawDebugLine(Sa, Xa, R, false, UE_KINDA_SMALL_NUMBER, uint8(Settings.DrawPriority), ConnectorThickness);
-				FDebugDrawQueue::GetInstance().DrawDebugLine(Sb, Xb, C, false, UE_KINDA_SMALL_NUMBER, uint8(Settings.DrawPriority), ConnectorThickness);
-			}
-			if (FeatureMask.bCoMConnector)
-			{
-				const FRealSingle ConnectorThickness = 1.5f * Settings.LineThickness;
-				const FReal CoMSize = Settings.DrawScale * Settings.JointComSize;
-				// Leave a gap around the body position so we can see where the center is
-				FVec3 Sa = Ca;
-				const FReal Lena = (Xa - Ca).Size();
-				if (Lena > UE_KINDA_SMALL_NUMBER)
-				{
-					Sa = FMath::Lerp(Ca, Xa, FMath::Clamp<FReal>(CoMSize / Lena, 0., 1.));
-				}
-				FVec3 Sb = Cb;
-				const FReal Lenb = (Xb - Cb).Size();
-				if (Lenb > UE_KINDA_SMALL_NUMBER)
-				{
-					Sb = FMath::Lerp(Cb, Xb, FMath::Clamp<FReal>(CoMSize / Lena, 0., 1.));
-				}
-				FDebugDrawQueue::GetInstance().DrawDebugLine(Ca, Sa, FColor::Black, false, UE_KINDA_SMALL_NUMBER, uint8(Settings.DrawPriority), ConnectorThickness);
-				FDebugDrawQueue::GetInstance().DrawDebugLine(Cb, Sb, FColor::Black, false, UE_KINDA_SMALL_NUMBER, uint8(Settings.DrawPriority), ConnectorThickness);
-				FDebugDrawQueue::GetInstance().DrawDebugLine(Sa, Xa, R, false, UE_KINDA_SMALL_NUMBER, uint8(Settings.DrawPriority), ConnectorThickness);
-				FDebugDrawQueue::GetInstance().DrawDebugLine(Sb, Xb, C, false, UE_KINDA_SMALL_NUMBER, uint8(Settings.DrawPriority), ConnectorThickness);
-			}
-			if (FeatureMask.bStretch)
-			{
-				const FRealSingle StretchThickness = 3.0f * Settings.LineThickness;
-				FDebugDrawQueue::GetInstance().DrawDebugLine(Xa, Xb, M, false, UE_KINDA_SMALL_NUMBER, uint8(Settings.DrawPriority), StretchThickness);
-			}
-			if (FeatureMask.bAxes)
-			{
-				FDebugDrawQueue::GetInstance().DrawDebugDirectionalArrow(Xa, Xa + Settings.DrawScale * Settings.ConstraintAxisLen * SpaceTransform.TransformVector(Ra.GetAxis(0)), Settings.DrawScale * Settings.ArrowSize, R, false, UE_KINDA_SMALL_NUMBER, uint8(Settings.DrawPriority), Settings.LineThickness);
-				FDebugDrawQueue::GetInstance().DrawDebugDirectionalArrow(Xa, Xa + Settings.DrawScale * Settings.ConstraintAxisLen * SpaceTransform.TransformVector(Ra.GetAxis(1)), Settings.DrawScale * Settings.ArrowSize, G, false, UE_KINDA_SMALL_NUMBER, uint8(Settings.DrawPriority), Settings.LineThickness);
-				FDebugDrawQueue::GetInstance().DrawDebugDirectionalArrow(Xa, Xa + Settings.DrawScale * Settings.ConstraintAxisLen * SpaceTransform.TransformVector(Ra.GetAxis(2)), Settings.DrawScale * Settings.ArrowSize, B, false, UE_KINDA_SMALL_NUMBER, uint8(Settings.DrawPriority), Settings.LineThickness);
-				FDebugDrawQueue::GetInstance().DrawDebugDirectionalArrow(Xb, Xb + Settings.DrawScale * Settings.ConstraintAxisLen * SpaceTransform.TransformVector(Rb.GetAxis(0)), Settings.DrawScale * Settings.ArrowSize, C, false, UE_KINDA_SMALL_NUMBER, uint8(Settings.DrawPriority), Settings.LineThickness);
-				FDebugDrawQueue::GetInstance().DrawDebugDirectionalArrow(Xb, Xb + Settings.DrawScale * Settings.ConstraintAxisLen * SpaceTransform.TransformVector(Rb.GetAxis(1)), Settings.DrawScale * Settings.ArrowSize, M, false, UE_KINDA_SMALL_NUMBER, uint8(Settings.DrawPriority), Settings.LineThickness);
-				FDebugDrawQueue::GetInstance().DrawDebugDirectionalArrow(Xb, Xb + Settings.DrawScale * Settings.ConstraintAxisLen * SpaceTransform.TransformVector(Rb.GetAxis(2)), Settings.DrawScale * Settings.ArrowSize, Y, false, UE_KINDA_SMALL_NUMBER, uint8(Settings.DrawPriority), Settings.LineThickness);
-			}
-		}
-
 		void DrawJointConstraintImpl(const FRigidTransform3& SpaceTransform, const FPBDJointConstraintHandle* ConstraintHandle, Chaos::FRealSingle ColorScale, const FChaosDebugDrawJointFeatures& FeatureMask, const FChaosDebugDrawSettings& Settings)
 		{
 			if (!ConstraintHandle->IsEnabled())
@@ -1583,7 +1507,11 @@ namespace Chaos
 			TVec2<FGeometryParticleHandle*> ConstrainedParticles = ConstraintHandle->GetConstrainedParticles();
 			auto RigidParticle0 = ConstrainedParticles[0]->CastToRigidParticle();
 			auto RigidParticle1 = ConstrainedParticles[1]->CastToRigidParticle();
-			if ((RigidParticle0 && RigidParticle0->ObjectState() == EObjectStateType::Dynamic) || (RigidParticle1 && RigidParticle1->ObjectState() == EObjectStateType::Dynamic))
+			if ((RigidParticle0 == nullptr) && (RigidParticle1 == nullptr))
+			{
+				return;
+			}
+			if (RigidParticle0->IsDynamic() || RigidParticle0->IsSleeping() || RigidParticle1->IsDynamic() || RigidParticle1->IsSleeping())
 			{
 				FVec3 Pa = FParticleUtilities::GetActorWorldTransform(FConstGenericParticleHandle(ConstraintHandle->GetConstrainedParticles()[1])).GetTranslation();
 				FVec3 Pb = FParticleUtilities::GetActorWorldTransform(FConstGenericParticleHandle(ConstraintHandle->GetConstrainedParticles()[0])).GetTranslation();
@@ -1592,7 +1520,87 @@ namespace Chaos
 				FVec3 Xa, Xb;
 				FMatrix33 Ra, Rb;
 				ConstraintHandle->CalculateConstraintSpace(Xa, Ra, Xb, Rb);
-				DrawJointConstraintImpl(SpaceTransform, Pa, Ca, Xa, Ra, Pb, Cb, Xb, Rb, ColorScale, FeatureMask, Settings);
+
+				Pa = SpaceTransform.TransformPosition(Pa);
+				Pb = SpaceTransform.TransformPosition(Pb);
+				Ca = SpaceTransform.TransformPosition(Ca);
+				Cb = SpaceTransform.TransformPosition(Cb);
+				Xa = SpaceTransform.TransformPosition(Xa);
+				Xb = SpaceTransform.TransformPosition(Xb);
+
+				FColor R = FColor::Red;
+				FColor G = FColor::Green;
+				FColor B = FColor::Blue;
+				FColor C = FColor::Cyan;
+				FColor M = FColor::Magenta;
+				FColor Y = FColor::Yellow;
+
+				if (FeatureMask.bActorConnector)
+				{
+					const FRealSingle ConnectorThickness = 1.5f * Settings.LineThickness;
+					const FReal CoMSize = Settings.DrawScale * Settings.JointComSize;
+					// Leave a gap around the actor position so we can see where the center is
+					FVec3 Sa = Pa;
+					const FReal Lena = (Xa - Pa).Size();
+					if (Lena > UE_KINDA_SMALL_NUMBER)
+					{
+						Sa = FMath::Lerp(Pa, Xa, FMath::Clamp<FReal>(CoMSize / Lena, 0., 1.));
+					}
+					FVec3 Sb = Pb;
+					const FReal Lenb = (Xb - Pb).Size();
+					if (Lenb > UE_KINDA_SMALL_NUMBER)
+					{
+						Sb = FMath::Lerp(Pb, Xb, FMath::Clamp<FReal>(CoMSize / Lena, 0., 1.));
+					}
+					FDebugDrawQueue::GetInstance().DrawDebugLine(Pa, Sa, FColor::White, false, UE_KINDA_SMALL_NUMBER, uint8(Settings.DrawPriority), ConnectorThickness);
+					FDebugDrawQueue::GetInstance().DrawDebugLine(Pb, Sb, FColor::White, false, UE_KINDA_SMALL_NUMBER, uint8(Settings.DrawPriority), ConnectorThickness);
+					FDebugDrawQueue::GetInstance().DrawDebugLine(Sa, Xa, R, false, UE_KINDA_SMALL_NUMBER, uint8(Settings.DrawPriority), ConnectorThickness);
+					FDebugDrawQueue::GetInstance().DrawDebugLine(Sb, Xb, C, false, UE_KINDA_SMALL_NUMBER, uint8(Settings.DrawPriority), ConnectorThickness);
+				}
+				if (FeatureMask.bCoMConnector)
+				{
+					const FRealSingle ConnectorThickness = 1.5f * Settings.LineThickness;
+					const FReal CoMSize = Settings.DrawScale * Settings.JointComSize;
+					// Leave a gap around the body position so we can see where the center is
+					FVec3 Sa = Ca;
+					const FReal Lena = (Xa - Ca).Size();
+					if (Lena > UE_KINDA_SMALL_NUMBER)
+					{
+						Sa = FMath::Lerp(Ca, Xa, FMath::Clamp<FReal>(CoMSize / Lena, 0., 1.));
+					}
+					FVec3 Sb = Cb;
+					const FReal Lenb = (Xb - Cb).Size();
+					if (Lenb > UE_KINDA_SMALL_NUMBER)
+					{
+						Sb = FMath::Lerp(Cb, Xb, FMath::Clamp<FReal>(CoMSize / Lena, 0., 1.));
+					}
+					FDebugDrawQueue::GetInstance().DrawDebugLine(Ca, Sa, FColor::Black, false, UE_KINDA_SMALL_NUMBER, uint8(Settings.DrawPriority), ConnectorThickness);
+					FDebugDrawQueue::GetInstance().DrawDebugLine(Cb, Sb, FColor::Black, false, UE_KINDA_SMALL_NUMBER, uint8(Settings.DrawPriority), ConnectorThickness);
+					FDebugDrawQueue::GetInstance().DrawDebugLine(Sa, Xa, R, false, UE_KINDA_SMALL_NUMBER, uint8(Settings.DrawPriority), ConnectorThickness);
+					FDebugDrawQueue::GetInstance().DrawDebugLine(Sb, Xb, C, false, UE_KINDA_SMALL_NUMBER, uint8(Settings.DrawPriority), ConnectorThickness);
+				}
+				if (FeatureMask.bStretch)
+				{
+					const FRealSingle StretchThickness = 3.0f * Settings.LineThickness;
+					FDebugDrawQueue::GetInstance().DrawDebugLine(Xa, Xb, M, false, UE_KINDA_SMALL_NUMBER, uint8(Settings.DrawPriority), StretchThickness);
+				}
+				if (FeatureMask.bAxes)
+				{
+					FDebugDrawQueue::GetInstance().DrawDebugDirectionalArrow(Xa, Xa + Settings.DrawScale * Settings.ConstraintAxisLen * SpaceTransform.TransformVector(Ra.GetAxis(0)), Settings.DrawScale * Settings.ArrowSize, R, false, UE_KINDA_SMALL_NUMBER, uint8(Settings.DrawPriority), Settings.LineThickness);
+					FDebugDrawQueue::GetInstance().DrawDebugDirectionalArrow(Xa, Xa + Settings.DrawScale * Settings.ConstraintAxisLen * SpaceTransform.TransformVector(Ra.GetAxis(1)), Settings.DrawScale * Settings.ArrowSize, G, false, UE_KINDA_SMALL_NUMBER, uint8(Settings.DrawPriority), Settings.LineThickness);
+					FDebugDrawQueue::GetInstance().DrawDebugDirectionalArrow(Xa, Xa + Settings.DrawScale * Settings.ConstraintAxisLen * SpaceTransform.TransformVector(Ra.GetAxis(2)), Settings.DrawScale * Settings.ArrowSize, B, false, UE_KINDA_SMALL_NUMBER, uint8(Settings.DrawPriority), Settings.LineThickness);
+					FDebugDrawQueue::GetInstance().DrawDebugDirectionalArrow(Xb, Xb + Settings.DrawScale * Settings.ConstraintAxisLen * SpaceTransform.TransformVector(Rb.GetAxis(0)), Settings.DrawScale * Settings.ArrowSize, C, false, UE_KINDA_SMALL_NUMBER, uint8(Settings.DrawPriority), Settings.LineThickness);
+					FDebugDrawQueue::GetInstance().DrawDebugDirectionalArrow(Xb, Xb + Settings.DrawScale * Settings.ConstraintAxisLen * SpaceTransform.TransformVector(Rb.GetAxis(1)), Settings.DrawScale * Settings.ArrowSize, M, false, UE_KINDA_SMALL_NUMBER, uint8(Settings.DrawPriority), Settings.LineThickness);
+					FDebugDrawQueue::GetInstance().DrawDebugDirectionalArrow(Xb, Xb + Settings.DrawScale * Settings.ConstraintAxisLen * SpaceTransform.TransformVector(Rb.GetAxis(2)), Settings.DrawScale * Settings.ArrowSize, Y, false, UE_KINDA_SMALL_NUMBER, uint8(Settings.DrawPriority), Settings.LineThickness);
+				}
+
+				// NOTE: GetLinearImpulse is the positional impulse (pushout)
+				if ((Settings.PushOutScale > 0) && !ConstraintHandle->GetLinearImpulse().IsNearlyZero())
+				{
+					FColor PushOutImpusleColor = FColor(0, 250, 250);
+					FColor Color = PushOutImpusleColor;
+					FDebugDrawQueue::GetInstance().DrawDebugLine(Xa, Xa + Settings.DrawScale * Settings.PushOutScale * SpaceTransform.TransformVectorNoScale(FVec3(ConstraintHandle->GetLinearImpulse())), Color, false, 0, uint8(Settings.DrawPriority), Settings.LineThickness);
+				}
 			}
 		}
 
