Geometry Collection - add OneWayInteractionLevel rather than IsOneWayInteraction to get more control over which particles behave like debris

diff --git a/Engine/Source/Runtime/Experimental/Chaos/Private/Chaos/ClusterUnionManager.cpp b/Engine/Source/Runtime/Experimental/Chaos/Private/Chaos/ClusterUnionManager.cpp
--- a/Engine/Source/Runtime/Experimental/Chaos/Private/Chaos/ClusterUnionManager.cpp
+++ b/Engine/Source/Runtime/Experimental/Chaos/Private/Chaos/ClusterUnionManager.cpp
@@ -653,12 +653,18 @@ namespace Chaos
 				MEvolution.SetParticleObjectState(Cluster->InternalCluster, Chaos::EObjectStateType::Sleeping);
 			}
 
-			if (!Cluster->IsGravityOverrideSet() && !FinalParticlesToAdd.IsEmpty())
+			if (!FinalParticlesToAdd.IsEmpty())
 			{
-				Cluster->InternalCluster->SetGravityGroupIndex(FinalParticlesToAdd[0]->GravityGroupIndex());
-			}
+				FPBDRigidParticleHandle* ParticlePropertySource = FinalParticlesToAdd[0];
 
-			Cluster->InternalCluster->SetInitialOverlapDepenetrationVelocity(FinalParticlesToAdd[0]->InitialOverlapDepenetrationVelocity());
+				if (!Cluster->IsGravityOverrideSet())
+				{
+					Cluster->InternalCluster->SetGravityGroupIndex(ParticlePropertySource->GravityGroupIndex());
+				}
+
+				Cluster->InternalCluster->SetInitialOverlapDepenetrationVelocity(ParticlePropertySource->InitialOverlapDepenetrationVelocity());
+				Cluster->InternalCluster->SetOneWayInteraction(ParticlePropertySource->OneWayInteraction());
+			}
 		}
 	}
 
diff --git a/Engine/Source/Runtime/Experimental/Chaos/Private/Chaos/PBDRigidClustering.cpp b/Engine/Source/Runtime/Experimental/Chaos/Private/Chaos/PBDRigidClustering.cpp
--- a/Engine/Source/Runtime/Experimental/Chaos/Private/Chaos/PBDRigidClustering.cpp
+++ b/Engine/Source/Runtime/Experimental/Chaos/Private/Chaos/PBDRigidClustering.cpp
@@ -251,9 +251,11 @@ namespace Chaos
 		MEvolution.DisableParticles(reinterpret_cast<TSet<FGeometryParticleHandle*>&>(ChildrenSet));
 
 		bool bClusterIsAsleep = true;
+		bool bClusterIsOneWayInteraction = true;
 		for (FPBDRigidParticleHandle* Child : ChildrenSet)
 		{
 			bClusterIsAsleep &= Child->Sleeping();
+			bClusterIsOneWayInteraction &= Child->OneWayInteraction();
 
 			if (FPBDRigidClusteredParticleHandle* ClusteredChild = Child->CastToClustered())
 			{
@@ -302,6 +304,8 @@ namespace Chaos
 
 		NewParticle->SetSleeping(bClusterIsAsleep);
 
+		NewParticle->SetOneWayInteraction(bClusterIsOneWayInteraction);
+
 		auto AddToClusterUnion = [this](int32 ClusterID, FPBDRigidClusteredParticleHandle* Handle)
 		{
 			if (ClusterID <= 0)
@@ -561,6 +565,7 @@ namespace Chaos
 			TSet<FGeometryParticleHandle*> ChildrenHandles(static_cast<TArray<FGeometryParticleHandle*>>(ChildrenArray));
 			MEvolution.DisableParticles(ChildrenHandles);
 		}
+		bool bClusterIsOneWayInteraction = true;
 		for (FPBDRigidParticleHandle* Child : ChildrenArray)
 		{
 			if (FPBDRigidClusteredParticleHandle* ClusteredChild = Child->CastToClustered())
@@ -578,6 +583,8 @@ namespace Chaos
 
 				ClusteredChild->SetCollisionImpulses(FMath::Max(NewParticle->CollisionImpulses(), ClusteredChild->CollisionImpulses()));
 				Child->SetCollisionGroup(FMath::Min(NewParticle->CollisionGroup(), Child->CollisionGroup()));
+
+				bClusterIsOneWayInteraction &= Child->OneWayInteraction();
 			}
 		}
 
@@ -594,6 +601,8 @@ namespace Chaos
 
 		UpdateGeometry(NewParticle, ChildrenSet, MChildren, FImplicitObjectPtr(nullptr), NoCleanParams);
 
+		NewParticle->SetOneWayInteraction(bClusterIsOneWayInteraction);
+
 		// Build the convex optimizer if required
 		FRigidClustering::BuildConvexOptimizer(NewParticle);
 
diff --git a/Engine/Source/Runtime/Experimental/Chaos/Private/PhysicsProxy/GeometryCollectionPhysicsProxy.cpp b/Engine/Source/Runtime/Experimental/Chaos/Private/PhysicsProxy/GeometryCollectionPhysicsProxy.cpp
--- a/Engine/Source/Runtime/Experimental/Chaos/Private/PhysicsProxy/GeometryCollectionPhysicsProxy.cpp
+++ b/Engine/Source/Runtime/Experimental/Chaos/Private/PhysicsProxy/GeometryCollectionPhysicsProxy.cpp
@@ -63,6 +63,7 @@ namespace
 	const FName MassToLocalAttributeName = "MassToLocal";
 	const FName MassAttributeName = "Mass";
 	const FName InertiaTensorAttributeName = "InertiaTensor";
+	const FName LevelAttributeName = "Level";
 }
 
 namespace Chaos{
@@ -579,6 +580,7 @@ void FGeometryCollectionPhysicsProxy::Initialize(Chaos::FPBDRigidsEvolutionBase
 	const FVector Scale = Parameters.WorldTransform.GetScale3D();
 	const TManagedArray<float>& Mass = Parameters.RestCollection->GetAttribute<float>(MassAttributeName, FTransformCollection::TransformGroup);
 	const TManagedArray<FTransform>& MassToLocal = Parameters.RestCollection->GetAttribute<FTransform>(MassToLocalAttributeName, FTransformCollection::TransformGroup);
+	const TManagedArray<int32>& Level = Parameters.RestCollection->GetAttribute<int32>(LevelAttributeName, FTransformCollection::TransformGroup);
 
 	TArray<int32> ChildrenToCheckForParentFix;
 	TManagedArray<Chaos::FImplicitObjectPtr>& Implicits = GameThreadCollection.ModifyAttribute<Chaos::FImplicitObjectPtr>(FGeometryDynamicCollection::ImplicitsAttribute, FTransformCollection::TransformGroup);
@@ -647,6 +649,9 @@ void FGeometryCollectionPhysicsProxy::Initialize(Chaos::FPBDRigidsEvolutionBase
 					P->SetSpatialIdx(Chaos::FSpatialAccelerationIdx{ 0,0 });
 				}
 
+				const bool bIsOneWayInteraction = (Parameters.OneWayInteractionLevel >= 0) && (Level[Index] >= Parameters.OneWayInteractionLevel);
+				P->SetOneWayInteraction(bIsOneWayInteraction);
+
 #if !(UE_BUILD_SHIPPING || UE_BUILD_TEST)
 				P->SetDebugName(MakeShared<FString, ESPMode::ThreadSafe>(FString::Printf(TEXT("%s-%d"), *Parameters.Name, Index)));
 #endif
@@ -1491,14 +1496,18 @@ void FGeometryCollectionPhysicsProxy::InitializeBodiesPT(Chaos::FPBDRigidsSolver
 
 		// apply various features on the handles 
 		const bool bEnableGravity = Parameters.EnableGravity && !DisableGeometryCollectionGravity;
-		for (Chaos::FPBDRigidParticleHandle* Handle: SolverParticleHandles)
+		const TManagedArray<int32>& Level = PhysicsThreadCollection.GetInitialLevels().Get();
+		for (int32 ParticleIndex = 0; ParticleIndex < SolverParticleHandles.Num(); ++ParticleIndex)
 		{
+			Chaos::FPBDRigidParticleHandle* Handle = SolverParticleHandles[ParticleIndex];
 			if (Handle)
 			{
+				const bool bIsOneWayInteraction = (Parameters.OneWayInteractionLevel >= 0) && (Level[ParticleIndex] >= Parameters.OneWayInteractionLevel);
+
 				Handle->SetGravityEnabled(bEnableGravity);
 				Handle->SetGravityGroupIndex(Parameters.GravityGroupIndex);
 				Handle->SetCCDEnabled(Parameters.UseCCD);
-				Handle->SetOneWayInteraction(Parameters.bIsOneWayInteraction);
+				Handle->SetOneWayInteraction(bIsOneWayInteraction);
 				Handle->SetInertiaConditioningEnabled(Parameters.UseInertiaConditioning);
 				Handle->SetLinearEtherDrag(Parameters.LinearDamping);
 				Handle->SetAngularEtherDrag(Parameters.AngularDamping);
@@ -3242,23 +3251,38 @@ void FGeometryCollectionPhysicsProxy::SetGravityGroupIndex_Internal(int32 Gravit
 	}
 }
 
-void FGeometryCollectionPhysicsProxy::SetIsOneWayInteraction_External(bool bIsOneWayInteraction)
+void FGeometryCollectionPhysicsProxy::SetOneWayInteractionLevel_External(int32 InOneWayInteractionLevel)
 {
-	ExecuteOnPhysicsThread(*this, [this, bIsOneWayInteraction]()
+	// @todo(chaos): not sure we need the OneWayInteraction flag on the GT side for geometry collection particles?
+	const TManagedArray<int32>& Level = Parameters.RestCollection->GetAttribute<int32>(LevelAttributeName, FTransformCollection::TransformGroup);
+	for (int32 ParticleIndex = 0; ParticleIndex < GTParticles.Num(); ++ParticleIndex)
+	{
+		TUniquePtr<FParticle>& Particle = GTParticles[ParticleIndex];
+		if (Particle.IsValid())
+		{
+			const bool bIsOneWayInteraction = (InOneWayInteractionLevel >= 0) && (Level[ParticleIndex] >= InOneWayInteractionLevel);
+			Particle->SetOneWayInteraction(bIsOneWayInteraction);
+		}
+	}
+
+	ExecuteOnPhysicsThread(*this, [this, InOneWayInteractionLevel]()
 		{
-			SetIsOneWayInteraction_Internal(bIsOneWayInteraction);
+			SetOneWayInteractionLevel_Internal(InOneWayInteractionLevel);
 		});
 }
 
-void FGeometryCollectionPhysicsProxy::SetIsOneWayInteraction_Internal(bool bInIsOneWayInteraction)
+void FGeometryCollectionPhysicsProxy::SetOneWayInteractionLevel_Internal(int32 InOneWayInteractionLevel)
 {
-	Parameters.bIsOneWayInteraction = bInIsOneWayInteraction;
+	Parameters.OneWayInteractionLevel = InOneWayInteractionLevel;
 
-	for (Chaos::FPBDRigidClusteredParticleHandle* Handle : SolverParticleHandles)
+	const TManagedArray<int32>& Level = PhysicsThreadCollection.GetInitialLevels().Get();
+	for (int32 ParticleIndex = 0; ParticleIndex < SolverParticleHandles.Num(); ++ParticleIndex)
 	{
+		Chaos::FPBDRigidClusteredParticleHandle* Handle = SolverParticleHandles[ParticleIndex];
 		if (Handle)
 		{
-			Handle->SetOneWayInteraction(Parameters.bIsOneWayInteraction);
+			const bool bIsOneWayInteraction = (Parameters.OneWayInteractionLevel >= 0) && (Level[ParticleIndex] >= Parameters.OneWayInteractionLevel);
+			Handle->SetOneWayInteraction(bIsOneWayInteraction);
 		}
 	}
 }
diff --git a/Engine/Source/Runtime/Experimental/Chaos/Public/GeometryCollection/GeometryCollectionSimulationCoreTypes.h b/Engine/Source/Runtime/Experimental/Chaos/Public/GeometryCollection/GeometryCollectionSimulationCoreTypes.h
--- a/Engine/Source/Runtime/Experimental/Chaos/Public/GeometryCollection/GeometryCollectionSimulationCoreTypes.h
+++ b/Engine/Source/Runtime/Experimental/Chaos/Public/GeometryCollection/GeometryCollectionSimulationCoreTypes.h
@@ -211,9 +211,9 @@ struct FSimulationParameters
 		, bGenerateGlobalCrumblingChildrenData(false)
 		, EnableGravity(true)
 		, GravityGroupIndex(0)
+		, OneWayInteractionLevel(INDEX_NONE)
 		, UseInertiaConditioning(true)
 		, UseCCD(false)
-		, bIsOneWayInteraction(false)
 		, LinearDamping(0.01f)
 		, AngularDamping(0)
 		, InitialOverlapDepenetrationVelocity(-1.0f)
@@ -272,9 +272,9 @@ struct FSimulationParameters
 		, Shared(Other.Shared)
 		, EnableGravity(Other.EnableGravity)
 		, GravityGroupIndex(Other.GravityGroupIndex)
+		, OneWayInteractionLevel(Other.OneWayInteractionLevel)
 		, UseInertiaConditioning(Other.UseInertiaConditioning)
 		, UseCCD(Other.UseCCD)
-		, bIsOneWayInteraction(Other.bIsOneWayInteraction)
 		, LinearDamping(Other.LinearDamping)
 		, AngularDamping(Other.AngularDamping)
 		, InitialOverlapDepenetrationVelocity(Other.InitialOverlapDepenetrationVelocity)
@@ -363,9 +363,9 @@ struct FSimulationParameters
 
 	bool EnableGravity;
 	int32 GravityGroupIndex;
+	int32 OneWayInteractionLevel;
 	bool UseInertiaConditioning;
 	bool UseCCD;
-	bool bIsOneWayInteraction;
 	float LinearDamping;
 	float AngularDamping;
 	float InitialOverlapDepenetrationVelocity;
diff --git a/Engine/Source/Runtime/Experimental/Chaos/Public/PhysicsProxy/GeometryCollectionPhysicsProxy.h b/Engine/Source/Runtime/Experimental/Chaos/Public/PhysicsProxy/GeometryCollectionPhysicsProxy.h
--- a/Engine/Source/Runtime/Experimental/Chaos/Public/PhysicsProxy/GeometryCollectionPhysicsProxy.h
+++ b/Engine/Source/Runtime/Experimental/Chaos/Public/PhysicsProxy/GeometryCollectionPhysicsProxy.h
@@ -473,7 +473,7 @@ public:
 	CHAOS_API void SetUseMaterialDamageModifiers_External(bool bUseMaterialDamageModifiers);
 	CHAOS_API void SetMaterialOverrideMassScaleMultiplier_External(float InMultiplier);
 	CHAOS_API void SetGravityGroupIndex_External(int32 GravityGroupIndex);
-	CHAOS_API void SetIsOneWayInteraction_External(bool bInIsOneWayInteraction);
+	CHAOS_API void SetOneWayInteractionLevel_External(int32 OneWayInteractionLevel);
 	CHAOS_API void SetPhysicsMaterial_External(const Chaos::FMaterialHandle& MaterialHandle);
 	/** 
 	 * Traverses the parents of TransformGroupIdx counting number of levels,
@@ -564,7 +564,7 @@ protected:
 	CHAOS_API void SetUseMaterialDamageModifiers_Internal(bool bUseMaterialDamageModifiers);
 	CHAOS_API void SetMaterialOverrideMassScaleMultiplier_Internal(float InMultiplier);
 	CHAOS_API void SetGravityGroupIndex_Internal(int32 GravityGroupIndex);
-	CHAOS_API void SetIsOneWayInteraction_Internal(bool bInIsOneWayInteraction);
+	CHAOS_API void SetOneWayInteractionLevel_Internal(int32 InOneWayInteractionLevel);
 	CHAOS_API void SetPhysicsMaterial_Internal(const Chaos::FMaterialHandle& MaterialHandle);
 
 private:
diff --git a/Engine/Source/Runtime/Experimental/GeometryCollectionEngine/Private/GeometryCollection/GeometryCollectionComponent.cpp b/Engine/Source/Runtime/Experimental/GeometryCollectionEngine/Private/GeometryCollection/GeometryCollectionComponent.cpp
--- a/Engine/Source/Runtime/Experimental/GeometryCollectionEngine/Private/GeometryCollection/GeometryCollectionComponent.cpp
+++ b/Engine/Source/Runtime/Experimental/GeometryCollectionEngine/Private/GeometryCollection/GeometryCollectionComponent.cpp
@@ -534,7 +534,7 @@ UGeometryCollectionComponent::UGeometryCollectionComponent(const FObjectInitiali
 	, InitializationState(ESimulationInitializationState::Unintialized)
 	, ObjectType(EObjectStateTypeEnum::Chaos_Object_Dynamic)
 	, GravityGroupIndex(0)
-	, bIsOneWayInteraction(false)
+	, OneWayInteractionLevel(INDEX_NONE)
 	, bDensityFromPhysicsMaterial(false)
 	, bForceMotionBlur()
 	, EnableClustering(true)
@@ -957,12 +957,12 @@ void UGeometryCollectionComponent::SetGravityGroupIndex(int32 InGravityGroupInde
 	}
 }
 
-void UGeometryCollectionComponent::SetIsOneWayInteraction(bool bInIsOneWayInteraction)
+void UGeometryCollectionComponent::SetOneWayInteractionLevel(int32 InOneWayInteractionLevel)
 {
-	bIsOneWayInteraction = bInIsOneWayInteraction;
+	OneWayInteractionLevel = InOneWayInteractionLevel;
 	if (PhysicsProxy)
 	{
-		PhysicsProxy->SetIsOneWayInteraction_External(bInIsOneWayInteraction);
+		PhysicsProxy->SetOneWayInteractionLevel_External(InOneWayInteractionLevel);
 	}
 }
 
@@ -3763,7 +3763,7 @@ void UGeometryCollectionComponent::RegisterAndInitializePhysicsProxy()
 		SimulationParameters.bGenerateGlobalCrumblingChildrenData = bGlobalCrumblingEventIncludesChildren;
 		SimulationParameters.EnableGravity = BodyInstance.bEnableGravity;
 		SimulationParameters.GravityGroupIndex = GravityGroupIndex;
-		SimulationParameters.bIsOneWayInteraction = bIsOneWayInteraction;
+		SimulationParameters.OneWayInteractionLevel = OneWayInteractionLevel;
 		SimulationParameters.UseInertiaConditioning = BodyInstance.IsInertiaConditioningEnabled();
 		SimulationParameters.UseCCD = BodyInstance.bUseCCD;
 		SimulationParameters.LinearDamping = BodyInstance.LinearDamping;
diff --git a/Engine/Source/Runtime/Experimental/GeometryCollectionEngine/Public/GeometryCollection/GeometryCollectionComponent.h b/Engine/Source/Runtime/Experimental/GeometryCollectionEngine/Public/GeometryCollection/GeometryCollectionComponent.h
--- a/Engine/Source/Runtime/Experimental/GeometryCollectionEngine/Public/GeometryCollection/GeometryCollectionComponent.h
+++ b/Engine/Source/Runtime/Experimental/GeometryCollectionEngine/Public/GeometryCollection/GeometryCollectionComponent.h
@@ -943,11 +943,13 @@ public:
 	UFUNCTION(BlueprintCallable, BlueprintInternalUseOnly)
 	GEOMETRYCOLLECTIONENGINE_API void SetGravityGroupIndex(int32 InGravityGroupIndex);
 
-	UPROPERTY(EditAnywhere, BlueprintReadWrite, BlueprintSetter = SetIsOneWayInteraction, Category = "ChaosPhysics|General")
-	bool bIsOneWayInteraction;
+	// All bodies with a level greater than or equal to this will have One-Way Interaction enabled and act like debris (will not apply forces to non-debris bodies)
+	// Set to -1 to disable (no bodies will have One-Way Interaction enabled)
+	UPROPERTY(EditAnywhere, BlueprintReadWrite, BlueprintSetter = SetOneWayInteractionLevel, Category = "ChaosPhysics|General")
+	int32 OneWayInteractionLevel;
 
 	UFUNCTION(BlueprintCallable, BlueprintInternalUseOnly)
-	GEOMETRYCOLLECTIONENGINE_API void SetIsOneWayInteraction(bool bInOneWayInteraction);
+	GEOMETRYCOLLECTIONENGINE_API void SetOneWayInteractionLevel(int32 InOneWayInteractionLevel);
 
 	/** when true, density will be used to compute mass using the assigned physics material */
 	UPROPERTY(EditAnywhere, BlueprintReadWrite, BlueprintSetter = SetDensityFromPhysicsMaterial, Category = "ChaosPhysics|General")
