ChaosSofts : Fix CIS Issue

diff --git a/Engine/Source/Runtime/Experimental/Chaos/Public/Chaos/Deformable/GaussSeidelCorotatedCodimensionalConstraints.h b/Engine/Source/Runtime/Experimental/Chaos/Public/Chaos/Deformable/GaussSeidelCorotatedCodimensionalConstraints.h
--- a/Engine/Source/Runtime/Experimental/Chaos/Public/Chaos/Deformable/GaussSeidelCorotatedCodimensionalConstraints.h
+++ b/Engine/Source/Runtime/Experimental/Chaos/Public/Chaos/Deformable/GaussSeidelCorotatedCodimensionalConstraints.h
@@ -375,7 +375,11 @@ namespace Chaos::Softs
 					Chaos::TVector<T, 3> L((T)0.);
 					for (int32 alpha = 0; alpha < 3; alpha++) 
 					{
-						L[alpha] = dJDmInvT.M[ElementIndexLocal* 3 - 3 + alpha];
+						const int32 IndexVisited = ElementIndexLocal* 3 - 3 + alpha;
+						if (IndexVisited < 6 && IndexVisited > INDEX_NONE)
+						{
+							L[alpha] = dJDmInvT.M[IndexVisited];
+						}
 					}
 					for (int32 alpha = 0; alpha < 3; alpha++)
 					{
