Chaos - Fix TransformedAABB2 ISPC version

diff --git a/Engine/Source/Runtime/Experimental/Chaos/Private/Chaos/AABB.ispc b/Engine/Source/Runtime/Experimental/Chaos/Private/Chaos/AABB.ispc
--- a/Engine/Source/Runtime/Experimental/Chaos/Private/Chaos/AABB.ispc
+++ b/Engine/Source/Runtime/Experimental/Chaos/Private/Chaos/AABB.ispc
@@ -99,7 +99,7 @@ export void TransformedAABB2(const uniform FTransform &SpaceTransform, const uni
 	// Transformed center
 	// NOTE: This is where positive/negative scales matters
 	const uniform FVector4 Center = 0.5 * (Min + Max);
-	const uniform FVector4 TCenter = SpaceTransform.Translation + SpaceTransform.Rotation * (SpaceTransform.Scale3D * Center);
+	const uniform FVector4 TCenter = SpaceTransform.Translation + VectorQuaternionRotateVector(SpaceTransform.Rotation, (SpaceTransform.Scale3D * Center));
 
 	// Transformed bounds
 	NewMin = SetVector(TCenter - RExtent);
