[Chaos] Removed use of TFunction as a parameter type for a callback in the AABBTree generation process, to remove eliminate the virtual call made by that type.

This hasn't show up on per profiles yet, but I introduced it a month ago and the virtual call it is not really needed.

diff --git a/Engine/Source/Runtime/Experimental/Chaos/Public/Chaos/AABBTree.h b/Engine/Source/Runtime/Experimental/Chaos/Public/Chaos/AABBTree.h
--- a/Engine/Source/Runtime/Experimental/Chaos/Public/Chaos/AABBTree.h
+++ b/Engine/Source/Runtime/Experimental/Chaos/Public/Chaos/AABBTree.h
@@ -2198,7 +2198,7 @@ public:
 		int32 SizeToCopyLeft = MaximumBytesToCopy;
 		check(From.CellHashToFlatArray.Num() == 0); // Partial Copy of TMAPs not implemented, and this should be empty for our current use cases
 
-		TFunction<bool(int32,int32)> CanContinueCopyingDataCallback = [this](int32 MaxSizeToCopy, int32 CurrentCopiedSize)
+		auto CanContinueCopyingDataCallback = [this](int32 MaxSizeToCopy, int32 CurrentCopiedSize)
 		{
 			bool bCanContinueCopy = true;
 			const bool bForceCopyAll = MaxSizeToCopy == -1;
@@ -3741,8 +3741,8 @@ private:
 		}
 	}
 
-	template<typename ContainerType>
-	static bool ContinueTimeSliceCopy(const ContainerType& ContainerFrom, ContainerType& ContainerTo, int32& InOutMaxSize, TFunctionRef<bool(int32,int32)> CanContinueCallback)
+	template<typename ContainerType, typename TCanContinueCallback>
+	static bool ContinueTimeSliceCopy(const ContainerType& ContainerFrom, ContainerType& ContainerTo, int32& InOutMaxSize, const TCanContinueCallback& CanContinueCallback)
 	{
 		int32 SizeCopied = 0;
 
