Fix issue where updating twice a geometry collection component transform per tick would imporperly rest the physics side previous transform causing the kinematic particle to not be placed properly in world space

diff --git a/Engine/Source/Runtime/Experimental/Chaos/Private/PhysicsProxy/GeometryCollectionPhysicsProxy.cpp b/Engine/Source/Runtime/Experimental/Chaos/Private/PhysicsProxy/GeometryCollectionPhysicsProxy.cpp
--- a/Engine/Source/Runtime/Experimental/Chaos/Private/PhysicsProxy/GeometryCollectionPhysicsProxy.cpp
+++ b/Engine/Source/Runtime/Experimental/Chaos/Private/PhysicsProxy/GeometryCollectionPhysicsProxy.cpp
@@ -2997,7 +2997,6 @@ void FGeometryCollectionPhysicsProxy::SetWorldTransform_Internal(const FTransfor
 {
 	using namespace Chaos;
 
-	Parameters.PrevWorldTransform = Parameters.WorldTransform;
 	Parameters.WorldTransform = InWorldTransform;
 
 	TSet<FClusterHandle*> ProcessedInternalClusters;
@@ -3658,6 +3657,9 @@ void FGeometryCollectionPhysicsProxy::BufferPhysicsResults_Internal(Chaos::FPBDR
 	// Advertise to game thread
 	Results.IsObjectDynamic = IsObjectDynamic;
 	Results.IsObjectLoading = IsObjectLoading;
+
+	//we are now done with this physics thread tick , we can set the previous world transform 
+	Parameters.PrevWorldTransform = Parameters.WorldTransform;
 }
 
 void FGeometryCollectionPhysicsProxy::FlipBuffer()
