Exposed some setters for various properties we were setting in construction scrips.

#rnx

diff --git a/Engine/Source/Runtime/Experimental/GeometryCollectionEngine/Public/GeometryCollection/GeometryCollectionComponent.h b/Engine/Source/Runtime/Experimental/GeometryCollectionEngine/Public/GeometryCollection/GeometryCollectionComponent.h
--- a/Engine/Source/Runtime/Experimental/GeometryCollectionEngine/Public/GeometryCollection/GeometryCollectionComponent.h
+++ b/Engine/Source/Runtime/Experimental/GeometryCollectionEngine/Public/GeometryCollection/GeometryCollectionComponent.h
@@ -1325,6 +1325,12 @@ public:
 
 	GEOMETRYCOLLECTIONENGINE_API double GetRootBrokenElapsedTimeInMs() const { return BrokenAndDecayedStates.GetRootBrokenElapsedTimeInMs(); }
 
+	/** Attn: these replication methods are helpers meant to be called before the component is fully registered, like a constructor! */
+	GEOMETRYCOLLECTIONENGINE_API void SetEnableReplication(bool bInEnableReplication) { bEnableReplication = bInEnableReplication; }
+	GEOMETRYCOLLECTIONENGINE_API void SetReplicationAbandonAfterLevel(int32 InReplicationAbandonAfterLevel) { ReplicationAbandonAfterLevel = InReplicationAbandonAfterLevel; }
+	GEOMETRYCOLLECTIONENGINE_API void SetReplicationMaxPositionAndVelocityCorrectionLevel(int32 InReplicationMaxPositionAndVelocityCorrectionLevel) { ReplicationMaxPositionAndVelocityCorrectionLevel = InReplicationMaxPositionAndVelocityCorrectionLevel; }
+	
+
 public:
 	UPROPERTY(BlueprintAssignable, Category = "Collision")
 	FOnChaosPhysicsCollision OnChaosPhysicsCollision;
