Reserve space for at least one entry at the end of the journal

In certain cases the remaining space for the journal might be only 16 bytes. In these cases we can only fit the phrase description and MaxEntries for the phrase willl be -1. This will cause an assert and over commit the journal. With this change we detect that case in the preceding phrase and add an identity entry which will be ignored when loading the journal.

diff --git a/Engine/Source/Runtime/Experimental/IoStoreOnDemand/Private/JournaledCache.cpp b/Engine/Source/Runtime/Experimental/IoStoreOnDemand/Private/JournaledCache.cpp
--- a/Engine/Source/Runtime/Experimental/IoStoreOnDemand/Private/JournaledCache.cpp
+++ b/Engine/Source/Runtime/Experimental/IoStoreOnDemand/Private/JournaledCache.cpp
@@ -327,7 +327,8 @@ struct FDiskPhrase
 	const FDataEntry*	GetEntries() const		{ return Entries.GetData() + Index; }
 	int32				GetEntryCount() const	{ return Entries.Num() - Index; }
 	uint8*				GetPhraseData() const	{ return Buffer.Get(); }
-	uint32				GetDataSize() const		{ return Cursor; } 
+	uint32				GetDataSize() const		{ return Cursor; }
+	int32				GetRemainingEntries() const { return MaxEntries; }
 
 private:
 	TUniquePtr<uint8[]>	Buffer;
@@ -466,7 +467,7 @@ FDiskPhrase FDiskJournal::OpenPhrase(uint32 DataSize)
 	check((Cursor & (sizeof(FDataEntry) - 1)) == 0);
 
 	Entries.Add(FDataEntry{});
-	const int32 MaxEntries = int32((MaxSize - Cursor) / sizeof(FDataEntry)) - 1;
+	int32 MaxEntries = int32((MaxSize - Cursor) / sizeof(FDataEntry)) - Entries.Num();
 
 	FDiskPhrase Ret(Entries, FMath::Min(MaxEntries, int32(UINT16_MAX)), DataSize);
 
@@ -482,6 +483,16 @@ void FDiskJournal::ClosePhrase(FDiskPhrase&& Phrase, uint64 DataCursor)
 		Entries.Pop();
 		return;
 	}
+	
+	// Since we minimally need two entries to write a complete phrase, if there
+	// is only room for one more entry, pad the phrase with an identity entry.
+	const uint32 Size = Entries.Num() * sizeof(Entries[0]);
+	const int32 PhraseEnd = Cursor + Size;
+	if (MaxSize - PhraseEnd == sizeof(FDataEntry))
+	{
+		Phrase.Add(0, FIoBuffer());
+		++EntryCount;
+	}
 
 	// Write the number of data entries in the first and last entry
 	check(EntryCount <= UINT16_MAX);
@@ -517,10 +528,7 @@ int32 FDiskJournal::Flush()
 
 	uint32 Size = Entries.Num() * sizeof(Entries[0]);
 
-	if (Cursor + Size > MaxSize)
-	{
-		Cursor = 0;
-	}
+	check(Cursor + Size <= MaxSize);
 
 	if (JrnHandle.IsValid())
 	{
@@ -532,6 +540,13 @@ int32 FDiskJournal::Flush()
 		Cursor += Size;
 	}
 
+	// We may end up exactly on the end of the journal file,
+	// immediately wrap in that case.
+	if (Cursor >= MaxSize)
+	{
+		Cursor = 0;
+	}
+
 	Entries.Reset();
 	OpenJrnFile();
 
@@ -1072,7 +1087,14 @@ static bool LoadCache(FDiskCache& DiskCache)
 	for (uint32 i = BasisIndex, n = Paragraphs.Num(); i < n; ++i)
 	{
 		const FPhraseDesc& Holm = Paragraphs[i].Phrase[0];
-		DiskCache.Insert(Holm.DataCursor, Holm.Entries, Paragraphs[i].EntryCount);
+		uint32 EntryCount = Paragraphs[i].EntryCount;
+		// The last entry may be a padded entry which should be skipped. See FDiskPhrase::ClosePhrase for details.
+		if (Holm.Entries[EntryCount-1].Key == 0)
+		{
+			check(Holm.Entries[EntryCount-1].Size == 0);
+			--EntryCount;
+		}
+		DiskCache.Insert(Holm.DataCursor, Holm.Entries, EntryCount);
 	}
 	
 	// Prime the journal's state
@@ -1265,7 +1287,7 @@ uint32 FCache::WriteMemToDisk(int32 Allowance)
 	for (int32 i = 0, n = PeelItems.Num(); i < n; ++i)
 	{
 		auto& [Key, Data] = PeelItems[i];
-		if (!Phrase.Add(Key, MoveTemp(Data)))
+		if (Phrase.GetRemainingEntries() < 1 || !Phrase.Add(Key, MoveTemp(Data)))
 		{
 			PeelIndex = i;
 			break;
@@ -2317,17 +2339,26 @@ static void CacheTests(FSupport& Support)
 {
 	using namespace JournaledCache;
 
+	// Small structure used to verify stored data
+	struct FStoredData
+	{
+		uint64 Key;
+		uint64 Size;
+		uint64 Hash;
+	};
+	
+	FCache::FConfig DefaultSettings;
+	DefaultSettings.Path = Support.TestDir / "cache_tests";
+	DefaultSettings.MemoryQuota = uint32(512_Ki);
+	DefaultSettings.DiskQuota = 8_Mi;
+	DefaultSettings.JournalQuota = uint32(7_Ki);
+
 	FCache* Cache = nullptr;
-	auto NewCache = [&Cache, &Support] (bool Drop=false)
+	auto NewCache = [&Cache, &Support, DefaultSettings] (bool Drop=false, const FCache::FConfig* Settings = nullptr)
 	{
 		delete Cache;
-		FCache::FConfig TestConfig;
-		TestConfig.Path = Support.TestDir / "cache_tests";
-		TestConfig.MemoryQuota = uint32(512_Ki);
-		TestConfig.DiskQuota = 8_Mi;
-		TestConfig.JournalQuota = uint32(7_Ki);
-		TestConfig.DropCache = Drop;
-		FCache::FConfig Config = TestConfig;
+		FCache::FConfig Config = Settings ? *Settings : DefaultSettings;
+		Config.DropCache = Drop;
 		Cache = new FCache(MoveTemp(Config));
 	};
 
@@ -2459,6 +2490,134 @@ static void CacheTests(FSupport& Support)
 
 	// don't load-and-sort so many paragraphs (only need max-data size)
 #endif // 0
+
+	{
+		FCache::FConfig Config;
+		Config.Path = Support.TestDir / "cache_jrn_wrap";
+		Config.MemoryQuota = uint32(4_Mi);
+		Config.DiskQuota = 16_Mi;
+		Config.JournalQuota = uint32(2_Ki - 1);
+		
+		NewCache(true, &Config);
+
+		auto PutAndCommit = [&] (int64 FillCnt, uint64 SizeMax, TArrayView<int32> Allowance) {
+			TArray<FStoredData> Ret;
+			for(uint32 i = 0; i < FillCnt; ++i)
+			{
+				const uint64 Size = Support.Mix() & (SizeMax - 1);
+				FIoBuffer Data = Support.DummyData(Size);
+				const uint64 Key = KeyGen(Data);
+				Ret.Push(FStoredData{Key, Size, CityHash64((const char*)Data.GetData(), uint32(Data.GetSize()))});
+				Cache->Put(Key, Data);
+				const uint32 AllowanceIdx = i % Allowance.Num();
+				if (Allowance[AllowanceIdx])
+				{
+					Cache->WriteMemToDisk(Allowance[AllowanceIdx]);
+					if (i % 3 == 0)
+					{
+						Cache->Flush();
+					}
+				}
+			}
+			Cache->Flush();
+			return Ret;
+		};
+
+		int32 Allowances[] = { int32(1_Mi) };
+		auto CommittedBuffers = PutAndCommit(2048, 16_Ki, Allowances);
+
+		NewCache(false, &Config);
+		Cache->Load();
+
+		uint32 Found(0), Lost(0);
+
+		for (const FStoredData& Committed : CommittedBuffers)
+		{
+			FIoBuffer Data;
+			if (const auto Token = Cache->Get(Committed.Key, Data); Token == Committed.Key)
+			{
+				check(Cache->Materialize(Token, Data, 0) == EIoErrorCode::Ok);
+			}
+			// Some items have been lost by this point, as expected
+			if (Data.GetSize())
+			{
+				check(Data.DataSize() == Committed.Size);
+				const uint64 Hash = CityHash64((const char*)Data.GetData(), uint32(Data.GetSize()));
+				check(Committed.Hash == Hash);
+				++Found;
+			}
+			else
+			{
+				++Lost;
+			}
+		}
+
+		UE_LOG(LogIas, VeryVerbose, TEXT("Journal wrap test found %u correct entries. %u entries were lost."), Found, Lost);
+	}
+	
+	{
+		FCache::FConfig Config;
+		Config.Path = Support.TestDir / "cache_jrn_wrap2";
+		Config.MemoryQuota = uint32(4_Mi);
+		Config.DiskQuota = uint64(16_Ki*819);
+		Config.JournalQuota = uint32(2_Ki);
+		
+		NewCache(true, &Config);
+
+		auto PutAndCommit = [&] (int64 FillCnt, uint64 SizeMax, TArrayView<int32> Allowance) {
+			TArray<FStoredData> Ret;
+			for(uint32 i = 0; i < FillCnt; ++i)
+			{
+				const uint64 Size = SizeMax;
+				FIoBuffer Data = Support.DummyData(Size);
+				const uint64 Key = KeyGen(Data);
+				Ret.Push(FStoredData{Key, Size, CityHash64((const char*)Data.GetData(), uint32(Data.GetSize()))});
+				Cache->Put(Key, Data);
+				const uint32 AllowanceIdx = i % Allowance.Num();
+				if (Allowance[AllowanceIdx])
+				{
+					Cache->WriteMemToDisk(Allowance[AllowanceIdx]);
+					if (i % 3 == 0)
+					{
+						Cache->Flush();
+					}
+				}
+			}
+			Cache->Flush();
+			return Ret;
+		};
+
+		int32 Allowances[] = { int32(1_Mi), 0, 0, 0, 0, int32(2_Mi), 0, 0, 0, int32(500_Ki), 0 };
+		auto CommittedBuffers = PutAndCommit(2048, 16_Ki, Allowances);
+
+		NewCache(false, &Config);
+		Cache->Load();
+
+		uint32 Found(0), Lost(0);
+
+		for (const FStoredData& Committed : CommittedBuffers)
+		{
+			FIoBuffer Data;
+			if (const auto Token = Cache->Get(Committed.Key, Data); Token == Committed.Key)
+			{
+				check(Cache->Materialize(Token, Data, 0) == EIoErrorCode::Ok);
+			}
+			// Some items have been lost by this point, as expected
+			if (Data.GetSize())
+			{
+				check(Data.DataSize() == Committed.Size);
+				const uint64 Hash = CityHash64((const char*)Data.GetData(), uint32(Data.GetSize()));
+				check(Committed.Hash == Hash);
+				++Found;
+			}
+			else
+			{
+				++Lost;
+			}
+		}
+
+		UE_LOG(LogIas, VeryVerbose, TEXT("Journal wrap test 2 found %u correct entries. %u entries were lost."), Found, Lost);
+	}
 }
 
 ////////////////////////////////////////////////////////////////////////////////
