Call FlushNetDormancy only when some replicated property are changing to avoid replicating all GCs and preventing the repmovement to be replicated

diff --git a/Engine/Source/Runtime/Engine/Private/PhysicsEngine/ClusterUnionReplicatedProxyComponent.cpp b/Engine/Source/Runtime/Engine/Private/PhysicsEngine/ClusterUnionReplicatedProxyComponent.cpp
--- a/Engine/Source/Runtime/Engine/Private/PhysicsEngine/ClusterUnionReplicatedProxyComponent.cpp
+++ b/Engine/Source/Runtime/Engine/Private/PhysicsEngine/ClusterUnionReplicatedProxyComponent.cpp
@@ -62,34 +62,41 @@ void UClusterUnionReplicatedProxyComponent::EndPlay(const EEndPlayReason::Type E
 
 void UClusterUnionReplicatedProxyComponent::SetParentClusterUnion(UClusterUnionComponent* InComponent)
 {
-	FlushNetDormancyIfNeeded();
+	if (ParentClusterUnion != InComponent)
+	{
+		FlushNetDormancyIfNeeded();
 
-	ParentClusterUnion = InComponent;
-	MARK_PROPERTY_DIRTY_FROM_NAME(UClusterUnionReplicatedProxyComponent, ParentClusterUnion, this);
+		ParentClusterUnion = InComponent;
+		MARK_PROPERTY_DIRTY_FROM_NAME(UClusterUnionReplicatedProxyComponent, ParentClusterUnion, this);
+	}
 }
 
 void UClusterUnionReplicatedProxyComponent::SetChildClusteredComponent(UPrimitiveComponent* InComponent)
 {
-	FlushNetDormancyIfNeeded();
+	if (ChildClusteredComponent != InComponent)
+	{
+		FlushNetDormancyIfNeeded();
 
-	ChildClusteredComponent = InComponent;
-	MARK_PROPERTY_DIRTY_FROM_NAME(UClusterUnionReplicatedProxyComponent, ChildClusteredComponent, this);
+		ChildClusteredComponent = InComponent;
+		MARK_PROPERTY_DIRTY_FROM_NAME(UClusterUnionReplicatedProxyComponent, ChildClusteredComponent, this);
+	}
 }
 
 void UClusterUnionReplicatedProxyComponent::SetParticleBoneIds(const TArray<int32>& InIds)
 {
-	FlushNetDormancyIfNeeded();
+	if(ParticleBoneIds != InIds)
+	{ 
+		FlushNetDormancyIfNeeded();
 
-	ParticleBoneIds = InIds;
+		ParticleBoneIds = InIds;
 
-	ParticleChildToParents.Empty();
-	ParticleChildToParents.Reserve(InIds.Num());
-	for (int32 Index = 0; Index < InIds.Num(); ++Index)
-	{
-		ParticleChildToParents.Add(FTransform::Identity);
+		ParticleChildToParents.Reset(InIds.Num());
+		for (int32 Index = 0; Index < InIds.Num(); ++Index)
+		{
+			ParticleChildToParents.Add(FTransform::Identity);
+		}
+		MARK_PROPERTY_DIRTY_FROM_NAME(UClusterUnionReplicatedProxyComponent, ParticleBoneIds, this);
 	}
-
-	MARK_PROPERTY_DIRTY_FROM_NAME(UClusterUnionReplicatedProxyComponent, ParticleBoneIds, this);
 }
 
 void UClusterUnionReplicatedProxyComponent::SetParticleChildToParent(int32 BoneId, const FTransform& ChildToParent)
@@ -97,10 +104,13 @@ void UClusterUnionReplicatedProxyComponent::SetParticleChildToParent(int32 BoneI
 	int32 Index = INDEX_NONE;
 	if (ParticleBoneIds.Find(BoneId, Index))
 	{
-		FlushNetDormancyIfNeeded();
+		if(!ParticleChildToParents[Index].Equals(ChildToParent))
+		{ 
+			FlushNetDormancyIfNeeded();
 
-		ParticleChildToParents[Index] = ChildToParent;
-		MARK_PROPERTY_DIRTY_FROM_NAME(UClusterUnionReplicatedProxyComponent, ParticleChildToParents, this);
+			ParticleChildToParents[Index] = ChildToParent;
+			MARK_PROPERTY_DIRTY_FROM_NAME(UClusterUnionReplicatedProxyComponent, ParticleChildToParents, this);
+		}
 	}
 }
 
diff --git a/Engine/Source/Runtime/Experimental/GeometryCollectionEngine/Private/GeometryCollection/GeometryCollectionComponent.cpp b/Engine/Source/Runtime/Experimental/GeometryCollectionEngine/Private/GeometryCollection/GeometryCollectionComponent.cpp
--- a/Engine/Source/Runtime/Experimental/GeometryCollectionEngine/Private/GeometryCollection/GeometryCollectionComponent.cpp
+++ b/Engine/Source/Runtime/Experimental/GeometryCollectionEngine/Private/GeometryCollection/GeometryCollectionComponent.cpp
@@ -208,6 +208,34 @@ bool FGeometryCollectionRepData::Identical(const FGeometryCollectionRepData* Oth
 	return Other && (Version == Other->Version);
 }
 
+bool FGeometryCollectionRepData::HasChanged(const FGeometryCollectionRepData& BaseData) const
+{
+	if((Version == 0) || (bIsRootAnchored != BaseData.bIsRootAnchored))
+	{
+		return true;
+	}
+	for (const FGeometryCollectionActivatedCluster& Data : OneOffActivated)
+	{
+		if (!BaseData.OneOffActivated.Contains(Data))
+		{
+			return true;
+		}
+	}
+	
+	for (const FGeometryCollectionClusterRep& Data : Clusters)
+	{
+		const FGeometryCollectionClusterRep* ExistingData = BaseData.Clusters.FindByPredicate(
+			[&TransformIndex = Data.ClusterIdx](const FGeometryCollectionClusterRep& OtherData) ->
+			bool { return OtherData.ClusterIdx == TransformIndex; });
+
+		if(!ExistingData || (ExistingData && Data.ClusterChanged(*ExistingData)))
+		{
+			return true;
+		}
+	}
+	return false;
+}
+
 bool FGeometryCollectionRepData::NetSerialize(FArchive& Ar, class UPackageMap* Map, bool& bOutSuccess)
 {
 	bOutSuccess = true;
@@ -286,6 +314,12 @@ bool FGeometryCollectionRepStateData::Identical(const FGeometryCollectionRepStat
 	return Other && (Version == Other->Version);
 }
 
+bool  FGeometryCollectionRepStateData::HasChanged(const FGeometryCollectionRepStateData& BaseData) const
+{
+	// We are not using the ReleasedData to check if something has changed since this array is fully controlled by the brokenstate one
+	return (Version == 0) || (BrokenState != BaseData.BrokenState) || (bIsRootAnchored != BaseData.bIsRootAnchored);
+}
+
 bool FGeometryCollectionRepStateData::NetSerialize(FArchive& Ar, class UPackageMap* Map, bool& bOutSuccess)
 {
 	bOutSuccess = true;
@@ -361,6 +395,26 @@ bool FGeometryCollectionRepDynamicData::Identical(const FGeometryCollectionRepDy
 	return Other && (Version == Other->Version);
 }
 
+bool FGeometryCollectionRepDynamicData::HasChanged(const FGeometryCollectionRepDynamicData& BaseData) const
+{
+	if (Version == 0)
+	{
+		return true;
+	}
+	for (const FGeometryCollectionRepDynamicData::FClusterData& Data : ClusterData)
+	{
+		const FGeometryCollectionRepDynamicData::FClusterData* ExistingData = BaseData.ClusterData.FindByPredicate(
+			[&TransformIndex = Data.TransformIndex](const FGeometryCollectionRepDynamicData::FClusterData& OtherData) ->
+			bool { return OtherData.TransformIndex == TransformIndex; });
+
+		if (!ExistingData || (ExistingData && !ExistingData->IsEqualPositionsAndVelocities(Data)))
+		{
+			return true;
+		}
+	}
+	return false;
+}
+
 bool FGeometryCollectionRepDynamicData::NetSerialize(FArchive& Ar, class UPackageMap* Map, bool& bOutSuccess)
 {
 	// we only support up to 2^16 bones for replication ( bandwidth reasons )
@@ -2020,8 +2074,6 @@ void UGeometryCollectionComponent::UpdateRepData()
 	
 	if (Owner && GetIsReplicated() && Owner->GetLocalRole() == ROLE_Authority)
 	{
-		FlushNetDormancyIfNeeded();
-
 		const bool bReplicateMovement = Owner->IsReplicatingMovement();
 		bool bFirstUpdate = false;
 		if(ClustersToRep == nullptr)
@@ -2041,7 +2093,11 @@ void UGeometryCollectionComponent::UpdateRepData()
 		//TODO: for now we have to iterate over all particles to find the clusters, would be better if we had the clusters and children already available
 		//We are relying on the fact that we fracture one level per step. This means we will see all one offs here
 
-		bool bClustersChanged = false;
+		FGeometryCollectionRepData LocalRepData;
+
+		LocalRepData.Version = RepData.Version;
+		LocalRepData.RepDataReceivedTime = RepData.RepDataReceivedTime;
+		LocalRepData.ServerFrame = RepData.ServerFrame;
 
 		FPBDRigidsSolver* Solver = PhysicsProxy->GetSolver<Chaos::FPBDRigidsSolver>();
 		const FRigidClustering& RigidClustering = Solver->GetEvolution()->GetRigidClustering();
@@ -2050,20 +2106,18 @@ void UGeometryCollectionComponent::UpdateRepData()
 		const TArray<FPBDRigidClusteredParticleHandle*>& ParticleHandles = PhysicsProxy->GetParticles();
 
 		// Replicate the anchored state of the root particle
+		LocalRepData.bIsRootAnchored = RepData.bIsRootAnchored;
 		const int32 InitialRootIndex = PhysicsProxy->GetSimParameters().InitialRootIndex;
 		if (FPBDRigidClusteredParticleHandle* RootHandle = ParticleHandles[InitialRootIndex])
 		{
-			const bool bIsRootAnchored = RootHandle->IsAnchored();
-			if (bFirstUpdate || bIsRootAnchored != RepData.bIsRootAnchored)
-			{
-				RepData.bIsRootAnchored = bIsRootAnchored;
-				bClustersChanged = true;
-			}
+			LocalRepData.bIsRootAnchored = RootHandle->IsAnchored();
 		}
-
+		
 		//see if we have any new clusters that are enabled
 		TSet<FPBDRigidClusteredParticleHandle*> Processed;
 
+		// It feels weird to start from the previous one. Need to be checked
+		LocalRepData.OneOffActivated = RepData.OneOffActivated;
 		for (FPBDRigidClusteredParticleHandle* Particle : ParticleHandles)
 		{
 			// Particle can be null if we have embedded geometry 
@@ -2115,7 +2169,6 @@ void UGeometryCollectionComponent::UpdateRepData()
 						if (Level <= ReplicationMaxPositionAndVelocityCorrectionLevel)
 						{
 							ClustersToRep->Add(Root);
-							bClustersChanged = true;
 						}
 					}
 
@@ -2130,10 +2183,9 @@ void UGeometryCollectionComponent::UpdateRepData()
 						// TODO: avoid search for entry for perf
 						// TODO: once we support deep fracture we should be able to remove one offs clusters that are now disabled, reducing the amount to be replicated
 						FGeometryCollectionActivatedCluster OneOffActivated(TransformGroupIdx, Root->V(), Root->W());
-						if(!RepData.OneOffActivated.Contains(OneOffActivated))
+						if(!LocalRepData.OneOffActivated.Contains(OneOffActivated))
 						{
-							bClustersChanged = true;
-							RepData.OneOffActivated.Add(OneOffActivated);
+							LocalRepData.OneOffActivated.Add(OneOffActivated);
 						}
 					}
 
@@ -2150,13 +2202,12 @@ void UGeometryCollectionComponent::UpdateRepData()
 			}
 		}
 
-		INC_DWORD_STAT_BY(STAT_GCReplicatedFractures, RepData.OneOffActivated.Num());
+		INC_DWORD_STAT_BY(STAT_GCReplicatedFractures, LocalRepData.OneOffActivated.Num());
 		
 		//Build up clusters to replicate and compare with previous frame
-		TArray<FGeometryCollectionClusterRep> Clusters;
-
 		if (bReplicateMovement)
 		{
+			LocalRepData.Clusters.Reserve(ClustersToRep->Num());
 			//remove disabled clusters and update rep data if needed
 			for (auto Itr = ClustersToRep->CreateIterator(); Itr; ++Itr)
 			{
@@ -2171,8 +2222,7 @@ void UGeometryCollectionComponent::UpdateRepData()
 				}
 				else
 				{
-					Clusters.AddDefaulted();
-					FGeometryCollectionClusterRep& ClusterRep = Clusters.Last();
+					FGeometryCollectionClusterRep& ClusterRep = LocalRepData.Clusters.AddDefaulted_GetRef();
 
 					ClusterRep.Position = Cluster->X();
 					ClusterRep.Rotation = Cluster->R();
@@ -2183,7 +2233,6 @@ void UGeometryCollectionComponent::UpdateRepData()
 					int32 TransformGroupIdx;
 					if (Cluster->InternalCluster())
 					{
-
 						ensureMsgf(Children.Num(), TEXT("Internal cluster yet we have no children? [Num %d vs Cached Empty %d]"), Children.Num(), bChildrenEmpty);
 						TransformGroupIdx = PhysicsProxy->GetTransformGroupIndexFromHandle(Children[0]);
 					}
@@ -2195,38 +2244,21 @@ void UGeometryCollectionComponent::UpdateRepData()
 
 					ensureMsgf(TransformGroupIdx < TNumericLimits<uint16>::Max(), TEXT("Trying to replicate GC with more than 65k pieces. We assumed uint16 would suffice"));
 					ClusterRep.ClusterIdx = TransformGroupIdx;
-
-					if (!bClustersChanged)
-					{
-						//compare to previous frame data
-						// this could be more efficient by having a way to find back the data from the idx
-						auto Predicate = [TransformGroupIdx](const FGeometryCollectionClusterRep& Entry)
-						{
-							return Entry.ClusterIdx == TransformGroupIdx;
-						};
-						if (const FGeometryCollectionClusterRep* PrevClusterData = RepData.Clusters.FindByPredicate(Predicate))
-						{
-							if (ClusterRep.ClusterChanged(*PrevClusterData))
-							{
-								bClustersChanged = true;
-							}
-						}
-					}
 				}
 			}
 		}
 
-		if (bClustersChanged)
+		if (LocalRepData.HasChanged(RepData))
 		{
-			RepData.Clusters = MoveTemp(Clusters);
-
 			if (Owner->GetWorld() && Owner->GetWorld()->GetPhysicsScene())
 			{
-				RepData.ServerFrame = Owner->GetWorld()->GetPhysicsScene()->ReplicationCache.ServerFrame;
+				LocalRepData.ServerFrame = Owner->GetWorld()->GetPhysicsScene()->ReplicationCache.ServerFrame;
 			}
+			INC_DWORD_STAT_BY(STAT_GCReplicatedClusters, LocalRepData.Clusters.Num());
 
-			INC_DWORD_STAT_BY(STAT_GCReplicatedClusters, RepData.Clusters.Num());
+			FlushNetDormancyIfNeeded();
 
+			RepData = MoveTemp(LocalRepData);
 			MARK_PROPERTY_DIRTY_FROM_NAME(UGeometryCollectionComponent, RepData, this);
 			++RepData.Version;
 
@@ -2263,32 +2295,28 @@ void UGeometryCollectionComponent::UpdateRepStateAndDynamicData()
 
 	if (GetIsReplicated() && Owner->GetLocalRole() == ROLE_Authority)
 	{
-		FlushNetDormancyIfNeeded();
-
 		if (RestCollection && RestCollection->GetGeometryCollection())
 		{
 			FPBDRigidsSolver* Solver = PhysicsProxy->GetSolver<Chaos::FPBDRigidsSolver>();
-			//const FRigidClustering& RigidClustering = Solver->GetEvolution()->GetRigidClustering();
 
 			// let go through all the transform and see which one has changed its state
 			const TArray<FPBDRigidClusteredParticleHandle*>& ParticleHandles = PhysicsProxy->GetParticles();
 			const int32 NumTransforms = ParticleHandles.Num();
 
 			const TManagedArrayAccessor<int32> InitialLevels = PhysicsProxy->GetPhysicsCollection().GetInitialLevels();
+			
+			FGeometryCollectionRepStateData LocalRepStateData;
+			FGeometryCollectionRepDynamicData LocalRepDynamicData;
 
-			bool bStateChanged = false;
-			bool bDynamicChanged = false;
+			LocalRepStateData.Version = RepStateData.Version;
+			LocalRepDynamicData.Version = RepDynamicData.Version;
 
 			// root level anchor state
+			LocalRepStateData.bIsRootAnchored = RepStateData.bIsRootAnchored;
 			const int32 InitialRootIndex = PhysicsProxy->GetSimParameters().InitialRootIndex;
 			if (FPBDRigidClusteredParticleHandle* RootHandle = ParticleHandles[InitialRootIndex])
 			{
-				const bool bIsRootAnchored = RootHandle->IsAnchored();
-				if (RepStateData.Version == 0 || bIsRootAnchored != (bool)RepStateData.bIsRootAnchored)
-				{
-					RepStateData.bIsRootAnchored = bIsRootAnchored? 1: 0;
-					bStateChanged = true;
-				}
+				LocalRepStateData.bIsRootAnchored = RootHandle->IsAnchored() ? 1 : 0;
 			}
 
 			struct FRootHandle
@@ -2307,6 +2335,8 @@ void UGeometryCollectionComponent::UpdateRepStateAndDynamicData()
 			const int32 RootIndex = GetRootIndex();
 
 			// go through particles and send replicated data 
+			LocalRepStateData.BrokenState.SetNum(NumTransforms, false);
+			LocalRepStateData.ReleasedData.Reserve(NumTransforms);
 			for (int32 TransformIndex = 0; TransformIndex < NumTransforms; TransformIndex++)
 			{
 				if (FPBDRigidClusteredParticleHandle* ParticleHandle = ParticleHandles[TransformIndex])
@@ -2316,13 +2346,20 @@ void UGeometryCollectionComponent::UpdateRepStateAndDynamicData()
 					const FPBDRigidClusteredParticleHandle* ParentHandle = ParticleHandle->Parent();
 
 					// no parent and not root means the particle has been broken  
-					const bool bIsBroken = (ParentHandle == nullptr && TransformIndex != GetRootIndex());
+					const bool bIsBroken = (ParentHandle == nullptr && TransformIndex != RootIndex);
 					if (bIsBroken)
 					{
 						// only record the one at the abandon level and before ( child break of abandon level clusters must be recorded )  
 						if (!bEnableAbandonAfterLevel || Level <= ReplicationAbandonAfterLevel)
 						{
-							bStateChanged |= RepStateData.SetBroken(TransformIndex, NumTransforms, ParticleHandle->Disabled(), ParticleHandle->V(), ParticleHandle->W());
+							LocalRepStateData.BrokenState[TransformIndex] = true;
+							if(!ParticleHandle->Disabled())
+							{
+								FGeometryCollectionRepStateData::FReleasedData& Data = LocalRepStateData.ReleasedData.AddDefaulted_GetRef();
+								Data.TransformIndex = TransformIndex;
+								Data.LinearVelocity = ParticleHandle->V();
+								Data.AngularVelocityInDegreesPerSecond = FMath::RadiansToDegrees(ParticleHandle->W());
+							}
 						}
 
 						// anything beyond ReplicationAbandonAfterLevel must be disabled to save server CPU as they are now client authoritative
@@ -2351,28 +2388,38 @@ void UGeometryCollectionComponent::UpdateRepStateAndDynamicData()
 			}
 
 			// let's update the data for the tracked clusters 
+			LocalRepDynamicData.ClusterData.Reserve(RootsToTrack.Num());
 			for (const FRootHandle& Root: RootsToTrack)
 			{
-				FGeometryCollectionRepDynamicData::FClusterData Data;
+				FGeometryCollectionRepDynamicData::FClusterData& Data = LocalRepDynamicData.ClusterData.AddDefaulted_GetRef();
 				Data.TransformIndex = Root.TransformIndex;
 				Data.bIsInternalCluster = Root.Handle->InternalCluster();
 				Data.Position = Root.Handle->X();
 				Data.EulerRotation = Root.Handle->R().Rotator().Euler();
 				Data.LinearVelocity = Root.Handle->V();
 				Data.AngularVelocityInDegreesPerSecond = FMath::RadiansToDegrees(Root.Handle->W());
-				bDynamicChanged |= RepDynamicData.SetData(Data);
+				Data.LastUpdatedVersion = LocalRepDynamicData.Version + 1;
 			}
 
-			RepDynamicData.RemoveOutOfDateClusterData();
+			// Check to see if the replicated state and dynamic Data has changed
+			const bool bStateChanged = LocalRepStateData.HasChanged(RepStateData);
+			const bool bDynamicChanged = LocalRepDynamicData.HasChanged(RepDynamicData);
+
+			if(bStateChanged || bDynamicChanged)
+			{
+				FlushNetDormancyIfNeeded();
+			}
 
 			if (bStateChanged)
 			{
+				RepStateData = MoveTemp(LocalRepStateData);
 				MARK_PROPERTY_DIRTY_FROM_NAME(UGeometryCollectionComponent, RepStateData, this);
 				++RepStateData.Version;
 			}
 
 			if (bDynamicChanged)
 			{
+				RepDynamicData = MoveTemp(LocalRepDynamicData);
 				MARK_PROPERTY_DIRTY_FROM_NAME(UGeometryCollectionComponent, RepDynamicData, this);
 				++RepDynamicData.Version;
 			}
diff --git a/Engine/Source/Runtime/Experimental/GeometryCollectionEngine/Public/GeometryCollection/GeometryCollectionComponent.h b/Engine/Source/Runtime/Experimental/GeometryCollectionEngine/Public/GeometryCollection/GeometryCollectionComponent.h
--- a/Engine/Source/Runtime/Experimental/GeometryCollectionEngine/Public/GeometryCollection/GeometryCollectionComponent.h
+++ b/Engine/Source/Runtime/Experimental/GeometryCollectionEngine/Public/GeometryCollection/GeometryCollectionComponent.h
@@ -423,6 +423,9 @@ struct FGeometryCollectionRepData
 	bool Identical(const FGeometryCollectionRepData* Other, uint32 PortFlags) const;
 	bool NetSerialize(FArchive& Ar, class UPackageMap* Map, bool& bOutSuccess);
 
+	// Check if the data has changed
+	bool HasChanged(const FGeometryCollectionRepData& BaseData) const;
+
 	void Reset()
 	{
 		OneOffActivated.Reset();
@@ -485,6 +488,9 @@ struct FGeometryCollectionRepStateData
 	bool Identical(const FGeometryCollectionRepStateData* Other, uint32 PortFlags) const;
 	bool NetSerialize(FArchive& Ar, class UPackageMap* Map, bool& bOutSuccess);
 
+	// Check if the data has changed
+	bool HasChanged(const FGeometryCollectionRepStateData& BaseData) const;
+
 	void Reset()
 	{
 		BrokenState.Reset();
@@ -546,6 +552,9 @@ struct FGeometryCollectionRepDynamicData
 	bool Identical(const FGeometryCollectionRepDynamicData* Other, uint32 PortFlags) const;
 	bool NetSerialize(FArchive& Ar, class UPackageMap* Map, bool& bOutSuccess);
 
+	// Check if the data has changed
+	bool HasChanged(const FGeometryCollectionRepDynamicData& BaseData) const;
+
 	void Reset()
 	{
 		ClusterData.Reset();
