Fix non-unity

diff --git a/Engine/Source/Runtime/Experimental/IoStoreOnDemand/Private/OnDemandHttpClient.cpp b/Engine/Source/Runtime/Experimental/IoStoreOnDemand/Private/OnDemandHttpClient.cpp
--- a/Engine/Source/Runtime/Experimental/IoStoreOnDemand/Private/OnDemandHttpClient.cpp
+++ b/Engine/Source/Runtime/Experimental/IoStoreOnDemand/Private/OnDemandHttpClient.cpp
@@ -6,7 +6,6 @@
 #include "Containers/StringConv.h"
 #include "HAL/IConsoleManager.h"
 #include "HAL/PlatformTime.h"
-#include "IO/IoBuffer.h"
 #include "IO/IoStoreOnDemand.h"
 
 namespace UE::IO::IAS
diff --git a/Engine/Source/Runtime/Experimental/IoStoreOnDemand/Private/OnDemandHttpClient.h b/Engine/Source/Runtime/Experimental/IoStoreOnDemand/Private/OnDemandHttpClient.h
--- a/Engine/Source/Runtime/Experimental/IoStoreOnDemand/Private/OnDemandHttpClient.h
+++ b/Engine/Source/Runtime/Experimental/IoStoreOnDemand/Private/OnDemandHttpClient.h
@@ -5,6 +5,7 @@
 #include "CoreHttp/Client.h"
 #include "Containers/StringView.h"
 #include "Containers/UnrealString.h"
+#include "IO/IoBuffer.h"
 #include "IO/IoOffsetLength.h"
 #include "IO/IoStatus.h"
 #include "Templates/Function.h"
