Ias: Set ias.HttpRangeRequestMinSizeKiB=128 to reduce the number of small http requests.

#rnx

#localization none
#tests frontend and replays on pc and console

diff --git a/Engine/Source/Runtime/Experimental/IoStoreOnDemand/Private/OnDemandIoDispatcherBackend.cpp b/Engine/Source/Runtime/Experimental/IoStoreOnDemand/Private/OnDemandIoDispatcherBackend.cpp
--- a/Engine/Source/Runtime/Experimental/IoStoreOnDemand/Private/OnDemandIoDispatcherBackend.cpp
+++ b/Engine/Source/Runtime/Experimental/IoStoreOnDemand/Private/OnDemandIoDispatcherBackend.cpp
@@ -213,7 +213,7 @@ static FAutoConsoleVariableRef CVar_IasAsyncTocGeneration(
 	ECVF_ReadOnly
 );
 
-static int32 GIasHttpRangeRequestMinSizeKiB = 0;
+static int32 GIasHttpRangeRequestMinSizeKiB = 128;
 static FAutoConsoleVariableRef CVar_IasHttpRangeRequestMinSizeKiB(
 	TEXT("ias.HttpRangeRequestMinSizeKiB"),
 	GIasHttpRangeRequestMinSizeKiB,
