Memory optimization when creating Geometry Collection
Don't create implicit geometry on children particle when GC is created.
Wait until destruction occurs to create them.

diff --git a/Engine/Source/Runtime/Experimental/Chaos/Private/Chaos/PBDRigidClustering.cpp b/Engine/Source/Runtime/Experimental/Chaos/Private/Chaos/PBDRigidClustering.cpp
--- a/Engine/Source/Runtime/Experimental/Chaos/Private/Chaos/PBDRigidClustering.cpp
+++ b/Engine/Source/Runtime/Experimental/Chaos/Private/Chaos/PBDRigidClustering.cpp
@@ -1062,7 +1062,7 @@ namespace Chaos
 		bool bUseDamagePropagation = false;
 		float BreakDamagePropagationFactor = 0.0f;
 		float ShockDamagePropagationFactor = 0.0f;
-		const FGeometryCollectionPhysicsProxy* ConcreteGCProxy = GetConcreteProxy<FGeometryCollectionPhysicsProxy>(ClusteredParticle);
+		FGeometryCollectionPhysicsProxy* ConcreteGCProxy = GetConcreteProxy<FGeometryCollectionPhysicsProxy>(ClusteredParticle);
 		if (ConcreteGCProxy)
 		{
 			const FSimulationParameters& SimParams = ConcreteGCProxy->GetSimParameters();
@@ -1215,6 +1215,10 @@ namespace Chaos
 
 		if (ActivatedChildren.Num() > 0)
 		{
+			if (ConcreteGCProxy)
+			{
+				ConcreteGCProxy->CreateChildrenGeometry_Internal();
+			}
 			const bool bIsClusterUnion = ClusterUnionManager.IsClusterUnionParticle(ClusteredParticle);
 			if (Children.Num() == 0)
 			{
diff --git a/Engine/Source/Runtime/Experimental/Chaos/Private/PhysicsProxy/GeometryCollectionPhysicsProxy.cpp b/Engine/Source/Runtime/Experimental/Chaos/Private/PhysicsProxy/GeometryCollectionPhysicsProxy.cpp
--- a/Engine/Source/Runtime/Experimental/Chaos/Private/PhysicsProxy/GeometryCollectionPhysicsProxy.cpp
+++ b/Engine/Source/Runtime/Experimental/Chaos/Private/PhysicsProxy/GeometryCollectionPhysicsProxy.cpp
@@ -271,72 +271,15 @@ static float AreaFromBoundingBoxOverlap(const Chaos::FAABB3& BoxA, const Chaos::
 	return static_cast<float>(Area);
 }
 
-
-DECLARE_CYCLE_STAT(TEXT("FGeometryCollectionPhysicsProxy::PopulateSimulatedParticle"), STAT_PopulateSimulatedParticle, STATGROUP_Chaos);
-void PopulateSimulatedParticle(
-	Chaos::TPBDRigidParticleHandle<Chaos::FReal,3>* Handle,
+void SetImplicitToPTParticles(Chaos::TPBDRigidParticleHandle<Chaos::FReal, 3>* Handle,
 	const FSharedSimulationParameters& SharedParams,
 	const FCollisionStructureManager::FSimplicial* Simplicial,
 	Chaos::FImplicitObjectPtr Implicit,
 	const FCollisionFilterData SimFilterIn,
 	const FCollisionFilterData QueryFilterIn,
-	Chaos::FReal MassIn,
-	Chaos::TVec3<Chaos::FRealSingle> InertiaTensorVec,
-	const FTransform& WorldTransform, 
-	const uint8 DynamicState, 
-	const int16 CollisionGroup,
+	const FTransform& WorldTransform,
 	float CollisionParticlesPerObjectFraction)
 {
-	SCOPE_CYCLE_COUNTER(STAT_PopulateSimulatedParticle);
-	Handle->SetDisabledLowLevel(false);
-	Handle->SetX(WorldTransform.GetTranslation());
-	Handle->SetV(Chaos::FVec3(0.f));
-	Handle->SetR(WorldTransform.GetRotation().GetNormalized());
-	Handle->SetW(Chaos::FVec3(0.f));
-	Handle->SetP(Handle->X());
-	Handle->SetQ(Handle->R());
-	Handle->SetCenterOfMass(FVector3f::ZeroVector);
-	Handle->SetRotationOfMass(FQuat::Identity);
-
-	//
-	// Setup Mass
-	//
-	{
-		Handle->SetObjectStateLowLevel(Chaos::EObjectStateType::Uninitialized);
-
-		if (!CHAOS_ENSURE_MSG(FMath::IsWithinInclusive<Chaos::FReal>(MassIn, SharedParams.MinimumMassClamp, SharedParams.MaximumMassClamp),
-			TEXT("Clamped mass[%3.5f] to range [%3.5f,%3.5f]"), MassIn, SharedParams.MinimumMassClamp, SharedParams.MaximumMassClamp))
-		{
-			MassIn = FMath::Clamp<Chaos::FReal>(MassIn, SharedParams.MinimumMassClamp, SharedParams.MaximumMassClamp);
-		}
-
-		if (!CHAOS_ENSURE_MSG(!FMath::IsNaN(InertiaTensorVec[0]) && !FMath::IsNaN(InertiaTensorVec[1]) && !FMath::IsNaN(InertiaTensorVec[2]),
-			TEXT("Nan Tensor, reset to unit tesor")))
-		{
-			InertiaTensorVec = FVector3f(1);
-		}
-		else if (!CHAOS_ENSURE_MSG(FMath::IsWithinInclusive(InertiaTensorVec[0], SharedParams.MinimumInertiaTensorDiagonalClamp, SharedParams.MaximumInertiaTensorDiagonalClamp)
-			&& FMath::IsWithinInclusive(InertiaTensorVec[1], SharedParams.MinimumInertiaTensorDiagonalClamp, SharedParams.MaximumInertiaTensorDiagonalClamp)
-			&& FMath::IsWithinInclusive(InertiaTensorVec[2], SharedParams.MinimumInertiaTensorDiagonalClamp, SharedParams.MaximumInertiaTensorDiagonalClamp),
-			TEXT("Clamped Inertia tensor[%3.5f,%3.5f,%3.5f]. Clamped each element to [%3.5f, %3.5f,]"), InertiaTensorVec[0], InertiaTensorVec[1], InertiaTensorVec[2],
-			SharedParams.MinimumInertiaTensorDiagonalClamp, SharedParams.MaximumInertiaTensorDiagonalClamp))
-		{
-			InertiaTensorVec[0] = FMath::Clamp(InertiaTensorVec[0], SharedParams.MinimumInertiaTensorDiagonalClamp, SharedParams.MaximumInertiaTensorDiagonalClamp);
-			InertiaTensorVec[1] = FMath::Clamp(InertiaTensorVec[1], SharedParams.MinimumInertiaTensorDiagonalClamp, SharedParams.MaximumInertiaTensorDiagonalClamp);
-			InertiaTensorVec[2] = FMath::Clamp(InertiaTensorVec[2], SharedParams.MinimumInertiaTensorDiagonalClamp, SharedParams.MaximumInertiaTensorDiagonalClamp);
-		}
-
-		Handle->SetM(MassIn);
-		Handle->SetI(InertiaTensorVec);
-		const Chaos::FReal MassInv = (MassIn > 0.0f) ? 1.0f / MassIn : 0.0f;
-		const Chaos::FVec3 InertiaInv = (MassIn > 0.0f) ? Chaos::FVec3(InertiaTensorVec).Reciprocal() : Chaos::FVec3::ZeroVector;
-		Handle->SetInvM(MassInv);
-		Handle->SetInvI(InertiaInv);
-		Handle->SetObjectStateLowLevel(Chaos::EObjectStateType::Dynamic); // this step sets InvM, InvInertia, P, Q
-	}
-
-	Handle->SetCollisionGroup(CollisionGroup);
-
 	// @todo(GCCollisionShapes) : add support for multiple shapes, currently just one. 
 	FCollectionCollisionTypeData SingleSupportedCollisionTypeData = FCollectionCollisionTypeData();
 	if (SharedParams.SizeSpecificData.Num() && SharedParams.SizeSpecificData[0].CollisionShapesData.Num())
@@ -360,7 +303,7 @@ void PopulateSimulatedParticle(
 
 		Chaos::EImplicitObjectType ImplicitType = Implicit->GetType();
 		// Don't copy if it is not a level set and scale is one
-		if (SingleSupportedCollisionTypeData.CollisionType != ECollisionTypeEnum::Chaos_Surface_Volumetric && 
+		if (SingleSupportedCollisionTypeData.CollisionType != ECollisionTypeEnum::Chaos_Surface_Volumetric &&
 			ImplicitType != Chaos::ImplicitObjectType::LevelSet && Scale.Equals(FVector::OneVector))
 		{
 			Handle->SetGeometry(Implicit);
@@ -401,7 +344,7 @@ void PopulateSimulatedParticle(
 		for (int32 ParticleIndex = 0; ParticleIndex < (int32)CollisionParticles->Size(); ++ParticleIndex)
 		{
 			CollisionParticles->X(ParticleIndex) *= Scale;
-			
+
 			// Make sure the collision particles are at least in the domain 
 			// of the implicit shape.
 			ensure(ImplicitShapeDomain.Contains(CollisionParticles->X(ParticleIndex)));
@@ -440,6 +383,74 @@ void PopulateSimulatedParticle(
 			Shape->SetQueryData(QueryFilterIn);
 		}
 	}
+}
+
+DECLARE_CYCLE_STAT(TEXT("FGeometryCollectionPhysicsProxy::PopulateSimulatedParticle"), STAT_PopulateSimulatedParticle, STATGROUP_Chaos);
+void PopulateSimulatedParticle(
+	Chaos::TPBDRigidParticleHandle<Chaos::FReal,3>* Handle,
+	const FSharedSimulationParameters& SharedParams,
+	const FCollisionStructureManager::FSimplicial* Simplicial,
+	Chaos::FImplicitObjectPtr Implicit,
+	const FCollisionFilterData SimFilterIn,
+	const FCollisionFilterData QueryFilterIn,
+	Chaos::FReal MassIn,
+	Chaos::TVec3<Chaos::FRealSingle> InertiaTensorVec,
+	const FTransform& WorldTransform, 
+	const uint8 DynamicState, 
+	const int16 CollisionGroup,
+	float CollisionParticlesPerObjectFraction)
+{
+	SCOPE_CYCLE_COUNTER(STAT_PopulateSimulatedParticle);
+	Handle->SetDisabledLowLevel(false);
+	Handle->SetX(WorldTransform.GetTranslation());
+	Handle->SetV(Chaos::FVec3(0.f));
+	Handle->SetR(WorldTransform.GetRotation().GetNormalized());
+	Handle->SetW(Chaos::FVec3(0.f));
+	Handle->SetP(Handle->X());
+	Handle->SetQ(Handle->R());
+	Handle->SetCenterOfMass(FVector3f::ZeroVector);
+	Handle->SetRotationOfMass(FQuat::Identity);
+
+	//
+	// Setup Mass
+	//
+	{
+		Handle->SetObjectStateLowLevel(Chaos::EObjectStateType::Uninitialized);
+
+		if (!CHAOS_ENSURE_MSG(FMath::IsWithinInclusive<Chaos::FReal>(MassIn, SharedParams.MinimumMassClamp, SharedParams.MaximumMassClamp),
+			TEXT("Clamped mass[%3.5f] to range [%3.5f,%3.5f]"), MassIn, SharedParams.MinimumMassClamp, SharedParams.MaximumMassClamp))
+		{
+			MassIn = FMath::Clamp<Chaos::FReal>(MassIn, SharedParams.MinimumMassClamp, SharedParams.MaximumMassClamp);
+		}
+
+		if (!CHAOS_ENSURE_MSG(!FMath::IsNaN(InertiaTensorVec[0]) && !FMath::IsNaN(InertiaTensorVec[1]) && !FMath::IsNaN(InertiaTensorVec[2]),
+			TEXT("Nan Tensor, reset to unit tesor")))
+		{
+			InertiaTensorVec = FVector3f(1);
+		}
+		else if (!CHAOS_ENSURE_MSG(FMath::IsWithinInclusive(InertiaTensorVec[0], SharedParams.MinimumInertiaTensorDiagonalClamp, SharedParams.MaximumInertiaTensorDiagonalClamp)
+			&& FMath::IsWithinInclusive(InertiaTensorVec[1], SharedParams.MinimumInertiaTensorDiagonalClamp, SharedParams.MaximumInertiaTensorDiagonalClamp)
+			&& FMath::IsWithinInclusive(InertiaTensorVec[2], SharedParams.MinimumInertiaTensorDiagonalClamp, SharedParams.MaximumInertiaTensorDiagonalClamp),
+			TEXT("Clamped Inertia tensor[%3.5f,%3.5f,%3.5f]. Clamped each element to [%3.5f, %3.5f,]"), InertiaTensorVec[0], InertiaTensorVec[1], InertiaTensorVec[2],
+			SharedParams.MinimumInertiaTensorDiagonalClamp, SharedParams.MaximumInertiaTensorDiagonalClamp))
+		{
+			InertiaTensorVec[0] = FMath::Clamp(InertiaTensorVec[0], SharedParams.MinimumInertiaTensorDiagonalClamp, SharedParams.MaximumInertiaTensorDiagonalClamp);
+			InertiaTensorVec[1] = FMath::Clamp(InertiaTensorVec[1], SharedParams.MinimumInertiaTensorDiagonalClamp, SharedParams.MaximumInertiaTensorDiagonalClamp);
+			InertiaTensorVec[2] = FMath::Clamp(InertiaTensorVec[2], SharedParams.MinimumInertiaTensorDiagonalClamp, SharedParams.MaximumInertiaTensorDiagonalClamp);
+		}
+
+		Handle->SetM(MassIn);
+		Handle->SetI(InertiaTensorVec);
+		const Chaos::FReal MassInv = (MassIn > 0.0f) ? 1.0f / MassIn : 0.0f;
+		const Chaos::FVec3 InertiaInv = (MassIn > 0.0f) ? Chaos::FVec3(InertiaTensorVec).Reciprocal() : Chaos::FVec3::ZeroVector;
+		Handle->SetInvM(MassInv);
+		Handle->SetInvI(InertiaInv);
+		Handle->SetObjectStateLowLevel(Chaos::EObjectStateType::Dynamic); // this step sets InvM, InvInertia, P, Q
+	}
+
+	Handle->SetCollisionGroup(CollisionGroup);
+
+	SetImplicitToPTParticles(Handle, SharedParams, Simplicial, Implicit, SimFilterIn, QueryFilterIn, WorldTransform, CollisionParticlesPerObjectFraction);
 
 	//
 	//  Manage Object State
@@ -497,6 +508,8 @@ FGeometryCollectionPhysicsProxy::FGeometryCollectionPhysicsProxy(
 	, GameThreadCollection(GameThreadCollectionIn)
 	, WorldTransform_External(FTransform::Identity)
 	, bIsGameThreadWorldTransformDirty(false)
+	, bHasBuiltGeometryOnPT(false)
+	, bHasBuiltGeometryOnGT(false)
 	, CollectorGuid(InCollectorGuid)
 {
 	// We rely on a guarded buffer.
@@ -510,6 +523,9 @@ FGeometryCollectionPhysicsProxy::~FGeometryCollectionPhysicsProxy()
 float ReportHighParticleFraction = -1.f;
 FAutoConsoleVariableRef CVarReportHighParticleFraction(TEXT("p.gc.ReportHighParticleFraction"), ReportHighParticleFraction, TEXT("Report any objects with particle fraction above this threshold"));
 
+bool bBuildGeometryForChildren = true;
+FAutoConsoleVariableRef CVarbBuildGeometryForChildren(TEXT("p.gc.BuildGeometryForChildren"), bBuildGeometryForChildren, TEXT("If true build all children geometry at initilaization time, otherwise wait until destruction occurs."));
+
 void FGeometryCollectionPhysicsProxy::Initialize(Chaos::FPBDRigidsEvolutionBase *Evolution)
 {
 	check(IsInGameThread());
@@ -576,22 +592,101 @@ void FGeometryCollectionPhysicsProxy::Initialize(Chaos::FPBDRigidsEvolutionBase
 	const Chaos::Facades::FCollectionAnchoringFacade DynamicCollectionAnchoringFacade(DynamicCollection);
 	Chaos::Facades::FCollectionAnchoringFacade PhysicsThreadCollectionAnchoringFacade(PhysicsThreadCollection);
 	PhysicsThreadCollectionAnchoringFacade.CopyAnchoredAttribute(DynamicCollectionAnchoringFacade);
-	
+
+	const FVector Scale = Parameters.WorldTransform.GetScale3D();
+	const TManagedArray<float>& Mass = Parameters.RestCollection->GetAttribute<float>(MassAttributeName, FTransformCollection::TransformGroup);
+
+	TManagedArray<Chaos::FImplicitObjectPtr>& Implicits = GameThreadCollection.ModifyAttribute<Chaos::FImplicitObjectPtr>(FGeometryDynamicCollection::ImplicitsAttribute, FTransformCollection::TransformGroup);
+	if (ensure(NumTransforms == Implicits.Num() && NumTransforms == GTParticles.Num())) // Implicits are in the transform group so this invariant should always hold
+	{
+		constexpr bool bInitializeRootOnly = true;
+		bHasBuiltGeometryOnGT = bBuildGeometryForChildren;
+		CreateParticles(EffectiveParticles, Implicits, Evolution, bInitializeRootOnly);
+
+		// Skip simplicials, as they're owned by unique pointers.
+		static const FAttributeAndGroupId SkipList[] =
+		{
+			{ FGeometryDynamicCollection::SimplicialsAttribute, FTransformCollection::TransformGroup },
+		};
+		static constexpr int32 SkipListSize = sizeof(SkipList) / sizeof(FAttributeAndGroupId);
+
+		// Adding the Implicits to PhysicsThreadCollection, before it copies all matching attributes
+		PhysicsThreadCollection.AddAttribute<Chaos::FImplicitObjectPtr>(FGeometryDynamicCollection::ImplicitsAttribute, FTransformCollection::TransformGroup);
+		PhysicsThreadCollection.CopyMatchingAttributesFrom(DynamicCollection, MakeArrayView(SkipList, SkipListSize));
+		PhysicsThreadCollection.CopyInitialVelocityAttributesFrom(DynamicCollection);
+
+		// Copy simplicials.
+		// TODO: Ryan - Should we just transfer ownership of the SimplicialsAttribute from the DynamicCollection to
+		// the PhysicsThreadCollection?
+		{
+			if (FGeometryCollection::AreCollisionParticlesEnabled()
+				&& DynamicCollection.HasAttribute(DynamicCollection.SimplicialsAttribute, FTransformCollection::TransformGroup))
+			{
+				const auto& SourceSimplicials = DynamicCollection.GetAttribute<TUniquePtr<FSimplicial>>(
+					DynamicCollection.SimplicialsAttribute, FTransformCollection::TransformGroup);
+				for (int32 Index = PhysicsThreadCollection.NumElements(FTransformCollection::TransformGroup) - 1; 0 <= Index; Index--)
+				{
+					PhysicsThreadCollection.Simplicials[Index].Reset(
+						SourceSimplicials[Index] ? SourceSimplicials[Index]->NewCopy() : nullptr);
+				}
+			}
+			else
+			{
+				for (int32 Index = PhysicsThreadCollection.NumElements(FTransformCollection::TransformGroup) - 1; 0 <= Index; Index--)
+				{
+					PhysicsThreadCollection.Simplicials[Index].Reset();
+				}
+			}
+		}
+
+
+		if (Parameters.EnableClustering)
+		{
+			// make sure we set Activate the right way when clustering is enabled ( only root should be enabled at start ) 
+			for (int32 TransformIndex = 0; TransformIndex < NumTransforms; ++TransformIndex)
+			{
+				const bool bIsRoot = !GameThreadCollection.GetHasParent(TransformIndex);
+				GameThreadCollection.Active[TransformIndex] = bIsRoot;
+				PhysicsThreadCollection.Active[TransformIndex] = bIsRoot;
+
+				if (FParticle* P = GTParticles[TransformIndex].Get())
+				{
+					if (P != nullptr)
+					{
+						P->SetDisabled(!bIsRoot);
+					}
+				}
+			}
+		}
+	}
+
+	if (bHasBuiltGeometryOnGT)
+	{
+		// The Implicits attributes from the Dynamic Collection are just used for initialization, after they can be removed and so free some memory.
+		GameThreadCollection.RemoveAttribute(FGeometryDynamicCollection::ImplicitsAttribute, FTransformCollection::TransformGroup);
+	}
+}
+
+void FGeometryCollectionPhysicsProxy::CreateParticles(const TBitArray<>& EffectiveParticles, TManagedArray<Chaos::FImplicitObjectPtr>& Implicits, Chaos::FPBDRigidsEvolutionBase* Evolution, bool bInitializeRootOnly)
+{
+	const Chaos::Facades::FCollectionAnchoringFacade DynamicCollectionAnchoringFacade(GameThreadCollection);
 	const FVector Scale = Parameters.WorldTransform.GetScale3D();
 	const TManagedArray<float>& Mass = Parameters.RestCollection->GetAttribute<float>(MassAttributeName, FTransformCollection::TransformGroup);
 	const TManagedArray<FTransform>& MassToLocal = Parameters.RestCollection->GetAttribute<FTransform>(MassToLocalAttributeName, FTransformCollection::TransformGroup);
 	const TManagedArray<int32>& Level = Parameters.RestCollection->GetAttribute<int32>(LevelAttributeName, FTransformCollection::TransformGroup);
+	const int32 NumTransforms = GameThreadCollection.GetNumTransforms();
 
 	TArray<int32> ChildrenToCheckForParentFix;
-	TManagedArray<Chaos::FImplicitObjectPtr>& Implicits = GameThreadCollection.ModifyAttribute<Chaos::FImplicitObjectPtr>(FGeometryDynamicCollection::ImplicitsAttribute, FTransformCollection::TransformGroup);
-	if(ensure(NumTransforms == Implicits.Num() && NumTransforms == GTParticles.Num())) // Implicits are in the transform group so this invariant should always hold
+
+	for (int32 Index = 0; Index < NumTransforms; ++Index)
 	{
-		for(int32 Index = 0; Index < NumTransforms; ++Index)
+		if (EffectiveParticles[Index])
 		{
-			if (EffectiveParticles[Index])
+			FParticle* P = GTParticles[Index].Get();
+			if (bInitializeRootOnly)
 			{
 				GTParticles[Index] = FParticle::CreateParticle();
-				FParticle* P = GTParticles[Index].Get();
+				P = GTParticles[Index].Get();
 				GTParticlesToTransformGroupIndex.Add(P, Index);
 
 				GTParticles[Index]->SetUniqueIdx(Evolution->GenerateUniqueIdx());
@@ -609,27 +704,35 @@ void FGeometryCollectionPhysicsProxy::Initialize(Chaos::FPBDRigidsEvolutionBase
 				P->SetM(ScaledMass);
 				P->SetUserData(Parameters.UserData);
 				P->SetProxy(this);
+			}
+			check(P != nullptr);
+			if (bInitializeRootOnly && Index == Parameters.InitialRootIndex && bUseStaticMeshCollisionForTraces && CreateTraceCollisionGeometryCallback != nullptr)
+			{
+				const FTransform ToLocal = MassToLocal[Index].Inverse();
+				TArray<Chaos::FImplicitObjectPtr> Geoms;
+				Chaos::FShapesArray Shapes;
+				CreateTraceCollisionGeometryCallback(ToLocal, Geoms, Shapes);
 
-				if (Index == Parameters.InitialRootIndex && bUseStaticMeshCollisionForTraces && CreateTraceCollisionGeometryCallback != nullptr)
+				Chaos::FImplicitObjectPtr ImplicitGeometry = MakeImplicitObjectPtr<Chaos::FImplicitObjectUnion>(MoveTemp(Geoms));
+				P->SetGeometry(ImplicitGeometry);
+			}
+			else if (bInitializeRootOnly == (Index == Parameters.InitialRootIndex) || bBuildGeometryForChildren)
+			{
+				Chaos::FImplicitObjectPtr ImplicitGeometry = Implicits[Index];
+				if (ImplicitGeometry && !Scale.Equals(FVector::OneVector))
 				{
-					const FTransform ToLocal = MassToLocal[Index].Inverse();
-					TArray<Chaos::FImplicitObjectPtr> Geoms;
-					Chaos::FShapesArray Shapes;
-					CreateTraceCollisionGeometryCallback(ToLocal, Geoms, Shapes);
-
-					Chaos::FImplicitObjectPtr ImplicitGeometry = MakeImplicitObjectPtr<Chaos::FImplicitObjectUnion>(MoveTemp(Geoms));
-					P->SetGeometry(ImplicitGeometry);
-				}
-				else
-				{
-					Chaos::FImplicitObjectPtr ImplicitGeometry = Implicits[Index];
-					if (ImplicitGeometry && !Scale.Equals(FVector::OneVector))
-					{
-						ImplicitGeometry = ImplicitGeometry->CopyGeometryWithScale(Scale);
-					}
-					P->SetGeometry(ImplicitGeometry);
+					ImplicitGeometry = ImplicitGeometry->CopyGeometryWithScale(Scale);
 				}
+				P->SetGeometry(ImplicitGeometry);
+			}
+			// Reset the root in the ManagedArray if not in mode bInitializeRootOnly
+			if (!bInitializeRootOnly && Index == Parameters.InitialRootIndex)
+			{
+				Implicits[Index] = P->GetGeometry();
+			}
 
+			if (bInitializeRootOnly)
+			{
 				if (DynamicCollectionAnchoringFacade.IsAnchored(Index))
 				{
 					P->SetObjectState(Chaos::EObjectStateType::Kinematic, false, false);
@@ -656,41 +759,42 @@ void FGeometryCollectionPhysicsProxy::Initialize(Chaos::FPBDRigidsEvolutionBase
 				P->SetDebugName(MakeShared<FString, ESPMode::ThreadSafe>(FString::Printf(TEXT("%s-%d"), *Parameters.Name, Index)));
 #endif
 			}
-			// this step is necessary for Phase 2 where we need to walk back the hierarchy from children to parent 
-			if (bGeometryCollectionAlwaysGenerateGTCollisionForClusters && !GameThreadCollection.HasChildren(Index))
-			{
-				ChildrenToCheckForParentFix.Add(Index);
-			}
 		}
+		// this step is necessary for Phase 2 where we need to walk back the hierarchy from children to parent 
+		if (bGeometryCollectionAlwaysGenerateGTCollisionForClusters && !GameThreadCollection.HasChildren(Index))
+		{
+			ChildrenToCheckForParentFix.Add(Index);
+		}
+	}
 
-		if (bGeometryCollectionAlwaysGenerateGTCollisionForClusters)
+	if (bGeometryCollectionAlwaysGenerateGTCollisionForClusters)
+	{
+		// second phase: fixing parent geometries
+		// @todo(chaos) this could certainly be done ahead at generation time rather than runtime
+		TSet<int32> ParentToPotentiallyFix;
+		while (ChildrenToCheckForParentFix.Num())
 		{
-			// second phase: fixing parent geometries
-			// @todo(chaos) this could certainly be done ahead at generation time rather than runtime
-			TSet<int32> ParentToPotentiallyFix;
-			while (ChildrenToCheckForParentFix.Num())
+			// step 1 : find parents
+			for (const int32 ChildIndex : ChildrenToCheckForParentFix)
 			{
-				// step 1 : find parents
-				for(const int32 ChildIndex: ChildrenToCheckForParentFix)
+				const int32 ParentIndex = GameThreadCollection.GetParent(ChildIndex);
+				if (ParentIndex != INDEX_NONE)
 				{
-					const int32 ParentIndex = GameThreadCollection.GetParent(ChildIndex);
-					if (ParentIndex != INDEX_NONE)
-					{
-						ParentToPotentiallyFix.Add(ParentIndex);
-					}
+					ParentToPotentiallyFix.Add(ParentIndex);
 				}
+			}
 
-				// step 2: fix the parent if necessary
-				for (const int32 ParentToFixIndex: ParentToPotentiallyFix)
+			// step 2: fix the parent if necessary
+			for (const int32 ParentToFixIndex : ParentToPotentiallyFix)
+			{
+				if (Implicits[ParentToFixIndex] == nullptr)
 				{
-					if (Implicits[ParentToFixIndex] == nullptr)
-					{
-						const Chaos::FRigidTransform3 ParentShapeTransform =  MassToLocal[ParentToFixIndex];
-				
-						// let's make sure all our children have an implicit defined, other wise, postpone to next iteration 
-						bool bAllChildrenHaveCollision = true;
-						
-						GameThreadCollection.IterateThroughChildren(ParentToFixIndex, [&](int32 ChildIndex)
+					const Chaos::FRigidTransform3 ParentShapeTransform = MassToLocal[ParentToFixIndex];
+
+					// let's make sure all our children have an implicit defined, other wise, postpone to next iteration 
+					bool bAllChildrenHaveCollision = true;
+
+					GameThreadCollection.IterateThroughChildren(ParentToFixIndex, [&](int32 ChildIndex)
 						{
 							// defer if any of the children is a cluster with no collision yet generated 
 							if (Implicits[ChildIndex] == nullptr && GameThreadCollection.HasChildren(ChildIndex))
@@ -701,11 +805,11 @@ void FGeometryCollectionPhysicsProxy::Initialize(Chaos::FPBDRigidsEvolutionBase
 							return true;
 						});
 
-						if (bAllChildrenHaveCollision)
-						{
-							// Make a union of the children geometry
-							TArray<Chaos::FImplicitObjectPtr> ChildImplicits;
-							GameThreadCollection.IterateThroughChildren(ParentToFixIndex, [&](int32 ChildIndex)
+					if (bAllChildrenHaveCollision)
+					{
+						// Make a union of the children geometry
+						TArray<Chaos::FImplicitObjectPtr> ChildImplicits;
+						GameThreadCollection.IterateThroughChildren(ParentToFixIndex, [&](int32 ChildIndex)
 							{
 								const Chaos::FImplicitObjectPtr& ChildImplicit = Implicits[ChildIndex];
 								if (ChildImplicit)
@@ -729,103 +833,61 @@ void FGeometryCollectionPhysicsProxy::Initialize(Chaos::FPBDRigidsEvolutionBase
 								}
 								return true;
 							});
-							if (ChildImplicits.Num() > 0)
-							{
-								Chaos::FImplicitObject* UnionImplicit = new Chaos::FImplicitObjectUnion(MoveTemp(ChildImplicits));
-								Implicits[ParentToFixIndex] = Chaos::FImplicitObjectPtr(UnionImplicit);
-							}
-							if (GTParticles[ParentToFixIndex] != nullptr)
-							{
-								GTParticles[ParentToFixIndex]->SetGeometry(Implicits[ParentToFixIndex]);
-							}
+						if (ChildImplicits.Num() > 0)
+						{
+							Chaos::FImplicitObject* UnionImplicit = new Chaos::FImplicitObjectUnion(MoveTemp(ChildImplicits));
+							Implicits[ParentToFixIndex] = Chaos::FImplicitObjectPtr(UnionImplicit);
+						}
+						if (GTParticles[ParentToFixIndex] != nullptr)
+						{
+							GTParticles[ParentToFixIndex]->SetGeometry(Implicits[ParentToFixIndex]);
 						}
 					}
 				}
-
-				// step 3 : make the parent the new child to go up the hierarchy and continue the fixing
-				ChildrenToCheckForParentFix = ParentToPotentiallyFix.Array(); 
-				ParentToPotentiallyFix.Reset();
-			}
-		}
-		
-		// Phase 3 : finalization of shapes
-		for(int32 Index = 0; Index < NumTransforms; ++Index)
-		{
-			FParticle* P = GTParticles[Index].Get();
-			if (P != nullptr)
-			{
-				const Chaos::FShapesArray& Shapes = P->ShapesArray();
-				const int32 NumShapes = Shapes.Num();
-				for (int32 ShapeIndex = 0; ShapeIndex < NumShapes; ++ShapeIndex)
-				{
-					Chaos::FPerShapeData* Shape = Shapes[ShapeIndex].Get();
-					Shape->SetSimData(SimFilter);
-					Shape->SetQueryData(QueryFilter);
-					Shape->SetProxy(this);
-					Shape->SetMaterial(Parameters.PhysicalMaterialHandle);
-				}
 			}
+
+			// step 3 : make the parent the new child to go up the hierarchy and continue the fixing
+			ChildrenToCheckForParentFix = ParentToPotentiallyFix.Array();
+			ParentToPotentiallyFix.Reset();
 		}
 	}
 
-	// Skip simplicials, as they're owned by unique pointers.
-	static const FAttributeAndGroupId SkipList[] =
-	{
-		{ FGeometryDynamicCollection::SimplicialsAttribute, FTransformCollection::TransformGroup },
-	};
-	static constexpr int32 SkipListSize = sizeof(SkipList) / sizeof(FAttributeAndGroupId);
-
-	// Adding the Implicits to PhysicsThreadCollection, before it copies all matching attributes
-	PhysicsThreadCollection.AddAttribute<Chaos::FImplicitObjectPtr>(FGeometryDynamicCollection::ImplicitsAttribute, FTransformCollection::TransformGroup);
-	PhysicsThreadCollection.CopyMatchingAttributesFrom(DynamicCollection, MakeArrayView(SkipList, SkipListSize));
-	PhysicsThreadCollection.CopyInitialVelocityAttributesFrom(DynamicCollection);
-
-	// Copy simplicials.
-	// TODO: Ryan - Should we just transfer ownership of the SimplicialsAttribute from the DynamicCollection to
-	// the PhysicsThreadCollection?
+	// Phase 3 : finalization of shapes
+	for (int32 Index = 0; Index < NumTransforms; ++Index)
 	{
-		if (FGeometryCollection::AreCollisionParticlesEnabled()
-			&& DynamicCollection.HasAttribute(DynamicCollection.SimplicialsAttribute, FTransformCollection::TransformGroup))
-		{
-			const auto& SourceSimplicials = DynamicCollection.GetAttribute<TUniquePtr<FSimplicial>>(
-				DynamicCollection.SimplicialsAttribute, FTransformCollection::TransformGroup);
-			for (int32 Index = PhysicsThreadCollection.NumElements(FTransformCollection::TransformGroup) - 1; 0 <= Index; Index--)
-			{
-				PhysicsThreadCollection.Simplicials[Index].Reset(
-					SourceSimplicials[Index] ? SourceSimplicials[Index]->NewCopy() : nullptr);
-			}
-		}
-		else
+		FParticle* P = GTParticles[Index].Get();
+		if (P != nullptr)
 		{
-			for (int32 Index = PhysicsThreadCollection.NumElements(FTransformCollection::TransformGroup) - 1; 0 <= Index; Index--)
+			const Chaos::FShapesArray& Shapes = P->ShapesArray();
+			const int32 NumShapes = Shapes.Num();
+			for (int32 ShapeIndex = 0; ShapeIndex < NumShapes; ++ShapeIndex)
 			{
-				PhysicsThreadCollection.Simplicials[Index].Reset();
+				Chaos::FPerShapeData* Shape = Shapes[ShapeIndex].Get();
+				Shape->SetSimData(SimFilter);
+				Shape->SetQueryData(QueryFilter);
+				Shape->SetProxy(this);
+				Shape->SetMaterial(Parameters.PhysicalMaterialHandle);
 			}
 		}
 	}
+}
 
-
-	if (Parameters.EnableClustering)
+void FGeometryCollectionPhysicsProxy::CreateChildrenGeometry_External()
+{
+	if (!bHasBuiltGeometryOnGT)
 	{
-		// make sure we set Activate the right way when clustering is enabled ( only root should be enabled at start ) 
-		for (int32 TransformIndex = 0; TransformIndex < NumTransforms; ++TransformIndex)
+		if (Chaos::FPhysicsSolver* RBDSolver = GetSolver<Chaos::FPhysicsSolver>())
 		{
-			const bool bIsRoot = !GameThreadCollection.GetHasParent(TransformIndex);
-			GameThreadCollection.Active[TransformIndex] = bIsRoot;
-			PhysicsThreadCollection.Active[TransformIndex] = bIsRoot;
+			TManagedArray<Chaos::FImplicitObjectPtr>& Implicits = GameThreadCollection.ModifyAttribute<Chaos::FImplicitObjectPtr>(FGeometryDynamicCollection::ImplicitsAttribute, FTransformCollection::TransformGroup);
+			TBitArray<> EffectiveParticles;
+			NumEffectiveParticles = CalculateEffectiveParticles(GameThreadCollection, NumParticles, Parameters.MaxSimulatedLevel, Parameters.EnableClustering, GetOwner(), EffectiveParticles);
+			CreateParticles(EffectiveParticles, Implicits, RBDSolver->GetEvolution(), /*InitilaizeOnlyRoot*/false);
 
-			if (FParticle* P = GTParticles[TransformIndex].Get())
-			{
-				if (P != nullptr)
-				{
-					P->SetDisabled(!bIsRoot);
-				}
-			}
+			// The Implicits attributes from the Dynamic Collection are just used for initialization, after they can be removed and so free some memory.
+			GameThreadCollection.RemoveAttribute(FGeometryDynamicCollection::ImplicitsAttribute, FTransformCollection::TransformGroup);
+			bHasBuiltGeometryOnGT = true;
 		}
 	}
-
-	// The Implicits attributes from the Dynamic Collection are just used for initialization, after they can be removed and so free some memory.
-	GameThreadCollection.RemoveAttribute(FGeometryDynamicCollection::ImplicitsAttribute, FTransformCollection::TransformGroup);
 }
 
 
@@ -1081,6 +1143,8 @@ void FGeometryCollectionPhysicsProxy::InitializeBodiesPT(Chaos::FPBDRigidsSolver
 	const FGeometryDynamicCollection& DynamicCollection = PhysicsThreadCollection;
 
 	bIsInitializedOnPhysicsThread = true;
+	// Building geometry according to the cVar
+	bHasBuiltGeometryOnPT = bBuildGeometryForChildren;
 
 	if (Parameters.Simulating && RestCollection)
 	{
@@ -1156,7 +1220,7 @@ void FGeometryCollectionPhysicsProxy::InitializeBodiesPT(Chaos::FPBDRigidsSolver
 					Handle,
 					Parameters.Shared,
 					Simplicials[TransformGroupIndex].Get(),
-					Implicits[TransformGroupIndex],
+					(bBuildGeometryForChildren || (TransformGroupIndex == Parameters.InitialRootIndex)) ? Implicits[TransformGroupIndex] : nullptr,
 					SimFilter,
 					QueryFilter,
 					ScaledMass,
@@ -1529,10 +1593,67 @@ void FGeometryCollectionPhysicsProxy::InitializeBodiesPT(Chaos::FPBDRigidsSolver
 		// call DirtyParticle to make sure the acceleration structure is up to date with all the changes happening here
 		DirtyAllParticles(*RigidsSolver);
 
+		if (bHasBuiltGeometryOnPT)
+		{
+			// The Implicits attributes from the Dynamic Collection are just used for geometry initialization, after they can be removed and so free some memory.
+			PhysicsThreadCollection.RemoveAttribute(FGeometryDynamicCollection::ImplicitsAttribute, FTransformCollection::TransformGroup);
+		}
 	} // end if simulating...
 
-	// The Implicits attributes from the Dynamic Collection are just used for initialization, after they can be removed and so free some memory.
-	PhysicsThreadCollection.RemoveAttribute(FGeometryDynamicCollection::ImplicitsAttribute, FTransformCollection::TransformGroup);
+}
+
+void FGeometryCollectionPhysicsProxy::CreateChildrenGeometry_Internal()
+{
+	if (!bHasBuiltGeometryOnPT)
+	{
+		const FGeometryCollection* RestCollection = Parameters.RestCollection;
+
+		if (Parameters.Simulating && ensure(RestCollection))
+		{
+			const int32 NumTransforms = PhysicsThreadCollection.NumElements(FTransformCollection::TransformGroup);
+			const TManagedArray<FTransform>& MassToLocal = RestCollection->GetAttribute<FTransform>(MassToLocalAttributeName, FTransformCollection::TransformGroup);
+			const TManagedArray<Chaos::FImplicitObjectPtr>& Implicits = PhysicsThreadCollection.GetAttribute<Chaos::FImplicitObjectPtr>(FGeometryDynamicCollection::ImplicitsAttribute, FTransformCollection::TransformGroup);
+			const TManagedArray<TUniquePtr<FCollisionStructureManager::FSimplicial>>& Simplicials = PhysicsThreadCollection.Simplicials;
+
+			// In PushToPhysicsState, we're going to compute a relative transform from Parameters.PrevWorldTransform
+			// to a particle's current world transform to get its relative transform. Then, that relative transform is
+			// applied on top of the new Parameters.WorldTransform to get the final world transform. This is problematic
+			// if the particle is initialized with some Parameters.WorldTransform because then the particle's world transform
+			// will have Parameters.WorldTransform baked in but Parameters.PrevWorldTransform will be an identity. Then
+			// when we go to set the kinematic target we will then effectively be applying Parameters.WorldTransform twice.
+			Parameters.PrevWorldTransform = Parameters.WorldTransform;
+
+			TArray<FTransform> Transform;
+			GeometryCollectionAlgo::Private::GlobalMatrices(PhysicsThreadCollection, Transform);
+
+			// Here Clean up Additional particles
+			TBitArray<> EffectiveParticles;
+			NumEffectiveParticles = CalculateEffectiveParticles(PhysicsThreadCollection, PhysicsThreadCollection.NumElements(FGeometryCollection::TransformGroup), Parameters.MaxSimulatedLevel, Parameters.EnableClustering, GetOwner(), EffectiveParticles);
+
+			for (int32 Index = 0; Index < NumTransforms; Index++)
+			{
+				if (EffectiveParticles[Index] && Index != Parameters.InitialRootIndex)
+				{
+					Chaos::TPBDGeometryCollectionParticleHandle<Chaos::FReal, 3>* Handle = static_cast<Chaos::TPBDGeometryCollectionParticleHandle<Chaos::FReal, 3>*>(SolverParticleHandles[Index]);
+					if (ensure(Handle) && Handle->GetGeometry() == nullptr)
+					{
+						const FTransform WorldTransform = MassToLocal[Index] * Transform[Index] * Parameters.WorldTransform;
+						SetImplicitToPTParticles(Handle, Parameters.Shared, Simplicials[Index].Get(), Implicits[Index], SimFilter, QueryFilter, WorldTransform, CollisionParticlesPerObjectFraction);
+						check(Handle->GetGeometry() != nullptr);
+					}
+				}
+			}
+
+			if (Chaos::FPhysicsSolver* RBDSolver = GetSolver<Chaos::FPhysicsSolver>())
+			{
+				// call DirtyParticle to make sure the acceleration structure is up to date with all the changes happening here
+				DirtyAllParticles(*RBDSolver);
+			}
+		}
+		bHasBuiltGeometryOnPT = true;
+		// The Implicits attributes from the Dynamic Collection are just used for geometry initialization, after they can be removed and so free some memory.
+		PhysicsThreadCollection.RemoveAttribute(FGeometryDynamicCollection::ImplicitsAttribute, FTransformCollection::TransformGroup);
+	}
 }
 
 int32 ReportNoLevelsetCluster = 0;
@@ -1723,7 +1844,7 @@ Chaos::TPBDGeometryCollectionParticleHandle<Chaos::FReal, 3>* FGeometryCollectio
 		Handle,
 		Parameters.Shared,
 		Simplicials[CollectionClusterIndex].Get(),
-		Implicits[CollectionClusterIndex],
+		(bBuildGeometryForChildren || (CollectionClusterIndex == Parameters.InitialRootIndex)) ? Implicits[CollectionClusterIndex] : nullptr,
 		SimFilter,
 		QueryFilter,
 		Mass,
@@ -2213,8 +2334,11 @@ void FGeometryCollectionPhysicsProxy::BreakClusters_External(TArray<FGeometryCol
 
 	if (Chaos::FPhysicsSolver* RBDSolver = GetSolver<Chaos::FPhysicsSolver>())
 	{
+		// todo(chaos) to explain that this is likely temporary and that we need to find a better way
+		CreateChildrenGeometry_External();
 		RBDSolver->EnqueueCommandImmediate([this, RBDSolver, IndicesToBreakParent = MoveTemp(ItemIndices)]()
 		{
+			CreateChildrenGeometry_Internal();
 			Chaos::FRigidClustering& Clustering = RBDSolver->GetEvolution()->GetRigidClustering();
 			for (const FGeometryCollectionItemIndex& ItemIndex : IndicesToBreakParent)
 			{
@@ -2234,8 +2358,11 @@ void FGeometryCollectionPhysicsProxy::BreakActiveClusters_External()
 
 	if (Chaos::FPhysicsSolver* RBDSolver = GetSolver<Chaos::FPhysicsSolver>())
 	{
+		// todo(chaos) to explain that this is likely temporary and that we need to find a better way
+		CreateChildrenGeometry_External();
 		RBDSolver->EnqueueCommandImmediate([this, RBDSolver]()
 		{
+			CreateChildrenGeometry_Internal();
 			Chaos::FRigidClustering& Clustering = RBDSolver->GetEvolution()->GetRigidClustering();
 			Clustering.BreakClustersByProxy(this);
 		});
@@ -2358,8 +2485,11 @@ void FGeometryCollectionPhysicsProxy::RemoveAllAnchors_External()
 	check(IsInGameThread());
 	if (Chaos::FPhysicsSolver* RBDSolver = GetSolver<Chaos::FPhysicsSolver>())
 	{
+		// todo(chaos) to explain that this is likely temporary and that we need to find a better way
+		CreateChildrenGeometry_External();
 		RBDSolver->EnqueueCommandImmediate([this, RBDSolver]()
 			{
+				CreateChildrenGeometry_Internal();
 				Chaos::FPBDRigidsEvolution* Evolution = RBDSolver->GetEvolution();
 				for (FClusterHandle* ParticleHandle : GetSolverParticleHandles())
 				{
@@ -3989,6 +4119,8 @@ bool FGeometryCollectionPhysicsProxy::PullFromPhysicsState(const Chaos::FDirtyGe
 			TManagedArray<FVector3f>* LinearVelocities = GameThreadCollection.GetLinearVelocitiesAttribute();
 			TManagedArray<FVector3f>* AngularVelocities = GameThreadCollection.GetAngularVelocitiesAttribute();
 
+			bool bAtLeatOneChildActive = false;
+
 			// for that case we cannot just go through the list of entries since Results and NextResults may have different number of entries that don't always match
 			// so we need to go through the transform indices and find the matching entries on both side 
 			// this is helped by the FResultAccessor structure
@@ -4009,6 +4141,8 @@ bool FGeometryCollectionPhysicsProxy::PullFromPhysicsState(const Chaos::FDirtyGe
 				const bool bAnimatingWhileDisabled = AnimationsActive ? (*AnimationsActive)[TransformGroupIndex] : false;
 				const bool bActive = GameThreadCollection.Active[TransformGroupIndex];
 
+				bAtLeatOneChildActive |= bActive && TransformGroupIndex != Parameters.InitialRootIndex;
+
 				if (bActive || bAnimatingWhileDisabled)
 				{
 					Chaos::FVec3 NewX;
@@ -4034,6 +4168,10 @@ bool FGeometryCollectionPhysicsProxy::PullFromPhysicsState(const Chaos::FDirtyGe
 					}
 				}
 			}
+			if (bAtLeatOneChildActive)
+			{
+				CreateChildrenGeometry_External();
+			}
 		}
 
 		// update the parent transform now that we have all X and R computed
diff --git a/Engine/Source/Runtime/Experimental/Chaos/Public/PhysicsProxy/GeometryCollectionPhysicsProxy.h b/Engine/Source/Runtime/Experimental/Chaos/Public/PhysicsProxy/GeometryCollectionPhysicsProxy.h
--- a/Engine/Source/Runtime/Experimental/Chaos/Public/PhysicsProxy/GeometryCollectionPhysicsProxy.h
+++ b/Engine/Source/Runtime/Experimental/Chaos/Public/PhysicsProxy/GeometryCollectionPhysicsProxy.h
@@ -499,6 +499,8 @@ public:
 	using FCreateTraceCollisionGeometryCallback = TFunction<void(const FTransform& InToLocal, TArray<Chaos::FImplicitObjectPtr>& OutGeoms, Chaos::FShapesArray& OutShapes)>;
 	void SetCreateTraceCollisionGeometryCallback(FCreateTraceCollisionGeometryCallback InCreateGeometryCallback) { CreateTraceCollisionGeometryCallback = InCreateGeometryCallback; }
 
+	CHAOS_API void CreateChildrenGeometry_Internal();
+
 protected:
 
 	CHAOS_API float ComputeMaterialBasedDamageThreshold_Internal(int32 TransformIndex) const;
@@ -572,6 +574,9 @@ private:
 	static TBitArray<> CalculateClustersToCreateFromChildren(const FGeometryDynamicCollection& DynamicCollection, int32 NumTransforms);
 	static int32 CalculateEffectiveParticles(const FGeometryDynamicCollection& DynamicCollection, int32 NumTransform, int32 MaxSimulatedLevel, bool bEnableClustering, const UObject* Owner, TBitArray<>& EffectiveParticles);
 
+	void CreateParticles(const TBitArray<>& EffectiveParticles, TManagedArray<Chaos::FImplicitObjectPtr>& Implicits, Chaos::FPBDRigidsEvolutionBase* Evolution, bool bInitializeRootOnly);
+	void CreateChildrenGeometry_External();
+	
 	/**
 	 * Since geometry collections only buffer data that has changed, when PullFromPhysicsState is given both PrevData and NextData it must
 	 * examine *both* PrevData and NextData for data about a particle (since that particle's data coudl be in PrevData and not NextData).
@@ -674,6 +679,9 @@ private:
 	FTransform WorldTransform_External;
 	uint8 bIsGameThreadWorldTransformDirty : 1;
 
+	uint8 bHasBuiltGeometryOnPT : 1;
+	uint8 bHasBuiltGeometryOnGT : 1;
+
 	// Currently this is using triple buffers for game-physics and 
 	// physics-game thread communication, but not for any reason other than this 
 	// is the only implementation we currently have of a guarded buffer - a buffer 
