Iris - Implicit conversion fixes.

#rnx

diff --git a/Engine/Source/Runtime/Experimental/Iris/Core/Private/Iris/ReplicationSystem/Conditionals/ReplicationConditionals.cpp b/Engine/Source/Runtime/Experimental/Iris/Core/Private/Iris/ReplicationSystem/Conditionals/ReplicationConditionals.cpp
--- a/Engine/Source/Runtime/Experimental/Iris/Core/Private/Iris/ReplicationSystem/Conditionals/ReplicationConditionals.cpp
+++ b/Engine/Source/Runtime/Experimental/Iris/Core/Private/Iris/ReplicationSystem/Conditionals/ReplicationConditionals.cpp
@@ -65,7 +65,7 @@ bool FReplicationConditionals::SetConditionConnectionFilter(FInternalNetRefIndex
 	if (ObjectInfo->AutonomousConnectionId != AutonomousConnectionId)
 	{
 		const uint32 ConnIdForBaselineInvalidation = (bEnable ? ConnectionId : ObjectInfo->AutonomousConnectionId);
-		ObjectInfo->AutonomousConnectionId = AutonomousConnectionId;
+		ObjectInfo->AutonomousConnectionId = uint16(AutonomousConnectionId);
 
 		BaselineInvalidationTracker->InvalidateBaselines(ObjectIndex, ConnIdForBaselineInvalidation);
 	}
diff --git a/Engine/Source/Runtime/Experimental/Iris/Core/Private/Iris/Serialization/IntNetSerializerBase.h b/Engine/Source/Runtime/Experimental/Iris/Core/Private/Iris/Serialization/IntNetSerializerBase.h
--- a/Engine/Source/Runtime/Experimental/Iris/Core/Private/Iris/Serialization/IntNetSerializerBase.h
+++ b/Engine/Source/Runtime/Experimental/Iris/Core/Private/Iris/Serialization/IntNetSerializerBase.h
@@ -93,7 +93,7 @@ struct FIntNetSerializerBase
 		// Sign-extend the value if needed
 		if constexpr (std::is_signed_v<SourceType>)
 		{
-			const QuantizedType SignMask = QuantizedType(1) << (BitCount - 1U);
+			const QuantizedType SignMask = static_cast<QuantizedType>(QuantizedType(1) << (BitCount - 1U));
 			Value = (Value ^ SignMask) - SignMask;
 		}
 
@@ -168,7 +168,7 @@ struct FIntNetSerializerBase
 		QuantizedType Value;
 		if constexpr (std::is_signed_v<SourceType>)
 		{
-			const QuantizedType SignMask = (QuantizedType(1) << (BitCount - 1U));
+			const QuantizedType SignMask = static_cast<QuantizedType>(QuantizedType(1) << (BitCount - 1U));
 			const QuantizedType ValueMask = SignMask | (SignMask - QuantizedType(1));
 
 			Value = static_cast<QuantizedType>(Source) & ValueMask;
diff --git a/Engine/Source/Runtime/Experimental/Iris/Core/Private/Iris/Serialization/QuatNetSerializers.cpp b/Engine/Source/Runtime/Experimental/Iris/Core/Private/Iris/Serialization/QuatNetSerializers.cpp
--- a/Engine/Source/Runtime/Experimental/Iris/Core/Private/Iris/Serialization/QuatNetSerializers.cpp
+++ b/Engine/Source/Runtime/Experimental/Iris/Core/Private/Iris/Serialization/QuatNetSerializers.cpp
@@ -89,7 +89,7 @@ inline typename FUnitQuatNetSerializerBase<T>::FloatType FUnitQuatNetSerializerB
 template<typename T>
 T FUnitQuatNetSerializerBase<T>::GetUnitQuat(const T& Value)
 {
-	constexpr FloatType SmallNumber = (sizeof(FloatType) == 4 ? SMALL_NUMBER : DOUBLE_SMALL_NUMBER);
+	constexpr FloatType SmallNumber = FloatType(UE_SMALL_NUMBER);
 	T UnitQuat = Value;
 	if (UnitQuat.SizeSquared() <= SmallNumber)
 	{
diff --git a/Engine/Source/Runtime/Experimental/Iris/Core/Private/Iris/Stats/NetStats.cpp b/Engine/Source/Runtime/Experimental/Iris/Core/Private/Iris/Stats/NetStats.cpp
--- a/Engine/Source/Runtime/Experimental/Iris/Core/Private/Iris/Stats/NetStats.cpp
+++ b/Engine/Source/Runtime/Experimental/Iris/Core/Private/Iris/Stats/NetStats.cpp
@@ -71,12 +71,12 @@ void FNetSendStats::ReportCsvStats()
 	if (Stats.ReplicatingConnectionCount > 0)
 	{
 		const float ConnectionCountFloat = float(Stats.ReplicatingConnectionCount);
-		CSV_CUSTOM_STAT(Iris, AvgScheduledForReplicationRootObjectCount, Stats.ScheduledForReplicationRootObjectCount/ConnectionCountFloat, ECsvCustomStatOp::Set);
-		CSV_CUSTOM_STAT(Iris, AvgReplicatedRootObjectCount, Stats.ReplicatedRootObjectCount/ConnectionCountFloat, ECsvCustomStatOp::Set);
-		CSV_CUSTOM_STAT(Iris, AvgReplicatedObjectCount, Stats.ReplicatedObjectCount/ConnectionCountFloat, ECsvCustomStatOp::Set);
-		CSV_CUSTOM_STAT(Iris, AvgReplicatedDestructionInfoCount, Stats.ReplicatedDestructionInfoCount/ConnectionCountFloat, ECsvCustomStatOp::Set);
-		CSV_CUSTOM_STAT(Iris, AvgReplicatedObjectStatesMaskedOut, Stats.ReplicatedObjectStatesMaskedOut/ConnectionCountFloat, ECsvCustomStatOp::Set);
-		CSV_CUSTOM_STAT(Iris, AvgDeltaCompressedObjectCount, Stats.DeltaCompressedObjectCount/ConnectionCountFloat, ECsvCustomStatOp::Set);
+		CSV_CUSTOM_STAT(Iris, AvgScheduledForReplicationRootObjectCount, float(Stats.ScheduledForReplicationRootObjectCount)/ConnectionCountFloat, ECsvCustomStatOp::Set);
+		CSV_CUSTOM_STAT(Iris, AvgReplicatedRootObjectCount, float(Stats.ReplicatedRootObjectCount)/ConnectionCountFloat, ECsvCustomStatOp::Set);
+		CSV_CUSTOM_STAT(Iris, AvgReplicatedObjectCount, float(Stats.ReplicatedObjectCount)/ConnectionCountFloat, ECsvCustomStatOp::Set);
+		CSV_CUSTOM_STAT(Iris, AvgReplicatedDestructionInfoCount, float(Stats.ReplicatedDestructionInfoCount)/ConnectionCountFloat, ECsvCustomStatOp::Set);
+		CSV_CUSTOM_STAT(Iris, AvgReplicatedObjectStatesMaskedOut, float(Stats.ReplicatedObjectStatesMaskedOut)/ConnectionCountFloat, ECsvCustomStatOp::Set);
+		CSV_CUSTOM_STAT(Iris, AvgDeltaCompressedObjectCount, float(Stats.DeltaCompressedObjectCount)/ConnectionCountFloat, ECsvCustomStatOp::Set);
 	}
 	else
 	{
@@ -184,7 +184,7 @@ void FNetTypeStats::Accumulate(FNetStatsContext& Context)
 
 #define UE_NET_STATS_RECORD_TYPESTATS_TIME(StatsName, ValueName, StatsData) FCsvProfiler::RecordCustomStat(StatsName, CSV_CATEGORY_INDEX(Iris##ValueName##MS), FGenericPlatformTime::ToMilliseconds64(StatsData.Values[FNetTypeStatsData::EStatsIndex::ValueName].Time), ECsvCustomStatOp::Set)
 #define UE_NET_STATS_RECORD_TYPESTATS_COUNT(StatsName, ValueName, StatsData) FCsvProfiler::RecordCustomStat(StatsName, CSV_CATEGORY_INDEX(Iris##ValueName##Count), static_cast<int32>(StatsData.Values[FNetTypeStatsData::EStatsIndex::ValueName].Count) , ECsvCustomStatOp::Set)
-#define UE_NET_STATS_RECORD_TYPESTATS_BITS(StatsName, ValueName, StatsData) FCsvProfiler::RecordCustomStat(StatsName, CSV_CATEGORY_INDEX(Iris##ValueName##KBytes), ((StatsData.Values[FNetTypeStatsData::EStatsIndex::ValueName].Bits + 7U) / 8) / 1000.f , ECsvCustomStatOp::Set)
+#define UE_NET_STATS_RECORD_TYPESTATS_BITS(StatsName, ValueName, StatsData) FCsvProfiler::RecordCustomStat(StatsName, CSV_CATEGORY_INDEX(Iris##ValueName##KBytes), float((StatsData.Values[FNetTypeStatsData::EStatsIndex::ValueName].Bits + 7U) / 8) / 1000.f , ECsvCustomStatOp::Set)
 
 void FNetTypeStats::ReportCSVStats()
 {
diff --git a/Engine/Source/Runtime/Experimental/Iris/Core/Public/Iris/ReplicationSystem/Prioritization/SphereNetObjectPrioritizer.h b/Engine/Source/Runtime/Experimental/Iris/Core/Public/Iris/ReplicationSystem/Prioritization/SphereNetObjectPrioritizer.h
--- a/Engine/Source/Runtime/Experimental/Iris/Core/Public/Iris/ReplicationSystem/Prioritization/SphereNetObjectPrioritizer.h
+++ b/Engine/Source/Runtime/Experimental/Iris/Core/Public/Iris/ReplicationSystem/Prioritization/SphereNetObjectPrioritizer.h
@@ -32,7 +32,7 @@ public:
 
 	UPROPERTY(Config)
 	/** Priority outside the sphere */
-	float OutsidePriority = 0.1;
+	float OutsidePriority = 0.1f;
 };
 
 UCLASS(Transient, MinimalAPI)
