Iris - Enable unsafe type cast warnings in IrisCore.

#rnx

diff --git a/Engine/Source/Runtime/Experimental/Iris/Core/IrisCore.Build.cs b/Engine/Source/Runtime/Experimental/Iris/Core/IrisCore.Build.cs
--- a/Engine/Source/Runtime/Experimental/Iris/Core/IrisCore.Build.cs
+++ b/Engine/Source/Runtime/Experimental/Iris/Core/IrisCore.Build.cs
@@ -31,6 +31,8 @@ namespace UnrealBuildTool.Rules
 			}
 
 			bAllowAutoRTFMInstrumentation = true;
+
+			UnsafeTypeCastWarningLevel = WarningLevel.Error;
 		}
 	}
 }
