Constraints: Speculative fix for SortConstraints crashing since we do an Algo::StableSort on a invalid/stale pointer

diff --git a/Engine/Source/Runtime/Experimental/Animation/Constraints/Private/ConstraintsManager.cpp b/Engine/Source/Runtime/Experimental/Animation/Constraints/Private/ConstraintsManager.cpp
--- a/Engine/Source/Runtime/Experimental/Animation/Constraints/Private/ConstraintsManager.cpp
+++ b/Engine/Source/Runtime/Experimental/Animation/Constraints/Private/ConstraintsManager.cpp
@@ -805,13 +805,18 @@ TArray< TWeakObjectPtr<UTickableConstraint> > FConstraintsManagerController::Get
 		return Empty;
 	}
 	TArray<TWeakObjectPtr<UTickableConstraint>> Constraints = Subsystem->GetConstraints(World);
-
+	
 	if (!bSorted)
 	{
 		return Constraints;
 	}
 
 	TArray< TWeakObjectPtr<UTickableConstraint> > SortedConstraints(Constraints);
+	// Remove stale constraints
+	Constraints.RemoveAll([](const TWeakObjectPtr<UTickableConstraint>& ExistingConstraint) -> bool
+	{
+		return ExistingConstraint.IsStale();
+	});
 	SortConstraints(World,SortedConstraints);
 	
 	return SortedConstraints;
