Fixed use-after-free detected by address sanitizer:
the ConnectionPool::Params object keeps views over the content of the service URL so we need the service URL to outlive the params.

#rnx

diff --git a/Engine/Source/Runtime/Experimental/IoStoreOnDemand/Private/OnDemandHttpClient.cpp b/Engine/Source/Runtime/Experimental/IoStoreOnDemand/Private/OnDemandHttpClient.cpp
--- a/Engine/Source/Runtime/Experimental/IoStoreOnDemand/Private/OnDemandHttpClient.cpp
+++ b/Engine/Source/Runtime/Experimental/IoStoreOnDemand/Private/OnDemandHttpClient.cpp
@@ -233,8 +233,9 @@ void FHttpClient::RetryRequest(FRequestParams&& Params, bool bNextEndpoint)
 
 TUniquePtr<HTTP::FConnectionPool> FHttpClient::CreateConnection(const FStringView& HostAddr)
 {
+	auto ServiceUrlAnsi = StringCast<ANSICHAR>(HostAddr.GetData(), HostAddr.Len());
 	HTTP::FConnectionPool::FParams Params;
-	ensure(Params.SetHostFromUrl(StringCast<ANSICHAR>(HostAddr.GetData(), HostAddr.Len())) >= 0);
+	ensure(Params.SetHostFromUrl(ServiceUrlAnsi) >= 0);
 	if (Config.ReceiveBufferSize >= 0)
 	{
 		UE_LOG(LogIas, Log, TEXT("HTTP client receive buffer size set to %d"), Config.ReceiveBufferSize);
