Create geometry before adding particles in the Cluster Union

diff --git a/Engine/Source/Runtime/Experimental/Chaos/Private/Chaos/ClusterUnionManager.cpp b/Engine/Source/Runtime/Experimental/Chaos/Private/Chaos/ClusterUnionManager.cpp
--- a/Engine/Source/Runtime/Experimental/Chaos/Private/Chaos/ClusterUnionManager.cpp
+++ b/Engine/Source/Runtime/Experimental/Chaos/Private/Chaos/ClusterUnionManager.cpp
@@ -9,6 +9,7 @@
 #include "Chaos/PBDRigidClusteringAlgo.h"
 #include "Chaos/PBDRigidsEvolution.h"
 #include "Chaos/PBDRigidsEvolutionGBF.h"
+#include "PhysicsProxy/GeometryCollectionPhysicsProxy.h"
 #include "ChaosStats.h"
 #include "RewindData.h"
 
@@ -239,6 +240,18 @@ namespace Chaos
 
 			return GetParentParticleInClusterUnion(ClusterUnion, Particle->Parent());
 		}
+
+		FGeometryCollectionPhysicsProxy* GetGCProxy(FPBDRigidParticleHandle* Particle)
+		{
+			if (IPhysicsProxyBase* Proxy = Particle->PhysicsProxy())
+			{
+				if (Proxy->GetType() == FGeometryCollectionPhysicsProxy::ConcreteType())
+				{
+					return static_cast<FGeometryCollectionPhysicsProxy*>(Proxy);
+				}
+			}
+			return nullptr;
+		}
 	}
 
 	const TArray<FPBDRigidParticleHandle*>& FClusterUnion::GetPendingGeometryOperationParticles(EClusterUnionGeometryOperation Op) const
@@ -592,6 +605,15 @@ namespace Chaos
 
 		for (FPBDRigidParticleHandle* Particle : FinalParticlesToAdd)
 		{
+			if (Particle->GetGeometry() == nullptr)
+			{
+				FGeometryCollectionPhysicsProxy* GCPhysicsProxy = GetGCProxy(Particle);
+				if (GCPhysicsProxy)
+				{
+					GCPhysicsProxy->CreateChildrenGeometry_Internal();
+				}
+			}
+
 			if (!bIsNewCluster)
 			{
 				Cluster->PendingConnectivityOperations.Add({ Particle, EClusterUnionConnectivityOperation::Add });
diff --git a/Engine/Source/Runtime/Experimental/Chaos/Private/PhysicsProxy/GeometryCollectionPhysicsProxy.cpp b/Engine/Source/Runtime/Experimental/Chaos/Private/PhysicsProxy/GeometryCollectionPhysicsProxy.cpp
--- a/Engine/Source/Runtime/Experimental/Chaos/Private/PhysicsProxy/GeometryCollectionPhysicsProxy.cpp
+++ b/Engine/Source/Runtime/Experimental/Chaos/Private/PhysicsProxy/GeometryCollectionPhysicsProxy.cpp
@@ -2334,11 +2334,8 @@ void FGeometryCollectionPhysicsProxy::BreakClusters_External(TArray<FGeometryCol
 
 	if (Chaos::FPhysicsSolver* RBDSolver = GetSolver<Chaos::FPhysicsSolver>())
 	{
-		// todo(chaos) to explain that this is likely temporary and that we need to find a better way
-		CreateChildrenGeometry_External();
 		RBDSolver->EnqueueCommandImmediate([this, RBDSolver, IndicesToBreakParent = MoveTemp(ItemIndices)]()
 		{
-			CreateChildrenGeometry_Internal();
 			Chaos::FRigidClustering& Clustering = RBDSolver->GetEvolution()->GetRigidClustering();
 			for (const FGeometryCollectionItemIndex& ItemIndex : IndicesToBreakParent)
 			{
@@ -2358,11 +2355,8 @@ void FGeometryCollectionPhysicsProxy::BreakActiveClusters_External()
 
 	if (Chaos::FPhysicsSolver* RBDSolver = GetSolver<Chaos::FPhysicsSolver>())
 	{
-		// todo(chaos) to explain that this is likely temporary and that we need to find a better way
-		CreateChildrenGeometry_External();
 		RBDSolver->EnqueueCommandImmediate([this, RBDSolver]()
 		{
-			CreateChildrenGeometry_Internal();
 			Chaos::FRigidClustering& Clustering = RBDSolver->GetEvolution()->GetRigidClustering();
 			Clustering.BreakClustersByProxy(this);
 		});
@@ -2485,11 +2479,8 @@ void FGeometryCollectionPhysicsProxy::RemoveAllAnchors_External()
 	check(IsInGameThread());
 	if (Chaos::FPhysicsSolver* RBDSolver = GetSolver<Chaos::FPhysicsSolver>())
 	{
-		// todo(chaos) to explain that this is likely temporary and that we need to find a better way
-		CreateChildrenGeometry_External();
 		RBDSolver->EnqueueCommandImmediate([this, RBDSolver]()
 			{
-				CreateChildrenGeometry_Internal();
 				Chaos::FPBDRigidsEvolution* Evolution = RBDSolver->GetEvolution();
 				for (FClusterHandle* ParticleHandle : GetSolverParticleHandles())
 				{
