Memory optimization
Remove implicits from dynamic collection

diff --git a/Engine/Source/Runtime/Experimental/Chaos/Private/PhysicsProxy/GeometryCollectionPhysicsProxy.cpp b/Engine/Source/Runtime/Experimental/Chaos/Private/PhysicsProxy/GeometryCollectionPhysicsProxy.cpp
--- a/Engine/Source/Runtime/Experimental/Chaos/Private/PhysicsProxy/GeometryCollectionPhysicsProxy.cpp
+++ b/Engine/Source/Runtime/Experimental/Chaos/Private/PhysicsProxy/GeometryCollectionPhysicsProxy.cpp
@@ -526,6 +526,12 @@ FAutoConsoleVariableRef CVarReportHighParticleFraction(TEXT("p.gc.ReportHighPart
 bool bBuildGeometryForChildren = true;
 FAutoConsoleVariableRef CVarbBuildGeometryForChildren(TEXT("p.gc.BuildGeometryForChildren"), bBuildGeometryForChildren, TEXT("If true build all children geometry at initilaization time, otherwise wait until destruction occurs."));
 
+bool bRemoveImplicitsInDynamicCollections = false;
+FAutoConsoleVariableRef CVarbRemoveImplicitsInDynamicCollections(TEXT("p.gc.RemoveImplicitsInDynamicCollections"), 
+	bBuildGeometryForChildren, TEXT("This cvar has an impact only if geometry are not added for children. It removes implicits from the Dynamic Collections, and recreate then from the rest collection. \
+										Using this cvar could have an impact if geometry are updated from the dynamic collection on the GT, then those changes won't be ported to the PT."));
+
+
 void FGeometryCollectionPhysicsProxy::Initialize(Chaos::FPBDRigidsEvolutionBase *Evolution)
 {
 	check(IsInGameThread());
@@ -660,7 +666,7 @@ void FGeometryCollectionPhysicsProxy::Initialize(Chaos::FPBDRigidsEvolutionBase
 		}
 	}
 
-	if (bHasBuiltGeometryOnGT)
+	if (bHasBuiltGeometryOnGT || bRemoveImplicitsInDynamicCollections)
 	{
 		// The Implicits attributes from the Dynamic Collection are just used for initialization, after they can be removed and so free some memory.
 		GameThreadCollection.RemoveAttribute(FGeometryDynamicCollection::ImplicitsAttribute, FTransformCollection::TransformGroup);
@@ -878,6 +884,12 @@ void FGeometryCollectionPhysicsProxy::CreateChildrenGeometry_External()
 	{
 		if (Chaos::FPhysicsSolver* RBDSolver = GetSolver<Chaos::FPhysicsSolver>())
 		{
+			if (bRemoveImplicitsInDynamicCollections)
+			{
+				GameThreadCollection.AddAttribute<Chaos::FImplicitObjectPtr>(FGeometryDynamicCollection::ImplicitsAttribute, FTransformCollection::TransformGroup);
+				GameThreadCollection.CopyAttribute(*Parameters.RestCollection, FGeometryDynamicCollection::ImplicitsAttribute, FTransformCollection::TransformGroup);
+			}
+
 			TManagedArray<Chaos::FImplicitObjectPtr>& Implicits = GameThreadCollection.ModifyAttribute<Chaos::FImplicitObjectPtr>(FGeometryDynamicCollection::ImplicitsAttribute, FTransformCollection::TransformGroup);
 			TBitArray<> EffectiveParticles;
 			NumEffectiveParticles = CalculateEffectiveParticles(GameThreadCollection, NumParticles, Parameters.MaxSimulatedLevel, Parameters.EnableClustering, GetOwner(), EffectiveParticles);
@@ -1593,7 +1605,7 @@ void FGeometryCollectionPhysicsProxy::InitializeBodiesPT(Chaos::FPBDRigidsSolver
 		// call DirtyParticle to make sure the acceleration structure is up to date with all the changes happening here
 		DirtyAllParticles(*RigidsSolver);
 
-		if (bHasBuiltGeometryOnPT)
+		if (bHasBuiltGeometryOnPT || bRemoveImplicitsInDynamicCollections)
 		{
 			// The Implicits attributes from the Dynamic Collection are just used for geometry initialization, after they can be removed and so free some memory.
 			PhysicsThreadCollection.RemoveAttribute(FGeometryDynamicCollection::ImplicitsAttribute, FTransformCollection::TransformGroup);
@@ -1610,6 +1622,12 @@ void FGeometryCollectionPhysicsProxy::CreateChildrenGeometry_Internal()
 
 		if (Parameters.Simulating && ensure(RestCollection))
 		{
+			if (bRemoveImplicitsInDynamicCollections)
+			{
+				PhysicsThreadCollection.AddAttribute<Chaos::FImplicitObjectPtr>(FGeometryDynamicCollection::ImplicitsAttribute, FTransformCollection::TransformGroup);
+				PhysicsThreadCollection.CopyAttribute(*Parameters.RestCollection, FGeometryDynamicCollection::ImplicitsAttribute, FTransformCollection::TransformGroup);
+			}
+
 			const int32 NumTransforms = PhysicsThreadCollection.NumElements(FTransformCollection::TransformGroup);
 			const TManagedArray<FTransform>& MassToLocal = RestCollection->GetAttribute<FTransform>(MassToLocalAttributeName, FTransformCollection::TransformGroup);
 			const TManagedArray<Chaos::FImplicitObjectPtr>& Implicits = PhysicsThreadCollection.GetAttribute<Chaos::FImplicitObjectPtr>(FGeometryDynamicCollection::ImplicitsAttribute, FTransformCollection::TransformGroup);
