Fix Cloth Collision ISPC bug from switching to new collision particle type. It might have still magically worked because of the order of the fields, but there's no guarantee!

diff --git a/Engine/Source/Runtime/Experimental/Chaos/Private/Chaos/PerParticlePBDCollisionConstraint.cpp b/Engine/Source/Runtime/Experimental/Chaos/Private/Chaos/PerParticlePBDCollisionConstraint.cpp
--- a/Engine/Source/Runtime/Experimental/Chaos/Private/Chaos/PerParticlePBDCollisionConstraint.cpp
+++ b/Engine/Source/Runtime/Experimental/Chaos/Private/Chaos/PerParticlePBDCollisionConstraint.cpp
@@ -121,7 +121,7 @@ namespace Chaos::Softs {
 // Helper function to call PhiWithNormal and return data to ISPC
 extern "C" void GetPhiWithNormal(const uint8* CollisionParticles, const FSolverReal* InV, FSolverReal* Normal, FSolverReal* Phi, const int32 i, const int32 ProgramCount, const int32 Mask)
 {
-	const TKinematicGeometryParticlesImp<FSolverReal, 3, EGeometryParticlesSimType::Other>& C = *(const TKinematicGeometryParticlesImp<FSolverReal, 3, EGeometryParticlesSimType::Other>*)CollisionParticles;
+	const FSolverCollisionParticles& C = *(const FSolverCollisionParticles*)CollisionParticles;
 	
 	for (int32 Index = 0; Index < ProgramCount; ++Index)
 	{
@@ -149,7 +149,7 @@ extern "C" void GetPhiWithNormal(const uint8* CollisionParticles, const FSolverR
 // Helper function to call PhiWithNormal and return data to ISPC
 extern "C" void GetPhiWithNormalAndVelocityBone(const uint8 * CollisionParticles, const FSolverReal * InV, FSolverReal * Normal, FSolverReal * Phi, int32 * VelocityBone, const int32 i, const FSolverReal Thickness, const int32 ProgramCount, const int32 Mask)
 {
-	const TKinematicGeometryParticlesImp<FSolverReal, 3, EGeometryParticlesSimType::Other>& C = *(const TKinematicGeometryParticlesImp<FSolverReal, 3, EGeometryParticlesSimType::Other>*)CollisionParticles;
+	const FSolverCollisionParticles& C = *(const FSolverCollisionParticles*)CollisionParticles;
 
 	for (int32 Index = 0; Index < ProgramCount; ++Index)
 	{
