Choas : Clustering
- Start the cluster unions sleeping.
- Make the sleeping state sticky as the cluster bounces between being anchored and not. Previously, it lost its sleeping state in favor of dynamic.

diff --git a/Engine/Source/Runtime/Experimental/Chaos/Private/Chaos/ClusterUnionManager.cpp b/Engine/Source/Runtime/Experimental/Chaos/Private/Chaos/ClusterUnionManager.cpp
--- a/Engine/Source/Runtime/Experimental/Chaos/Private/Chaos/ClusterUnionManager.cpp
+++ b/Engine/Source/Runtime/Experimental/Chaos/Private/Chaos/ClusterUnionManager.cpp
@@ -889,7 +889,8 @@ namespace Chaos
 				}
 				else
 				{
-					MEvolution.SetParticleObjectState(ClusterUnion.InternalCluster, EObjectStateType::Dynamic);
+					const EObjectStateType TransferedState = ClusterUnion.InternalCluster->IsSleeping() ? EObjectStateType::Sleeping : EObjectStateType::Dynamic;
+					MEvolution.SetParticleObjectState(ClusterUnion.InternalCluster, TransferedState);
 					MEvolution.SetParticleKinematicTarget(ClusterUnion.InternalCluster, FKinematicTarget());
 				}
 			}
