cluster union needs to properly reset its geometry to an FImplicitObjectUnionClustered in RemoveParticlesFromClusterUnionGeometry if it ends up empty

diff --git a/Engine/Source/Runtime/Experimental/Chaos/Private/Chaos/ClusterUnionManager.cpp b/Engine/Source/Runtime/Experimental/Chaos/Private/Chaos/ClusterUnionManager.cpp
--- a/Engine/Source/Runtime/Experimental/Chaos/Private/Chaos/ClusterUnionManager.cpp
+++ b/Engine/Source/Runtime/Experimental/Chaos/Private/Chaos/ClusterUnionManager.cpp
@@ -908,22 +908,7 @@ namespace Chaos
 		// parent transform is not set properly for a new empty cluster until UpdateClusterMassProperties is called for the first time.
 		if (EnumHasAnyFlags(Flags, EUpdateClusterUnionPropertiesFlags::ForceGenerateGeometry))
 		{
-			ClusterUnion.InternalCluster->SetGeometry(MakeImplicitObjectPtr<FImplicitObjectUnionClustered>());
-
-			ModifyAdditionOfChildrenToClusterUnionGeometry(
-				ClusterUnion.InternalCluster,
-				ClusterUnion.ChildParticles,
-				ClusterUnion.ClusterUnionParameters.ActorId,
-				ClusterUnion.ClusterUnionParameters.ComponentId,
-				[this, &ClusterUnion, &FullChildrenSet]()
-				{
-					ClusterUnion.Geometry = ForceRecreateClusterUnionGeometry(ClusterUnion);
-					UpdateGeometry(ClusterUnion.InternalCluster, FullChildrenSet, MClustering.GetChildrenMap(), ClusterUnion.Geometry, ClusterUnion.Parameters);
-				}
-			);
-
-			ClusterUnion.GeometryChildParticles = ClusterUnion.ChildParticles;
-			ClusterUnion.ClearAllPendingGeometryOperations();
+			ForceRegenerateGeometry(ClusterUnion, FullChildrenSet);
 		}
 		else if (EnumHasAnyFlags(Flags, EUpdateClusterUnionPropertiesFlags::IncrementalGenerateGeometry))
 		{
@@ -982,12 +967,40 @@ namespace Chaos
 		ClusterUnion.PendingConnectivityOperations.Empty();
 	}
 
+	DECLARE_CYCLE_STAT(TEXT("FClusterUnionManager::ForceRegenerateGeometry"), STAT_ForceRegenerateGeometry, STATGROUP_Chaos);
+	void FClusterUnionManager::ForceRegenerateGeometry(FClusterUnion& ClusterUnion, const TSet<FPBDRigidParticleHandle*>& FullChildrenSet)
+	{
+		SCOPE_CYCLE_COUNTER(STAT_ForceRegenerateGeometry);
+		ClusterUnion.InternalCluster->SetGeometry(MakeImplicitObjectPtr<FImplicitObjectUnionClustered>());
+
+		ModifyAdditionOfChildrenToClusterUnionGeometry(
+			ClusterUnion.InternalCluster,
+			ClusterUnion.ChildParticles,
+			ClusterUnion.ClusterUnionParameters.ActorId,
+			ClusterUnion.ClusterUnionParameters.ComponentId,
+			[this, &ClusterUnion, &FullChildrenSet]()
+			{
+				ClusterUnion.Geometry = ForceRecreateClusterUnionGeometry(ClusterUnion);
+				UpdateGeometry(ClusterUnion.InternalCluster, FullChildrenSet, MClustering.GetChildrenMap(), ClusterUnion.Geometry, ClusterUnion.Parameters);
+			}
+		);
+
+		ClusterUnion.GeometryChildParticles = ClusterUnion.ChildParticles;
+		ClusterUnion.ClearAllPendingGeometryOperations();
+	}
+
 	DECLARE_CYCLE_STAT(TEXT("FClusterUnionManager::FlushIncrementalGeometryOperations"), STAT_FlushIncrementalGeometryOperations, STATGROUP_Chaos);
 	void FClusterUnionManager::FlushIncrementalGeometryOperations(FClusterUnion& ClusterUnion)
 	{
 		SCOPE_CYCLE_COUNTER(STAT_FlushIncrementalGeometryOperations);
 		check(ClusterUnion.Geometry != nullptr);
 
+		if (ClusterUnion.Geometry->template IsA<FImplicitObjectUnionClustered>())
+		{
+			ForceRegenerateGeometry(ClusterUnion, TSet<FPBDRigidParticleHandle*>{ClusterUnion.ChildParticles});
+			return;
+		}
+
 		FImplicitObjectUnion* ImplicitUnion = ClusterUnion.Geometry->template AsA<FImplicitObjectUnion>();
 		check(ImplicitUnion != nullptr);
 
@@ -1042,10 +1055,16 @@ namespace Chaos
 		if (!PendingGeometryRemovals.IsEmpty())
 		{
 			RemoveParticlesFromClusterUnionGeometry(ClusterUnion.InternalCluster, PendingGeometryRemovals, ClusterUnion.GeometryChildParticles);
+			ClusterUnion.Geometry = ClusterUnion.InternalCluster->GetGeometry();
 			ClusterUnion.ClearPendingGeometryOperations(EClusterUnionGeometryOperation::Remove);
 		}
 
-		if (!PendingGeometryRefresh.IsEmpty())
+		check(ClusterUnion.Geometry != nullptr);
+		// Need to double check the geometry here. Did we get switched to a FImplicitObjectUnionClustered?
+		// In that case we need to skip the refresh.
+		const bool bIsUnionClustered = ClusterUnion.Geometry->template IsA<FImplicitObjectUnionClustered>();
+
+		if (!PendingGeometryRefresh.IsEmpty() && !bIsUnionClustered)
 		{
 			// For each particle we need to find the corresponding shape.
 			// Note that by the time we get to handling PendimgGeometryRefresh, we can once again
diff --git a/Engine/Source/Runtime/Experimental/Chaos/Public/Chaos/ClusterUnionManager.h b/Engine/Source/Runtime/Experimental/Chaos/Public/Chaos/ClusterUnionManager.h
--- a/Engine/Source/Runtime/Experimental/Chaos/Public/Chaos/ClusterUnionManager.h
+++ b/Engine/Source/Runtime/Experimental/Chaos/Public/Chaos/ClusterUnionManager.h
@@ -278,6 +278,9 @@ namespace Chaos
 		// Flush the cluster union's incremental connectivity operations
 		CHAOS_API void FlushIncrementalConnectivityGraphOperations(FClusterUnion& ClusterUnion);
 
+		// Forcefully regenerate the cluster union's geometry.
+		void ForceRegenerateGeometry(FClusterUnion& ClusterUnion, const TSet<FPBDRigidParticleHandle*>& FullChildrenSet);
+
 		// Flush the cluster union's incremental geometry operations.
 		void FlushIncrementalGeometryOperations(FClusterUnion& ClusterUnion);
 	};
@@ -408,5 +411,11 @@ namespace Chaos
 		RemoveArrayItemsAtSortedIndices(AllChildParticles, ShapeIndicesToRemove);
 
 		check(AllChildParticles.Num() == ClusterParticle->ShapesArray().Num());
+
+		// If we remove particles from the cluster union geometry then we need to switch the geometry back to a FImplicitObjectUnionClustered to avoid errors with empty unions.
+		if (ClusterParticle->ShapesArray().IsEmpty())
+		{
+			ClusterParticle->SetGeometry(MakeImplicitObjectPtr<FImplicitObjectUnionClustered>());
+		}
 	}
 }
