Fix logging error in geometry collection component
fallout from CL 29361608

diff --git a/Engine/Source/Runtime/Experimental/Chaos/Private/PhysicsProxy/GeometryCollectionPhysicsProxy.cpp b/Engine/Source/Runtime/Experimental/Chaos/Private/PhysicsProxy/GeometryCollectionPhysicsProxy.cpp
--- a/Engine/Source/Runtime/Experimental/Chaos/Private/PhysicsProxy/GeometryCollectionPhysicsProxy.cpp
+++ b/Engine/Source/Runtime/Experimental/Chaos/Private/PhysicsProxy/GeometryCollectionPhysicsProxy.cpp
@@ -570,7 +570,7 @@ void FGeometryCollectionPhysicsProxy::Initialize(Chaos::FPBDRigidsEvolutionBase
 		}
 	}
 
-	// we need to make sure the world transform is kept up to dat eon the game thread 
+	// we need to make sure the world transform is kept up to date on the game thread 
 	WorldTransform_External = Parameters.WorldTransform;
 
 	//
@@ -2891,9 +2891,11 @@ int32 FGeometryCollectionPhysicsProxy::CalculateEffectiveParticles(const FGeomet
 	EffectiveParticles.Init(false, NumTransform);
 	const int32 MaxSimulatedLevel = FMath::Min(GlobalMaxSimulatedLevel, InMaxSimulatedLevel);
 	const TManagedArray<Chaos::FImplicitObjectPtr>& Implicits = DynamicCollection.GetAttribute<Chaos::FImplicitObjectPtr>(FGeometryDynamicCollection::ImplicitsAttribute, FTransformCollection::TransformGroup);
+	TManagedArrayAccessor<int32> Levels = DynamicCollection.GetInitialLevels();
+
 	for (int32 TransformIndex = 0; TransformIndex < NumTransform; ++TransformIndex)
 	{
-		const int32 Level = FMath::Clamp(CalculateHierarchyLevel(DynamicCollection, TransformIndex), 0, INT_MAX);
+		const int32 Level = Levels[TransformIndex];
 		if (Level <= MaxSimulatedLevel || !bEnableClustering)
 		{
 			const bool bIsClusterUsingChildGeometry = (ClustersUsingChildGeometry.Num() > 0) && ClustersUsingChildGeometry[TransformIndex];
