Iris:
* Added cvar to force all RPC to be held in the reliable queue

#rnx

diff --git a/Engine/Source/Runtime/Experimental/Iris/Core/Private/Iris/ReplicationSystem/AttachmentReplication.cpp b/Engine/Source/Runtime/Experimental/Iris/Core/Private/Iris/ReplicationSystem/AttachmentReplication.cpp
--- a/Engine/Source/Runtime/Experimental/Iris/Core/Private/Iris/ReplicationSystem/AttachmentReplication.cpp
+++ b/Engine/Source/Runtime/Experimental/Iris/Core/Private/Iris/ReplicationSystem/AttachmentReplication.cpp
@@ -31,6 +31,9 @@ namespace AttachmentReplicationCVars
 
 	static int32 MaxSimultaneousObjectsWithRPCs = 4096;
 	FAutoConsoleVariableRef CVarMaxSimultaneousObjectsWithRPCs(TEXT("net.MaxSimultaneousObjectsWithRPCs"), MaxSimultaneousObjectsWithRPCs, TEXT("Maximum number of objects that can have unsent RPCs at the same time. "));
+
+	static bool bMakeAllRPCsReliable = false;
+	FAutoConsoleVariableRef CVarAllRPCsReliable(TEXT("net.Iris.AllRPCReliable"), bMakeAllRPCsReliable, TEXT("Forces all the queued RPCs to be reliable. "));
 }
 
 static const FName NetError_UnreliableQueueFull("Unreliable attachment queue full");
@@ -187,7 +190,7 @@ FNetObjectAttachmentSendQueue::~FNetObjectAttachmentSendQueue()
 bool FNetObjectAttachmentSendQueue::Enqueue(TArrayView<const TRefCountPtr<FNetBlob>> Attachments)
 {
 	const FNetBlobCreationInfo& CreationInfo = Attachments[0]->GetCreationInfo();
-	if (EnumHasAnyFlags(CreationInfo.Flags, ENetBlobFlags::Reliable))
+	if (EnumHasAnyFlags(CreationInfo.Flags, ENetBlobFlags::Reliable) || UE::Net::Private::AttachmentReplicationCVars::bMakeAllRPCsReliable)
 	{
 		if (ReliableQueue == nullptr)
 		{
