Fix a rare crash if the game fails to resolve the distributed endpoint correctly.

#rnx

### FOnDemandIoBackend
- Resolving the distributed endpoint is now done as one of the first things we do on the Ias Http thread. Since the thread can do no meaningful work until the distributed endpoint is resolved, If the attempt to resolve fails we will wait for a set length of time then attempt to resolve it again until we resolve successfully.
-- This resolving loop will be canceled when the game is shutting down by the request to shutdown the Ias Http thread. We now do this by implementing ::Stop and gathering all the required operations there.
- If the distributed endpoint failed to resolve by the time that the game is shutdown then it is likely that we will have a number of pending requests built up and not dealing with them will result in a crash when cleaning up our allocators. So after the Ias Http thread has been shutdown we must iterate over any pending request and cancel them.

### FDistributionEndpoints
- Removed the async requests from FDistributionEndpoints and replaced with a synchronous call. This call will accept an FEvent that another thread can use to cancel the request which we use when shutting down the game to prevent long stalls.

diff --git a/Engine/Source/Runtime/Experimental/IoStoreOnDemand/Private/DistributionEndpoints.cpp b/Engine/Source/Runtime/Experimental/IoStoreOnDemand/Private/DistributionEndpoints.cpp
--- a/Engine/Source/Runtime/Experimental/IoStoreOnDemand/Private/DistributionEndpoints.cpp
+++ b/Engine/Source/Runtime/Experimental/IoStoreOnDemand/Private/DistributionEndpoints.cpp
@@ -17,199 +17,86 @@
 namespace UE::IO::IAS
 {
 
-static int32 GIasHttpDistributionRetryCount = 2;
-static FAutoConsoleVariableRef CVar_IasHttpDistributionRetryCount(
-	TEXT("ias.HttpDistributionRetryCount"),
-	GIasHttpDistributionRetryCount,
-	TEXT("Number of HTTP distribution request retries.")
+static int32 GDistributedEndpointTimeout = 30;
+static FAutoConsoleVariableRef CVar_DistributedEndpointTimeout(
+	TEXT("ias.DistributedEndpointTimeout"),
+	GDistributedEndpointTimeout,
+	TEXT("How long to wait (in seconds) for a distributed endoint resolve request before timing out")
 );
 
-
-FDistributionEndpoints::~FDistributionEndpoints()
+FDistributionEndpoints::EResult FDistributionEndpoints::ResolveEndpoints(const FString& DistributionUrl, TArray<FString>& OutServiceUrls)
 {
-	CancelEndpointRequests();
+	FEventRef Event;
+	return ResolveEndpoints(DistributionUrl, OutServiceUrls, *Event.Get());
 }
 
-#if IS_PROGRAM || WITH_EDITOR
-
-bool FDistributionEndpoints::Flush(double TimeOut)
+FDistributionEndpoints::EResult FDistributionEndpoints::ResolveEndpoints(const FString& DistributionUrl, TArray<FString>& OutServiceUrls, FEvent& Event)
 {
-	TRACE_CPUPROFILER_EVENT_SCOPE(FDistributionEndpoints::Flush);
+	TRACE_CPUPROFILER_EVENT_SCOPE(FDistributionEndpoints::ResolveEndpoints);
 
-	if (!bInitialized)
-	{
-		ResolveDeferredEndpoints();
-	}
+	std::atomic_bool bHasResponse = false;
+	EResult Result = EResult::Failure;
 
-	FHttpManager& HttpManager = FHttpModule::Get().GetHttpManager();
+	// The timeout is a cvar but should not be less than 10 seconds.
+	const float HttpTimeout = FMath::Max(static_cast<float>(GDistributedEndpointTimeout), 10.0f);
 
-	const double StartTime = FPlatformTime::Seconds();
+	UE_LOG(LogIas, Log, TEXT("Resolving distributed endpoint '%s'"), *DistributionUrl);
 
-	while (!PendingRequests.IsEmpty())
-	{
-		HttpManager.Tick(0.0);
-
-		FPlatformProcess::SleepNoStats(0.0f);
-
-		if (TimeOut > 0.0 && (FPlatformTime::Seconds() - StartTime) > TimeOut)
+	FHttpRequestPtr HttpRequest = FHttpModule::Get().CreateRequest();
+	HttpRequest->SetTimeout(HttpTimeout);
+	HttpRequest->SetURL(DistributionUrl);
+	HttpRequest->SetVerb(TEXT("GET"));
+	HttpRequest->SetHeader(TEXT("Accept"), TEXT("application/json"));
+	HttpRequest->OnProcessRequestComplete().BindLambda(
+		[this, &bHasResponse, &Event, &Result, &OutServiceUrls](FHttpRequestPtr, FHttpResponsePtr Response, bool bOk)
 		{
-			return false;
-		}
-	}
+			LLM_SCOPE_BYTAG(Ias);
+			
+			bHasResponse = true;
+			Result = ParseResponse(Response, OutServiceUrls);
 
-	return true;
-}
+			Event.Trigger();
+		});
 
-#endif // IS_PROGRAM || WITH_EDITOR
+	HttpRequest->ProcessRequest();
 
-void FDistributionEndpoints::ResolveEndpoints(const FString& DistributionUrl, FOnEndpointResolved&& OnResolved)
-{
-	TRACE_CPUPROFILER_EVENT_SCOPE(FOnDemandIoBackend::ResolveEndpoints);
-	const FResolvedEndpoint* Ep = nullptr;
-	{
-		FReadScopeLock _(Lock);
-		if (TUniquePtr<FResolvedEndpoint>* Entry = ResolvedEndpoints.Find(DistributionUrl))
-		{
-			Ep = Entry->Get();
-		}
-	}
+	const uint32 WaitTime = GDistributedEndpointTimeout >= 0 ? (static_cast<uint32>(GDistributedEndpointTimeout) * 1000) : MAX_uint32;
 
-	if (Ep != nullptr)
+	if (!Event.Wait(FTimespan::FromSeconds(WaitTime)) || bHasResponse == false)
 	{
-		return OnResolved(DistributionUrl, Ep->ServiceUrls);
+		HttpRequest->CancelRequest();
+		HttpRequest->OnProcessRequestComplete().Unbind();
 	}
 
-	bool bIssueRequest = false;
-	{
-		FWriteScopeLock _(Lock);
-		TUniquePtr<FResolveRequest>& Request = PendingRequests.FindOrAdd(DistributionUrl);
-		if (!Request.IsValid())
-		{
-			Request.Reset(new FResolveRequest{ DistributionUrl });
-			bIssueRequest = bInitialized;
-		}
-		Request->Callbacks.Add(MoveTemp(OnResolved));
-	}
-
-	if (bIssueRequest)
-	{
-		IssueEndpointRequests();
-	}
-}
-
-void FDistributionEndpoints::ResolveDeferredEndpoints()
-{
-	TRACE_CPUPROFILER_EVENT_SCOPE(FOnDemandIoBackend::ResolveDeferredEndpoints);
-	{
-		FWriteScopeLock _(Lock);
-		bInitialized = true;
-	}
+	UE_CLOG(Result == EResult::Success, LogIas, Log, TEXT("Successfully resolved distributed endpoint '%s' %d urls found"), *DistributionUrl, OutServiceUrls.Num());
+	UE_CLOG(Result != EResult::Success, LogIas, Log, TEXT("Failed to resolve distributed endpoint '%s'"), *DistributionUrl);
 
-	IssueEndpointRequests();
+	return Result;
 }
 
-void FDistributionEndpoints::IssueEndpointRequests()
+FDistributionEndpoints::EResult FDistributionEndpoints::ParseResponse(FHttpResponsePtr HttpResponse, TArray<FString>& OutUrls)
 {
-	TRACE_CPUPROFILER_EVENT_SCOPE(FOnDemandIoBackend::IssueEndpointRequests);
-	// Currently we need to use the HTTP module in order to resolve service endpoints due to HTTPS
-	FHttpModule& HttpModule = FHttpModule::Get();
-	const int32 MaxAttempts = GIasHttpDistributionRetryCount;
-
-	TArray<FHttpRequestPtr, TInlineAllocator<2>> HttpRequests;
-	{
-		FWriteScopeLock _(Lock);
-		check(bInitialized);
-
-		for (const TPair<FString, TUniquePtr<FResolveRequest>>& Kv : PendingRequests)
-		{
-			if (Kv.Value->HttpRequest.IsValid())
-			{
-				continue;
-			}
-
-			FResolveRequest& ResolveRequest = *Kv.Value.Get();
-			UE_LOG(LogIas, Log, TEXT("Resolving '%s' (#%d/%d)"), *ResolveRequest.DistributionUrl, ResolveRequest.RetryCount + 1, MaxAttempts);
-
-			FHttpRequestPtr HttpRequest = HttpModule.CreateRequest();
-			HttpRequest->SetTimeout(3.0f);
-			HttpRequest->SetURL(Kv.Key);
-			HttpRequest->SetVerb(TEXT("GET"));
-			HttpRequest->SetHeader(TEXT("Accept"), TEXT("application/json"));
-			HttpRequest->OnProcessRequestComplete().BindLambda(
-				[this, &ResolveRequest, MaxAttempts](FHttpRequestPtr, FHttpResponsePtr Response, bool bOk)
-				{
-					LLM_SCOPE_BYTAG(Ias);
-					FHttpRequestPtr Request = MoveTemp(ResolveRequest.HttpRequest);
-
-					// Response will be null if the connection timed out
-					if (Response == nullptr || Response->GetResponseCode() != 200)
-					{
-						if (++ResolveRequest.RetryCount <= MaxAttempts)
-						{
-							FDistributionEndpoints* CachedPtr = this;
-							Request->OnProcessRequestComplete().Unbind(); // <- Invalidates the lambdas captures
-
-							return CachedPtr->IssueEndpointRequests();
-						}
-					}
-
-					this->CompleteEndpointRequest(ResolveRequest, Response);
-				});
-
-			ResolveRequest.HttpRequest = HttpRequest;
-			HttpRequests.Add(HttpRequest);
-		}
-	}
-
-	for (FHttpRequestPtr& Request : HttpRequests)
-	{
-		Request->ProcessRequest();
-	}
-}
-
-void FDistributionEndpoints::CancelEndpointRequests()
-{
-	TArray<FHttpRequestPtr, TInlineAllocator<2>> HttpRequests;
-	{
-		FWriteScopeLock _(Lock);
-		for (const TPair<FString, TUniquePtr<FResolveRequest>>& Kv : PendingRequests)
-		{
-			if (Kv.Value->HttpRequest.IsValid())
-			{
-				HttpRequests.Add(Kv.Value->HttpRequest);
-			}
-		}
-	}
-
-	if (!HttpRequests.IsEmpty())
-	{
-		FHttpModule& HttpModule = FHttpModule::Get();
-		for (FHttpRequestPtr& Request : HttpRequests)
-		{
-			HttpModule.GetHttpManager().RemoveRequest(Request.ToSharedRef());
-			//TODO: Flush?
-		}
-	}
-}
-
-void FDistributionEndpoints::CompleteEndpointRequest(FResolveRequest& ResolveRequest, FHttpResponsePtr HttpResponse)
-{
-	TRACE_CPUPROFILER_EVENT_SCOPE(FOnDemandIoBackend::CompleteEndpointRequest);
+	TRACE_CPUPROFILER_EVENT_SCOPE(FDistributionEndpoints::ParseResponse);
 
 	using FJsonValuePtr = TSharedPtr<FJsonValue>;
 	using FJsonObjPtr = TSharedPtr<FJsonObject>;
 	using FJsonReader = TJsonReader<TCHAR>;
 	using FJsonReaderPtr = TSharedRef<FJsonReader>;
 
-	TArray<FString> ServiceUrls;
 	if (HttpResponse != nullptr && HttpResponse->GetResponseCode() == 200)
 	{
 		FString Json = HttpResponse->GetContentAsString();
+		
 		FJsonReaderPtr JsonReader = TJsonReaderFactory<TCHAR>::Create(Json);
 
 		FJsonObjPtr JsonObj;
 		if (FJsonSerializer::Deserialize(JsonReader, JsonObj))
 		{
+			if (!JsonObj->HasTypedField< EJson::Array>(TEXT("distributions")))
+			{
+				return EResult::Failure;
+			}
+
 			TArray<FJsonValuePtr> JsonValues = JsonObj->GetArrayField(TEXT("distributions"));
 			for (const FJsonValuePtr& JsonValue : JsonValues)
 			{
@@ -218,35 +105,14 @@ void FDistributionEndpoints::CompleteEndpointRequest(FResolveRequest& ResolveReq
 				{
 					ServiceUrl.LeftInline(ServiceUrl.Len() - 1);
 				}
-				ServiceUrls.Add(MoveTemp(ServiceUrl));
+				OutUrls.Add(MoveTemp(ServiceUrl));
 			}
-		}
-	}
 
-	const FResolvedEndpoint* ResolvedEndpoint = nullptr;
-	FString DistributionUrl;
-	TArray<FOnEndpointResolved> Callbacks;
-
-	{
-		FWriteScopeLock _(Lock);
-		if (!ServiceUrls.IsEmpty())
-		{
-			ResolvedEndpoint = ResolvedEndpoints.Emplace(
-				ResolveRequest.DistributionUrl,
-				new FResolvedEndpoint{ MoveTemp(ServiceUrls) })
-				.Get();
+			return !OutUrls.IsEmpty() ? EResult::Success : EResult::Failure;
 		}
-
-		Callbacks = MoveTemp(ResolveRequest.Callbacks);
-		DistributionUrl = MoveTemp(ResolveRequest.DistributionUrl);
-		PendingRequests.Remove(DistributionUrl);
 	}
 
-	TConstArrayView<FString> Urls = ResolvedEndpoint ? ResolvedEndpoint->ServiceUrls : TConstArrayView<FString>();
-	for (FOnEndpointResolved& Callback : Callbacks)
-	{
-		Callback(DistributionUrl, Urls);
-	}
+	return EResult::Failure;
 }
 
 } // namespace UE::IO::IAS
diff --git a/Engine/Source/Runtime/Experimental/IoStoreOnDemand/Private/DistributionEndpoints.h b/Engine/Source/Runtime/Experimental/IoStoreOnDemand/Private/DistributionEndpoints.h
--- a/Engine/Source/Runtime/Experimental/IoStoreOnDemand/Private/DistributionEndpoints.h
+++ b/Engine/Source/Runtime/Experimental/IoStoreOnDemand/Private/DistributionEndpoints.h
@@ -14,41 +14,35 @@ namespace UE::IO::IAS
 class FDistributionEndpoints
 {
 public:
-	using FOnEndpointResolved = TFunction<void(const FString&, TConstArrayView<FString>)>;
-
-	FDistributionEndpoints() = default;
-	~FDistributionEndpoints();
-
-#if IS_PROGRAM || WITH_EDITOR
-	bool Flush(double TimeOut);
-#endif //IS_PROGRAM || WITH_EDITOR
-
-	void ResolveEndpoints(const FString& DistributionUrl, FOnEndpointResolved&& OnResolved);
-	void ResolveDeferredEndpoints();
-
-private:
-	struct FResolvedEndpoint
+	enum class EResult : uint8
 	{
-		TArray<FString> ServiceUrls;
+		Success = 0,
+		Failure
 	};
 
-	struct FResolveRequest
-	{
-		FString DistributionUrl;
-		FHttpRequestPtr HttpRequest;
-		TArray<FOnEndpointResolved> Callbacks;
-		int32 RetryCount = 0;
-	};
-
-	void IssueEndpointRequests();
-	void CancelEndpointRequests();
-	void CompleteEndpointRequest(FResolveRequest& ResolveRequest, FHttpResponsePtr HttpResponse);
+	FDistributionEndpoints() = default;
+	~FDistributionEndpoints() = default;
+
+	/**
+	 * Queries the distributed endpoint in order to get a list of service urls from which we can download data.
+	 * 
+	 * @param DistributionUrl	The url of the distributed endpoint to resolve
+	 * @param OutServiceUrls	An array into which the resolved urls will be placed
+	 */
+	FDistributionEndpoints::EResult ResolveEndpoints(const FString& DistributionUrl, TArray<FString>& OutServiceUrls);
+
+	/**
+	 * Queries the distributed endpoint in order to get a list of service urls from which we can download data.
+	 * 
+	 * @param DistributionUrl	The url of the distributed endpoint to resolve
+	 * @param OutServiceUrls	An array into which the resolved urls will be placed
+	 * @param Event				An event which can be triggered to cancel the request
+	 */
+	FDistributionEndpoints::EResult ResolveEndpoints(const FString& DistributionUrl, TArray<FString>& OutServiceUrls, FEvent& Event);
 
-	TMap<FString, TUniquePtr<FResolvedEndpoint>> ResolvedEndpoints;
-	TMap<FString, TUniquePtr<FResolveRequest>> PendingRequests;
+private:
 
-	FRWLock Lock;
-	bool bInitialized = false;
+	FDistributionEndpoints::EResult ParseResponse(FHttpResponsePtr HttpResponse, TArray<FString>& OutUrls);
 };
 
 } // namespace UE::IO::IAS
diff --git a/Engine/Source/Runtime/Experimental/IoStoreOnDemand/Private/OnDemandIoDispatcherBackend.cpp b/Engine/Source/Runtime/Experimental/IoStoreOnDemand/Private/OnDemandIoDispatcherBackend.cpp
--- a/Engine/Source/Runtime/Experimental/IoStoreOnDemand/Private/OnDemandIoDispatcherBackend.cpp
+++ b/Engine/Source/Runtime/Experimental/IoStoreOnDemand/Private/OnDemandIoDispatcherBackend.cpp
@@ -29,9 +29,11 @@
 #include "IasCache.h"
 #include "Math/NumericLimits.h"
 #include "Misc/CommandLine.h"
+#include "Misc/ConfigCacheIni.h"
 #include "Misc/EnumClassFlags.h"
 #include "Misc/PathViews.h"
 #include "Misc/Paths.h"
+#include "Misc/ScopeExit.h"
 #include "Misc/ScopeLock.h"
 #include "Misc/ScopeRWLock.h"
 #include "OnDemandHttpClient.h"
@@ -220,6 +222,13 @@ static FAutoConsoleVariableRef CVar_IasHttpRangeRequestMinSizeKiB(
 	TEXT("Minimum chunk size for partial chunk request(s)")
 );
 
+static int32 GDistributedEndpointRetryWaitTime = 15;
+static FAutoConsoleVariableRef CVar_DistributedEndpointRetryWaitTime(
+	TEXT("ias.DistributedEndpointRetryWaitTime"),
+	GDistributedEndpointRetryWaitTime,
+	TEXT("How long to wait (in seconds) after failing to resolve a distributed endpoint before retrying")
+);
+
 #if !UE_BUILD_SHIPPING
 static FAutoConsoleCommand CVar_IasAbandonCache(
 	TEXT("Ias.AbandonCache"),
@@ -1535,11 +1544,13 @@ public:
 
 	// Runnable
 	virtual bool Init() override { return true; }
-	virtual void Stop() override { bStopRequested = true; }
 	virtual uint32 Run() override;
+	virtual void Stop() override;
 
 private:
 
+	void CancelAllPendingRequests();
+
 	FString GetEndpointTestPath() const;
 	void ConditionallyStartBackendThread();
 	void CompleteRequest(FChunkRequest* ChunkRequest);
@@ -1548,6 +1559,7 @@ private:
 	FIoStatus ApplyGeneratedOnDemandToc(const FString& CdnUrl, const FString& TocPath);
 	FIoStatus DownloadoadOnDemandToc(const FString& CdnUrl, const FString& TocPath);
 
+	bool ResolveDistributedEndpoint(const FString& Url);
 	void InitializePrimaryEndpoint();
 
 	/** Mode to control which types of .iochunktoc are flush */
@@ -1563,6 +1575,8 @@ private:
 	};
 
 	void FlushDeferredTocs(EFlushMode FlushMode);
+
+	bool SetupHttpThread();
 	void ProcessHttpRequests(FHttpClient& HttpClient, FBitWindow& HttpErrors, int32 MaxConcurrentRequests);
 	int32 WaitForPendingRequests(float WaitTimeSeconds, float PollTimeSeconds);
 
@@ -1586,6 +1600,8 @@ private:
 	FBackendStatus BackendStatus;
 	FAvailableEps AvailableEps;
 	FString DistributionUrl;
+	FEventRef DistributedEndpointEvent;
+
 	mutable FRWLock Lock;
 	std::atomic_uint32_t InflightCacheRequestCount{0};
 	std::atomic_bool bStopRequested{false};
@@ -1642,35 +1658,7 @@ void FOnDemandIoBackend::Initialize(TSharedRef<const FIoDispatcherBackendContext
 	UE_LOG(LogIas, Log, TEXT("Initializing on demand I/O dispatcher backend"));
 	BackendContext = Context;
 
-	// We need to resolve the end point in this method which occurs after the config system has initialized
-	// rather than in ::Mount which can occur before that.
-	// Without the config system initialized the http module will not work properly and we will always fail
-	// to resolve and the OnDemand system will not recover.
-	if (DistributionUrl.IsEmpty() == false)
-	{
-		TSharedPtr<FDistributionEndpoints> Resolver = MakeShared<FDistributionEndpoints>();
-		Resolver->ResolveEndpoints(DistributionUrl,
-			[this, Resolver](const FString& DistributionEp, TConstArrayView<FString> Eps)
-			{
-				UE_CLOG(Eps.IsEmpty(), LogIas, Error, TEXT("Failed to resolve available endpoint(s) from '%s'"), *DistributionEp);
-				{
-					FWriteScopeLock _(Lock);
-					for (const FString& Ep : Eps)
-					{
-						AvailableEps.Urls.Add(Ep.Replace(TEXT("https"), TEXT("http")));
-					}
-				}
-
-				// A disk only flush should be fairly quick at this point, especially if 's.IasEnableThreadedTocGeneration' is true
-				// and this will help mount .iochunktoc as soon as possible rather than waiting for the primary endpoint to be
-				// chosen.
-				FlushDeferredTocs(EFlushMode::Disk);
-
-				ConditionallyStartBackendThread();
-			});
-		Resolver->ResolveDeferredEndpoints();
-		DistributionUrl.Reset();
-	}
+	ConditionallyStartBackendThread();
 }
 
 void FOnDemandIoBackend::Shutdown()
@@ -1680,11 +1668,18 @@ void FOnDemandIoBackend::Shutdown()
 		return;
 	}
 
+	TRACE_CPUPROFILER_EVENT_SCOPE(IasBackend::Shutdown);
+
 	UE_LOG(LogIas, Log, TEXT("Shutting down on demand I/O dispatcher backend"));
 
-	bStopRequested = true;
-	TickBackendEvent->Trigger();
-	BackendThread.Reset();
+	Stop();
+
+	{
+		TRACE_CPUPROFILER_EVENT_SCOPE(IasBackend::Shutdown::StoppingIasHttpThread);
+		BackendThread.Reset();
+	}
+
+	CancelAllPendingRequests();	
 
 	const int32 NumPending = WaitForPendingRequests(5.0f, 0.1f);
 	UE_CLOG(NumPending > 0, LogIas, Warning, TEXT("%d request(s) still pending after shutdown"), NumPending);
@@ -1692,6 +1687,24 @@ void FOnDemandIoBackend::Shutdown()
 	BackendContext.Reset();
 }
 
+void FOnDemandIoBackend::CancelAllPendingRequests()
+{
+	FChunkRequest* Iterator = HttpRequests.Dequeue();
+	while (Iterator != nullptr)
+	{
+		FChunkRequest* Request = Iterator;
+		Iterator = Iterator->NextRequest;
+
+		// We need to call this to increment the number of requests in flight or the subsequent call to Stats.OnHttpCancel will start to mismatch that counter.
+		Stats.OnHttpDequeue();
+
+		Request->bCancelled = true;
+		CompleteRequest(Request);
+		
+		Stats.OnHttpCancel();
+	}
+}
+
 FString FOnDemandIoBackend::GetEndpointTestPath() const
 {
 	FString TestPath = IoStore.IsValid() ? IoStore->GetFirstTocPath() : FString();
@@ -2062,6 +2075,45 @@ FIoStatus FOnDemandIoBackend::DownloadoadOnDemandToc(const FString& CdnUrl, cons
 	}
 }
 
+bool FOnDemandIoBackend::ResolveDistributedEndpoint(const FString& DistributedEndpointUrl)
+{
+	TRACE_CPUPROFILER_EVENT_SCOPE(IasBackend::ResolveDistributedEndpoint);
+
+	check(!DistributedEndpointUrl.IsEmpty());
+
+	// We need to resolve the end point in this method which occurs after the config system has initialized
+	// rather than in ::Mount which can occur before that.
+	// Without the config system initialized the http module will not work properly and we will always fail
+	// to resolve and the OnDemand system will not recover.
+	check(GConfig->IsReadyForUse());
+
+	while (!bStopRequested)
+	{
+		TArray<FString> ServiceUrls;
+
+		FDistributionEndpoints Resolver;
+		FDistributionEndpoints::EResult Result = Resolver.ResolveEndpoints(DistributedEndpointUrl, ServiceUrls, *DistributedEndpointEvent.Get());
+		if (Result == FDistributionEndpoints::EResult::Success)
+		{
+			FWriteScopeLock _(Lock);
+			for (const FString& Url : ServiceUrls)
+			{
+				AvailableEps.Urls.Add(Url.Replace(TEXT("https"), TEXT("http")));
+			}
+
+			return true;
+		}
+
+		if (!bStopRequested)
+		{
+			const uint32 WaitTime = GDistributedEndpointRetryWaitTime >= 0 ? (static_cast<uint32>(GDistributedEndpointRetryWaitTime) * 1000) : MAX_uint32;
+			DistributedEndpointEvent->Wait(WaitTime);
+		}
+	}
+
+	return false;
+}
+
 void FOnDemandIoBackend::InitializePrimaryEndpoint()
 {
 	TRACE_CPUPROFILER_EVENT_SCOPE(IasBackend::InitializePrimaryEndpoint);
@@ -2219,8 +2271,6 @@ void FOnDemandIoBackend::Mount(const FOnDemandEndpoint& Endpoint)
 					DeferredTocs.Add(FTocParams{Endpoint.TocPath, Endpoint.bForceTocDownload});
 			}
 		}
-
-		ConditionallyStartBackendThread();
 	}
 	else
 	{
@@ -2397,11 +2447,38 @@ void FOnDemandIoBackend::ProcessHttpRequests(FHttpClient& HttpClient, FBitWindow
 	} 
 }
 
+bool FOnDemandIoBackend::SetupHttpThread()
+{
+	// A disk only flush should be fairly quick at this point, especially if 's.IasEnableThreadedTocGeneration' is true
+	// and this will help mount .iochunktoc as soon as possible rather than waiting for the primary endpoint to be
+	// chosen.
+	FlushDeferredTocs(EFlushMode::Disk);
+
+	// Note that the following method will block until the distributed endpoint (if we have one) has
+	// been resolved as the thread cannot function without it.
+	if (!DistributionUrl.IsEmpty())
+	{
+		if (!ResolveDistributedEndpoint(DistributionUrl))
+		{
+			return false;
+		}
+		DistributionUrl.Reset();
+	}
+
+	InitializePrimaryEndpoint();
+
+	return true;
+}
+
 uint32 FOnDemandIoBackend::Run()
 {
 	LLM_SCOPE_BYTAG(Ias);
 
-	InitializePrimaryEndpoint();
+	// This call can take some time depending on the IAS setup and the network conditions
+	if (!SetupHttpThread())
+	{
+		return 0;
+	}
 
 	FBitWindow HttpErrors;
 	HttpErrors.Reset(GIasHttpErrorSampleCount);
@@ -2485,6 +2562,13 @@ uint32 FOnDemandIoBackend::Run()
 	return 0;
 }
 
+void FOnDemandIoBackend::Stop()
+{
+	bStopRequested = true;
+	TickBackendEvent->Trigger();
+	DistributedEndpointEvent->Trigger();
+}
+
 int32 FOnDemandIoBackend::WaitForPendingRequests(float WaitTimeSeconds, float PollTimeSeconds)
 {
 	const double StartTime = FPlatformTime::Seconds();
diff --git a/Engine/Source/Runtime/Experimental/IoStoreOnDemand/Private/PrimeEndpoint.cpp b/Engine/Source/Runtime/Experimental/IoStoreOnDemand/Private/PrimeEndpoint.cpp
--- a/Engine/Source/Runtime/Experimental/IoStoreOnDemand/Private/PrimeEndpoint.cpp
+++ b/Engine/Source/Runtime/Experimental/IoStoreOnDemand/Private/PrimeEndpoint.cpp
@@ -34,17 +34,14 @@ static TIoStatusOr<TArray<FString>> FindCDNUrls(const FString& DistributionUrl)
 	TArray<FString> CDNUrls;
 
 	FDistributionEndpoints DistributionInfo;
-	DistributionInfo.ResolveEndpoints(DistributionUrl, [&CDNUrls](const FString& DistributionUrl, TConstArrayView<FString> Urls)
-		{
-			CDNUrls = Urls;
-		});
-
-	if (!DistributionInfo.Flush(30.0))
+	if (DistributionInfo.ResolveEndpoints(DistributionUrl, CDNUrls) == FDistributionEndpoints::FDistributionEndpoints::EResult::Success)
 	{
-		return FIoStatus(EIoErrorCode::Unknown, TEXT("Unable to connect to endpoint, timed out..."));
+		return CDNUrls;
+	}
+	else
+	{
+		return FIoStatus(EIoErrorCode::Unknown, TEXT("Unable to connect to distributed endpoint"));
 	}
-
-	return CDNUrls;
 }
 
 /** Utility to format time in seconds to a more human readable form */
