ChaosSofts : Fix compile warning 1053

diff --git a/Engine/Source/Runtime/Experimental/Chaos/Public/Chaos/Deformable/GaussSeidelCorotatedConstraints.h b/Engine/Source/Runtime/Experimental/Chaos/Public/Chaos/Deformable/GaussSeidelCorotatedConstraints.h
--- a/Engine/Source/Runtime/Experimental/Chaos/Public/Chaos/Deformable/GaussSeidelCorotatedConstraints.h
+++ b/Engine/Source/Runtime/Experimental/Chaos/Public/Chaos/Deformable/GaussSeidelCorotatedConstraints.h
@@ -72,7 +72,7 @@ namespace Chaos::Softs
 				X_k.Init(TVector<T, 3>(T(0.)), ParticleEndIndex - ParticleStartIndex);
 			}
 			InitColor(InParticles);
-			InitializeLambdas();
+			InitializeCorotatedLambdas();
 		}
 
 		virtual ~FGaussSeidelCorotatedConstraints() {}
@@ -282,7 +282,7 @@ namespace Chaos::Softs
 			ParticlesPerColor = ComputeNodalColoring(MeshConstraints, Particles, ParticleStartIndex, ParticleEndIndex, IncidentElements, IncidentElementsLocal);
 		}
 
-		 virtual void InitializeLambdas()
+		 void InitializeCorotatedLambdas()
 		 {
 			 ComputeStress = [](const Chaos::PMatrix<T, 3, 3>& Fe, const T mu, const T lambda, Chaos::PMatrix<T, 3, 3>& P)
 			 {
diff --git a/Engine/Source/Runtime/Experimental/Chaos/Public/Chaos/Deformable/GaussSeidelNeohookeanConstraints.h b/Engine/Source/Runtime/Experimental/Chaos/Public/Chaos/Deformable/GaussSeidelNeohookeanConstraints.h
--- a/Engine/Source/Runtime/Experimental/Chaos/Public/Chaos/Deformable/GaussSeidelNeohookeanConstraints.h
+++ b/Engine/Source/Runtime/Experimental/Chaos/Public/Chaos/Deformable/GaussSeidelNeohookeanConstraints.h
@@ -55,14 +55,14 @@ namespace Chaos::Softs
 		)
 			: GSBase(InParticles, InMesh, EMeshArray, NuMeshArray, MoveTemp(AlphaJMeshArray), MoveTemp(IncidentElementsIn), MoveTemp(IncidentElementsLocalIn), ParticleStartIndexIn, ParticleEndIndexIn, bDoQuasistaticsIn, bDoSORIn, InOmegaSOR, InParams, NuMesh, bRecordMetricIn)
 		{
-			InitializeLambdas();
+			InitializeNeohookeanLambdas();
 		}
 
 		virtual ~FGaussSeidelNeohookeanConstraints() {}
 
 	protected:
 
-		void InitializeLambdas() override
+		void InitializeNeohookeanLambdas() 
 		{
 			ComputeStress = [](const Chaos::PMatrix<T, 3, 3>& Fe, const T mu, const T lambda, Chaos::PMatrix<T, 3, 3>& P)
 			{
