Added the ability to set ISM components in Geometry Collection ISM pools to just use their parent bounds and change their transform mode to absolute.
This allows the actor containing the ISM components to have full control over the root transforms of each ISM for faster or customised processing.

diff --git a/Engine/Source/Runtime/Engine/Classes/Components/InstancedStaticMeshComponent.h b/Engine/Source/Runtime/Engine/Classes/Components/InstancedStaticMeshComponent.h
--- a/Engine/Source/Runtime/Engine/Classes/Components/InstancedStaticMeshComponent.h
+++ b/Engine/Source/Runtime/Engine/Classes/Components/InstancedStaticMeshComponent.h
@@ -247,6 +247,7 @@ class UInstancedStaticMeshComponent : public UStaticMeshComponent, public ISMIns
 	ENGINE_API bool GetInstancePrevTransform(int32 InstanceIndex, FTransform& OutInstanceTransform, bool bWorldSpace = false) const;
 
 	ENGINE_API virtual void OnUpdateTransform(EUpdateTransformFlags UpdateTransformFlags, ETeleportType Teleport) override;
+	ENGINE_API void UpdateComponentTransform(EUpdateTransformFlags UpdateTransformFlags, ETeleportType Teleport);
 
 	/** Get the scale comming form the component, when computing StreamingTexture data. Used to support instanced meshes. */
 	ENGINE_API virtual float GetTextureStreamingTransformScale() const override;
diff --git a/Engine/Source/Runtime/Engine/Private/InstancedStaticMesh.cpp b/Engine/Source/Runtime/Engine/Private/InstancedStaticMesh.cpp
--- a/Engine/Source/Runtime/Engine/Private/InstancedStaticMesh.cpp
+++ b/Engine/Source/Runtime/Engine/Private/InstancedStaticMesh.cpp
@@ -3628,9 +3628,13 @@ bool UInstancedStaticMeshComponent::GetInstancePrevTransform(int32 InstanceIndex
 
 void UInstancedStaticMeshComponent::OnUpdateTransform(EUpdateTransformFlags UpdateTransformFlags, ETeleportType Teleport)
 {
-	// We are handling the physics move below, so don't handle it at higher levels
+	// We are handling the physics move in UpdateComponentTransform below, so don't handle it at higher levels
 	Super::OnUpdateTransform(UpdateTransformFlags | EUpdateTransformFlags::SkipPhysicsUpdate, Teleport);
+	UpdateComponentTransform(UpdateTransformFlags, Teleport);
+}
 
+void UInstancedStaticMeshComponent::UpdateComponentTransform(EUpdateTransformFlags UpdateTransformFlags, ETeleportType Teleport)
+{
 	const bool bTeleport = TeleportEnumToFlag(Teleport);
 	const bool bDoPartialNavigationUpdate = IsNavigationRelevant() && SupportsPartialNavigationUpdate();
 	const bool bUpdateBodies = bPhysicsStateCreated && !(EUpdateTransformFlags::SkipPhysicsUpdate & UpdateTransformFlags);
diff --git a/Engine/Source/Runtime/Experimental/GeometryCollectionEngine/Private/GeometryCollection/GeometryCollectionISMPoolComponent.cpp b/Engine/Source/Runtime/Experimental/GeometryCollectionEngine/Private/GeometryCollection/GeometryCollectionISMPoolComponent.cpp
--- a/Engine/Source/Runtime/Experimental/GeometryCollectionEngine/Private/GeometryCollection/GeometryCollectionISMPoolComponent.cpp
+++ b/Engine/Source/Runtime/Experimental/GeometryCollectionEngine/Private/GeometryCollection/GeometryCollectionISMPoolComponent.cpp
@@ -81,12 +81,12 @@ void FGeometryCollectionISM::CreateISM(AActor* InOwningActor)
 	ISMComponent->SetCanEverAffectNavigation(false);
 	ISMComponent->SetCollisionEnabled(ECollisionEnabled::NoCollision);
 	ISMComponent->SetupAttachment(InOwningActor->GetRootComponent());
-
+	
 	InOwningActor->AddInstanceComponent(ISMComponent);
 	ISMComponent->RegisterComponent();
 }
 
-void FGeometryCollectionISM::InitISM(const FGeometryCollectionStaticMeshInstance& InMeshInstance, bool bKeepAlive)
+void FGeometryCollectionISM::InitISM(const FGeometryCollectionStaticMeshInstance& InMeshInstance, bool bKeepAlive, bool bOverrideTransformUpdates)
 {
 	MeshInstance = InMeshInstance;
 	check(MeshInstance.StaticMesh);
@@ -98,6 +98,9 @@ void FGeometryCollectionISM::InitISM(const FGeometryCollectionStaticMeshInstance
 	ISMComponent->Rename(*ISMNameString);
 #endif
 
+	ISMComponent->bUseAttachParentBound = bOverrideTransformUpdates;
+	ISMComponent->SetAbsolute(bOverrideTransformUpdates, bOverrideTransformUpdates, bOverrideTransformUpdates);
+
 	ISMComponent->EmptyOverrideMaterials();
 	for (int32 MaterialIndex = 0; MaterialIndex < MeshInstance.MaterialsOverrides.Num(); MaterialIndex++)
 	{
@@ -117,15 +120,35 @@ void FGeometryCollectionISM::InitISM(const FGeometryCollectionStaticMeshInstance
 	// Instead of reverse culling we put the mirror in the component transform so that PRIMITIVE_SCENE_DATA_FLAG_DETERMINANT_SIGN will be set for use by materials.
 	//ISMComponent->SetReverseCulling(bReverseCulling);
 	const FVector Scale = bReverseCulling ? FVector(-1, 1, 1) : FVector(1, 1, 1);
-	const FTransform NewRelativeTransform(FQuat::Identity, MeshInstance.Desc.Position, Scale);
-	if (!ISMComponent->GetRelativeTransform().Equals(NewRelativeTransform))
+
+	if(bOverrideTransformUpdates)
+	{
+		FTransform TempTm = ISMComponent->GetAttachParent() ?
+			ISMComponent->GetAttachParent()->GetComponentToWorld() :
+			FTransform::Identity;
+
+		// Apply above identified scale to the transform directly
+		TempTm.SetScale3D(TempTm.GetScale3D() * Scale);
+
+		ISMComponent->SetComponentToWorld(TempTm);
+		ISMComponent->UpdateComponentTransform(EUpdateTransformFlags::None, ETeleportType::None);
+	}
+	else
 	{
-		ISMComponent->SetRelativeTransform(FTransform(FQuat::Identity, MeshInstance.Desc.Position, Scale));
+		const FTransform NewRelativeTransform(FQuat::Identity, MeshInstance.Desc.Position, Scale);
+
+		if(!ISMComponent->GetRelativeTransform().Equals(NewRelativeTransform))
+		{
+			// If we're not overriding the transform and need a relative offset, apply that here
+			ISMComponent->SetRelativeTransform(FTransform(FQuat::Identity, MeshInstance.Desc.Position, Scale));
+		}
 	}
+
 	if ((MeshInstance.Desc.Flags & FISMComponentDescription::DistanceCullPrimitive) != 0)
 	{
 		ISMComponent->SetCachedMaxDrawDistance(MeshInstance.Desc.EndCullDistance);
 	}
+
 	ISMComponent->SetCullDistances(MeshInstance.Desc.StartCullDistance, MeshInstance.Desc.EndCullDistance);
 	ISMComponent->SetCastShadow((MeshInstance.Desc.Flags & FISMComponentDescription::AffectShadow) != 0);
 	ISMComponent->bAffectDynamicIndirectLighting = (MeshInstance.Desc.Flags & FISMComponentDescription::AffectDynamicIndirectLighting) != 0;
@@ -214,7 +237,7 @@ FGeometryCollectionISMPool::FISMIndex FGeometryCollectionISMPool::GetOrAddISM(UG
 		ISMs[ISMIndex].CreateISM(OwningComponent->GetOwner());
 	}
 	
-	ISMs[ISMIndex].InitISM(MeshInstance, bCachedKeepAlive);
+	ISMs[ISMIndex].InitISM(MeshInstance, bCachedKeepAlive, bDisableBoundsAndTransformUpdate);
 	
 	bOutISMCreated = true;
 	MeshToISMIndex.Add(MeshInstance, ISMIndex);
@@ -433,6 +456,33 @@ void FGeometryCollectionISMPool::ProcessPreallocationRequests(UGeometryCollectio
 	}
 }
 
+void FGeometryCollectionISMPool::UpdateAbsoluteTransforms(const FTransform& BaseTransform, EUpdateTransformFlags UpdateTransformFlags, ETeleportType Teleport)
+{
+	for(const FGeometryCollectionISM& GcIsm : ISMs)
+	{
+		const bool bReverseCulling = (GcIsm.MeshInstance.Desc.Flags & FISMComponentDescription::ReverseCulling) != 0;
+		check(GcIsm.MeshInstance.Desc.Position == FVector::ZeroVector);
+		
+		if(bReverseCulling)
+		{
+			// As in InitISM we need to apply the inverted X scale for reverse culling.
+			// Just copy the transform and set an inverted scale to apply to the ISM
+			FVector BaseScale = BaseTransform.GetScale3D();
+			BaseScale.X = -BaseScale.X;
+			FTransform Flipped = BaseTransform;
+			Flipped.SetScale3D(BaseScale);
+			
+			GcIsm.ISMComponent->SetComponentToWorld(Flipped);
+		}
+		else
+		{
+			GcIsm.ISMComponent->SetComponentToWorld(BaseTransform);
+		}
+
+		GcIsm.ISMComponent->UpdateComponentTransform(UpdateTransformFlags | EUpdateTransformFlags::SkipPhysicsUpdate, Teleport);
+	}
+}
+
 void FGeometryCollectionISMPool::Tick(UGeometryCollectionISMPoolComponent* OwningComponent)
 {
 	// Recache component lifecycle state from cvar.
@@ -556,6 +606,16 @@ void UGeometryCollectionISMPoolComponent::PreallocateMeshInstance(const FGeometr
 	Pool.RequestPreallocateMeshInstance(MeshInstance);
 }
 
+void UGeometryCollectionISMPoolComponent::SetOverrideTransformUpdates(bool bOverrideUpdates)
+{
+	Pool.bDisableBoundsAndTransformUpdate = bOverrideUpdates;
+}
+
+void UGeometryCollectionISMPoolComponent::UpdateAbsoluteTransforms(const FTransform& BaseTransform, EUpdateTransformFlags UpdateTransformFlags, ETeleportType Teleport)
+{
+	Pool.UpdateAbsoluteTransforms(BaseTransform, UpdateTransformFlags, Teleport);
+}
+
 void UGeometryCollectionISMPoolComponent::GetResourceSizeEx(FResourceSizeEx& CumulativeResourceSize)
 {
 	Super::GetResourceSizeEx(CumulativeResourceSize);
diff --git a/Engine/Source/Runtime/Experimental/GeometryCollectionEngine/Public/GeometryCollection/GeometryCollectionISMPoolComponent.h b/Engine/Source/Runtime/Experimental/GeometryCollectionEngine/Public/GeometryCollection/GeometryCollectionISMPoolComponent.h
--- a/Engine/Source/Runtime/Experimental/GeometryCollectionEngine/Public/GeometryCollection/GeometryCollectionISMPoolComponent.h
+++ b/Engine/Source/Runtime/Experimental/GeometryCollectionEngine/Public/GeometryCollection/GeometryCollectionISMPoolComponent.h
@@ -252,7 +252,7 @@ struct FGeometryCollectionISM
 	/** Create the ISMComponent according to settings on the mesh instance. */
 	void CreateISM(AActor* InOwningActor);
 	/** Initialize the ISMComponent according to settings on the mesh instance. */
-	void InitISM(const FGeometryCollectionStaticMeshInstance& InMeshInstance, bool bKeepAlive);
+	void InitISM(const FGeometryCollectionStaticMeshInstance& InMeshInstance, bool bKeepAlive, bool bOverrideTransformUpdates = false);
 	/** Add a group to the ISM. Returns the group index. */
 	FInstanceGroups::FInstanceGroupId AddInstanceGroup(int32 InstanceCount, TArrayView<const float> CustomDataFloats);
 
@@ -296,6 +296,8 @@ struct FGeometryCollectionISMPool
 	/** Process the preallocation queue. Processing is timesliced so that only some of the queue will be processed in every call. */
 	void ProcessPreallocationRequests(UGeometryCollectionISMPoolComponent* OwningComponent, int32 MaxPreallocations);
 
+	void UpdateAbsoluteTransforms(const FTransform& BaseTransform, EUpdateTransformFlags UpdateTransformFlags, ETeleportType Teleport);
+
 	/** Array of ISM objects. */
 	TArray<FGeometryCollectionISM> ISMs;
 	/** Mapping from mesh description to ISMs array slot. */
@@ -312,6 +314,9 @@ struct FGeometryCollectionISMPool
 	// Cached state of lifecycle cvars from the last Tick()
 	bool bCachedKeepAlive = false;
 	bool bCachedRecycle = false;
+
+	// Whether we force ISMs to use parent bounds and disable transform updates
+	bool bDisableBoundsAndTransformUpdate = false;
 };
 
 
@@ -361,6 +366,10 @@ public:
 	 */
 	GEOMETRYCOLLECTIONENGINE_API void PreallocateMeshInstance(const FGeometryCollectionStaticMeshInstance& MeshInstance);
 
+	GEOMETRYCOLLECTIONENGINE_API void SetOverrideTransformUpdates(bool bOverrideUpdates);
+
+	GEOMETRYCOLLECTIONENGINE_API void UpdateAbsoluteTransforms(const FTransform& BaseTransform, EUpdateTransformFlags UpdateTransformFlags, ETeleportType Teleport);
+
 private:
 	uint32 NextMeshGroupId = 0;
 	TMap<FMeshGroupId, FGeometryCollectionMeshGroup> MeshGroups;
