Deprecate PhysXIncludes and NvClothIncludes.h

diff --git a/Engine/Source/Editor/UnrealEd/Private/ApexClothingUtils.cpp b/Engine/Source/Editor/UnrealEd/Private/ApexClothingUtils.cpp
--- a/Engine/Source/Editor/UnrealEd/Private/ApexClothingUtils.cpp
+++ b/Engine/Source/Editor/UnrealEd/Private/ApexClothingUtils.cpp
@@ -17,7 +17,6 @@
 
 #include "ClothingAssetFactory.h"
 #include "PhysicsPublic.h"
-#include "PhysXIncludes.h"
 
 DEFINE_LOG_CATEGORY_STATIC(LogApexClothingUtils, Log, All);
 
diff --git a/Engine/Source/Runtime/ClothingSystemRuntimeNv/Private/ClothingSystemRuntimeModuleNv.cpp b/Engine/Source/Runtime/ClothingSystemRuntimeNv/Private/ClothingSystemRuntimeModuleNv.cpp
--- a/Engine/Source/Runtime/ClothingSystemRuntimeNv/Private/ClothingSystemRuntimeModuleNv.cpp
+++ b/Engine/Source/Runtime/ClothingSystemRuntimeNv/Private/ClothingSystemRuntimeModuleNv.cpp
@@ -7,7 +7,6 @@
 #include "HAL/PlatformProcess.h"
 #include "Modules/ModuleManager.h"
 
-#include "NvClothIncludes.h"
 #include "NvClothSupport.h"
 
 
diff --git a/Engine/Source/Runtime/ClothingSystemRuntimeNv/Public/NvClothIncludes.h b/Engine/Source/Runtime/ClothingSystemRuntimeNv/Public/NvClothIncludes.h
--- a/Engine/Source/Runtime/ClothingSystemRuntimeNv/Public/NvClothIncludes.h
+++ b/Engine/Source/Runtime/ClothingSystemRuntimeNv/Public/NvClothIncludes.h
@@ -2,5 +2,9 @@
 
 #pragma once
 
+// HEADER_UNIT_SKIP - Deprecated
+
+#pragma message (__FILE__"(5): warning: NvClothIncludes.h is no longer useful")
+
 #include "CoreMinimal.h"
 #include "EngineDefines.h"
diff --git a/Engine/Source/Runtime/Engine/Private/Collision/CollisionConversions.h b/Engine/Source/Runtime/Engine/Private/Collision/CollisionConversions.h
--- a/Engine/Source/Runtime/Engine/Private/Collision/CollisionConversions.h
+++ b/Engine/Source/Runtime/Engine/Private/Collision/CollisionConversions.h
@@ -8,7 +8,6 @@
 #include "EngineDefines.h"
 #include "PhysicsEngine/BodyInstance.h"
 #include "PhysicsPublic.h"
-#include "PhysXIncludes.h"
 
 enum class EConvertQueryResult
 {
diff --git a/Engine/Source/Runtime/Engine/Private/Collision/CollisionDebugDrawing.h b/Engine/Source/Runtime/Engine/Private/Collision/CollisionDebugDrawing.h
--- a/Engine/Source/Runtime/Engine/Private/Collision/CollisionDebugDrawing.h
+++ b/Engine/Source/Runtime/Engine/Private/Collision/CollisionDebugDrawing.h
@@ -7,7 +7,6 @@
 #include "CoreMinimal.h"
 #include "EngineDefines.h"
 #include "PhysicsPublic.h"
-#include "PhysXIncludes.h"
 
 struct FHitResult;
 
diff --git a/Engine/Source/Runtime/Engine/Private/EnginePrivatePCH.h b/Engine/Source/Runtime/Engine/Private/EnginePrivatePCH.h
--- a/Engine/Source/Runtime/Engine/Private/EnginePrivatePCH.h
+++ b/Engine/Source/Runtime/Engine/Private/EnginePrivatePCH.h
@@ -518,7 +518,6 @@
 #include "CanvasTypes.h"
 #include "Components/StaticMeshComponent.h"
 #include "ContentStreaming.h"
-#include "PhysXIncludes.h"
 #include "GameFramework/MovementComponent.h"
 #include "CanvasItem.h"
 #include "Model.h"
diff --git a/Engine/Source/Runtime/Engine/Private/PhysicsEngine/PhysDerivedData.h b/Engine/Source/Runtime/Engine/Private/PhysicsEngine/PhysDerivedData.h
--- a/Engine/Source/Runtime/Engine/Private/PhysicsEngine/PhysDerivedData.h
+++ b/Engine/Source/Runtime/Engine/Private/PhysicsEngine/PhysDerivedData.h
@@ -5,6 +5,5 @@
 #include "CoreMinimal.h"
 #include "Misc/Guid.h"
 #include "EngineDefines.h"
-#include "PhysXIncludes.h"
 #include "Stats/Stats.h"
 #include "PhysDerivedDataPublic.h"
diff --git a/Engine/Source/Runtime/Engine/Public/Physics/PhysicsFiltering.h b/Engine/Source/Runtime/Engine/Public/Physics/PhysicsFiltering.h
--- a/Engine/Source/Runtime/Engine/Public/Physics/PhysicsFiltering.h
+++ b/Engine/Source/Runtime/Engine/Public/Physics/PhysicsFiltering.h
@@ -5,7 +5,6 @@
 #include "CoreMinimal.h"
 #include "EngineDefines.h"
 #include "PhysicsPublic.h"
-#include "PhysXIncludes.h"
 #include "Physics/PhysicsInterfaceTypes.h"
 #include "Chaos/ChaosEngineInterface.h"
 
diff --git a/Engine/Source/Runtime/Experimental/GeometryCollectionEngine/Private/GeometryCollection/PhysicsAssetSimulation.cpp b/Engine/Source/Runtime/Experimental/GeometryCollectionEngine/Private/GeometryCollection/PhysicsAssetSimulation.cpp
--- a/Engine/Source/Runtime/Experimental/GeometryCollectionEngine/Private/GeometryCollection/PhysicsAssetSimulation.cpp
+++ b/Engine/Source/Runtime/Experimental/GeometryCollectionEngine/Private/GeometryCollection/PhysicsAssetSimulation.cpp
@@ -19,8 +19,6 @@
 #include "Rendering/SkeletalMeshRenderData.h"
 #include "PhysicsProxy/SkeletalMeshPhysicsProxy.h"
 
-#include "PhysXIncludes.h"
-
 #include "DrawDebugHelpers.h"
 #include "Async/ParallelFor.h"
 #include "Math/Box.h"
diff --git a/Engine/Source/Runtime/PhysicsCore/Public/PhysXIncludes.h b/Engine/Source/Runtime/PhysicsCore/Public/PhysXIncludes.h
--- a/Engine/Source/Runtime/PhysicsCore/Public/PhysXIncludes.h
+++ b/Engine/Source/Runtime/PhysicsCore/Public/PhysXIncludes.h
@@ -1,6 +1,9 @@
 // Copyright Epic Games, Inc. All Rights Reserved.
-// Includes all necessary PhysX headers
 
 #pragma once
 
-#include "CoreMinimal.h"
+// HEADER_UNIT_SKIP - Deprecated
+
+#pragma message (__FILE__"(5): warning: PhysXIncludes.h is no longer useful")
+
+#include "CoreMinimal.h"
\ No newline at end of file
