Support writing partial buffers to the cache

Current caching mechanism only supports peeling off complete buffers from the memory cache to the cache file. This change adds support for writing partial buffers, with the remainder written in consecutive passes. Once the entire buffer is written to the file cache the journal entry will be added. This allows for greater utilization of the alloted write budget on certain platforms.

diff --git a/Engine/Source/Runtime/Experimental/IoStoreOnDemand/Private/JournaledCache.cpp b/Engine/Source/Runtime/Experimental/IoStoreOnDemand/Private/JournaledCache.cpp
--- a/Engine/Source/Runtime/Experimental/IoStoreOnDemand/Private/JournaledCache.cpp
+++ b/Engine/Source/Runtime/Experimental/IoStoreOnDemand/Private/JournaledCache.cpp
@@ -117,7 +117,7 @@ public:
 	uint32			GetMax() const		{ return MaxSize; }
 	const FIoBuffer*Get(uint64 Key) const;
 	bool			Put(uint64 Key, FIoBuffer&& Data);
-	int32			Peel(int32 PeelSize, PeelItems& Out);
+	int32			Peel(int32 TargetPeelSize, PeelItems& Out, int32 MaxPeelSize);
 	uint32			DebugVisit(void* Param, FDebugCacheEntry::Callback* Callback);
 
 private:
@@ -186,7 +186,7 @@ bool FMemCache::Put(uint64 Key, FIoBuffer&& Data)
 }
 
 ////////////////////////////////////////////////////////////////////////////////
-int32 FMemCache::Peel(int32 PeelSize, PeelItems& Out)
+int32 FMemCache::Peel(int32 TargetPeelSize, PeelItems& Out, int32 MaxPeelSize)
 {
 	TRACE_CPUPROFILER_EVENT_SCOPE(IasCache::Peel);
 
@@ -197,30 +197,15 @@ int32 FMemCache::Peel(int32 PeelSize, PeelItems& Out)
 	// Add large items
 	int32 NumItems = Items.Num();
 	int32 DropSize = 0;
-	while (NumItems > 0)
-	{
-		FItem& Item = Items[NumItems - 1];
-
-		int32 NextDropSize = DropSize + int32(Item.Data.GetSize());
-		if (NextDropSize > PeelSize)
-		{
-			break;
-		}
-
-		Out.Add(MoveTemp(Item));
-		DropSize = NextDropSize;
-		--NumItems;
-	}
-
-	// Fill remaining space with small items
-	for (int32 i = 0; i < NumItems; ++i)
+	for (int32 i=Items.Num() - 1; i >= 0 && DropSize < TargetPeelSize; --i)
 	{
 		FItem& Item = Items[i];
 
 		int32 NextDropSize = DropSize + int32(Item.Data.GetSize());
-		if (NextDropSize > PeelSize)
+		
+		if (NextDropSize > MaxPeelSize)
 		{
-			break;
+			continue;
 		}
 
 		Out.Add(MoveTemp(Item));
@@ -294,6 +279,7 @@ int32 FMemCache::Drop(uint32 Size)
 static const uint32 MAGIC = 0x04930002;
 static const uint32 SIZE_BITS = 25;
 static const uint32 MARKER_MAX = 0x3fffffff;
+static const uint32	HASH_CHECKSUM_SIZE = 64;
 
 using EntryHash = uint32;
 
@@ -320,22 +306,26 @@ static_assert(sizeof(FPhraseDesc) == sizeof(FDataEntry));
 ////////////////////////////////////////////////////////////////////////////////
 struct FDiskPhrase
 {
-						FDiskPhrase(TArray<FDataEntry>& InEntries, int32 InMaxEntries, uint32 DataSize);
+						FDiskPhrase(TArray<FDataEntry>& InEntries, int32 InMaxEntries, uint32 DataSize, uint32 PreviousPartial);
 						FDiskPhrase(FDiskPhrase&&) = default;
 	bool				Add(uint64 Key, FIoBuffer&& Data);
 	void				Drop()					{ return Entries.SetNumUninitialized(Index); }
 	const FDataEntry*	GetEntries() const		{ return Entries.GetData() + Index; }
 	int32				GetEntryCount() const	{ return Entries.Num() - Index; }
 	uint8*				GetPhraseData() const	{ return Buffer.Get(); }
-	uint32				GetDataSize() const		{ return Cursor; }
+	uint32				GetDataSize() const		{ return EntriesSize; }
 	int32				GetRemainingEntries() const { return MaxEntries; }
+	uint32				GetWriteSize() const	{ return Cursor; }
 
 private:
 	TUniquePtr<uint8[]>	Buffer;
 	TArray<FDataEntry>&	Entries;
 	uint32				Cursor = 0;
+	uint32				CurrentOffset = 0;
+	uint32				EntriesSize = 0;
 	uint32				Index;
 	int32				MaxEntries;
+	uint32				PreviousPartial;
 
 private:
 						FDiskPhrase(const FDiskPhrase&) = delete;
@@ -344,10 +334,11 @@ private:
 };
 
 ////////////////////////////////////////////////////////////////////////////////
-FDiskPhrase::FDiskPhrase(TArray<FDataEntry>& InEntries, int32 InMaxEntries, uint32 DataSize)
+FDiskPhrase::FDiskPhrase(TArray<FDataEntry>& InEntries, int32 InMaxEntries, uint32 DataSize, uint32 InPreviousPartial)
 : Entries(InEntries)
 , Index(Entries.Num())
 , MaxEntries(InMaxEntries)
+, PreviousPartial(InPreviousPartial)
 {
 	Buffer = TUniquePtr<uint8[]>(new uint8[DataSize]);
 }
@@ -356,12 +347,27 @@ FDiskPhrase::FDiskPhrase(TArray<FDataEntry>& InEntries, int32 InMaxEntries, uint
 bool FDiskPhrase::Add(uint64 Key, FIoBuffer&& Data)
 {
 	check(MaxEntries > 0);
-	uint32 DataSize = uint32(Data.GetSize());
-	check(DataSize < (1 << SIZE_BITS));
-	Entries.Add({Key, Cursor, DataSize});
+	const uint32 DataSize = uint32(Data.GetSize());
+	// Actual entries (Key!=0) 
+	// Padding (Key=0, Size=0) can appear at the end of journal file
+	// Partials (Key=0) write data but does not add entry
+	if (Key || Data.GetSize() == 0)
+	{
+		const uint32 FullSize = DataSize + PreviousPartial;
+		PreviousPartial = 0;
+		check(FullSize < (1 << SIZE_BITS));
+		Entries.Add({Key, CurrentOffset, FullSize});
+		EntriesSize += FullSize;
+		CurrentOffset += FullSize;
+		--MaxEntries;
+	}
+	else
+	{
+		PreviousPartial += DataSize;
+	}
 	std::memcpy(Buffer.Get() + Cursor, Data.GetData(), DataSize);
 	Cursor += DataSize;
-	return (--MaxEntries > 0);
+	return MaxEntries > 0;
 }
 
 
@@ -381,6 +387,7 @@ public:
 	uint32					GetMaxSize() const	{ return MaxSize; }
 	uint32					GetCursor() const	{ return Cursor; }
 	uint32					GetMarker() const	{ return Marker; }
+	uint32					GetPreviousPartial() const { return PreviousPartial; };
 
 private:
 	friend bool				LoadCache(FDiskCache&);
@@ -393,7 +400,8 @@ private:
 	TUniquePtr<IFileHandle>	JrnHandle;
 	uint32					Cursor = 0;
 	uint32					MaxSize;
-	static constexpr uint32	BytesToHash = 64;
+	uint32					PreviousPartial = 0;
+	uint32					PreviousHash = 0;
 };
 
 ////////////////////////////////////////////////////////////////////////////////
@@ -469,7 +477,7 @@ FDiskPhrase FDiskJournal::OpenPhrase(uint32 DataSize)
 	Entries.Add(FDataEntry{});
 	int32 MaxEntries = int32((MaxSize - Cursor) / sizeof(FDataEntry)) - Entries.Num();
 
-	FDiskPhrase Ret(Entries, FMath::Min(MaxEntries, int32(UINT16_MAX)), DataSize);
+	FDiskPhrase Ret(Entries, FMath::Min(MaxEntries, int32(UINT16_MAX)), DataSize, PreviousPartial);
 
 	return Ret;
 }
@@ -477,10 +485,27 @@ FDiskPhrase FDiskJournal::OpenPhrase(uint32 DataSize)
 ////////////////////////////////////////////////////////////////////////////////
 void FDiskJournal::ClosePhrase(FDiskPhrase&& Phrase, uint64 DataCursor)
 {
+	auto SavePreviousHash = [&]
+	{
+		// Save the hash for the next phrase, if we have just written the (first) partial
+		if (PreviousPartial && !PreviousHash)
+		{
+			const uint64 PartialOffset = Phrase.GetWriteSize() - PreviousPartial;
+			check(PreviousPartial >= HASH_CHECKSUM_SIZE);
+			PreviousHash = HashBytes(Phrase.GetPhraseData() + PartialOffset, HASH_CHECKSUM_SIZE, Marker);
+		}
+	};
+	
 	uint32 EntryCount = Phrase.GetEntryCount();
 	if (EntryCount == 0)
 	{
 		Entries.Pop();
+		// We could write partial data but not complete any entries
+		if (const uint32 WriteSize = Phrase.GetWriteSize(); WriteSize > 0)
+		{
+			PreviousPartial += WriteSize;
+			SavePreviousHash();
+		}
 		return;
 	}
 	
@@ -504,16 +529,35 @@ void FDiskJournal::ClosePhrase(FDiskPhrase&& Phrase, uint64 DataCursor)
 	auto& Desc = (FPhraseDesc&)(Phrase.GetEntries()[-1]);
 	Desc.Magic = MAGIC;
 	Desc.Marker = Marker;
-	Desc.DataCursor = DataCursor;
-	
-	const uint32 HashSize = FMath::Min(BytesToHash, Phrase.GetDataSize());
-	Desc.Hash = HashBytes(Phrase.GetPhraseData(), HashSize, Desc.Marker);
+	// Potentially adjust the data cursor to the start of the previous
+	// partial write.
+	Desc.DataCursor = DataCursor - PreviousPartial;
+
+	// If we have already written a partial fragment of this phrase
+	// we use the hash of the start
+	if (PreviousHash)
+	{
+		Desc.Hash = PreviousHash;
+		PreviousHash = 0;
+	}
+	else
+	{
+		const uint32 HashSize = FMath::Min(HASH_CHECKSUM_SIZE, Phrase.GetDataSize());
+		Desc.Hash = HashBytes(Phrase.GetPhraseData(), HashSize, Desc.Marker);
+	}
+
+	// If we wrote a partial entry at the end we need adjust the
+	// next phrase accordingly.
+	check((PreviousPartial + Phrase.GetWriteSize()) >= Phrase.GetDataSize());
+	PreviousPartial = (PreviousPartial + Phrase.GetWriteSize()) - Phrase.GetDataSize();
 
 	// Increment and wrap marker
 	if (++Marker > MARKER_MAX)
 	{
 		Marker = 0;
 	}
+
+	SavePreviousHash();
 }
 
 ////////////////////////////////////////////////////////////////////////////////
@@ -569,6 +613,8 @@ public:
 	EIoErrorCode			Materialize(uint64 Key, FIoBuffer& Out, uint32 Offset=0) const;
 	int32					Flush();
 	void					Drop();
+	void					Wrap();
+	uint64					RemainingUntilWrap() const;
 	uint32					DebugVisit(void* Param, FDebugCacheEntry::Callback* Callback);
 
 private:
@@ -638,11 +684,18 @@ FDiskPhrase FDiskCache::OpenPhrase(uint32 DataSize)
 	return Journal.OpenPhrase(DataSize);
 }
 
+////////////////////////////////////////////////////////////////////////////////
+void FDiskCache::Wrap()
+{
+	OverRemoval = 0;
+	DataCursor = 0;
+}
+
 ////////////////////////////////////////////////////////////////////////////////
 void FDiskCache::ClosePhrase(FDiskPhrase&& Phrase)
 {
 	int32 EntryCount = Phrase.GetEntryCount();
-	if (EntryCount <= 0)
+	if (Phrase.GetWriteSize() == 0)
 	{
 		Journal.ClosePhrase(MoveTemp(Phrase), 0);
 		return;
@@ -657,11 +710,10 @@ void FDiskCache::ClosePhrase(FDiskPhrase&& Phrase)
 		return;
 	}
 
-	uint32 WriteSize = Phrase.GetDataSize();
+	uint32 WriteSize = Phrase.GetWriteSize();
 	if (DataCursor + WriteSize > MaxDataSize)
 	{
-		OverRemoval = 0;
-		DataCursor = 0;
+		Wrap();
 	}
 
 	bool bWriteOk;
@@ -682,7 +734,7 @@ void FDiskCache::ClosePhrase(FDiskPhrase&& Phrase)
 	{
 		FWriteScopeLock _(Lock);
 		Prune(DataCursor, WriteSize);
-		Insert(DataCursor, Phrase.GetEntries(), EntryCount);
+		Insert(DataCursor - Journal.GetPreviousPartial(), Phrase.GetEntries(), EntryCount);
 	}
 
 	Journal.ClosePhrase(MoveTemp(Phrase), DataCursor);
@@ -740,6 +792,7 @@ uint64 FDiskCache::Insert(uint64 DataBase, const FDataEntry& Entry)
 {
 	FMapEntry Value;
 	Value.DataCursor = DataBase + Entry.Offset;
+	check(Value.DataCursor < MaxDataSize);
 	Value.Size = Entry.Size;
 	Value.First = (Entry.Offset == 0);
 	DataMap.Add(Entry.Key, Value);
@@ -832,6 +885,12 @@ void FDiskCache::Drop()
 	OpenDataFile();
 }
 
+////////////////////////////////////////////////////////////////////////////////
+uint64 FDiskCache::RemainingUntilWrap() const
+{
+	return MaxDataSize - DataCursor;
+}
+
 ////////////////////////////////////////////////////////////////////////////////
 void FDiskCache::Spam()
 {
@@ -1053,8 +1112,8 @@ static bool LoadCache(FDiskCache& DiskCache)
 
 		File->Seek(Cursor);
 		
-		const uint32 HashSize = FMath::Min(FDiskJournal::BytesToHash, MaxHashSize);
-		TArray<uint8, TInlineAllocator<FDiskJournal::BytesToHash>> BytesToHash;
+		const uint32 HashSize = FMath::Min(HASH_CHECKSUM_SIZE, MaxHashSize);
+		TArray<uint8, TInlineAllocator<HASH_CHECKSUM_SIZE>> BytesToHash;
 		BytesToHash.AddZeroed(HashSize);
 		if (const bool Result = File->Read(BytesToHash.GetData(), HashSize))
 		{
@@ -1075,14 +1134,16 @@ static bool LoadCache(FDiskCache& DiskCache)
 		}
 
 		uint32 Hash;
+		const uint32 Seed = uint32(Stock.Phrase->Marker);
 		// There could be a phrase with only one very short data entry. Make sure we don't hash too much.
-		if (ReadHash(DataBase, Hash, Stock.DataSize, uint32(Stock.Phrase->Marker)) && (Hash == Stock.Phrase->Hash))
+		if (ReadHash(DataBase, Hash, Stock.DataSize, Seed) && (Hash == Stock.Phrase->Hash))
 		{
 			break;
 		}
 	}
 
 	// Add known entries into the tree.
+	uint64 MappedBytes = 0; uint32 MappedItems = 0;
 	for (uint32 i = BasisIndex, n = Paragraphs.Num(); i < n; ++i)
 	{
 		const FPhraseDesc& Holm = Paragraphs[i].Phrase[0];
@@ -1093,9 +1154,12 @@ static bool LoadCache(FDiskCache& DiskCache)
 			check(Holm.Entries[EntryCount-1].Size == 0);
 			--EntryCount;
 		}
-		DiskCache.Insert(Holm.DataCursor, Holm.Entries, EntryCount);
+		MappedItems += EntryCount;
+		MappedBytes += DiskCache.Insert(Holm.DataCursor, Holm.Entries, EntryCount);
 	}
 	
+	UE_LOG(LogIas, VeryVerbose, TEXT("JournaledCache: Mapped %u items with %llu bytes"), MappedItems, MappedBytes);
+	
 	// Prime the journal's state
 	const FParagraph& LastPara = Paragraphs.Last();
 	const FPhraseDesc* LastPhrase = LastPara.Phrase;
@@ -1159,10 +1223,19 @@ public:
 	uint32			DebugVisit(void* Param, FDebugCacheEntry::Callback* Callback);
 
 private:
+	struct FPartialItem
+	{
+		uint64		Key;
+		FIoBuffer	Data;
+		uint32		Remaining;
+	};
+	using FPartial = TOptional<FPartialItem>;
+	
 	mutable FRWLock	MemLock;
 	FMemCache		MemCache;
 	FDiskCache		DiskCache;
 	std::atomic_int	Demand;
+	FPartial		Partial;
 };
 
 ////////////////////////////////////////////////////////////////////////////////
@@ -1271,21 +1344,75 @@ uint32 FCache::WriteMemToDisk(int32 Allowance)
 {
 	TRACE_CPUPROFILER_EVENT_SCOPE(IasCache::Flush_MemCache);
 
+	bool bCloseToWrap = false;
 	FMemCache::PeelItems PeelItems;
-	int32 MemCacheSize;
+	int32 WriteSize = 0;
 	{
 		FWriteScopeLock _(MemLock);
 
-		MemCacheSize = MemCache.Peel(Allowance, PeelItems);
+		// If we have any partials that was previously written process that first
+		// and peel off as much of that buffer as possible.
+		if (Partial)
+		{
+			const uint32 BufferOffset = uint32(Partial->Data.GetSize()) - Partial->Remaining;
+			const FMemoryView PartialSlice = Partial->Data.GetView().Mid(BufferOffset, Allowance);
+			const int32 PartialSize = int32(PartialSlice.GetSize());
+			Partial->Remaining -= PartialSize;
+			WriteSize += PartialSize;
+			const uint64 Key = Partial->Remaining ? 0 : Partial->Key;
+			PeelItems.Push(FMemCache::FItem {Key, FIoBuffer(PartialSlice, Partial->Data)});
+			if (Partial->Remaining == 0)
+			{
+				Partial.Reset();
+			}
+		}
+		
+		// We don't want to deal with wrapping partials around the end of the buffer
+		const int32 UntilWrap = int32(DiskCache.RemainingUntilWrap()) - WriteSize;
+		check(UntilWrap >= 0);
+		bCloseToWrap = UntilWrap <= Allowance;
+
+		// If there is any allowance left start peeling of buffers from the memcache
+		if (WriteSize < Allowance)
+		{
+			WriteSize += MemCache.Peel(Allowance - WriteSize, PeelItems, UntilWrap);
+		}
+
+		// Finally split any overshooting buffers into a partial slice and save the
+		// buffer in the Partial member. While being dropped this buffer exists neither
+		// in the memcache or the disk cache. The partial fragment needs to be at least
+		// large enough for the hash checksum.
+		if (WriteSize > Allowance)
+		{
+			auto [Key, Data] = PeelItems.Pop();
+			check(Data.GetSize() <= UntilWrap);
+			
+			int32 RemainderSize = WriteSize - Allowance;
+			const int32 PartialSize = int32(Data.GetSize()) - RemainderSize;
+			if (PartialSize < HASH_CHECKSUM_SIZE)
+			{
+				RemainderSize = int32(Data.GetSize());
+			}
+			else
+			{
+				const FMemoryView PartialSlice = Data.GetView().Left(PartialSize);
+				PeelItems.Push(FMemCache::FItem{0, FIoBuffer(PartialSlice, Data)});
+			}
+			Partial.Emplace(FPartialItem{Key, MoveTemp(Data), uint32(RemainderSize)});
+			WriteSize -= RemainderSize;
+		}
 
 		uint32 NewDemand = MemCache.GetDemand();
 		Demand.store(NewDemand, std::memory_order_relaxed);
 	}
-	FDiskPhrase Phrase = DiskCache.OpenPhrase(MemCacheSize);
+
+	check(WriteSize >= 0 && WriteSize <= Allowance);
+	FDiskPhrase Phrase = DiskCache.OpenPhrase(WriteSize);
 	int32 PeelIndex = -1;
 	for (int32 i = 0, n = PeelItems.Num(); i < n; ++i)
 	{
 		auto& [Key, Data] = PeelItems[i];
+		check(Key || i == (PeelItems.Num() - 1)); // Partials must be last
 		if (Phrase.GetRemainingEntries() < 1 || !Phrase.Add(Key, MoveTemp(Data)))
 		{
 			PeelIndex = i;
@@ -1297,12 +1424,22 @@ uint32 FCache::WriteMemToDisk(int32 Allowance)
 	{
 		/* end of journal reached so not all peeled items could be added, may
 		 * we can re-add leftover peeled items back to mem-cache? */
-		MemCacheSize = Phrase.GetDataSize();
+		WriteSize = Phrase.GetDataSize();
 	}
 
 	DiskCache.ClosePhrase(MoveTemp(Phrase));
 
-	return MemCacheSize;
+	// If we are close to the end of the file and cannot find any suitable
+	// buffers to peel off, reset file to beginning. There should be no partials
+	// at this point, but release it if any exists (it will be lost)
+	if (bCloseToWrap && PeelItems.IsEmpty())
+	{
+		check(!Partial.IsSet());
+		Partial.Reset();
+		DiskCache.Wrap();
+	}
+
+	return WriteSize;
 }
 
 ////////////////////////////////////////////////////////////////////////////////
@@ -2278,9 +2415,9 @@ static void MemCacheTests(FSupport& Support)
 
 	FMemCache MemCache(64);
 	MemCache.Put(1, Support.DummyData(1));
-	check(MemCache.Peel(0, Peeled) == 0);
+	check(MemCache.Peel(0, Peeled, 0) == 0);
 	check(Peeled.Num() == 0);
-	check(MemCache.Peel(64, Peeled) == 1);
+	check(MemCache.Peel(64, Peeled, 64) == 1);
 	check(Peeled.Num() == 1);
 	check(MemCache.GetUsed() == 0);
 	Peeled.Reset();
@@ -2291,7 +2428,7 @@ static void MemCacheTests(FSupport& Support)
 		MemCache.Put(i + 1, Support.DummyData(1));
 	}
 
-	check(MemCache.Peel(32, Peeled) == 32);
+	check(MemCache.Peel(32, Peeled, 32) == 32);
 	check(Peeled.Num() == 32);
 	check(MemCache.GetUsed() == 32);
 	for (auto& [Key, _] : Peeled)
@@ -2551,7 +2688,8 @@ static void CacheTests(FSupport& Support)
 			}
 		}
 
-		UE_LOG(LogIas, VeryVerbose, TEXT("Journal wrap test found %u correct entries. %u entries were lost."), Found, Lost);
+		UE_LOG(LogIas, Display, TEXT("Journal wrap test found %u correct entries. %u entries were lost."), Found, Lost);
+		check(Found > 40);
 	}
 	
 	{
@@ -2615,8 +2753,131 @@ static void CacheTests(FSupport& Support)
 			}
 		}
 
-		UE_LOG(LogIas, VeryVerbose, TEXT("Journal wrap test 2 found %u correct entries. %u entries were lost."), Found, Lost);
+		UE_LOG(LogIas, Display, TEXT("Journal wrap test 2 found %u correct entries. %u entries were lost."), Found, Lost);
+		check(Found > 40);
+	}
+
+	// random writes and allowances, wrap both cache and journal
+	{
+		constexpr uint64 AllowanceMax = 1_Mi;
+		int32 Allowances[] = {
+			int32(Support.Mix() & (AllowanceMax-1)),
+			int32(Support.Mix() & (AllowanceMax-1)),
+			int32(Support.Mix() & (AllowanceMax-1)),
+			int32(Support.Mix() & (AllowanceMax-1)),
+			0, 0,
+			int32(Support.Mix() & (AllowanceMax-1)),
+			int32(Support.Mix() & (AllowanceMax-1)),
+			int32(Support.Mix() & (AllowanceMax-1)),
+			int32(Support.Mix() & (AllowanceMax-1)),
+			int32(Support.Mix() & (AllowanceMax-1)),
+			0,
+			int32(Support.Mix() & (AllowanceMax-1)),
+			int32(Support.Mix() & (AllowanceMax-1)),
+		};
+		
+		auto PutAndCommit = [&] (int64 FillCnt, uint64 SizeMax, TArrayView<int32> Allowance) {
+			TArray<FStoredData> Ret;
+			for(uint32 i = 0; i < FillCnt; ++i)
+			{
+				const uint64 Size = Support.Mix() & (SizeMax - 1);
+				FIoBuffer Data = Support.DummyData(Size);
+				const uint64 Key = KeyGen(Data);
+				Ret.Push(FStoredData{Key, Size, CityHash64((const char*)Data.GetData(), uint32(Data.GetSize()))});
+				Cache->Put(Key, Data);
+				const uint32 AllowanceIdx = i % Allowance.Num();
+				if (Allowance[AllowanceIdx])
+				{
+					Cache->WriteMemToDisk(Allowance[AllowanceIdx]);
+					if (i % 3 == 0)
+					{
+						Cache->Flush();
+					}
+				}
+			}
+			Cache->Flush();
+			return Ret;
+		};
+
+		auto CheckCommitted = [&Cache](const TArray<FStoredData>& CommittedBuffers)
+		{
+			uint32 Found(0);
+
+			for (const FStoredData& Committed : CommittedBuffers)
+			{
+				FIoBuffer Data;
+				if (const auto Token = Cache->Get(Committed.Key, Data); Token == Committed.Key)
+				{
+					check(Cache->Materialize(Token, Data, 0) == EIoErrorCode::Ok);
+				}
+				// Some items have been lost by this point, as expected
+				if (Data.GetSize())
+				{
+					check(Data.DataSize() == Committed.Size);
+					const uint64 Hash = CityHash64((const char*)Data.GetData(), uint32(Data.GetSize()));
+					check(Committed.Hash == Hash);
+					++Found;
+				}
+			}
+
+			return Found;
+		};
+
+		auto CheckCached = [&Cache](const TArray<FStoredData>& CommittedBuffers)
+		{
+			uint32 Found(0);
+			TArray<uint64> KnownKeys;
+			Cache->DebugVisit(&KnownKeys, [](void* Param, const FDebugCacheEntry& Entry)
+			{
+				TArray<uint64>* KnownKeys = (TArray<uint64>*) Param;
+				KnownKeys->Add(Entry.Key);
+			});
+
+			for (uint64 Key : KnownKeys)
+			{
+				FIoBuffer Data;
+				if (const auto Token = Cache->Get(Key, Data); Token == Key)
+				{
+					check(Cache->Materialize(Token, Data, 0) == EIoErrorCode::Ok);
+				}
+				if (const FStoredData* Stored = CommittedBuffers.FindByPredicate([&](const FStoredData& Entry){ return Key == Entry.Key; }))
+				{
+					check(Data.DataSize() == Stored->Size);
+					const uint64 Hash = CityHash64((const char*)Data.GetData(), uint32(Data.GetSize()));
+					check(Stored->Hash == Hash);
+					++Found;
+				}
+			}
+			return Found;
+		};
+
+		FCache::FConfig Config;
+		Config.Path = Support.TestDir / "cache_random";
+		Config.MemoryQuota = uint32(4_Mi);
+		Config.DiskQuota = 16_Mi;
+		Config.JournalQuota = uint32(32_Ki);
+		
+		NewCache(true, &Config);
+		
+		uint32 Loops = 8;
+		while (--Loops)
+		{
+			auto CommittedBuffers = PutAndCommit(32, 1_Mi, Allowances);
+
+			uint32 FoundBefore = CheckCached(CommittedBuffers);
+
+			NewCache(false, &Config);
+			Cache->Load();
+
+			uint32 FoundAfter = CheckCommitted(CommittedBuffers);
+
+			UE_LOG(LogIas, Display, TEXT("Random cache test found %u correct entries before reload and %u after out of %d."), FoundBefore, FoundAfter, CommittedBuffers.Num());
+			check(FoundBefore > 10 && FoundAfter > 10);
+		}
 	}
+
+	// Make sure we delete the cache at the end to release the file handle
+	delete Cache;
 }
 
 ////////////////////////////////////////////////////////////////////////////////
