Fix destruction debris not always cleaning up correctly.

GeometryCollection root particles were still flagged as active on the game thread when the root was removed from a cluster untion on the same tick that it gets broken. This was preventing the cleanup logic from being executed.

diff --git a/Engine/Source/Runtime/Experimental/Chaos/Private/PhysicsProxy/GeometryCollectionPhysicsProxy.cpp b/Engine/Source/Runtime/Experimental/Chaos/Private/PhysicsProxy/GeometryCollectionPhysicsProxy.cpp
--- a/Engine/Source/Runtime/Experimental/Chaos/Private/PhysicsProxy/GeometryCollectionPhysicsProxy.cpp
+++ b/Engine/Source/Runtime/Experimental/Chaos/Private/PhysicsProxy/GeometryCollectionPhysicsProxy.cpp
@@ -172,6 +172,7 @@ static Chaos::FImplicitObjectPtr MakeTransformImplicitObject(const Chaos::FImpli
 FGeometryCollectionResults::FGeometryCollectionResults()
 	: IsObjectDynamic(false)
 	, IsObjectLoading(false)
+	, IsRootBroken(false)
 {}
 
 void FGeometryCollectionResults::Reset()
@@ -185,6 +186,7 @@ void FGeometryCollectionResults::Reset()
 
 	IsObjectDynamic = false;
 	IsObjectLoading = false;
+	IsRootBroken = false;
 }
 
 //==============================================================================
@@ -3634,6 +3636,21 @@ void FGeometryCollectionPhysicsProxy::BufferPhysicsResults_Internal(Chaos::FPBDR
 
 		const TManagedArray<FTransform>& MassToLocalArray = Parameters.RestCollection->GetAttribute<FTransform>(MassToLocalAttributeName, FTransformCollection::TransformGroup);
 
+		// Explicitly handle root breaking to ensure that the Root GT particle Active state matches the PT particle.
+		// If the root was in a cluster union and the cluster gets broken immediately on removal, we
+		// will never hit the logic below that checks for the particle changing from active to inactive.
+		// The parent checks also miss the case where the parent has an external parent (GC particles
+		// historically assume that the root has no parent which is not true any more).
+		if (Parameters.InitialRootIndex != INDEX_NONE)
+		{
+			const Chaos::FPBDRigidClusteredParticleHandle* RootHandle = SolverParticleHandles[Parameters.InitialRootIndex];
+			if (RootHandle != nullptr)
+			{
+				const bool bIsRootBroken = RootHandle->Disabled() && (RootHandle->Parent() == nullptr);
+				Results.IsRootBroken = bIsRootBroken;
+			}
+		}
+
 		for (int32 TransformGroupIndex = 0; TransformGroupIndex < NumTransforms; ++TransformGroupIndex)
 		{
 			Chaos::FPBDRigidClusteredParticleHandle* Handle = SolverParticleHandles[TransformGroupIndex];
@@ -4065,6 +4082,17 @@ bool FGeometryCollectionPhysicsProxy::PullNonInterpolatableDataFromSinglePhysics
 		}
 	}
 
+	// See comments in BufferPhysicsResults_Internal where IsRootBroken is set
+	const int32 RootIndex = Parameters.InitialRootIndex;
+	if (RootIndex != INDEX_NONE)
+	{
+		// NOTE: We never return to unbroken once broken
+		if (CurrentResults.IsRootBroken)
+		{
+			GameThreadCollection.Active[RootIndex] = false;
+		}
+	}
+
 	return bIsCollectionDirty;
 }
 
diff --git a/Engine/Source/Runtime/Experimental/Chaos/Public/GeometryCollectionProxyData.h b/Engine/Source/Runtime/Experimental/Chaos/Public/GeometryCollectionProxyData.h
--- a/Engine/Source/Runtime/Experimental/Chaos/Public/GeometryCollectionProxyData.h
+++ b/Engine/Source/Runtime/Experimental/Chaos/Public/GeometryCollectionProxyData.h
@@ -370,4 +370,5 @@ private:
 public:
 	uint8 IsObjectDynamic: 1;
 	uint8 IsObjectLoading: 1;
+	uint8 IsRootBroken : 1;
 };
