Memory optimization don't create GT particles at initialization, wait until destruction occurs

diff --git a/Engine/Source/Runtime/Experimental/Chaos/Private/PhysicsProxy/GeometryCollectionPhysicsProxy.cpp b/Engine/Source/Runtime/Experimental/Chaos/Private/PhysicsProxy/GeometryCollectionPhysicsProxy.cpp
--- a/Engine/Source/Runtime/Experimental/Chaos/Private/PhysicsProxy/GeometryCollectionPhysicsProxy.cpp
+++ b/Engine/Source/Runtime/Experimental/Chaos/Private/PhysicsProxy/GeometryCollectionPhysicsProxy.cpp
@@ -528,6 +528,9 @@ FAutoConsoleVariableRef CVarReportHighParticleFraction(TEXT("p.gc.ReportHighPart
 bool bBuildGeometryForChildren = true;
 FAutoConsoleVariableRef CVarbBuildGeometryForChildren(TEXT("p.gc.BuildGeometryForChildren"), bBuildGeometryForChildren, TEXT("If true build all children geometry at initilaization time, otherwise wait until destruction occurs."));
 
+bool bCreateGTParticleForChildren = true;
+FAutoConsoleVariableRef CVarCreateGTParticleForChildren(TEXT("p.gc.CreateGTParticlesForChildren"), bCreateGTParticleForChildren, TEXT("If true create all children particles at initilaization time, otherwise wait until destruction occurs."));
+
 bool bRemoveImplicitsInDynamicCollections = false;
 FAutoConsoleVariableRef CVarbRemoveImplicitsInDynamicCollections(TEXT("p.gc.RemoveImplicitsInDynamicCollections"), 
 	bBuildGeometryForChildren, TEXT("This cvar has an impact only if geometry are not added for children. It removes implicits from the Dynamic Collections, and recreate then from the rest collection. \
@@ -554,6 +557,7 @@ void FGeometryCollectionPhysicsProxy::Initialize(Chaos::FPBDRigidsEvolutionBase
 	SolverClusterHandles.Init(nullptr, NumParticles);
 	SolverParticleHandles.Init(nullptr, NumParticles);
 	GTParticles.SetNum(NumParticles);
+	UniqueIdxs.SetNum(NumParticles);
 	
 	TBitArray<> EffectiveParticles;
 	NumEffectiveParticles = CalculateEffectiveParticles(DynamicCollection, NumParticles, Parameters.MaxSimulatedLevel, Parameters.EnableClustering, GetOwner(), EffectiveParticles);
@@ -607,9 +611,9 @@ void FGeometryCollectionPhysicsProxy::Initialize(Chaos::FPBDRigidsEvolutionBase
 	TManagedArray<Chaos::FImplicitObjectPtr>& Implicits = GameThreadCollection.ModifyAttribute<Chaos::FImplicitObjectPtr>(FGeometryDynamicCollection::ImplicitsAttribute, FTransformCollection::TransformGroup);
 	if (ensure(NumTransforms == Implicits.Num() && NumTransforms == GTParticles.Num())) // Implicits are in the transform group so this invariant should always hold
 	{
-		constexpr bool bInitializeRootOnly = true;
+		constexpr bool bInitializationTime = true;
 		bHasBuiltGeometryOnGT = bBuildGeometryForChildren;
-		CreateParticles(EffectiveParticles, Implicits, Evolution, bInitializeRootOnly);
+		CreateGTParticles(EffectiveParticles, Implicits, Evolution, bInitializationTime);
 
 		// Skip simplicials, as they're owned by unique pointers.
 		static const FAttributeAndGroupId SkipList[] =
@@ -673,9 +677,14 @@ void FGeometryCollectionPhysicsProxy::Initialize(Chaos::FPBDRigidsEvolutionBase
 		// The Implicits attributes from the Dynamic Collection are just used for initialization, after they can be removed and so free some memory.
 		GameThreadCollection.RemoveAttribute(FGeometryDynamicCollection::ImplicitsAttribute, FTransformCollection::TransformGroup);
 	}
+	// If GT particle have been created, we can use directly the UniqueIdx from the GTParticle, then free up some space.
+	if (bBuildGeometryForChildren || bCreateGTParticleForChildren)
+	{
+		UniqueIdxs.Empty();
+	}
 }
 
-void FGeometryCollectionPhysicsProxy::CreateParticles(const TBitArray<>& EffectiveParticles, TManagedArray<Chaos::FImplicitObjectPtr>& Implicits, Chaos::FPBDRigidsEvolutionBase* Evolution, bool bInitializeRootOnly)
+void FGeometryCollectionPhysicsProxy::CreateGTParticles(const TBitArray<>& EffectiveParticles, TManagedArray<Chaos::FImplicitObjectPtr>& Implicits, Chaos::FPBDRigidsEvolutionBase* Evolution, bool bInitializationTime)
 {
 	const Chaos::Facades::FCollectionAnchoringFacade DynamicCollectionAnchoringFacade(GameThreadCollection);
 	const FVector Scale = Parameters.WorldTransform.GetScale3D();
@@ -691,13 +700,20 @@ void FGeometryCollectionPhysicsProxy::CreateParticles(const TBitArray<>& Effecti
 		if (EffectiveParticles[Index])
 		{
 			FParticle* P = GTParticles[Index].Get();
-			if (bInitializeRootOnly)
+			// Generate all particles unique idx at initialization time
+			if (bInitializationTime)
+			{
+				UniqueIdxs[Index] = Evolution->GenerateUniqueIdx();
+			}
+
+			if (bInitializationTime && Index == Parameters.InitialRootIndex || // When initializing always create particle for the root
+				(bInitializationTime && (bCreateGTParticleForChildren || bBuildGeometryForChildren) || // When initializing create all particles if one of the flag is true 
+					(!bInitializationTime && Index != Parameters.InitialRootIndex && !bCreateGTParticleForChildren && !bBuildGeometryForChildren))) // When not initializing create other particles if flag was set to not create
 			{
 				GTParticles[Index] = FParticle::CreateParticle();
 				P = GTParticles[Index].Get();
 				GTParticlesToTransformGroupIndex.Add(P, Index);
-
-				GTParticles[Index]->SetUniqueIdx(Evolution->GenerateUniqueIdx());
+				GTParticles[Index]->SetUniqueIdx(UniqueIdxs[Index]);
 
 				const float ScaledMass = AdjustMassForScale(Mass[Index]);
 
@@ -713,8 +729,7 @@ void FGeometryCollectionPhysicsProxy::CreateParticles(const TBitArray<>& Effecti
 				P->SetUserData(Parameters.UserData);
 				P->SetProxy(this);
 			}
-			check(P != nullptr);
-			if (bInitializeRootOnly && Index == Parameters.InitialRootIndex && bUseStaticMeshCollisionForTraces && CreateTraceCollisionGeometryCallback != nullptr)
+			if (bInitializationTime && Index == Parameters.InitialRootIndex && bUseStaticMeshCollisionForTraces && CreateTraceCollisionGeometryCallback != nullptr)
 			{
 				const FTransform ToLocal = MassToLocal[Index].Inverse();
 				TArray<Chaos::FImplicitObjectPtr> Geoms;
@@ -724,7 +739,7 @@ void FGeometryCollectionPhysicsProxy::CreateParticles(const TBitArray<>& Effecti
 				Chaos::FImplicitObjectPtr ImplicitGeometry = MakeImplicitObjectPtr<Chaos::FImplicitObjectUnion>(MoveTemp(Geoms));
 				P->SetGeometry(ImplicitGeometry);
 			}
-			else if (bInitializeRootOnly == (Index == Parameters.InitialRootIndex) || bBuildGeometryForChildren)
+			else if (bInitializationTime == (Index == Parameters.InitialRootIndex) || bBuildGeometryForChildren)
 			{
 				Chaos::FImplicitObjectPtr ImplicitGeometry = Implicits[Index];
 				if (ImplicitGeometry && !Scale.Equals(FVector::OneVector))
@@ -734,12 +749,12 @@ void FGeometryCollectionPhysicsProxy::CreateParticles(const TBitArray<>& Effecti
 				P->SetGeometry(ImplicitGeometry);
 			}
 			// Reset the root in the ManagedArray if not in mode bInitializeRootOnly
-			if (!bInitializeRootOnly && Index == Parameters.InitialRootIndex)
+			if (!bInitializationTime && Index == Parameters.InitialRootIndex)
 			{
 				Implicits[Index] = P->GetGeometry();
 			}
 
-			if (bInitializeRootOnly)
+			if (bInitializationTime == (Index == Parameters.InitialRootIndex) || bBuildGeometryForChildren)
 			{
 				if (DynamicCollectionAnchoringFacade.IsAnchored(Index))
 				{
@@ -895,7 +910,9 @@ void FGeometryCollectionPhysicsProxy::CreateChildrenGeometry_External()
 			TManagedArray<Chaos::FImplicitObjectPtr>& Implicits = GameThreadCollection.ModifyAttribute<Chaos::FImplicitObjectPtr>(FGeometryDynamicCollection::ImplicitsAttribute, FTransformCollection::TransformGroup);
 			TBitArray<> EffectiveParticles;
 			NumEffectiveParticles = CalculateEffectiveParticles(GameThreadCollection, NumParticles, Parameters.MaxSimulatedLevel, Parameters.EnableClustering, GetOwner(), EffectiveParticles);
-			CreateParticles(EffectiveParticles, Implicits, RBDSolver->GetEvolution(), /*InitilaizeOnlyRoot*/false);
+			CreateGTParticles(EffectiveParticles, Implicits, RBDSolver->GetEvolution(), /*bInitializationTime*/false);
+			SyncParticles_External();
+			UniqueIdxs.Empty();
 
 			// The Implicits attributes from the Dynamic Collection are just used for initialization, after they can be removed and so free some memory.
 			GameThreadCollection.RemoveAttribute(FGeometryDynamicCollection::ImplicitsAttribute, FTransformCollection::TransformGroup);
@@ -1011,7 +1028,16 @@ void FGeometryCollectionPhysicsProxy::CreateNonClusteredParticles(Chaos::FPBDRig
 			if (SimulatableParticles[Idx] && !RestCollection.IsClustered(Idx) && RestCollection.IsGeometry(Idx))
 			{
 				NumRigids++;
-				UniqueIndices.Add(GTParticles[Idx]->UniqueIdx());
+				Chaos::FUniqueIdx ExistingIndex;
+				if (GTParticles[Idx] == nullptr)
+				{
+					ExistingIndex = RigidsSolver->GetEvolution()->GenerateUniqueIdx();
+				}
+				else
+				{
+					ExistingIndex = GTParticles[Idx]->UniqueIdx();
+				}
+				UniqueIndices.Add(ExistingIndex);
 			}
 		}
 	}
@@ -1429,9 +1455,16 @@ void FGeometryCollectionPhysicsProxy::InitializeBodiesPT(Chaos::FPBDRigidsSolver
 
 					// Hook the handle up with the GT particle
 					FParticle* GTParticle = GTParticles[TransformGroupIndex].Get();
-					check(GTParticle != nullptr);
+					Chaos::FUniqueIdx ExistingIndex;
+					if (GTParticle != nullptr)
+					{
+						ExistingIndex = GTParticle->UniqueIdx();
+					}
+					else
+					{
+						ExistingIndex = UniqueIdxs[TransformGroupIndex];
+					}
 
-					Chaos::FUniqueIdx ExistingIndex = GTParticle->UniqueIdx();
 					Chaos::FPBDRigidClusteredParticleHandle* Handle = nullptr;
 					if (RigidChildren.Num())
 					{
@@ -1439,9 +1472,12 @@ void FGeometryCollectionPhysicsProxy::InitializeBodiesPT(Chaos::FPBDRigidsSolver
 					}
 					else
 					{
-						Handle = BuildNonClusters_Internal(TransformGroupIndex, RigidsSolver, Masses[TransformGroupIndex], Inertias[TransformGroupIndex]);
+						Handle = BuildNonClusters_Internal(TransformGroupIndex, RigidsSolver, Masses[TransformGroupIndex], Inertias[TransformGroupIndex], &ExistingIndex);
+					}
+					if (GTParticle != nullptr)
+					{
+						Handle->GTGeometryParticle() = GTParticle;
 					}
-					Handle->GTGeometryParticle() = GTParticle;
 
 					int32 RigidChildrenIdx = 0;
 					for(const int32 ChildTransformIndex : RigidChildrenTransformGroupIndex)
@@ -1676,6 +1712,25 @@ void FGeometryCollectionPhysicsProxy::CreateChildrenGeometry_Internal()
 	}
 }
 
+void FGeometryCollectionPhysicsProxy::SyncParticles_External()
+{
+	ExecuteOnPhysicsThread(*this, [this]()
+	{
+		int32 ParticlNum = GTParticles.Num();
+		for (int32 Index = 0; Index < ParticlNum; ++Index)
+		{
+			FParticle* GTParticle = GTParticles[Index].Get();
+			if (GTParticle)
+			{
+				FClusterHandle* Handle = SolverParticleHandles[Index];
+				check(Handle != nullptr);
+				Handle->GTGeometryParticle() = GTParticle;
+			}
+		}
+	});
+}
+
+
 int32 ReportNoLevelsetCluster = 0;
 FAutoConsoleVariableRef CVarReportNoLevelsetCluster(TEXT("p.gc.ReportNoLevelsetCluster"), ReportNoLevelsetCluster, TEXT("Report any cluster objects without levelsets"));
 
@@ -1809,7 +1864,7 @@ Chaos::FVec3f FGeometryCollectionPhysicsProxy::AdjustInertiaForScale(const Chaos
 	return WorldScaledInertia * Parameters.MaterialOverrideMassScaleMultiplier;
 }
 
-Chaos::TPBDGeometryCollectionParticleHandle<Chaos::FReal, 3>* FGeometryCollectionPhysicsProxy::BuildNonClusters_Internal(const uint32 CollectionClusterIndex, Chaos::FPBDRigidsSolver* RigidsSolver, float Mass, Chaos::FVec3f Inertia)
+Chaos::TPBDGeometryCollectionParticleHandle<Chaos::FReal, 3>* FGeometryCollectionPhysicsProxy::BuildNonClusters_Internal(const uint32 CollectionClusterIndex, Chaos::FPBDRigidsSolver* RigidsSolver, float Mass, Chaos::FVec3f Inertia, const Chaos::FUniqueIdx* ExistingIndex)
 {
 	FGeometryDynamicCollection& DynamicCollection = PhysicsThreadCollection;
 	TManagedArray<uint8>& DynamicState = DynamicCollection.DynamicState;
@@ -1824,7 +1879,7 @@ Chaos::TPBDGeometryCollectionParticleHandle<Chaos::FReal, 3>* FGeometryCollectio
 
 	// Gather unique indices from GT to pass into PT handle creation
 	TArray<Chaos::FUniqueIdx> UniqueIndices;
-	UniqueIndices.Add(GTParticles[CollectionClusterIndex]->UniqueIdx());
+	UniqueIndices.Add(*ExistingIndex);
 
 	// Add entries into simulation array
 	RigidsSolver->GetEvolution()->ReserveParticles(1);
@@ -4010,20 +4065,28 @@ bool FGeometryCollectionPhysicsProxy::PullNonInterpolatableDataFromSinglePhysics
 			continue;
 		}
 
-		if (GTParticles[TransformGroupIndex] == nullptr)
-		{
-			continue;
-		}
-
-		FParticle& GTParticle = *GTParticles[TransformGroupIndex];
-
+		
 		const bool bIsActive = !StateData.State.DisabledState;
 		if (UpdateValue(GameThreadCollection.Active[TransformGroupIndex], bIsActive))
 		{
-			GTParticle.SetDisabled(!bIsActive);
+			if (bIsActive && Parameters.InitialRootIndex != TransformGroupIndex)
+			{
+				CreateChildrenGeometry_External();
+			}
+			if (GTParticles[TransformGroupIndex] == nullptr)
+			{
+				continue;
+			}
+			GTParticles[TransformGroupIndex]->SetDisabled(!bIsActive);
 			bIsCollectionDirty = true;
 		}
 
+		if (GTParticles[TransformGroupIndex] == nullptr)
+		{
+			continue;
+		}
+		FParticle& GTParticle = *GTParticles[TransformGroupIndex];
+
 		if (UpdateValue(GameThreadCollection.DynamicState[TransformGroupIndex], static_cast<uint8>(StateData.State.DynamicState)))
 		{
 			GTParticle.SetObjectState(static_cast<Chaos::EObjectStateType>(StateData.State.DynamicState));
@@ -4165,6 +4228,13 @@ bool FGeometryCollectionPhysicsProxy::PullFromPhysicsState(const Chaos::FDirtyGe
 			// this is helped by the FResultAccessor structure
 			for (int32 TransformGroupIndex = 0; TransformGroupIndex < NumTransforms; ++TransformGroupIndex)
 			{
+				const bool bActive = GameThreadCollection.Active[TransformGroupIndex];
+				bAtLeatOneChildActive |= bActive && TransformGroupIndex != Parameters.InitialRootIndex;
+				if (bAtLeatOneChildActive)
+				{
+					CreateChildrenGeometry_External();
+				}
+
 				if (GTParticles[TransformGroupIndex] == nullptr)
 				{
 					continue;
@@ -4178,9 +4248,6 @@ bool FGeometryCollectionPhysicsProxy::PullFromPhysicsState(const Chaos::FDirtyGe
 				}
 
 				const bool bAnimatingWhileDisabled = AnimationsActive ? (*AnimationsActive)[TransformGroupIndex] : false;
-				const bool bActive = GameThreadCollection.Active[TransformGroupIndex];
-
-				bAtLeatOneChildActive |= bActive && TransformGroupIndex != Parameters.InitialRootIndex;
 
 				if (bActive || bAnimatingWhileDisabled)
 				{
@@ -4207,10 +4274,6 @@ bool FGeometryCollectionPhysicsProxy::PullFromPhysicsState(const Chaos::FDirtyGe
 					}
 				}
 			}
-			if (bAtLeatOneChildActive)
-			{
-				CreateChildrenGeometry_External();
-			}
 		}
 
 		// update the parent transform now that we have all X and R computed
diff --git a/Engine/Source/Runtime/Experimental/Chaos/Public/PhysicsProxy/GeometryCollectionPhysicsProxy.h b/Engine/Source/Runtime/Experimental/Chaos/Public/PhysicsProxy/GeometryCollectionPhysicsProxy.h
--- a/Engine/Source/Runtime/Experimental/Chaos/Public/PhysicsProxy/GeometryCollectionPhysicsProxy.h
+++ b/Engine/Source/Runtime/Experimental/Chaos/Public/PhysicsProxy/GeometryCollectionPhysicsProxy.h
@@ -520,7 +520,7 @@ protected:
 	/** adjust inertia to account for per component scale properties ( from material override and world transform scale ) */
 	CHAOS_API Chaos::FVec3f AdjustInertiaForScale(const Chaos::FVec3f& Inertia) const;
 
-	CHAOS_API Chaos::TPBDGeometryCollectionParticleHandle<Chaos::FReal, 3>* BuildNonClusters_Internal(const uint32 CollectionClusterIndex, Chaos::FPBDRigidsSolver* RigidsSolver, float Mass, Chaos::FVec3f Inertia);
+	CHAOS_API Chaos::TPBDGeometryCollectionParticleHandle<Chaos::FReal, 3>* BuildNonClusters_Internal(const uint32 CollectionClusterIndex, Chaos::FPBDRigidsSolver* RigidsSolver, float Mass, Chaos::FVec3f Inertia, const Chaos::FUniqueIdx* ExistingIndex);
 
 	/**
 	 * Build a physics thread cluster parent particle.
@@ -574,8 +574,9 @@ private:
 	static TBitArray<> CalculateClustersToCreateFromChildren(const FGeometryDynamicCollection& DynamicCollection, int32 NumTransforms);
 	static int32 CalculateEffectiveParticles(const FGeometryDynamicCollection& DynamicCollection, int32 NumTransform, int32 MaxSimulatedLevel, bool bEnableClustering, const UObject* Owner, TBitArray<>& EffectiveParticles);
 
-	void CreateParticles(const TBitArray<>& EffectiveParticles, TManagedArray<Chaos::FImplicitObjectPtr>& Implicits, Chaos::FPBDRigidsEvolutionBase* Evolution, bool bInitializeRootOnly);
+	void CreateGTParticles(const TBitArray<>& EffectiveParticles, TManagedArray<Chaos::FImplicitObjectPtr>& Implicits, Chaos::FPBDRigidsEvolutionBase* Evolution, bool bInitializeRootOnly);
 	void CreateChildrenGeometry_External();
+	void SyncParticles_External();
 	
 	/**
 	 * Since geometry collections only buffer data that has changed, when PullFromPhysicsState is given both PrevData and NextData it must
@@ -612,6 +613,7 @@ private:
 	TSet<FClusterHandle*> SolverAnchors;
 	TMap<FParticleHandle*, int32> HandleToTransformGroupIndex;
 	TMap<int32, FClusterHandle*> UniqueIdxToInternalClusterHandle;
+	TArray<Chaos::FUniqueIdx> UniqueIdxs;
 
 	//
 	// Buffer Results State Information
