Fix compilation error fallout from CL 29445954

diff --git a/Engine/Source/Runtime/Experimental/Chaos/Private/PhysicsProxy/GeometryCollectionPhysicsProxy.cpp b/Engine/Source/Runtime/Experimental/Chaos/Private/PhysicsProxy/GeometryCollectionPhysicsProxy.cpp
--- a/Engine/Source/Runtime/Experimental/Chaos/Private/PhysicsProxy/GeometryCollectionPhysicsProxy.cpp
+++ b/Engine/Source/Runtime/Experimental/Chaos/Private/PhysicsProxy/GeometryCollectionPhysicsProxy.cpp
@@ -706,8 +706,8 @@ void FGeometryCollectionPhysicsProxy::CreateGTParticles(const TBitArray<>& Effec
 				UniqueIdxs[Index] = Evolution->GenerateUniqueIdx();
 			}
 
-			if (bInitializationTime && Index == Parameters.InitialRootIndex || // When initializing always create particle for the root
-				(bInitializationTime && (bCreateGTParticleForChildren || bBuildGeometryForChildren) || // When initializing create all particles if one of the flag is true 
+			if ((bInitializationTime && Index == Parameters.InitialRootIndex) || // When initializing always create particle for the root
+				((bInitializationTime && (bCreateGTParticleForChildren || bBuildGeometryForChildren)) || // When initializing create all particles if one of the flag is true 
 					(!bInitializationTime && Index != Parameters.InitialRootIndex && !bCreateGTParticleForChildren && !bBuildGeometryForChildren))) // When not initializing create other particles if flag was set to not create
 			{
 				GTParticles[Index] = FParticle::CreateParticle();
