Iris
- Fixed bug with bExecuteReliableRPCsBeforeOnReps not executing rpcs at the expected time
- Renamed cvar to net.Iris.ExecuteReliableRPCsBeforeApplyState which better describes what it does.
- Fixed disconnect if pre-dispatch of attachments cleared out attachement queue

diff --git a/Engine/Source/Runtime/Experimental/Iris/Core/Private/Iris/ReplicationSystem/ReplicationReader.cpp b/Engine/Source/Runtime/Experimental/Iris/Core/Private/Iris/ReplicationSystem/ReplicationReader.cpp
--- a/Engine/Source/Runtime/Experimental/Iris/Core/Private/Iris/ReplicationSystem/ReplicationReader.cpp
+++ b/Engine/Source/Runtime/Experimental/Iris/Core/Private/Iris/ReplicationSystem/ReplicationReader.cpp
@@ -41,11 +41,11 @@
 namespace UE::Net::Private
 {
 
-static bool bExecuteReliableRPCsBeforeOnReps = false;
-static FAutoConsoleVariableRef CVarExecuteReliableRPCsBeforeOnReps(
-		TEXT("net.Iris.ExecuteReliableRPCsBeforeOnReps"),
-		bExecuteReliableRPCsBeforeOnReps,
-		TEXT("If true and Iris runs in backwards compatibility mode then reliable RPCs will be executed before OnReps on the target object. Default is false."
+static bool bExecuteReliableRPCsBeforeApplyState = true;
+static FAutoConsoleVariableRef CVarExecuteReliableRPCsBeforeApplyState(
+		TEXT("net.Iris.ExecuteReliableRPCsBeforeApplyState"),
+		bExecuteReliableRPCsBeforeApplyState,
+		TEXT("If true and Iris runs in backwards compatibility mode then reliable RPCs will be executed before we apply state data on the target object unless we first need to spawn the object."
 		));
 
 static bool bDeferEndReplication = true;
@@ -1344,7 +1344,7 @@ void FReplicationReader::ResolveAndDispatchUnresolvedReferencesForObject(FNetSer
 			// Apply resolved references, this is a blunt tool as we currently push out full dirty properties rather than only the resolved references
 			if (ResolvedChangeMask.IsAnyBitSet() || bHasResolvedInitReferences)
 			{
-				if (bObjectHasAttachments && bExecuteReliableRPCsBeforeOnReps && !bHasResolvedInitReferences)
+				if (bObjectHasAttachments && bExecuteReliableRPCsBeforeApplyState && !bHasResolvedInitReferences)
 				{
 					AttachmentDispatchedFlags = ENetObjectAttachmentDispatchFlags::Reliable;
 					ResolveAndDispatchAttachments(Context, ReplicationInfo, ENetObjectAttachmentDispatchFlags::Reliable);
@@ -1409,13 +1409,6 @@ void FReplicationReader::DispatchStateData(FNetSerializationContext& Context)
 		{
 			FDispatchObjectInfo& Info = *PostDispatchObjectInfo.Info;
 
-			// If we are running in backwards compatibility mode, execute Reliable RPC`s before RepNotify callbacks
-			if (Info.bHasAttachments && bExecuteReliableRPCsBeforeOnReps && !Info.bIsInitialState)
-			{
-				PostDispatchObjectInfo.AttachmentDispatchedFlags |= ENetObjectAttachmentDispatchFlags::Reliable;
-				ResolveAndDispatchAttachments(Context, PostDispatchObjectInfo.ReplicationInfo, ENetObjectAttachmentDispatchFlags::Reliable);
-			}
-
 			// Execute legacy post replicate functions
 			if (Info.bHasState && PostDispatchObjectInfo.DequantizeAndApplyContext)
 			{
@@ -1478,6 +1471,15 @@ void FReplicationReader::DispatchStateData(FNetSerializationContext& Context)
 			}
 		}
 
+		// If we are running in backwards compatibility mode, execute Reliable RPC`s before applying state data unless object is already created.
+		if (Info.bHasAttachments && bExecuteReliableRPCsBeforeApplyState && !Info.bIsInitialState)
+		{
+			PostDispatchObjectInfo.AttachmentDispatchedFlags |= ENetObjectAttachmentDispatchFlags::Reliable;
+			ResolveAndDispatchAttachments(Context, PostDispatchObjectInfo.ReplicationInfo, ENetObjectAttachmentDispatchFlags::Reliable);
+			// Update if we have attachments or not since we might have processed all of them in the first pass.
+			Info.bHasAttachments = PostDispatchObjectInfo.ReplicationInfo->bHasAttachments;
+		}
+
 		// If we have any object references we want to update any unresolved ones, including previously unresolved references
 		if (Info.bHasState)
 		{
