Fixed double free segfault.

When one IAS object is nearing completion we send out a request for the next one. If the current request has in fact finished, and requesting the next one fails, the current request will end up added to the "done list" twice. This results in said request getting freed twice when the done list is processed.

#rnx
#tests wrote tests to reliably reproduce use-after-free and implement and verify fix

diff --git a/Engine/Source/Runtime/Experimental/IoStoreOnDemand/Private/CoreHttp/Http.cpp b/Engine/Source/Runtime/Experimental/IoStoreOnDemand/Private/CoreHttp/Http.cpp
--- a/Engine/Source/Runtime/Experimental/IoStoreOnDemand/Private/CoreHttp/Http.cpp
+++ b/Engine/Source/Runtime/Experimental/IoStoreOnDemand/Private/CoreHttp/Http.cpp
@@ -2777,6 +2777,11 @@ void FSocketGroup::RecvInternal(FTickState& State)
 				check(Send == nullptr);
 				Send = Next;
 				SendInternal(State);
+
+				if (!Socket.IsValid())
+				{
+					return;
+				}
 			}
 		}
 	}
