# Pushout players that are overlapping physics objects

This change has a number of parts to it

1. Add a mid-phase modification for disabling convex optimization for specific interactions, and use it for interactions involving player capsules so that players do not collide with empty volumes of concave buildings.
2. Ignore debris in hit callback
3. Try to depenetrate player capsule in the direction of the contact, in the up direction, and in the directions normal to the plane shared by those two vectors.

#swarm https://p4-swarm.epicgames.net/reviews/29417971

#localization none
#tests Local tests in PIE

diff --git a/Engine/Source/Runtime/Experimental/Chaos/Private/Chaos/Collision/ParticlePairMidPhase.cpp b/Engine/Source/Runtime/Experimental/Chaos/Private/Chaos/Collision/ParticlePairMidPhase.cpp
--- a/Engine/Source/Runtime/Experimental/Chaos/Private/Chaos/Collision/ParticlePairMidPhase.cpp
+++ b/Engine/Source/Runtime/Experimental/Chaos/Private/Chaos/Collision/ParticlePairMidPhase.cpp
@@ -812,6 +812,7 @@ namespace Chaos
 		Flags.bIsCCDActive = false;
 		Flags.bIsSleeping = false;
 		Flags.bIsModified = false;
+		Flags.bIsConvexOptimizationActive = true;
 
 		ResetImpl();
 	}
@@ -822,6 +823,7 @@ namespace Chaos
 		{
 			Flags.bIsActive = true;
 			Flags.bIsCCDActive = Flags.bIsCCD;
+			Flags.bIsConvexOptimizationActive = true;
 			Flags.bIsModified = false;
 		}
 	}
@@ -851,6 +853,9 @@ namespace Chaos
 		Flags.bIsCCD = bIsCCD;
 		Flags.bIsCCDActive = bIsCCD;
 
+		// Initially we allow for convex optimization where available
+		Flags.bIsConvexOptimizationActive = true;
+
 		BuildDetectorsImpl();
 
 		InitThresholds();
@@ -1185,8 +1190,14 @@ namespace Chaos
 		FPBDRigidClusteredParticleHandle* ClusteredHandle0 = GetParticle0()->CastToClustered();
 		FPBDRigidClusteredParticleHandle* ClusteredHandle1 = GetParticle1()->CastToClustered();
 
-		Private::FConvexOptimizer* ConvexOptimizer0 = ClusteredHandle0 ? ClusteredHandle0->ConvexOptimizer().Get() : nullptr;
-		Private::FConvexOptimizer* ConvexOptimizer1 = ClusteredHandle1 ? ClusteredHandle1->ConvexOptimizer().Get() : nullptr;
+		Private::FConvexOptimizer* ConvexOptimizer0
+			= Flags.bIsConvexOptimizationActive && ClusteredHandle0
+			? ClusteredHandle0->ConvexOptimizer().Get()
+			: nullptr;
+		Private::FConvexOptimizer* ConvexOptimizer1
+			= Flags.bIsConvexOptimizationActive && ClusteredHandle1
+			? ClusteredHandle1->ConvexOptimizer().Get()
+			: nullptr;
 
 		// See if we have a BVH for either/both of the particles
 		const Private::FImplicitBVH* BVH0 = nullptr;
diff --git a/Engine/Source/Runtime/Experimental/Chaos/Private/Chaos/ContactModification.cpp b/Engine/Source/Runtime/Experimental/Chaos/Private/Chaos/ContactModification.cpp
--- a/Engine/Source/Runtime/Experimental/Chaos/Private/Chaos/ContactModification.cpp
+++ b/Engine/Source/Runtime/Experimental/Chaos/Private/Chaos/ContactModification.cpp
@@ -489,6 +489,16 @@ namespace Chaos
 		return { Constraint->GetShape0(), Constraint->GetShape1() };
 	}
 
+	const FShapeInstance* FContactPairModifier::GetShape(int32 ParticleIdx) const
+	{
+		return Constraint->GetShape(ParticleIdx);
+	}
+
+	const FConstImplicitObjectRef FContactPairModifier::GetImplicit(int32 ParticleIdx) const
+	{
+		return Constraint->GetImplicit(ParticleIdx);
+	}
+
 	bool FContactPairModifier::IsEdgeContactPoint(int32 ContactPointIdx) const
 	{
 		const FManifoldPoint& ManifoldPoint = Constraint->GetManifoldPoint(ContactPointIdx);
diff --git a/Engine/Source/Runtime/Experimental/Chaos/Private/Chaos/MidPhaseModification.cpp b/Engine/Source/Runtime/Experimental/Chaos/Private/Chaos/MidPhaseModification.cpp
--- a/Engine/Source/Runtime/Experimental/Chaos/Private/Chaos/MidPhaseModification.cpp
+++ b/Engine/Source/Runtime/Experimental/Chaos/Private/Chaos/MidPhaseModification.cpp
@@ -20,6 +20,11 @@ namespace Chaos
 		MidPhase->SetCCDIsActive(false);
 	}
 
+	void FMidPhaseModifier::DisableConvexOptimization()
+	{
+		MidPhase->SetConvexOptimizationIsActive(false);
+	}
+
 	void FMidPhaseModifier::GetParticles(
 		FGeometryParticleHandle** Particle0,
 		FGeometryParticleHandle** Particle1) const
diff --git a/Engine/Source/Runtime/Experimental/Chaos/Public/Chaos/Collision/ParticlePairMidPhase.h b/Engine/Source/Runtime/Experimental/Chaos/Public/Chaos/Collision/ParticlePairMidPhase.h
--- a/Engine/Source/Runtime/Experimental/Chaos/Public/Chaos/Collision/ParticlePairMidPhase.h
+++ b/Engine/Source/Runtime/Experimental/Chaos/Public/Chaos/Collision/ParticlePairMidPhase.h
@@ -333,6 +333,12 @@ namespace Chaos
 			Flags.bIsCCDActive = bCCDIsActive;
 		}
 
+		void SetConvexOptimizationIsActive(bool bSetConvexOptimizationIsActive)
+		{
+			Flags.bIsModified = true;
+			Flags.bIsConvexOptimizationActive = bSetConvexOptimizationIsActive;
+		}
+
 		/*
 		 * True if CCD is supported by either particle
 		 */
@@ -387,6 +393,7 @@ namespace Chaos
 				uint16 bIsCCD : 1;       // True if CCD is supported by either particle
 				uint16 bIsCCDActive : 1; // True if CCD is active for this midphase on this frame. This can be changed by modifiers and resets to bIsCCD each frame.
 				uint16 bUseSweep : 1;    // True if CCD is active (this frame) and we are moving fast enough to require a sweep
+				uint16 bIsConvexOptimizationActive : 1; // True if convex optimization is active for this midphase
 				uint16 bIsSleeping : 1;
 				uint16 bIsModified : 1;  // True if a modifier applied any changes to this midphase
 			};
diff --git a/Engine/Source/Runtime/Experimental/Chaos/Public/Chaos/ContactModification.h b/Engine/Source/Runtime/Experimental/Chaos/Public/Chaos/ContactModification.h
--- a/Engine/Source/Runtime/Experimental/Chaos/Public/Chaos/ContactModification.h
+++ b/Engine/Source/Runtime/Experimental/Chaos/Public/Chaos/ContactModification.h
@@ -3,6 +3,7 @@
 
 #include "Chaos/Core.h"
 #include "ParticleHandleFwd.h"
+#include "Chaos/ImplicitFwd.h"
 
 namespace
 {
@@ -22,6 +23,7 @@ namespace Chaos
 	class FCollisionContactModifier;
 	class FContactPairModifier;
 	class FPerShapeData;
+	class FShapeInstance;
 	class FPBDCollisionConstraintContainerCookie;
 
 	class FContactPairModifier
@@ -247,6 +249,16 @@ namespace Chaos
 		*/
 		CHAOS_API TVec2<const FPerShapeData*> GetShapePair() const;
 
+		/*
+		* Get one of the particle's shapes.
+		*/
+		CHAOS_API const FShapeInstance* GetShape(int32 ParticleIdx) const;
+
+		/*
+		* Get one of the particle's implicit objects.
+		*/
+		CHAOS_API const FConstImplicitObjectRef GetImplicit(int32 ParticleIdx) const;
+
 		/*
 		* Check to see if a contact point index is an edge contact.
 		*/
diff --git a/Engine/Source/Runtime/Experimental/Chaos/Public/Chaos/MidPhaseModification.h b/Engine/Source/Runtime/Experimental/Chaos/Public/Chaos/MidPhaseModification.h
--- a/Engine/Source/Runtime/Experimental/Chaos/Public/Chaos/MidPhaseModification.h
+++ b/Engine/Source/Runtime/Experimental/Chaos/Public/Chaos/MidPhaseModification.h
@@ -52,6 +52,9 @@ namespace Chaos
 		// Disable CCD for this pair
 		CHAOS_API void DisableCCD();
 
+		// Disable convex optimization for this pair
+		CHAOS_API void DisableConvexOptimization();
+
 		//
 		// Accessor functions
 		//
