Chaos - Fix OneWayInteraction for new Cluster Particles. The cluster should only act like a one-way particle if all particles in it are one-way particles.

This fixes an issue where the particles added to the cluster are a mix of one-way and standard particles. This can happen in the partial destruction case where we have geometry collection particles from different levels in the same cluster. In this case it is possible that the L1 or lower particles could be added to the cluster first and the cluster would inherit their one-way property.

diff --git a/Engine/Source/Runtime/Experimental/Chaos/Private/Chaos/ClusterUnionManager.cpp b/Engine/Source/Runtime/Experimental/Chaos/Private/Chaos/ClusterUnionManager.cpp
--- a/Engine/Source/Runtime/Experimental/Chaos/Private/Chaos/ClusterUnionManager.cpp
+++ b/Engine/Source/Runtime/Experimental/Chaos/Private/Chaos/ClusterUnionManager.cpp
@@ -603,6 +603,9 @@ namespace Chaos
 		// not the new parent when AddParticlesToCluster is called.
 		MClustering.AddParticlesToCluster(Cluster->InternalCluster, FinalParticlesToAdd, ChildToParentMap);
 
+		// Cluster has one-way interaction only if all children are also one-way
+		bool bIsOneWayInteraction = (bIsNewCluster) || Cluster->InternalCluster->OneWayInteraction();
+
 		for (FPBDRigidParticleHandle* Particle : FinalParticlesToAdd)
 		{
 			if (Particle->GetGeometry() == nullptr)
@@ -620,6 +623,8 @@ namespace Chaos
 				Cluster->AddPendingGeometryOperation(EClusterUnionGeometryOperation::Add, Particle);
 			}
 
+			bIsOneWayInteraction &= Particle->OneWayInteraction();
+
 			if (!Cluster->ChildProperties.Contains(Particle))
 			{
 				FClusterUnionParticleProperties Properties;
@@ -684,10 +689,12 @@ namespace Chaos
 					Cluster->InternalCluster->SetGravityGroupIndex(ParticlePropertySource->GravityGroupIndex());
 				}
 
+				// Should be min or max of all children or something else that doesn't depend on order? 
 				Cluster->InternalCluster->SetInitialOverlapDepenetrationVelocity(ParticlePropertySource->InitialOverlapDepenetrationVelocity());
-				Cluster->InternalCluster->SetOneWayInteraction(ParticlePropertySource->OneWayInteraction());
 			}
 		}
+
+		Cluster->InternalCluster->SetOneWayInteraction(bIsOneWayInteraction);
 	}
 
 	DECLARE_CYCLE_STAT(TEXT("FClusterUnionManager::DeferredClusterUnionUpdate"), STAT_DeferredClusterUnionUpdate, STATGROUP_Chaos);
