Ias: Fix RemoveDispatcherRequest It nullptr crash when cancelling requests.
Fix Reprioritize/RemoveInternal bug that could loose track of requests so that they would never ever be issued nor completed.

#rnx
#tests frontend, replays and joining server

#localization none

diff --git a/Engine/Source/Runtime/Experimental/IoStoreOnDemand/Private/OnDemandIoDispatcherBackend.cpp b/Engine/Source/Runtime/Experimental/IoStoreOnDemand/Private/OnDemandIoDispatcherBackend.cpp
--- a/Engine/Source/Runtime/Experimental/IoStoreOnDemand/Private/OnDemandIoDispatcherBackend.cpp
+++ b/Engine/Source/Runtime/Experimental/IoStoreOnDemand/Private/OnDemandIoDispatcherBackend.cpp
@@ -674,6 +674,10 @@ private:
 			{
 				It->NextRequest = It->NextRequest->NextRequest;
 				Request->NextRequest = nullptr;
+				if (Tail == Request)
+				{
+					Tail = It;
+				}
 				return true;
 			}
 		}
@@ -802,7 +806,7 @@ struct FChunkRequest
 		return bPriorityChanged;
 	}
 
-	uint32 RemoveDispatcherRequest(FIoRequestImpl* Request)
+	int32 RemoveDispatcherRequest(FIoRequestImpl* Request)
 	{
 		check(Request != nullptr);
 		check(RequestCount > 0);
@@ -822,6 +826,10 @@ struct FChunkRequest
 			while (It->NextRequest != Request)
 			{
 				It = It->NextRequest;
+				if (It == nullptr)
+				{
+					return INDEX_NONE; // Not found
+				}
 			}
 			check(It->NextRequest == Request);
 			It->NextRequest = It->NextRequest->NextRequest;
@@ -1439,7 +1447,16 @@ class FOnDemandIoBackend final
 			if (FChunkRequest** InflightRequest = Inflight.Find(BackendData->ChunkKey))
 			{
 				FChunkRequest& ChunkRequest = **InflightRequest;
-				const uint32 RemainingCount = ChunkRequest.RemoveDispatcherRequest(Request);
+				const int32 RemainingCount = ChunkRequest.RemoveDispatcherRequest(Request);
+				if (RemainingCount == INDEX_NONE)
+				{
+					// Not found
+					// When a request A with ChunkKey X enters CompleteRequest its Inflight entry X->A is removed.
+					// If a new request B with the same ChunkKey X is made, then Resolve will add a new Infligt entry X->B.
+					// If we at this point cancel A, we will find the Inflight entry for B, which will not contain A, which is fine.
+					return false;
+				}
+
 				check(Request->NextRequest == nullptr);
 
 				if (RemainingCount == 0)
