Guard against crashes from OOB file reads.

Note that this change builds upon a more complete fix already submitted in 29178406. This additional change makes reads more robust and tolerant of error conditions, prefering to return read errors instead of causing the OS to terminate the process.

#rnx
#tests induced incorrect read parameters on reported platform to confirm.
#localization none

diff --git a/Engine/Source/Runtime/Experimental/IoStoreOnDemand/Private/JournaledCache.cpp b/Engine/Source/Runtime/Experimental/IoStoreOnDemand/Private/JournaledCache.cpp
--- a/Engine/Source/Runtime/Experimental/IoStoreOnDemand/Private/JournaledCache.cpp
+++ b/Engine/Source/Runtime/Experimental/IoStoreOnDemand/Private/JournaledCache.cpp
@@ -780,6 +780,11 @@ EIoErrorCode FDiskCache::Materialize(uint64 Key, FIoBuffer& Out, uint32 Offset)
 
 	ReadSize = FMath::Min<uint32>(uint32(Out.GetSize()), ReadSize);
 
+	if (EntryDataCursor + Offset + ReadSize > uint64(DataHandle->Size()))
+	{
+		return EIoErrorCode::ReadError;
+	}
+
 	TRACE_COUNTER_SET(IasReadCursor, EntryDataCursor + Offset);
 
 	DataHandle->Seek(EntryDataCursor + Offset);
