ISMPool, fix issue where preallocation queue contains raw object pointers that can be invalid by the time we process the entries.
Use weak object pointers instead.

diff --git a/Engine/Source/Runtime/Experimental/GeometryCollectionEngine/Private/GeometryCollection/GeometryCollectionISMPoolComponent.cpp b/Engine/Source/Runtime/Experimental/GeometryCollectionEngine/Private/GeometryCollection/GeometryCollectionISMPoolComponent.cpp
--- a/Engine/Source/Runtime/Experimental/GeometryCollectionEngine/Private/GeometryCollection/GeometryCollectionISMPoolComponent.cpp
+++ b/Engine/Source/Runtime/Experimental/GeometryCollectionEngine/Private/GeometryCollection/GeometryCollectionISMPoolComponent.cpp
@@ -89,11 +89,14 @@ void FGeometryCollectionISM::CreateISM(AActor* InOwningActor)
 void FGeometryCollectionISM::InitISM(const FGeometryCollectionStaticMeshInstance& InMeshInstance, bool bKeepAlive, bool bOverrideTransformUpdates)
 {
 	MeshInstance = InMeshInstance;
-	check(MeshInstance.StaticMesh);
 	check(ISMComponent != nullptr);
 
+	UStaticMesh* StaticMesh = MeshInstance.StaticMesh.Get();
+	// We should only get here for valid static mesh objects.
+	check(StaticMesh != nullptr);
+
 #if WITH_EDITOR
-	const FName ISMName = MakeUniqueObjectName(ISMComponent->GetOwner(), UInstancedStaticMeshComponent::StaticClass(), InMeshInstance.StaticMesh->GetFName());
+	const FName ISMName = MakeUniqueObjectName(ISMComponent->GetOwner(), UInstancedStaticMeshComponent::StaticClass(), StaticMesh->GetFName());
 	const FString ISMNameString = ISMName.ToString();
 	ISMComponent->Rename(*ISMNameString);
 #endif
@@ -104,10 +107,13 @@ void FGeometryCollectionISM::InitISM(const FGeometryCollectionStaticMeshInstance
 	ISMComponent->EmptyOverrideMaterials();
 	for (int32 MaterialIndex = 0; MaterialIndex < MeshInstance.MaterialsOverrides.Num(); MaterialIndex++)
 	{
-		ISMComponent->SetMaterial(MaterialIndex, MeshInstance.MaterialsOverrides[MaterialIndex]);
+		UMaterialInterface* Material = MeshInstance.MaterialsOverrides[MaterialIndex].Get();
+		// We should only get here for valid material objects.
+		check(Material != nullptr);
+		ISMComponent->SetMaterial(MaterialIndex, Material);
 	}
 
-	ISMComponent->SetStaticMesh(MeshInstance.StaticMesh);
+	ISMComponent->SetStaticMesh(StaticMesh);
 	ISMComponent->SetMobility((MeshInstance.Desc.Flags & FISMComponentDescription::StaticMobility) != 0 ? EComponentMobility::Static : EComponentMobility::Movable);
 
 	ISMComponent->NumCustomDataFloats = MeshInstance.Desc.NumCustomDataFloats;
@@ -437,13 +443,37 @@ void FGeometryCollectionISMPool::RequestPreallocateMeshInstance(const FGeometryC
 	}
 }
 
+static bool AreWeakPointersValid(FGeometryCollectionStaticMeshInstance& InMeshInstance)
+{
+	if (!InMeshInstance.StaticMesh.IsValid())
+	{
+		return false;
+	}
+
+	for (TWeakObjectPtr<UMaterialInterface> Material : InMeshInstance.MaterialsOverrides)
+	{
+		if (!Material.IsValid())
+		{
+			return false;
+		}
+	}
+
+	return true;
+}
+
 void FGeometryCollectionISMPool::ProcessPreallocationRequests(UGeometryCollectionISMPoolComponent* OwningComponent, int32 MaxPreallocations)
 {
 	int32 NumAdded = 0;
 	for (TSet<FGeometryCollectionStaticMeshInstance>::TIterator It(PrellocationQueue); It; ++It)
 	{
 		bool bISMCreated = false;
-		GetOrAddISM(OwningComponent, *It, bISMCreated);
+
+		// Objects in the entries of the preallocation queue may no longer be loaded.
+		if (AreWeakPointersValid(*It))
+		{
+			GetOrAddISM(OwningComponent, *It, bISMCreated);
+		}
+
 		It.RemoveCurrent();
 
 		if (bISMCreated)
diff --git a/Engine/Source/Runtime/Experimental/GeometryCollectionEngine/Private/GeometryCollection/GeometryCollectionISMPoolRenderer.cpp b/Engine/Source/Runtime/Experimental/GeometryCollectionEngine/Private/GeometryCollection/GeometryCollectionISMPoolRenderer.cpp
--- a/Engine/Source/Runtime/Experimental/GeometryCollectionEngine/Private/GeometryCollection/GeometryCollectionISMPoolRenderer.cpp
+++ b/Engine/Source/Runtime/Experimental/GeometryCollectionEngine/Private/GeometryCollection/GeometryCollectionISMPoolRenderer.cpp
@@ -149,7 +149,8 @@ void UGeometryCollectionISMPoolRenderer::InitInstancesFromGeometryCollection(UGe
 			StaticMeshInstance.StaticMesh = const_cast<UStaticMesh*>(StaticMesh);
 			if (bMaterialOverride)
 			{
-				StaticMeshInstance.MaterialsOverrides = AutoInstanceMesh.Materials;
+				StaticMeshInstance.MaterialsOverrides.Reset();
+				StaticMeshInstance.MaterialsOverrides.Append(AutoInstanceMesh.Materials);
 			}
 
 			TArray<float> DummyCustomData;
diff --git a/Engine/Source/Runtime/Experimental/GeometryCollectionEngine/Public/GeometryCollection/GeometryCollectionISMPoolComponent.h b/Engine/Source/Runtime/Experimental/GeometryCollectionEngine/Public/GeometryCollection/GeometryCollectionISMPoolComponent.h
--- a/Engine/Source/Runtime/Experimental/GeometryCollectionEngine/Public/GeometryCollection/GeometryCollectionISMPoolComponent.h
+++ b/Engine/Source/Runtime/Experimental/GeometryCollectionEngine/Public/GeometryCollection/GeometryCollectionISMPoolComponent.h
@@ -157,14 +157,18 @@ FORCEINLINE uint32 GetTypeHash(const FISMComponentDescription& Desc)
  */
 struct FGeometryCollectionStaticMeshInstance
 {
-	UStaticMesh* StaticMesh = nullptr;
-	TArray<UMaterialInterface*> MaterialsOverrides;
+	TWeakObjectPtr<UStaticMesh> StaticMesh;
+	TArray<TWeakObjectPtr<UMaterialInterface>> MaterialsOverrides;
 	TArray<float> CustomPrimitiveData;
 	FISMComponentDescription Desc;
 
 	bool operator==(const FGeometryCollectionStaticMeshInstance& Other) const 
 	{
-		if (StaticMesh != Other.StaticMesh || !(Desc == Other.Desc))
+		if (!(Desc == Other.Desc))
+		{
+			return false;
+		}
+		if (!StaticMesh.HasSameIndexAndSerialNumber(Other.StaticMesh))
 		{
 			return false;
 		}
@@ -174,9 +178,7 @@ struct FGeometryCollectionStaticMeshInstance
 		}
 		for (int32 MatIndex = 0; MatIndex < MaterialsOverrides.Num(); MatIndex++)
 		{
-			const FName MatName = MaterialsOverrides[MatIndex] ? MaterialsOverrides[MatIndex]->GetFName() : NAME_None;
-			const FName OtherName = Other.MaterialsOverrides[MatIndex] ? Other.MaterialsOverrides[MatIndex]->GetFName() : NAME_None;
-			if (MatName != OtherName)
+			if (!MaterialsOverrides[MatIndex].HasSameIndexAndSerialNumber(Other.MaterialsOverrides[MatIndex]))
 			{
 				return false;
 			}
@@ -201,7 +203,7 @@ FORCEINLINE uint32 GetTypeHash(const FGeometryCollectionStaticMeshInstance& Mesh
 {
 	uint32 CombinedHash = GetTypeHash(MeshInstance.StaticMesh);
 	CombinedHash = HashCombineFast(CombinedHash, GetTypeHash(MeshInstance.MaterialsOverrides.Num()));
-	for (const UMaterialInterface* Material: MeshInstance.MaterialsOverrides)
+	for (const TWeakObjectPtr<UMaterialInterface> Material: MeshInstance.MaterialsOverrides)
 	{
 		CombinedHash = HashCombineFast(CombinedHash, GetTypeHash(Material));
 	}
