Chaos - Fix Geometry Collection Proxy so that it disabled the root particle when its children are in an external cluster union. (addresses failed fix from CL 24931330)

diff --git a/Engine/Source/Runtime/Experimental/Chaos/Private/PhysicsProxy/GeometryCollectionPhysicsProxy.cpp b/Engine/Source/Runtime/Experimental/Chaos/Private/PhysicsProxy/GeometryCollectionPhysicsProxy.cpp
--- a/Engine/Source/Runtime/Experimental/Chaos/Private/PhysicsProxy/GeometryCollectionPhysicsProxy.cpp
+++ b/Engine/Source/Runtime/Experimental/Chaos/Private/PhysicsProxy/GeometryCollectionPhysicsProxy.cpp
@@ -138,6 +138,12 @@ FAutoConsoleVariableRef CVarGeometryCollectionRotationUpdateTolerance(
 	GeometryCollectionRotationUpdateTolerance,
 	TEXT("Tolerance to detect if particle rotation has changed has changed when syncing PT to GT"));
 
+bool bGeometryCollectionUseRootBrokenFlag = true;
+FAutoConsoleVariableRef CVarGeometryCollectionUseRootBrokenFlag(
+	TEXT("p.GeometryCollection.UseRootBrokenFlag"),
+	bGeometryCollectionUseRootBrokenFlag,
+	TEXT("If enabled, check if the root transform is broken in the proxy and disable the GT particle if so. Should be enabled - cvar is a failsafe to revert behaviour"));
+
 enum EOverrideGCCollisionSetupForTraces
 {
 	GCCSFT_Property   = -1,  // Default: do what property says
@@ -4146,13 +4152,15 @@ bool FGeometryCollectionPhysicsProxy::PullNonInterpolatableDataFromSinglePhysics
 	}
 
 	// See comments in BufferPhysicsResults_Internal where IsRootBroken is set
+	// NOTE: We should never be returning to unbroken once broken although we aren't checking for that...
 	const int32 RootIndex = Parameters.InitialRootIndex;
-	if (RootIndex != INDEX_NONE)
+	if ((RootIndex != INDEX_NONE) && CurrentResults.IsRootBroken && GameThreadCollection.Active[RootIndex] && bGeometryCollectionUseRootBrokenFlag)
 	{
-		// NOTE: We never return to unbroken once broken
-		if (CurrentResults.IsRootBroken)
+		GameThreadCollection.Active[RootIndex] = false;
+		if (GTParticles[RootIndex] != nullptr)
 		{
-			GameThreadCollection.Active[RootIndex] = false;
+			GTParticles[RootIndex]->SetDisabled(true);
+			bIsCollectionDirty = true;
 		}
 	}
 
