Fix use-after-free in the cluster union manager: ensure the ImplicitUnion remains valid if the geometry is updated in FClusterUnionManager::FlushIncrementalGeometryOperations

#localization none

diff --git a/Engine/Source/Runtime/Experimental/Chaos/Private/Chaos/ClusterUnionManager.cpp b/Engine/Source/Runtime/Experimental/Chaos/Private/Chaos/ClusterUnionManager.cpp
--- a/Engine/Source/Runtime/Experimental/Chaos/Private/Chaos/ClusterUnionManager.cpp
+++ b/Engine/Source/Runtime/Experimental/Chaos/Private/Chaos/ClusterUnionManager.cpp
@@ -1064,6 +1064,9 @@ namespace Chaos
 			RemoveParticlesFromClusterUnionGeometry(ClusterUnion.InternalCluster, PendingGeometryRemovals, ClusterUnion.GeometryChildParticles);
 			ClusterUnion.Geometry = ClusterUnion.InternalCluster->GetGeometry();
 			ClusterUnion.ClearPendingGeometryOperations(EClusterUnionGeometryOperation::Remove);
+
+			ImplicitUnion = ClusterUnion.Geometry->template AsA<FImplicitObjectUnion>();
+			check(ImplicitUnion != nullptr);
 		}
 
 		check(ClusterUnion.Geometry != nullptr);
