The Acceleration structure Async task in Evolution will now Cancel if it did not have a chance to run before destruction
Solves:
Editor Soft-Locks at end of MRQ Render

diff --git a/Engine/Source/Runtime/Engine/Private/World.cpp b/Engine/Source/Runtime/Engine/Private/World.cpp
--- a/Engine/Source/Runtime/Engine/Private/World.cpp
+++ b/Engine/Source/Runtime/Engine/Private/World.cpp
@@ -1363,6 +1363,7 @@ bool UWorld::IsReadyForFinishDestroy()
 	{
 		if (PhysicsScene != nullptr)
 		{
+			PhysicsScene->KillSafeAsyncTasks();
 			if (PhysicsScene->AreAnyTasksPending())
 			{
 				return false;
diff --git a/Engine/Source/Runtime/Experimental/Chaos/Private/Chaos/PBDRigidsEvolution.cpp b/Engine/Source/Runtime/Experimental/Chaos/Private/Chaos/PBDRigidsEvolution.cpp
--- a/Engine/Source/Runtime/Experimental/Chaos/Private/Chaos/PBDRigidsEvolution.cpp
+++ b/Engine/Source/Runtime/Experimental/Chaos/Private/Chaos/PBDRigidsEvolution.cpp
@@ -216,6 +216,8 @@ namespace Chaos
 		, bIsSingleThreaded(InIsSingleThreaded)
 		, bCanStartAsyncTasks(true)
 		, LatestExternalTimestampConsumed_Internal(-1)
+		, bAccelerationStructureTaskStarted(nullptr)
+		, bAccelerationStructureTaskSignalKill(nullptr)
 		, SpatialCollectionFactory(new FDefaultCollectionFactory())
 	{
 		Particles.GetParticleHandles().AddArray(&PhysicsMaterials);
@@ -255,7 +257,9 @@ namespace Chaos
 		, FAccelerationStructure* InExternalAccelerationStructure
 		, bool InForceFullBuild
 		, bool InIsSingleThreaded
-		, bool InNeedsReset)
+		, bool InNeedsReset
+		, std::atomic<bool>** bOutStarted
+		, std::atomic<bool>** bOutKillTask)
 		: SpatialCollectionFactory(InSpatialCollectionFactory)
 		, SpatialAccelerationCache(InSpatialAccelerationCache)
 		, InternalStructure(InInternalAccelerationStructure)
@@ -263,8 +267,11 @@ namespace Chaos
 		, IsForceFullBuild(InForceFullBuild)
 		, bIsSingleThreaded(InIsSingleThreaded)
 		, bNeedsReset(InNeedsReset)
+		, bStarted(false)
+		, bKillTask(false)
 	{
-
+		*bOutStarted = &bStarted;
+		*bOutKillTask = &bKillTask;
 	}
 
 	TStatId FPBDRigidsEvolutionBase::FChaosAccelerationStructureTask::GetStatId()
@@ -510,6 +517,12 @@ namespace Chaos
 
 	void FPBDRigidsEvolutionBase::FChaosAccelerationStructureTask::DoTask(ENamedThreads::Type CurrentThread, const FGraphEventRef& MyCompletionGraphEvent)
 	{
+		bStarted = true;
+		if (bKillTask) // We can kill the task before it has really started
+		{
+			return;
+		}		
+
 		LLM_SCOPE(ELLMTag::ChaosAcceleration);
 
 #if !WITH_EDITOR
@@ -814,7 +827,7 @@ namespace Chaos
 			if (bCanStartAsyncTasks)
 			{
 				// we run the task for both starting a new accel structure as well as for the time-slicing
-				AccelerationStructureTaskComplete = TGraphTask<FChaosAccelerationStructureTask>::CreateTask().ConstructAndDispatchWhenReady(*SpatialCollectionFactory, SpatialAccelerationCache, AsyncInternalAcceleration, AsyncExternalAcceleration, ForceFullBuild, bIsSingleThreaded, bNeedsReset);
+				AccelerationStructureTaskComplete = TGraphTask<FChaosAccelerationStructureTask>::CreateTask().ConstructAndDispatchWhenReady(*SpatialCollectionFactory, SpatialAccelerationCache, AsyncInternalAcceleration, AsyncExternalAcceleration, ForceFullBuild, bIsSingleThreaded, bNeedsReset, &bAccelerationStructureTaskStarted, &bAccelerationStructureTaskSignalKill);
 			}
 		}
 		else
diff --git a/Engine/Source/Runtime/Experimental/Chaos/Public/Chaos/Framework/PhysicsSolverBase.h b/Engine/Source/Runtime/Experimental/Chaos/Public/Chaos/Framework/PhysicsSolverBase.h
--- a/Engine/Source/Runtime/Experimental/Chaos/Public/Chaos/Framework/PhysicsSolverBase.h
+++ b/Engine/Source/Runtime/Experimental/Chaos/Public/Chaos/Framework/PhysicsSolverBase.h
@@ -438,6 +438,10 @@ namespace Chaos
 			}
 		}
 
+		virtual void KillSafeAsyncTasks()
+		{
+		}
+
 		virtual bool AreAnyTasksPending() const
 		{
 			return false;
diff --git a/Engine/Source/Runtime/Experimental/Chaos/Public/Chaos/PBDRigidsEvolution.h b/Engine/Source/Runtime/Experimental/Chaos/Public/Chaos/PBDRigidsEvolution.h
--- a/Engine/Source/Runtime/Experimental/Chaos/Public/Chaos/PBDRigidsEvolution.h
+++ b/Engine/Source/Runtime/Experimental/Chaos/Public/Chaos/PBDRigidsEvolution.h
@@ -916,9 +916,18 @@ public:
 		return UniqueIndicesPendingRelease.Contains(UniqueIdx) || PendingReleaseIndices.Contains(UniqueIdx);
 	}
 
+	void KillSafeAsyncTasks()
+	{
+		if (AccelerationStructureTaskComplete.GetReference() && !AccelerationStructureTaskComplete->IsComplete() && bAccelerationStructureTaskSignalKill != nullptr)
+		{
+			*bAccelerationStructureTaskSignalKill = true;			
+		}
+	}
+
 	bool AreAnyTasksPending() const
 	{
-		return (AccelerationStructureTaskComplete.GetReference() && !AccelerationStructureTaskComplete->IsComplete());
+		return AccelerationStructureTaskComplete.GetReference() && !AccelerationStructureTaskComplete->IsComplete() && 
+			(bAccelerationStructureTaskSignalKill == nullptr || bAccelerationStructureTaskStarted  == nullptr || *bAccelerationStructureTaskSignalKill == false || *bAccelerationStructureTaskStarted == true);
 	}
 
 	void SetCanStartAsyncTasks(bool bInCanStartAsyncTasks)
@@ -1096,7 +1105,9 @@ protected:
 			, FAccelerationStructure* InExternalAccelerationStructure
 			, bool InForceFullBuild
 			, bool InIsSingleThreaded
-			, bool bNeedsReset);
+			, bool bNeedsReset
+			, std::atomic<bool>** bOutStarted
+			, std::atomic<bool>** bOutKillTask);
 		static FORCEINLINE TStatId GetStatId();
 		static FORCEINLINE ENamedThreads::Type GetDesiredThread();
 		static FORCEINLINE ESubsequentsMode::Type GetSubsequentsMode();
@@ -1109,11 +1120,15 @@ protected:
 		bool IsForceFullBuild;
 		bool bIsSingleThreaded;
 		bool bNeedsReset;
+		std::atomic<bool> bStarted;
+		std::atomic<bool> bKillTask;
 
 	private:
 		void UpdateStructure(FAccelerationStructure* AccelerationStructure, FAccelerationStructure* CopyToAccelerationStructure = nullptr);
 	};
 	FGraphEventRef AccelerationStructureTaskComplete;
+	std::atomic<bool>* bAccelerationStructureTaskStarted;
+	std::atomic<bool>* bAccelerationStructureTaskSignalKill;
 
 	TUniquePtr<ISpatialAccelerationCollectionFactory> SpatialCollectionFactory;
 
diff --git a/Engine/Source/Runtime/Experimental/Chaos/Public/PBDRigidsSolver.h b/Engine/Source/Runtime/Experimental/Chaos/Public/PBDRigidsSolver.h
--- a/Engine/Source/Runtime/Experimental/Chaos/Public/PBDRigidsSolver.h
+++ b/Engine/Source/Runtime/Experimental/Chaos/Public/PBDRigidsSolver.h
@@ -298,6 +298,11 @@ namespace Chaos
 		/** Apply a solver configuration to this solver, set externally by the owner of a solver (see UPhysicsSettings for world solver settings) */
 		CHAOS_API void ApplyConfig(const FChaosSolverConfiguration& InConfig);
 
+		virtual void KillSafeAsyncTasks() override
+		{
+			GetEvolution()->KillSafeAsyncTasks();
+		}
+
 		virtual bool AreAnyTasksPending() const override
 		{
 			if (IsPendingTasksComplete() == false || GetEvolution()->AreAnyTasksPending())
diff --git a/Engine/Source/Runtime/PhysicsCore/Private/ChaosScene.cpp b/Engine/Source/Runtime/PhysicsCore/Private/ChaosScene.cpp
--- a/Engine/Source/Runtime/PhysicsCore/Private/ChaosScene.cpp
+++ b/Engine/Source/Runtime/PhysicsCore/Private/ChaosScene.cpp
@@ -399,6 +399,15 @@ void FChaosScene::OnSyncBodies(Chaos::FPhysicsSolverBase* Solver)
 	Solver->PullPhysicsStateForEachDirtyProxy_External(Dispatcher);
 }
 
+void FChaosScene::KillSafeAsyncTasks()
+{
+	Chaos::FPBDRigidsSolver* Solver = GetSolver();
+	if (Solver )
+	{
+		Solver->KillSafeAsyncTasks();
+	}
+}
+
 bool FChaosScene::AreAnyTasksPending() const
 {
 	const Chaos::FPBDRigidsSolver* Solver = GetSolver();
diff --git a/Engine/Source/Runtime/PhysicsCore/Public/Chaos/ChaosScene.h b/Engine/Source/Runtime/PhysicsCore/Public/Chaos/ChaosScene.h
--- a/Engine/Source/Runtime/PhysicsCore/Public/Chaos/ChaosScene.h
+++ b/Engine/Source/Runtime/PhysicsCore/Public/Chaos/ChaosScene.h
@@ -144,6 +144,7 @@ public:
 	DECLARE_MULTICAST_DELEGATE_OneParam(FOnPhysScenePostTick,FChaosScene*);
 	FOnPhysScenePostTick OnPhysScenePostTick;
 
+	PHYSICSCORE_API void KillSafeAsyncTasks();
 	PHYSICSCORE_API bool AreAnyTasksPending() const;
 	PHYSICSCORE_API void BeginDestroy();
 	PHYSICSCORE_API bool IsCompletionEventComplete() const;
