Fixed ensures when geometry collection ISM pools were updated but their render transform wasn't flagged as dirty.

diff --git a/Engine/Source/Runtime/Experimental/GeometryCollectionEngine/Private/GeometryCollection/GeometryCollectionISMPoolComponent.cpp b/Engine/Source/Runtime/Experimental/GeometryCollectionEngine/Private/GeometryCollection/GeometryCollectionISMPoolComponent.cpp
--- a/Engine/Source/Runtime/Experimental/GeometryCollectionEngine/Private/GeometryCollection/GeometryCollectionISMPoolComponent.cpp
+++ b/Engine/Source/Runtime/Experimental/GeometryCollectionEngine/Private/GeometryCollection/GeometryCollectionISMPoolComponent.cpp
@@ -138,6 +138,7 @@ void FGeometryCollectionISM::InitISM(const FGeometryCollectionStaticMeshInstance
 
 		ISMComponent->SetComponentToWorld(TempTm);
 		ISMComponent->UpdateComponentTransform(EUpdateTransformFlags::None, ETeleportType::None);
+		ISMComponent->MarkRenderTransformDirty();
 	}
 	else
 	{
@@ -493,6 +494,8 @@ void FGeometryCollectionISMPool::UpdateAbsoluteTransforms(const FTransform& Base
 		const bool bReverseCulling = (GcIsm.MeshInstance.Desc.Flags & FISMComponentDescription::ReverseCulling) != 0;
 		check(GcIsm.MeshInstance.Desc.Position == FVector::ZeroVector);
 		
+		UInstancedStaticMeshComponent* Ism = GcIsm.ISMComponent;
+
 		if(bReverseCulling)
 		{
 			// As in InitISM we need to apply the inverted X scale for reverse culling.
@@ -502,14 +505,15 @@ void FGeometryCollectionISMPool::UpdateAbsoluteTransforms(const FTransform& Base
 			FTransform Flipped = BaseTransform;
 			Flipped.SetScale3D(BaseScale);
 			
-			GcIsm.ISMComponent->SetComponentToWorld(Flipped);
+			Ism->SetComponentToWorld(Flipped);
 		}
 		else
 		{
-			GcIsm.ISMComponent->SetComponentToWorld(BaseTransform);
+			Ism->SetComponentToWorld(BaseTransform);
 		}
 
-		GcIsm.ISMComponent->UpdateComponentTransform(UpdateTransformFlags | EUpdateTransformFlags::SkipPhysicsUpdate, Teleport);
+		Ism->UpdateComponentTransform(UpdateTransformFlags | EUpdateTransformFlags::SkipPhysicsUpdate, Teleport);
+		Ism->MarkRenderTransformDirty();
 	}
 }
 
