disabling actor collision on a GC should not transfer that collision setting to a parent cluster union particle

diff --git a/Engine/Source/Runtime/Experimental/Chaos/Private/PhysicsProxy/GeometryCollectionPhysicsProxy.cpp b/Engine/Source/Runtime/Experimental/Chaos/Private/PhysicsProxy/GeometryCollectionPhysicsProxy.cpp
--- a/Engine/Source/Runtime/Experimental/Chaos/Private/PhysicsProxy/GeometryCollectionPhysicsProxy.cpp
+++ b/Engine/Source/Runtime/Experimental/Chaos/Private/PhysicsProxy/GeometryCollectionPhysicsProxy.cpp
@@ -4429,7 +4429,7 @@ void FGeometryCollectionPhysicsProxy::SetFilterData_Internal(const FCollisionFil
 				}
 
 				FClusterHandle* ParentHandle = Handle->Parent();
-				if (ParentHandle && ParentHandle->InternalCluster())
+				if (ParentHandle && ParentHandle->InternalCluster() && ParentHandle->PhysicsProxy() == this)
 				{
 					if (!ProcessedInternalClusters.Contains(ParentHandle))
 					{
