properly cleanup the cluster union manager's PendingChildToParentUpdates and PendingParticlesToUndoChildToParentLock containers so we don't have garbage pointers lying around in containers that continue to grow uncontrolled

diff --git a/Engine/Source/Runtime/Experimental/Chaos/Private/Chaos/ClusterUnionManager.cpp b/Engine/Source/Runtime/Experimental/Chaos/Private/Chaos/ClusterUnionManager.cpp
--- a/Engine/Source/Runtime/Experimental/Chaos/Private/Chaos/ClusterUnionManager.cpp
+++ b/Engine/Source/Runtime/Experimental/Chaos/Private/Chaos/ClusterUnionManager.cpp
@@ -738,6 +738,16 @@ namespace Chaos
 			}
 		}
 
+		// Pending particles to undo child to parent lock is global across all cluster unions hence why it has to be done after.
+		for (FPBDRigidClusteredParticleHandle* ChildParticle : PendingParticlesToUndoChildToParentLock)
+		{
+			if (ChildParticle)
+			{
+				ChildParticle->SetChildToParentLocked(false);
+			}
+		}
+		PendingParticlesToUndoChildToParentLock.Empty();
+
 		DeferredClusterUnionsForUpdateProperties.Reset();
 	}
 
@@ -819,6 +829,25 @@ namespace Chaos
 			// and the index of its corresponding shape in the cluster union's shape array. There is currently
 			// an assumption that the two will always match each other.
 			Cluster->ChildParticles.RemoveAt(ParticleIndex);
+
+			if (FPBDRigidClusteredParticleHandle* ClusterParticle = Particle->CastToClustered())
+			{
+				// If we remove the particle we don't need to update its child to parent any longer.
+				if (FClusterUnionChildToParentUpdate* Update = PendingChildToParentUpdates.Find(ClusterParticle))
+				{
+					// Doing this here ensures that this particle is still part of this cluster union and hasn't yet been added into another cluster union.
+					if (Update->ClusterUnionIndex == ClusterIndex)
+					{
+						PendingChildToParentUpdates.Remove(ClusterParticle);
+					}
+				}
+
+				// TODO: We probably won't ever run into a situation where this actually does anything since
+				// this container gets cleared pretty shortly after things are added to it. But to be safe
+				// we probably need some sort of verification that we're removing a particle that isn't associated
+				// with some other cluster union.
+				PendingParticlesToUndoChildToParentLock.Remove(ClusterParticle);
+			}
 		}
 
 		MClustering.RemoveParticlesFromCluster(Cluster->InternalCluster, ParticleSet.Array());
@@ -943,16 +972,6 @@ namespace Chaos
 		}
 		// Build the convex optimizer if required
 		MClustering.BuildConvexOptimizer(ClusterUnion.InternalCluster);
-		
-		for (FPBDRigidClusteredParticleHandle* ChildParticle : PendingParticlesToUndoChildToParentLock)
-		{
-			if (ChildParticle)
-			{
-				ChildParticle->SetChildToParentLocked(false);
-			}
-		}
-
-		PendingParticlesToUndoChildToParentLock.Empty();
 	}
 
 	DECLARE_CYCLE_STAT(TEXT("FClusterUnionManager::FlushIncrementalConnectivityGraphOperations"), STAT_FlushIncrementalConnectivityGraphOperations, STATGROUP_Chaos);
@@ -1249,12 +1268,15 @@ namespace Chaos
 				return;
 			}
 
-			const int32 ChildIndex = ClusterUnion->ChildParticles.Find(Particle->CastToRigidParticle());
+			FPBDRigidParticleHandle* RigidParticle = Particle->CastToRigidParticle();
+			FPBDRigidClusteredParticleHandle* ClusteredParticle = Particle->CastToClustered();
+
+			const int32 ChildIndex = ClusterUnion->ChildParticles.Find(RigidParticle);
 			if (ChildIndex != INDEX_NONE && ClusterUnion->InternalCluster)
 			{
-				if (FPBDRigidClusteredParticleHandle* ChildHandle = ClusterUnion->ChildParticles[ChildIndex]->CastToClustered())
+				if (ClusteredParticle)
 				{
-					if (const FClusterUnionChildToParentUpdate* Update = PendingChildToParentUpdates.Find(ChildHandle))
+					if (const FClusterUnionChildToParentUpdate* Update = PendingChildToParentUpdates.Find(ClusteredParticle))
 					{
 						if (Update->ClusterUnionIndex != ClusterIndex)
 						{
@@ -1262,25 +1284,24 @@ namespace Chaos
 						}
 
 						const FRigidTransform3 ChildToParent = Update->ChildToParent;
-						ChildHandle->SetChildToParent(ChildToParent);
+						ClusteredParticle->SetChildToParent(ChildToParent);
 
-						if (!Update->bLock && !ChildHandle->IsChildToParentLocked())
+						if (!Update->bLock && !ClusteredParticle->IsChildToParentLocked())
 						{
-							PendingParticlesToUndoChildToParentLock.Add(ChildHandle);
+							PendingParticlesToUndoChildToParentLock.Add(ClusteredParticle);
 						}
-						ChildHandle->SetChildToParentLocked(true);
-						PendingChildToParentUpdates.Remove(ChildHandle);
+						ClusteredParticle->SetChildToParentLocked(true);
 
 						// Update the child's world transform to be consistent with its ChildToParent transform
 						const FPBDRigidClusteredParticleHandle* ParentHandle = ClusterUnion->InternalCluster;
 						const FRigidTransform3 ParticleToWorld = ChildToParent * FRigidTransform3(ParentHandle->X(), ParentHandle->R());
-						ChildHandle->SetX(ParticleToWorld.GetTranslation());
-						ChildHandle->SetP(ParticleToWorld.GetTranslation());
-						ChildHandle->SetR(ParticleToWorld.GetRotation());
-						ChildHandle->SetQ(ParticleToWorld.GetRotation());
+						ClusteredParticle->SetX(ParticleToWorld.GetTranslation());
+						ClusteredParticle->SetP(ParticleToWorld.GetTranslation());
+						ClusteredParticle->SetR(ParticleToWorld.GetRotation());
+						ClusteredParticle->SetQ(ParticleToWorld.GetRotation());
 
 						// We need to mark the child handle to be dirty so that its proxy gets a chance to sync back to the GC 
-						MEvolution.GetParticles().MarkTransientDirtyParticle(ChildHandle);
+						MEvolution.GetParticles().MarkTransientDirtyParticle(ClusteredParticle);
 
 						// A child to parent update needs to remove *and* add to the connectivity graph (in that order) since
 						// the child to parent update might move the node so far away as to make the old connectivity edges incorrect.
@@ -1294,6 +1315,8 @@ namespace Chaos
 					}
 				}
 			}
+
+			PendingChildToParentUpdates.Remove(ClusteredParticle);
 		}
 
 		if (bMadeChanges)
