[Chaos] Making RefreshCustomRenderer on Geometry Collection Component public so other objects can force an update if needed.

Also updated a class to make use of this method until we find a better solution.

diff --git a/Engine/Source/Runtime/Experimental/GeometryCollectionEngine/Public/GeometryCollection/GeometryCollectionComponent.h b/Engine/Source/Runtime/Experimental/GeometryCollectionEngine/Public/GeometryCollection/GeometryCollectionComponent.h
--- a/Engine/Source/Runtime/Experimental/GeometryCollectionEngine/Public/GeometryCollection/GeometryCollectionComponent.h
+++ b/Engine/Source/Runtime/Experimental/GeometryCollectionEngine/Public/GeometryCollection/GeometryCollectionComponent.h
@@ -1732,8 +1732,13 @@ private:
 
 	GEOMETRYCOLLECTIONENGINE_API void RegisterCustomRenderer();
 	GEOMETRYCOLLECTIONENGINE_API void UnregisterCustomRenderer();
+
+public:
+	/** Updates the custom renderer to the reflect the current state of the Geometry Collection */
 	GEOMETRYCOLLECTIONENGINE_API void RefreshCustomRenderer();
 
+private:
+
 	GEOMETRYCOLLECTIONENGINE_API void IncrementSleepTimer(float DeltaTime);
 	GEOMETRYCOLLECTIONENGINE_API void IncrementBreakTimer(float DeltaTime);
 	GEOMETRYCOLLECTIONENGINE_API bool CalculateInnerSphere(int32 TransformIndex, UE::Math::TSphere<double>& SphereOut) const;
