Get rid of some non-load-bearing physics data on both clients and server.

- Chaos indices are only used for the debug drawing.

diff --git a/Engine/Source/Runtime/Engine/Private/PhysicsEngine/BodySetup.cpp b/Engine/Source/Runtime/Engine/Private/PhysicsEngine/BodySetup.cpp
--- a/Engine/Source/Runtime/Engine/Private/PhysicsEngine/BodySetup.cpp
+++ b/Engine/Source/Runtime/Engine/Private/PhysicsEngine/BodySetup.cpp
@@ -1775,10 +1775,15 @@ void FKConvexElem::ResetChaosConvexMesh()
 
 ENGINE_API void FKConvexElem::ComputeChaosConvexIndices(bool bForceCompute)
 {
+	// these indices are not needed for simulation, but only get used for debug visualization
+#if (UE_ENABLE_DEBUG_DRAWING && !USE_NULL_RHI)	// FIXME: UE_ENABLE_DEBUG_DRAWING needs to be fixed to incorportate !USE_NULL_RHI
 	if (bForceCompute || IndexData.Num() == 0)
 	{
 		IndexData = GetChaosConvexIndices();
 	}
+#else
+	IndexData.Reset();
+#endif
 }
 
 TArray<int32> FKConvexElem::GetChaosConvexIndices() const
diff --git a/Engine/Source/Runtime/Experimental/Chaos/Public/Chaos/TriangleMeshImplicitObject.h b/Engine/Source/Runtime/Experimental/Chaos/Public/Chaos/TriangleMeshImplicitObject.h
--- a/Engine/Source/Runtime/Experimental/Chaos/Public/Chaos/TriangleMeshImplicitObject.h
+++ b/Engine/Source/Runtime/Experimental/Chaos/Public/Chaos/TriangleMeshImplicitObject.h
@@ -668,8 +668,12 @@ namespace Chaos
 				// unless it is required for gameplay, as we are wasting disk space.
 				if (Ar.IsLoading())
 				{
-					ExternalFaceIndexMap = MakeUnique<TArray<int32>>(TArray<int32>());
-					Ar << *ExternalFaceIndexMap;
+					TUniquePtr<TArray<int32>> ExternalFaceIndexMapTemp = MakeUnique<TArray<int32>>(TArray<int32>());
+					Ar << *ExternalFaceIndexMapTemp;
+					if (!ExternalFaceIndexMapTemp->IsEmpty())
+					{
+						ExternalFaceIndexMap = MoveTemp(ExternalFaceIndexMapTemp);
+					}
 				}
 				else
 				{
