Ias: Fix broken client side cache telemetry.
Analytics: Report the existing IasCacheCachedBytes & IasCacheMaxBytes and the new IasCacheUsagePercent as absolute values.
Analytics: Add reporting of the existing IasCacheGetCount IasCachePutCount stats as delta values.
Analytics: Add a new stat IasCacheWriteBytes for actual bytes written as a delta value.
Add the new CacheWrittenBytes to trace/csv stats and periodic logging, and updated logging to KiB from MiB.

#rnx
#tests frontend, replays and joining server with and without ias

#localization none

diff --git a/Engine/Source/Runtime/Experimental/IoStoreOnDemand/Private/JournaledCache.cpp b/Engine/Source/Runtime/Experimental/IoStoreOnDemand/Private/JournaledCache.cpp
--- a/Engine/Source/Runtime/Experimental/IoStoreOnDemand/Private/JournaledCache.cpp
+++ b/Engine/Source/Runtime/Experimental/IoStoreOnDemand/Private/JournaledCache.cpp
@@ -731,6 +731,8 @@ void FDiskCache::ClosePhrase(FDiskPhrase&& Phrase)
 		return;
 	}
 
+	FOnDemandIoBackendStats::Get()->OnCacheWriteBytes(WriteSize);
+
 	{
 		FWriteScopeLock _(Lock);
 		Prune(DataCursor, WriteSize);
diff --git a/Engine/Source/Runtime/Experimental/IoStoreOnDemand/Private/Statistics.cpp b/Engine/Source/Runtime/Experimental/IoStoreOnDemand/Private/Statistics.cpp
--- a/Engine/Source/Runtime/Experimental/IoStoreOnDemand/Private/Statistics.cpp
+++ b/Engine/Source/Runtime/Experimental/IoStoreOnDemand/Private/Statistics.cpp
@@ -8,6 +8,7 @@
 #include "HAL/PlatformTime.h"
 #include "IO/IoStoreOnDemand.h"
 #include "Internationalization/Internationalization.h"
+#include "Math/UnrealMathUtility.h"
 #include "Misc/CoreDelegates.h"
 
 LLM_DEFINE_TAG(Ias);
@@ -256,6 +257,7 @@ FCounterAtomicInt		GCachePutCount(TEXT("Ias/CachePutCount"), TraceCounterDisplay
 FCounterAtomicInt		GCachePutExistingCount(TEXT("Ias/CachePutExistingCount"), TraceCounterDisplayHint_None);
 FCounterAtomicInt		GCachePutRejectCount(TEXT("Ias/CachePutRejectCount"), TraceCounterDisplayHint_None);
 FCounterAtomicInt		GCacheCachedBytes(TEXT("Ias/CacheCachedBytes"), TraceCounterDisplayHint_Memory);
+FCounterAtomicInt		GCacheWrittenBytes(TEXT("Ias/CacheWrittenBytes"), TraceCounterDisplayHint_Memory);
 int64					GCacheMaxBytes = 0;
 FCounterAtomicInt		GCachePendingBytes(TEXT("Ias/CachePendingBytes"), TraceCounterDisplayHint_Memory);
 FCounterAtomicInt		GCacheReadBytes(TEXT("Ias/CacheReadBytes"), TraceCounterDisplayHint_Memory);
@@ -318,6 +320,7 @@ CSV_DEFINE_STAT(Ias, CachePutCount);
 CSV_DEFINE_STAT(Ias, CachePutExistingCount);
 CSV_DEFINE_STAT(Ias, CachePutRejectCount);
 CSV_DEFINE_STAT(Ias, CacheCachedMB);
+CSV_DEFINE_STAT(Ias, CacheWrittenMB);
 CSV_DEFINE_STAT(Ias, CacheReadMB);
 CSV_DEFINE_STAT(Ias, CacheRejectedMB);
 // http stat totals
@@ -347,17 +350,19 @@ FOnDemandIoBackendStats::FOnDemandIoBackendStats()
 		int32 CPutCount = int32(GCachePutCount.Get());
 		int32 CPutExistingCount = int32(GCachePutExistingCount.Get());
 		int32 CPutRejectCount = int32(GCachePutRejectCount.Get());
-		int32 CCachedMB = BytesToApproxMB(GCacheCachedBytes.Get());
-		int32 CReadMB = BytesToApproxMB(GCacheReadBytes.Get());
-		int32 CRejectedMB = BytesToApproxMB(GCacheRejectBytes.Get());
+		int32 CCachedKiB = BytesToApproxKB(GCacheCachedBytes.Get());
+		int32 CWrittenKiB = BytesToApproxKB(GCacheWrittenBytes.Get());
+		int32 CReadKiB = BytesToApproxKB(GCacheReadBytes.Get());
+		int32 CRejectedKiB = BytesToApproxKB(GCacheRejectBytes.Get());
 		CSV_CUSTOM_STAT_DEFINED(CacheGetCount, CGetCount, ECsvCustomStatOp::Set);
 		CSV_CUSTOM_STAT_DEFINED(CacheErrorCount, CErrorCount, ECsvCustomStatOp::Set);
 		CSV_CUSTOM_STAT_DEFINED(CachePutCount, CPutCount, ECsvCustomStatOp::Set);
 		CSV_CUSTOM_STAT_DEFINED(CachePutExistingCount, CPutExistingCount, ECsvCustomStatOp::Set);
 		CSV_CUSTOM_STAT_DEFINED(CachePutRejectCount, CPutRejectCount, ECsvCustomStatOp::Set);
-		CSV_CUSTOM_STAT_DEFINED(CacheCachedMB, CCachedMB, ECsvCustomStatOp::Set);
-		CSV_CUSTOM_STAT_DEFINED(CacheReadMB, CReadMB, ECsvCustomStatOp::Set);
-		CSV_CUSTOM_STAT_DEFINED(CacheRejectedMB, CRejectedMB, ECsvCustomStatOp::Set);
+		CSV_CUSTOM_STAT_DEFINED(CacheCachedMB, CCachedKiB >> 10, ECsvCustomStatOp::Set);
+		CSV_CUSTOM_STAT_DEFINED(CacheWrittenMB, CWrittenKiB >> 10, ECsvCustomStatOp::Set);
+		CSV_CUSTOM_STAT_DEFINED(CacheReadMB, CReadKiB >> 10, ECsvCustomStatOp::Set);
+		CSV_CUSTOM_STAT_DEFINED(CacheRejectedMB, CRejectedKiB >> 10, ECsvCustomStatOp::Set);
 
 		// http stat totals
 		int32 HGetCount = int32(GHttpGetCount.Get());
@@ -365,7 +370,7 @@ FOnDemandIoBackendStats::FOnDemandIoBackendStats()
 		int32 HCancelCount = int32(GHttpCancelCount.Get());
 		int32 HErrorCount = int32(GHttpErrorCount.Get());
 		int32 HPendingCount = int32(GHttpPendingCount.Get());
-		int32 HDownloadedMB = BytesToApproxMB(GHttpDownloadedBytes.Get());
+		int32 HDownloadedKiB = BytesToApproxKB(GHttpDownloadedBytes.Get());
 		int32 HBandwidthMpbs = int32(GHttpBandwidthMpbs.Get());
 		int32 HDurationMsAvg = int32(GHttpDurationMsAvg.Get());
 		int32 HDurationMsMax = int32(GHttpDurationMsMax.Get());
@@ -373,7 +378,7 @@ FOnDemandIoBackendStats::FOnDemandIoBackendStats()
 		CSV_CUSTOM_STAT_DEFINED(HttpCancelCount, HCancelCount, ECsvCustomStatOp::Set);
 		CSV_CUSTOM_STAT_DEFINED(HttpErrorCount, HErrorCount, ECsvCustomStatOp::Set);
 		CSV_CUSTOM_STAT_DEFINED(HttpPendingCount, HPendingCount, ECsvCustomStatOp::Set);
-		CSV_CUSTOM_STAT_DEFINED(HttpDownloadedMB, HDownloadedMB, ECsvCustomStatOp::Set);
+		CSV_CUSTOM_STAT_DEFINED(HttpDownloadedMB, HDownloadedKiB >> 10, ECsvCustomStatOp::Set);
 		CSV_CUSTOM_STAT_DEFINED(HttpBandwidthMpbs, HBandwidthMpbs, ECsvCustomStatOp::Set);
 		CSV_CUSTOM_STAT_DEFINED(HttpDurationMsAvg, HDurationMsAvg, ECsvCustomStatOp::Set);
 		CSV_CUSTOM_STAT_DEFINED(HttpDurationMsMax, HDurationMsMax, ECsvCustomStatOp::Set);
@@ -383,10 +388,10 @@ FOnDemandIoBackendStats::FOnDemandIoBackendStats()
 			static double LastLogTime = 0.0;
 			if (double Time = FPlatformTime::Seconds(); Time - LastLogTime > (double)GIasStatisticsLogInterval)
 			{
-				UE_LOG(LogIas, Log, TEXT("CacheStats: CachedMB=%d, RejectedMB=%d, ReadMB=%d, Get=%d, Error=%d, Put=%d, PutReject=%d, PutExisting=%d"),
-					CCachedMB, CRejectedMB, CReadMB, CGetCount, CErrorCount, CPutCount, CPutRejectCount, CPutExistingCount);
-				UE_LOG(LogIas, Log, TEXT("HttpStats: DownloadedMB=%d, Get=%d, Retry=%d, Cancel=%d, Error=%d, CurPending=%d, CurDurationMsAvg=%d, CurDurationMsMax=%d"),
-					HDownloadedMB, HGetCount, HRetryCount, HCancelCount, HErrorCount, HPendingCount, HDurationMsAvg, HDurationMsMax);
+				UE_LOG(LogIas, Log, TEXT("CacheStats: CachedKiB=%d, WrittenKiB=%d, ReadKiB=%d, RejectedKiB=%d, Get=%d, Error=%d, Put=%d, PutReject=%d, PutExisting=%d"),
+					CCachedKiB, CWrittenKiB, CReadKiB, CRejectedKiB, CGetCount, CErrorCount, CPutCount, CPutRejectCount, CPutExistingCount);
+				UE_LOG(LogIas, Log, TEXT("HttpStats: DownloadedKiB=%d, Get=%d, Retry=%d, Cancel=%d, Error=%d, CurPending=%d, CurDurationMsAvg=%d, CurDurationMsMax=%d"),
+					HDownloadedKiB, HGetCount, HRetryCount, HCancelCount, HErrorCount, HPendingCount, HDurationMsAvg, HDurationMsMax);
 				LastLogTime = Time;
 			}
 		}
@@ -459,13 +464,20 @@ void FOnDemandIoBackendStats::ReportAnalytics(TArray<FAnalyticsEventAttribute>&
 
 	if (GIasReportCacheAnalyticsEnabled)
 	{
+		int64 CacheTotalCount = GCacheGetCount.Get() + GCachePutCount.Get();
+		float CacheUsagePercent = GCacheMaxBytes > 0 ? FMath::Clamp(0.f, (100.f*float(GCacheCachedBytes.Get())) / float(GCacheMaxBytes), 100.f) : 0.f;
 		AppendAnalyticsEventAttributeArray(OutAnalyticsArray
 
+			,TRACK_DELTA("IasCacheTotalCount", CacheTotalCount)
 			,TRACK_DELTA("IasCacheErrorCount", GCacheErrorCount.Get())
+			,TRACK_DELTA("IasCacheGetCount", GCacheGetCount.Get())
+			,TRACK_DELTA("IasCachePutCount", GCachePutCount.Get())
 
-			,TRACK_DELTA("IasCacheCachedBytes", GCacheCachedBytes.Get())
-			,TRACK_DELTA("IasCacheMaxBytes", GCacheMaxBytes)
+			,TEXT("IasCacheCachedBytes"), GCacheCachedBytes.Get()
+			,TEXT("IasCacheMaxBytes"), GCacheMaxBytes
+			,TEXT("IasCacheUsagePercent"), CacheUsagePercent
 
+			,TRACK_DELTA("IasCacheWriteBytes", GCacheWrittenBytes.Get())
 			,TRACK_DELTA("IasCacheReadBytes", GCacheReadBytes.Get())
 			,TRACK_DELTA("IasCacheRejectBytes", GCacheRejectBytes.Get())
 		);
@@ -538,6 +550,11 @@ void FOnDemandIoBackendStats::OnCachePersistedBytes(uint64 TotalSize)
 	GCacheCachedBytes.Set(TotalSize);
 }
 
+void FOnDemandIoBackendStats::OnCacheWriteBytes(uint64 WriteSize)
+{
+	GCacheWrittenBytes.Add(WriteSize);
+}
+
 void FOnDemandIoBackendStats::OnCacheSetMaxBytes(uint64 TotalSize)
 {
 	GCacheMaxBytes = TotalSize;
diff --git a/Engine/Source/Runtime/Experimental/IoStoreOnDemand/Private/Statistics.h b/Engine/Source/Runtime/Experimental/IoStoreOnDemand/Private/Statistics.h
--- a/Engine/Source/Runtime/Experimental/IoStoreOnDemand/Private/Statistics.h
+++ b/Engine/Source/Runtime/Experimental/IoStoreOnDemand/Private/Statistics.h
@@ -47,6 +47,7 @@ public:
 	void OnCachePutReject(uint64 DataSize) IAS_STATISTICS_IMPL()
 	void OnCachePendingBytes(uint64 TotalSize) IAS_STATISTICS_IMPL()
 	void OnCachePersistedBytes(uint64 TotalSize) IAS_STATISTICS_IMPL()
+	void OnCacheWriteBytes(uint64 WriteSize) IAS_STATISTICS_IMPL()
 	void OnCacheSetMaxBytes(uint64 TotalSize) IAS_STATISTICS_IMPL()
 
 	void OnHttpEnqueue() IAS_STATISTICS_IMPL()
