Fix updated cvar
A cvar was wrongly updated, so it cancel memory optimization
#tests Switch client, Win64 client

diff --git a/Engine/Source/Runtime/Experimental/Chaos/Private/PhysicsProxy/GeometryCollectionPhysicsProxy.cpp b/Engine/Source/Runtime/Experimental/Chaos/Private/PhysicsProxy/GeometryCollectionPhysicsProxy.cpp
--- a/Engine/Source/Runtime/Experimental/Chaos/Private/PhysicsProxy/GeometryCollectionPhysicsProxy.cpp
+++ b/Engine/Source/Runtime/Experimental/Chaos/Private/PhysicsProxy/GeometryCollectionPhysicsProxy.cpp
@@ -539,7 +539,7 @@ FAutoConsoleVariableRef CVarCreateGTParticleForChildren(TEXT("p.gc.CreateGTParti
 
 bool bRemoveImplicitsInDynamicCollections = false;
 FAutoConsoleVariableRef CVarbRemoveImplicitsInDynamicCollections(TEXT("p.gc.RemoveImplicitsInDynamicCollections"), 
-	bBuildGeometryForChildren, TEXT("This cvar has an impact only if geometry are not added for children. It removes implicits from the Dynamic Collections, and recreate then from the rest collection. \
+	bRemoveImplicitsInDynamicCollections, TEXT("This cvar has an impact only if geometry are not added for children. It removes implicits from the Dynamic Collections, and recreate then from the rest collection. \
 										Using this cvar could have an impact if geometry are updated from the dynamic collection on the GT, then those changes won't be ported to the PT."));
 
 
