Fix bugs in Cloth TriangleMeshCollisions caused by holding onto pointers after Arrays are resized.

diff --git a/Engine/Source/Runtime/Experimental/Chaos/Private/Chaos/PBDTriangleMeshCollisions.cpp b/Engine/Source/Runtime/Experimental/Chaos/Private/Chaos/PBDTriangleMeshCollisions.cpp
--- a/Engine/Source/Runtime/Experimental/Chaos/Private/Chaos/PBDTriangleMeshCollisions.cpp
+++ b/Engine/Source/Runtime/Experimental/Chaos/Private/Chaos/PBDTriangleMeshCollisions.cpp
@@ -274,10 +274,10 @@ namespace GIA
 
 				FIntersectionContour* FaceContour = &ContourPair.ColorContours[InitialFaceContourIndex];
 				FIntersectionContour* EdgeContour = &ContourPair.ColorContours[InitialEdgeContourIndex];
-				ContourPair.ContourPointCurves[InitialFaceContourIndex].Add(ContourPoints.Num());
-				TArray<FPBDTriangleMeshCollisions::FBarycentricPoint>* FaceContourPoints = &ContourPoints.AddDefaulted_GetRef();
-				ContourPair.ContourPointCurves[InitialEdgeContourIndex].Add(ContourPoints.Num());
-				TArray<FPBDTriangleMeshCollisions::FBarycentricPoint>* EdgeContourPoints = &ContourPoints.AddDefaulted_GetRef();
+				int32 FaceContourPointsIndex = ContourPoints.AddDefaulted();
+				int32 EdgeContourPointsIndex = ContourPoints.AddDefaulted();
+				ContourPair.ContourPointCurves[InitialFaceContourIndex].Add(FaceContourPointsIndex);
+				ContourPair.ContourPointCurves[InitialEdgeContourIndex].Add(EdgeContourPointsIndex);
 
 				const int32 FirstEdgeFace = EdgeToFaces[FirstIntersection.EdgeIndex][0];
 				check(FirstEdgeFace != -1); // Each Edge should be connected to at least one face
@@ -301,8 +301,8 @@ namespace GIA
 					const TVec3<int32>& FaceVertices = Elements[CurrIntersection.FaceIndex];
 					const TVec3<int32>& EdgeFaceVertices = Elements[EdgeFace];
 
-					FaceContourPoints->Add({ CurrIntersection.FaceCoordinate, FaceVertices });
-					EdgeContourPoints->Add({ {CurrIntersection.EdgeCoordinate, 0.f}, {EdgeFaceVertices[0], EdgeFaceVertices[1], EdgeFaceVertices[1]} });
+					ContourPoints[FaceContourPointsIndex].Add({CurrIntersection.FaceCoordinate, FaceVertices});
+					ContourPoints[EdgeContourPointsIndex].Add({ {CurrIntersection.EdgeCoordinate, 0.f}, {EdgeFaceVertices[0], EdgeFaceVertices[1], EdgeFaceVertices[1]} });
 
 					// Walk the intersection contour. Next point in contour is either
 					// 1) Loop vertex
@@ -353,10 +353,10 @@ namespace GIA
 						EdgeContour = &ContourPair.ColorContours[InitialEdgeContourIndex];
 
 						// ContourPoints are currently just used for debug drawing. Start new contours for reverse section
-						ContourPair.ContourPointCurves[InitialFaceContourIndex].Add(ContourPoints.Num());
-						FaceContourPoints = &ContourPoints.AddDefaulted_GetRef();
-						ContourPair.ContourPointCurves[InitialEdgeContourIndex].Add(ContourPoints.Num());
-						EdgeContourPoints = &ContourPoints.AddDefaulted_GetRef();
+						FaceContourPointsIndex = ContourPoints.AddDefaulted();
+						EdgeContourPointsIndex = ContourPoints.AddDefaulted();
+						ContourPair.ContourPointCurves[InitialFaceContourIndex].Add(FaceContourPointsIndex);
+						ContourPair.ContourPointCurves[InitialEdgeContourIndex].Add(EdgeContourPointsIndex);
 
 						FaceSection = &FaceContour->Contour[0];
 						EdgeFaceSection = &EdgeContour->Contour.Add_GetRef(FIntersectionContourTriangleSection(EdgeFace));
@@ -373,9 +373,14 @@ namespace GIA
 
 						FIntersectionContourTriangleSection& FirstFaceSection = ContourPair.ColorContours[InitialFaceContourIndex].Contour[0];
 
+						if (!ensure(FaceContour == &ContourPair.ColorContours[InitialFaceContourIndex]) ||
+							!ensure(FaceSection->TriangleIndex == FirstFaceSection.TriangleIndex))
+						{
+							// Somehow our two contours have crossed. Just bail on this contour.
+							break;
+						}
+
 						// Merge current FaceSection and FirstFaceSection (if they're not already the same)
-						check(FaceContour == &ContourPair.ColorContours[InitialFaceContourIndex]);
-						check(FaceSection->TriangleIndex == FirstFaceSection.TriangleIndex);
 						if (FaceSection != &FirstFaceSection)
 						{
 							check(FaceSection->CrossingEdgeLocalIndex[0] != INDEX_NONE);
@@ -387,8 +392,8 @@ namespace GIA
 						ContourPair.ClosedStatus = FIntersectionContourPair::EClosedStatus::SimpleClosed;
 
 						// Repeat first point in contour points for ease of drawing closed loop
-						FaceContourPoints->Add(FPBDTriangleMeshCollisions::FBarycentricPoint((*FaceContourPoints)[0]));
-						EdgeContourPoints->Add(FPBDTriangleMeshCollisions::FBarycentricPoint((*EdgeContourPoints)[0]));
+						ContourPoints[FaceContourPointsIndex].Add(FPBDTriangleMeshCollisions::FBarycentricPoint(ContourPoints[FaceContourPointsIndex][0]));
+						ContourPoints[EdgeContourPointsIndex].Add(FPBDTriangleMeshCollisions::FBarycentricPoint(ContourPoints[EdgeContourPointsIndex][0]));
 						break;
 					}
 
@@ -400,8 +405,13 @@ namespace GIA
 
 						// Merge current EdgeFaceSection with FirstFaceSection
 						FIntersectionContourTriangleSection& FirstFaceSection = ContourPair.ColorContours[InitialFaceContourIndex].Contour[0];
-						check(EdgeContour == &ContourPair.ColorContours[InitialFaceContourIndex]);
-						check(EdgeFaceSection->TriangleIndex == FirstFaceSection.TriangleIndex);
+						if (!ensure(EdgeContour == &ContourPair.ColorContours[InitialFaceContourIndex]) ||
+							!ensure(EdgeFaceSection->TriangleIndex == FirstFaceSection.TriangleIndex))
+						{
+							// Somehow our two contours have crossed. Just bail on this contour.
+							break;
+						}
+
 						check(EdgeFaceSection->CrossingEdgeLocalIndex[0] != INDEX_NONE);
 						FirstFaceSection.CrossingEdgeLocalIndex[0] = EdgeFaceSection->CrossingEdgeLocalIndex[0];
 
@@ -411,8 +421,8 @@ namespace GIA
 						ContourPair.ClosedStatus = FIntersectionContourPair::EClosedStatus::SimpleClosed;
 
 						// Repeat first point in contour points for ease of drawing closed loop
-						FaceContourPoints->Add(FPBDTriangleMeshCollisions::FBarycentricPoint((*FaceContourPoints)[0]));
-						EdgeContourPoints->Add(FPBDTriangleMeshCollisions::FBarycentricPoint((*EdgeContourPoints)[0]));
+						ContourPoints[FaceContourPointsIndex].Add(FPBDTriangleMeshCollisions::FBarycentricPoint(ContourPoints[FaceContourPointsIndex][0]));
+						ContourPoints[EdgeContourPointsIndex].Add(FPBDTriangleMeshCollisions::FBarycentricPoint(ContourPoints[EdgeContourPointsIndex][0]));
 						break;
 					}
 
@@ -472,10 +482,10 @@ namespace GIA
 							EdgeContour = &ContourPair.ColorContours[InitialEdgeContourIndex];
 
 							// ContourPoints are currently just used for debug drawing. Start new contours for reverse section
-							ContourPair.ContourPointCurves[InitialFaceContourIndex].Add(ContourPoints.Num());
-							FaceContourPoints = &ContourPoints.AddDefaulted_GetRef();
-							ContourPair.ContourPointCurves[InitialEdgeContourIndex].Add(ContourPoints.Num());
-							EdgeContourPoints = &ContourPoints.AddDefaulted_GetRef();
+							FaceContourPointsIndex = ContourPoints.AddDefaulted();
+							EdgeContourPointsIndex = ContourPoints.AddDefaulted();
+							ContourPair.ContourPointCurves[InitialFaceContourIndex].Add(FaceContourPointsIndex);
+							ContourPair.ContourPointCurves[InitialEdgeContourIndex].Add(EdgeContourPointsIndex);
 
 							FaceSection = &ContourPair.ColorContours[InitialFaceContourIndex].Contour[0];
 							EdgeFaceSection = &EdgeContour->Contour.Add_GetRef(FIntersectionContourTriangleSection(EdgeFace));
@@ -548,10 +558,10 @@ namespace GIA
 							EdgeContour = &ContourPair.ColorContours[1];
 
 							// ContourPoints are currently just used for debug drawing. Start new contours for reverse section
-							ContourPair.ContourPointCurves[InitialFaceContourIndex].Add(ContourPoints.Num());
-							FaceContourPoints = &ContourPoints.AddDefaulted_GetRef();
-							ContourPair.ContourPointCurves[InitialEdgeContourIndex].Add(ContourPoints.Num());
-							EdgeContourPoints = &ContourPoints.AddDefaulted_GetRef();
+							FaceContourPointsIndex = ContourPoints.AddDefaulted();
+							EdgeContourPointsIndex = ContourPoints.AddDefaulted();
+							ContourPair.ContourPointCurves[InitialFaceContourIndex].Add(FaceContourPointsIndex);
+							ContourPair.ContourPointCurves[InitialEdgeContourIndex].Add(EdgeContourPointsIndex);
 
 							FaceSection = &ContourPair.ColorContours[InitialFaceContourIndex].Contour[0];
 							EdgeFaceSection = &EdgeContour->Contour.Add_GetRef(FIntersectionContourTriangleSection(EdgeFace));
@@ -564,12 +574,12 @@ namespace GIA
 						CurrIntersection = NextIntersection;
 
 						// FaceContour and EdgeContour swap
-						FIntersectionContour* TmpContourSwap = EdgeContour;
+						FIntersectionContour* const TmpContourSwap = EdgeContour;
 						EdgeContour = FaceContour;
 						FaceContour = TmpContourSwap;
-						TArray<FPBDTriangleMeshCollisions::FBarycentricPoint>* TmpPointsSwap = EdgeContourPoints;
-						EdgeContourPoints = FaceContourPoints;
-						FaceContourPoints = TmpPointsSwap;
+						const int32 TmpContourPointsSwap = EdgeContourPointsIndex;
+						EdgeContourPointsIndex = FaceContourPointsIndex;
+						FaceContourPointsIndex = TmpContourPointsSwap;
 
 						FaceSection = EdgeFaceSection;
 
