Fix rendering offsets with cluster unions
- This is caused by the  transform not being properly defined when receiving a ComponentAdded event
- Override ChildToParentUpdated method to update the rendering transform

diff --git a/Engine/Source/Runtime/Engine/Classes/PhysicsEngine/ClusterUnionComponent.h b/Engine/Source/Runtime/Engine/Classes/PhysicsEngine/ClusterUnionComponent.h
--- a/Engine/Source/Runtime/Engine/Classes/PhysicsEngine/ClusterUnionComponent.h
+++ b/Engine/Source/Runtime/Engine/Classes/PhysicsEngine/ClusterUnionComponent.h
@@ -356,6 +356,7 @@ protected:
 	Chaos::FClusterUnionPhysicsProxy* GetPhysicsProxy() { return PhysicsProxy; }
 
 	ENGINE_API virtual void SetRigidState(Chaos::EObjectStateType ObjectState);
+	ENGINE_API virtual void OnChildToParentUpdated(UPrimitiveComponent* ChangedComponent, const FLocalBonesToTransformMap& PerBoneChildToParent) {}
 
 	// We need to keep track of the mapping of primitive components to physics objects.
 	// This way we know the right physics objects to pass when removing the component (because
diff --git a/Engine/Source/Runtime/Engine/Private/PhysicsEngine/ClusterUnionComponent.cpp b/Engine/Source/Runtime/Engine/Private/PhysicsEngine/ClusterUnionComponent.cpp
--- a/Engine/Source/Runtime/Engine/Private/PhysicsEngine/ClusterUnionComponent.cpp
+++ b/Engine/Source/Runtime/Engine/Private/PhysicsEngine/ClusterUnionComponent.cpp
@@ -1118,6 +1118,8 @@ void UClusterUnionComponent::HandleAddOrModifiedClusteredComponent(const FMapped
 		}
 	}
 
+	OnChildToParentUpdated(ChangedComponentData.ComponentPtr, PerBoneChildToParent);
+
 	if (PerBoneChildToParent.Num())
 	{
 		// get the interface from the first element 
diff --git a/Engine/Source/Runtime/Experimental/GeometryCollectionEngine/Private/GeometryCollection/GeometryCollectionComponent.cpp b/Engine/Source/Runtime/Experimental/GeometryCollectionEngine/Private/GeometryCollection/GeometryCollectionComponent.cpp
--- a/Engine/Source/Runtime/Experimental/GeometryCollectionEngine/Private/GeometryCollection/GeometryCollectionComponent.cpp
+++ b/Engine/Source/Runtime/Experimental/GeometryCollectionEngine/Private/GeometryCollection/GeometryCollectionComponent.cpp
@@ -6557,6 +6557,24 @@ FTransform UGeometryCollectionComponent::GetRootCurrentComponentSpaceTransform()
 	return (RestCollection) ? FTransform(ComponentSpaceTransforms.RequestRootTransform()) : FTransform::Identity;
 }
 
+FTransform UGeometryCollectionComponent::GetRootParticleMassOffset() const
+{
+	FTransform MassOffset{ FTransform::Identity };
+	if (RestCollection)
+	{
+		const int32 RootIndex = RestCollection->GetRootIndex();
+		if (const FGeometryCollection* RestGeometryCollection = RestCollection->GetGeometryCollection().Get())
+		{
+			const TManagedArray<FTransform>* MassToLocal = RestGeometryCollection->FindAttribute<FTransform>("MassToLocal", FGeometryCollection::TransformGroup);
+			if (MassToLocal && MassToLocal->IsValidIndex(RootIndex))
+			{
+				MassOffset = (*MassToLocal)[RootIndex];
+			}
+		}
+	}
+	return MassOffset;
+}
+
 TArray<FTransform> UGeometryCollectionComponent::GetInitialLocalRestTransforms() const
 {
 	TArray<FTransform> InitialLocalTransforms;
diff --git a/Engine/Source/Runtime/Experimental/GeometryCollectionEngine/Public/GeometryCollection/GeometryCollectionComponent.h b/Engine/Source/Runtime/Experimental/GeometryCollectionEngine/Public/GeometryCollection/GeometryCollectionComponent.h
--- a/Engine/Source/Runtime/Experimental/GeometryCollectionEngine/Public/GeometryCollection/GeometryCollectionComponent.h
+++ b/Engine/Source/Runtime/Experimental/GeometryCollectionEngine/Public/GeometryCollection/GeometryCollectionComponent.h
@@ -807,6 +807,8 @@ public:
 
 	GEOMETRYCOLLECTIONENGINE_API FTransform GetRootCurrentComponentSpaceTransform() const;
 
+	GEOMETRYCOLLECTIONENGINE_API FTransform GetRootParticleMassOffset() const;
+
 	/** return true if the root cluster is not longer active at runtime */
 	UFUNCTION(BlueprintCallable, Category = "ChaosPhysics")
 	GEOMETRYCOLLECTIONENGINE_API bool IsRootBroken() const { return BrokenAndDecayedStates.GetIsRootBroken(); }
@@ -1334,6 +1336,8 @@ public:
 
 	/** update of the custom renderer when post physics sync callback is executing ( valid if custom redner is being used ) - true by default */
 	GEOMETRYCOLLECTIONENGINE_API void SetUpdateCustomRendererOnPostPhysicsSync(bool bValue) { bUpdateCustomRendererOnPostPhysicsSync = bValue; }
+	GEOMETRYCOLLECTIONENGINE_API bool GetUpdateCustomRendererOnPostPhysicsSync() const { return bUpdateCustomRendererOnPostPhysicsSync; }
+
 
 	GEOMETRYCOLLECTIONENGINE_API bool ShouldUpdateComponentTransformToRootBone() const { return bUpdateComponentTransformToRootBone; }
 
