geometry collection needs to turn on break events at the solver level if either notify breakings or notify global breakings is set to true - also made sure the flag gets set properly for removals and collisions

diff --git a/Engine/Source/Runtime/Experimental/GeometryCollectionEngine/Private/GeometryCollection/GeometryCollectionComponent.cpp b/Engine/Source/Runtime/Experimental/GeometryCollectionEngine/Private/GeometryCollection/GeometryCollectionComponent.cpp
--- a/Engine/Source/Runtime/Experimental/GeometryCollectionEngine/Private/GeometryCollection/GeometryCollectionComponent.cpp
+++ b/Engine/Source/Runtime/Experimental/GeometryCollectionEngine/Private/GeometryCollection/GeometryCollectionComponent.cpp
@@ -1172,6 +1172,7 @@ void UGeometryCollectionComponent::SetNotifyGlobalBreaks(bool bNewNotifyGlobalBr
 			PhysicsProxy->SetNotifyGlobalBreakings_External(bNewNotifyGlobalBreaks);
 		}
 		bNotifyGlobalBreaks = bNewNotifyGlobalBreaks;
+		RegisterForEvents();
 	}
 }
 
@@ -1181,6 +1182,7 @@ void UGeometryCollectionComponent::SetNotifyGlobalCollision(bool bNewNotifyGloba
 	{
 		bNotifyGlobalCollisions = bNewNotifyGlobalCollisions;
 		UpdateGlobalCollisionEventRegistration();
+		RegisterForEvents();
 	}
 }
 
@@ -1194,6 +1196,7 @@ void UGeometryCollectionComponent::SetNotifyGlobalRemovals(bool bNewNotifyGlobal
 		}
 		bNotifyGlobalRemovals = bNewNotifyGlobalRemovals;
 		UpdateGlobalRemovalEventRegistration();
+		RegisterForEvents();
 	}
 }
 
@@ -1872,7 +1875,19 @@ void UGeometryCollectionComponent::DispatchChaosPhysicsCollisionBlueprintEvents(
 // call when first registering
 void UGeometryCollectionComponent::RegisterForEvents()
 {
-	Chaos::FPhysicsSolver* Solver = GetWorld()->GetPhysicsScene()->GetSolver();
+	UWorld* World = GetWorld();
+	if (!World)
+	{
+		return;
+	}
+
+	FPhysScene* Scene = World->GetPhysicsScene();
+	if (!Scene)
+	{
+		return;
+	}
+
+	Chaos::FPhysicsSolver* Solver = Scene->GetSolver();
 	if (Solver)
 	{
 		if (EventDispatcher)
