Constraints: fix notification when the control's transform has been modified externally
    - this does not fix the issue completely but it's a prerequisite

#rnx

diff --git a/Engine/Plugins/Animation/ControlRig/Source/ControlRig/Private/Constraints/ControlRigTransformableHandle.cpp b/Engine/Plugins/Animation/ControlRig/Source/ControlRig/Private/Constraints/ControlRigTransformableHandle.cpp
--- a/Engine/Plugins/Animation/ControlRig/Source/ControlRig/Private/Constraints/ControlRigTransformableHandle.cpp
+++ b/Engine/Plugins/Animation/ControlRig/Source/ControlRig/Private/Constraints/ControlRigTransformableHandle.cpp
@@ -440,7 +440,7 @@ void UTransformableControlHandle::OnControlModified(
 					Control->Pose.IsDirty(ERigTransformType::CurrentGlobal);
 				if (bIsTransformDirty)
 				{
-					OnHandleModified.Broadcast(this, Event);
+					OnHandleModified.Broadcast(this, EHandleEvent::UpperDependencyUpdated);
 				}
 			}
 		}
diff --git a/Engine/Source/Runtime/Experimental/Animation/Constraints/Private/TransformConstraint.cpp b/Engine/Source/Runtime/Experimental/Animation/Constraints/Private/TransformConstraint.cpp
--- a/Engine/Source/Runtime/Experimental/Animation/Constraints/Private/TransformConstraint.cpp
+++ b/Engine/Source/Runtime/Experimental/Animation/Constraints/Private/TransformConstraint.cpp
@@ -581,17 +581,27 @@ void UTickableTransformConstraint::OnHandleModified(UTransformableHandle* InHand
 		return;
 	}
 
+	const UObject* Target = InHandle->GetTarget().Get();
+	UWorld* World = Target != nullptr ? Target->GetWorld() : nullptr;
+	if (!World)
+	{
+		return;
+	}
+	
 	if(InHandle == ChildTRSHandle || InHandle == ParentTRSHandle)
 	{
 		if (InNotification == EHandleEvent::ComponentUpdated)
 		{
-			UWorld* World = nullptr;
-			if (InHandle->GetTarget().Get())
+			SetupDependencies(World);
+			return;
+		}
+		
+		if (InNotification == EHandleEvent::UpperDependencyUpdated)
+		{
+			const FConstraintTickFunction* ConstraintTick = ConstraintTicks.Find(World->GetCurrentLevel());
+			if (IsFullyActive() && ConstraintTick->IsTickFunctionRegistered() && ConstraintTick->IsTickFunctionEnabled())
 			{
-				if (InHandle->GetTarget().Get()->GetWorld())
-				{
-					SetupDependencies(InHandle->GetTarget().Get()->GetWorld());
-				}
+				Evaluate();
 			}
 			return;
 		}
@@ -600,17 +610,10 @@ void UTickableTransformConstraint::OnHandleModified(UTransformableHandle* InHand
 	// update the constraint if the parent's transform has been modified 
 	if (InHandle == ParentTRSHandle && InNotification == EHandleEvent::GlobalTransformUpdated)
 	{
-		if (InHandle->GetTarget().Get())
+		const FConstraintTickFunction* ConstraintTick = ConstraintTicks.Find(World->GetCurrentLevel());
+		if (IsFullyActive() && ConstraintTick->IsTickFunctionRegistered() && ConstraintTick->IsTickFunctionEnabled())
 		{
-			if (UWorld* World = InHandle->GetTarget().Get()->GetWorld())
-			{
-				FConstraintTickFunction* ConstraintTick = ConstraintTicks.Find(World->GetCurrentLevel());
-				if (IsFullyActive() && ConstraintTick->IsTickFunctionRegistered() && ConstraintTick->IsTickFunctionEnabled())
-				{
-					Evaluate();
-				}
-
-			}
+			Evaluate();
 		}
 		return;
 	}
diff --git a/Engine/Source/Runtime/Experimental/Animation/Constraints/Public/TransformableHandle.h b/Engine/Source/Runtime/Experimental/Animation/Constraints/Public/TransformableHandle.h
--- a/Engine/Source/Runtime/Experimental/Animation/Constraints/Public/TransformableHandle.h
+++ b/Engine/Source/Runtime/Experimental/Animation/Constraints/Public/TransformableHandle.h
@@ -26,6 +26,7 @@ enum class EHandleEvent : uint8
 	LocalTransformUpdated,
 	GlobalTransformUpdated,
 	ComponentUpdated,
+	UpperDependencyUpdated,
 
 	/** MAX - invalid */
 	Max UMETA(Hidden),
