Lots of small changes:
*Fix scaling on suspension parameters
*Update use of cluster union events
*Adding extra inputs, boost, drift, reverse
*Some extra CVars for debugging
*Autoreverse added
*Suppling some common input state for the modules to use
*Exporting a few values bach to the game
*Reference ParticleIdx instead of an array index to prevent mismatch between simulation module and physical part when parts are removed.

diff --git a/Engine/Source/Runtime/Experimental/ChaosVehicles/ChaosVehiclesCore/Private/SimModule/AerofoilModule.cpp b/Engine/Source/Runtime/Experimental/ChaosVehicles/ChaosVehiclesCore/Private/SimModule/AerofoilModule.cpp
--- a/Engine/Source/Runtime/Experimental/ChaosVehicles/ChaosVehiclesCore/Private/SimModule/AerofoilModule.cpp
+++ b/Engine/Source/Runtime/Experimental/ChaosVehicles/ChaosVehiclesCore/Private/SimModule/AerofoilModule.cpp
@@ -142,7 +142,7 @@ namespace Chaos
 	{
 		if (FPBDRigidClusteredParticleHandle* ClusterChild = GetClusterParticle(Proxy))
 		{
-			FQuat ControlRotation = FQuat(Setup().ControlRotationAxis, FMath::DegreesToRadians(ControlSurfaceAngle));
+			FQuat ControlRotation = FQuat(Setup().ControlRotationAxis, FMath::DegreesToRadians(ControlSurfaceAngle * Setup().AnimationMagnitudeMultiplier));
 
 			FTransform InitialTransform = GetInitialParticleTransform();
 			ClusterChild->ChildToParent().SetRotation(InitialTransform.GetRotation() * ControlRotation);
diff --git a/Engine/Source/Runtime/Experimental/ChaosVehicles/ChaosVehiclesCore/Private/SimModule/DeferredForcesModular.cpp b/Engine/Source/Runtime/Experimental/ChaosVehicles/ChaosVehiclesCore/Private/SimModule/DeferredForcesModular.cpp
--- a/Engine/Source/Runtime/Experimental/ChaosVehicles/ChaosVehiclesCore/Private/SimModule/DeferredForcesModular.cpp
+++ b/Engine/Source/Runtime/Experimental/ChaosVehicles/ChaosVehiclesCore/Private/SimModule/DeferredForcesModular.cpp
@@ -19,10 +19,27 @@ FAutoConsoleVariableRef CVarChaosModularVehiclesDrawForceScaling(TEXT("p.Modular
 FAutoConsoleVariableRef CVarChaosModularVehiclesDisableForces(TEXT("p.ModularVehicle.DisableForces"), GCoreModularVehicleDebugParams.DisableForces, TEXT("Disable all forces."));
 #endif
 
+Chaos::FPBDRigidParticleHandle* FDeferredForcesModular::GetParticleFromUniqueIndex(int32 ParticleUniqueIdx, const TArray<Chaos::FPBDRigidParticleHandle*>& Particles) const
+{
+	for (Chaos::FPBDRigidParticleHandle* Particle : Particles)
+	{
+		if (Particle && Particle->UniqueIdx().IsValid())
+		{
+			if (ParticleUniqueIdx == Particle->UniqueIdx().Idx)
+			{
+				return Particle;
+			}
+		}
+	}
+
+	return nullptr;
+}
+
 // #TODO: passing all these parameters in is horrible, tidy this up now we know what needs to be done
 Chaos::FPBDRigidParticleHandle* FDeferredForcesModular::GetParticle(TArray<Chaos::FPBDRigidClusteredParticleHandle*>& Particles
 		, TArray<Chaos::FPBDRigidClusteredParticleHandle*>& ClusterParticles
 		, int TransformIndex
+		, int32 ParticleIdx
 		, const FVector& PositionalOffset
 		, const TManagedArray<FTransform>& Transforms
 		, const TManagedArray<FTransform>& CollectionMassToLocal
@@ -54,6 +71,7 @@ Chaos::FPBDRigidParticleHandle* FDeferredForcesModular::GetParticle(TArray<Chaos
 Chaos::FPBDRigidParticleHandle* FDeferredForcesModular::GetParticle(TArray<Chaos::FPBDRigidClusteredParticleHandle*>& Particles
 	, TArray<Chaos::FPBDRigidClusteredParticleHandle*>& ClusterParticles
 	, int TransformIndex
+	, int32 ParticleIdx
 	, const FVector& PositionalOffset
 	, const FTransform& Transform
 	, const TManagedArray<FTransform>& CollectionMassToLocal
@@ -84,7 +102,7 @@ Chaos::FPBDRigidParticleHandle* FDeferredForcesModular::GetParticle(TArray<Chaos
 Chaos::FPBDRigidParticleHandle* FDeferredForcesModular::GetParticle(const FTransform& OffsetTransform
 	, TArray<Chaos::FPBDRigidParticleHandle*>& Particles
 	, TArray<Chaos::FPBDRigidClusteredParticleHandle*>& ClusterParticles
-	, int TransformIndex
+	, int32 ParticleIdx
 	, const FVector& PositionalOffset
 	, FTransform& TransformOut)
 {
@@ -97,11 +115,8 @@ Chaos::FPBDRigidParticleHandle* FDeferredForcesModular::GetParticle(const FTrans
 	const FRigidTransform3 ClusterWorldTM(ClusterParticles[SingleChassisIndex]->X(), ClusterParticles[SingleChassisIndex]->R());
 	FRigidTransform3 Frame = FRigidTransform3::Identity;
 
-	//UE_LOG(LogInit, Warning, TEXT("TransformIndex %d  vs  Particles.Num() %d  vs ClusterParticles.Num() %d"), TransformIndex, Particles.Num(), ClusterParticles.Num());
-
-	if (TransformIndex < Particles.Num())
+	if (Chaos::FPBDRigidParticleHandle* Child = GetParticleFromUniqueIndex(ParticleIdx, Particles))
 	{
-		FPBDRigidParticleHandle* Child = Particles[TransformIndex];
 		if (FPBDRigidClusteredParticleHandle* ClusterChild = Child->CastToClustered(); ClusterChild && ClusterChild->IsChildToParentLocked())
 		{
 			Frame = ClusterChild->ChildToParent();
@@ -113,8 +128,6 @@ Chaos::FPBDRigidParticleHandle* FDeferredForcesModular::GetParticle(const FTrans
 			const FRigidTransform3 ChildWorldTM(Child->X(), Child->R());
 			Frame = ChildWorldTM.GetRelativeTransform(ClusterWorldTM);
 			Frame.SetLocation(Frame.GetLocation() + OffsetTransform.TransformVector(PositionalOffset));
-
-			//UE_LOG(LogChaos, Warning, TEXT("ApplyForce Data TIdx %d %s"), TransformIndex, *Frame.GetTranslation().ToString());
 		}
 	}
 
@@ -135,7 +148,7 @@ void FDeferredForcesModular::ApplyTemplate(TArray<Chaos::FPBDRigidClusteredParti
 		FTransform RelativeTransform;
 		for (const FApplyForceData& Data : ApplyForceDatas)
 		{
-			Chaos::FPBDRigidParticleHandle* RigidHandle = GetParticle(Particles, ClusterParticles, Data.TransformIndex, FVector::ZeroVector, FTransform(Transforms[Data.TransformIndex]), CollectionMassToLocal, Parent, RelativeTransform);
+			Chaos::FPBDRigidParticleHandle* RigidHandle = GetParticle(Particles, ClusterParticles, Data.TransformIndex, Data.ParticleIdx, FVector::ZeroVector, FTransform(Transforms[Data.TransformIndex]), CollectionMassToLocal, Parent, RelativeTransform);
 			if (RigidHandle)
 			{
 				AddForce(RigidHandle, Data, RelativeTransform);
@@ -144,7 +157,7 @@ void FDeferredForcesModular::ApplyTemplate(TArray<Chaos::FPBDRigidClusteredParti
 
 		for (const FApplyForceAtPositionData& Data : ApplyForceAtPositionDatas)
 		{
-			Chaos::FPBDRigidParticleHandle* RigidHandle = GetParticle(Particles, ClusterParticles, Data.TransformIndex, Data.Position, FTransform(Transforms[Data.TransformIndex]), CollectionMassToLocal, Parent, RelativeTransform);
+			Chaos::FPBDRigidParticleHandle* RigidHandle = GetParticle(Particles, ClusterParticles, Data.TransformIndex, Data.ParticleIdx, Data.Position, FTransform(Transforms[Data.TransformIndex]), CollectionMassToLocal, Parent, RelativeTransform);
 			if (RigidHandle)
 			{
 				AddForceAtPosition(RigidHandle, Data, RelativeTransform);
@@ -153,7 +166,7 @@ void FDeferredForcesModular::ApplyTemplate(TArray<Chaos::FPBDRigidClusteredParti
 
 		for (const FAddTorqueInRadiansData& Data : ApplyTorqueDatas)
 		{
-			Chaos::FPBDRigidParticleHandle* RigidHandle = GetParticle(Particles, ClusterParticles, Data.TransformIndex, FVector::ZeroVector, FTransform(Transforms[Data.TransformIndex]), CollectionMassToLocal, Parent, RelativeTransform);
+			Chaos::FPBDRigidParticleHandle* RigidHandle = GetParticle(Particles, ClusterParticles, Data.TransformIndex, Data.ParticleIdx, FVector::ZeroVector, FTransform(Transforms[Data.TransformIndex]), CollectionMassToLocal, Parent, RelativeTransform);
 			if (RigidHandle)
 			{
 				AddTorque(RigidHandle, Data, RelativeTransform);
@@ -193,7 +206,7 @@ void FDeferredForcesModular::Apply(TArray<Chaos::FPBDRigidParticleHandle*>& Part
 		FTransform RelativeTransform;
 		for (const FApplyForceData& Data : ApplyForceDatas)
 		{
-			Chaos::FPBDRigidParticleHandle* RigidHandle = GetParticle(Data.OffsetTransform, Particles, ClusterParticles, Data.TransformIndex, FVector::ZeroVector, RelativeTransform);
+			Chaos::FPBDRigidParticleHandle* RigidHandle = GetParticle(Data.OffsetTransform, Particles, ClusterParticles, Data.ParticleIdx, FVector::ZeroVector, RelativeTransform);
 			if (RigidHandle)
 			{
 				AddForce(RigidHandle, Data, RelativeTransform);
@@ -202,7 +215,7 @@ void FDeferredForcesModular::Apply(TArray<Chaos::FPBDRigidParticleHandle*>& Part
 
 		for (const FApplyForceAtPositionData& Data : ApplyForceAtPositionDatas)
 		{
-			Chaos::FPBDRigidParticleHandle* RigidHandle = GetParticle(Data.OffsetTransform, Particles, ClusterParticles, Data.TransformIndex, Data.Position, RelativeTransform);
+			Chaos::FPBDRigidParticleHandle* RigidHandle = GetParticle(Data.OffsetTransform, Particles, ClusterParticles, Data.ParticleIdx, Data.Position, RelativeTransform);
 			if (RigidHandle)
 			{
 				AddForceAtPosition(RigidHandle, Data, RelativeTransform);
@@ -211,7 +224,7 @@ void FDeferredForcesModular::Apply(TArray<Chaos::FPBDRigidParticleHandle*>& Part
 
 		for (const FAddTorqueInRadiansData& Data : ApplyTorqueDatas)
 		{
-			Chaos::FPBDRigidParticleHandle* RigidHandle = GetParticle(Data.OffsetTransform, Particles, ClusterParticles, Data.TransformIndex, FVector::ZeroVector, RelativeTransform);
+			Chaos::FPBDRigidParticleHandle* RigidHandle = GetParticle(Data.OffsetTransform, Particles, ClusterParticles, Data.ParticleIdx, FVector::ZeroVector, RelativeTransform);
 			if (RigidHandle)
 			{
 				AddTorque(RigidHandle, Data, RelativeTransform);
diff --git a/Engine/Source/Runtime/Experimental/ChaosVehicles/ChaosVehiclesCore/Private/SimModule/EngineModule.cpp b/Engine/Source/Runtime/Experimental/ChaosVehicles/ChaosVehiclesCore/Private/SimModule/EngineModule.cpp
--- a/Engine/Source/Runtime/Experimental/ChaosVehicles/ChaosVehiclesCore/Private/SimModule/EngineModule.cpp
+++ b/Engine/Source/Runtime/Experimental/ChaosVehicles/ChaosVehiclesCore/Private/SimModule/EngineModule.cpp
@@ -86,12 +86,13 @@ namespace Chaos
 		const FEngineOutputData& Next = static_cast<const FEngineOutputData&>(InNext);
 
 		RPM = FMath::Lerp(Current.RPM, Next.RPM, Alpha);
+		Torque = FMath::Lerp(Current.Torque, Next.Torque, Alpha);
 	}
 
 	FString FEngineOutputData::ToString()
 	{
-		return  FString::Printf(TEXT("RPM=%3.3f")
-			, RPM);
+		return  FString::Printf(TEXT("RPM=%3.3f, Torque=%3.3f")
+			, RPM, Torque);
 	}
 
 } // namespace Chaos
diff --git a/Engine/Source/Runtime/Experimental/ChaosVehicles/ChaosVehiclesCore/Private/SimModule/SimModuleTree.cpp b/Engine/Source/Runtime/Experimental/ChaosVehicles/ChaosVehiclesCore/Private/SimModule/SimModuleTree.cpp
--- a/Engine/Source/Runtime/Experimental/ChaosVehicles/ChaosVehiclesCore/Private/SimModule/SimModuleTree.cpp
+++ b/Engine/Source/Runtime/Experimental/ChaosVehicles/ChaosVehiclesCore/Private/SimModule/SimModuleTree.cpp
@@ -2,6 +2,7 @@
 
 #include "SimModule/SimModuleTree.h"
 #include "Chaos/ParticleHandleFwd.h"
+#include "Chaos/GeometryParticlesfwd.h"
 #include "PhysicsProxy/GeometryCollectionPhysicsProxy.h"
 #include "PhysicsProxy/ClusterUnionPhysicsProxy.h"
 #include "Chaos/PhysicsObjectInternalInterface.h"
@@ -18,9 +19,11 @@ DECLARE_CYCLE_STAT(TEXT("ModularVehicle_SetNetState"), STAT_ModularVehicle_SetNe
 DECLARE_CYCLE_STAT(TEXT("ModularVehicle_SetSimState"), STAT_ModularVehicle_SetSimState, STATGROUP_ModularVehicleSimTree);
 DECLARE_CYCLE_STAT(TEXT("ModularVehicle_AppendTreeUpdates"), STAT_ModularVehicle_AppendTreeUpdates, STATGROUP_ModularVehicleSimTree);
 
-bool bModularVehicle_NetworkData_Enable = false;
+bool bModularVehicle_NetworkData_Enable = true;
 FAutoConsoleVariableRef CVarModularVehicleNetworkDataEnable(TEXT("p.ModularVehicle.NetworkData.Enable"), bModularVehicle_NetworkData_Enable, TEXT("Enable/Disable additional module network data."));
 
+bool bModularVehicle_DisableAllSimulationAfterDestruction_Enable = false;
+FAutoConsoleVariableRef CVarModularVehicleDisableAllSimulationAfterDestruction(TEXT("p.ModularVehicle.DisableAllSimulationAfterDestruction.Enable"), bModularVehicle_DisableAllSimulationAfterDestruction_Enable, TEXT("Enable/Disable whole vehicle simulation after destruction has occured."));
 
 namespace Chaos
 {
@@ -123,52 +126,6 @@ void FSimModuleTree::AppendTreeUpdates(const FSimTreeUpdates& TreeUpdates)
 		}				
 	}
 
-	//// when the particle is removed the references will be out by one unless we fix them up
-	//// this may only apply to ClusterUnion and not GeometryCollection
-	//for (int ComponentIndex : ComponentIndices)
-	//{
-	//	for (int I = 0; I < GetNumNodes(); I++)
-	//	{
-	//		if (Chaos::ISimulationModuleBase* SimModule = GetNode(I).SimModule)
-	//		{
-	//			if (SimModule->GetTransformIndex() > ComponentIndex)
-	//			{
-	//				SimModule->SetTransformIndex(SimModule->GetTransformIndex() - 1);
-	//			}
-
-	//		}
-	//	}
-	//}
-
-	for (int ComponentIndex : ComponentIndices)
-	{
-		// find the largest transform index
-		int LargestIndex = -1;
-		for (int I = 0; I < GetNumNodes(); I++)
-		{
-			if (Chaos::ISimulationModuleBase* SimModule = GetNode(I).SimModule)
-			{
-				if (SimModule->GetTransformIndex() > LargestIndex)
-				{
-					LargestIndex = SimModule->GetTransformIndex();
-				}
-			}
-		}
-
-		// swap with the one that has just been deleted 
-		// - there can be more than one SimModule referencing the same component index
-		for (int I = 0; I < GetNumNodes(); I++)
-		{
-			if (Chaos::ISimulationModuleBase* SimModule = GetNode(I).SimModule)
-			{
-				if (SimModule->GetTransformIndex() == LargestIndex)
-				{
-					SimModule->SetTransformIndex(ComponentIndex);
-				}
-			}
-		}
-	}
-
 }
 
 
@@ -218,6 +175,12 @@ int FSimModuleTree::InsertNodeAbove(int AtIndex, ISimulationModuleBase* SimModul
 
 void FSimModuleTree::DeleteNode(int AtIndex)
 {
+	// if is there is ever an issue then we have the option of disabling ALL module simulation after first destruction occurs
+	if (bModularVehicle_DisableAllSimulationAfterDestruction_Enable)
+	{
+		SetSimulationEnabled(false);
+	}
+
 	// multiple children might become equal parents?
 	
 	int ParentIndex = SimulationModuleTree[AtIndex].Parent;
@@ -252,18 +215,24 @@ void FSimModuleTree::Simulate(float DeltaTime, FAllInputs& Inputs, FClusterUnion
 {
 	SCOPE_CYCLE_COUNTER(STAT_ModularVehicle_SimulateTree);
 
-	if (PhysicsProxy)
+	if (IsSimulationEnabled())
 	{
-		UpdateModuleVelocites(PhysicsProxy, Inputs.ControlInputs.InputNonZero());
-	}
+		if (PhysicsProxy)
+		{
+			UpdateVehicleState(PhysicsProxy);
 
-	TArray<int> RootNodes;
-	GetRootNodes(RootNodes);
+			UpdateModuleVelocites(PhysicsProxy, Inputs.ControlInputs.InputNonZero() || Inputs.bKeepVehicleAwake);
+		}
 
-	for (int RootIndex : RootNodes)
-	{
-		SimulateNode(DeltaTime, Inputs, RootIndex, PhysicsProxy);
+		TArray<int> RootNodes;
+		GetRootNodes(RootNodes);
+
+		for (int RootIndex : RootNodes)
+		{
+			SimulateNode(DeltaTime, Inputs, RootIndex, PhysicsProxy);
+		}
 	}
+
 }
 
 
@@ -455,6 +424,24 @@ void FSimModuleTree::UpdateModuleVelocites(FClusterUnionPhysicsProxy* PhysicsPro
 	}
 }
 
+void FSimModuleTree::UpdateVehicleState(FClusterUnionPhysicsProxy* PhysicsProxy)
+{
+	check(PhysicsProxy);
+	Chaos::EnsureIsInPhysicsThreadContext();
+
+	if (Chaos::FClusterUnionPhysicsProxy::FInternalParticle* ParentParticle = PhysicsProxy->GetParticle_Internal())
+	{
+		const FTransform BodyTransform(ParentParticle->R(), ParentParticle->X());
+
+		VehicleState.ForwardDir = BodyTransform.GetUnitAxis(EAxis::X);
+		VehicleState.UpDir = BodyTransform.GetUnitAxis(EAxis::Z);
+		VehicleState.RightDir = BodyTransform.GetUnitAxis(EAxis::Y);
+		VehicleState.ForwardSpeedKmh = Chaos::CmSToKmH(FVector::DotProduct(ParentParticle->V(), VehicleState.ForwardDir));
+		VehicleState.AngularVelocityRad = ParentParticle->W();
+	}
+
+}
+
 void FSimModuleTree::GenerateReplicationStructure(Chaos::FModuleNetDataArray& NetData)
 {
 	SCOPE_CYCLE_COUNTER(STAT_ModularVehicle_GenerateReplicationStructure);
diff --git a/Engine/Source/Runtime/Experimental/ChaosVehicles/ChaosVehiclesCore/Private/SimModule/SimulationModuleBase.cpp b/Engine/Source/Runtime/Experimental/ChaosVehicles/ChaosVehiclesCore/Private/SimModule/SimulationModuleBase.cpp
--- a/Engine/Source/Runtime/Experimental/ChaosVehicles/ChaosVehiclesCore/Private/SimModule/SimulationModuleBase.cpp
+++ b/Engine/Source/Runtime/Experimental/ChaosVehicles/ChaosVehiclesCore/Private/SimModule/SimulationModuleBase.cpp
@@ -15,7 +15,7 @@ void ISimulationModuleBase::AddLocalForceAtPosition(const FVector& Force, const
 	AppliedForce = Force;
 	if (SimModuleTree)
 	{
-		SimModuleTree->AccessDeferredForces().Add(FDeferredForcesModular::FApplyForceAtPositionData(ComponentTransform, TransformIndex, Force, Position, bAllowSubstepping, bIsLocalForce, bLevelSlope, DebugColorIn));
+		SimModuleTree->AccessDeferredForces().Add(FDeferredForcesModular::FApplyForceAtPositionData(ComponentTransform, TransformIndex, ParticleIdx.Idx, Force, Position, bAllowSubstepping, bIsLocalForce, bLevelSlope, DebugColorIn));
 	}
 }
 
@@ -24,7 +24,7 @@ void ISimulationModuleBase::AddLocalForce(const FVector& Force, bool bAllowSubst
 	AppliedForce = Force;
 	if (SimModuleTree)
 	{
-		SimModuleTree->AccessDeferredForces().Add(FDeferredForcesModular::FApplyForceData(ComponentTransform, TransformIndex, Force, bAllowSubstepping, bIsLocalForce, bLevelSlope, DebugColorIn));
+		SimModuleTree->AccessDeferredForces().Add(FDeferredForcesModular::FApplyForceData(ComponentTransform, TransformIndex, ParticleIdx.Idx, Force, bAllowSubstepping, bIsLocalForce, bLevelSlope, DebugColorIn));
 	}
 }
 
@@ -32,7 +32,7 @@ void ISimulationModuleBase::AddLocalTorque(const FVector& Torque, bool bAllowSub
 {
 	if (SimModuleTree)
 	{
-		SimModuleTree->AccessDeferredForces().Add(FDeferredForcesModular::FAddTorqueInRadiansData(ComponentTransform, TransformIndex, Torque, bAllowSubstepping, bAccelChangeIn, DebugColorIn));
+		SimModuleTree->AccessDeferredForces().Add(FDeferredForcesModular::FAddTorqueInRadiansData(ComponentTransform, TransformIndex, ParticleIdx.Idx, Torque, bAllowSubstepping, bAccelChangeIn, DebugColorIn));
 	}
 }
 
@@ -68,15 +68,31 @@ FPBDRigidClusteredParticleHandle* ISimulationModuleBase::GetClusterParticle(Chao
 		FPBDRigidClusteredParticleHandle* ClusterHandle = ClusterUnion->InternalCluster;
 		TArray<FPBDRigidParticleHandle*> Particles = ClusterUnion->ChildParticles;
 
-		if (TransformIndex < Particles.Num())
+		if (FPBDRigidParticleHandle* Particle = GetParticleFromUniqueIndex(ParticleIdx.Idx, Particles))
 		{
-			ClusterChild = Particles[TransformIndex]->CastToClustered();
+			ClusterChild = Particle->CastToClustered();
 		}
 	}
 
 	return ClusterChild;
 }
 
+FPBDRigidParticleHandle* ISimulationModuleBase::GetParticleFromUniqueIndex(int32 ParticleUniqueIdx, TArray<FPBDRigidParticleHandle*>& Particles)
+{
+	for (FPBDRigidParticleHandle* Particle : Particles)
+	{
+		if (Particle && Particle->UniqueIdx().IsValid())
+		{
+			if (ParticleUniqueIdx == Particle->UniqueIdx().Idx)
+			{
+				return Particle;
+			}
+		}
+	}
+
+	return nullptr;
+}
+
 bool ISimulationModuleBase::GetDebugString(FString& StringOut) const 
 {
 	StringOut += FString::Format(TEXT("{0}: TreeIndex {1}, Enabled {2}, InCluster {3}, TFormIdx {4}, ")
diff --git a/Engine/Source/Runtime/Experimental/ChaosVehicles/ChaosVehiclesCore/Private/SimModule/SuspensionModule.cpp b/Engine/Source/Runtime/Experimental/ChaosVehicles/ChaosVehiclesCore/Private/SimModule/SuspensionModule.cpp
--- a/Engine/Source/Runtime/Experimental/ChaosVehicles/ChaosVehiclesCore/Private/SimModule/SuspensionModule.cpp
+++ b/Engine/Source/Runtime/Experimental/ChaosVehicles/ChaosVehiclesCore/Private/SimModule/SuspensionModule.cpp
@@ -22,6 +22,7 @@ namespace Chaos
 		, SpringSpeed(0.f)
 		, WheelSimTreeIndex(INVALID_IDX)
 		, Constraint(nullptr)
+		, ConstraintIndex(INVALID_IDX)
 		, TargetPos(FVector::ZeroVector)
 		, ImpactNormal(FVector::ZeroVector)
 		, WheelInContact(false)
@@ -59,7 +60,7 @@ namespace Chaos
 			float ForceIntoSurface = 0.0f;
 			if (SpringDisplacement > 0)
 			{
-				float Damping = (SpringDisplacement < LastDisplacement) ? Setup().CompressionDamping : Setup().ReboundDamping;
+				float Damping = Setup().SpringDamping;
 				SpringSpeed = (LastDisplacement - SpringDisplacement) / DeltaTime;
 
 				float StiffnessForce = SpringDisplacement * Setup().SpringRate;
@@ -69,7 +70,7 @@ namespace Chaos
 
 				if (SuspensionForce > 0)
 				{
-					ForceIntoSurface = SuspensionForce;
+					ForceIntoSurface = SuspensionForce * Setup().SuspensionForceEffect;
 
 					if (Constraint == nullptr)
 					{
@@ -126,7 +127,6 @@ namespace Chaos
 			if (FSuspensionConstraintPhysicsProxy* Proxy = Constraint->GetProxy<FSuspensionConstraintPhysicsProxy>())
 			{
 				Chaos::FPhysicsSolver* Solver = Proxy->GetSolver<Chaos::FPhysicsSolver>();
-				UE_LOG(LogTemp, Warning, TEXT("TargetPos %s, WheelInContact %d"), *TargetPos.ToString(), WheelInContact);
 				Solver->SetSuspensionTarget(Constraint, TargetPos, ImpactNormal, WheelInContact);
 			}
 		}
diff --git a/Engine/Source/Runtime/Experimental/ChaosVehicles/ChaosVehiclesCore/Private/SimModule/ThrusterModule.cpp b/Engine/Source/Runtime/Experimental/ChaosVehicles/ChaosVehiclesCore/Private/SimModule/ThrusterModule.cpp
--- a/Engine/Source/Runtime/Experimental/ChaosVehicles/ChaosVehiclesCore/Private/SimModule/ThrusterModule.cpp
+++ b/Engine/Source/Runtime/Experimental/ChaosVehicles/ChaosVehiclesCore/Private/SimModule/ThrusterModule.cpp
@@ -26,7 +26,8 @@ namespace Chaos
 		}
 
 		// applies continuous force
-		FVector Force = Setup().ForceAxis * Setup().MaxThrustForce * Inputs.ControlInputs.Throttle;
+		float BoostEffect = Inputs.ControlInputs.Boost * Setup().BoostMultiplier;
+		FVector Force = Setup().ForceAxis * Setup().MaxThrustForce * Inputs.ControlInputs.Throttle * (1.0f + BoostEffect);
 		FQuat Steer = FQuat(Setup().SteeringAxis, FMath::DegreesToRadians(SteerAngleDegrees) * Setup().SteeringForceEffect);
 		AddLocalForceAtPosition(Steer.RotateVector(Force), Setup().ForceOffset, true, false, false, FColor::Magenta);
 	}
diff --git a/Engine/Source/Runtime/Experimental/ChaosVehicles/ChaosVehiclesCore/Private/SimModule/TorqueSimModule.cpp b/Engine/Source/Runtime/Experimental/ChaosVehicles/ChaosVehiclesCore/Private/SimModule/TorqueSimModule.cpp
--- a/Engine/Source/Runtime/Experimental/ChaosVehicles/ChaosVehiclesCore/Private/SimModule/TorqueSimModule.cpp
+++ b/Engine/Source/Runtime/Experimental/ChaosVehicles/ChaosVehiclesCore/Private/SimModule/TorqueSimModule.cpp
@@ -46,7 +46,8 @@ void FTorqueSimModule::TransmitTorque(const FSimModuleTree& ModuleTree, float Pu
 			}
 		}
 
-		BrakingTorque /= GearingRatio;
+		BrakingTorque /= FMath::Abs(GearingRatio);
+
 	}
 }
 
diff --git a/Engine/Source/Runtime/Experimental/ChaosVehicles/ChaosVehiclesCore/Private/SimModule/TransmissionModule.cpp b/Engine/Source/Runtime/Experimental/ChaosVehicles/ChaosVehiclesCore/Private/SimModule/TransmissionModule.cpp
--- a/Engine/Source/Runtime/Experimental/ChaosVehicles/ChaosVehiclesCore/Private/SimModule/TransmissionModule.cpp
+++ b/Engine/Source/Runtime/Experimental/ChaosVehicles/ChaosVehiclesCore/Private/SimModule/TransmissionModule.cpp
@@ -2,6 +2,7 @@
 
 #include "SimModule/TransmissionModule.h"
 #include "SimModule/SimModuleTree.h"
+#include "SimModule/SimulationModuleBase.h"
 #include "VehicleUtility.h"
 
 #if VEHICLE_DEBUGGING_ENABLED
@@ -13,6 +14,27 @@ namespace Chaos
 
 	void FTransmissionSimModule::Simulate(float DeltaTime, const FAllInputs& Inputs, FSimModuleTree& VehicleModuleSystem)
 	{
+
+		if (Setup().AutoReverse)
+		{
+			if (Inputs.ControlInputs.IsReversing)
+			{
+				// if reversing change to reverse gear if currently in a forwards gear
+				if (TargetGear > 0)
+				{
+					TargetGear = -1;
+				}
+			}
+			else
+			{
+				// if not revering change to forwards gear if currently in a reverse gear
+				if (TargetGear < 0)
+				{
+					TargetGear = 1;
+				}
+			}
+		}
+
 		if (Setup().TransmissionType == FTransmissionSettings::ETransType::AutomaticType)
 		{
 			// not currently changing gear, also don't want to change up because the wheels are spinning up due to having no load
diff --git a/Engine/Source/Runtime/Experimental/ChaosVehicles/ChaosVehiclesCore/Private/SimModule/WheelModule.cpp b/Engine/Source/Runtime/Experimental/ChaosVehicles/ChaosVehiclesCore/Private/SimModule/WheelModule.cpp
--- a/Engine/Source/Runtime/Experimental/ChaosVehicles/ChaosVehiclesCore/Private/SimModule/WheelModule.cpp
+++ b/Engine/Source/Runtime/Experimental/ChaosVehicles/ChaosVehiclesCore/Private/SimModule/WheelModule.cpp
@@ -18,7 +18,7 @@ namespace Chaos
 		, SurfaceFriction(1.0f)
 		, SuspensionSimTreeIndex(INVALID_IDX)
 		, ForceFromFriction(FVector::ZeroVector)
-		, MassPerWheel(500.0f * 0.25f)
+		, MassPerWheel(500.0f)
 		, SteerAngleDegrees(0.0f)
 	{
 
diff --git a/Engine/Source/Runtime/Experimental/ChaosVehicles/ChaosVehiclesCore/Public/SimModule/AerofoilModule.h b/Engine/Source/Runtime/Experimental/ChaosVehicles/ChaosVehiclesCore/Public/SimModule/AerofoilModule.h
--- a/Engine/Source/Runtime/Experimental/ChaosVehicles/ChaosVehiclesCore/Public/SimModule/AerofoilModule.h
+++ b/Engine/Source/Runtime/Experimental/ChaosVehicles/ChaosVehiclesCore/Public/SimModule/AerofoilModule.h
@@ -31,9 +31,10 @@ namespace Chaos
 			, Type(EAerofoil::Fixed)
 			, LiftMultiplier(1.0f)
 			, DragMultiplier(1.0f)
+			, AnimationMagnitudeMultiplier(1.0f)
 		{
 		}
-
+		
 		FVector Offset;
 		FVector ForceAxis;
 		FVector ControlRotationAxis;
@@ -45,6 +46,7 @@ namespace Chaos
 		EAerofoil Type;
 		float LiftMultiplier;
 		float DragMultiplier;
+		float AnimationMagnitudeMultiplier;
 
 	};
 
diff --git a/Engine/Source/Runtime/Experimental/ChaosVehicles/ChaosVehiclesCore/Public/SimModule/DeferredForcesModular.h b/Engine/Source/Runtime/Experimental/ChaosVehicles/ChaosVehiclesCore/Public/SimModule/DeferredForcesModular.h
--- a/Engine/Source/Runtime/Experimental/ChaosVehicles/ChaosVehiclesCore/Public/SimModule/DeferredForcesModular.h
+++ b/Engine/Source/Runtime/Experimental/ChaosVehicles/ChaosVehiclesCore/Public/SimModule/DeferredForcesModular.h
@@ -26,9 +26,10 @@ public:
 
 	struct FApplyForceData
 	{
-		FApplyForceData(const FTransform& OffsetTransformIn, int TransformIndexIn, const FVector& ForceIn, bool bAllowSubsteppingIn, bool bAccelChangeIn, bool bLevelSlope, const FColor& ColorIn)
+		FApplyForceData(const FTransform& OffsetTransformIn, int TransformIndexIn, int ParticleIndexIn, const FVector& ForceIn, bool bAllowSubsteppingIn, bool bAccelChangeIn, bool bLevelSlope, const FColor& ColorIn)
 			: OffsetTransform(OffsetTransformIn)
 			, TransformIndex(TransformIndexIn)
+			, ParticleIdx(ParticleIndexIn)
 			, Force(ForceIn)
 			, Flags(EForceFlags::None)
 			, DebugColor(ColorIn)
@@ -40,6 +41,7 @@ public:
 
 		FTransform OffsetTransform;
 		int TransformIndex;
+		int32 ParticleIdx;
 		FVector Force;
 		EForceFlags Flags;
 		FColor DebugColor;
@@ -47,9 +49,10 @@ public:
 
 	struct FApplyForceAtPositionData
 	{
-		FApplyForceAtPositionData(const FTransform& OffsetTransformIn, int TransformIndexIn, const FVector& ForceIn, const FVector& PositionIn, bool bAllowSubsteppingIn, bool bIsLocalForceIn, bool bLevelSlope, const FColor& ColorIn)
+		FApplyForceAtPositionData(const FTransform& OffsetTransformIn, int TransformIndexIn, int ParticleIndexIn, const FVector& ForceIn, const FVector& PositionIn, bool bAllowSubsteppingIn, bool bIsLocalForceIn, bool bLevelSlope, const FColor& ColorIn)
 			: OffsetTransform(OffsetTransformIn)
 			, TransformIndex(TransformIndexIn)
+			, ParticleIdx(ParticleIndexIn)
 			, Force(ForceIn)
 			, Position(PositionIn)
 			, Flags(EForceFlags::None)
@@ -62,6 +65,7 @@ public:
 
 		FTransform OffsetTransform;
 		int TransformIndex;
+		int32 ParticleIdx;
 		FVector Force;
 		FVector Position;
 		EForceFlags Flags;
@@ -70,9 +74,10 @@ public:
 
 	struct FAddTorqueInRadiansData
 	{
-		FAddTorqueInRadiansData(const FTransform& OffsetTransformIn, int TransformIndexIn, const FVector& TorqueIn, bool bAllowSubsteppingIn, bool bAccelChangeIn, const FColor& ColorIn)
+		FAddTorqueInRadiansData(const FTransform& OffsetTransformIn, int TransformIndexIn, int ParticleIndexIn, const FVector& TorqueIn, bool bAllowSubsteppingIn, bool bAccelChangeIn, const FColor& ColorIn)
 			: OffsetTransform(OffsetTransformIn)
 			, TransformIndex(TransformIndexIn)
+			, ParticleIdx(ParticleIndexIn)
 			, Torque(TorqueIn)
 			, Flags(EForceFlags::None)
 			, DebugColor(ColorIn)
@@ -83,6 +88,7 @@ public:
 
 		FTransform OffsetTransform;
 		int TransformIndex;
+		int32 ParticleIdx;
 		FVector Torque;
 		EForceFlags Flags;
 		FColor DebugColor;
@@ -90,9 +96,10 @@ public:
 
 	struct FAddImpulseData
 	{
-		FAddImpulseData(const FTransform& OffsetTransformIn, int TransformIndexIn, const FVector& ImpulseIn, const bool bVelChangeIn)
+		FAddImpulseData(const FTransform& OffsetTransformIn, int TransformIndexIn, int ParticleIndexIn, const FVector& ImpulseIn, const bool bVelChangeIn)
 			: OffsetTransform(OffsetTransformIn)
 			, TransformIndex(TransformIndexIn)
+			, ParticleIdx(ParticleIndexIn)
 			, Impulse(ImpulseIn)
 			, Flags(EForceFlags::None)
 		{
@@ -101,15 +108,17 @@ public:
 
 		FTransform OffsetTransform;
 		int TransformIndex;
+		int32 ParticleIdx;
 		FVector Impulse;
 		EForceFlags Flags;
 	};
 
 	struct FAddImpulseAtPositionData
 	{
-		FAddImpulseAtPositionData(const FTransform& OffsetTransformIn, int TransformIndexIn, const FVector& ImpulseIn, const FVector& PositionIn)
+		FAddImpulseAtPositionData(const FTransform& OffsetTransformIn, int TransformIndexIn, int ParticleIndexIn, const FVector& ImpulseIn, const FVector& PositionIn)
 			: OffsetTransform(OffsetTransformIn)
 			, TransformIndex(TransformIndexIn)
+			, ParticleIdx(ParticleIndexIn)
 			, Impulse(ImpulseIn)
 			, Position(PositionIn)
 		{
@@ -118,6 +127,7 @@ public:
 
 		FTransform OffsetTransform;
 		int TransformIndex;
+		int32 ParticleIdx;
 		FVector Impulse;
 		FVector Position;
 	};
@@ -147,9 +157,12 @@ public:
 		ApplyImpulseAtPositionDatas.Add(ApplyImpulseAtPositionDataIn);
 	}	
 
+	Chaos::FPBDRigidParticleHandle* GetParticleFromUniqueIndex(int32 ParticleUniqueIdx, const TArray<Chaos::FPBDRigidParticleHandle*>& Particles) const;
+
 	Chaos::FPBDRigidParticleHandle* GetParticle(TArray<Chaos::FPBDRigidClusteredParticleHandle*>& Particles
 			, TArray<Chaos::FPBDRigidClusteredParticleHandle*>& ClusterParticles
 			, int TransformIndex
+			, int32 ParticleIdx
 			, const FVector& PositionalOffset
 			, const TManagedArray<FTransform>& Transforms
 			, const TManagedArray<FTransform>& CollectionMassToLocal
@@ -159,6 +172,7 @@ public:
 	Chaos::FPBDRigidParticleHandle* GetParticle(TArray<Chaos::FPBDRigidClusteredParticleHandle*>& Particles
 		, TArray<Chaos::FPBDRigidClusteredParticleHandle*>& ClusterParticles
 		, int TransformIndex
+		, int32 ParticleIdx
 		, const FVector& PositionalOffset
 		, const FTransform& Transform
 		, const TManagedArray<FTransform>& CollectionMassToLocal
@@ -168,7 +182,7 @@ public:
 	Chaos::FPBDRigidParticleHandle* GetParticle(const FTransform& OffsetTransform
 		, TArray<Chaos::FPBDRigidParticleHandle*>& Particles
 		, TArray<Chaos::FPBDRigidClusteredParticleHandle*>& ClusterParticles
-		, int TransformIndex
+		, int32 ParticleIdx
 		, const FVector& PositionalOffset
 		, FTransform& TransformOut);
 
diff --git a/Engine/Source/Runtime/Experimental/ChaosVehicles/ChaosVehiclesCore/Public/SimModule/EngineModule.h b/Engine/Source/Runtime/Experimental/ChaosVehicles/ChaosVehiclesCore/Public/SimModule/EngineModule.h
--- a/Engine/Source/Runtime/Experimental/ChaosVehicles/ChaosVehiclesCore/Public/SimModule/EngineModule.h
+++ b/Engine/Source/Runtime/Experimental/ChaosVehicles/ChaosVehiclesCore/Public/SimModule/EngineModule.h
@@ -39,11 +39,13 @@ namespace Chaos
 		virtual FSimOutputData* MakeNewData() override { return FEngineOutputData::MakeNew(); }
 		static FSimOutputData* MakeNew() { return new FEngineOutputData(); }
 
+		virtual eSimType GetType() override { return eSimType::Engine; }
 		virtual void FillOutputState(const ISimulationModuleBase* SimModule) override;
 		virtual void Lerp(const FSimOutputData& InCurrent, const FSimOutputData& InNext, float Alpha) override;
 		virtual FString ToString() override;
 
 		float RPM;
+		float Torque;
 	};
 
 	struct CHAOSVEHICLESCORE_API FEngineSettings
diff --git a/Engine/Source/Runtime/Experimental/ChaosVehicles/ChaosVehiclesCore/Public/SimModule/SimModuleTree.h b/Engine/Source/Runtime/Experimental/ChaosVehicles/ChaosVehiclesCore/Public/SimModule/SimModuleTree.h
--- a/Engine/Source/Runtime/Experimental/ChaosVehicles/ChaosVehiclesCore/Public/SimModule/SimModuleTree.h
+++ b/Engine/Source/Runtime/Experimental/ChaosVehicles/ChaosVehiclesCore/Public/SimModule/SimModuleTree.h
@@ -67,6 +67,14 @@ namespace Chaos
 		TArray<FPendingModuleDeletions> DeletedModules;
 	};
 
+	struct CHAOSVEHICLESCORE_API FVehicleState
+	{
+		float ForwardSpeedKmh;
+		FVector ForwardDir;
+		FVector RightDir;
+		FVector UpDir;
+		FVector AngularVelocityRad;
+	};
 
 	class CHAOSVEHICLESCORE_API FSimModuleTree
 	{
@@ -93,7 +101,8 @@ namespace Chaos
 
 		FSimModuleTree()
 		{
-			bAnimationEnabled = false;
+			bAnimationEnabled = true;
+			bSimulationEnabled = true;
 		}
 
 		~FSimModuleTree()
@@ -138,12 +147,19 @@ namespace Chaos
 
 		void SetAnimationEnabled(bool bInEnabled) { bAnimationEnabled = bInEnabled; }
 		bool IsAnimationEnabled() { return bAnimationEnabled; }
+		void SetSimulationEnabled(bool bInEnabled) { bSimulationEnabled = bInEnabled; }
+		bool IsSimulationEnabled() { return bSimulationEnabled; }
 
 		FControlInputs& GetControlInputs()
 		{
 			return AllInputs.ControlInputs;
 		}
 
+		const FVehicleState& GetVehicleState() const
+		{
+			return VehicleState;
+		}
+
 		FSimModuleNode* LocateNodeByType(Chaos::eSimType InType)
 		{
 			for (FSimModuleNode& Node : SimulationModuleTree)
@@ -187,6 +203,7 @@ namespace Chaos
 
 		void UpdateModuleVelocites(FGeometryCollectionPhysicsProxy* PhysicsProxy);
 		void UpdateModuleVelocites(FClusterUnionPhysicsProxy* PhysicsProxy, bool bWake);
+		void UpdateVehicleState(FClusterUnionPhysicsProxy* PhysicsProxy);
 
 		TArray<FSimModuleNode> SimulationModuleTree;
 		TArray<int> FreeList;
@@ -196,6 +213,9 @@ namespace Chaos
 		Chaos::FAllInputs AllInputs;
 
 		bool bAnimationEnabled;
+		bool bSimulationEnabled;
+
+		FVehicleState VehicleState;
 	};
 
 
diff --git a/Engine/Source/Runtime/Experimental/ChaosVehicles/ChaosVehiclesCore/Public/SimModule/SimulationModuleBase.h b/Engine/Source/Runtime/Experimental/ChaosVehicles/ChaosVehiclesCore/Public/SimModule/SimulationModuleBase.h
--- a/Engine/Source/Runtime/Experimental/ChaosVehicles/ChaosVehiclesCore/Public/SimModule/SimulationModuleBase.h
+++ b/Engine/Source/Runtime/Experimental/ChaosVehicles/ChaosVehiclesCore/Public/SimModule/SimulationModuleBase.h
@@ -5,6 +5,7 @@
 #include "CoreMinimal.h"
 #include "Logging/LogMacros.h"
 #include "Chaos/ParticleHandleFwd.h"
+#include "Chaos/GeometryParticlesfwd.h"
 
 
 DECLARE_LOG_CATEGORY_EXTERN(LogSimulationModule, Warning, All);
@@ -14,7 +15,7 @@ struct CHAOSVEHICLESCORE_API FCoreModularVehicleDebugParams
 	bool ShowMass = false;
 	bool ShowForces = false;
 	float DrawForceScaling = 0.0004f;
-	float LevelSlopeThreshold = 0.96f; // ~16 degrees
+	float LevelSlopeThreshold = 0.86f;
 	bool DisableForces = false;
 };
 
@@ -30,6 +31,7 @@ namespace Chaos
 
 		FControlInputs()
 			: IsValid(false)
+			, IsReversing(false)
 			, Throttle(0)
 			, Brake(0)
 			, Steering(0)
@@ -39,6 +41,7 @@ namespace Chaos
 			, Pitch(0)
 			, Yaw(0)
 			, Boost(0)
+			, Drift(0)
 			, ChangeUp(false)
 			, ChangeDown(false)
 			, GearNumber(0)
@@ -54,10 +57,12 @@ namespace Chaos
 				|| FMath::Abs(Roll) > SMALL_NUMBER
 				|| FMath::Abs(Pitch) > SMALL_NUMBER
 				|| FMath::Abs(Yaw) > SMALL_NUMBER
-				|| FMath::Abs(Boost) > SMALL_NUMBER;
+				|| FMath::Abs(Boost) > SMALL_NUMBER
+				|| FMath::Abs(Drift) > SMALL_NUMBER;
 		}
 
 		bool IsValid;
+		bool IsReversing;
 		float Throttle;
 		float Brake;
 		float Steering;
@@ -67,6 +72,7 @@ namespace Chaos
 		float Pitch;
 		float Yaw;
 		float Boost;
+		float Drift;
 		bool ChangeUp;
 		bool ChangeDown;
 		int GearNumber;
@@ -86,6 +92,7 @@ namespace Chaos
 		FTransform VehicleWorldTransform;
 		TMap<int32, FModuleHitResults> HitResults;
 		FControlInputs ControlInputs;
+		bool bKeepVehicleAwake;
 	};
 
 	/**
@@ -163,6 +170,7 @@ namespace Chaos
 			, SimTreeIndex(INVALID_IDX)
 			, StateFlags(Enabled)
 			, TransformIndex(INVALID_IDX)
+			, ParticleIdx(INVALID_IDX)
 			, LocalLinearVelocity(FVector::ZeroVector)
 			, LocalAngularVelocity(FVector::ZeroVector)
 			, bClustered(true)
@@ -229,6 +237,12 @@ namespace Chaos
 		void SetTransformIndex(int TransformIndexIn) { TransformIndex = TransformIndexIn; }
 		const int GetTransformIndex() const { return TransformIndex; }
 
+		/**
+		 * The Particle unique index, should be valid on game and physics threads
+		 */
+		void SetParticleIndex(FUniqueIdx ParticleIndexIn) { ParticleIdx = ParticleIndexIn; }
+		const FUniqueIdx GetParticleIndex() const { return ParticleIdx; }
+
 		/**
 		 * The modules own index in the simulation tree array
 		 */
@@ -312,12 +326,16 @@ namespace Chaos
 		//void SetClusterParticle(FPBDRigidClusteredParticleHandle* ParticleIn) { ClusterParticle = ParticleIn; }
 		Chaos::FPBDRigidClusteredParticleHandle* GetClusterParticle(Chaos::FClusterUnionPhysicsProxy* Proxy);
 
+		Chaos::FPBDRigidParticleHandle* GetParticleFromUniqueIndex(int32 ParticleUniqueIdx, TArray<Chaos::FPBDRigidParticleHandle*>& Particles);
+
+
 	protected:
 
 		FSimModuleTree* SimModuleTree;	// A pointer back to the simulation tree where we are stored
 		int SimTreeIndex;	// Index of this SimModule in the FSimModuleTree
 		eSimModuleState StateFlags;	// TODO: make this more like flags
 		int TransformIndex; // Index of this Sim Module's node in Geometry Collection Transform array
+		FUniqueIdx ParticleIdx; // Physics particle unique index
 
 		FTransform InitialParticleTransform;
 		FTransform RelativeOffsetTransform;
@@ -372,6 +390,7 @@ namespace Chaos
 		FSimOutputData() = default;
 		virtual ~FSimOutputData() {}
 
+		virtual eSimType GetType() = 0;
 		virtual FSimOutputData* MakeNewData() = 0;
 		virtual void FillOutputState(const ISimulationModuleBase* SimModule) = 0;
 		virtual void Lerp(const FSimOutputData& InCurrent, const FSimOutputData& InNext, float Alpha) = 0;
diff --git a/Engine/Source/Runtime/Experimental/ChaosVehicles/ChaosVehiclesCore/Public/SimModule/SuspensionModule.h b/Engine/Source/Runtime/Experimental/ChaosVehicles/ChaosVehiclesCore/Public/SimModule/SuspensionModule.h
--- a/Engine/Source/Runtime/Experimental/ChaosVehicles/ChaosVehiclesCore/Public/SimModule/SuspensionModule.h
+++ b/Engine/Source/Runtime/Experimental/ChaosVehicles/ChaosVehiclesCore/Public/SimModule/SuspensionModule.h
@@ -46,6 +46,7 @@ namespace Chaos
 		virtual FSimOutputData* MakeNewData() override { return FSuspensionOutputData::MakeNew(); }
 		static FSimOutputData* MakeNew() { return new FSuspensionOutputData(); }
 
+		virtual eSimType GetType() override { return eSimType::Suspension; }
 		virtual void FillOutputState(const ISimulationModuleBase* SimModule) override;
 		virtual void Lerp(const FSimOutputData& InCurrent, const FSimOutputData& InNext, float Alpha) override;
 
@@ -65,10 +66,9 @@ namespace Chaos
 			, MaxLength(0.f)
 			, SpringRate(1.f)
 			, SpringPreload(0.5f)
-			, CompressionDamping(0.9f)
-			, ReboundDamping(0.9f)
+			, SpringDamping(0.9f)
+			, SuspensionForceEffect(100.0f)
 			//	, SwaybarEffect(0.5f)
-			//	, DampingRatio(0.3f)
 		{
 
 		}
@@ -81,12 +81,11 @@ namespace Chaos
 
 		float SpringRate;			// spring constant
 		float SpringPreload;		// Amount of Spring force (independent spring movement)
-		float CompressionDamping;	// limit compression speed
-		float ReboundDamping;		// limit rebound speed
+		float SpringDamping;		// limit compression/rebound speed
 
-		//	float Swaybar;				// Anti-roll bar
+		float SuspensionForceEffect; // force that presses the wheels into the ground - producing grip
 
-		//	float DampingRatio;			// value between (0-no damping) and (1-critical damping)
+		//	float Swaybar;				// Anti-roll bar
 	};
 
 	/** Suspension world ray/shape trace start and end positions */
@@ -153,6 +152,8 @@ namespace Chaos
 		void UpdateConstraint();
 
 		void SetSuspensionConstraint(FSuspensionConstraint* InConstraint);
+		void SetConstraintIndex(int32 InConstraintIndex) { ConstraintIndex = InConstraintIndex; }
+		int32 GetConstraintIndex() const { return ConstraintIndex; }
 		void SetTargetPoint(const FVector& InTargetPoint, const FVector& InImpactNormal, bool InWheelInContact)
 		{
 			TargetPos = InTargetPoint;
@@ -167,6 +168,7 @@ namespace Chaos
 		int WheelSimTreeIndex;
 
 		FSuspensionConstraint* Constraint;
+		int32 ConstraintIndex;
 		FVector TargetPos;
 		FVector ImpactNormal;
 		bool WheelInContact;
diff --git a/Engine/Source/Runtime/Experimental/ChaosVehicles/ChaosVehiclesCore/Public/SimModule/ThrusterModule.h b/Engine/Source/Runtime/Experimental/ChaosVehicles/ChaosVehiclesCore/Public/SimModule/ThrusterModule.h
--- a/Engine/Source/Runtime/Experimental/ChaosVehicles/ChaosVehiclesCore/Public/SimModule/ThrusterModule.h
+++ b/Engine/Source/Runtime/Experimental/ChaosVehicles/ChaosVehiclesCore/Public/SimModule/ThrusterModule.h
@@ -22,6 +22,8 @@ namespace Chaos
 			, ForceOffset(FVector::ZeroVector)
 			, MaxSteeringAngle(0)
 			, SteeringForceEffect(1.0f)
+			, BoostMultiplier(2.0f)
+			, MaxSpeed(125.0f)
 			, SteeringEnabled(false)
 		{
 
@@ -33,6 +35,8 @@ namespace Chaos
 		FVector ForceOffset;
 		float MaxSteeringAngle;
 		float SteeringForceEffect;
+		float BoostMultiplier;
+		float MaxSpeed;
 		bool SteeringEnabled;
 	};
 
diff --git a/Engine/Source/Runtime/Experimental/ChaosVehicles/ChaosVehiclesCore/Public/SimModule/TorqueSimModule.h b/Engine/Source/Runtime/Experimental/ChaosVehicles/ChaosVehiclesCore/Public/SimModule/TorqueSimModule.h
--- a/Engine/Source/Runtime/Experimental/ChaosVehicles/ChaosVehiclesCore/Public/SimModule/TorqueSimModule.h
+++ b/Engine/Source/Runtime/Experimental/ChaosVehicles/ChaosVehiclesCore/Public/SimModule/TorqueSimModule.h
@@ -72,7 +72,7 @@ namespace Chaos
 		void SetLoadTorque(float TorqueIn) { LoadTorque = TorqueIn; }
 		float GetLoadTorque() const { return LoadTorque; }
 
-		void SetBrakingTorque(float TorqueIn) { BrakingTorque = TorqueIn; }
+		void SetBrakingTorque(float TorqueIn) { ensure(TorqueIn >= 0.0f); BrakingTorque = TorqueIn; }
 		float GetBrakingTorque() const { return BrakingTorque; }
 
 		void SetAngularVelocity(float AngularVelocityIn) { AngularVelocity = AngularVelocityIn; }
diff --git a/Engine/Source/Runtime/Experimental/ChaosVehicles/ChaosVehiclesCore/Public/SimModule/TransmissionModule.h b/Engine/Source/Runtime/Experimental/ChaosVehicles/ChaosVehiclesCore/Public/SimModule/TransmissionModule.h
--- a/Engine/Source/Runtime/Experimental/ChaosVehicles/ChaosVehiclesCore/Public/SimModule/TransmissionModule.h
+++ b/Engine/Source/Runtime/Experimental/ChaosVehicles/ChaosVehiclesCore/Public/SimModule/TransmissionModule.h
@@ -47,6 +47,7 @@ namespace Chaos
 		virtual FSimOutputData* MakeNewData() override { return FTransmissionOutputData::MakeNew(); }
 		static FSimOutputData* MakeNew() { return new FTransmissionOutputData(); }
 
+		virtual eSimType GetType() override { return eSimType::Transmission; }
 		virtual void FillOutputState(const ISimulationModuleBase* SimModule) override;
 		virtual void Lerp(const FSimOutputData& InCurrent, const FSimOutputData& InNext, float Alpha) override;
 		virtual FString ToString() override;
diff --git a/Engine/Source/Runtime/Experimental/ChaosVehicles/ChaosVehiclesCore/Public/SimModule/WheelModule.h b/Engine/Source/Runtime/Experimental/ChaosVehicles/ChaosVehiclesCore/Public/SimModule/WheelModule.h
--- a/Engine/Source/Runtime/Experimental/ChaosVehicles/ChaosVehiclesCore/Public/SimModule/WheelModule.h
+++ b/Engine/Source/Runtime/Experimental/ChaosVehicles/ChaosVehiclesCore/Public/SimModule/WheelModule.h
@@ -45,6 +45,7 @@ namespace Chaos
 		virtual FSimOutputData* MakeNewData() override { return FWheelOutputData::MakeNew(); }
 		static FSimOutputData* MakeNew() { return new FWheelOutputData(); }
 
+		virtual eSimType GetType() override { return eSimType::Wheel; }
 		virtual void FillOutputState(const ISimulationModuleBase* SimModule) override;
 		virtual void Lerp(const FSimOutputData& InCurrent, const FSimOutputData& InNext, float Alpha) override;
 		virtual FString ToString() override;
@@ -149,6 +150,7 @@ namespace Chaos
 		float GetSteerAngleDegrees() const { return SteerAngleDegrees; }
 
 		void SetForceIntoSurface(float ForceIntoSurfaceIn) { ForceIntoSurface = ForceIntoSurfaceIn; }
+		float GetForceIntoSurface() const { return ForceIntoSurface; }
 
 		void SetSurfaceFriction(float FrictionIn) { SurfaceFriction = FrictionIn; }
 
