ensure that cluster union child to parent replication via the proxy component is authoriative and other sources (e.g. the GC proxy world transform update) don't try to fight that

- also fix an issue where the cluster union FClusterUnionManager::HandleAddOperation would do a remove and add clearing out the child to parent lock and causing that data to be potentially lost

diff --git a/Engine/Source/Runtime/Experimental/Chaos/Private/Chaos/ClusterUnionManager.cpp b/Engine/Source/Runtime/Experimental/Chaos/Private/Chaos/ClusterUnionManager.cpp
--- a/Engine/Source/Runtime/Experimental/Chaos/Private/Chaos/ClusterUnionManager.cpp
+++ b/Engine/Source/Runtime/Experimental/Chaos/Private/Chaos/ClusterUnionManager.cpp
@@ -514,7 +514,7 @@ namespace Chaos
 	}
 
 	DECLARE_CYCLE_STAT(TEXT("FClusterUnionManager::HandleAddOperation"), STAT_HandleAddOperation, STATGROUP_Chaos);
-	void FClusterUnionManager::HandleAddOperation(FClusterUnionIndex ClusterIndex, const TArray<FPBDRigidParticleHandle*>& Particles, bool bReleaseClustersFirst)
+	void FClusterUnionManager::HandleAddOperation(FClusterUnionIndex ClusterIndex, const TArray<FPBDRigidParticleHandle*>& InParticles, bool bReleaseClustersFirst)
 	{
 		SCOPE_CYCLE_COUNTER(STAT_HandleAddOperation);
 		FClusterUnion* Cluster = ClusterUnions.Find(ClusterIndex);
@@ -523,6 +523,16 @@ namespace Chaos
 			return;
 		}
 
+		// This ensures us that we only try to add particles that aren't already in the cluster already.
+		// TODO: There's probably a better way to do this without having to reallocate another array - but these arrays should be small.
+		TArray<FPBDRigidParticleHandle*> Particles = InParticles.FilterByPredicate(
+			[this, ClusterIndex](FPBDRigidParticleHandle* P)
+			{
+				const int32 CompareIndex = FindClusterUnionIndexFromParticle(P);
+				return CompareIndex != ClusterIndex;
+			}
+		);
+
 		// If we're adding particles to a cluster we need to first make sure they're not part of any other cluster.
 		// Book-keeping might get a bit odd if we try to add a particle to a new cluster and then only later remove the particle from its old cluster.
 		HandleRemoveOperationWithClusterLookup(Particles, EClusterUnionOperationTiming::Defer);
diff --git a/Engine/Source/Runtime/Experimental/Chaos/Private/PhysicsProxy/GeometryCollectionPhysicsProxy.cpp b/Engine/Source/Runtime/Experimental/Chaos/Private/PhysicsProxy/GeometryCollectionPhysicsProxy.cpp
--- a/Engine/Source/Runtime/Experimental/Chaos/Private/PhysicsProxy/GeometryCollectionPhysicsProxy.cpp
+++ b/Engine/Source/Runtime/Experimental/Chaos/Private/PhysicsProxy/GeometryCollectionPhysicsProxy.cpp
@@ -3273,7 +3273,10 @@ void FGeometryCollectionPhysicsProxy::SetWorldTransform_Internal(const FTransfor
 			}
 		}
 
-		if (ClusterUnionIndex != INDEX_NONE && !DeferredClusterUnionParticleUpdates.IsEmpty() && !DeferredClusterUnionChildToParentUpdates.IsEmpty())
+		// Should be safe enough to access on the PT.
+		const bool bIsAuthority = GetReplicationMode() == EReplicationMode::Server;
+		// This should only happen on the server otherwise the client may override the replicated child to parent before it's even set.
+		if (bIsAuthority && ClusterUnionIndex != INDEX_NONE && !DeferredClusterUnionParticleUpdates.IsEmpty() && !DeferredClusterUnionChildToParentUpdates.IsEmpty())
 		{
 			ClusterUnionManager.UpdateClusterUnionParticlesChildToParent(ClusterUnionIndex, DeferredClusterUnionParticleUpdates, DeferredClusterUnionChildToParentUpdates, false);
 		}
diff --git a/Engine/Source/Runtime/Experimental/Chaos/Public/Chaos/ClusterUnionManager.h b/Engine/Source/Runtime/Experimental/Chaos/Public/Chaos/ClusterUnionManager.h
--- a/Engine/Source/Runtime/Experimental/Chaos/Public/Chaos/ClusterUnionManager.h
+++ b/Engine/Source/Runtime/Experimental/Chaos/Public/Chaos/ClusterUnionManager.h
@@ -174,7 +174,7 @@ namespace Chaos
 		CHAOS_API void AddPendingClusterIndexOperation(FClusterUnionIndex Index, EClusterUnionOperation Op, const TArray<FPBDRigidParticleHandle*>& Particles);
 
 		// Actually performs the change specified in the FClusterUnionOperationData structure.
-		CHAOS_API void HandleAddOperation(FClusterUnionIndex ClusterIndex, const TArray<FPBDRigidParticleHandle*>& Particles, bool bReleaseClustersFirst);
+		CHAOS_API void HandleAddOperation(FClusterUnionIndex ClusterIndex, const TArray<FPBDRigidParticleHandle*>& InParticles, bool bReleaseClustersFirst);
 
 		// Removes the specified particles from the specified cluster.
 		CHAOS_API void HandleRemoveOperation(FClusterUnionIndex ClusterIndex, const TArray<FPBDRigidParticleHandle*>& Particles, EClusterUnionOperationTiming UpdateClusterPropertiesTiming);
