Iris - Fix compile warning/error.

#rnx

diff --git a/Engine/Source/Runtime/Experimental/Iris/Core/Private/Iris/ReplicationSystem/Prioritization/ReplicationPrioritization.cpp b/Engine/Source/Runtime/Experimental/Iris/Core/Private/Iris/ReplicationSystem/Prioritization/ReplicationPrioritization.cpp
--- a/Engine/Source/Runtime/Experimental/Iris/Core/Private/Iris/ReplicationSystem/Prioritization/ReplicationPrioritization.cpp
+++ b/Engine/Source/Runtime/Experimental/Iris/Core/Private/Iris/ReplicationSystem/Prioritization/ReplicationPrioritization.cpp
@@ -761,7 +761,7 @@ void FReplicationPrioritization::SetHighPriorityOnViewTargets(const TArrayView<f
 		}
 	}
 
-	for (const FNetHandle NetHandle : ViewTargets)
+	for (FNetHandle NetHandle : ViewTargets)
 	{
 		const FInternalNetRefIndex ViewTargetInternalIndex = NetRefHandleManager->GetInternalIndexFromNetHandle(NetHandle);
 		if (ViewTargetInternalIndex != FNetRefHandleManager::InvalidInternalIndex)
