The application crashes in-game: "Chaos::TEventContainer<Chaos::FCollisionEventData>::DispatchConsumerData".

diff --git a/Engine/Source/Runtime/Experimental/Chaos/Public/EventManager.h b/Engine/Source/Runtime/Experimental/Chaos/Public/EventManager.h
--- a/Engine/Source/Runtime/Experimental/Chaos/Public/EventManager.h
+++ b/Engine/Source/Runtime/Experimental/Chaos/Public/EventManager.h
@@ -323,9 +323,13 @@ namespace Chaos
 				for (auto& KeyValue : *Map) // Only iterating over objects that are associated with events here
 				{
 					const IPhysicsProxyBase* Proxy = KeyValue.Get<0>();
-					for (TMultiMap<UObject*, FEventHandlerPtr>::TConstKeyIterator It = ProxyOwnerToHandlerMap.CreateConstKeyIterator(Proxy->GetOwner()); It; ++It)
+					const UObject* Owner = Proxy->GetOwner();
+					if (ProxyOwnerToHandlerMap.Contains(Owner))
 					{
-						UniqueHandlers.Add(It.Value());
+						for (TMultiMap<UObject*, FEventHandlerPtr>::TConstKeyIterator It = ProxyOwnerToHandlerMap.CreateConstKeyIterator(Owner); It; ++It)
+						{
+							UniqueHandlers.Add(It.Value());
+						}
 					}
 				}
 
