cluster union manager should only not add a particle into a cluster union only if the particle itself is in the cluster union - not caring if its parents is in the cluster union

diff --git a/Engine/Source/Runtime/Experimental/Chaos/Private/Chaos/ClusterUnionManager.cpp b/Engine/Source/Runtime/Experimental/Chaos/Private/Chaos/ClusterUnionManager.cpp
--- a/Engine/Source/Runtime/Experimental/Chaos/Private/Chaos/ClusterUnionManager.cpp
+++ b/Engine/Source/Runtime/Experimental/Chaos/Private/Chaos/ClusterUnionManager.cpp
@@ -526,10 +526,19 @@ namespace Chaos
 		// This ensures us that we only try to add particles that aren't already in the cluster already.
 		// TODO: There's probably a better way to do this without having to reallocate another array - but these arrays should be small.
 		TArray<FPBDRigidParticleHandle*> Particles = InParticles.FilterByPredicate(
-			[this, ClusterIndex](FPBDRigidParticleHandle* P)
+			[this, Cluster](FPBDRigidParticleHandle* P)
 			{
-				const int32 CompareIndex = FindClusterUnionIndexFromParticle(P);
-				return CompareIndex != ClusterIndex;
+				if (!P)
+				{
+					return false;
+				}
+
+				if (FPBDRigidClusteredParticleHandle* ClusteredP = P->CastToClustered())
+				{
+					return Cluster->InternalCluster != ClusteredP->Parent();
+				}
+
+				return false;
 			}
 		);
 
