geometry collection should track its previously set world transform external so the GT can compute the difference between the two if needed (e.g. if it needs to query the difference prior to the PT->GT sync)

diff --git a/Engine/Source/Runtime/Experimental/Chaos/Private/PhysicsProxy/GeometryCollectionPhysicsProxy.cpp b/Engine/Source/Runtime/Experimental/Chaos/Private/PhysicsProxy/GeometryCollectionPhysicsProxy.cpp
--- a/Engine/Source/Runtime/Experimental/Chaos/Private/PhysicsProxy/GeometryCollectionPhysicsProxy.cpp
+++ b/Engine/Source/Runtime/Experimental/Chaos/Private/PhysicsProxy/GeometryCollectionPhysicsProxy.cpp
@@ -584,6 +584,7 @@ void FGeometryCollectionPhysicsProxy::Initialize(Chaos::FPBDRigidsEvolutionBase
 
 	// we need to make sure the world transform is kept up to date on the game thread 
 	WorldTransform_External = Parameters.WorldTransform;
+	PreviousWorldTransform_External = WorldTransform_External;
 
 	//
 	// Collision vertices down sampling validation.  
@@ -3187,6 +3188,7 @@ void FGeometryCollectionPhysicsProxy::SetWorldTransform_External(const FTransfor
 	if (bHasTransformChanged)
 	{
 		bIsGameThreadWorldTransformDirty = bHasTransformChanged;
+		PreviousWorldTransform_External = WorldTransform_External;
 		WorldTransform_External = WorldTransform;
 
 		ExecuteOnPhysicsThread(*this,
diff --git a/Engine/Source/Runtime/Experimental/Chaos/Public/PhysicsProxy/GeometryCollectionPhysicsProxy.h b/Engine/Source/Runtime/Experimental/Chaos/Public/PhysicsProxy/GeometryCollectionPhysicsProxy.h
--- a/Engine/Source/Runtime/Experimental/Chaos/Public/PhysicsProxy/GeometryCollectionPhysicsProxy.h
+++ b/Engine/Source/Runtime/Experimental/Chaos/Public/PhysicsProxy/GeometryCollectionPhysicsProxy.h
@@ -271,6 +271,7 @@ public:
 
 	// set the world transform ( this needs to be called on the game thread ) 
 	CHAOS_API void SetWorldTransform_External(const FTransform& WorldTransform);
+	CHAOS_API const FTransform& GetPreviousWorldTransform_External() const { return PreviousWorldTransform_External; }
 
 	// todo(chaos): Remove this and move to a cook time approach of the SM data based on the GC property
 	// Set whether the GC should be using collision from the Static Mesh or the GC itself for game thread traces ( this needs to be called on the game thread )
@@ -679,6 +680,7 @@ private:
 
 	// todo : we should probably keep a simulation parameter copy on the game thread instead 
 	FTransform WorldTransform_External;
+	FTransform PreviousWorldTransform_External;
 	uint8 bIsGameThreadWorldTransformDirty : 1;
 
 	uint8 bHasBuiltGeometryOnPT : 1;
diff --git a/Engine/Source/Runtime/Experimental/GeometryCollectionEngine/Private/GeometryCollection/GeometryCollectionComponent.cpp b/Engine/Source/Runtime/Experimental/GeometryCollectionEngine/Private/GeometryCollection/GeometryCollectionComponent.cpp
--- a/Engine/Source/Runtime/Experimental/GeometryCollectionEngine/Private/GeometryCollection/GeometryCollectionComponent.cpp
+++ b/Engine/Source/Runtime/Experimental/GeometryCollectionEngine/Private/GeometryCollection/GeometryCollectionComponent.cpp
@@ -7043,6 +7043,16 @@ bool UGeometryCollectionComponent::CanBeUsedInPhysicsReplication(const FName Bon
 	return false;
 }
 
+const FTransform& UGeometryCollectionComponent::GetPreviousComponentToWorld() const
+{
+	if (!PhysicsProxy)
+	{
+		return FTransform::Identity;
+	}
+
+	return PhysicsProxy->GetPreviousWorldTransform_External();
+}
+
 #if WITH_EDITOR
 
 bool UGeometryCollectionComponent::IsHLODRelevant() const
diff --git a/Engine/Source/Runtime/Experimental/GeometryCollectionEngine/Public/GeometryCollection/GeometryCollectionComponent.h b/Engine/Source/Runtime/Experimental/GeometryCollectionEngine/Public/GeometryCollection/GeometryCollectionComponent.h
--- a/Engine/Source/Runtime/Experimental/GeometryCollectionEngine/Public/GeometryCollection/GeometryCollectionComponent.h
+++ b/Engine/Source/Runtime/Experimental/GeometryCollectionEngine/Public/GeometryCollection/GeometryCollectionComponent.h
@@ -1348,7 +1348,7 @@ public:
 	GEOMETRYCOLLECTIONENGINE_API void SetReplicationAbandonAfterLevel(int32 InReplicationAbandonAfterLevel) { ReplicationAbandonAfterLevel = InReplicationAbandonAfterLevel; }
 	GEOMETRYCOLLECTIONENGINE_API void SetReplicationMaxPositionAndVelocityCorrectionLevel(int32 InReplicationMaxPositionAndVelocityCorrectionLevel) { ReplicationMaxPositionAndVelocityCorrectionLevel = InReplicationMaxPositionAndVelocityCorrectionLevel; }
 	
-
+	GEOMETRYCOLLECTIONENGINE_API const FTransform& GetPreviousComponentToWorld() const;
 public:
 	UPROPERTY(BlueprintAssignable, Category = "Collision")
 	FOnChaosPhysicsCollision OnChaosPhysicsCollision;
