revert previous change to call RegisterForEvents every time we change a notify global event flag - only set the flag directly on the PT directly for breaks

diff --git a/Engine/Source/Runtime/Experimental/Chaos/Private/PhysicsProxy/GeometryCollectionPhysicsProxy.cpp b/Engine/Source/Runtime/Experimental/Chaos/Private/PhysicsProxy/GeometryCollectionPhysicsProxy.cpp
--- a/Engine/Source/Runtime/Experimental/Chaos/Private/PhysicsProxy/GeometryCollectionPhysicsProxy.cpp
+++ b/Engine/Source/Runtime/Experimental/Chaos/Private/PhysicsProxy/GeometryCollectionPhysicsProxy.cpp
@@ -2846,6 +2846,11 @@ void FGeometryCollectionPhysicsProxy::SetNotifyGlobalBreakings_External(bool bNo
 		[this, bNotify]()
 		{
 			Parameters.bGenerateGlobalBreakingData = bNotify;
+
+			if (Chaos::FPBDRigidsSolver* Solver = GetSolver<Chaos::FPBDRigidsSolver>())
+			{
+				Solver->SetGenerateBreakingData(true);
+			}
 		});
 }
 
diff --git a/Engine/Source/Runtime/Experimental/GeometryCollectionEngine/Private/GeometryCollection/GeometryCollectionComponent.cpp b/Engine/Source/Runtime/Experimental/GeometryCollectionEngine/Private/GeometryCollection/GeometryCollectionComponent.cpp
--- a/Engine/Source/Runtime/Experimental/GeometryCollectionEngine/Private/GeometryCollection/GeometryCollectionComponent.cpp
+++ b/Engine/Source/Runtime/Experimental/GeometryCollectionEngine/Private/GeometryCollection/GeometryCollectionComponent.cpp
@@ -1172,7 +1172,6 @@ void UGeometryCollectionComponent::SetNotifyGlobalBreaks(bool bNewNotifyGlobalBr
 			PhysicsProxy->SetNotifyGlobalBreakings_External(bNewNotifyGlobalBreaks);
 		}
 		bNotifyGlobalBreaks = bNewNotifyGlobalBreaks;
-		RegisterForEvents();
 	}
 }
 
@@ -1182,7 +1181,6 @@ void UGeometryCollectionComponent::SetNotifyGlobalCollision(bool bNewNotifyGloba
 	{
 		bNotifyGlobalCollisions = bNewNotifyGlobalCollisions;
 		UpdateGlobalCollisionEventRegistration();
-		RegisterForEvents();
 	}
 }
 
@@ -1196,7 +1194,6 @@ void UGeometryCollectionComponent::SetNotifyGlobalRemovals(bool bNewNotifyGlobal
 		}
 		bNotifyGlobalRemovals = bNewNotifyGlobalRemovals;
 		UpdateGlobalRemovalEventRegistration();
-		RegisterForEvents();
 	}
 }
 
