Update Chaos::TEventContainer to use a map for storing handlers to optimize unregister calls.

When the container has several thousand handlers the linear search is very slow (several MS). Switching to a map increases this significantly.

We do lose a bit of speed in DispatchConsumerData as we need to iterate over all the events, but this didn't seem to cause a noticeable regression when profiling.

diff --git a/Engine/Source/Runtime/Experimental/Chaos/Public/EventManager.h b/Engine/Source/Runtime/Experimental/Chaos/Public/EventManager.h
--- a/Engine/Source/Runtime/Experimental/Chaos/Public/EventManager.h
+++ b/Engine/Source/Runtime/Experimental/Chaos/Public/EventManager.h
@@ -170,10 +170,10 @@ namespace Chaos
 		~TEventContainer()
 		{
 			HandlerLock.WriteLock();
-			for (FEventHandlerPtr Handler : HandlerArray)
+			for (TPair<void*, FEventHandlerPtr>& Pair : HandlerMap)
 			{
-				delete Handler;
-				Handler = nullptr;
+				delete Pair.Value;
+				Pair.Value = nullptr;
 			}
 			HandlerLock.WriteUnlock();
 		}
@@ -187,7 +187,7 @@ namespace Chaos
 			// so if we cannot lock it we store the event handler to try to register later
 			if (HandlerLock.TryWriteLock())
 			{
-				HandlerArray.AddUnique(Handler);
+				HandlerMap.Add(Handler->GetHandler(), Handler);
 				TArray<UObject*> ProxyOwners;
 				bool bValidProxyFilter = Handler->GetInterestedProxyOwners(ProxyOwners);
 
@@ -202,7 +202,7 @@ namespace Chaos
 				{
 					if (GetProxyToIndexMap(EventBuffer.Get()->GetConsumerBuffer()) != nullptr) // Only if our type supports getting the ProxyToIndexMap do we bother adding this
 					{
-						HandlersNotInMap.AddUnique(Handler);
+						HandlersNotInProxyOwnerMap.Add(Handler->GetHandler(), Handler);
 					}
 				}
 				HandlerLock.WriteUnlock();
@@ -239,23 +239,12 @@ namespace Chaos
 					ProxyOwnerToHandlerMap.Remove(KeyAndValue.Get<0>(), KeyAndValue.Get<1>());
 				}
 
-				for (int i = 0; i < HandlersNotInMap.Num(); i++)
-				{
-					if (HandlersNotInMap[i]->GetHandler() == InHandler)
-					{
-						HandlersNotInMap.RemoveAtSwap(i, 1, false);
-						break;
-					}
-				}
+				HandlersNotInProxyOwnerMap.Remove(InHandler);
 
-				for (int i = 0; i < HandlerArray.Num(); i++)
+				if (FEventHandlerPtr* HandlerPtr = HandlerMap.Find(InHandler))
 				{
-					if (HandlerArray[i]->GetHandler() == InHandler)
-					{
-						DeleteHandler(HandlerArray[i]);
-						HandlerArray.RemoveAtSwap(i, 1, false);
-						break;
-					}
+					DeleteHandler(*HandlerPtr);
+					HandlerMap.Remove(InHandler);
 				}
 				HandlerLock.WriteUnlock();
 			}
@@ -316,7 +305,7 @@ namespace Chaos
 			HandlerLock.ReadLock();
 			const TMap<IPhysicsProxyBase*, TArray<int32>>* Map = GetProxyToIndexMap(Buffer); // Use t his map to get all proxies used in the event buffer
 			// Only take this path if we have fewer Events than Handlers
-			if (Map && Map->Num() + HandlersNotInMap.Num() < HandlerArray.Num())
+			if (Map && Map->Num() + HandlersNotInProxyOwnerMap.Num() < HandlerMap.Num())
 			{
 				TSet<FEventHandlerPtr> UniqueHandlers;
 				UniqueHandlers.Reserve(Map->Num());
@@ -338,17 +327,17 @@ namespace Chaos
 					Handler->HandleEvent(Buffer);
 				}
 
-				for (const FEventHandlerPtr Handler : HandlersNotInMap)
+				for (const TPair<void*, FEventHandlerPtr>& Pair : HandlersNotInProxyOwnerMap)
 				{
-					Handler->HandleEvent(Buffer);
+					Pair.Value->HandleEvent(Buffer);
 				}
 			}
 			else
 			// This path is taken if there are fewer Handlers than events or the handler does not support GetProxyToIndexMap
 			{
-				for (const FEventHandlerPtr Handler : HandlerArray)
+				for (const TPair<void*, FEventHandlerPtr>& Pair : HandlerMap)
 				{
-					Handler->HandleEvent(Buffer);
+					Pair.Value->HandleEvent(Buffer);
 				}
 			}
 			HandlerLock.ReadUnlock();
@@ -392,15 +381,15 @@ private:
 		 */
 		TUniquePtr<IBufferResource<PayloadType>> EventBuffer;
 
-		TMultiMap<UObject*, FEventHandlerPtr> ProxyOwnerToHandlerMap; // Used to prevent us from iterating through the whole HandlerArray
-		TArray<FEventHandlerPtr> HandlersNotInMap; // Handlers not added to ProxyOwnerToHandlerMap since they do not support it
+		TMultiMap<UObject*, FEventHandlerPtr> ProxyOwnerToHandlerMap; // Used to prevent us from iterating through the whole HandlerMap
+		TMap<void*, FEventHandlerPtr> HandlersNotInProxyOwnerMap; // Handlers not added to ProxyOwnerToHandlerMap since they do not support it
 
 		/**
 		 * Delegate function registered to handle this event when it is dispatched
 		 */
-		TArray<FEventHandlerPtr> HandlerArray;
+		TMap<void*, FEventHandlerPtr> HandlerMap;
 
-		FRWLock HandlerLock; // protect access ProxyOwnerToHandlerMap, HandlersNotInMap, HandlerArray
+		FRWLock HandlerLock; // protect access ProxyOwnerToHandlerMap, HandlersNotInProxyOwnerMap, HandlerMap
 		TArray<FEventHandlerPtr> DeferredHandlers; // Store handler to register, they couldn't registered to avoid reentrant lock
 		TArray<const void*> DeferredUnregisterHandlers; // Store handler to unregister, they couldn't unregistered to avoid reentrant lock
 		FRWLock DeferredHandlerLock; // protect access to DeferredHandlers and DeferredUnregisterHandlers
