Add ISMPool option to not pay cost of ISM component update and ISMPool Tick associated with preallocation.

diff --git a/Engine/Source/Runtime/Experimental/GeometryCollectionEngine/Private/GeometryCollection/GeometryCollectionISMPoolComponent.cpp b/Engine/Source/Runtime/Experimental/GeometryCollectionEngine/Private/GeometryCollection/GeometryCollectionISMPoolComponent.cpp
--- a/Engine/Source/Runtime/Experimental/GeometryCollectionEngine/Private/GeometryCollection/GeometryCollectionISMPoolComponent.cpp
+++ b/Engine/Source/Runtime/Experimental/GeometryCollectionEngine/Private/GeometryCollection/GeometryCollectionISMPoolComponent.cpp
@@ -640,6 +640,19 @@ void UGeometryCollectionISMPoolComponent::PreallocateMeshInstance(const FGeometr
 	Pool.RequestPreallocateMeshInstance(MeshInstance);
 }
 
+void UGeometryCollectionISMPoolComponent::SetTickablePoolManagement(bool bEnablePoolManagement)
+{
+	if (!bEnablePoolManagement)
+	{
+		// Disable the keep alive and recycle pool management systems.
+		// This also disables preallocation for this pool.
+		Pool.bCachedKeepAlive = false;
+		Pool.bCachedRecycle = false;
+	}
+	// Disable the Tick that is used to manage the pool.
+	PrimaryComponentTick.SetTickFunctionEnable(bEnablePoolManagement);
+}
+
 void UGeometryCollectionISMPoolComponent::SetOverrideTransformUpdates(bool bOverrideUpdates)
 {
 	Pool.bDisableBoundsAndTransformUpdate = bOverrideUpdates;
diff --git a/Engine/Source/Runtime/Experimental/GeometryCollectionEngine/Public/GeometryCollection/GeometryCollectionISMPoolComponent.h b/Engine/Source/Runtime/Experimental/GeometryCollectionEngine/Public/GeometryCollection/GeometryCollectionISMPoolComponent.h
--- a/Engine/Source/Runtime/Experimental/GeometryCollectionEngine/Public/GeometryCollection/GeometryCollectionISMPoolComponent.h
+++ b/Engine/Source/Runtime/Experimental/GeometryCollectionEngine/Public/GeometryCollection/GeometryCollectionISMPoolComponent.h
@@ -368,6 +368,8 @@ public:
 	 */
 	GEOMETRYCOLLECTIONENGINE_API void PreallocateMeshInstance(const FGeometryCollectionStaticMeshInstance& MeshInstance);
 
+	GEOMETRYCOLLECTIONENGINE_API void SetTickablePoolManagement(bool bEnablePoolManagement);
+
 	GEOMETRYCOLLECTIONENGINE_API void SetOverrideTransformUpdates(bool bOverrideUpdates);
 
 	GEOMETRYCOLLECTIONENGINE_API void UpdateAbsoluteTransforms(const FTransform& BaseTransform, EUpdateTransformFlags UpdateTransformFlags, ETeleportType Teleport);
