Iris
* Added ensure when RenamePathForPie fails for objects with invalid names

#rnx

diff --git a/Engine/Source/Runtime/Experimental/Iris/Core/Private/Iris/ReplicationSystem/ObjectReferenceCache.cpp b/Engine/Source/Runtime/Experimental/Iris/Core/Private/Iris/ReplicationSystem/ObjectReferenceCache.cpp
--- a/Engine/Source/Runtime/Experimental/Iris/Core/Private/Iris/ReplicationSystem/ObjectReferenceCache.cpp
+++ b/Engine/Source/Runtime/Experimental/Iris/Core/Private/Iris/ReplicationSystem/ObjectReferenceCache.cpp
@@ -360,6 +360,11 @@ bool FObjectReferenceCache::CreateObjectReferenceInternal(const UObject* Object,
 		constexpr bool bReading = false;
 		// The connection ID isn't used unless reading.
 		RenamePathForPie(UE::Net::InvalidConnectionId, ObjectPath, bReading);
+
+#if WITH_EDITOR
+		ensureMsgf(!ObjectPath.IsEmpty(), TEXT("NetworkRemapPath found %s to be an invalid name for %s. This object will not replicate!"), *Object->GetName(), *GetPathNameSafe(Object));
+#endif
+
 		PathToken = StringTokenStore->GetOrCreateToken(ObjectPath);
 	}
 
@@ -783,6 +788,10 @@ UObject* FObjectReferenceCache::ResolveObjectReferenceHandleInternal(FNetRefHand
 	constexpr bool bReading = true;
 	RenamePathForPie(ResolveContext.ConnectionId, ObjectPath, bReading);
 
+#if WITH_EDITOR
+	ensureMsgf(!ObjectPath.IsEmpty(), TEXT("NetworkRemapPath found %s to be an invalid name. This object will not be binded and replicated!"), ResolvedToken);
+#endif
+
 	const FName ObjectPathName(ObjectPath);
 
 	// See if this object is in memory
diff --git a/Engine/Source/Runtime/Experimental/Iris/Core/Private/Iris/ReplicationSystem/ObjectReplicationBridge.cpp b/Engine/Source/Runtime/Experimental/Iris/Core/Private/Iris/ReplicationSystem/ObjectReplicationBridge.cpp
--- a/Engine/Source/Runtime/Experimental/Iris/Core/Private/Iris/ReplicationSystem/ObjectReplicationBridge.cpp
+++ b/Engine/Source/Runtime/Experimental/Iris/Core/Private/Iris/ReplicationSystem/ObjectReplicationBridge.cpp
@@ -652,12 +652,39 @@ FReplicationBridgeCreateNetRefHandleResult UObjectReplicationBridge::CreateNetRe
 
 	FReplicationBridgeCreateNetRefHandleResult CreateResult;
 
+	if (UE_LOG_ACTIVE(LogIrisBridge, Verbose))
+	{
+		if (RootObjectOfSubObject.IsValid())
+		{
+			UE_LOG(LogIrisBridge, Verbose, TEXT("CreateNetRefHandleFromRemote: SubObject: %s of RootObject: %s"), *WantedNetHandle.ToString(), *RootObjectOfSubObject.ToString());
+		}
+		else
+		{
+			UE_LOG(LogIrisBridge, Verbose, TEXT("CreateNetRefHandleFromRemote: RootObject: %s"), *WantedNetHandle.ToString());
+		}
+	}
+	
+	
+	
+
 	// Currently we need to always instantiate remote objects, moving forward we want to make this optional so that can be deferred until it is time to apply received state data.
 	// https://jira.it.epicgames.com/browse/UE-127369	
 	FObjectReplicationBridgeInstantiateResult InstantiateResult = BeginInstantiateFromRemote(RootObjectOfSubObject, Context.SerializationContext.GetInternalContext()->ResolveContext, Header.Get());
 	UObject* InstancePtr = InstantiateResult.Object;
 	if (!InstancePtr)
 	{
+		if (UE_LOG_ACTIVE(LogIrisBridge, Warning) && !bSuppressCreateInstanceFailedEnsure)
+		{
+			if (RootObjectOfSubObject.IsValid())
+			{
+				UE_LOG(LogIrisBridge, Warning, TEXT("CreateNetRefHandleFromRemote: Failed to instantiate SubObject NetHandle: %s of RootObject: %s (%s)"), *WantedNetHandle.ToString(), *RootObjectOfSubObject.ToString(), *GetNameSafe(GetReplicatedObject(RootObjectOfSubObject)));
+			}
+			else
+			{
+				UE_LOG(LogIrisBridge, Warning, TEXT("CreateNetRefHandleFromRemote: Failed to instantiate RootObject NetHandle: %s"), *WantedNetHandle.ToString());
+			}
+		}
+		
 		ensureMsgf(bSuppressCreateInstanceFailedEnsure, TEXT("Failed to instantiate Handle: %s"), *WantedNetHandle.ToString());
 		return CreateResult;
 	}
