Iris - Added Rotator3f and Rotator3d NetSerializers.

diff --git a/Engine/Plugins/Runtime/ReplicationSystemTestPlugin/Source/Private/Tests/Serialization/TestRotatorNetSerializer.cpp b/Engine/Plugins/Runtime/ReplicationSystemTestPlugin/Source/Private/Tests/Serialization/TestRotatorNetSerializer.cpp
--- a/Engine/Plugins/Runtime/ReplicationSystemTestPlugin/Source/Private/Tests/Serialization/TestRotatorNetSerializer.cpp
+++ b/Engine/Plugins/Runtime/ReplicationSystemTestPlugin/Source/Private/Tests/Serialization/TestRotatorNetSerializer.cpp
@@ -7,7 +7,12 @@
 namespace UE::Net
 {
 
-FTestMessage& operator<<(FTestMessage& Message, const FRotator& Rotator)
+FTestMessage& operator<<(FTestMessage& Message, const FRotator3f& Rotator)
+{
+	return Message << Rotator.ToString();
+}
+
+FTestMessage& operator<<(FTestMessage& Message, const FRotator3d& Rotator)
 {
 	return Message << Rotator.ToString();
 }
@@ -22,10 +27,11 @@ static FTestMessage& PrintRotatorNetSerializerConfig(FTestMessage& Message, cons
 	return Message;
 }
 
-class FTestRotatorNetSerializerBase : public TTestNetSerializerFixture<PrintRotatorNetSerializerConfig, FRotator>
+template<typename RotatorType>
+class FTestRotatorNetSerializerBase : public TTestNetSerializerFixture<PrintRotatorNetSerializerConfig, RotatorType>
 {
 private:
-	typedef TTestNetSerializerFixture<PrintRotatorNetSerializerConfig, FRotator> Super;
+	typedef TTestNetSerializerFixture<PrintRotatorNetSerializerConfig, RotatorType> Super;
 
 public:
 	FTestRotatorNetSerializerBase(const FNetSerializer& Serializer) : Super(Serializer) {}
@@ -36,7 +42,7 @@ public:
 	void TestValidate();
 
 protected:
-	using ScalarType = decltype(FRotator::Pitch);
+	using ScalarType = decltype(RotatorType::Pitch);
 
 	struct FFloatTriplet
 	{
@@ -47,14 +53,15 @@ protected:
 
 	static const FFloatTriplet BadValues[];
 	static const SIZE_T BadValueCount;
-	static const FRotator Values[];
+	static const RotatorType Values[];
 	static const SIZE_T ValueCount;
 	float ValueTolerance = 0.0f;
 	const FNetSerializerConfig* NetSerializerConfig = nullptr;
 };
 
 // In case there's any sort of sanity checking in the FRotator constructors this will avoid them.
-const FTestRotatorNetSerializerBase::FFloatTriplet FTestRotatorNetSerializerBase::BadValues[] =
+template<typename RotatorType>
+const typename FTestRotatorNetSerializerBase<RotatorType>::FFloatTriplet FTestRotatorNetSerializerBase<RotatorType>::BadValues[] =
 {
 	{-INFINITY, NAN, -NAN},
 	{-INFINITY, INFINITY, -INFINITY},
@@ -67,53 +74,79 @@ const FTestRotatorNetSerializerBase::FFloatTriplet FTestRotatorNetSerializerBase
 	{std::numeric_limits<ScalarType>::max(), std::numeric_limits<ScalarType>::max(), std::numeric_limits<ScalarType>::max()},
 	{std::numeric_limits<ScalarType>::lowest(), std::numeric_limits<ScalarType>::lowest(), std::numeric_limits<ScalarType>::lowest()},
 };
-const SIZE_T FTestRotatorNetSerializerBase::BadValueCount = sizeof(BadValues)/sizeof(BadValues[0]);
-
-const FRotator FTestRotatorNetSerializerBase::Values[] =
-{
-	FRotator(0.0f),
-	FRotator(1.0f, 0.0f, 0.0f),
-	FRotator(0.0f, 1.0f, 0.0f),
-	FRotator(0.0f, 0.0f, 1.0f),
-	FRotator(-0.0f), // 0.0f == -0.0f so we expect this to be handled gracefully
-	FRotator(359.999f, 129.130f, 45.0f),
-	FRotator(1.26698556E-23),
-	FRotator(47.11f),
+
+template<typename RotatorType>
+const SIZE_T FTestRotatorNetSerializerBase<RotatorType>::BadValueCount = sizeof(BadValues)/sizeof(BadValues[0]);
+
+template<typename RotatorType>
+const RotatorType FTestRotatorNetSerializerBase<RotatorType>::Values[] =
+{
+	RotatorType(0.0f),
+	RotatorType(1.0f, 0.0f, 0.0f),
+	RotatorType(0.0f, 1.0f, 0.0f),
+	RotatorType(0.0f, 0.0f, 1.0f),
+	RotatorType(-0.0f), // 0.0f == -0.0f so we expect this to be handled gracefully
+	RotatorType(359.999f, 129.130f, 45.0f),
+	RotatorType(1.26698556E-23),
+	RotatorType(47.11f),
 };
-const SIZE_T FTestRotatorNetSerializerBase::ValueCount = sizeof(Values)/sizeof(Values[0]);
+template<typename RotatorType>
+const SIZE_T FTestRotatorNetSerializerBase<RotatorType>::ValueCount = sizeof(Values)/sizeof(Values[0]);
 
-class FTestRotatorAsShortNetSerializer : public FTestRotatorNetSerializerBase
+class FTestRotatorAsShortNetSerializer : public FTestRotatorNetSerializerBase<FRotator>
 {
 private:
-	typedef FTestRotatorNetSerializerBase Super;
+	typedef FTestRotatorNetSerializerBase<FRotator> Super;
 
 public:
 	FTestRotatorAsShortNetSerializer()
 	: Super(UE_NET_GET_SERIALIZER(FRotatorAsShortNetSerializer))
 	{
 		ValueTolerance = 1.5f*(360.0f/65536.0f);
-		NetSerializerConfig = &RotatorAsShortNetSerializerConfig;
+		NetSerializerConfig = UE_NET_GET_SERIALIZER_DEFAULT_CONFIG(FRotatorAsShortNetSerializer);
 	}
-
-private:
-	inline static const FRotatorAsShortNetSerializerConfig RotatorAsShortNetSerializerConfig;
 };
 
-class FTestRotatorAsByteNetSerializer : public FTestRotatorNetSerializerBase
+class FTestRotatorAsByteNetSerializer : public FTestRotatorNetSerializerBase<FRotator>
 {
 private:
-	typedef FTestRotatorNetSerializerBase Super;
+	typedef FTestRotatorNetSerializerBase<FRotator> Super;
 
 public:
 	FTestRotatorAsByteNetSerializer()
 	: Super(UE_NET_GET_SERIALIZER(FRotatorAsByteNetSerializer))
 	{
 		ValueTolerance = 1.5f*(360.0f/256.0f);
-		NetSerializerConfig = &RotatorAsByteNetSerializerConfig;
+		NetSerializerConfig = UE_NET_GET_SERIALIZER_DEFAULT_CONFIG(FRotatorAsByteNetSerializer);
 	}
+};
 
+class FTestRotator3fNetSerializer : public FTestRotatorNetSerializerBase<FRotator3f>
+{
 private:
-	inline static const FRotatorAsByteNetSerializerConfig RotatorAsByteNetSerializerConfig;
+	typedef FTestRotatorNetSerializerBase<FRotator3f> Super;
+
+public:
+	FTestRotator3fNetSerializer()
+	: Super(UE_NET_GET_SERIALIZER(FRotator3fNetSerializer))
+	{
+		ValueTolerance = 1.5f*(360.0f/65536.0f);
+		NetSerializerConfig = UE_NET_GET_SERIALIZER_DEFAULT_CONFIG(FRotator3fNetSerializer);
+	}
+};
+
+class FTestRotator3dNetSerializer : public FTestRotatorNetSerializerBase<FRotator3d>
+{
+private:
+	typedef FTestRotatorNetSerializerBase<FRotator3d> Super;
+
+public:
+	FTestRotator3dNetSerializer()
+	: Super(UE_NET_GET_SERIALIZER(FRotator3dNetSerializer))
+	{
+		ValueTolerance = 1.5f*(360.0f/65536.0f);
+		NetSerializerConfig = UE_NET_GET_SERIALIZER_DEFAULT_CONFIG(FRotator3dNetSerializer);
+	}
 };
 
 // FRotatorAsShortNetSerializer tests
@@ -158,20 +191,63 @@ UE_NET_TEST_FIXTURE(FTestRotatorAsByteNetSerializer, TestValidate)
 	TestValidate();
 }
 
-void FTestRotatorNetSerializerBase::TestIsEqual()
+// FRotator3fNetSerializer tests
+UE_NET_TEST_FIXTURE(FTestRotator3fNetSerializer, TestIsEqual)
+{
+	TestIsEqual();
+}
+
+UE_NET_TEST_FIXTURE(FTestRotator3fNetSerializer, TestSerialize)
+{
+	TestSerialize();
+}
+
+UE_NET_TEST_FIXTURE(FTestRotator3fNetSerializer, TestSerializeDelta)
+{
+	TestSerializeDelta();
+}
+
+UE_NET_TEST_FIXTURE(FTestRotator3fNetSerializer, TestValidate)
+{
+	TestValidate();
+}
+
+// FRotator3dNetSerializer tests
+UE_NET_TEST_FIXTURE(FTestRotator3dNetSerializer, TestIsEqual)
+{
+	TestIsEqual();
+}
+
+UE_NET_TEST_FIXTURE(FTestRotator3dNetSerializer, TestSerialize)
+{
+	TestSerialize();
+}
+
+UE_NET_TEST_FIXTURE(FTestRotator3dNetSerializer, TestSerializeDelta)
+{
+	TestSerializeDelta();
+}
+
+UE_NET_TEST_FIXTURE(FTestRotator3dNetSerializer, TestValidate)
+{
+	TestValidate();
+}
+
+template<typename RotatorType>
+void FTestRotatorNetSerializerBase<RotatorType>::TestIsEqual()
 {
-	FRotator CompareValues[2][sizeof(Values)/sizeof(Values[0])];
+	RotatorType CompareValues[2][sizeof(Values)/sizeof(Values[0])];
 	bool ExpectedResults[2][sizeof(Values)/sizeof(Values[0])];
 
 	for (SIZE_T ValueIt = 0; ValueIt < ValueCount; ++ValueIt)
 	{
 		// We need to memcopy the data to avoid any validation/fixing in the constructor
-		FPlatformMemory::Memcpy(&CompareValues[0][ValueIt], &Values[ValueIt], sizeof(FRotator));
+		FPlatformMemory::Memcpy(&CompareValues[0][ValueIt], &Values[ValueIt], sizeof(RotatorType));
 		ExpectedResults[0][ValueIt] = true;
 	}
 	for (SIZE_T ValueIt = 0; ValueIt < ValueCount; ++ValueIt)
 	{
-		FPlatformMemory::Memcpy(&CompareValues[1][ValueIt], &Values[(ValueIt + 1) % ValueCount], sizeof(FRotator));
+		FPlatformMemory::Memcpy(&CompareValues[1][ValueIt], &Values[(ValueIt + 1) % ValueCount], sizeof(RotatorType));
 		ExpectedResults[1][ValueIt] = false;
 	}
 
@@ -186,24 +262,27 @@ void FTestRotatorNetSerializerBase::TestIsEqual()
 	}
 }
 
-void FTestRotatorNetSerializerBase::TestSerialize()
+template<typename RotatorType>
+void FTestRotatorNetSerializerBase<RotatorType>::TestSerialize()
 {
 	constexpr bool bQuantizedCompare = true;
 	Super::TestSerialize(Values, Values, ValueCount, *NetSerializerConfig, bQuantizedCompare);
 }
 
-void FTestRotatorNetSerializerBase::TestSerializeDelta()
+template<typename RotatorType>
+void FTestRotatorNetSerializerBase<RotatorType>::TestSerializeDelta()
 {
 	Super::TestSerializeDelta(Values, ValueCount, *NetSerializerConfig);
 }
 
-void FTestRotatorNetSerializerBase::TestValidate()
+template<typename RotatorType>
+void FTestRotatorNetSerializerBase<RotatorType>::TestValidate()
 {
 	// Test bad values
 	{
-		static_assert(sizeof(BadValues[0]) == sizeof(FRotator), "BadValues and FRotator type mismatch.");
+		static_assert(sizeof(BadValues[0]) == sizeof(RotatorType), "BadValues and RotatorType type mismatch.");
 		bool ExpectedResults[BadValueCount] = {};
-		const bool bSuccess = Super::TestValidate(reinterpret_cast<const FRotator*>(BadValues), ExpectedResults, BadValueCount, *NetSerializerConfig);
+		const bool bSuccess = Super::TestValidate(reinterpret_cast<const RotatorType*>(BadValues), ExpectedResults, BadValueCount, *NetSerializerConfig);
 		if (!bSuccess)
 		{
 			return;
diff --git a/Engine/Source/Runtime/Experimental/Iris/Core/Private/Iris/Serialization/RotatorNetSerializers.cpp b/Engine/Source/Runtime/Experimental/Iris/Core/Private/Iris/Serialization/RotatorNetSerializers.cpp
--- a/Engine/Source/Runtime/Experimental/Iris/Core/Private/Iris/Serialization/RotatorNetSerializers.cpp
+++ b/Engine/Source/Runtime/Experimental/Iris/Core/Private/Iris/Serialization/RotatorNetSerializers.cpp
@@ -8,10 +8,9 @@
 namespace UE::Net
 {
 
-struct FRotatorAsShortNetSerializer
+template<typename RotatorType>
+struct FRotatorAsShortNetSerializerBase
 {
-	static constexpr uint32 Version = 0;
-
 	struct FQuantizedType
 	{
 		// Using three bits to indicate whether the components are zero or not.
@@ -21,11 +20,8 @@ struct FRotatorAsShortNetSerializer
 		uint16 Z;
 	};
 
-	typedef FRotator SourceType;
+	typedef RotatorType SourceType;
 	typedef FQuantizedType QuantizedType;
-	typedef FRotatorAsShortNetSerializerConfig ConfigType;
-
-	inline static const ConfigType DefaultConfig;
 
 	static void Serialize(FNetSerializationContext&, const FNetSerializeArgs& Args);
 	static void Deserialize(FNetSerializationContext&, const FNetDeserializeArgs& Args);
@@ -55,22 +51,10 @@ protected:
 	static constexpr ScalarType Bias = ScalarType(0.5f);
 	static constexpr ScalarType InvScale = ScalarType(360)/ScalarType(65536);
 };
-UE_NET_IMPLEMENT_SERIALIZER(FRotatorAsShortNetSerializer);
-
-struct FRotatorNetSerializer : public FRotatorAsShortNetSerializer
-{
-	static constexpr uint32 Version = 0;
-
-	typedef FRotatorNetSerializerConfig ConfigType;
-
-	inline static const ConfigType DefaultConfig;
-};
-UE_NET_IMPLEMENT_SERIALIZER(FRotatorNetSerializer);
 
-struct FRotatorAsByteNetSerializer
+template<typename RotatorType>
+struct FRotatorAsByteNetSerializerBase
 {
-	static constexpr uint32 Version = 0;
-
 	struct FQuantizedType
 	{
 		// Using three bits to indicate whether the components are zero or not.
@@ -80,11 +64,8 @@ struct FRotatorAsByteNetSerializer
 		uint8 Z;
 	};
 
-	typedef FRotator SourceType;
+	typedef RotatorType SourceType;
 	typedef FQuantizedType QuantizedType;
-	typedef FRotatorAsByteNetSerializerConfig ConfigType;
-
-	inline static const ConfigType DefaultConfig;
 
 	static void Serialize(FNetSerializationContext&, const FNetSerializeArgs& Args);
 	static void Deserialize(FNetSerializationContext&, const FNetDeserializeArgs& Args);
@@ -114,10 +95,10 @@ protected:
 	static constexpr ScalarType Bias = ScalarType(0.5f);
 	static constexpr ScalarType InvScale = ScalarType(360)/ScalarType(256);
 };
-UE_NET_IMPLEMENT_SERIALIZER(FRotatorAsByteNetSerializer);
 
-// FRotatorAsShortNetSerializer implementation
-void FRotatorAsShortNetSerializer::Serialize(FNetSerializationContext& Context, const FNetSerializeArgs& Args)
+// FRotatorAsShortNetSerializerBase implementation
+template<typename T>
+void FRotatorAsShortNetSerializerBase<T>::Serialize(FNetSerializationContext& Context, const FNetSerializeArgs& Args)
 {
 	const QuantizedType& Value = *reinterpret_cast<const QuantizedType*>(Args.Source);
 
@@ -137,7 +118,8 @@ void FRotatorAsShortNetSerializer::Serialize(FNetSerializationContext& Context,
 	}
 }
 
-void FRotatorAsShortNetSerializer::Deserialize(FNetSerializationContext& Context, const FNetDeserializeArgs& Args)
+template<typename T>
+void FRotatorAsShortNetSerializerBase<T>::Deserialize(FNetSerializationContext& Context, const FNetDeserializeArgs& Args)
 {
 	FNetBitStreamReader* Reader = Context.GetBitStreamReader();
 
@@ -152,7 +134,8 @@ void FRotatorAsShortNetSerializer::Deserialize(FNetSerializationContext& Context
 	Target = Value;
 }
 
-void FRotatorAsShortNetSerializer::SerializeDelta(FNetSerializationContext& Context, const FNetSerializeDeltaArgs& Args)
+template<typename T>
+void FRotatorAsShortNetSerializerBase<T>::SerializeDelta(FNetSerializationContext& Context, const FNetSerializeDeltaArgs& Args)
 {
 	// Per component equality. With the current scaling a 10 degree change would require 12 bits to replicate
 	// due to 11 bits for the value and 1 for the sign. Add a small lookup table index for that and you're
@@ -186,7 +169,8 @@ void FRotatorAsShortNetSerializer::SerializeDelta(FNetSerializationContext& Cont
 	}
 }
 
-void FRotatorAsShortNetSerializer::DeserializeDelta(FNetSerializationContext& Context, const FNetDeserializeDeltaArgs& Args)
+template<typename T>
+void FRotatorAsShortNetSerializerBase<T>::DeserializeDelta(FNetSerializationContext& Context, const FNetDeserializeDeltaArgs& Args)
 {
 	FNetBitStreamReader* Reader = Context.GetBitStreamReader();
 
@@ -221,7 +205,8 @@ void FRotatorAsShortNetSerializer::DeserializeDelta(FNetSerializationContext& Co
 	Target = TempValue;
 }
 
-void FRotatorAsShortNetSerializer::Quantize(FNetSerializationContext&, const FNetQuantizeArgs& Args)
+template<typename T>
+void FRotatorAsShortNetSerializerBase<T>::Quantize(FNetSerializationContext&, const FNetQuantizeArgs& Args)
 {
 	const SourceType& Source = *reinterpret_cast<const SourceType*>(Args.Source);
 	QuantizedType& Target = *reinterpret_cast<QuantizedType*>(Args.Target);
@@ -237,7 +222,8 @@ void FRotatorAsShortNetSerializer::Quantize(FNetSerializationContext&, const FNe
 	Target = TempValue;
 }
 
-void FRotatorAsShortNetSerializer::Dequantize(FNetSerializationContext&, const FNetDequantizeArgs& Args)
+template<typename T>
+void FRotatorAsShortNetSerializerBase<T>::Dequantize(FNetSerializationContext&, const FNetDequantizeArgs& Args)
 {
 	const QuantizedType& Source = *reinterpret_cast<const QuantizedType*>(Args.Source);
 	SourceType& Target = *reinterpret_cast<SourceType*>(Args.Target);
@@ -250,7 +236,8 @@ void FRotatorAsShortNetSerializer::Dequantize(FNetSerializationContext&, const F
 	Target = TempValue;
 }
 
-bool FRotatorAsShortNetSerializer::IsEqual(FNetSerializationContext& Context, const FNetIsEqualArgs& Args)
+template<typename T>
+bool FRotatorAsShortNetSerializerBase<T>::IsEqual(FNetSerializationContext& Context, const FNetIsEqualArgs& Args)
 {
 	if (Args.bStateIsQuantized)
 	{
@@ -275,15 +262,17 @@ bool FRotatorAsShortNetSerializer::IsEqual(FNetSerializationContext& Context, co
 	}
 }
 
-bool FRotatorAsShortNetSerializer::Validate(FNetSerializationContext& Context, const FNetValidateArgs& Args)
+template<typename T>
+bool FRotatorAsShortNetSerializerBase<T>::Validate(FNetSerializationContext& Context, const FNetValidateArgs& Args)
 {
 	const SourceType& Value = *reinterpret_cast<SourceType*>(Args.Source);
 	// Make sure values are valid degree angles. This should catch NaNs as well.
-	return (Value.Pitch >= 0.0f && Value.Pitch < 360.0f) && (Value.Yaw >= 0.0f && Value.Yaw < 360.0f) & (Value.Roll >= 0.0f && Value.Roll < 360.0f); 
+	return (Value.Pitch >= ScalarType(0) && Value.Pitch < ScalarType(360)) && (Value.Yaw >= ScalarType(0) && Value.Yaw < ScalarType(360)) & (Value.Roll >= ScalarType(0) && Value.Roll < ScalarType(360)); 
 }
 
-// FRotatorAsByteNetSerializer implementation
-void FRotatorAsByteNetSerializer::Serialize(FNetSerializationContext& Context, const FNetSerializeArgs& Args)
+// FRotatorAsByteNetSerializerBase implementation
+template<typename T>
+void FRotatorAsByteNetSerializerBase<T>::Serialize(FNetSerializationContext& Context, const FNetSerializeArgs& Args)
 {
 	const QuantizedType& Value = *reinterpret_cast<const QuantizedType*>(Args.Source);
 
@@ -303,7 +292,8 @@ void FRotatorAsByteNetSerializer::Serialize(FNetSerializationContext& Context, c
 	}
 }
 
-void FRotatorAsByteNetSerializer::Deserialize(FNetSerializationContext& Context, const FNetDeserializeArgs& Args)
+template<typename T>
+void FRotatorAsByteNetSerializerBase<T>::Deserialize(FNetSerializationContext& Context, const FNetDeserializeArgs& Args)
 {
 	FNetBitStreamReader* Reader = Context.GetBitStreamReader();
 
@@ -318,7 +308,8 @@ void FRotatorAsByteNetSerializer::Deserialize(FNetSerializationContext& Context,
 	Target = Value;
 }
 
-void FRotatorAsByteNetSerializer::SerializeDelta(FNetSerializationContext& Context, const FNetSerializeDeltaArgs& Args)
+template<typename T>
+void FRotatorAsByteNetSerializerBase<T>::SerializeDelta(FNetSerializationContext& Context, const FNetSerializeDeltaArgs& Args)
 {
 	const QuantizedType& Value = *reinterpret_cast<QuantizedType*>(Args.Source);
 	const QuantizedType& PrevValue = *reinterpret_cast<QuantizedType*>(Args.Prev);
@@ -344,7 +335,8 @@ void FRotatorAsByteNetSerializer::SerializeDelta(FNetSerializationContext& Conte
 	}
 }
 
-void FRotatorAsByteNetSerializer::DeserializeDelta(FNetSerializationContext& Context, const FNetDeserializeDeltaArgs& Args)
+template<typename T>
+void FRotatorAsByteNetSerializerBase<T>::DeserializeDelta(FNetSerializationContext& Context, const FNetDeserializeDeltaArgs& Args)
 {
 	FNetBitStreamReader* Reader = Context.GetBitStreamReader();
 
@@ -376,7 +368,8 @@ void FRotatorAsByteNetSerializer::DeserializeDelta(FNetSerializationContext& Con
 	Target = TempValue;
 }
 
-void FRotatorAsByteNetSerializer::Quantize(FNetSerializationContext&, const FNetQuantizeArgs& Args)
+template<typename T>
+void FRotatorAsByteNetSerializerBase<T>::Quantize(FNetSerializationContext&, const FNetQuantizeArgs& Args)
 {
 	const SourceType& Source = *reinterpret_cast<const SourceType*>(Args.Source);
 	QuantizedType& Target = *reinterpret_cast<QuantizedType*>(Args.Target);
@@ -392,7 +385,8 @@ void FRotatorAsByteNetSerializer::Quantize(FNetSerializationContext&, const FNet
 	Target = TempValue;
 }
 
-void FRotatorAsByteNetSerializer::Dequantize(FNetSerializationContext&, const FNetDequantizeArgs& Args)
+template<typename T>
+void FRotatorAsByteNetSerializerBase<T>::Dequantize(FNetSerializationContext&, const FNetDequantizeArgs& Args)
 {
 	const QuantizedType& Source = *reinterpret_cast<const QuantizedType*>(Args.Source);
 	SourceType& Target = *reinterpret_cast<SourceType*>(Args.Target);
@@ -405,7 +399,8 @@ void FRotatorAsByteNetSerializer::Dequantize(FNetSerializationContext&, const FN
 	Target = TempValue;
 }
 
-bool FRotatorAsByteNetSerializer::IsEqual(FNetSerializationContext& Context, const FNetIsEqualArgs& Args)
+template<typename T>
+bool FRotatorAsByteNetSerializerBase<T>::IsEqual(FNetSerializationContext& Context, const FNetIsEqualArgs& Args)
 {
 	if (Args.bStateIsQuantized)
 	{
@@ -430,11 +425,62 @@ bool FRotatorAsByteNetSerializer::IsEqual(FNetSerializationContext& Context, con
 	}
 }
 
-bool FRotatorAsByteNetSerializer::Validate(FNetSerializationContext& Context, const FNetValidateArgs& Args)
+template<typename T>
+bool FRotatorAsByteNetSerializerBase<T>::Validate(FNetSerializationContext& Context, const FNetValidateArgs& Args)
 {
 	const SourceType& Value = *reinterpret_cast<SourceType*>(Args.Source);
 	// Make sure values are valid degree angles. This should catch NaNs as well.
-	return (Value.Pitch >= 0.0f && Value.Pitch < 360.0f) && (Value.Yaw >= 0.0f && Value.Yaw < 360.0f) & (Value.Roll >= 0.0f && Value.Roll < 360.0f); 
+	return (Value.Pitch >= ScalarType(0) && Value.Pitch < ScalarType(360)) && (Value.Yaw >= ScalarType(0) && Value.Yaw < ScalarType(360)) & (Value.Roll >= ScalarType(0) && Value.Roll < ScalarType(360)); 
 }
 
+struct FRotatorNetSerializer : public FRotatorAsShortNetSerializerBase<FRotator>
+{
+	static constexpr uint32 Version = 0;
+
+	typedef FRotatorNetSerializerConfig ConfigType;
+
+	inline static const ConfigType DefaultConfig;
+};
+UE_NET_IMPLEMENT_SERIALIZER(FRotatorNetSerializer);
+
+struct FRotatorAsByteNetSerializer : public FRotatorAsByteNetSerializerBase<FRotator>
+{
+	static constexpr uint32 Version = 0;
+
+	typedef FRotatorAsByteNetSerializerConfig ConfigType;
+
+	inline static const ConfigType DefaultConfig;
+};
+UE_NET_IMPLEMENT_SERIALIZER(FRotatorAsByteNetSerializer);
+
+struct FRotatorAsShortNetSerializer : public FRotatorAsShortNetSerializerBase<FRotator>
+{
+	static constexpr uint32 Version = 0;
+
+	typedef FRotatorAsShortNetSerializerConfig ConfigType;
+
+	inline static const ConfigType DefaultConfig;
+};
+UE_NET_IMPLEMENT_SERIALIZER(FRotatorAsShortNetSerializer);
+
+struct FRotator3fNetSerializer : public FRotatorAsShortNetSerializerBase<FRotator3f>
+{
+	static constexpr uint32 Version = 0;
+
+	typedef FRotator3fNetSerializerConfig ConfigType;
+
+	inline static const ConfigType DefaultConfig;
+};
+UE_NET_IMPLEMENT_SERIALIZER(FRotator3fNetSerializer);
+
+struct FRotator3dNetSerializer : public FRotatorAsShortNetSerializerBase<FRotator3d>
+{
+	static constexpr uint32 Version = 0;
+
+	typedef FRotator3dNetSerializerConfig ConfigType;
+
+	inline static const ConfigType DefaultConfig;
+};
+UE_NET_IMPLEMENT_SERIALIZER(FRotator3dNetSerializer);
+
 }
diff --git a/Engine/Source/Runtime/Experimental/Iris/Core/Public/Iris/Serialization/RotatorNetSerializers.h b/Engine/Source/Runtime/Experimental/Iris/Core/Public/Iris/Serialization/RotatorNetSerializers.h
--- a/Engine/Source/Runtime/Experimental/Iris/Core/Public/Iris/Serialization/RotatorNetSerializers.h
+++ b/Engine/Source/Runtime/Experimental/Iris/Core/Public/Iris/Serialization/RotatorNetSerializers.h
@@ -24,11 +24,25 @@ struct FRotatorAsShortNetSerializerConfig : public FNetSerializerConfig
 	GENERATED_BODY()
 };
 
+USTRUCT()
+struct FRotator3fNetSerializerConfig : public FNetSerializerConfig
+{
+	GENERATED_BODY()
+};
+
+USTRUCT()
+struct FRotator3dNetSerializerConfig : public FNetSerializerConfig
+{
+	GENERATED_BODY()
+};
+
 namespace UE::Net
 {
 
 UE_NET_DECLARE_SERIALIZER(FRotatorNetSerializer, IRISCORE_API);
 UE_NET_DECLARE_SERIALIZER(FRotatorAsByteNetSerializer, IRISCORE_API);
 UE_NET_DECLARE_SERIALIZER(FRotatorAsShortNetSerializer, IRISCORE_API);
+UE_NET_DECLARE_SERIALIZER(FRotator3fNetSerializer, IRISCORE_API);
+UE_NET_DECLARE_SERIALIZER(FRotator3dNetSerializer, IRISCORE_API);
 
 }
