Implement shallow copying for Chaos geometry speeding up modifications to geometry and shape arrays on particles.
Added access specifier to ModifyGeometry to be able to control how a caller will modify the underlying geometry that is returned.
Changed leaf object counting to no longer use a virtual visitor on shape array updates and update the count directly.

#tests Client builds against expensive geometry test case

diff --git a/Engine/Source/Runtime/Experimental/Chaos/Private/Chaos/Convex.cpp b/Engine/Source/Runtime/Experimental/Chaos/Private/Chaos/Convex.cpp
--- a/Engine/Source/Runtime/Experimental/Chaos/Private/Chaos/Convex.cpp
+++ b/Engine/Source/Runtime/Experimental/Chaos/Private/Chaos/Convex.cpp
@@ -8,10 +8,26 @@
 
 namespace Chaos
 {
-	Chaos::FImplicitObjectPtr FConvex::CopyGeometryWithScale(const FVec3& Scale) const
+	FImplicitObjectPtr FConvex::CopyGeometry() const
 	{
-		FConvexPtr ConvexCopy( new FConvex(*this));
-		return  Chaos::FImplicitObjectPtr(new TImplicitObjectScaled<FConvex>(ConvexCopy, Scale));
+		// const_cast required as this object has an intrusive reference count that need to be mutable
+		return FImplicitObjectPtr(const_cast<FConvex*>(this));
+	}
+
+	FImplicitObjectPtr FConvex::CopyGeometryWithScale(const FVec3& Scale) const
+	{
+		// const_cast required as this object has an intrusive reference count that need to be mutable
+		return FImplicitObjectPtr(new TImplicitObjectScaled<FConvex>(const_cast<FConvex*>(this), Scale));
+	}
+
+	FImplicitObjectPtr FConvex::DeepCopyGeometry() const
+	{
+		return FImplicitObjectPtr(new FConvex(*this));
+	}
+
+	FImplicitObjectPtr FConvex::DeepCopyGeometryWithScale(const FVec3& Scale) const
+	{
+		return FImplicitObjectPtr(new TImplicitObjectScaled<FConvex>(new FConvex(*this), Scale));
 	}
 
 	bool FConvex::Raycast(const FVec3& StartPoint, const FVec3& Dir, const FReal Length, const FReal Thickness, FReal& OutTime, FVec3& OutPosition, FVec3& OutNormal, int32& OutFaceIndex) const
diff --git a/Engine/Source/Runtime/Experimental/Chaos/Private/Chaos/ImplicitObjectUnion.cpp b/Engine/Source/Runtime/Experimental/Chaos/Private/Chaos/ImplicitObjectUnion.cpp
--- a/Engine/Source/Runtime/Experimental/Chaos/Private/Chaos/ImplicitObjectUnion.cpp
+++ b/Engine/Source/Runtime/Experimental/Chaos/Private/Chaos/ImplicitObjectUnion.cpp
@@ -76,7 +76,7 @@ FImplicitObjectUnion::FImplicitObjectUnion(TArray<Chaos::FImplicitObjectPtr>&& O
 	ensure(MObjects.Num());
 
 	MLocalBoundingBox = CalculateObjectsBounds(MakeArrayView(MObjects));
-	SetNumLeafObjects(Private::FImplicitBVH::CountLeafObjects(MakeArrayView(GetObjects())));
+	SetNumLeafObjects(CountObjectsInHierarchy());
 }
 
 FImplicitObjectUnion::FImplicitObjectUnion(FImplicitObjectUnion&& Other)
diff --git a/Engine/Source/Runtime/Experimental/Chaos/Private/Chaos/ParticleHandle.cpp b/Engine/Source/Runtime/Experimental/Chaos/Private/Chaos/ParticleHandle.cpp
--- a/Engine/Source/Runtime/Experimental/Chaos/Private/Chaos/ParticleHandle.cpp
+++ b/Engine/Source/Runtime/Experimental/Chaos/Private/Chaos/ParticleHandle.cpp
@@ -14,6 +14,18 @@
 
 namespace Chaos
 {
+	namespace CVars
+	{
+		static bool GForceDeepCopyOnModifyGeometry = false;
+
+		FAutoConsoleVariableRef CVarForceDeepCopyOnModifyGeometry(TEXT("p.Chaos.Geometry.ForceDeepCopyAccess"), GForceDeepCopyOnModifyGeometry, TEXT("Whether we always use a deep copy when modifying particle geometry"));
+
+		bool ForceDeepCopyOnModifyGeometry()
+		{
+			return GForceDeepCopyOnModifyGeometry;
+		}
+	}
+
 	extern void UpdateShapesArrayFromGeometry(FShapeInstanceProxyArray& ShapesArray, const FImplicitObjectPtr& Geometry, const FRigidTransform3& ActorTM, IPhysicsProxyBase* Proxy);
 
 	void SetObjectStateHelper(IPhysicsProxyBase& Proxy, FPBDRigidParticleHandle& Rigid, EObjectStateType InState, bool bAllowEvents, bool bInvalidate)
@@ -73,13 +85,15 @@ namespace Chaos
 
 		if (MNonFrequentData.Read().GetGeometry()->GetType() == FImplicitObjectUnion::StaticType())
 		{
-			ModifyGeometry([&Objects, this](FImplicitObject& GeomToModify)
-			{
-				if (FImplicitObjectUnion* Union = GeomToModify.template GetObject<FImplicitObjectUnion>())
+			// Only adding to the root union - shallow copy allowed here.
+			ModifyGeometry(EGeometryAccess::ShallowCopy,
+				[&Objects, this](FImplicitObject& GeomToModify)
 				{
-					Union->Combine(Objects);
-				}
-			});
+					if (FImplicitObjectUnion* Union = GeomToModify.template GetObject<FImplicitObjectUnion>())
+					{
+						Union->Combine(Objects);
+					}
+				});
 		}
 	}
 
@@ -127,7 +141,8 @@ namespace Chaos
 		// NOTE: only intended use is to remove objects from inside a FImplicitObjectUnion
 		CHAOS_ENSURE(MNonFrequentData.Read().GetGeometry()->GetType() == FImplicitObjectUnion::StaticType());
 
-		ModifyGeometry(
+		// Only removing shapes, shallow copy is allowed
+		ModifyGeometry(EGeometryAccess::ShallowCopy,
 			[this, &InIndices](FImplicitObject& GeomToModify)
 			{
 				if (FImplicitObjectUnion* Union = GeomToModify.template AsA<FImplicitObjectUnion>())
@@ -248,10 +263,12 @@ namespace Chaos
 	template <typename T, int d>
 	void TGeometryParticle<T,d>::SetIgnoreAnalyticCollisions(bool bIgnoreAnalyticCollisions)
 	{
-		ModifyGeometry([this, bIgnoreAnalyticCollisions](FImplicitObject& GeomToModify)
-		{
-			SetIgnoreAnalyticCollisionsImp(&GeomToModify, bIgnoreAnalyticCollisions);
-		});
+		// Deep copy required as we modify the actual geometries
+		ModifyGeometry(EGeometryAccess::DeepCopy,
+			[this, bIgnoreAnalyticCollisions](FImplicitObject& GeomToModify)
+			{
+				SetIgnoreAnalyticCollisionsImp(&GeomToModify, bIgnoreAnalyticCollisions);
+			});
 	}
 
 	template <typename T, int d, bool bPersistent>
diff --git a/Engine/Source/Runtime/Experimental/Chaos/Public/Chaos/Convex.h b/Engine/Source/Runtime/Experimental/Chaos/Public/Chaos/Convex.h
--- a/Engine/Source/Runtime/Experimental/Chaos/Public/Chaos/Convex.h
+++ b/Engine/Source/Runtime/Experimental/Chaos/Public/Chaos/Convex.h
@@ -190,12 +190,10 @@ namespace Chaos
             return nullptr;
         }
 
-		virtual Chaos::FImplicitObjectPtr CopyGeometry() const
-		{
-			return Chaos::FImplicitObjectPtr(new FConvex(*this));
-		}
-
-		CHAOS_API virtual Chaos::FImplicitObjectPtr CopyGeometryWithScale(const FVec3& Scale) const override;
+		CHAOS_API Chaos::FImplicitObjectPtr CopyGeometry() const override;
+		CHAOS_API Chaos::FImplicitObjectPtr CopyGeometryWithScale(const FVec3& Scale) const override;
+		CHAOS_API Chaos::FImplicitObjectPtr DeepCopyGeometry() const override;
+		CHAOS_API Chaos::FImplicitObjectPtr DeepCopyGeometryWithScale(const FVec3& Scale) const override;
 
 		CHAOS_API void MovePlanesAndRebuild(FRealType InDelta);
 
diff --git a/Engine/Source/Runtime/Experimental/Chaos/Public/Chaos/ImplicitObjectTransformed.h b/Engine/Source/Runtime/Experimental/Chaos/Public/Chaos/ImplicitObjectTransformed.h
--- a/Engine/Source/Runtime/Experimental/Chaos/Public/Chaos/ImplicitObjectTransformed.h
+++ b/Engine/Source/Runtime/Experimental/Chaos/Public/Chaos/ImplicitObjectTransformed.h
@@ -93,20 +93,14 @@ public:
 	
 	virtual Chaos::FImplicitObjectPtr CopyGeometry() const
 	{
-		if (MObject)
-		{
-			Chaos::FImplicitObjectPtr ObjectCopy(MObject->CopyGeometry());
-			return MakeImplicitObjectPtr<TImplicitObjectTransformed<T,d>>(MoveTemp(ObjectCopy), MTransform);
-		}
-		else
-		{
-			check(false);
-			return nullptr;
-		}
+		// shallow copying this just requires a copy of the pointer, which is just "this" which will convert into TRefCountPtr<FImplicitObject>
+		// const_cast required here as the invasive ref count needs to be mutated
+		return const_cast<TImplicitObjectTransformed*>(this);
 	}
 	
 	virtual Chaos::FImplicitObjectPtr CopyGeometryWithScale(const FVec3& Scale) const override
 	{
+		// This is a deep copy - we can't take a shallow copy of this and apply the scale without modifying all other instances
 		if(MObject)
 		{
 			//return MakeCopyWithScaleTransformed(MObjectOwner, MTransform, Scale);
@@ -127,6 +121,7 @@ public:
 
 	virtual Chaos::FImplicitObjectPtr DeepCopyGeometry() const
 	{
+		// Deep copy both the transformed wrapper, and the inner object
 		if(MObject)
 		{
 			Chaos::FImplicitObjectPtr ObjectCopy(MObject->DeepCopyGeometry());
@@ -443,7 +438,10 @@ namespace Utilities
 
 		if(GetInnerType(OuterType) == ImplicitObjectType::Transformed)
 		{
-			Chaos::FImplicitObjectPtr NewTransformed = InObject->CopyGeometry();
+			// Take a deep copy here as we're modifying the transformed itself.
+			// #TODO - Deep copy the transformed wrapper but shallow copy the internal shape as that isn't modified.
+			// Likely need to expand the copy functions to handle deep copy wrappers but not concrete geoms
+			Chaos::FImplicitObjectPtr NewTransformed = InObject->DeepCopyGeometry();
 			TImplicitObjectTransformed<FReal, 3>* InnerTransformed = static_cast<TImplicitObjectTransformed<FReal, 3>*>(NewTransformed.GetReference());
 			InnerTransformed->SetTransform(NewTransform);
 
@@ -451,6 +449,7 @@ namespace Utilities
 		}
 		else
 		{
+			// Shallow copy the inner object and wrap it in a new transformed
 			Chaos::FImplicitObjectPtr NewInnerObject = InObject->CopyGeometry();
 			return MakeImplicitObjectPtr<TImplicitObjectTransformed<FReal, 3>>(MoveTemp(NewInnerObject), NewTransform);
 		}
diff --git a/Engine/Source/Runtime/Experimental/Chaos/Public/Chaos/ParticleHandle.h b/Engine/Source/Runtime/Experimental/Chaos/Public/Chaos/ParticleHandle.h
--- a/Engine/Source/Runtime/Experimental/Chaos/Public/Chaos/ParticleHandle.h
+++ b/Engine/Source/Runtime/Experimental/Chaos/Public/Chaos/ParticleHandle.h
@@ -32,6 +32,11 @@ namespace Chaos
 		class FConvexOptimizer;
 	}
 
+	namespace CVars
+	{
+		bool ForceDeepCopyOnModifyGeometry();
+	}
+
 struct FGeometryParticleParameters
 {
 	FGeometryParticleParameters()
@@ -2857,30 +2862,78 @@ protected:
 	// Right now it's exposed to lubricate the creation of the whole proxy system.
 	class IPhysicsProxyBase* Proxy;
 
+	/*
+	* Describes how a call to ModifyGeometry wishes to access the underlying particle geometry
+	*/
+	enum class EGeometryAccess
+	{
+		// Grant direct access to the geometry, use with care - if the geometry is in
+		// use inside a solver this will likely be unsafe. If the geometry is not
+		// in use, all calls to ModifyGeometry will use Direct
+		Direct,
+
+		// Create a new root union, but shallow copy all geometries within the union.
+		// This is fine if you are only removing shapes, or adding shapes. If you require
+		// the ability to modify the actual geometry you must take a deep copy.
+		ShallowCopy,
+
+		// Create a new root union and deep copy all geometries within it. This is required
+		// if you intend to actually modify the geometry itself and not just grow/shrink
+		// the number of shapes on the particle
+		DeepCopy
+	};
+
 	template <typename Lambda>
 	void ModifyGeometry(const Lambda& Func, const bool bDirectAccess = false)
 	{
 		ensure(IsInGameThread());
+
+		return ModifyGeometry(bDirectAccess ? EGeometryAccess::Direct : EGeometryAccess::DeepCopy, Func);
+	}
+
+	template <typename Lambda>
+	void ModifyGeometry(EGeometryAccess AccessType, const Lambda& Func)
+	{
+		ensure(IsInGameThread());
+
 		FPhysicsSolverBase* Solver = Proxy ? Proxy->GetSolverBase() : nullptr;
-		MNonFrequentData.Modify(true, MDirtyFlags, Proxy, [this, Solver, &Func, bDirectAccess](auto& Data)
+
+		if(Solver == nullptr)
+		{
+			//not registered yet so we can still modify geometry safely
+			AccessType = EGeometryAccess::Direct;
+		}
+		else if(AccessType == EGeometryAccess::ShallowCopy && CVars::ForceDeepCopyOnModifyGeometry())
+		{
+			AccessType = EGeometryAccess::DeepCopy;
+		}
+
+		MNonFrequentData.Modify(true, MDirtyFlags, Proxy, [this, Solver, &Func, AccessType](auto& Data)
 		{
 			FImplicitObjectPtr GeomToModify = nullptr;
 			bool bNewGeom = false;
 			if(Data.GetGeometry())
 			{
-				if ((Solver == nullptr) || bDirectAccess)
+				switch(AccessType)
 				{
-					//not registered yet so we can still modify geometry
+				case EGeometryAccess::Direct:
 					GeomToModify = Data.AccessGeometryDangerous();
-				}
-				else
-				{
-					//already registered and used by physics thread, so need to duplicate
+					break;
+				case EGeometryAccess::ShallowCopy:
+					GeomToModify = Data.GetGeometry()->CopyGeometry();
+					bNewGeom = true;
+					break;
+				case EGeometryAccess::DeepCopy:
 					GeomToModify = Data.GetGeometry()->DeepCopyGeometry();
 					bNewGeom = true;
+					break;
+				default:
+					check(false);
+					break;
 				}
+
 				Func(*GeomToModify);
-				
+
 				if(bNewGeom)
 				{
 					//must set geometry after because shapes are rebuilt and we want them to know about anything Func did
