Updated ChaosDebugDrawComponent to render in the Editor world

diff --git a/Engine/Source/Runtime/Experimental/ChaosSolverEngine/Private/Chaos/ChaosDebugDrawComponent.cpp b/Engine/Source/Runtime/Experimental/ChaosSolverEngine/Private/Chaos/ChaosDebugDrawComponent.cpp
--- a/Engine/Source/Runtime/Experimental/ChaosSolverEngine/Private/Chaos/ChaosDebugDrawComponent.cpp
+++ b/Engine/Source/Runtime/Experimental/ChaosSolverEngine/Private/Chaos/ChaosDebugDrawComponent.cpp
@@ -85,7 +85,7 @@ void DebugDrawChaos(AActor* DebugDrawActor, const TArray<Chaos::FLatentDrawComma
 		return;
 	}
 
-	if (!World->IsGameWorld())
+	if (World->IsPreviewWorld())
 	{
 		return;
 	}
@@ -205,6 +205,8 @@ UChaosDebugDrawComponent::UChaosDebugDrawComponent()
 	PrimaryComponentTick.bCanEverTick = true;
 	PrimaryComponentTick.SetTickFunctionEnable(true);
 	PrimaryComponentTick.TickGroup = TG_PostUpdateWork;
+
+	bTickInEditor = true;
 }
 
 void UChaosDebugDrawComponent::BeginDestroy()
@@ -245,6 +247,36 @@ void UChaosDebugDrawComponent::EndPlay(const EEndPlayReason::Type EndPlayReason)
 #endif
 }
 
+#if WITH_EDITOR && CHAOS_DEBUG_DRAW
+void UChaosDebugDrawComponent::OnRegister()
+{
+	Super::OnRegister();
+
+	const UWorld* World = GetWorld();
+	check(World);
+
+	// GameWorld is handled by BeginPlay/EndPlay
+	if (!World->IsGameWorld() && !World->IsPreviewWorld())
+	{
+		Chaos::FDebugDrawQueue::GetInstance().SetConsumerActive(this, /*bConsumerActive*/true);
+	}
+}
+
+void UChaosDebugDrawComponent::OnUnregister()
+{
+	const UWorld* World = GetWorld();
+	check(World);
+
+	// GameWorld is handled by BeginPlay/EndPlay
+	if (!World->IsGameWorld() && !World->IsPreviewWorld())
+	{
+		Chaos::FDebugDrawQueue::GetInstance().SetConsumerActive(this, /*bConsumerActive*/false);
+	}
+
+	Super::OnUnregister();
+}
+#endif // WITH_EDITOR && CHAOS_DEBUG_DRAW
+
 void UChaosDebugDrawComponent::TickComponent(float DeltaTime, enum ELevelTick TickType, FActorComponentTickFunction *ThisTickFunction)
 {
 	using namespace Chaos;
@@ -317,7 +349,7 @@ void UChaosDebugDrawComponent::BindWorldDelegates()
 void UChaosDebugDrawComponent::HandlePostWorldInitialization(UWorld* World, const UWorld::InitializationValues IVS)
 {
 #if CHAOS_DEBUG_DRAW
-	if ((World != nullptr) && World->IsGameWorld())
+	if ((World != nullptr) && !World->IsPreviewWorld())
 	{
 		CreateDebugDrawActor(World);
 	}
diff --git a/Engine/Source/Runtime/Experimental/ChaosSolverEngine/Public/Chaos/ChaosDebugDrawComponent.h b/Engine/Source/Runtime/Experimental/ChaosSolverEngine/Public/Chaos/ChaosDebugDrawComponent.h
--- a/Engine/Source/Runtime/Experimental/ChaosSolverEngine/Public/Chaos/ChaosDebugDrawComponent.h
+++ b/Engine/Source/Runtime/Experimental/ChaosSolverEngine/Public/Chaos/ChaosDebugDrawComponent.h
@@ -28,6 +28,12 @@ public:
 
 	static CHAOSSOLVERENGINE_API void BindWorldDelegates();
 
+protected:
+#if WITH_EDITOR && CHAOS_DEBUG_DRAW
+	virtual void OnRegister() override;
+	virtual void OnUnregister() override;
+#endif // WITH_EDITOR && CHAOS_DEBUG_DRAW
+
 private:
 	static void HandlePostWorldInitialization(UWorld* World, const UWorld::InitializationValues IVS);
 	static void CreateDebugDrawActor(UWorld* World);
