SoftsEvolution -- rework PBDEvolution to be more group-oriented and setup to add force-based solver.

SoftsEvolution currently just has the same functionality of PBDEvolution, but it's structured in a way that will make it easier to add the force-based solver where we need to build and solve a linear system.

Main change is adding the concept of a ParticleRange which holds a pointer to the Particles SOA and an index and size to create an ArrayView. This works similar to the existing ParticleHandle, only it's for a contiguous range of particles.

The new Softs::FEvolution manages ranges, and the constraints have been updated to know how to act on ranges.

On the constraint side, the main change is the expectation that all internal indexing is working on local indices for the range rather than global indices. This makes things much simpler for data that only exists for a particular cloth object (range) like weight maps and triangle mesh indices.

Additionally, it is now explicitly enforced in Softs::FEvolution that each range corresponds with a single groupId. This simplifies collision code and some code related to managing transforms.

Finally, some more things were moved out of the Evolution and into constraint rules: per-substep initialization, external forces, and collision constraints. This removes a lot of the PropertyCollection management and weird group override settings in the cloth plugin. This also made it easier to refactor the initial calculation of P (euler stepping external forces and applying damping) to allow future additional ISPC or otherwise optimizations. Per-substep initialization was split into rules that are const on particles that can run in parallel and per-substep constraints that affect particles (e.g., long range constraint).

diff --git a/Engine/Plugins/Experimental/ChaosCloth/Source/Chaos/Private/ChaosCloth/ChaosClothConstraints.cpp b/Engine/Plugins/Experimental/ChaosCloth/Source/Chaos/Private/ChaosCloth/ChaosClothConstraints.cpp
--- a/Engine/Plugins/Experimental/ChaosCloth/Source/Chaos/Private/ChaosCloth/ChaosClothConstraints.cpp
+++ b/Engine/Plugins/Experimental/ChaosCloth/Source/Chaos/Private/ChaosCloth/ChaosClothConstraints.cpp
@@ -1668,19 +1668,7 @@ void FClothConstraints::SetAnimDriveConstraints(const TConstArrayView<FRealSingl
 			AnimDriveStiffnessMultipliers,
 			AnimDriveDampingMultipliers);
 	}
-	else
-	{
-		PRAGMA_DISABLE_DEPRECATION_WARNINGS
-		// Deprecated behavior until old Initialize can be removed
-		AnimDriveConstraints = MakeShareable( new Softs::FPBDAnimDriveConstraint(
-			ParticleOffset,
-			NumParticles,
-			*AnimationPositions,
-			*OldAnimationPositions_Deprecated,
-			AnimDriveStiffnessMultipliers,
-			AnimDriveDampingMultipliers));
-		PRAGMA_ENABLE_DEPRECATION_WARNINGS
-	}
+
 	++NumConstraintInits;  // Uses init to update the property tables
 	++NumConstraintRules;
 }
diff --git a/Engine/Source/Runtime/Experimental/Chaos/Private/Chaos/GraphColoring.cpp b/Engine/Source/Runtime/Experimental/Chaos/Private/Chaos/GraphColoring.cpp
--- a/Engine/Source/Runtime/Experimental/Chaos/Private/Chaos/GraphColoring.cpp
+++ b/Engine/Source/Runtime/Experimental/Chaos/Private/Chaos/GraphColoring.cpp
@@ -4,9 +4,10 @@
 #include "ChaosLog.h"
 #include "Chaos/Framework/Parallel.h"
 #include "Containers/BitArray.h"
+#include "Chaos/SoftsSolverParticlesRange.h"
 
-template<typename T>
-static bool VerifyGraph(TArray<TArray<int32>> ColorGraph, const TArray<Chaos::TVec2<int32>>& Graph, const Chaos::TDynamicParticles<T, 3>& InParticles)
+template<typename DynamicParticlesType>
+static bool VerifyGraph(TArray<TArray<int32>> ColorGraph, const TArray<Chaos::TVec2<int32>>& Graph, const DynamicParticlesType& InParticles)
 {
 	for (int32 i = 0; i < ColorGraph.Num(); ++i)
 	{
@@ -38,8 +39,8 @@ static bool VerifyGraph(TArray<TArray<int32>> ColorGraph, const TArray<Chaos::TV
 	return true;
 }
 
-template<typename T>
-static bool VerifyGraph(TArray<TArray<int32>> ColorGraph, const TArray<Chaos::TVec3<int32>>& Graph, const Chaos::TDynamicParticles<T, 3>& InParticles)
+template<typename DynamicParticlesType>
+static bool VerifyGraph(TArray<TArray<int32>> ColorGraph, const TArray<Chaos::TVec3<int32>>& Graph, const DynamicParticlesType& InParticles)
 {
 	for (int32 i = 0; i < ColorGraph.Num(); ++i)
 	{
@@ -81,8 +82,8 @@ static bool VerifyGraph(TArray<TArray<int32>> ColorGraph, const TArray<Chaos::TV
 	return true;
 }
 
-template<typename T>
-static bool VerifyGraph(TArray<TArray<int32>> ColorGraph, const TArray<Chaos::TVec4<int32>>& Graph, const Chaos::TDynamicParticles<T, 3>& InParticles)
+template<typename DynamicParticlesType>
+static bool VerifyGraph(TArray<TArray<int32>> ColorGraph, const TArray<Chaos::TVec4<int32>>& Graph, const DynamicParticlesType& InParticles)
 {	
 	for (int32 i = 0; i < ColorGraph.Num(); ++i)
 	{
@@ -135,8 +136,8 @@ static bool VerifyGraph(TArray<TArray<int32>> ColorGraph, const TArray<Chaos::TV
 }
 
 
-template<typename T>
-static bool VerifyGraphAllDynamic(TArray<TArray<int32>> ColorGraph, const TArray<Chaos::TVec4<int32>>& Graph, const Chaos::TDynamicParticles<T, 3>& InParticles)
+template<typename DynamicParticlesType>
+static bool VerifyGraphAllDynamic(TArray<TArray<int32>> ColorGraph, const TArray<Chaos::TVec4<int32>>& Graph, const DynamicParticlesType& InParticles)
 {
 	for (int32 i = 0; i < ColorGraph.Num(); ++i)
 	{
@@ -465,8 +466,8 @@ static bool VerifyExtraNodalColoring(const Chaos::TDynamicParticles<T, 3>& InPar
 	return true;
 }
 
-template<typename T>
-TArray<TArray<int32>> Chaos::FGraphColoring::ComputeGraphColoring(const TArray<Chaos::TVec2<int32>>& Graph, const Chaos::TDynamicParticles<T, 3>& InParticles, const int32 GraphParticlesStart, const int32 GraphParticlesEnd)
+template<typename DynamicParticlesType>
+TArray<TArray<int32>> Chaos::FGraphColoring::ComputeGraphColoringParticlesOrRange(const TArray<Chaos::TVec2<int32>>& Graph, const DynamicParticlesType& InParticles, const int32 GraphParticlesStart, const int32 GraphParticlesEnd)
 {
 	checkSlow(GraphParticlesStart <= GraphParticlesEnd);
 	checkSlow(GraphParticlesEnd <= (int32)InParticles.Size());
@@ -495,7 +496,7 @@ TArray<TArray<int32>> Chaos::FGraphColoring::ComputeGraphColoring(const TArray<C
 
 	for (int32 ParticleNodeIndex = GraphParticlesStart; ParticleNodeIndex < GraphParticlesEnd; ++ParticleNodeIndex)
 	{
-		const bool bIsParticleDynamic = InParticles.InvM(ParticleNodeIndex) != (T)0.;
+		const bool bIsParticleDynamic = InParticles.InvM(ParticleNodeIndex) != (decltype(InParticles.InvM(ParticleNodeIndex)))0.;
 		if (ProcessedNodes.Contains(ParticleNodeIndex) || !bIsParticleDynamic)
 		{
 			continue;
@@ -544,7 +545,7 @@ TArray<TArray<int32>> Chaos::FGraphColoring::ComputeGraphColoring(const TArray<C
 				{
 					FGraphNode& OtherNode = Nodes[OtherNodeIndex];
 
-					const bool bIsOtherGraphNodeDynamic = InParticles.InvM(OtherNodeIndex) != (T)0.;
+					const bool bIsOtherGraphNodeDynamic = InParticles.InvM(OtherNodeIndex) != (decltype(InParticles.InvM(ParticleNodeIndex)))0.;
 					if (bIsOtherGraphNodeDynamic)
 					{
 						while (OtherNode.UsedColors.Contains(ColorToUse) || GraphNode.UsedColors.Contains(ColorToUse))
@@ -574,7 +575,7 @@ TArray<TArray<int32>> Chaos::FGraphColoring::ComputeGraphColoring(const TArray<C
 				if (OtherNodeIndex != INDEX_NONE)
 				{
 					FGraphNode& OtherGraphNode = Nodes[OtherNodeIndex];
-					const bool bIsOtherGraphNodeDynamic = InParticles.InvM(OtherNodeIndex) != (T)0.;
+					const bool bIsOtherGraphNodeDynamic = InParticles.InvM(OtherNodeIndex) != (decltype(InParticles.InvM(ParticleNodeIndex)))0.;
 					if (bIsOtherGraphNodeDynamic)
 					{
 						// Mark other node as not allowing use of this color
@@ -598,8 +599,8 @@ TArray<TArray<int32>> Chaos::FGraphColoring::ComputeGraphColoring(const TArray<C
 	return ColorGraph;
 }
 
-template<typename T>
-TArray<TArray<int32>> Chaos::FGraphColoring::ComputeGraphColoring(const TArray<TVec3<int32>>& Graph, const Chaos::TDynamicParticles<T, 3>& InParticles, const int32 GraphParticlesStart, const int32 GraphParticlesEnd)
+template<typename DynamicParticlesType>
+TArray<TArray<int32>> Chaos::FGraphColoring::ComputeGraphColoringParticlesOrRange(const TArray<TVec3<int32>>& Graph, const DynamicParticlesType& InParticles, const int32 GraphParticlesStart, const int32 GraphParticlesEnd)
 {
 	using namespace Chaos;
 
@@ -633,7 +634,7 @@ TArray<TArray<int32>> Chaos::FGraphColoring::ComputeGraphColoring(const TArray<T
 
 	for (int32 ParticleNodeIndex = GraphParticlesStart; ParticleNodeIndex < GraphParticlesEnd; ++ParticleNodeIndex)
 	{
-		const bool bIsParticleDynamic = InParticles.InvM(ParticleNodeIndex) != (T)0.;
+		const bool bIsParticleDynamic = InParticles.InvM(ParticleNodeIndex) != (decltype(InParticles.InvM(ParticleNodeIndex)))0.;
 		if (ProcessedNodes.Contains(ParticleNodeIndex) || !bIsParticleDynamic)
 		{
 			continue;
@@ -690,7 +691,7 @@ TArray<TArray<int32>> Chaos::FGraphColoring::ComputeGraphColoring(const TArray<T
 				{
 					FGraphNode& OtherNode = Nodes[OtherNodeIndex];
 
-					const bool bIsOtherGraphNodeDynamic = InParticles.InvM(OtherNodeIndex) != (T)0.;
+					const bool bIsOtherGraphNodeDynamic = InParticles.InvM(OtherNodeIndex) != (decltype(InParticles.InvM(ParticleNodeIndex)))0.;
 					if (bIsOtherGraphNodeDynamic)
 					{
 						while (OtherNode.UsedColors.Contains(ColorToUse) || GraphNode.UsedColors.Contains(ColorToUse))
@@ -703,7 +704,7 @@ TArray<TArray<int32>> Chaos::FGraphColoring::ComputeGraphColoring(const TArray<T
 				{
 					FGraphNode& OtherNode = Nodes[OtherNodeIndex2];
 
-					const bool bIsOtherGraphNodeDynamic = InParticles.InvM(OtherNodeIndex2) != (T)0.;
+					const bool bIsOtherGraphNodeDynamic = InParticles.InvM(OtherNodeIndex2) != (decltype(InParticles.InvM(ParticleNodeIndex)))0.;
 					if (bIsOtherGraphNodeDynamic)
 					{
 						if (OtherNodeIndex == INDEX_NONE)
@@ -744,7 +745,7 @@ TArray<TArray<int32>> Chaos::FGraphColoring::ComputeGraphColoring(const TArray<T
 				if (OtherNodeIndex != INDEX_NONE)
 				{
 					FGraphNode& OtherGraphNode = Nodes[OtherNodeIndex];
-					const bool bIsOtherGraphNodeDynamic = InParticles.InvM(OtherNodeIndex) != (T)0.;
+					const bool bIsOtherGraphNodeDynamic = InParticles.InvM(OtherNodeIndex) != (decltype(InParticles.InvM(ParticleNodeIndex)))0.;
 					if (bIsOtherGraphNodeDynamic)
 					{
 						// Mark other node as not allowing use of this color
@@ -763,7 +764,7 @@ TArray<TArray<int32>> Chaos::FGraphColoring::ComputeGraphColoring(const TArray<T
 				if (OtherNodeIndex2 != INDEX_NONE)
 				{
 					FGraphNode& OtherGraphNode = Nodes[OtherNodeIndex2];
-					const bool bIsOtherGraphNodeDynamic = InParticles.InvM(OtherNodeIndex2) != (T)0.;
+					const bool bIsOtherGraphNodeDynamic = InParticles.InvM(OtherNodeIndex2) != (decltype(InParticles.InvM(ParticleNodeIndex)))0.;
 					if (bIsOtherGraphNodeDynamic)
 					{
 						// Mark other node as not allowing use of this color
@@ -787,8 +788,8 @@ TArray<TArray<int32>> Chaos::FGraphColoring::ComputeGraphColoring(const TArray<T
 	return ColorGraph;
 }
 
-template<typename T>
-TArray<TArray<int32>> Chaos::FGraphColoring::ComputeGraphColoring(const TArray<TVec4<int32>>& Graph, const Chaos::TDynamicParticles<T, 3>& InParticles, const int32 GraphParticlesStart, const int32 GraphParticlesEnd)
+template<typename DynamicParticlesType>
+TArray<TArray<int32>> Chaos::FGraphColoring::ComputeGraphColoringParticlesOrRange(const TArray<TVec4<int32>>& Graph, const DynamicParticlesType& InParticles, const int32 GraphParticlesStart, const int32 GraphParticlesEnd)
 {
 	using namespace Chaos;
 
@@ -825,7 +826,7 @@ TArray<TArray<int32>> Chaos::FGraphColoring::ComputeGraphColoring(const TArray<T
 
 	for (int32 ParticleNodeIndex = GraphParticlesStart; ParticleNodeIndex < GraphParticlesEnd; ++ParticleNodeIndex)
 	{
-		const bool bIsParticleDynamic = InParticles.InvM(ParticleNodeIndex) != (T)0.;
+		const bool bIsParticleDynamic = InParticles.InvM(ParticleNodeIndex) != (decltype(InParticles.InvM(ParticleNodeIndex)))0.;
 		if (ProcessedNodes.Contains(ParticleNodeIndex) || !bIsParticleDynamic)
 		{
 			continue;
@@ -892,7 +893,7 @@ TArray<TArray<int32>> Chaos::FGraphColoring::ComputeGraphColoring(const TArray<T
 				{
 					FGraphNode& OtherNode = Nodes[OtherNodeIndex];
 
-					const bool bIsOtherGraphNodeDynamic = InParticles.InvM(OtherNodeIndex) != (T)0.;
+					const bool bIsOtherGraphNodeDynamic = InParticles.InvM(OtherNodeIndex) != (decltype(InParticles.InvM(ParticleNodeIndex)))0.;
 					if (bIsOtherGraphNodeDynamic)
 					{
 						while (OtherNode.UsedColors.Contains(ColorToUse) || GraphNode.UsedColors.Contains(ColorToUse))
@@ -905,7 +906,7 @@ TArray<TArray<int32>> Chaos::FGraphColoring::ComputeGraphColoring(const TArray<T
 				{
 					FGraphNode& OtherNode = Nodes[OtherNodeIndex2];
 
-					const bool bIsOtherGraphNodeDynamic = InParticles.InvM(OtherNodeIndex2) != (T)0.;
+					const bool bIsOtherGraphNodeDynamic = InParticles.InvM(OtherNodeIndex2) != (decltype(InParticles.InvM(ParticleNodeIndex)))0.;
 					if (bIsOtherGraphNodeDynamic)
 					{
 						if (OtherNodeIndex == INDEX_NONE)
@@ -929,7 +930,7 @@ TArray<TArray<int32>> Chaos::FGraphColoring::ComputeGraphColoring(const TArray<T
 				{
 					FGraphNode& OtherNode = Nodes[OtherNodeIndex3];
 
-					const bool bIsOtherGraphNodeDynamic = InParticles.InvM(OtherNodeIndex3) != (T)0.;
+					const bool bIsOtherGraphNodeDynamic = InParticles.InvM(OtherNodeIndex3) != (decltype(InParticles.InvM(ParticleNodeIndex)))0.;
 					if (bIsOtherGraphNodeDynamic)
 					{
 						if (OtherNodeIndex == INDEX_NONE && OtherNodeIndex2 == INDEX_NONE)
@@ -988,7 +989,7 @@ TArray<TArray<int32>> Chaos::FGraphColoring::ComputeGraphColoring(const TArray<T
 				if (OtherNodeIndex != INDEX_NONE)
 				{
 					FGraphNode& OtherGraphNode = Nodes[OtherNodeIndex];
-					const bool bIsOtherGraphNodeDynamic = InParticles.InvM(OtherNodeIndex) != (T)0.;
+					const bool bIsOtherGraphNodeDynamic = InParticles.InvM(OtherNodeIndex) != (decltype(InParticles.InvM(ParticleNodeIndex)))0.;
 					if (bIsOtherGraphNodeDynamic)
 					{
 						// Mark other node as not allowing use of this color
@@ -1007,7 +1008,7 @@ TArray<TArray<int32>> Chaos::FGraphColoring::ComputeGraphColoring(const TArray<T
 				if (OtherNodeIndex2 != INDEX_NONE)
 				{
 					FGraphNode& OtherGraphNode = Nodes[OtherNodeIndex2];
-					const bool bIsOtherGraphNodeDynamic = InParticles.InvM(OtherNodeIndex2) != (T)0.;
+					const bool bIsOtherGraphNodeDynamic = InParticles.InvM(OtherNodeIndex2) != (decltype(InParticles.InvM(ParticleNodeIndex)))0.;
 					if (bIsOtherGraphNodeDynamic)
 					{
 						// Mark other node as not allowing use of this color
@@ -1026,7 +1027,7 @@ TArray<TArray<int32>> Chaos::FGraphColoring::ComputeGraphColoring(const TArray<T
 				if (OtherNodeIndex3 != INDEX_NONE)
 				{
 					FGraphNode& OtherGraphNode = Nodes[OtherNodeIndex3];
-					const bool bIsOtherGraphNodeDynamic = InParticles.InvM(OtherNodeIndex3) != (T)0.;
+					const bool bIsOtherGraphNodeDynamic = InParticles.InvM(OtherNodeIndex3) != (decltype(InParticles.InvM(ParticleNodeIndex)))0.;
 					if (bIsOtherGraphNodeDynamic)
 					{
 						// Mark other node as not allowing use of this color
@@ -1050,8 +1051,8 @@ TArray<TArray<int32>> Chaos::FGraphColoring::ComputeGraphColoring(const TArray<T
 	return ColorGraph;
 }
 
-template<typename T>
-TArray<TArray<int32>> Chaos::FGraphColoring::ComputeGraphColoringAllDynamic(const TArray<TVec4<int32>>& Graph, const Chaos::TDynamicParticles<T, 3>& InParticles, const int32 GraphParticlesStart, const int32 GraphParticlesEnd)
+template<typename DynamicParticlesType>
+TArray<TArray<int32>> Chaos::FGraphColoring::ComputeGraphColoringAllDynamicParticlesOrRange(const TArray<TVec4<int32>>& Graph, const DynamicParticlesType& InParticles, const int32 GraphParticlesStart, const int32 GraphParticlesEnd)
 {
 	using namespace Chaos;
 
@@ -1790,14 +1791,22 @@ void Chaos::ComputeExtraNodalColoring(const TArray<TArray<int32>>& Graph, const
 
 
 
-template CHAOS_API TArray<TArray<int32>> Chaos::FGraphColoring::ComputeGraphColoring<Chaos::FRealSingle>(const TArray<Chaos::TVector<int32, 2>>&, const Chaos::TDynamicParticles<Chaos::FRealSingle, 3>&, const int32 GraphParticlesStart, const int32 GraphParticlesEnd);
-template CHAOS_API TArray<TArray<int32>> Chaos::FGraphColoring::ComputeGraphColoring<Chaos::FRealDouble>(const TArray<Chaos::TVector<int32, 2>>&, const Chaos::TDynamicParticles<Chaos::FRealDouble, 3>&, const int32 GraphParticlesStart, const int32 GraphParticlesEnd);
-template CHAOS_API TArray<TArray<int32>> Chaos::FGraphColoring::ComputeGraphColoring<Chaos::FRealSingle>(const TArray<Chaos::TVector<int32, 3>>&, const Chaos::TDynamicParticles<Chaos::FRealSingle, 3>&, const int32 GraphParticlesStart, const int32 GraphParticlesEnd);
-template CHAOS_API TArray<TArray<int32>> Chaos::FGraphColoring::ComputeGraphColoring<Chaos::FRealDouble>(const TArray<Chaos::TVector<int32, 3>>&, const Chaos::TDynamicParticles<Chaos::FRealDouble, 3>&, const int32 GraphParticlesStart, const int32 GraphParticlesEnd);
-template CHAOS_API TArray<TArray<int32>> Chaos::FGraphColoring::ComputeGraphColoring<Chaos::FRealSingle>(const TArray<Chaos::TVector<int32, 4>>&, const Chaos::TDynamicParticles<Chaos::FRealSingle, 3>&, const int32 GraphParticlesStart, const int32 GraphParticlesEnd);
-template CHAOS_API TArray<TArray<int32>> Chaos::FGraphColoring::ComputeGraphColoring<Chaos::FRealDouble>(const TArray<Chaos::TVector<int32, 4>>&, const Chaos::TDynamicParticles<Chaos::FRealDouble, 3>&, const int32 GraphParticlesStart, const int32 GraphParticlesEnd);
-template CHAOS_API TArray<TArray<int32>> Chaos::FGraphColoring::ComputeGraphColoringAllDynamic<Chaos::FRealSingle>(const TArray<Chaos::TVector<int32, 4>>&, const Chaos::TDynamicParticles<Chaos::FRealSingle, 3>&, const int32 GraphParticlesStart, const int32 GraphParticlesEnd);
-template CHAOS_API TArray<TArray<int32>> Chaos::FGraphColoring::ComputeGraphColoringAllDynamic<Chaos::FRealDouble>(const TArray<Chaos::TVector<int32, 4>>&, const Chaos::TDynamicParticles<Chaos::FRealDouble, 3>&, const int32 GraphParticlesStart, const int32 GraphParticlesEnd);
+template CHAOS_API TArray<TArray<int32>> Chaos::FGraphColoring::ComputeGraphColoringParticlesOrRange(const TArray<Chaos::TVector<int32, 2>>&, const Chaos::TDynamicParticles<Chaos::FRealSingle, 3>&, const int32 GraphParticlesStart, const int32 GraphParticlesEnd);
+template CHAOS_API TArray<TArray<int32>> Chaos::FGraphColoring::ComputeGraphColoringParticlesOrRange(const TArray<Chaos::TVector<int32, 2>>&, const Chaos::TDynamicParticles<Chaos::FRealDouble, 3>&, const int32 GraphParticlesStart, const int32 GraphParticlesEnd);
+template CHAOS_API TArray<TArray<int32>> Chaos::FGraphColoring::ComputeGraphColoringParticlesOrRange(const TArray<Chaos::TVector<int32, 2>>&, const Chaos::Softs::FSolverParticles&, const int32 GraphParticlesStart, const int32 GraphParticlesEnd);
+template CHAOS_API TArray<TArray<int32>> Chaos::FGraphColoring::ComputeGraphColoringParticlesOrRange(const TArray<Chaos::TVector<int32, 2>>&, const Chaos::Softs::FSolverParticlesRange&, const int32 GraphParticlesStart, const int32 GraphParticlesEnd);
+template CHAOS_API TArray<TArray<int32>> Chaos::FGraphColoring::ComputeGraphColoringParticlesOrRange(const TArray<Chaos::TVector<int32, 3>>&, const Chaos::TDynamicParticles<Chaos::FRealSingle, 3>&, const int32 GraphParticlesStart, const int32 GraphParticlesEnd);
+template CHAOS_API TArray<TArray<int32>> Chaos::FGraphColoring::ComputeGraphColoringParticlesOrRange(const TArray<Chaos::TVector<int32, 3>>&, const Chaos::TDynamicParticles<Chaos::FRealDouble, 3>&, const int32 GraphParticlesStart, const int32 GraphParticlesEnd);
+template CHAOS_API TArray<TArray<int32>> Chaos::FGraphColoring::ComputeGraphColoringParticlesOrRange(const TArray<Chaos::TVector<int32, 3>>&, const Chaos::Softs::FSolverParticles&, const int32 GraphParticlesStart, const int32 GraphParticlesEnd);
+template CHAOS_API TArray<TArray<int32>> Chaos::FGraphColoring::ComputeGraphColoringParticlesOrRange(const TArray<Chaos::TVector<int32, 3>>&, const Chaos::Softs::FSolverParticlesRange&, const int32 GraphParticlesStart, const int32 GraphParticlesEnd);
+template CHAOS_API TArray<TArray<int32>> Chaos::FGraphColoring::ComputeGraphColoringParticlesOrRange(const TArray<Chaos::TVector<int32, 4>>&, const Chaos::TDynamicParticles<Chaos::FRealSingle, 3>&, const int32 GraphParticlesStart, const int32 GraphParticlesEnd);
+template CHAOS_API TArray<TArray<int32>> Chaos::FGraphColoring::ComputeGraphColoringParticlesOrRange(const TArray<Chaos::TVector<int32, 4>>&, const Chaos::TDynamicParticles<Chaos::FRealDouble, 3>&, const int32 GraphParticlesStart, const int32 GraphParticlesEnd);
+template CHAOS_API TArray<TArray<int32>> Chaos::FGraphColoring::ComputeGraphColoringParticlesOrRange(const TArray<Chaos::TVector<int32, 4>>&, const Chaos::Softs::FSolverParticles&, const int32 GraphParticlesStart, const int32 GraphParticlesEnd);
+template CHAOS_API TArray<TArray<int32>> Chaos::FGraphColoring::ComputeGraphColoringParticlesOrRange(const TArray<Chaos::TVector<int32, 4>>&, const Chaos::Softs::FSolverParticlesRange&, const int32 GraphParticlesStart, const int32 GraphParticlesEnd);
+template CHAOS_API TArray<TArray<int32>> Chaos::FGraphColoring::ComputeGraphColoringAllDynamicParticlesOrRange(const TArray<Chaos::TVector<int32, 4>>&, const Chaos::TDynamicParticles<Chaos::FRealSingle, 3>&, const int32 GraphParticlesStart, const int32 GraphParticlesEnd);
+template CHAOS_API TArray<TArray<int32>> Chaos::FGraphColoring::ComputeGraphColoringAllDynamicParticlesOrRange(const TArray<Chaos::TVector<int32, 4>>&, const Chaos::TDynamicParticles<Chaos::FRealDouble, 3>&, const int32 GraphParticlesStart, const int32 GraphParticlesEnd);
+template CHAOS_API TArray<TArray<int32>> Chaos::FGraphColoring::ComputeGraphColoringAllDynamicParticlesOrRange(const TArray<Chaos::TVector<int32, 4>>&, const Chaos::Softs::FSolverParticles&, const int32 GraphParticlesStart, const int32 GraphParticlesEnd);
+template CHAOS_API TArray<TArray<int32>> Chaos::FGraphColoring::ComputeGraphColoringAllDynamicParticlesOrRange(const TArray<Chaos::TVector<int32, 4>>&, const Chaos::Softs::FSolverParticlesRange&, const int32 GraphParticlesStart, const int32 GraphParticlesEnd);
 template CHAOS_API void Chaos::ComputeGridBasedGraphSubColoringPointer(const TArray<TArray<int32>>& ElementsPerColor, const TMPMGrid<Chaos::FRealSingle>& Grid, const int32 GridSize, TArray<TArray<int32>>*& PreviousColoring, const TArray<TArray<int32>>& ConstraintsNodesSet, TArray<TArray<TArray<int32>>>& ElementsPerSubColors);
 template CHAOS_API void Chaos::ComputeGridBasedGraphSubColoringPointer(const TArray<TArray<int32>>& ElementsPerColor, const TMPMGrid<Chaos::FRealDouble>& Grid, const int32 GridSize, TArray<TArray<int32>>*& PreviousColoring, const TArray<TArray<int32>>& ConstraintsNodesSet, TArray<TArray<TArray<int32>>>& ElementsPerSubColors);
 template CHAOS_API void Chaos::ComputeWeakConstraintsColoring<Chaos::FRealSingle>(const TArray<TArray<int32>>& Indices, const TArray<TArray<int32>>& SecondIndices, const Chaos::TDynamicParticles<Chaos::FRealSingle, 3>& InParticles, TArray<TArray<int32>>& ConstraintsPerColor);
diff --git a/Engine/Source/Runtime/Experimental/Chaos/Private/Chaos/PBDAxialSpringConstraints.cpp b/Engine/Source/Runtime/Experimental/Chaos/Private/Chaos/PBDAxialSpringConstraints.cpp
--- a/Engine/Source/Runtime/Experimental/Chaos/Private/Chaos/PBDAxialSpringConstraints.cpp
+++ b/Engine/Source/Runtime/Experimental/Chaos/Private/Chaos/PBDAxialSpringConstraints.cpp
@@ -23,14 +23,15 @@ namespace Chaos::Softs {
 int32 Chaos_AxialSpring_ParallelConstraintCount = 100;
 FAutoConsoleVariableRef CVarChaosAxialSpringParallelConstraintCount(TEXT("p.Chaos.AxialSpring.ParallelConstraintCount"), Chaos_AxialSpring_ParallelConstraintCount, TEXT("If we have more constraints than this, use parallel-for in Apply."));
 
-void FPBDAxialSpringConstraints::InitColor(const FSolverParticles& InParticles)
+template<typename SolverParticlesOrRange>
+void FPBDAxialSpringConstraints::InitColor(const SolverParticlesOrRange& InParticles)
 {
 	// In dev builds we always color so we can tune the system without restarting. See Apply()
 #if UE_BUILD_SHIPPING || UE_BUILD_TEST
 	if (Constraints.Num() > Chaos_AxialSpring_ParallelConstraintCount)
 #endif
 	{
-		const TArray<TArray<int32>> ConstraintsPerColor = FGraphColoring::ComputeGraphColoring(Constraints, InParticles, ParticleOffset, ParticleOffset + ParticleCount);
+		const TArray<TArray<int32>> ConstraintsPerColor = FGraphColoring::ComputeGraphColoringParticlesOrRange(Constraints, InParticles, ParticleOffset, ParticleOffset + ParticleCount);
 
 		// Reorder constraints based on color so each array in ConstraintsPerColor contains contiguous elements.
 		TArray<TVec3<int32>> ReorderedConstraints;
@@ -66,8 +67,11 @@ void FPBDAxialSpringConstraints::InitColor(const FSolverParticles& InParticles)
 		Stiffness.ReorderIndices(OrigToReorderedIndices);
 	}
 }
+template CHAOS_API void FPBDAxialSpringConstraints::InitColor(const FSolverParticles& InParticles);
+template CHAOS_API void FPBDAxialSpringConstraints::InitColor(const FSolverParticlesRange& InParticles);
 
-void FPBDAxialSpringConstraints::ApplyHelper(FSolverParticles& Particles, const FSolverReal Dt, const int32 ConstraintIndex, const FSolverReal ExpStiffnessValue) const
+template<typename SolverParticlesOrRange>
+void FPBDAxialSpringConstraints::ApplyHelper(SolverParticlesOrRange& Particles, const FSolverReal Dt, const int32 ConstraintIndex, const FSolverReal ExpStiffnessValue) const
 {
 		const TVec3<int32>& Constraint = Constraints[ConstraintIndex];
 		const int32 i1 = Constraint[0];
@@ -88,8 +92,11 @@ void FPBDAxialSpringConstraints::ApplyHelper(FSolverParticles& Particles, const
 			Particles.P(i3) += Multiplier * Particles.InvM(i3) * ((FSolverReal)1. - Barys[ConstraintIndex]) * Delta;
 		}
 }
+template void FPBDAxialSpringConstraints::ApplyHelper(FSolverParticles& Particles, const FSolverReal Dt, const int32 ConstraintIndex, const FSolverReal ExpStiffnessValue) const;
+template void FPBDAxialSpringConstraints::ApplyHelper(FSolverParticlesRange& Particles, const FSolverReal Dt, const int32 ConstraintIndex, const FSolverReal ExpStiffnessValue) const;
 
-void FPBDAxialSpringConstraints::Apply(FSolverParticles& Particles, const FSolverReal Dt) const
+template<typename SolverParticlesOrRange>
+void FPBDAxialSpringConstraints::Apply(SolverParticlesOrRange& Particles, const FSolverReal Dt) const
 {
 	TRACE_CPUPROFILER_EVENT_SCOPE(FPBDAxialSpringConstraints_Apply);
 	SCOPE_CYCLE_COUNTER(STAT_PBD_AxialSpring);
@@ -187,6 +194,8 @@ void FPBDAxialSpringConstraints::Apply(FSolverParticles& Particles, const FSolve
 		}
 	}
 }
+template CHAOS_API void FPBDAxialSpringConstraints::Apply(FSolverParticles& Particles, const FSolverReal Dt) const;
+template CHAOS_API void FPBDAxialSpringConstraints::Apply(FSolverParticlesRange& Particles, const FSolverReal Dt) const;
 
 void FPBDAreaSpringConstraints::SetProperties(
 	const FCollectionPropertyConstFacade& PropertyCollection,
diff --git a/Engine/Source/Runtime/Experimental/Chaos/Private/Chaos/PBDBendingConstraints.cpp b/Engine/Source/Runtime/Experimental/Chaos/Private/Chaos/PBDBendingConstraints.cpp
--- a/Engine/Source/Runtime/Experimental/Chaos/Private/Chaos/PBDBendingConstraints.cpp
+++ b/Engine/Source/Runtime/Experimental/Chaos/Private/Chaos/PBDBendingConstraints.cpp
@@ -23,14 +23,15 @@ namespace Chaos::Softs {
 int32 Chaos_Bending_ParallelConstraintCount = 100;
 FAutoConsoleVariableRef CVarChaosBendingParallelConstraintCount(TEXT("p.Chaos.Bending.ParallelConstraintCount"), Chaos_Bending_ParallelConstraintCount, TEXT("If we have more constraints than this, use parallel-for in Apply."));
 
-void FPBDBendingConstraints::InitColor(const FSolverParticles& InParticles)
+template<typename SolverParticlesOrRange>
+void FPBDBendingConstraints::InitColor(const SolverParticlesOrRange& InParticles)
 {
 	// In dev builds we always color so we can tune the system without restarting. See Apply()
 #if UE_BUILD_SHIPPING || UE_BUILD_TEST
 	if (Constraints.Num() > Chaos_Bending_ParallelConstraintCount)
 #endif
 	{
-		const TArray<TArray<int32>> ConstraintsPerColor = FGraphColoring::ComputeGraphColoring(Constraints, InParticles, ParticleOffset, ParticleOffset + ParticleCount);
+		const TArray<TArray<int32>> ConstraintsPerColor = FGraphColoring::ComputeGraphColoringParticlesOrRange(Constraints, InParticles, ParticleOffset, ParticleOffset + ParticleCount);
 
 		// Reorder constraints based on color so each array in ConstraintsPerColor contains contiguous elements.
 		TArray<TVec4<int32>> ReorderedConstraints; 
@@ -67,6 +68,8 @@ void FPBDBendingConstraints::InitColor(const FSolverParticles& InParticles)
 		BucklingStiffness.ReorderIndices(OrigToReorderedIndices);
 	}
 }
+template CHAOS_API void FPBDBendingConstraints::InitColor(const FSolverParticles& InParticles);
+template CHAOS_API void FPBDBendingConstraints::InitColor(const FSolverParticlesRange& InParticles);
 
 void FPBDBendingConstraints::SetProperties(
 	const FCollectionPropertyConstFacade& PropertyCollection,
@@ -114,7 +117,8 @@ void FPBDBendingConstraints::SetProperties(
 	}
 }
 
-void FPBDBendingConstraints::ApplyHelper(FSolverParticles& Particles, const FSolverReal Dt, const int32 ConstraintIndex, const FSolverReal ExpStiffnessValue, const FSolverReal ExpBucklingValue) const
+template<typename SolverParticlesOrRange>
+void FPBDBendingConstraints::ApplyHelper(SolverParticlesOrRange& Particles, const FSolverReal Dt, const int32 ConstraintIndex, const FSolverReal ExpStiffnessValue, const FSolverReal ExpBucklingValue) const
 {
 
 	const TVec4<int32>& Constraint = Constraints[ConstraintIndex];
@@ -130,7 +134,8 @@ void FPBDBendingConstraints::ApplyHelper(FSolverParticles& Particles, const FSol
 	Particles.P(i4) -= S * Particles.InvM(i4) * Grads[3];
 }
 
-void FPBDBendingConstraints::Apply(FSolverParticles& Particles, const FSolverReal Dt) const
+template<typename SolverParticlesOrRange>
+void FPBDBendingConstraints::Apply(SolverParticlesOrRange& Particles, const FSolverReal Dt) const
 {
 	TRACE_CPUPROFILER_EVENT_SCOPE(FPBDBendingConstraints_Apply);
 	SCOPE_CYCLE_COUNTER(STAT_PBD_Bending);
@@ -244,5 +249,7 @@ void FPBDBendingConstraints::Apply(FSolverParticles& Particles, const FSolverRea
 		}
 	}
 }
+template CHAOS_API void FPBDBendingConstraints::Apply(FSolverParticles& Particles, const FSolverReal Dt) const;
+template CHAOS_API void FPBDBendingConstraints::Apply(FSolverParticlesRange& Particles, const FSolverReal Dt) const;
 
 }  // End namespace Chaos::Softs
diff --git a/Engine/Source/Runtime/Experimental/Chaos/Private/Chaos/PBDBendingConstraintsBase.cpp b/Engine/Source/Runtime/Experimental/Chaos/Private/Chaos/PBDBendingConstraintsBase.cpp
--- a/Engine/Source/Runtime/Experimental/Chaos/Private/Chaos/PBDBendingConstraintsBase.cpp
+++ b/Engine/Source/Runtime/Experimental/Chaos/Private/Chaos/PBDBendingConstraintsBase.cpp
@@ -5,8 +5,9 @@ namespace Chaos::Softs
 {
 	namespace Private
 	{
+		template<typename SolverParticlesOrRange>
 		void Calculate3DRestAngles(
-			const FSolverParticles& InParticles,
+			const SolverParticlesOrRange& InParticles,
 			const TArray<TVec4<int32>>& Constraints,
 			TArray<FSolverReal>& RestAngles)
 		{
@@ -21,8 +22,9 @@ namespace Chaos::Softs
 			}
 		}
 
+		template<typename SolverParticlesOrRange>
 		void CalculateFlatnessRestAngles(
-			const FSolverParticles& InParticles,
+			const SolverParticlesOrRange& InParticles,
 			int32 InParticleOffset,
 			int32 InParticleCount,
 			const TConstArrayView<FRealSingle>& RestAngleMap,
@@ -95,8 +97,9 @@ namespace Chaos::Softs
 		}
 	}
 
+	template<typename SolverParticlesOrRange>
 	void FPBDBendingConstraintsBase::CalculateRestAngles(
-		const FSolverParticles& InParticles,
+		const SolverParticlesOrRange& InParticles,
 		int32 InParticleOffset,
 		int32 InParticleCount,
 		const TConstArrayView<FRealSingle>& RestAngleMap,
@@ -116,4 +119,19 @@ namespace Chaos::Softs
 			Private::CalculateExplicitRestAngles(InParticleOffset, InParticleCount, RestAngleMap, RestAngleValue, Constraints, RestAngles);
 		}
 	}
+
+	template CHAOS_API void FPBDBendingConstraintsBase::CalculateRestAngles(
+		const FSolverParticles& InParticles,
+		int32 InParticleOffset,
+		int32 InParticleCount,
+		const TConstArrayView<FRealSingle>& RestAngleMap,
+		const FSolverVec2& RestAngleValue,
+		ERestAngleConstructionType RestAngleConstructionType);
+	template CHAOS_API void FPBDBendingConstraintsBase::CalculateRestAngles(
+		const FSolverParticlesRange& InParticles,
+		int32 InParticleOffset,
+		int32 InParticleCount,
+		const TConstArrayView<FRealSingle>& RestAngleMap,
+		const FSolverVec2& RestAngleValue,
+		ERestAngleConstructionType RestAngleConstructionType);
 }  // End namespace Chaos::Softs
diff --git a/Engine/Source/Runtime/Experimental/Chaos/Private/Chaos/PBDCollisionSpringConstraintsBase.cpp b/Engine/Source/Runtime/Experimental/Chaos/Private/Chaos/PBDCollisionSpringConstraintsBase.cpp
--- a/Engine/Source/Runtime/Experimental/Chaos/Private/Chaos/PBDCollisionSpringConstraintsBase.cpp
+++ b/Engine/Source/Runtime/Experimental/Chaos/Private/Chaos/PBDCollisionSpringConstraintsBase.cpp
@@ -5,6 +5,7 @@
 #include "Chaos/TriangleCollisionPoint.h"
 #include "Chaos/TriangleMesh.h"
 #include "Chaos/PBDSoftsSolverParticles.h"
+#include "Chaos/SoftsSolverParticlesRange.h"
 
 #if !COMPILE_WITHOUT_UNREAL_SUPPORT
 #include "Chaos/Framework/Parallel.h"
@@ -12,6 +13,24 @@
 
 namespace Chaos::Softs {
 
+static TConstArrayView<FSolverVec3> GetConstArrayView(const FSolverParticles& Particles, int32 Offset, int32 NumParticles, const TArray<FSolverVec3>* const Data)
+{
+	if (Data)
+	{
+		return TConstArrayView<FSolverVec3>(Data->GetData(), NumParticles + Offset);
+	}
+	return TConstArrayView<FSolverVec3>();
+}
+
+static TConstArrayView<FSolverVec3> GetConstArrayView(const FSolverParticlesRange& Particles, int32 Offset, int32 NumParticles, const TArray<FSolverVec3>* const Data)
+{
+	if (Data)
+	{
+		return Particles.GetConstArrayView(*Data);
+	}
+	return TConstArrayView<FSolverVec3>();
+}
+
 FPBDCollisionSpringConstraintsBase::FPBDCollisionSpringConstraintsBase(
 	const int32 InOffset,
 	const int32 InNumParticles,
@@ -53,8 +72,8 @@ void FPBDCollisionSpringConstraintsBase::Init(const FSolverParticles& Particles)
 	return Init(Particles, BVH, static_cast<TConstArrayView<FPBDTriangleMeshCollisions::FGIAColor>>(EmptyGIAColors), EmptyGIAColors);
 }
 
-template<typename SpatialAccelerator>
-void FPBDCollisionSpringConstraintsBase::Init(const FSolverParticles& Particles, const SpatialAccelerator& Spatial, 
+template<typename SpatialAccelerator, typename SolverParticlesOrRange>
+void FPBDCollisionSpringConstraintsBase::Init(const SolverParticlesOrRange& Particles, const SpatialAccelerator& Spatial,
 	const TConstArrayView<FPBDTriangleMeshCollisions::FGIAColor>& VertexGIAColors, const TArray<FPBDTriangleMeshCollisions::FGIAColor>& TriangleGIAColors)
 {
 	if (!Elements.Num())
@@ -78,8 +97,10 @@ void FPBDCollisionSpringConstraintsBase::Init(const FSolverParticles& Particles,
 
 		const FSolverReal HeightSq = FMath::Square(Thickness + Thickness);
 
+		const TConstArrayView<FSolverVec3> ReferencePositionsView = GetConstArrayView(Particles, Offset, NumParticles, ReferencePositions);
+
 		PhysicsParallelFor(NumParticles,
-			[this, &Spatial, &Particles, &ConstraintIndex, HeightSq, MaxConnectionsPerPoint, &VertexGIAColors, &TriangleGIAColors](int32 i)
+			[this, &Spatial, &Particles, &ConstraintIndex, HeightSq, MaxConnectionsPerPoint, &VertexGIAColors, &TriangleGIAColors, &ReferencePositionsView](int32 i)
 			{
 				const int32 Index = i + Offset;
 				constexpr FSolverReal ExtraThicknessMult = 1.5f;
@@ -139,12 +160,12 @@ void FPBDCollisionSpringConstraintsBase::Init(const FSolverParticles& Particles,
 					for (const TTriangleCollisionPoint<FSolverReal>& CollisionPoint : Result)
 					{
 						const TVector<int32, 3>& Elem = Elements[CollisionPoint.Indices[1]];
-						if (ReferencePositions)
+						if (ReferencePositionsView.Num())
 						{
-							const FSolverVec3& RefP = (*ReferencePositions)[Index];
-							const FSolverVec3& RefP0 = (*ReferencePositions)[Elem[0]];
-							const FSolverVec3& RefP1 = (*ReferencePositions)[Elem[1]];
-							const FSolverVec3& RefP2 = (*ReferencePositions)[Elem[2]];
+							const FSolverVec3& RefP = ReferencePositionsView[Index];
+							const FSolverVec3& RefP0 = ReferencePositionsView[Elem[0]];
+							const FSolverVec3& RefP1 = ReferencePositionsView[Elem[1]];
+							const FSolverVec3& RefP2 = ReferencePositionsView[Elem[2]];
 							const FSolverVec3 RefDiff = RefP - CollisionPoint.Bary[1] * RefP0 - CollisionPoint.Bary[2] * RefP1 - CollisionPoint.Bary[3] * RefP2;
 							if (RefDiff.SizeSquared() < HeightSq)
 							{
@@ -188,8 +209,13 @@ template void CHAOS_API FPBDCollisionSpringConstraintsBase::Init<FTriangleMesh::
 	const TConstArrayView<FPBDTriangleMeshCollisions::FGIAColor>& VertexGIAColors, const TArray<FPBDTriangleMeshCollisions::FGIAColor>& TriangleGIAColors);
 template void CHAOS_API FPBDCollisionSpringConstraintsBase::Init<FTriangleMesh::TSpatialHashType<FSolverReal>>(const FSolverParticles& Particles, const FTriangleMesh::TSpatialHashType<FSolverReal>& Spatial, 
 	const TConstArrayView<FPBDTriangleMeshCollisions::FGIAColor>& VertexGIAColors, const TArray<FPBDTriangleMeshCollisions::FGIAColor>& TriangleGIAColors);
+template void CHAOS_API FPBDCollisionSpringConstraintsBase::Init<FTriangleMesh::TBVHType<FSolverReal>>(const FSolverParticlesRange& Particles, const FTriangleMesh::TBVHType<FSolverReal>& Spatial,
+	const TConstArrayView<FPBDTriangleMeshCollisions::FGIAColor>& VertexGIAColors, const TArray<FPBDTriangleMeshCollisions::FGIAColor>& TriangleGIAColors);
+template void CHAOS_API FPBDCollisionSpringConstraintsBase::Init<FTriangleMesh::TSpatialHashType<FSolverReal>>(const FSolverParticlesRange& Particles, const FTriangleMesh::TSpatialHashType<FSolverReal>& Spatial,
+	const TConstArrayView<FPBDTriangleMeshCollisions::FGIAColor>& VertexGIAColors, const TArray<FPBDTriangleMeshCollisions::FGIAColor>& TriangleGIAColors);
 
-FSolverVec3 FPBDCollisionSpringConstraintsBase::GetDelta(const FSolverParticles& Particles, const int32 i) const
+template<typename SolverParticlesOrRange>
+FSolverVec3 FPBDCollisionSpringConstraintsBase::GetDelta(const SolverParticlesOrRange& Particles, const int32 i) const
 {
 	const TVec4<int32>& Constraint = Constraints[i];
 	const int32 i1 = Constraint[0];
@@ -247,6 +273,8 @@ FSolverVec3 FPBDCollisionSpringConstraintsBase::GetDelta(const FSolverParticles&
 		return RepulsionDelta;
 	}
 }
+template CHAOS_API FSolverVec3 FPBDCollisionSpringConstraintsBase::GetDelta(const FSolverParticles& Particles, const int32 i) const;
+template CHAOS_API FSolverVec3 FPBDCollisionSpringConstraintsBase::GetDelta(const FSolverParticlesRange& Particles, const int32 i) const;
 
 }  // End namespace Chaos::Softs
 
diff --git a/Engine/Source/Runtime/Experimental/Chaos/Private/Chaos/PBDLongRangeConstraints.cpp b/Engine/Source/Runtime/Experimental/Chaos/Private/Chaos/PBDLongRangeConstraints.cpp
--- a/Engine/Source/Runtime/Experimental/Chaos/Private/Chaos/PBDLongRangeConstraints.cpp
+++ b/Engine/Source/Runtime/Experimental/Chaos/Private/Chaos/PBDLongRangeConstraints.cpp
@@ -48,7 +48,8 @@ void FPBDLongRangeConstraints::SetProperties(
 	}
 }
 
-void FPBDLongRangeConstraints::Apply(FSolverParticles& Particles, const FSolverReal /*Dt*/) const
+template<typename SolverParticlesOrRange>
+void FPBDLongRangeConstraints::Apply(SolverParticlesOrRange& Particles, const FSolverReal /*Dt*/) const
 {
 	TRACE_CPUPROFILER_EVENT_SCOPE(FPBDLongRangeConstraints_Apply);
 	SCOPE_CYCLE_COUNTER(STAT_PBD_LongRange);
@@ -200,5 +201,7 @@ void FPBDLongRangeConstraints::Apply(FSolverParticles& Particles, const FSolverR
 		}
 	}
 }
+template CHAOS_API void FPBDLongRangeConstraints::Apply(FSolverParticles& Particles, const FSolverReal /*Dt*/) const;
+template CHAOS_API void FPBDLongRangeConstraints::Apply(FSolverParticlesRange& Particles, const FSolverReal /*Dt*/) const;
 
 }  // End namespace Chaos::Softs
diff --git a/Engine/Source/Runtime/Experimental/Chaos/Private/Chaos/PBDLongRangeConstraintsBase.cpp b/Engine/Source/Runtime/Experimental/Chaos/Private/Chaos/PBDLongRangeConstraintsBase.cpp
--- a/Engine/Source/Runtime/Experimental/Chaos/Private/Chaos/PBDLongRangeConstraintsBase.cpp
+++ b/Engine/Source/Runtime/Experimental/Chaos/Private/Chaos/PBDLongRangeConstraintsBase.cpp
@@ -24,6 +24,34 @@ int32 FPBDLongRangeConstraintsBase::GetMinParallelBatchSize()
 	return ChaosPBDLongRangeConstraintsMinParallelBatchSize;
 }
 
+FPBDLongRangeConstraintsBase::FPBDLongRangeConstraintsBase(
+	const FSolverParticlesRange& Particles,
+	const TArray<TConstArrayView<TTuple<int32, int32, FRealSingle>>>& InTethers,
+	const TConstArrayView<FRealSingle>& StiffnessMultipliers,
+	const TConstArrayView<FRealSingle>& ScaleMultipliers,
+	const FSolverVec2& InStiffness,
+	const FSolverVec2& InScale,
+	FSolverReal MaxStiffness,
+	FSolverReal MeshScale)
+	: Tethers(InTethers)
+	, ParticleOffset(0)
+	, ParticleCount(Particles.GetRangeSize())
+	, Stiffness(
+		InStiffness,
+		StiffnessMultipliers,
+		ParticleCount,
+		FPBDStiffness::DefaultTableSize,
+		FPBDStiffness::DefaultParameterFitBase,
+		MaxStiffness)
+	, TetherScale(
+		InScale.ClampAxes(MinTetherScale, MaxTetherScale)* MeshScale,
+		ScaleMultipliers,
+		ParticleCount)
+{
+	// Apply default properties
+	ApplyProperties((FSolverReal)(1. / FPBDStiffness::ParameterFrequency), 1);
+}
+
 FPBDLongRangeConstraintsBase::FPBDLongRangeConstraintsBase(
 	const FSolverParticles& Particles,
 	const int32 InParticleOffset,
diff --git a/Engine/Source/Runtime/Experimental/Chaos/Private/Chaos/PBDSelfCollisionSphereConstraints.cpp b/Engine/Source/Runtime/Experimental/Chaos/Private/Chaos/PBDSelfCollisionSphereConstraints.cpp
--- a/Engine/Source/Runtime/Experimental/Chaos/Private/Chaos/PBDSelfCollisionSphereConstraints.cpp
+++ b/Engine/Source/Runtime/Experimental/Chaos/Private/Chaos/PBDSelfCollisionSphereConstraints.cpp
@@ -45,7 +45,8 @@ namespace Chaos::Softs {
 	{
 	}
 
-	void FPBDSelfCollisionSphereConstraintsBase::Init(const FSolverParticles& Particles)
+	template<typename SolverParticlesOrRange>
+	void FPBDSelfCollisionSphereConstraintsBase::Init(const SolverParticlesOrRange& Particles)
 	{
 		Constraints.Reset();
 
@@ -81,8 +82,10 @@ namespace Chaos::Softs {
 		constexpr int32 CellRadius = 1; // We set the cell size of the spatial hash such that we only need to look 1 cell away to find proximities.
 		constexpr int32 MaxNumExpectedConnectionsPerParticle = 3;
 		const int32 MaxNumExpectedConnections = MaxNumExpectedConnectionsPerParticle * Entries.Num();
+
+		const TConstArrayView<FSolverVec3> ReferencePositionsView = ReferencePositions ? Particles.GetConstArrayView(*ReferencePositions) : TConstArrayView<FSolverVec3>();
 		Constraints = SpatialHash.FindAllSelfProximities(CellRadius, MaxNumExpectedConnections,
-			[this, &Particles, DiamSq](const int32 i1, const int32 i2)
+			[this, &Particles, DiamSq, &ReferencePositionsView](const int32 i1, const int32 i2)
 		{
 			const FSolverReal CombinedMass = Particles.InvM(i1) + Particles.InvM(i2);
 			if (CombinedMass < (FSolverReal)1e-7)
@@ -91,7 +94,7 @@ namespace Chaos::Softs {
 			}
 			if (ReferencePositions)
 			{
-				if (FSolverVec3::DistSquared((*ReferencePositions)[i1], (*ReferencePositions)[i2]) < DiamSq)
+				if (FSolverVec3::DistSquared(ReferencePositionsView[i1], ReferencePositionsView[i2]) < DiamSq)
 				{
 					return false;
 				}
@@ -100,8 +103,11 @@ namespace Chaos::Softs {
 		}
 		);
 	}
+	template CHAOS_API void FPBDSelfCollisionSphereConstraintsBase::Init(const FSolverParticles& Particles);
+	template CHAOS_API void FPBDSelfCollisionSphereConstraintsBase::Init(const FSolverParticlesRange& Particles);
 
-	void FPBDSelfCollisionSphereConstraintsBase::Apply(FSolverParticles& Particles, const FSolverReal Dt) const
+	template<typename SolverParticlesOrRange>
+	void FPBDSelfCollisionSphereConstraintsBase::Apply(SolverParticlesOrRange& Particles, const FSolverReal Dt) const
 	{
 		const FSolverReal Diameter = 2.f * Radius;
 		const FSolverReal DiamSq = FMath::Square(Diameter);
@@ -132,6 +138,8 @@ namespace Chaos::Softs {
 			}
 		}
 	}
+	template CHAOS_API void FPBDSelfCollisionSphereConstraintsBase::Apply(FSolverParticles& Particles, const FSolverReal Dt) const;
+	template CHAOS_API void FPBDSelfCollisionSphereConstraintsBase::Apply(FSolverParticlesRange& Particles, const FSolverReal Dt) const;
 
 	void FPBDSelfCollisionSphereConstraints::SetProperties(
 		const FCollectionPropertyConstFacade& PropertyCollection,
diff --git a/Engine/Source/Runtime/Experimental/Chaos/Private/Chaos/PBDSoftBodyCollisionConstraint.cpp b/Engine/Source/Runtime/Experimental/Chaos/Private/Chaos/PBDSoftBodyCollisionConstraint.cpp
new file mode 100644
--- /dev/null
+++ b/Engine/Source/Runtime/Experimental/Chaos/Private/Chaos/PBDSoftBodyCollisionConstraint.cpp
@@ -0,0 +1,526 @@
+// Copyright Epic Games, Inc. All Rights Reserved.
+
+#include "Chaos/PBDSoftBodyCollisionConstraint.h"
+#include "ChaosStats.h"
+#include "HAL/IConsoleManager.h"
+#include "Chaos/Framework/Parallel.h"
+#include "Misc/ScopeLock.h"
+#if INTEL_ISPC
+#include "PerParticlePBDCollisionConstraint.ispc.generated.h"
+
+// These includes are used for the static_asserts below
+#include "Chaos/Capsule.h"
+#include "Chaos/Convex.h"
+#include "Chaos/ImplicitObjectUnion.h"
+#include "Chaos/Sphere.h"
+#include "Chaos/TaperedCapsule.h"
+#include "Chaos/TaperedCylinder.h"
+#include "Chaos/WeightedLatticeImplicitObject.h"
+#include "Chaos/Levelset.h"
+#endif
+
+#if INTEL_ISPC
+static_assert(sizeof(ispc::FVector3f) == sizeof(Chaos::Softs::FSolverVec3), "sizeof(ispc::FVector3f) != sizeof(Chaos::Softs::FSolverVec3)");
+static_assert(sizeof(ispc::FVector4f) == sizeof(Chaos::Softs::FSolverRotation3), "sizeof(ispc::FVector4f) != sizeof(Chaos::Softs::FSolverRotation3)");
+static_assert(ispc::ImplicitObjectType::Sphere == Chaos::ImplicitObjectType::Sphere);
+static_assert(ispc::ImplicitObjectType::Capsule == Chaos::ImplicitObjectType::Capsule);
+static_assert(ispc::ImplicitObjectType::Union == Chaos::ImplicitObjectType::Union);
+static_assert(ispc::ImplicitObjectType::TaperedCapsule == Chaos::ImplicitObjectType::TaperedCapsule);
+static_assert(ispc::ImplicitObjectType::Convex == Chaos::ImplicitObjectType::Convex);
+static_assert(ispc::ImplicitObjectType::WeightedLatticeBone == Chaos::ImplicitObjectType::WeightedLatticeBone);
+static_assert(ispc::ImplicitObjectType::IsWeightedLattice == Chaos::ImplicitObjectType::IsWeightedLattice);
+static_assert(ispc::ImplicitObjectType::WeightedLatticeLevelSetType == (Chaos::ImplicitObjectType::IsWeightedLattice | Chaos::ImplicitObjectType::LevelSet));
+static_assert(sizeof(ispc::TArray) == sizeof(TArray<int>));
+// Sphere
+static_assert(sizeof(Chaos::FImplicitObject) == Chaos::FSphere::FISPCDataVerifier::OffsetOfCenter());
+static_assert(sizeof(ispc::FVector) == Chaos::FSphere::FISPCDataVerifier::SizeOfCenter());
+// Capsule
+static_assert(sizeof(Chaos::FImplicitObject) == Chaos::FCapsule::FISPCDataVerifier::OffsetOfMSegment());
+static_assert(sizeof(ispc::Segment) == Chaos::FCapsule::FISPCDataVerifier::SizeOfMSegment());
+// Union (only specific case of MObjects = [FTaperedCylinder, Sphere, Sphere] is used here.
+static_assert(sizeof(Chaos::FImplicitObject) == Chaos::FImplicitObjectUnion::FISPCDataVerifier::OffsetOfMObjects());
+static_assert(sizeof(ispc::TArray) == Chaos::FImplicitObjectUnion::FISPCDataVerifier::SizeOfMObjects());
+// TaperedCylinder
+static_assert(sizeof(Chaos::FImplicitObject) + offsetof(ispc::FTaperedCylinder, MPlane1) == Chaos::FTaperedCylinder::FISPCDataVerifier::OffsetOfMPlane1());
+static_assert(sizeof(ispc::FTaperedCylinder::MPlane1) == Chaos::FTaperedCylinder::FISPCDataVerifier::SizeOfMPlane1());
+static_assert(sizeof(Chaos::FImplicitObject) + offsetof(ispc::FTaperedCylinder, MPlane2) == Chaos::FTaperedCylinder::FISPCDataVerifier::OffsetOfMPlane2());
+static_assert(sizeof(ispc::FTaperedCylinder::MPlane2) == Chaos::FTaperedCylinder::FISPCDataVerifier::SizeOfMPlane2());
+static_assert(sizeof(Chaos::FImplicitObject) + offsetof(ispc::FTaperedCylinder, Height) == Chaos::FTaperedCylinder::FISPCDataVerifier::OffsetOfMHeight());
+static_assert(sizeof(ispc::FTaperedCylinder::Height) == Chaos::FTaperedCylinder::FISPCDataVerifier::SizeOfMHeight());
+static_assert(sizeof(Chaos::FImplicitObject) + offsetof(ispc::FTaperedCylinder, Radius1) == Chaos::FTaperedCylinder::FISPCDataVerifier::OffsetOfMRadius1());
+static_assert(sizeof(ispc::FTaperedCylinder::Radius1) == Chaos::FTaperedCylinder::FISPCDataVerifier::SizeOfMRadius1());
+static_assert(sizeof(Chaos::FImplicitObject) + offsetof(ispc::FTaperedCylinder, Radius2) == Chaos::FTaperedCylinder::FISPCDataVerifier::OffsetOfMRadius2());
+static_assert(sizeof(ispc::FTaperedCylinder::Radius2) == Chaos::FTaperedCylinder::FISPCDataVerifier::SizeOfMRadius2());
+// TaperedCapsule
+static_assert(sizeof(Chaos::FImplicitObject) + offsetof(ispc::FTaperedCapsule, Origin) == Chaos::FTaperedCapsule::FISPCDataVerifier::OffsetOfOrigin());
+static_assert(sizeof(ispc::FTaperedCapsule::Origin) == Chaos::FTaperedCapsule::FISPCDataVerifier::SizeOfOrigin());
+static_assert(sizeof(Chaos::FImplicitObject) + offsetof(ispc::FTaperedCapsule, Axis) == Chaos::FTaperedCapsule::FISPCDataVerifier::OffsetOfAxis());
+static_assert(sizeof(ispc::FTaperedCapsule::Axis) == Chaos::FTaperedCapsule::FISPCDataVerifier::SizeOfAxis());
+static_assert(sizeof(Chaos::FImplicitObject) + offsetof(ispc::FTaperedCapsule, Height) == Chaos::FTaperedCapsule::FISPCDataVerifier::OffsetOfHeight());
+static_assert(sizeof(ispc::FTaperedCapsule::Height) == Chaos::FTaperedCapsule::FISPCDataVerifier::SizeOfHeight());
+static_assert(sizeof(Chaos::FImplicitObject) + offsetof(ispc::FTaperedCapsule, Radius1) == Chaos::FTaperedCapsule::FISPCDataVerifier::OffsetOfRadius1());
+static_assert(sizeof(ispc::FTaperedCapsule::Radius1) == Chaos::FTaperedCapsule::FISPCDataVerifier::SizeOfRadius1());
+static_assert(sizeof(Chaos::FImplicitObject) + offsetof(ispc::FTaperedCapsule, Radius2) == Chaos::FTaperedCapsule::FISPCDataVerifier::OffsetOfRadius2());
+static_assert(sizeof(ispc::FTaperedCapsule::Radius2) == Chaos::FTaperedCapsule::FISPCDataVerifier::SizeOfRadius2());
+// Convex
+static_assert(sizeof(Chaos::FImplicitObject) + offsetof(ispc::FConvex, Planes) == Chaos::FConvex::FISPCDataVerifier::OffsetOfPlanes());
+static_assert(sizeof(ispc::FConvex::Planes) == Chaos::FConvex::FISPCDataVerifier::SizeOfPlanes());
+static_assert(sizeof(Chaos::FImplicitObject) + offsetof(ispc::FConvex, Vertices) == Chaos::FConvex::FISPCDataVerifier::OffsetOfVertices());
+static_assert(sizeof(ispc::FConvex::Vertices) == Chaos::FConvex::FISPCDataVerifier::SizeOfVertices());
+static_assert(sizeof(Chaos::FImplicitObject) + offsetof(ispc::FConvex, StructureData) == Chaos::FConvex::FISPCDataVerifier::OffsetOfStructureData());
+static_assert(sizeof(ispc::FConvex::StructureData) == Chaos::FConvex::FISPCDataVerifier::SizeOfStructureData());
+
+static_assert(sizeof(ispc::FPlaneConcrete3f) == sizeof(Chaos::FConvex::FPlaneType));
+static_assert(sizeof(ispc::FVector3f) == sizeof(Chaos::FConvex::FVec3Type)); // Vertices array
+
+static_assert(offsetof(ispc::FConvexStructureData, Data) == Chaos::FConvexStructureData::FISPCDataVerifier::OffsetOfData());
+static_assert(sizeof(ispc::FConvexStructureData::Data) == Chaos::FConvexStructureData::FISPCDataVerifier::SizeOfData());
+static_assert(offsetof(ispc::FConvexStructureData, IndexType) == Chaos::FConvexStructureData::FISPCDataVerifier::OffsetOfIndexType());
+static_assert(sizeof(ispc::FConvexStructureData::IndexType) == Chaos::FConvexStructureData::FISPCDataVerifier::SizeOfIndexType());
+
+static_assert(offsetof(ispc::FConvexStructureDataImp, Planes) == Chaos::FConvexStructureData::FConvexStructureDataSmall::FISPCDataVerifier::OffsetOfPlanes());
+static_assert(offsetof(ispc::FConvexStructureDataImp, Planes) == Chaos::FConvexStructureData::FConvexStructureDataMedium::FISPCDataVerifier::OffsetOfPlanes());
+static_assert(offsetof(ispc::FConvexStructureDataImp, Planes) == Chaos::FConvexStructureData::FConvexStructureDataLarge::FISPCDataVerifier::OffsetOfPlanes());
+static_assert(sizeof(ispc::FConvexStructureDataImp::Planes) == Chaos::FConvexStructureData::FConvexStructureDataSmall::FISPCDataVerifier::SizeOfPlanes());
+static_assert(sizeof(ispc::FConvexStructureDataImp::Planes) == Chaos::FConvexStructureData::FConvexStructureDataMedium::FISPCDataVerifier::SizeOfPlanes());
+static_assert(sizeof(ispc::FConvexStructureDataImp::Planes) == Chaos::FConvexStructureData::FConvexStructureDataLarge::FISPCDataVerifier::SizeOfPlanes());
+
+static_assert(offsetof(ispc::FConvexStructureDataImp, HalfEdges) == Chaos::FConvexStructureData::FConvexStructureDataSmall::FISPCDataVerifier::OffsetOfHalfEdges());
+static_assert(offsetof(ispc::FConvexStructureDataImp, HalfEdges) == Chaos::FConvexStructureData::FConvexStructureDataMedium::FISPCDataVerifier::OffsetOfHalfEdges());
+static_assert(offsetof(ispc::FConvexStructureDataImp, HalfEdges) == Chaos::FConvexStructureData::FConvexStructureDataLarge::FISPCDataVerifier::OffsetOfHalfEdges());
+static_assert(sizeof(ispc::FConvexStructureDataImp::HalfEdges) == Chaos::FConvexStructureData::FConvexStructureDataSmall::FISPCDataVerifier::SizeOfHalfEdges());
+static_assert(sizeof(ispc::FConvexStructureDataImp::HalfEdges) == Chaos::FConvexStructureData::FConvexStructureDataMedium::FISPCDataVerifier::SizeOfHalfEdges());
+static_assert(sizeof(ispc::FConvexStructureDataImp::HalfEdges) == Chaos::FConvexStructureData::FConvexStructureDataLarge::FISPCDataVerifier::SizeOfHalfEdges());
+
+static_assert(offsetof(ispc::FConvexStructureDataImp, Vertices) == Chaos::FConvexStructureData::FConvexStructureDataSmall::FISPCDataVerifier::OffsetOfVertices());
+static_assert(offsetof(ispc::FConvexStructureDataImp, Vertices) == Chaos::FConvexStructureData::FConvexStructureDataMedium::FISPCDataVerifier::OffsetOfVertices());
+static_assert(offsetof(ispc::FConvexStructureDataImp, Vertices) == Chaos::FConvexStructureData::FConvexStructureDataLarge::FISPCDataVerifier::OffsetOfVertices());
+static_assert(sizeof(ispc::FConvexStructureDataImp::Vertices) == Chaos::FConvexStructureData::FConvexStructureDataSmall::FISPCDataVerifier::SizeOfVertices());
+static_assert(sizeof(ispc::FConvexStructureDataImp::Vertices) == Chaos::FConvexStructureData::FConvexStructureDataMedium::FISPCDataVerifier::SizeOfVertices());
+static_assert(sizeof(ispc::FConvexStructureDataImp::Vertices) == Chaos::FConvexStructureData::FConvexStructureDataLarge::FISPCDataVerifier::SizeOfVertices());
+
+static_assert(sizeof(ispc::PlanesS) == sizeof(Chaos::FConvexStructureData::FConvexStructureDataSmall::FPlaneData));
+static_assert(sizeof(ispc::PlanesM) == sizeof(Chaos::FConvexStructureData::FConvexStructureDataMedium::FPlaneData));
+static_assert(sizeof(ispc::PlanesL) == sizeof(Chaos::FConvexStructureData::FConvexStructureDataLarge::FPlaneData));
+static_assert(sizeof(ispc::HalfEdgesS) == sizeof(Chaos::FConvexStructureData::FConvexStructureDataSmall::FHalfEdgeData));
+static_assert(sizeof(ispc::HalfEdgesM) == sizeof(Chaos::FConvexStructureData::FConvexStructureDataMedium::FHalfEdgeData));
+static_assert(sizeof(ispc::HalfEdgesL) == sizeof(Chaos::FConvexStructureData::FConvexStructureDataLarge::FHalfEdgeData));
+static_assert(sizeof(ispc::VerticesS) == sizeof(Chaos::FConvexStructureData::FConvexStructureDataSmall::FVertexData));
+static_assert(sizeof(ispc::VerticesM) == sizeof(Chaos::FConvexStructureData::FConvexStructureDataMedium::FVertexData));
+static_assert(sizeof(ispc::VerticesL) == sizeof(Chaos::FConvexStructureData::FConvexStructureDataLarge::FVertexData));
+
+#if !UE_BUILD_SHIPPING
+bool bChaos_SoftBodyCollision_ISPC_Enabled = CHAOS_PER_PARTICLE_COLLISION_ISPC_ENABLED_DEFAULT;
+static FAutoConsoleVariableRef CVarChaosSoftBodyCollisionISPCEnabled(TEXT("p.Chaos.SoftBodyCollision.ISPC"), bChaos_SoftBodyCollision_ISPC_Enabled, TEXT("Whether to use ISPC optimizations in per particle collisions"));
+#endif
+#endif
+
+static int32 Chaos_SoftBodyCollision_ISPC_ParallelBatchSize = 128;
+#if !UE_BUILD_SHIPPING
+FAutoConsoleVariableRef CVarChaosSoftBodyCollisionISPCParallelBatchSize(TEXT("p.Chaos.SoftBodyCollision.ISPC.ParallelBatchSize"), Chaos_SoftBodyCollision_ISPC_ParallelBatchSize, TEXT("Parallel batch size for ISPC"));
+#endif
+
+namespace Chaos::Softs {
+
+// Helper function to call PhiWithNormal and return data to ISPC
+extern "C" void GetPhiWithNormalCollisionParticleRange(const uint8 * CollisionParticlesRange, const FSolverReal * InV, FSolverReal * Normal, FSolverReal * Phi, const int32 i, const int32 ProgramCount, const int32 Mask)
+{
+	const FSolverCollisionParticlesRange& C = *(const FSolverCollisionParticlesRange*)CollisionParticlesRange;
+
+	for (int32 Index = 0; Index < ProgramCount; ++Index)
+	{
+		if (Mask & (1 << Index))
+		{
+			FSolverVec3 V;
+
+			// aos_to_soa3
+			V.X = InV[Index];
+			V.Y = InV[Index + ProgramCount];
+			V.Z = InV[Index + 2 * ProgramCount];
+
+			FVec3 ImplicitNormal;
+			Phi[Index] = (FSolverReal)C.GetGeometry(i)->PhiWithNormal(FVec3(V), ImplicitNormal);
+			FSolverVec3 Norm(ImplicitNormal);
+
+			// aos_to_soa3
+			Normal[Index] = Norm.X;
+			Normal[Index + ProgramCount] = Norm.Y;
+			Normal[Index + 2 * ProgramCount] = Norm.Z;
+		}
+	}
+}
+
+// Helper function to call PhiWithNormal and return data to ISPC
+extern "C" void GetPhiWithNormalAndVelocityBoneCollisionParticleRange(const uint8 * CollisionParticlesRange, const FSolverReal * InV, FSolverReal * Normal, FSolverReal * Phi, int32 * VelocityBone, const int32 i, const FSolverReal Thickness, const int32 ProgramCount, const int32 Mask)
+{
+	const FSolverCollisionParticlesRange& C = *(const FSolverCollisionParticlesRange*)CollisionParticlesRange;
+
+	for (int32 Index = 0; Index < ProgramCount; ++Index)
+	{
+		if (Mask & (1 << Index))
+		{
+			FSolverVec3 V;
+
+			// aos_to_soa3
+			V.X = InV[Index];
+			V.Y = InV[Index + ProgramCount];
+			V.Z = InV[Index + 2 * ProgramCount];
+
+			VelocityBone[Index] = i;
+
+			FVec3 ImplicitNormal;
+			if (const TWeightedLatticeImplicitObject<FLevelSet>* LevelSet = C.GetGeometry(i)->GetObject< TWeightedLatticeImplicitObject<FLevelSet> >())
+			{
+				FWeightedLatticeImplicitObject::FEmbeddingCoordinate SurfaceCoord;
+				Phi[Index] = (FSolverReal)LevelSet->PhiWithNormalAndSurfacePoint(FVec3(V), ImplicitNormal, SurfaceCoord);
+				const FSolverReal Penetration = Thickness - Phi[Index]; // This is related to the Normal impulse
+				if (Penetration > (FSolverReal)0.)
+				{
+					const int32 StrongestBone = SurfaceCoord.GreatestInfluenceBone(LevelSet->GetBoneData());
+					if (StrongestBone != INDEX_NONE)
+					{
+						VelocityBone[Index] = LevelSet->GetSolverBoneIndices()[StrongestBone];
+					}
+				}
+			}
+			else
+			{
+				Phi[Index] = (FSolverReal)C.GetGeometry(i)->PhiWithNormal(FVec3(V), ImplicitNormal);
+			}
+			FSolverVec3 Norm(ImplicitNormal);
+
+			// aos_to_soa3
+			Normal[Index] = Norm.X;
+			Normal[Index + ProgramCount] = Norm.Y;
+			Normal[Index + 2 * ProgramCount] = Norm.Z;
+		}
+	}
+}
+
+
+void FPBDSoftBodyCollisionConstraintBase::Apply(FSolverParticlesRange& Particles, const FSolverReal Dt, const TArray<FSolverCollisionParticlesRange>& CollisionParticles) const
+{
+	TRACE_CPUPROFILER_EVENT_SCOPE(FPBDSoftBodyCollisionConstraint);
+	const bool bLockAndWriteContacts = bWriteDebugContacts && CollisionParticleCollided && Contacts && Normals && Phis;
+	const bool bWithFriction = FrictionCoefficient > (FSolverReal)UE_KINDA_SMALL_NUMBER;
+
+	if (CollisionParticles.IsEmpty())
+	{
+		return;
+	}
+
+	if (bUseCCD)
+	{
+		if (bWithFriction)
+		{
+			if (bLockAndWriteContacts)
+			{
+				ApplyInternalCCD<true, true>(Particles, Dt, CollisionParticles);
+			}
+			else
+			{
+				ApplyInternalCCD<true, false>(Particles, Dt, CollisionParticles);
+			}
+		}
+		else
+		{
+			if (bLockAndWriteContacts)
+			{
+				ApplyInternalCCD<false, true>(Particles, Dt, CollisionParticles);
+			}
+			else
+			{
+				ApplyInternalCCD<false, false>(Particles, Dt, CollisionParticles);
+			}
+		}
+	}
+	else
+	{
+		if (bWithFriction)
+		{
+			if (bRealTypeCompatibleWithISPC && bChaos_SoftBodyCollision_ISPC_Enabled)
+			{
+				ApplyInternalISPC<true>(Particles, Dt, CollisionParticles);
+			}
+			else if (bLockAndWriteContacts)
+			{
+				ApplyInternal<true, true>(Particles, Dt, CollisionParticles);
+			}
+			else
+			{
+				ApplyInternal<true, false>(Particles, Dt, CollisionParticles);
+			}
+		}
+		else
+		{
+			if (bRealTypeCompatibleWithISPC && bChaos_SoftBodyCollision_ISPC_Enabled)
+			{
+				ApplyInternalISPC<false>(Particles, Dt, CollisionParticles);
+			}
+			else if (bLockAndWriteContacts)
+			{
+				ApplyInternal<false, true>(Particles, Dt, CollisionParticles);
+			}
+			else
+			{
+				ApplyInternal<false, false>(Particles, Dt, CollisionParticles);
+			}
+		}
+	}
+}
+
+
+template<bool bLockAndWriteContacts, bool bWithFriction>
+void FPBDSoftBodyCollisionConstraintBase::ApplyInternal(FSolverParticlesRange& Particles, const FSolverReal Dt, const TArray<FSolverCollisionParticlesRange>& CollisionParticles) const
+{
+	FPAndInvM* const PAndInvM = Particles.GetPAndInvM().GetData();
+	const FSolverVec3* const X = Particles.XArray().GetData();
+	PhysicsParallelFor(Particles.GetRangeSize(), [this, Dt, &PAndInvM, &X, &CollisionParticles](int32 Index)
+	{
+		if (PAndInvM[Index].InvM == (FSolverReal)0.)
+		{
+			return;
+		}
+
+		for (const FSolverCollisionParticlesRange& CollisionParticlesRange : CollisionParticles)
+		{
+			for(int32 CollisionIndex = 0; CollisionIndex < CollisionParticlesRange.GetRangeSize(); ++CollisionIndex)
+			{
+				if (CollisionParticlesRange.GetGeometry(CollisionIndex)->GetType() == Chaos::ImplicitObjectType::WeightedLatticeBone)
+				{
+					continue;
+				}
+
+				const FSolverRigidTransform3 Frame(CollisionParticlesRange.X(CollisionIndex), CollisionParticlesRange.R(CollisionIndex));
+				const FVec3 RigidSpacePosition(Frame.InverseTransformPosition(PAndInvM[Index].P));  // PhiWithNormal requires FReal based arguments
+				FVec3 ImplicitNormal;                                                                // since implicits don't use FSolverReal
+				FSolverReal Phi;
+				int32 VelocityBone = CollisionIndex;
+				FSolverReal Penetration;
+				if (const TWeightedLatticeImplicitObject<FLevelSet>* LevelSet = CollisionParticlesRange.GetGeometry(CollisionIndex)->GetObject< TWeightedLatticeImplicitObject<FLevelSet> >())
+				{
+					FWeightedLatticeImplicitObject::FEmbeddingCoordinate SurfaceCoord;
+					Phi = (FSolverReal)LevelSet->PhiWithNormalAndSurfacePoint(RigidSpacePosition, ImplicitNormal, SurfaceCoord);
+					Penetration = CollisionThickness - Phi; // This is related to the Normal impulse
+					if (Penetration > (FSolverReal)0.)
+					{
+						const int32 StrongestBone = SurfaceCoord.GreatestInfluenceBone(LevelSet->GetBoneData());
+						if (StrongestBone != INDEX_NONE)
+						{
+							VelocityBone = LevelSet->GetSolverBoneIndices()[StrongestBone];
+						}
+					}
+				}
+				else
+				{
+					Phi = (FSolverReal)CollisionParticlesRange.GetGeometry(CollisionIndex)->PhiWithNormal(RigidSpacePosition, ImplicitNormal);
+					Penetration = CollisionThickness - Phi; // This is related to the Normal impulse
+				}
+				const FSolverVec3 Normal(ImplicitNormal);
+
+				if (Penetration > (FSolverReal)0.)
+				{
+					const FSolverVec3 NormalWorld = Frame.TransformVector(Normal);
+
+					if constexpr (bLockAndWriteContacts)
+					{
+						checkSlow(Contacts);
+						checkSlow(Normals);
+						checkSlow(Phis);
+						FScopeLock Lock(&DebugMutex);
+						CollisionParticlesRange.GetArrayView(*CollisionParticleCollided)[CollisionIndex] = true;
+						Contacts->Emplace(PAndInvM[Index].P);
+						Normals->Emplace(NormalWorld);
+						Phis->Emplace(Phi);
+					}
+
+					PAndInvM[Index].P += Penetration * NormalWorld;
+
+					if constexpr (bWithFriction)
+					{
+						const FSolverVec3 VectorToPoint = PAndInvM[Index].P - CollisionParticlesRange.X(VelocityBone);
+						const FSolverVec3 RelativeDisplacement = (PAndInvM[Index].P - X[Index]) - (CollisionParticlesRange.V(VelocityBone) + FSolverVec3::CrossProduct(CollisionParticlesRange.W(VelocityBone), VectorToPoint)) * Dt; // This corresponds to the tangential velocity multiplied by dt (friction will drive this to zero if it is high enough)
+						const FSolverVec3 RelativeDisplacementTangent = RelativeDisplacement - FSolverVec3::DotProduct(RelativeDisplacement, NormalWorld) * NormalWorld; // Project displacement into the tangential plane
+						const FSolverReal RelativeDisplacementTangentLength = RelativeDisplacementTangent.Size();
+						if (RelativeDisplacementTangentLength >= UE_SMALL_NUMBER)
+						{
+							const FSolverReal PositionCorrection = FMath::Min<FSolverReal>(Penetration * FrictionCoefficient, RelativeDisplacementTangentLength);
+							const FSolverReal CorrectionRatio = PositionCorrection / RelativeDisplacementTangentLength;
+							PAndInvM[Index].P -= CorrectionRatio * RelativeDisplacementTangent;
+						}
+					}
+				}
+			}
+		}
+	});
+}
+
+template<bool bLockAndWriteContacts, bool bWithFriction>
+void FPBDSoftBodyCollisionConstraintBase::ApplyInternalCCD(FSolverParticlesRange& Particles, const FSolverReal Dt, const TArray<FSolverCollisionParticlesRange>& CollisionParticles) const
+{
+	FPAndInvM* const PAndInvM = Particles.GetPAndInvM().GetData();
+	const FSolverVec3* const X = Particles.XArray().GetData();
+	PhysicsParallelFor(Particles.GetRangeSize(), [this, Dt, &PAndInvM, &X, &CollisionParticles](int32 Index)
+	{
+		if (PAndInvM[Index].InvM == (FSolverReal)0.)
+		{
+			return;
+		}
+
+		for (const FSolverCollisionParticlesRange& CollisionParticlesRange : CollisionParticles)
+		{
+			const TConstArrayView<FSolverRigidTransform3> CollisionTransforms = CollisionParticlesRange.GetConstArrayView(LastCollisionTransforms);
+
+			for (int32 CollisionIndex = 0; CollisionIndex < CollisionParticlesRange.GetRangeSize(); ++CollisionIndex)
+			{
+				if (CollisionParticlesRange.GetGeometry(CollisionIndex)->GetType() == Chaos::ImplicitObjectType::WeightedLatticeBone)
+				{
+					continue;
+				}
+
+				const FSolverRigidTransform3 Frame(CollisionParticlesRange.X(CollisionIndex), CollisionParticlesRange.R(CollisionIndex));
+
+				const Pair<FVec3, bool> PointPair = CollisionParticlesRange.GetGeometry(CollisionIndex)->FindClosestIntersection(  // Geometry operates in FReal
+					FVec3(CollisionTransforms[CollisionIndex].InverseTransformPositionNoScale(X[Index])),        // hence the back and forth
+					FVec3(Frame.InverseTransformPositionNoScale(PAndInvM[Index].P)), (FReal)CollisionThickness);                   // FVec3/FReal conversions
+
+				if (PointPair.Second)
+				{
+					const FSolverVec3 Normal = FSolverVec3(CollisionParticlesRange.GetGeometry(CollisionIndex)->Normal(PointPair.First));
+					const FSolverVec3 NormalWorld = Frame.TransformVectorNoScale(Normal);
+					const FSolverVec3 ContactWorld = Frame.TransformPositionNoScale(UE::Math::TVector<FSolverReal>(PointPair.First));
+
+					if constexpr (bLockAndWriteContacts)
+					{
+						checkSlow(Contacts);
+						checkSlow(Normals);
+						FScopeLock Lock(&DebugMutex);
+						CollisionParticlesRange.GetArrayView(*CollisionParticleCollided)[CollisionIndex] = true;
+						Contacts->Emplace(ContactWorld);
+						Normals->Emplace(NormalWorld);
+					}
+					const FSolverVec3 Direction = ContactWorld - PAndInvM[Index].P;
+					const FSolverReal Penetration = FMath::Max((FSolverReal)0., FSolverVec3::DotProduct(NormalWorld, Direction)) + (FSolverReal)UE_THRESH_POINT_ON_PLANE;
+
+					PAndInvM[Index].P += Penetration * NormalWorld;
+
+					// Friction
+					int32 VelocityBone = CollisionIndex;
+					if (const TWeightedLatticeImplicitObject<FLevelSet>* LevelSet = CollisionParticlesRange.GetGeometry(CollisionIndex)->GetObject< TWeightedLatticeImplicitObject<FLevelSet> >())
+					{
+						TArray<FWeightedLatticeImplicitObject::FEmbeddingCoordinate> Coordinates;
+						LevelSet->GetEmbeddingCoordinates(PointPair.First, Coordinates, false);
+						int32 ClosestCoordIndex = INDEX_NONE;
+						double ClosestCoordPhi = UE_BIG_NUMBER;
+						for (int32 CoordIndex = 0; CoordIndex < Coordinates.Num(); ++CoordIndex)
+						{
+							FVec3 NormalUnused;
+							const double CoordPhi = FMath::Abs(LevelSet->GetEmbeddedObject()->PhiWithNormal(Coordinates[CoordIndex].UndeformedPosition(LevelSet->GetGrid()), NormalUnused));
+							if (CoordPhi < ClosestCoordPhi)
+							{
+								ClosestCoordIndex = CoordIndex;
+								ClosestCoordPhi = CoordPhi;
+							}
+						}
+						if (ClosestCoordIndex != INDEX_NONE)
+						{
+							const int32 StrongestBone = Coordinates[ClosestCoordIndex].GreatestInfluenceBone(LevelSet->GetBoneData());
+							if (StrongestBone != INDEX_NONE)
+							{
+								VelocityBone = LevelSet->GetSolverBoneIndices()[StrongestBone];
+							}
+						}
+					}
+
+					const FSolverVec3 VectorToPoint = PAndInvM[Index].P - CollisionParticlesRange.X(VelocityBone);
+					const FSolverVec3 RelativeDisplacement = (PAndInvM[Index].P - X[Index]) - (CollisionParticlesRange.V(VelocityBone) + FSolverVec3::CrossProduct(CollisionParticlesRange.W(VelocityBone), VectorToPoint)) * Dt;  // This corresponds to the tangential velocity multiplied by dt (friction will drive this to zero if it is high enough)
+					const FSolverVec3 RelativeDisplacementTangent = RelativeDisplacement - FSolverVec3::DotProduct(RelativeDisplacement, NormalWorld) * NormalWorld;  // Project displacement into the tangential plane
+					const FSolverReal RelativeDisplacementTangentLength = RelativeDisplacementTangent.Size();
+					if (RelativeDisplacementTangentLength >= UE_SMALL_NUMBER)
+					{
+						const FSolverReal PositionCorrection = FMath::Min<FSolverReal>(Penetration * FrictionCoefficient, RelativeDisplacementTangentLength);
+						const FSolverReal CorrectionRatio = PositionCorrection / RelativeDisplacementTangentLength;
+						PAndInvM[Index].P -= CorrectionRatio * RelativeDisplacementTangent;
+					}
+				}
+			}
+		}
+	});
+}
+
+template<bool bWithFriction>
+void FPBDSoftBodyCollisionConstraintBase::ApplyInternalISPC(FSolverParticlesRange& Particles, const FSolverReal Dt, const TArray<FSolverCollisionParticlesRange>& CollisionParticles) const
+{
+#if INTEL_ISPC
+	check(bRealTypeCompatibleWithISPC);
+	const int32 NumBatches = FMath::CeilToInt((FSolverReal)(Particles.GetRangeSize()) / (FSolverReal)Chaos_SoftBodyCollision_ISPC_ParallelBatchSize);
+	PhysicsParallelFor(NumBatches, [this, &Particles, Dt, &CollisionParticles](int32 BatchNumber)
+	{
+		const int32 BatchBegin = (Chaos_SoftBodyCollision_ISPC_ParallelBatchSize * BatchNumber);
+		const int32 BatchEnd = FMath::Min(Particles.GetRangeSize(), BatchBegin + Chaos_SoftBodyCollision_ISPC_ParallelBatchSize);
+
+		for (const FSolverCollisionParticlesRange& CollisionParticlesRange : CollisionParticles)
+		{
+			if constexpr (bWithFriction)
+			{
+				ispc::ApplyPerParticleCollisionFastFrictionNoGroupCheck(
+					(ispc::FVector4f*)Particles.GetPAndInvM().GetData(),
+					(const ispc::FVector3f*)Particles.XArray().GetData(),
+					(const ispc::FVector3f*)CollisionParticlesRange.GetV().GetData(),
+					(const ispc::FVector3f*)CollisionParticlesRange.XArray().GetData(),
+					(const ispc::FVector3f*)CollisionParticlesRange.GetW().GetData(),
+					(const ispc::FVector4f*)CollisionParticlesRange.GetR().GetData(),
+					FrictionCoefficient,
+					CollisionThickness,
+					(const uint8*)&CollisionParticlesRange,
+					(const uint8*)CollisionParticlesRange.GetAllGeometry().GetData(),
+					sizeof(FImplicitObject),
+					FImplicitObject::GetOffsetOfType(),
+					FImplicitObject::GetOffsetOfMargin(),
+					Dt,
+					CollisionParticlesRange.GetRangeSize(),
+					BatchBegin,
+					BatchEnd);
+			}
+			else
+			{
+				ispc::ApplyPerParticleCollisionNoFrictionNoGroupCheck(
+					(ispc::FVector4f*)Particles.GetPAndInvM().GetData(),
+					(const ispc::FVector3f*)Particles.XArray().GetData(),
+					(const ispc::FVector3f*)CollisionParticlesRange.GetV().GetData(),
+					(const ispc::FVector3f*)CollisionParticlesRange.XArray().GetData(),
+					(const ispc::FVector3f*)CollisionParticlesRange.GetW().GetData(),
+					(const ispc::FVector4f*)CollisionParticlesRange.GetR().GetData(),
+					CollisionThickness,
+					(const uint8*)&CollisionParticlesRange,
+					(const uint8*)CollisionParticlesRange.GetAllGeometry().GetData(),
+					sizeof(FImplicitObject),
+					FImplicitObject::GetOffsetOfType(),
+					FImplicitObject::GetOffsetOfMargin(),
+					Dt,
+					CollisionParticlesRange.GetRangeSize(),
+					BatchBegin,
+					BatchEnd);
+			}
+		}
+	});
+#endif  // #if INTEL_ISPC
+}
+
+void FPBDSoftBodyCollisionConstraint::SetProperties(const FCollectionPropertyConstFacade& PropertyCollection)
+{
+	if (IsCollisionThicknessMutable(PropertyCollection))
+	{
+		CollisionThickness = MeshScale * GetCollisionThickness(PropertyCollection);
+	}
+	if (IsFrictionCoefficientMutable(PropertyCollection))
+	{
+		FrictionCoefficient = GetFrictionCoefficient(PropertyCollection);
+	}
+	if (IsUseCCDMutable(PropertyCollection))
+	{
+		bUseCCD = GetUseCCD(PropertyCollection);
+	}
+}
+
+}  // End namespace Chaos::Softs
diff --git a/Engine/Source/Runtime/Experimental/Chaos/Private/Chaos/PBDSphericalConstraint.cpp b/Engine/Source/Runtime/Experimental/Chaos/Private/Chaos/PBDSphericalConstraint.cpp
--- a/Engine/Source/Runtime/Experimental/Chaos/Private/Chaos/PBDSphericalConstraint.cpp
+++ b/Engine/Source/Runtime/Experimental/Chaos/Private/Chaos/PBDSphericalConstraint.cpp
@@ -38,16 +38,18 @@ void FPBDSphericalConstraint::SetProperties(
 	SetScale(MeshScale);
 }
 
-void FPBDSphericalConstraint::ApplyHelperISPC(FSolverParticles& Particles, const FSolverReal /*Dt*/) const
+template<typename SolverParticlesOrRange>
+void FPBDSphericalConstraint::ApplyHelperISPC(SolverParticlesOrRange& Particles, const FSolverReal /*Dt*/) const
 {
 	check(bRealTypeCompatibleWithISPC);
+	const TConstArrayView<FSolverVec3> AnimationPositionsView = Particles.GetConstArrayView(AnimationPositions);
 
 #if INTEL_ISPC
 	if (SphereRadii.Num() == ParticleCount)
 	{
 		ispc::ApplySphericalConstraintsWithMap(
 			(ispc::FVector4f*)Particles.GetPAndInvM().GetData(),
-			(const ispc::FVector3f*)AnimationPositions.GetData(),
+			(const ispc::FVector3f*)AnimationPositionsView.GetData(),
 			SphereRadii.GetData(),
 			MaxDistanceBase * GetScale(),
 			MaxDistanceRange * GetScale(),
@@ -58,13 +60,15 @@ void FPBDSphericalConstraint::ApplyHelperISPC(FSolverParticles& Particles, const
 	{
 		ispc::ApplySphericalConstraintsWithoutMap(
 			(ispc::FVector4f*)Particles.GetPAndInvM().GetData(),
-			(const ispc::FVector3f*)AnimationPositions.GetData(),
+			(const ispc::FVector3f*)AnimationPositionsView.GetData(),
 			MaxDistanceBase * GetScale(),
 			ParticleOffset,
 			ParticleCount);
 	}
 #endif
 }
+template CHAOS_API void FPBDSphericalConstraint::ApplyHelperISPC(FSolverParticles& Particles, const FSolverReal /*Dt*/) const;
+template CHAOS_API void FPBDSphericalConstraint::ApplyHelperISPC(FSolverParticlesRange& Particles, const FSolverReal /*Dt*/) const;
 
 void FPBDSphericalBackstopConstraint::SetProperties(
 	const FCollectionPropertyConstFacade& PropertyCollection,
@@ -101,10 +105,13 @@ void FPBDSphericalBackstopConstraint::SetProperties(
 	SetScale(MeshScale);
 }
 
-void FPBDSphericalBackstopConstraint::ApplyLegacyHelperISPC(FSolverParticles& Particles, const FSolverReal /*Dt*/) const
+template<typename SolverParticlesOrRange>
+void FPBDSphericalBackstopConstraint::ApplyLegacyHelperISPC(SolverParticlesOrRange& Particles, const FSolverReal /*Dt*/) const
 {
 	check(bRealTypeCompatibleWithISPC);
 
+	const TConstArrayView<FSolverVec3> AnimationPositionsView = Particles.GetConstArrayView(AnimationPositions);
+	const TConstArrayView<FSolverVec3> AnimationNormalsView = Particles.GetConstArrayView(AnimationNormals);
 #if INTEL_ISPC
 	if (SphereRadii.Num() == ParticleCount)
 	{
@@ -112,8 +119,8 @@ void FPBDSphericalBackstopConstraint::ApplyLegacyHelperISPC(FSolverParticles& Pa
 		{
 			ispc::ApplyLegacySphericalBackstopConstraintsWithMaps(
 				(ispc::FVector4f*)Particles.GetPAndInvM().GetData(),
-				(const ispc::FVector3f*)AnimationPositions.GetData(),
-				(const ispc::FVector3f*)AnimationNormals.GetData(),
+				(const ispc::FVector3f*)AnimationPositionsView.GetData(),
+				(const ispc::FVector3f*)AnimationNormalsView.GetData(),
 				SphereOffsetDistances.GetData(),
 				SphereRadii.GetData(),
 				BackstopDistanceBase * GetScale(),
@@ -127,8 +134,8 @@ void FPBDSphericalBackstopConstraint::ApplyLegacyHelperISPC(FSolverParticles& Pa
 		{
 			ispc::ApplyLegacySphericalBackstopConstraintsWithRadiusMap(
 				(ispc::FVector4f*)Particles.GetPAndInvM().GetData(),
-				(const ispc::FVector3f*)AnimationPositions.GetData(),
-				(const ispc::FVector3f*)AnimationNormals.GetData(),
+				(const ispc::FVector3f*)AnimationPositionsView.GetData(),
+				(const ispc::FVector3f*)AnimationNormalsView.GetData(),
 				SphereRadii.GetData(),
 				BackstopDistanceBase * GetScale(),
 				BackstopRadiusBase * GetScale(),
@@ -141,8 +148,8 @@ void FPBDSphericalBackstopConstraint::ApplyLegacyHelperISPC(FSolverParticles& Pa
 	{
 		ispc::ApplyLegacySphericalBackstopConstraintsWithDistanceMap(
 			(ispc::FVector4f*)Particles.GetPAndInvM().GetData(),
-			(const ispc::FVector3f*)AnimationPositions.GetData(),
-			(const ispc::FVector3f*)AnimationNormals.GetData(),
+			(const ispc::FVector3f*)AnimationPositionsView.GetData(),
+			(const ispc::FVector3f*)AnimationNormalsView.GetData(),
 			SphereOffsetDistances.GetData(),
 			BackstopDistanceBase * GetScale(),
 			BackstopDistanceRange * GetScale(),
@@ -154,8 +161,8 @@ void FPBDSphericalBackstopConstraint::ApplyLegacyHelperISPC(FSolverParticles& Pa
 	{
 		ispc::ApplyLegacySphericalBackstopConstraintsWithoutMaps(
 			(ispc::FVector4f*)Particles.GetPAndInvM().GetData(),
-			(const ispc::FVector3f*)AnimationPositions.GetData(),
-			(const ispc::FVector3f*)AnimationNormals.GetData(),
+			(const ispc::FVector3f*)AnimationPositionsView.GetData(),
+			(const ispc::FVector3f*)AnimationNormalsView.GetData(),
 			BackstopDistanceBase * GetScale(),
 			BackstopRadiusBase * GetScale(),
 			ParticleOffset,
@@ -163,11 +170,16 @@ void FPBDSphericalBackstopConstraint::ApplyLegacyHelperISPC(FSolverParticles& Pa
 	}
 #endif
 }
+template CHAOS_API void FPBDSphericalBackstopConstraint::ApplyLegacyHelperISPC(FSolverParticles& Particles, const FSolverReal /*Dt*/) const;
+template CHAOS_API void FPBDSphericalBackstopConstraint::ApplyLegacyHelperISPC(FSolverParticlesRange& Particles, const FSolverReal /*Dt*/) const;
 
-void FPBDSphericalBackstopConstraint::ApplyHelperISPC(FSolverParticles& Particles, const FSolverReal /*Dt*/) const
+template<typename SolverParticlesOrRange>
+void FPBDSphericalBackstopConstraint::ApplyHelperISPC(SolverParticlesOrRange& Particles, const FSolverReal /*Dt*/) const
 {
 	check(bRealTypeCompatibleWithISPC);
 
+	const TConstArrayView<FSolverVec3> AnimationPositionsView = Particles.GetConstArrayView(AnimationPositions);
+	const TConstArrayView<FSolverVec3> AnimationNormalsView = Particles.GetConstArrayView(AnimationNormals);
 #if INTEL_ISPC
 	if (SphereRadii.Num() == ParticleCount)
 	{
@@ -175,8 +187,8 @@ void FPBDSphericalBackstopConstraint::ApplyHelperISPC(FSolverParticles& Particle
 		{
 			ispc::ApplySphericalBackstopConstraintsWithMaps(
 				(ispc::FVector4f*)Particles.GetPAndInvM().GetData(),
-				(const ispc::FVector3f*)AnimationPositions.GetData(),
-				(const ispc::FVector3f*)AnimationNormals.GetData(),
+				(const ispc::FVector3f*)AnimationPositionsView.GetData(),
+				(const ispc::FVector3f*)AnimationNormalsView.GetData(),
 				SphereOffsetDistances.GetData(),
 				SphereRadii.GetData(),
 				BackstopDistanceBase * GetScale(),
@@ -190,8 +202,8 @@ void FPBDSphericalBackstopConstraint::ApplyHelperISPC(FSolverParticles& Particle
 		{
 			ispc::ApplySphericalBackstopConstraintsWithRadiusMap(
 				(ispc::FVector4f*)Particles.GetPAndInvM().GetData(),
-				(const ispc::FVector3f*)AnimationPositions.GetData(),
-				(const ispc::FVector3f*)AnimationNormals.GetData(),
+				(const ispc::FVector3f*)AnimationPositionsView.GetData(),
+				(const ispc::FVector3f*)AnimationNormalsView.GetData(),
 				SphereRadii.GetData(),
 				BackstopDistanceBase * GetScale(),
 				BackstopRadiusBase * GetScale(),
@@ -204,8 +216,8 @@ void FPBDSphericalBackstopConstraint::ApplyHelperISPC(FSolverParticles& Particle
 	{
 		ispc::ApplySphericalBackstopConstraintsWithDistanceMap(
 			(ispc::FVector4f*)Particles.GetPAndInvM().GetData(),
-			(const ispc::FVector3f*)AnimationPositions.GetData(),
-			(const ispc::FVector3f*)AnimationNormals.GetData(),
+			(const ispc::FVector3f*)AnimationPositionsView.GetData(),
+			(const ispc::FVector3f*)AnimationNormalsView.GetData(),
 			SphereOffsetDistances.GetData(),
 			BackstopDistanceBase * GetScale(),
 			BackstopDistanceRange * GetScale(),
@@ -217,8 +229,8 @@ void FPBDSphericalBackstopConstraint::ApplyHelperISPC(FSolverParticles& Particle
 	{
 		ispc::ApplySphericalBackstopConstraintsWithoutMaps(
 			(ispc::FVector4f*)Particles.GetPAndInvM().GetData(),
-			(const ispc::FVector3f*)AnimationPositions.GetData(),
-			(const ispc::FVector3f*)AnimationNormals.GetData(),
+			(const ispc::FVector3f*)AnimationPositionsView.GetData(),
+			(const ispc::FVector3f*)AnimationNormalsView.GetData(),
 			BackstopDistanceBase * GetScale(),
 			BackstopRadiusBase * GetScale(),
 			ParticleOffset,
@@ -226,5 +238,7 @@ void FPBDSphericalBackstopConstraint::ApplyHelperISPC(FSolverParticles& Particle
 	}
 #endif
 }
+template CHAOS_API void FPBDSphericalBackstopConstraint::ApplyHelperISPC(FSolverParticles& Particles, const FSolverReal /*Dt*/) const;
+template CHAOS_API void FPBDSphericalBackstopConstraint::ApplyHelperISPC(FSolverParticlesRange& Particles, const FSolverReal /*Dt*/) const;
 
 }  // End namespace Chaos::Softs
diff --git a/Engine/Source/Runtime/Experimental/Chaos/Private/Chaos/PBDSpringConstraints.cpp b/Engine/Source/Runtime/Experimental/Chaos/Private/Chaos/PBDSpringConstraints.cpp
--- a/Engine/Source/Runtime/Experimental/Chaos/Private/Chaos/PBDSpringConstraints.cpp
+++ b/Engine/Source/Runtime/Experimental/Chaos/Private/Chaos/PBDSpringConstraints.cpp
@@ -30,14 +30,15 @@ static int32 Chaos_Spring_ParallelConstraintCount = 100;
 FAutoConsoleVariableRef CVarChaosSpringParallelConstraintCount(TEXT("p.Chaos.Spring.ParallelConstraintCount"), Chaos_Spring_ParallelConstraintCount, TEXT("If we have more constraints than this, use parallel-for in Apply."));
 #endif
 
-void FPBDSpringConstraints::InitColor(const FSolverParticles& Particles)
+template<typename SolverParticlesOrRange>
+void FPBDSpringConstraints::InitColor(const SolverParticlesOrRange& Particles)
 {
 	// In dev builds we always color so we can tune the system without restarting. See Apply()
 #if UE_BUILD_SHIPPING || UE_BUILD_TEST
 	if (Constraints.Num() > Chaos_Spring_ParallelConstraintCount)
 #endif
 	{
-		const TArray<TArray<int32>> ConstraintsPerColor = FGraphColoring::ComputeGraphColoring(Constraints, Particles, ParticleOffset, ParticleOffset + ParticleCount);
+		const TArray<TArray<int32>> ConstraintsPerColor = FGraphColoring::ComputeGraphColoringParticlesOrRange(Constraints, Particles, ParticleOffset, ParticleOffset + ParticleCount);
 		
 		// Reorder constraints based on color so each array in ConstraintsPerColor contains contiguous elements.
 		TArray<TVec2<int32>> ReorderedConstraints;
@@ -70,8 +71,11 @@ void FPBDSpringConstraints::InitColor(const FSolverParticles& Particles)
 		Stiffness.ReorderIndices(OrigToReorderedIndices);
 	}
 }
+template CHAOS_API void FPBDSpringConstraints::InitColor(const FSolverParticles& Particles);
+template CHAOS_API void FPBDSpringConstraints::InitColor(const FSolverParticlesRange& Particles);
 
-void FPBDSpringConstraints::ApplyHelper(FSolverParticles& Particles, const FSolverReal Dt, const int32 ConstraintIndex, const FSolverReal ExpStiffnessValue) const
+template<typename SolverParticlesOrRange>
+void FPBDSpringConstraints::ApplyHelper(SolverParticlesOrRange& Particles, const FSolverReal Dt, const int32 ConstraintIndex, const FSolverReal ExpStiffnessValue) const
 {
 	const TVec2<int32>& Constraint = Constraints[ConstraintIndex];
 	const int32 i1 = Constraint[0];
@@ -87,7 +91,8 @@ void FPBDSpringConstraints::ApplyHelper(FSolverParticles& Particles, const FSolv
 	}
 }
 
-void FPBDSpringConstraints::Apply(FSolverParticles& Particles, const FSolverReal Dt) const
+template<typename SolverParticlesOrRange>
+void FPBDSpringConstraints::Apply(SolverParticlesOrRange& Particles, const FSolverReal Dt) const
 {
 	TRACE_CPUPROFILER_EVENT_SCOPE(FPBDSpringConstraints_Apply);
 	SCOPE_CYCLE_COUNTER(STAT_PBD_Spring);
@@ -183,6 +188,8 @@ void FPBDSpringConstraints::Apply(FSolverParticles& Particles, const FSolverReal
 		}
 	}
 }
+template CHAOS_API void FPBDSpringConstraints::Apply(FSolverParticles& Particles, const FSolverReal Dt) const;
+template CHAOS_API void FPBDSpringConstraints::Apply(FSolverParticlesRange& Particles, const FSolverReal Dt) const;
 
 void FPBDEdgeSpringConstraints::SetProperties(
 	const FCollectionPropertyConstFacade& PropertyCollection,
diff --git a/Engine/Source/Runtime/Experimental/Chaos/Private/Chaos/PBDTriangleMeshCollisions.cpp b/Engine/Source/Runtime/Experimental/Chaos/Private/Chaos/PBDTriangleMeshCollisions.cpp
--- a/Engine/Source/Runtime/Experimental/Chaos/Private/Chaos/PBDTriangleMeshCollisions.cpp
+++ b/Engine/Source/Runtime/Experimental/Chaos/Private/Chaos/PBDTriangleMeshCollisions.cpp
@@ -2,6 +2,7 @@
 #include "Chaos/PBDTriangleMeshCollisions.h"
 #include "Chaos/Plane.h"
 #include "Chaos/PBDSoftsSolverParticles.h"
+#include "Chaos/SoftsSolverParticlesRange.h"
 #include "Chaos/Triangle.h"
 #include "Chaos/TriangleCollisionPoint.h"
 #include "Chaos/TriangleMesh.h"
@@ -22,8 +23,8 @@ struct FEdgeFaceIntersection
 };
 
 // Returned array has NOT been shrunk
-template<typename SpatialAccelerator>
-static void FindEdgeFaceIntersections(const FTriangleMesh& TriangleMesh, const SpatialAccelerator& Spatial, const FSolverParticles& Particles, TArray<FEdgeFaceIntersection>& Intersections)
+template<typename SpatialAccelerator, typename SolverParticlesOrRange>
+static void FindEdgeFaceIntersections(const FTriangleMesh& TriangleMesh, const SpatialAccelerator& Spatial, const SolverParticlesOrRange& Particles, TArray<FEdgeFaceIntersection>& Intersections)
 {
 	TRACE_CPUPROFILER_EVENT_SCOPE(ChaosFPBDTriangleMeshCollisions_IntersectionQuery);
 
@@ -1120,7 +1121,8 @@ namespace GIA
 // Calculating gradient direction to minimize contour length.
 namespace ContourMinimization
 {
-	static void BuildLocalContourMinimizationIntersection(const FEdgeFaceIntersection& EdgeFaceIntersection, const FTriangleMesh& TriangleMesh, const FSegmentMesh& SegmentMesh, const FSolverParticles& Particles, FPBDTriangleMeshCollisions::FContourMinimizationIntersection& ContourIntersection)
+	template<typename SolverParticlesOrRange>
+	static void BuildLocalContourMinimizationIntersection(const FEdgeFaceIntersection& EdgeFaceIntersection, const FTriangleMesh& TriangleMesh, const FSegmentMesh& SegmentMesh, const SolverParticlesOrRange& Particles, FPBDTriangleMeshCollisions::FContourMinimizationIntersection& ContourIntersection)
 	{
 		ContourIntersection.EdgeVertices = SegmentMesh.GetElements()[EdgeFaceIntersection.EdgeIndex];
 		ContourIntersection.FaceVertices = TriangleMesh.GetElements()[EdgeFaceIntersection.FaceIndex];
@@ -1183,7 +1185,8 @@ namespace ContourMinimization
 		ContourIntersection.GlobalGradientVector = ContourIntersection.LocalGradientVector;
 	}
 
-	static void BuildLocalContourMinimizationIntersections(const FTriangleMesh& TriangleMesh, const FSolverParticles& Particles, const TArray<FEdgeFaceIntersection>& Intersections, TArray<FPBDTriangleMeshCollisions::FContourMinimizationIntersection>& ContourMinimizationIntersections)
+	template<typename SolverParticlesOrRange>
+	static void BuildLocalContourMinimizationIntersections(const FTriangleMesh& TriangleMesh, const SolverParticlesOrRange& Particles, const TArray<FEdgeFaceIntersection>& Intersections, TArray<FPBDTriangleMeshCollisions::FContourMinimizationIntersection>& ContourMinimizationIntersections)
 	{
 		TRACE_CPUPROFILER_EVENT_SCOPE(ChaosFPBDTriangleMeshCollisions_BuildLocalContourMinimizationIntersections);
 		const FSegmentMesh& SegmentMesh = TriangleMesh.GetSegmentMesh();
@@ -1199,7 +1202,8 @@ namespace ContourMinimization
 		);
 	}
 
-	static void BuildGlobalContourMinimizationIntersections(const FTriangleMesh& TriangleMesh, const FSolverParticles& Particles, const TArray<GIA::FIntersectionContourPair>& IntersectionContours, TArray<FPBDTriangleMeshCollisions::FContourMinimizationIntersection>& ContourMinimizationIntersections)
+	template<typename SolverParticlesOrRange>
+	static void BuildGlobalContourMinimizationIntersections(const FTriangleMesh& TriangleMesh, const SolverParticlesOrRange& Particles, const TArray<GIA::FIntersectionContourPair>& IntersectionContours, TArray<FPBDTriangleMeshCollisions::FContourMinimizationIntersection>& ContourMinimizationIntersections)
 	{
 		TRACE_CPUPROFILER_EVENT_SCOPE(ChaosFPBDTriangleMeshCollisions_BuildGlobalContourMinimizationIntersections);
 		const FSegmentMesh& SegmentMesh = TriangleMesh.GetSegmentMesh();
@@ -1270,7 +1274,8 @@ struct FPBDTriangleMeshCollisions::FScratchBuffers
 	}
 };
 
-void FPBDTriangleMeshCollisions::Init(const FSolverParticles& Particles, const FSolverReal MinProximityQueryRadius)
+template<typename SolverParticlesOrRange>
+void FPBDTriangleMeshCollisions::Init(const SolverParticlesOrRange& Particles, const FSolverReal MinProximityQueryRadius)
 {
 	if (TriangleMesh.GetNumElements() == 0)
 	{
@@ -1331,8 +1336,11 @@ void FPBDTriangleMeshCollisions::Init(const FSolverParticles& Particles, const F
 		}
 	}
 }
+template CHAOS_API void FPBDTriangleMeshCollisions::Init(const FSolverParticles& Particles, const FSolverReal MinProximityQueryRadius);
+template CHAOS_API void FPBDTriangleMeshCollisions::Init(const FSolverParticlesRange& Particles, const FSolverReal MinProximityQueryRadius);
 
-void FPBDTriangleMeshCollisions::PostStepInit(const FSolverParticles& Particles)
+template<typename SolverParticlesOrRange>
+void FPBDTriangleMeshCollisions::PostStepInit(const SolverParticlesOrRange& Particles)
 {
 	if (TriangleMesh.GetNumElements() == 0)
 	{
@@ -1372,4 +1380,6 @@ void FPBDTriangleMeshCollisions::PostStepInit(const FSolverParticles& Particles)
 		}
 	}
 }
+template CHAOS_API void FPBDTriangleMeshCollisions::PostStepInit(const FSolverParticles& Particles);
+template CHAOS_API void FPBDTriangleMeshCollisions::PostStepInit(const FSolverParticlesRange& Particles);
 }  // End namespace Chaos::Softs
diff --git a/Engine/Source/Runtime/Experimental/Chaos/Private/Chaos/PBDTriangleMeshIntersections.cpp b/Engine/Source/Runtime/Experimental/Chaos/Private/Chaos/PBDTriangleMeshIntersections.cpp
--- a/Engine/Source/Runtime/Experimental/Chaos/Private/Chaos/PBDTriangleMeshIntersections.cpp
+++ b/Engine/Source/Runtime/Experimental/Chaos/Private/Chaos/PBDTriangleMeshIntersections.cpp
@@ -1,6 +1,7 @@
 // Copyright Epic Games, Inc. All Rights Reserved.
 #include "Chaos/PBDTriangleMeshIntersections.h"
 #include "Chaos/PBDSoftsSolverParticles.h"
+#include "Chaos/SoftsSolverParticlesRange.h"
 
 float Chaos_TriangleIntersections_MaxDelta = 0.01f;
 FAutoConsoleVariableRef CVarChaosTriangleIntersectionMaxImpulse(TEXT("p.Chaos.TriangleIntersections.MaxDelta"), Chaos_TriangleIntersections_MaxDelta, TEXT("Maximum delta position applied to resolve triangle intersections."));
@@ -17,8 +18,8 @@ static inline FSolverVec3 GetDelta(const FPBDTriangleMeshCollisions::FContourMin
 	return MaxDelta * GradientLength * FMath::InvSqrt(GradientLength * GradientLength + RegularizeEpsilonSq) * GradientDir;
 }
 
-
-void FPBDTriangleMeshIntersections::Apply(FSolverParticles& Particles, const TArray<FPBDTriangleMeshCollisions::FContourMinimizationIntersection>& Intersections, const FSolverReal Dt) const
+template<typename SolverParticlesOrRange>
+void FPBDTriangleMeshIntersections::Apply(SolverParticlesOrRange& Particles, const TArray<FPBDTriangleMeshCollisions::FContourMinimizationIntersection>& Intersections, const FSolverReal Dt) const
 {
 	TRACE_CPUPROFILER_EVENT_SCOPE(ChaosFPBDTriangleMeshIntersections_Apply);
 	static FSolverReal RegularizeEpsilonSq = 1.f;
@@ -47,5 +48,7 @@ void FPBDTriangleMeshIntersections::Apply(FSolverParticles& Particles, const TAr
 		}
 	}
 }
+template CHAOS_API void FPBDTriangleMeshIntersections::Apply(FSolverParticles& Particles, const TArray<FPBDTriangleMeshCollisions::FContourMinimizationIntersection>& Intersections, const FSolverReal Dt) const;
+template CHAOS_API void FPBDTriangleMeshIntersections::Apply(FSolverParticlesRange& Particles, const TArray<FPBDTriangleMeshCollisions::FContourMinimizationIntersection>& Intersections, const FSolverReal Dt) const;
 
 }  // End namespace Chaos::Softs
diff --git a/Engine/Source/Runtime/Experimental/Chaos/Private/Chaos/PerParticleDampVelocity.cpp b/Engine/Source/Runtime/Experimental/Chaos/Private/Chaos/PerParticleDampVelocity.cpp
--- a/Engine/Source/Runtime/Experimental/Chaos/Private/Chaos/PerParticleDampVelocity.cpp
+++ b/Engine/Source/Runtime/Experimental/Chaos/Private/Chaos/PerParticleDampVelocity.cpp
@@ -2,6 +2,7 @@
 
 #include "Chaos/PerParticleDampVelocity.h"
 #include "Chaos/Matrix.h"
+#include "Chaos/SoftsSolverParticlesRange.h"
 #include "GenericPlatform/GenericPlatformMath.h"
 
 #if INTEL_ISPC && !UE_BUILD_SHIPPING
@@ -21,6 +22,12 @@ static FAutoConsoleVariableRef CVarChaosDampVelocityISPCEnabled(TEXT("p.Chaos.Da
 namespace Chaos::Softs {
 
 void FPerParticleDampVelocity::UpdatePositionBasedState(const FSolverParticles& Particles, const int32 Offset, const int32 Range)
+{
+	const FSolverParticlesRange ParticlesRange(const_cast<FSolverParticles*>(&Particles), Offset, Range - Offset);
+	return UpdatePositionBasedState(ParticlesRange);
+}
+
+void FPerParticleDampVelocity::UpdatePositionBasedState(const FSolverParticlesRange& Particles)
 {
 #if INTEL_ISPC
 	if (bRealTypeCompatibleWithISPC && bChaos_DampVelocity_ISPC_Enabled)
@@ -33,8 +40,8 @@ void FPerParticleDampVelocity::UpdatePositionBasedState(const FSolverParticles&
 			(const ispc::FVector3f*)Particles.GetV().GetData(),
 			Particles.GetM().GetData(),
 			Particles.GetInvM().GetData(),
-			Offset,
-			Range);
+			0,
+			Particles.GetRangeSize());
 	}
 	else
 #endif
@@ -43,7 +50,7 @@ void FPerParticleDampVelocity::UpdatePositionBasedState(const FSolverParticles&
 		Vcm = FSolverVec3(0.);
 		FSolverReal Mcm = (FSolverReal)0.;
 
-		for (int32 Index = Offset; Index < Range; ++Index)
+		for (int32 Index = 0; Index < Particles.GetRangeSize(); ++Index)
 		{
 			if (Particles.InvM(Index) == (FSolverReal)0.)
 			{
@@ -63,7 +70,7 @@ void FPerParticleDampVelocity::UpdatePositionBasedState(const FSolverParticles&
 
 		FSolverVec3 L = FSolverVec3(0.);
 		FSolverMatrix33 I(0.);
-		for (int32 Index = Offset; Index < Range; ++Index)
+		for (int32 Index = 0; Index < Particles.GetRangeSize(); ++Index)
 		{
 			if (Particles.InvM(Index) == (FSolverReal)0.)
 			{
diff --git a/Engine/Source/Runtime/Experimental/Chaos/Private/Chaos/PerParticlePBDCollisionConstraint.ispc b/Engine/Source/Runtime/Experimental/Chaos/Private/Chaos/PerParticlePBDCollisionConstraint.ispc
--- a/Engine/Source/Runtime/Experimental/Chaos/Private/Chaos/PerParticlePBDCollisionConstraint.ispc
+++ b/Engine/Source/Runtime/Experimental/Chaos/Private/Chaos/PerParticlePBDCollisionConstraint.ispc
@@ -584,6 +584,8 @@ static inline void PhiWithNormalConvex(const uniform uint8 *uniform TypedObjectP
 
 extern "C" unmasked void GetPhiWithNormal(const uniform uint8 *uniform CollisionParticles, const uniform float *uniform V, uniform float *uniform Normal, uniform float *uniform Phi, const uniform int i, const uniform int ProgramCount, const uniform int Mask);
 extern "C" unmasked void GetPhiWithNormalAndVelocityBone(const uniform uint8 *uniform CollisionParticles, const uniform float *uniform V, uniform float *uniform Normal, uniform float *uniform Phi, uniform int *uniform VelocityBone, const uniform int i, const uniform float Thickness, const uniform int ProgramCount, const uniform int Mask);
+extern "C" unmasked void GetPhiWithNormalCollisionParticleRange(const uniform uint8 *uniform CollisionParticlesRange, const uniform float *uniform V, uniform float *uniform Normal, uniform float *uniform Phi, const uniform int i, const uniform int ProgramCount, const uniform int Mask);
+extern "C" unmasked void GetPhiWithNormalAndVelocityBoneCollisionParticleRange(const uniform uint8 *uniform CollisionParticlesRange, const uniform float *uniform V, uniform float *uniform Normal, uniform float *uniform Phi, uniform int *uniform VelocityBone, const uniform int i, const uniform float Thickness, const uniform int ProgramCount, const uniform int Mask);
 
 export void ApplyPerParticleCollisionFastFriction(uniform FVector4f ParticlesPandInvM[],
 													const uniform FVector3f ParticlesX[],
@@ -774,6 +776,168 @@ export void ApplyPerParticleCollisionFastFriction(uniform FVector4f ParticlesPan
 	}
 }
 
+export void ApplyPerParticleCollisionFastFrictionNoGroupCheck(uniform FVector4f ParticlesPandInvM[],
+													const uniform FVector3f ParticlesX[],
+													const uniform FVector3f CollisionV[],
+													const uniform FVector3f CollisionX[],
+													const uniform FVector3f CollisionW[],
+													const uniform FVector4f CollisionR[],
+													const uniform float PerGroupFriction,
+													const uniform float PerGroupThickness,
+													const uniform uint8 *uniform CollisionParticlesRange,
+													const uniform uint8 *uniform Geometry,
+													const uniform int SizeofFImplicitObject,
+													const uniform int OffsetofGeometryType,
+													const uniform int OffsetOfMargin,
+													const uniform float Dt,
+													const uniform int InnerCollisionCount,
+													const uniform int OuterParticleBatchBegin,
+													const uniform int OuterParticleBatchEnd)
+{
+	foreach(Index = OuterParticleBatchBegin ... OuterParticleBatchEnd)
+	{
+		const FVector4f PandInvM = VectorLoad(&ParticlesPandInvM[extract(Index, 0)]);
+		varying FVector3f P;
+		varying float IM;
+		UnzipPandInvM(PandInvM, P, IM);
+
+		if(all(IM == 0))
+		{
+			continue;
+		}
+
+		for(uniform int i = 0; i < InnerCollisionCount; i++) // SequentialFor
+		{
+			// InverseTransformPosition. Can skip Scale calculation because it is 1 for xyz.
+			const FVector3f TranslatedVec = P - CollisionX[i];
+			const FVector4f InverseR = QuatInverse(CollisionR[i]);
+			const FVector3f PhiWithNormalInputFloat = VectorQuaternionRotateVector(InverseR, TranslatedVec);
+			const FVector PhiWithNormalInput = ConvertVector3fTo3Native(PhiWithNormalInputFloat);
+			FVector NormalReal;
+			FReal PhiReal;
+			FVector3f Normal;
+			float Phi;
+
+			const uniform SerializedPtr *uniform GeometryPtr = (const uniform SerializedPtr *uniform)Geometry;
+			const uniform uint8 *uniform Ptr = GeometryPtr[i].Ptr;
+			const uniform uint8 *uniform TypedPtr = Ptr + SizeofFImplicitObject;
+			const uniform uint8 Type = Ptr[OffsetofGeometryType] & 0x3F;
+
+			int32 VelocityBone = i;
+			switch(Type)
+			{
+			case Sphere:
+				{
+					const uniform FReal Margin = *((const uniform FReal *uniform)(Ptr + OffsetOfMargin));
+					PhiWithNormalSphere(TypedPtr, PhiWithNormalInput, NormalReal, PhiReal, Margin);
+					Phi = (varying float)PhiReal;
+					Normal = ConvertVector3NativeTo3f(NormalReal);
+					break;
+				}
+			case Capsule:
+				{
+					const uniform FReal Margin = *((const uniform FReal *uniform)(Ptr + OffsetOfMargin));
+					PhiWithNormalCapsule(TypedPtr, PhiWithNormalInput, NormalReal, PhiReal, Margin);
+					Phi = (varying float)PhiReal;
+					Normal = ConvertVector3NativeTo3f(NormalReal);
+					break;
+				}
+			case Union:
+				{
+					PhiWithNormalUnion(TypedPtr, PhiWithNormalInput, NormalReal, PhiReal, SizeofFImplicitObject, OffsetofGeometryType, OffsetOfMargin);
+					Phi = (varying float)PhiReal;
+					Normal = ConvertVector3NativeTo3f(NormalReal);
+					break;
+				}
+			case TaperedCapsule:
+				{
+					PhiWithNormalTaperedCapsule(TypedPtr, PhiWithNormalInput, NormalReal, PhiReal);
+					Phi = (varying float)PhiReal;
+					Normal = ConvertVector3NativeTo3f(NormalReal);
+					break;
+				}
+			case Convex:
+				{
+					PhiWithNormalConvex(TypedPtr, PhiWithNormalInputFloat, Normal, Phi);
+					break;
+				}
+			case WeightedLatticeBone:
+				{
+					Phi = PerGroupThickness + 10;
+					break;
+				}
+			case WeightedLatticeLevelSetType:
+				{
+					uniform int IMMask;
+					if(IM != 0)
+					{
+						IMMask = lanemask();
+					}
+					GetPhiWithNormalAndVelocityBoneCollisionParticleRange(CollisionParticlesRange, (const uniform float *uniform)&PhiWithNormalInputFloat, (uniform float *uniform)&Normal, (uniform float *uniform)&Phi, (uniform int *uniform)&VelocityBone, i, PerGroupThickness, programCount, IMMask);
+					break;
+				}
+			default:
+				{
+					uniform int IMMask;
+					if(IM != 0)
+					{
+						IMMask = lanemask();
+					}
+
+					GetPhiWithNormalCollisionParticleRange(CollisionParticlesRange, (const uniform float *uniform)&PhiWithNormalInputFloat, (uniform float *uniform)&Normal, (uniform float *uniform)&Phi, i, programCount, IMMask);
+					break;
+				}
+			}
+
+			const float Penetration = PerGroupThickness - Phi; // This is related to the Normal impulse
+			if(Penetration > 0.0f)
+			{
+				// TransformVector. Scale is 1 so that part can be removed.
+				const FVector3f NormalWorld = VectorQuaternionRotateVector(CollisionR[i], Normal);
+				P = P + (Penetration * NormalWorld);
+				if(Type == WeightedLatticeLevelSetType)
+				{
+					#pragma ignore warning(perf)
+					const FVector3f CollisionXVal = CollisionX[VelocityBone];
+					const FVector3f VectorToPoint = P - CollisionXVal;
+					const FVector3f X = VectorLoad(&ParticlesX[extract(Index, 0)]);
+					
+					#pragma ignore warning(perf)
+					const FVector3f CollisionVVal = CollisionV[VelocityBone];
+					#pragma ignore warning(perf)
+					const FVector3f CollisionWVal = CollisionW[VelocityBone];
+					const FVector3f RelativeDisplacement = (P - X) - (CollisionVVal + VectorCross(CollisionWVal, VectorToPoint)) * Dt; // This corresponds to the tangential velocity multiplied by dt (friction will drive this to zero if it is high enough)
+					const FVector3f RelativeDisplacementTangent = RelativeDisplacement - VectorDot(RelativeDisplacement, NormalWorld) * NormalWorld; // Project displacement into the tangential plane
+					const float RelativeDisplacementTangentLength = VectorSize(RelativeDisplacementTangent);
+
+					const float PositionCorrection = min(Penetration * PerGroupFriction, RelativeDisplacementTangentLength);
+					const float CorrectionRatio = select(RelativeDisplacementTangentLength < FLOAT_SMALL_NUMBER, 0.0f, PositionCorrection / RelativeDisplacementTangentLength);
+					P = P - (CorrectionRatio * RelativeDisplacementTangent);
+				
+				}
+				else
+				{
+					const FVector3f VectorToPoint = P - CollisionX[i];
+					const FVector3f X = VectorLoad(&ParticlesX[extract(Index, 0)]);
+
+					const FVector3f RelativeDisplacement = (P - X) - (CollisionV[i] + VectorCross(CollisionW[i], VectorToPoint)) * Dt; // This corresponds to the tangential velocity multiplied by dt (friction will drive this to zero if it is high enough)
+					const FVector3f RelativeDisplacementTangent = RelativeDisplacement - VectorDot(RelativeDisplacement, NormalWorld) * NormalWorld; // Project displacement into the tangential plane
+					const float RelativeDisplacementTangentLength = VectorSize(RelativeDisplacementTangent);
+
+					const float PositionCorrection = min(Penetration * PerGroupFriction, RelativeDisplacementTangentLength);
+					const float CorrectionRatio = select(RelativeDisplacementTangentLength < FLOAT_SMALL_NUMBER, 0.0f, PositionCorrection / RelativeDisplacementTangentLength);
+					P = P - (CorrectionRatio * RelativeDisplacementTangent);
+				}
+			}
+		}
+
+		if(IM != 0)
+		{
+			VectorStore(&ParticlesPandInvM[extract(Index, 0)], SetVector4(P, IM));
+		}
+	}
+}
+
 export void ApplyPerParticleCollisionNoFriction(uniform FVector4f ParticlesPandInvM[],
 												const uniform FVector3f ParticlesX[],
 												const uniform FVector3f CollisionV[],
@@ -917,3 +1081,121 @@ export void ApplyPerParticleCollisionNoFriction(uniform FVector4f ParticlesPandI
 		}
 	}
 }
+
+export void ApplyPerParticleCollisionNoFrictionNoGroupCheck(uniform FVector4f ParticlesPandInvM[],
+												const uniform FVector3f ParticlesX[],
+												const uniform FVector3f CollisionV[],
+												const uniform FVector3f CollisionX[],
+												const uniform FVector3f CollisionW[],
+												const uniform FVector4f CollisionR[],
+												const uniform float PerGroupThickness,
+												const uniform uint8 *uniform CollisionParticlesRange,
+												const uniform uint8 *uniform Geometry,
+												const uniform int SizeofFImplicitObject,
+												const uniform int OffsetofGeometryType,
+												const uniform int OffsetOfMargin,
+												const uniform float Dt,
+												const uniform int InnerCollisionCount,
+												const uniform int OuterParticleBatchBegin,
+												const uniform int OuterParticleBatchEnd)
+{
+
+	foreach(Index = OuterParticleBatchBegin ... OuterParticleBatchEnd)
+	{
+		const FVector4f PandInvM = VectorLoad(&ParticlesPandInvM[extract(Index, 0)]);
+		varying FVector3f P;
+		varying float IM;
+		UnzipPandInvM(PandInvM, P, IM);
+
+		if(all(IM == 0))
+		{
+			continue;
+		}
+
+		for(uniform int i = 0; i < InnerCollisionCount; i++) // SequentialFor
+		{
+			// InverseTransformPosition. Can skip Scale calculation because it is 1 for xyz.
+			const FVector3f TranslatedVec = P - CollisionX[i];
+			const FVector4f InverseR = QuatInverse(CollisionR[i]);
+			const FVector3f PhiWithNormalInputFloat = VectorQuaternionRotateVector(InverseR, TranslatedVec);
+			const FVector PhiWithNormalInput = ConvertVector3fTo3Native(PhiWithNormalInputFloat);
+			FVector NormalReal;
+			FReal PhiReal;
+			FVector3f Normal;
+			float Phi;
+
+			const uniform SerializedPtr *uniform GeometryPtr = (const uniform SerializedPtr *uniform)Geometry;
+			const uniform uint8 *uniform Ptr = GeometryPtr[i].Ptr;
+			const uniform uint8 *uniform TypedPtr = Ptr + SizeofFImplicitObject;
+			const uniform uint8 Type = Ptr[OffsetofGeometryType] & 0x3F;
+
+			switch(Type)
+			{
+			case Sphere:
+				{
+					const uniform FReal Margin = *((const uniform FReal *uniform)(Ptr + OffsetOfMargin));
+					PhiWithNormalSphere(TypedPtr, PhiWithNormalInput, NormalReal, PhiReal, Margin);
+					Phi = (varying float)PhiReal;
+					Normal = ConvertVector3NativeTo3f(NormalReal);
+					break;
+				}
+			case Capsule:
+				{
+					const uniform FReal Margin = *((const uniform FReal *uniform)(Ptr + OffsetOfMargin));
+					PhiWithNormalCapsule(TypedPtr, PhiWithNormalInput, NormalReal, PhiReal, Margin);
+					Phi = (varying float)PhiReal;
+					Normal = ConvertVector3NativeTo3f(NormalReal);
+					break;
+				}
+			case Union:
+				{
+					PhiWithNormalUnion(TypedPtr, PhiWithNormalInput, NormalReal, PhiReal, SizeofFImplicitObject, OffsetofGeometryType, OffsetOfMargin);
+					Phi = (varying float)PhiReal;
+					Normal = ConvertVector3NativeTo3f(NormalReal);
+					break;
+				}
+			case TaperedCapsule:
+				{
+					PhiWithNormalTaperedCapsule(TypedPtr, PhiWithNormalInput, NormalReal, PhiReal);
+					Phi = (varying float)PhiReal;
+					Normal = ConvertVector3NativeTo3f(NormalReal);
+					break;
+				}
+			case Convex:
+				{
+					PhiWithNormalConvex(TypedPtr, PhiWithNormalInputFloat, Normal, Phi);
+					break;
+				}
+			case WeightedLatticeBone:
+				{
+					Phi = PerGroupThickness + 10;
+					break;
+				}
+			default:
+				{
+					uniform int IMMask;
+					if(IM != 0)
+					{
+						IMMask = lanemask();
+					}
+
+					GetPhiWithNormalCollisionParticleRange(CollisionParticlesRange, (const uniform float *uniform)&PhiWithNormalInputFloat, (uniform float *uniform)&Normal, (uniform float *uniform)&Phi, i, programCount, IMMask);
+					break;
+				}
+			}
+
+			const float Penetration = PerGroupThickness - Phi; // This is related to the Normal impulse
+			if(Penetration > 0.0f)
+			{
+				// TransformVector. Scale is 1 so that part can be removed.
+				const FVector3f NormalWorld = VectorQuaternionRotateVector(CollisionR[i], Normal);
+				P = P + (Penetration * NormalWorld);
+			}
+		}
+
+		if(IM != 0)
+		{
+			VectorStore(&ParticlesPandInvM[extract(Index, 0)], SetVector4(P, IM));
+		}
+	}
+}
\ No newline at end of file
diff --git a/Engine/Source/Runtime/Experimental/Chaos/Private/Chaos/SoftsEvolution.cpp b/Engine/Source/Runtime/Experimental/Chaos/Private/Chaos/SoftsEvolution.cpp
new file mode 100644
--- /dev/null
+++ b/Engine/Source/Runtime/Experimental/Chaos/Private/Chaos/SoftsEvolution.cpp
@@ -0,0 +1,479 @@
+// Copyright Epic Games, Inc. All Rights Reserved.
+#include "Chaos/SoftsEvolution.h"
+#include "Chaos/PerParticleDampVelocity.h"
+
+namespace Chaos::Softs {
+
+namespace Private
+{
+namespace EvolutionSolverDefault
+{
+constexpr FSolverReal SolverFrequency = (FSolverReal)60.;
+constexpr int32 NumIterations = 1;
+constexpr int32 MinNumIterations = 1;
+constexpr int32 MaxNumIterations = 10;
+constexpr bool bDoQuasistatics = false;
+}
+namespace EvolutionSoftBodyDefault
+{
+constexpr FSolverReal GlobalDamping = (FSolverReal)0.01;
+constexpr FSolverReal LocalDamping = (FSolverReal)0.;
+constexpr bool bUsePerParticleDamping = false;
+}
+
+static void EulerStepVelocity(FSolverParticlesRange& Particles, const FSolverReal Dt)
+{
+	// TODO: ISPC and/or ParallelFor
+	FSolverVec3* const Velocity = Particles.GetV().GetData();
+	const FSolverVec3* const Acceleration = Particles.GetAcceleration().GetData();
+	const FSolverReal* const InvM = Particles.GetInvM().GetData();
+	for (int32 Index = 0; Index < Particles.GetRangeSize(); ++Index)
+	{
+		if (InvM[Index] != (FSolverReal)0.f)
+		{
+			Velocity[Index] += Acceleration[Index] * Dt;
+		}
+	}
+}
+
+static void DampLocalVelocity(FSolverParticlesRange& Particles, const FPerParticleDampVelocity& DampVelocityRule)
+{
+	// TODO: ISPC and/or ParallelFor
+	FSolverVec3* const Velocity = Particles.GetV().GetData();
+	const FSolverVec3* const X = Particles.XArray().GetData();
+	const FSolverReal* const InvM = Particles.GetInvM().GetData();
+	for (int32 Index = 0; Index < Particles.GetRangeSize(); ++Index)
+	{
+		if (InvM[Index] != (FSolverReal)0.f)
+		{
+			DampVelocityRule.Apply(X[Index], Velocity[Index]);
+		}
+	}
+}
+
+static void CalculateGlobalDamping(const FSolverReal DampingScale, const FSolverReal SolverFrequency, const FSolverReal Dt, FSolverReal& OutDampingPowDt, FSolverReal& OutDampingIntegrated)
+{
+	const FSolverReal Damping = FMath::Clamp(DampingScale, (FSolverReal)0., (FSolverReal)1.);
+	if (Damping > (FSolverReal)1. - (FSolverReal)UE_KINDA_SMALL_NUMBER)
+	{
+		OutDampingIntegrated = OutDampingPowDt = (FSolverReal)0.;
+	}
+	else if (Damping > (FSolverReal)UE_SMALL_NUMBER)
+	{
+		const FSolverReal LogValueByFrequency = FMath::Loge((FSolverReal)1. - Damping) * SolverFrequency;
+
+		OutDampingPowDt = FMath::Exp(LogValueByFrequency * Dt);  // DampingPowDt = FMath::Pow(OneMinusDamping, Dt * SolverFrequency);
+		OutDampingIntegrated = (OutDampingPowDt - (FSolverReal)1.) / LogValueByFrequency;
+	}
+	else
+	{
+		OutDampingPowDt = (FSolverReal)1.;
+		OutDampingIntegrated = Dt;
+	}
+}
+
+static void EulerStepPositionWithGlobalDamping(FSolverParticlesRange& Particles, const FSolverReal DampingScale, const FSolverReal SolverFrequency, const FSolverReal Dt)
+{
+	FSolverReal DampingPowDt;
+	FSolverReal DampingIntegrated;
+	CalculateGlobalDamping(DampingScale, SolverFrequency, Dt, DampingPowDt, DampingIntegrated);
+	const FSolverVec3* const X = Particles.XArray().GetData();
+	FSolverVec3* const Velocity = Particles.GetV().GetData();
+	FPAndInvM* const PAndInvM = Particles.GetPAndInvM().GetData();
+	for (int32 Index = 0; Index < Particles.GetRangeSize(); ++Index)
+	{
+		if (PAndInvM[Index].InvM != (FSolverReal)0.)
+		{
+			PAndInvM[Index].P = X[Index] + Velocity[Index] * DampingIntegrated;
+			Velocity[Index] *= DampingPowDt;
+		}
+	}
+}
+
+static void EulerStepPositionWithGlobalDampingArray(FSolverParticlesRange& Particles, const TArrayCollectionArray<FSolverReal>& ParticleDampings, const FSolverReal SolverFrequency, const FSolverReal Dt)
+{
+	const FSolverVec3* const X = Particles.XArray().GetData();
+	FSolverVec3* const Velocity = Particles.GetV().GetData();
+	FPAndInvM* const PAndInvM = Particles.GetPAndInvM().GetData();
+	const FSolverReal* const DampingScale = Particles.GetConstArrayView(ParticleDampings).GetData();
+	for (int32 Index = 0; Index < Particles.GetRangeSize(); ++Index)
+	{
+		if (PAndInvM[Index].InvM != (FSolverReal)0.)
+		{
+			FSolverReal DampingPowDt;
+			FSolverReal DampingIntegrated;
+			CalculateGlobalDamping(DampingScale[Index], SolverFrequency, Dt, DampingPowDt, DampingIntegrated);
+			PAndInvM[Index].P = X[Index] + Velocity[Index] * DampingIntegrated;
+			Velocity[Index] *= DampingPowDt;
+		}
+	}
+}
+
+static void QuasistaticUpdate(FSolverParticlesRange& Particles)
+{
+	const FSolverVec3* const X = Particles.XArray().GetData();
+	FPAndInvM* const PAndInvM = Particles.GetPAndInvM().GetData();
+	for (int32 Index = 0; Index < Particles.GetRangeSize(); ++Index)
+	{
+		if (PAndInvM[Index].InvM != (FSolverReal)0.)
+		{
+			PAndInvM[Index].P = X[Index];
+		}
+	}
+}
+
+static void PostStepUpdate(FSolverParticlesRange& Particles, const FSolverReal Dt)
+{
+	const FPAndInvM* const PAndInvM = Particles.GetPAndInvM().GetData();
+	FSolverVec3* const X = Particles.XArray().GetData();
+	FSolverVec3* const V = Particles.GetV().GetData();
+	for (int32 Index = 0; Index < Particles.GetRangeSize(); ++Index)
+	{
+		X[Index] = PAndInvM[Index].P;
+		V[Index] = (PAndInvM[Index].P - X[Index]) / Dt;
+	}
+}
+}
+
+FEvolution::FEvolution(const FCollectionPropertyConstFacade& Properties)
+	: Time((FSolverReal)0.f)
+{
+	SetSolverProperties(Properties);
+
+	Particles.AddArray(&ParticleDampings);
+}
+
+void FEvolution::Reset()
+{
+	CollisionParticles.Resize(0);
+	Particles.Resize(0);
+
+	SoftBodies.Reset();
+	CollisionRanges.Reset();
+
+	Groups.Reset();
+
+	ConstParticleRules.Reset();
+	ParticleRules.Reset();
+	CollisionRules.Reset();
+
+	KinematicUpdate.Reset();
+	CollisionKinematicUpdate.Reset();
+}
+
+int32 FEvolution::NumActiveParticles() const
+{
+	int32 Result = 0;
+	for (const uint32 GroupId : GetActiveGroups())
+	{
+		for (const int32 SoftBodyId : GetGroupActiveSoftBodies(GroupId))
+		{
+			Result += GetSoftBodyParticleNum(SoftBodyId);
+		}
+	}
+	return Result;
+}
+
+int32 FEvolution::AddSoftBody(uint32 GroupId, int32 NumParticles, bool bEnable)
+{
+	if (GroupId >= Groups.Size())
+	{
+		Groups.AddGroupsToSize(GroupId + 1);
+	}
+	const int32 SoftBodyId = SoftBodies.AddSoftBody();
+	Groups.SoftBodies[GroupId].Emplace(SoftBodyId);
+
+	SoftBodies.ParticleRanges[SoftBodyId] = FSolverParticlesRange::AddParticleRange(Particles, NumParticles);
+
+	// Set properties
+	SoftBodies.Active[SoftBodyId] = false;
+	SoftBodies.GroupId[SoftBodyId] = GroupId;
+	SoftBodies.GlobalDampings[SoftBodyId] = Private::EvolutionSoftBodyDefault::GlobalDamping;
+	SoftBodies.LocalDampings[SoftBodyId] = Private::EvolutionSoftBodyDefault::LocalDamping;
+	SoftBodies.UsePerParticleDamping[SoftBodyId] = Private::EvolutionSoftBodyDefault::bUsePerParticleDamping;
+	ActivateSoftBody(SoftBodyId, bEnable);
+	return SoftBodyId;
+}
+
+void FEvolution::SetSoftBodyProperties(int32 SoftBodyId, const FCollectionPropertyConstFacade& PropertyCollection)
+{
+	SoftBodies.GlobalDampings[SoftBodyId] = PropertyCollection.GetValue<float>(TEXT("DampingCoefficient"), Private::EvolutionSoftBodyDefault::GlobalDamping);
+	SoftBodies.LocalDampings[SoftBodyId] = PropertyCollection.GetValue<float>(TEXT("LocalDampingCoefficient"), Private::EvolutionSoftBodyDefault::LocalDamping);
+	SoftBodies.UsePerParticleDamping[SoftBodyId] = PropertyCollection.GetValue<bool>(TEXT("UsePerParticleDamping"), Private::EvolutionSoftBodyDefault::bUsePerParticleDamping);
+}
+
+void FEvolution::ActivateSoftBody(int32 SoftBodyId, bool bActivate)
+{
+	SoftBodies.Active[SoftBodyId] = bActivate;
+	if (bActivate)
+	{
+		Groups.ActiveSoftBodies[SoftBodies.GroupId[SoftBodyId]].Add(SoftBodyId);
+		ActiveGroups.Add(SoftBodies.GroupId[SoftBodyId]);
+	}
+	else
+	{
+		Groups.ActiveSoftBodies[SoftBodies.GroupId[SoftBodyId]].Remove(SoftBodyId);
+		if (Groups.ActiveSoftBodies[SoftBodies.GroupId[SoftBodyId]].IsEmpty())
+		{
+			ActiveGroups.Remove(SoftBodies.GroupId[SoftBodyId]);
+		}
+	}
+}
+
+int32 FEvolution::AddCollisionParticleRange(uint32 GroupId, int32 NumParticles, bool bEnable)
+{
+	if (GroupId >= Groups.Size())
+	{
+		Groups.AddGroupsToSize(GroupId + 1);
+	}
+
+	int32 CollisionRangeId = INDEX_NONE;
+	TArray<int32>* const FreeList = CollisionRangeFreeList.Find(NumParticles);
+	if (FreeList && FreeList->Num())
+	{
+		CollisionRangeId = FreeList->Pop();
+		check(CollisionRanges.Status[CollisionRangeId] == FCollisionBodyRanges::EStatus::Free);
+	}
+	else
+	{
+		CollisionRangeId = CollisionRanges.AddRange();
+		CollisionRanges.ParticleRanges[CollisionRangeId] = FSolverCollisionParticlesRange::AddParticleRange(CollisionParticles, NumParticles);
+	}
+
+	// Set properties
+	CollisionRanges.Status[CollisionRangeId] = FCollisionBodyRanges::EStatus::Inactive;
+	CollisionRanges.GroupId[CollisionRangeId] = GroupId;
+	ActivateCollisionParticleRange(CollisionRangeId, bEnable);
+	return CollisionRangeId;
+}
+
+void FEvolution::RemoveCollisionParticleRange(int32 CollisionRangeId)
+{
+	check(CollisionRanges.Status[CollisionRangeId] == FCollisionBodyRanges::EStatus::Active || CollisionRanges.Status[CollisionRangeId] == FCollisionBodyRanges::EStatus::Inactive);
+
+	// Remove from active group if it was active.
+	Groups.ActiveCollisionParticleRanges[CollisionRanges.GroupId[CollisionRangeId]].Remove(CollisionRangeId);
+
+	// Mark free
+	CollisionRanges.Status[CollisionRangeId] = FCollisionBodyRanges::EStatus::Free;
+
+	// Add to free list
+	const int32 NumParticles = CollisionRanges.ParticleRanges[CollisionRangeId].GetRangeSize();
+	TArray<int32>& FreeList = CollisionRangeFreeList.FindOrAdd(NumParticles);
+	FreeList.Add(CollisionRangeId);
+}
+
+void FEvolution::ActivateCollisionParticleRange(int32 CollisionRangeId, bool bEnable)
+{
+	check(CollisionRanges.Status[CollisionRangeId] != FCollisionBodyRanges::EStatus::Free);
+
+	CollisionRanges.Status[CollisionRangeId] = bEnable ? FCollisionBodyRanges::EStatus::Active : FCollisionBodyRanges::EStatus::Inactive;
+	if (bEnable)
+	{
+		Groups.ActiveCollisionParticleRanges[CollisionRanges.GroupId[CollisionRangeId]].Add(CollisionRangeId);
+	}
+	else
+	{
+		Groups.ActiveCollisionParticleRanges[CollisionRanges.GroupId[CollisionRangeId]].Remove(CollisionRangeId);
+	}
+
+}
+
+TArray<FSolverCollisionParticlesRange> FEvolution::GetActiveCollisionParticles(uint32 GroupId) const
+{
+	TArray<FSolverCollisionParticlesRange> Ranges;
+	Ranges.Reserve(Groups.ActiveCollisionParticleRanges[GroupId].Num());
+	for (const int32 RangeId : Groups.ActiveCollisionParticleRanges[GroupId])
+	{
+		check(CollisionRanges.Status[RangeId] == FCollisionBodyRanges::EStatus::Active);
+		check(CollisionRanges.GroupId[RangeId] == GroupId);
+		Ranges.Emplace(CollisionRanges.ParticleRanges[RangeId]);
+	}
+	return Ranges;
+}
+
+void FEvolution::SetSolverProperties(const FCollectionPropertyConstFacade& Properties)
+{
+	MaxNumIterations = FMath::Max(
+		Properties.GetValue<int32>(TEXT("MaxNumIterations"), Private::EvolutionSolverDefault::MaxNumIterations),
+		Private::EvolutionSolverDefault::MinNumIterations);
+
+	NumIterations = FMath::Clamp(
+		Properties.GetValue<int32>(TEXT("NumIterations"), Private::EvolutionSolverDefault::NumIterations),
+		Private::EvolutionSolverDefault::MinNumIterations,
+		MaxNumIterations);
+
+	bDoQuasistatics = Properties.GetValue<bool>(TEXT("DoQuasistatics"), Private::EvolutionSolverDefault::bDoQuasistatics);
+	
+	SolverFrequency = (FSolverReal)Properties.GetValue<FRealSingle>(TEXT("SolverFrequency"), Private::EvolutionSolverDefault::SolverFrequency);
+}
+
+void FEvolution::AdvanceOneTimeStep(const FSolverReal Dt, const FSolverReal TimeDependentIterationMultiplier)
+{
+	TRACE_CPUPROFILER_EVENT_SCOPE(FSoftsEvolution_AdvanceOneTimeStep);
+
+	// Advance time
+	Time += Dt;
+
+	const int32 TimeDependentNumIterations = bDisableTimeDependentNumIterations ?
+		NumIterations :
+		FMath::Clamp(FMath::RoundToInt32(SolverFrequency * Dt * TimeDependentIterationMultiplier * (Softs::FSolverReal)NumIterations), Private::EvolutionSolverDefault::MinNumIterations, MaxNumIterations);
+
+	const TArray<uint32> ActiveGroupsArray = ActiveGroups.Array();
+
+	PhysicsParallelFor(ActiveGroupsArray.Num(),
+		[this, &ActiveGroupsArray, Dt, TimeDependentNumIterations](int32 ActiveGroupIndex)
+	{
+		AdvanceOneTimeStepInternal(Dt, TimeDependentNumIterations, ActiveGroupsArray[ActiveGroupIndex]);
+	});
+}
+
+void FEvolution::AdvanceOneTimeStepInternal(const FSolverReal Dt, const int32 TimeDependentNumIterations, uint32 GroupId)
+{
+	TRACE_CPUPROFILER_EVENT_SCOPE(FSoftsEvolution_AdvanceOneTimeStepInternal);
+
+	TArray<FSolverCollisionParticlesRange> ActiveCollisionRanges = GetActiveCollisionParticles(GroupId);
+	// Right now, typically groups only have one active softbody in them, so don't parallelize softbodies
+	for (const int32 SoftBodyId : Groups.ActiveSoftBodies[GroupId])
+	{
+		check(SoftBodies.Active[SoftBodyId]);
+
+		// PreSubstepParallelInits + LocalDamping
+		const bool bDoLocalDamping = SoftBodies.LocalDampings[SoftBodyId] > (FSolverReal)0.;
+		FPerParticleDampVelocity DampVelocityRule(SoftBodies.LocalDampings[SoftBodyId]);
+		{
+			TRACE_CPUPROFILER_EVENT_SCOPE(FSoftsEvolution_AdvanceOneTimeStepInternal_PreSubstepParallelInits);
+
+			TConstArrayView<PreSubstepParallelInitFunc> PreSubstepParallelInits = SoftBodies.PreSubstepParallelInits[SoftBodyId].GetConstArrayView();
+			const int32 NumPreSubstepInits = PreSubstepParallelInits.Num() + (bDoLocalDamping ? 1 : 0);
+			PhysicsParallelFor(NumPreSubstepInits,
+				[this, &PreSubstepParallelInits, &DampVelocityRule, Dt, SoftBodyId](int32 FuncIdx)
+			{
+				if (FuncIdx == PreSubstepParallelInits.Num())
+				{
+					DampVelocityRule.UpdatePositionBasedState(SoftBodies.ParticleRanges[SoftBodyId]);
+				}
+				else
+				{
+					PreSubstepParallelInits[FuncIdx](SoftBodies.ParticleRanges[SoftBodyId], Dt);
+				}
+			}
+			);
+		}
+
+		// TODO: hypothetically, kinematic, dynamic, and collision kinematic updates can run concurrently
+		{
+			// Kinematic update
+			check(KinematicUpdate);
+			TRACE_CPUPROFILER_EVENT_SCOPE(FSoftsEvolution_AdvanceOneTimeStepInternal_KinematicUpdates);
+			KinematicUpdate(SoftBodies.ParticleRanges[SoftBodyId], Dt, Time);
+		}
+		{
+			// Dynamic initial guess (Euler step external forces and apply damping)
+			TRACE_CPUPROFILER_EVENT_SCOPE(FSoftsEvolution_AdvanceOneTimeStepInternal_InitialGuess);
+			if (bDoQuasistatics)
+			{
+				Private::QuasistaticUpdate(SoftBodies.ParticleRanges[SoftBodyId]);
+			}
+			else
+			{
+				// ExternalForceRules
+				TConstArrayView<ExternalForceRuleFunc> ExternalForceRules = SoftBodies.ExternalForceRules[SoftBodyId].GetConstArrayView();
+				for (const ExternalForceRuleFunc& ExternalForceRule : ExternalForceRules)
+				{
+					ExternalForceRule(SoftBodies.ParticleRanges[SoftBodyId], Dt);
+				}
+
+				// Euler Step Velocity
+				Private::EulerStepVelocity(SoftBodies.ParticleRanges[SoftBodyId], Dt);
+
+				if (bDoLocalDamping)
+				{
+					Private::DampLocalVelocity(SoftBodies.ParticleRanges[SoftBodyId], DampVelocityRule);
+				}
+
+				if (SoftBodies.UsePerParticleDamping[SoftBodyId])
+				{
+					Private::EulerStepPositionWithGlobalDampingArray(SoftBodies.ParticleRanges[SoftBodyId], ParticleDampings, SolverFrequency, Dt);
+				}
+				else
+				{
+					Private::EulerStepPositionWithGlobalDamping(SoftBodies.ParticleRanges[SoftBodyId], SoftBodies.GlobalDampings[SoftBodyId], SolverFrequency, Dt);
+				}
+			}
+		}
+
+		{
+			// Collision kinematic update
+			TRACE_CPUPROFILER_EVENT_SCOPE(FSoftsEvolution_AdvanceOneTimeStepInternal_CollisionKinematicUpdate);
+			check(CollisionKinematicUpdate);
+			for (FSolverCollisionParticlesRange& ActiveRange : ActiveCollisionRanges)
+			{
+				CollisionKinematicUpdate(ActiveRange, Dt, Time);
+			}
+		}
+
+		{
+			// Parallel Inits
+			TRACE_CPUPROFILER_EVENT_SCOPE(FSoftsEvolution_AdvanceOneTimeStepInternal_ConstraintParallelInits);
+			TConstArrayView<ConstraintParallelInitFunc> ConstraintParallelInits = SoftBodies.ConstraintParallelInits[SoftBodyId].GetConstArrayView();
+			PhysicsParallelFor(ConstraintParallelInits.Num(),
+				[this, &ConstraintParallelInits, Dt, SoftBodyId](int32 Index)
+			{
+				ConstraintParallelInits[Index](SoftBodies.ParticleRanges[SoftBodyId], Dt);
+			});
+		}
+
+		{
+			// PreSubstep rules
+			TRACE_CPUPROFILER_EVENT_SCOPE(FSoftsEvolution_AdvanceOneTimeStepInternal_PreSubstepConstraintRules);
+			TConstArrayView<ConstraintRuleFunc> PreSubstepConstraintRules = SoftBodies.PreSubstepConstraintRules[SoftBodyId].GetConstArrayView();
+			for (const ConstraintRuleFunc& ConstraintRule : PreSubstepConstraintRules)
+			{
+				ConstraintRule(SoftBodies.ParticleRanges[SoftBodyId], Dt);
+			}
+		}
+
+		{
+			// Iteration loop
+			TRACE_CPUPROFILER_EVENT_SCOPE(FSoftsEvolution_AdvanceOneTimeStepInternal_IterationLoop);
+			for (int32 i = 0; i < NumIterations; ++i)
+			{
+				TConstArrayView<ConstraintRuleFunc> PerIterationConstraintRules = SoftBodies.PerIterationConstraintRules[SoftBodyId].GetConstArrayView();
+				for (const ConstraintRuleFunc& ConstraintRule : PerIterationConstraintRules)
+				{
+					ConstraintRule(SoftBodies.ParticleRanges[SoftBodyId], Dt);
+				}
+
+				TConstArrayView<CollisionConstraintRuleFunc> PerIterationCollisionConstraintRules = SoftBodies.PerIterationCollisionConstraintRules[SoftBodyId].GetConstArrayView();
+				for (const CollisionConstraintRuleFunc& ConstraintRule : PerIterationCollisionConstraintRules)
+				{
+					ConstraintRule(SoftBodies.ParticleRanges[SoftBodyId], Dt, ActiveCollisionRanges);
+				}
+
+				TConstArrayView<ConstraintRuleFunc> PerIterationPostCollisionsConstraintRules = SoftBodies.PerIterationPostCollisionsConstraintRules[SoftBodyId].GetConstArrayView();
+				for (const ConstraintRuleFunc& ConstraintRule : PerIterationPostCollisionsConstraintRules)
+				{
+					ConstraintRule(SoftBodies.ParticleRanges[SoftBodyId], Dt);
+				}
+			}
+
+			{
+				// Post Substep Rules
+				TRACE_CPUPROFILER_EVENT_SCOPE(FSoftsEvolution_AdvanceOneTimeStepInternal_PostSubstep);
+				TConstArrayView<ConstraintRuleFunc> PostSubstepConstraintRules = SoftBodies.PostSubstepConstraintRules[SoftBodyId].GetConstArrayView();
+				for (const ConstraintRuleFunc& ConstraintRule : PostSubstepConstraintRules)
+				{
+					ConstraintRule(SoftBodies.ParticleRanges[SoftBodyId], Dt);
+				}
+			}
+
+			{
+				// Post Step Update
+				TRACE_CPUPROFILER_EVENT_SCOPE(FSoftsEvolution_AdvanceOneTimeStepInternal_PostStepUpdate);
+				Private::PostStepUpdate(SoftBodies.ParticleRanges[SoftBodyId], Dt);
+			}
+		}
+	}
+}
+}
\ No newline at end of file
diff --git a/Engine/Source/Runtime/Experimental/Chaos/Private/Chaos/VelocityField.cpp b/Engine/Source/Runtime/Experimental/Chaos/Private/Chaos/VelocityField.cpp
--- a/Engine/Source/Runtime/Experimental/Chaos/Private/Chaos/VelocityField.cpp
+++ b/Engine/Source/Runtime/Experimental/Chaos/Private/Chaos/VelocityField.cpp
@@ -81,6 +81,18 @@ void FVelocityAndPressureField::SetProperties(
 	QuarterRho = bEnableAerodynamics ? Rho * OneQuarter : (FSolverReal)0.;
 }
 
+void FVelocityAndPressureField::SetPropertiesAndWind(
+	const FCollectionPropertyConstFacade& PropertyCollection,
+	const TMap<FString, TConstArrayView<FRealSingle>>& WeightMaps,
+	FSolverReal WorldScale,
+	bool bEnableAerodynamics,
+	const FSolverVec3& SolverWind)
+{
+	SetProperties(PropertyCollection, WeightMaps, WorldScale, bEnableAerodynamics);
+	const FSolverVec3 WindVelocity = WindVelocityIndex != INDEX_NONE ? WorldScale * FSolverVec3(GetWindVelocity(PropertyCollection)) : FSolverVec3(0.f);
+	SetVelocity(WindVelocity + SolverWind);
+}
+
 void FVelocityAndPressureField::SetProperties(
 	const FSolverVec2& Drag,
 	const FSolverVec2& Lift,
@@ -203,7 +215,8 @@ void FVelocityAndPressureField::SetMultipliers(
 	}
 }
 
-void FVelocityAndPressureField::UpdateForces(const FSolverParticles& InParticles, const FSolverReal /*Dt*/)
+template<typename SolverParticlesOrRange>
+void FVelocityAndPressureField::UpdateForces(const SolverParticlesOrRange& InParticles, const FSolverReal /*Dt*/)
 {
 	const FSolverReal MaxVelocitySquared = (Private::VelocityFieldMaxVelocity > 0.f) ? FMath::Square((FSolverReal)Private::VelocityFieldMaxVelocity) : TNumericLimits<FSolverReal>::Max();
 
@@ -334,5 +347,7 @@ void FVelocityAndPressureField::UpdateForces(const FSolverParticles& InParticles
 		}
 	}
 }
+template CHAOS_API void FVelocityAndPressureField::UpdateForces(const FSolverParticles& InParticles, const FSolverReal /*Dt*/);
+template CHAOS_API void FVelocityAndPressureField::UpdateForces(const FSolverParticlesRange& InParticles, const FSolverReal /*Dt*/);
 
 }  // End namespace Chaos::Softs
diff --git a/Engine/Source/Runtime/Experimental/Chaos/Private/Chaos/XPBDAnisotropicBendingConstraints.cpp b/Engine/Source/Runtime/Experimental/Chaos/Private/Chaos/XPBDAnisotropicBendingConstraints.cpp
--- a/Engine/Source/Runtime/Experimental/Chaos/Private/Chaos/XPBDAnisotropicBendingConstraints.cpp
+++ b/Engine/Source/Runtime/Experimental/Chaos/Private/Chaos/XPBDAnisotropicBendingConstraints.cpp
@@ -22,6 +22,84 @@ namespace Chaos::Softs {
 // @todo(chaos): the parallel threshold (or decision to run parallel) should probably be owned by the solver and passed to the constraint container
 extern int32 Chaos_XPBDBending_ParallelConstraintCount;
 
+FXPBDAnisotropicBendingConstraints::FXPBDAnisotropicBendingConstraints(const FSolverParticlesRange& InParticles,
+	const FTriangleMesh& TriangleMesh,
+	const TArray<TVec3<FVec2f>>& FaceVertexPatternPositions,
+	const TMap<FString, TConstArrayView<FRealSingle>>& WeightMaps,
+	const FCollectionPropertyConstFacade& PropertyCollection,
+	bool bTrimKinematicConstraints)
+	: Base(
+		InParticles,
+		TriangleMesh.GetUniqueAdjacentElements(),
+		WeightMaps.FindRef(GetXPBDAnisoBendingStiffnessWarpString(PropertyCollection, XPBDAnisoBendingStiffnessWarpName.ToString())),
+		WeightMaps.FindRef(GetXPBDAnisoBucklingStiffnessWarpString(PropertyCollection, XPBDAnisoBucklingStiffnessWarpName.ToString())),
+		GetRestAngleMapFromCollection(WeightMaps, PropertyCollection),
+		FSolverVec2(GetWeightedFloatXPBDAnisoBendingStiffnessWarp(PropertyCollection, MaxStiffness)),
+		(FSolverReal)GetXPBDAnisoBucklingRatio(PropertyCollection, 0.f),
+		FSolverVec2(GetWeightedFloatXPBDAnisoBucklingStiffnessWarp(PropertyCollection, MaxStiffness)),
+		GetRestAngleValueFromCollection(PropertyCollection),
+		(ERestAngleConstructionType)GetXPBDAnisoRestAngleType(PropertyCollection, (int32)ERestAngleConstructionType::Use3DRestAngles),
+		bTrimKinematicConstraints,
+		MaxStiffness)
+	, StiffnessWeft(
+		FSolverVec2(GetWeightedFloatXPBDAnisoBendingStiffnessWeft(PropertyCollection, MaxStiffness)),
+		WeightMaps.FindRef(GetXPBDAnisoBendingStiffnessWeftString(PropertyCollection, XPBDAnisoBendingStiffnessWeftName.ToString())),
+		TConstArrayView<TVec2<int32>>(ConstraintSharedEdges),
+		ParticleOffset,
+		ParticleCount,
+		FPBDStiffness::DefaultTableSize,
+		FPBDStiffness::DefaultParameterFitBase,
+		MaxStiffness)
+	, StiffnessBias(
+		FSolverVec2(GetWeightedFloatXPBDAnisoBendingStiffnessBias(PropertyCollection, MaxStiffness)),
+		WeightMaps.FindRef(GetXPBDAnisoBendingStiffnessBiasString(PropertyCollection, XPBDAnisoBendingStiffnessBiasName.ToString())),
+		TConstArrayView<TVec2<int32>>(ConstraintSharedEdges),
+		ParticleOffset,
+		ParticleCount,
+		FPBDStiffness::DefaultTableSize,
+		FPBDStiffness::DefaultParameterFitBase,
+		MaxStiffness)
+	, BucklingStiffnessWeft(
+		FSolverVec2(GetWeightedFloatXPBDAnisoBucklingStiffnessWeft(PropertyCollection, MaxStiffness)),
+		WeightMaps.FindRef(GetXPBDAnisoBucklingStiffnessWeftString(PropertyCollection, XPBDAnisoBucklingStiffnessWeftName.ToString())),
+		TConstArrayView<TVec2<int32>>(ConstraintSharedEdges),
+		ParticleOffset,
+		ParticleCount,
+		FPBDStiffness::DefaultTableSize,
+		FPBDStiffness::DefaultParameterFitBase,
+		MaxStiffness)
+	, BucklingStiffnessBias(
+		FSolverVec2(GetWeightedFloatXPBDAnisoBucklingStiffnessBias(PropertyCollection, MaxStiffness)),
+		WeightMaps.FindRef(GetXPBDAnisoBucklingStiffnessBiasString(PropertyCollection, XPBDAnisoBucklingStiffnessBiasName.ToString())),
+		TConstArrayView<TVec2<int32>>(ConstraintSharedEdges),
+		ParticleOffset,
+		ParticleCount,
+		FPBDStiffness::DefaultTableSize,
+		FPBDStiffness::DefaultParameterFitBase,
+		MaxStiffness)
+	, DampingRatio(
+		FSolverVec2(GetWeightedFloatXPBDAnisoBendingDamping(PropertyCollection, MinDamping)).ClampAxes(MinDamping, MaxDamping),
+		WeightMaps.FindRef(GetXPBDAnisoBendingDampingString(PropertyCollection, XPBDAnisoBendingDampingName.ToString())),
+		TConstArrayView<TVec2<int32>>(ConstraintSharedEdges),
+		ParticleOffset,
+		ParticleCount)
+	, WarpWeftBiasBaseMultipliers(GenerateWarpWeftBiasBaseMultipliers(FaceVertexPatternPositions, TriangleMesh))
+	, XPBDAnisoBendingStiffnessWarpIndex(PropertyCollection)
+	, XPBDAnisoBendingStiffnessWeftIndex(PropertyCollection)
+	, XPBDAnisoBendingStiffnessBiasIndex(PropertyCollection)
+	, XPBDAnisoBendingDampingIndex(PropertyCollection)
+	, XPBDAnisoBucklingRatioIndex(PropertyCollection)
+	, XPBDAnisoBucklingStiffnessWarpIndex(PropertyCollection)
+	, XPBDAnisoBucklingStiffnessWeftIndex(PropertyCollection)
+	, XPBDAnisoBucklingStiffnessBiasIndex(PropertyCollection)
+	, XPBDAnisoFlatnessRatioIndex(PropertyCollection)
+	, XPBDAnisoRestAngleIndex(PropertyCollection)
+	, XPBDAnisoRestAngleTypeIndex(PropertyCollection)
+{
+	Lambdas.Init((FSolverReal)0., Constraints.Num());
+	InitColor(InParticles);
+}
+
 FXPBDAnisotropicBendingConstraints::FXPBDAnisotropicBendingConstraints(const FSolverParticles& InParticles,
 	int32 InParticleOffset,
 	int32 InParticleCount,
@@ -252,14 +330,15 @@ TArray<FSolverVec3> FXPBDAnisotropicBendingConstraints::GenerateWarpWeftBiasBase
 	return WarpWeftBaseMultiplierResult;
 }
 
-void FXPBDAnisotropicBendingConstraints::InitColor(const FSolverParticles& InParticles)
+template<typename SolverParticlesOrRange>
+void FXPBDAnisotropicBendingConstraints::InitColor(const SolverParticlesOrRange& InParticles)
 {
 	// In dev builds we always color so we can tune the system without restarting. See Apply()
 #if UE_BUILD_SHIPPING || UE_BUILD_TEST
 	if (Constraints.Num() > Chaos_XPBDBending_ParallelConstraintCount)
 #endif
 	{
-		const TArray<TArray<int32>> ConstraintsPerColor = FGraphColoring::ComputeGraphColoring(Constraints, InParticles, ParticleOffset, ParticleOffset + ParticleCount);
+		const TArray<TArray<int32>> ConstraintsPerColor = FGraphColoring::ComputeGraphColoringParticlesOrRange(Constraints, InParticles, ParticleOffset, ParticleOffset + ParticleCount);
 
 		// Reorder constraints based on color so each array in ConstraintsPerColor contains contiguous elements.
 		TArray<TVec4<int32>> ReorderedConstraints; 
@@ -460,7 +539,8 @@ void FXPBDAnisotropicBendingConstraints::SetProperties(
 	}
 }
 
-void FXPBDAnisotropicBendingConstraints::ApplyHelper(FSolverParticles& Particles, const FSolverReal Dt, const int32 ConstraintIndex, const FSolverVec3& ExpStiffnessValues, 
+template<typename SolverParticlesOrRange>
+void FXPBDAnisotropicBendingConstraints::ApplyHelper(SolverParticlesOrRange& Particles, const FSolverReal Dt, const int32 ConstraintIndex, const FSolverVec3& ExpStiffnessValues,
 	const FSolverVec3& ExpBucklingStiffnessValues, const FSolverReal DampingRatioValue) const
 {
 	const TVec4<int32>& Constraint = Constraints[ConstraintIndex];
@@ -505,7 +585,8 @@ void FXPBDAnisotropicBendingConstraints::ApplyHelper(FSolverParticles& Particles
 	Lambda += DLambda;
 }
 
-void FXPBDAnisotropicBendingConstraints::Apply(FSolverParticles& Particles, const FSolverReal Dt) const
+template<typename SolverParticlesOrRange>
+void FXPBDAnisotropicBendingConstraints::Apply(SolverParticlesOrRange& Particles, const FSolverReal Dt) const
 {
 	TRACE_CPUPROFILER_EVENT_SCOPE(FXPBDAnisotropicBendingConstraints_Apply);
 	SCOPE_CYCLE_COUNTER(STAT_XPBD_AnisoBending);
@@ -545,7 +626,7 @@ void FXPBDAnisotropicBendingConstraints::Apply(FSolverParticles& Particles, cons
 						const int32 ColorSize = ConstraintsPerColorStartIndex[ConstraintColorIndex + 1] - ColorStart;
 						ispc::ApplyXPBDAnisotropicBendingConstraintsWithDamping(
 							(ispc::FVector4f*)Particles.GetPAndInvM().GetData(),
-							(const ispc::FVector3f*)Particles.X().GetData(),
+							(const ispc::FVector3f*)Particles.XArray().GetData(),
 							(ispc::FIntVector4*)&Constraints.GetData()[ColorStart],
 							&RestAngles.GetData()[ColorStart],
 							&IsBuckled.GetData()[ColorStart],
@@ -608,7 +689,7 @@ void FXPBDAnisotropicBendingConstraints::Apply(FSolverParticles& Particles, cons
 						const int32 ColorSize = ConstraintsPerColorStartIndex[ConstraintColorIndex + 1] - ColorStart;
 						ispc::ApplyXPBDAnisotropicBendingConstraintsWithDampingAndMaps(
 							(ispc::FVector4f*)Particles.GetPAndInvM().GetData(),
-							(const ispc::FVector3f*)Particles.X().GetData(),
+							(const ispc::FVector3f*)Particles.XArray().GetData(),
 							(ispc::FIntVector4*)&Constraints.GetData()[ColorStart],
 							&RestAngles.GetData()[ColorStart],
 							&IsBuckled.GetData()[ColorStart],
@@ -752,6 +833,8 @@ void FXPBDAnisotropicBendingConstraints::Apply(FSolverParticles& Particles, cons
 		}
 	}
 }
+template CHAOS_API void FXPBDAnisotropicBendingConstraints::Apply(FSolverParticles& Particles, const FSolverReal Dt) const;
+template CHAOS_API void FXPBDAnisotropicBendingConstraints::Apply(FSolverParticlesRange& Particles, const FSolverReal Dt) const;
 
 void FXPBDAnisotropicBendingConstraints::ComputeGradTheta(const FSolverVec3& X0, const FSolverVec3& X1, const FSolverVec3& X2, const FSolverVec3& X3, const int32 Index, FSolverVec3& dThetadx, FSolverReal& Theta) 
 {
diff --git a/Engine/Source/Runtime/Experimental/Chaos/Private/Chaos/XPBDBendingConstraints.cpp b/Engine/Source/Runtime/Experimental/Chaos/Private/Chaos/XPBDBendingConstraints.cpp
--- a/Engine/Source/Runtime/Experimental/Chaos/Private/Chaos/XPBDBendingConstraints.cpp
+++ b/Engine/Source/Runtime/Experimental/Chaos/Private/Chaos/XPBDBendingConstraints.cpp
@@ -23,14 +23,15 @@ namespace Chaos::Softs {
 int32 Chaos_XPBDBending_ParallelConstraintCount = 100;
 FAutoConsoleVariableRef CVarChaosXPBDBendingParallelConstraintCount(TEXT("p.Chaos.XPBDBending.ParallelConstraintCount"), Chaos_XPBDBending_ParallelConstraintCount, TEXT("If we have more constraints than this, use parallel-for in Apply."));
 
-void FXPBDBendingConstraints::InitColor(const FSolverParticles& InParticles)
+template<typename SolverParticlesOrRange>
+void FXPBDBendingConstraints::InitColor(const SolverParticlesOrRange& InParticles)
 {
 	// In dev builds we always color so we can tune the system without restarting. See Apply()
 #if UE_BUILD_SHIPPING || UE_BUILD_TEST
 	if (Constraints.Num() > Chaos_XPBDBending_ParallelConstraintCount)
 #endif
 	{
-		const TArray<TArray<int32>> ConstraintsPerColor = FGraphColoring::ComputeGraphColoring(Constraints, InParticles, ParticleOffset, ParticleOffset + ParticleCount);
+		const TArray<TArray<int32>> ConstraintsPerColor = FGraphColoring::ComputeGraphColoringParticlesOrRange(Constraints, InParticles, ParticleOffset, ParticleOffset + ParticleCount);
 
 		// Reorder constraints based on color so each array in ConstraintsPerColor contains contiguous elements.
 		TArray<TVec4<int32>> ReorderedConstraints; 
@@ -67,6 +68,8 @@ void FXPBDBendingConstraints::InitColor(const FSolverParticles& InParticles)
 		BucklingStiffness.ReorderIndices(OrigToReorderedIndices);
 	}
 }
+template CHAOS_API void FXPBDBendingConstraints::InitColor(const FSolverParticles& InParticles);
+template CHAOS_API void FXPBDBendingConstraints::InitColor(const FSolverParticlesRange& InParticles);
 
 void FXPBDBendingConstraints::SetProperties(
 	const FCollectionPropertyConstFacade& PropertyCollection,
@@ -138,7 +141,8 @@ void FXPBDBendingConstraints::SetProperties(
 	}
 }
 
-void FXPBDBendingConstraints::ApplyHelper(FSolverParticles& Particles, const FSolverReal Dt, const int32 ConstraintIndex, const FSolverReal StiffnessValue, const FSolverReal BucklingValue, const FSolverReal DampingRatioValue) const
+template<typename SolverParticlesOrRange>
+void FXPBDBendingConstraints::ApplyHelper(SolverParticlesOrRange& Particles, const FSolverReal Dt, const int32 ConstraintIndex, const FSolverReal StiffnessValue, const FSolverReal BucklingValue, const FSolverReal DampingRatioValue) const
 {
 	const TVec4<int32>& Constraint = Constraints[ConstraintIndex];
 	const int32 i1 = Constraint[0];
@@ -179,7 +183,8 @@ void FXPBDBendingConstraints::ApplyHelper(FSolverParticles& Particles, const FSo
 	Lambda += DLambda;
 }
 
-void FXPBDBendingConstraints::Apply(FSolverParticles& Particles, const FSolverReal Dt) const
+template<typename SolverParticlesOrRange>
+void FXPBDBendingConstraints::Apply(SolverParticlesOrRange& Particles, const FSolverReal Dt) const
 {
 	TRACE_CPUPROFILER_EVENT_SCOPE(FXPBDBendingConstraints_Apply);
 	SCOPE_CYCLE_COUNTER(STAT_XPBD_Bending);
@@ -213,7 +218,7 @@ void FXPBDBendingConstraints::Apply(FSolverParticles& Particles, const FSolverRe
 						const int32 ColorSize = ConstraintsPerColorStartIndex[ConstraintColorIndex + 1] - ColorStart;
 						ispc::ApplyXPBDBendingConstraintsWithDamping(
 							(ispc::FVector4f*)Particles.GetPAndInvM().GetData(),
-							(const ispc::FVector3f*)Particles.X().GetData(),
+							(const ispc::FVector3f*)Particles.XArray().GetData(),
 							(ispc::FIntVector4*)&Constraints.GetData()[ColorStart],
 							&RestAngles.GetData()[ColorStart],
 							&IsBuckled.GetData()[ColorStart],
@@ -274,7 +279,7 @@ void FXPBDBendingConstraints::Apply(FSolverParticles& Particles, const FSolverRe
 						const int32 ColorSize = ConstraintsPerColorStartIndex[ConstraintColorIndex + 1] - ColorStart;
 						ispc::ApplyXPBDBendingConstraintsWithDampingAndMaps(
 							(ispc::FVector4f*)Particles.GetPAndInvM().GetData(),
-							(const ispc::FVector3f*)Particles.X().GetData(),
+							(const ispc::FVector3f*)Particles.XArray().GetData(),
 							(ispc::FIntVector4*)&Constraints.GetData()[ColorStart],
 							&RestAngles.GetData()[ColorStart],
 							&IsBuckled.GetData()[ColorStart],
@@ -370,6 +375,8 @@ void FXPBDBendingConstraints::Apply(FSolverParticles& Particles, const FSolverRe
 		}
 	}
 }
+template CHAOS_API void FXPBDBendingConstraints::Apply(FSolverParticles& Particles, const FSolverReal Dt) const;
+template CHAOS_API void FXPBDBendingConstraints::Apply(FSolverParticlesRange& Particles, const FSolverReal Dt) const;
 
 FSolverReal FXPBDBendingConstraints::ComputeTotalEnergy(const FSolverParticles& InParticles, const FSolverReal ExplicitStiffness)
 {
diff --git a/Engine/Source/Runtime/Experimental/Chaos/Private/Chaos/XPBDSpringConstraints.cpp b/Engine/Source/Runtime/Experimental/Chaos/Private/Chaos/XPBDSpringConstraints.cpp
--- a/Engine/Source/Runtime/Experimental/Chaos/Private/Chaos/XPBDSpringConstraints.cpp
+++ b/Engine/Source/Runtime/Experimental/Chaos/Private/Chaos/XPBDSpringConstraints.cpp
@@ -30,14 +30,15 @@ static int32 Chaos_XPBDSpring_ParallelConstraintCount = 100;
 FAutoConsoleVariableRef CVarChaosXPBDSpringParallelConstraintCount(TEXT("p.Chaos.XPBDSpring.ParallelConstraintCount"), Chaos_XPBDSpring_ParallelConstraintCount, TEXT("If we have more constraints than this, use parallel-for in Apply."));
 #endif
 
-void FXPBDSpringConstraints::InitColor(const FSolverParticles& Particles)
+template<typename SolverParticlesOrRange>
+void FXPBDSpringConstraints::InitColor(const SolverParticlesOrRange& Particles)
 {
 	// In dev builds we always color so we can tune the system without restarting. See Apply()
 #if UE_BUILD_SHIPPING || UE_BUILD_TEST
 	if (Constraints.Num() > Chaos_XPBDSpring_ParallelConstraintCount)
 #endif
 	{
-		const TArray<TArray<int32>> ConstraintsPerColor = FGraphColoring::ComputeGraphColoring(Constraints, Particles, ParticleOffset, ParticleOffset + ParticleCount);
+		const TArray<TArray<int32>> ConstraintsPerColor = FGraphColoring::ComputeGraphColoringParticlesOrRange(Constraints, Particles, ParticleOffset, ParticleOffset + ParticleCount);
 		
 		// Reorder constraints based on color so each array in ConstraintsPerColor contains contiguous elements.
 		TArray<TVec2<int32>> ReorderedConstraints;
@@ -71,8 +72,11 @@ void FXPBDSpringConstraints::InitColor(const FSolverParticles& Particles)
 		DampingRatio.ReorderIndices(OrigToReorderedIndices);
 	}
 }
+template CHAOS_API void FXPBDSpringConstraints::InitColor(const FSolverParticles& Particles);
+template CHAOS_API void FXPBDSpringConstraints::InitColor(const FSolverParticlesRange& Particles);
 
-void FXPBDSpringConstraints::ApplyHelper(FSolverParticles& Particles, const FSolverReal Dt, const int32 ConstraintIndex, const FSolverReal ExpStiffnessValue, const FSolverReal DampingRatioValue) const
+template<typename SolverParticlesOrRange>
+void FXPBDSpringConstraints::ApplyHelper(SolverParticlesOrRange& Particles, const FSolverReal Dt, const int32 ConstraintIndex, const FSolverReal ExpStiffnessValue, const FSolverReal DampingRatioValue) const
 {
 	const TVec2<int32>& Constraint = Constraints[ConstraintIndex];
 	const int32 i1 = Constraint[0];
@@ -88,7 +92,8 @@ void FXPBDSpringConstraints::ApplyHelper(FSolverParticles& Particles, const FSol
 	}
 }
 
-void FXPBDSpringConstraints::Apply(FSolverParticles& Particles, const FSolverReal Dt) const
+template<typename SolverParticlesOrRange>
+void FXPBDSpringConstraints::Apply(SolverParticlesOrRange& Particles, const FSolverReal Dt) const
 {
 	TRACE_CPUPROFILER_EVENT_SCOPE(FXPBDSpringConstraints_Apply);
 	SCOPE_CYCLE_COUNTER(STAT_XPBD_Spring);
@@ -119,7 +124,7 @@ void FXPBDSpringConstraints::Apply(FSolverParticles& Particles, const FSolverRea
 						const int32 ColorSize = ConstraintsPerColorStartIndex[ConstraintColorIndex + 1] - ColorStart;
 						ispc::ApplyXPBDSpringConstraintsWithDamping(
 							(ispc::FVector4f*)Particles.GetPAndInvM().GetData(),
-							(const ispc::FVector3f*)Particles.X().GetData(),
+							(const ispc::FVector3f*)Particles.XArray().GetData(),
 							(ispc::FIntVector2*)&Constraints.GetData()[ColorStart],
 							&Dists.GetData()[ColorStart],
 							&Lambdas.GetData()[ColorStart],
@@ -174,7 +179,7 @@ void FXPBDSpringConstraints::Apply(FSolverParticles& Particles, const FSolverRea
 						const int32 ColorSize = ConstraintsPerColorStartIndex[ConstraintColorIndex + 1] - ColorStart;
 						ispc::ApplyXPBDSpringConstraintsWithDampingAndWeightMaps(
 							(ispc::FVector4f*)Particles.GetPAndInvM().GetData(),
-							(const ispc::FVector3f*)Particles.X().GetData(),
+							(const ispc::FVector3f*)Particles.XArray().GetData(),
 							(ispc::FIntVector2*)&Constraints.GetData()[ColorStart],
 							&Dists.GetData()[ColorStart],
 							&Lambdas.GetData()[ColorStart],
@@ -256,6 +261,8 @@ void FXPBDSpringConstraints::Apply(FSolverParticles& Particles, const FSolverRea
 		}
 	}
 }
+template CHAOS_API void FXPBDSpringConstraints::Apply(FSolverParticles& Particles, const FSolverReal Dt) const;
+template CHAOS_API void FXPBDSpringConstraints::Apply(FSolverParticlesRange& Particles, const FSolverReal Dt) const;
 
 void FXPBDEdgeSpringConstraints::SetProperties(
 	const FCollectionPropertyConstFacade& PropertyCollection,
diff --git a/Engine/Source/Runtime/Experimental/Chaos/Private/Chaos/XPBDStretchBiasElementConstraints.cpp b/Engine/Source/Runtime/Experimental/Chaos/Private/Chaos/XPBDStretchBiasElementConstraints.cpp
--- a/Engine/Source/Runtime/Experimental/Chaos/Private/Chaos/XPBDStretchBiasElementConstraints.cpp
+++ b/Engine/Source/Runtime/Experimental/Chaos/Private/Chaos/XPBDStretchBiasElementConstraints.cpp
@@ -1,5 +1,6 @@
 // Copyright Epic Games, Inc. All Rights Reserved.
 #include "Chaos/XPBDStretchBiasElementConstraints.h"
+#include "Chaos/SoftsSolverParticlesRange.h"
 #include "Chaos/GraphColoring.h"
 #include "Chaos/Framework/Parallel.h"
 #include "Chaos/TriangleMesh.h"
@@ -31,6 +32,73 @@ namespace Chaos::Softs {
 // @todo(chaos): the parallel threshold (or decision to run parallel) should probably be owned by the solver and passed to the constraint container
 static int32 Chaos_XPBDStretchBias_ParallelConstraintCount = 100;
 
+
+ FXPBDStretchBiasElementConstraints::FXPBDStretchBiasElementConstraints(const FSolverParticlesRange& InParticles,
+	const FTriangleMesh& TriangleMesh,
+	const TArray<TVec3<FVec2f>>& FaceVertexUVs,
+	const TMap<FString, TConstArrayView<FRealSingle>>& WeightMaps,
+	const FCollectionPropertyConstFacade& PropertyCollection,
+	bool bTrimKinematicConstraints)
+	: ParticleOffset(0)
+	, ParticleCount(InParticles.GetRangeSize())
+	, StiffnessWarp(
+		FSolverVec2(GetWeightedFloatXPBDAnisoStretchStiffnessWarp(PropertyCollection, MaxStiffness)),
+		WeightMaps.FindRef(GetXPBDAnisoStretchStiffnessWarpString(PropertyCollection, XPBDAnisoStretchStiffnessWarpName.ToString())),
+		TConstArrayView<TVec3<int32>>(Constraints),
+		ParticleOffset,
+		ParticleCount,
+		FPBDStiffness::DefaultTableSize,
+		FPBDStiffness::DefaultParameterFitBase,
+		MaxStiffness)
+	, StiffnessWeft(
+		FSolverVec2(GetWeightedFloatXPBDAnisoStretchStiffnessWeft(PropertyCollection, MaxStiffness)),
+		WeightMaps.FindRef(GetXPBDAnisoStretchStiffnessWeftString(PropertyCollection, XPBDAnisoStretchStiffnessWeftName.ToString())),
+		TConstArrayView<TVec3<int32>>(Constraints),
+		ParticleOffset,
+		ParticleCount,
+		FPBDStiffness::DefaultTableSize,
+		FPBDStiffness::DefaultParameterFitBase,
+		MaxStiffness)
+	, StiffnessBias(
+		FSolverVec2(GetWeightedFloatXPBDAnisoStretchStiffnessBias(PropertyCollection, MaxStiffness)),
+		WeightMaps.FindRef(GetXPBDAnisoStretchStiffnessBiasString(PropertyCollection, XPBDAnisoStretchStiffnessBiasName.ToString())),
+		TConstArrayView<TVec3<int32>>(Constraints),
+		ParticleOffset,
+		ParticleCount,
+		FPBDStiffness::DefaultTableSize,
+		FPBDStiffness::DefaultParameterFitBase,
+		MaxStiffness)
+	, DampingRatio(
+		FSolverVec2(GetWeightedFloatXPBDAnisoStretchDamping(PropertyCollection, MinDamping)).ClampAxes(MinDamping, MaxDamping),
+		WeightMaps.FindRef(GetXPBDAnisoStretchDampingString(PropertyCollection, XPBDAnisoStretchDampingName.ToString())),
+		TConstArrayView<TVec3<int32>>(Constraints),
+		ParticleOffset,
+		ParticleCount)
+	, WarpScale(
+		FSolverVec2(GetWeightedFloatXPBDAnisoStretchWarpScale(PropertyCollection, DefaultWarpWeftScale)).ClampAxes(MinWarpWeftScale, MaxWarpWeftScale),
+		WeightMaps.FindRef(GetXPBDAnisoStretchWarpScaleString(PropertyCollection, XPBDAnisoStretchWarpScaleName.ToString())),
+		TConstArrayView<TVec3<int32>>(Constraints),
+		ParticleOffset,
+		ParticleCount)
+	, WeftScale(
+		FSolverVec2(GetWeightedFloatXPBDAnisoStretchWeftScale(PropertyCollection, DefaultWarpWeftScale)).ClampAxes(MinWarpWeftScale, MaxWarpWeftScale),
+		WeightMaps.FindRef(GetXPBDAnisoStretchWeftScaleString(PropertyCollection, XPBDAnisoStretchWeftScaleName.ToString())),
+		TConstArrayView<TVec3<int32>>(Constraints),
+		ParticleOffset,
+		ParticleCount)
+	, XPBDAnisoStretchUse3dRestLengthsIndex(PropertyCollection)
+	, XPBDAnisoStretchStiffnessWarpIndex(PropertyCollection)
+	, XPBDAnisoStretchStiffnessWeftIndex(PropertyCollection)
+	, XPBDAnisoStretchStiffnessBiasIndex(PropertyCollection)
+	, XPBDAnisoStretchDampingIndex(PropertyCollection)
+	, XPBDAnisoStretchWarpScaleIndex(PropertyCollection)
+	, XPBDAnisoStretchWeftScaleIndex(PropertyCollection)
+{
+	Lambdas.Init(FSolverVec3(0.), Constraints.Num());
+	InitConstraintsAndRestData(InParticles, TriangleMesh, FaceVertexUVs, GetXPBDAnisoStretchUse3dRestLengths(PropertyCollection, bDefaultUse3dRestLengths), bTrimKinematicConstraints);
+	InitColor(InParticles);
+}
+
 FXPBDStretchBiasElementConstraints::FXPBDStretchBiasElementConstraints(const FSolverParticles& InParticles,
 	int32 InParticleOffset,
 	int32 InParticleCount,
@@ -178,7 +246,8 @@ FXPBDStretchBiasElementConstraints::FXPBDStretchBiasElementConstraints(const FSo
 	InitColor(InParticles);
 }
 
-void FXPBDStretchBiasElementConstraints::InitConstraintsAndRestData(const FSolverParticles& InParticles, const FTriangleMesh& TriangleMesh,
+template<typename SolverParticlesOrRange>
+void FXPBDStretchBiasElementConstraints::InitConstraintsAndRestData(const SolverParticlesOrRange& InParticles, const FTriangleMesh& TriangleMesh,
 	const TArray<TVec3<FSolverVec2>>& FaceVertexUVs, const bool bUse3dRestLengths, const bool bTrimKinematicConstraints)
 {
 	const TArray<TVec3<int32>>& Elements = TriangleMesh.GetElements();
@@ -269,7 +338,8 @@ void FXPBDStretchBiasElementConstraints::InitConstraintsAndRestData(const FSolve
 	}
 }
 
-void FXPBDStretchBiasElementConstraints::InitColor(const FSolverParticles& InParticles)
+template<typename SolverParticlesOrRange>
+void FXPBDStretchBiasElementConstraints::InitColor(const SolverParticlesOrRange& InParticles)
 {
 	ConstraintsPerColorStartIndex.Reset();
 
@@ -278,7 +348,7 @@ void FXPBDStretchBiasElementConstraints::InitColor(const FSolverParticles& InPar
 	if (Constraints.Num() > Chaos_XPBDStretchBias_ParallelConstraintCount)
 #endif
 	{
-		const TArray<TArray<int32>> ConstraintsPerColor = FGraphColoring::ComputeGraphColoring(Constraints, InParticles, ParticleOffset, ParticleOffset + ParticleCount);
+		const TArray<TArray<int32>> ConstraintsPerColor = FGraphColoring::ComputeGraphColoringParticlesOrRange(Constraints, InParticles, ParticleOffset, ParticleOffset + ParticleCount);
 
 		// Reorder constraints based on color so each array in ConstraintsPerColor contains contiguous elements.
 		TArray<TVec3<int32>> ReorderedConstraints;
@@ -455,7 +525,8 @@ void FXPBDStretchBiasElementConstraints::SetProperties(
 	}
 }
 
-void FXPBDStretchBiasElementConstraints::Apply(FSolverParticles& Particles, const FSolverReal Dt) const
+template<typename SolverParticlesOrRange>
+void FXPBDStretchBiasElementConstraints::Apply(SolverParticlesOrRange& Particles, const FSolverReal Dt) const
 {
 	TRACE_CPUPROFILER_EVENT_SCOPE(FXPBDStretchBiasElementConstraints_Apply);
 	SCOPE_CYCLE_COUNTER(STAT_XPBD_StretchBias);
@@ -496,7 +567,7 @@ void FXPBDStretchBiasElementConstraints::Apply(FSolverParticles& Particles, cons
 						const int32 ColorSize = ConstraintsPerColorStartIndex[ConstraintColorIndex + 1] - ColorStart;
 						ispc::ApplyXPBDStretchBiasConstraintsWithDamping(
 							(ispc::FVector4f*)Particles.GetPAndInvM().GetData(),
-							(const ispc::FVector3f*)Particles.X().GetData(),
+							(const ispc::FVector3f*)Particles.XArray().GetData(),
 							(ispc::FIntVector*)&Constraints.GetData()[ColorStart],
 							(ispc::FVector3f*)&RestStretchLengths.GetData()[ColorStart],
 							(ispc::FVector4f*)&DeltaUVInverse.GetData()[ColorStart],
@@ -563,7 +634,7 @@ void FXPBDStretchBiasElementConstraints::Apply(FSolverParticles& Particles, cons
 						const int32 ColorSize = ConstraintsPerColorStartIndex[ConstraintColorIndex + 1] - ColorStart;
 						ispc::ApplyXPBDStretchBiasConstraintsWithDampingAndMaps(
 							(ispc::FVector4f*)Particles.GetPAndInvM().GetData(),
-							(const ispc::FVector3f*)Particles.X().GetData(),
+							(const ispc::FVector3f*)Particles.XArray().GetData(),
 							(ispc::FIntVector*)&Constraints.GetData()[ColorStart],
 							(ispc::FVector3f*)&RestStretchLengths.GetData()[ColorStart],
 							(ispc::FVector4f*)&DeltaUVInverse.GetData()[ColorStart],
@@ -679,6 +750,8 @@ void FXPBDStretchBiasElementConstraints::Apply(FSolverParticles& Particles, cons
 		}
 	}
 }
+template CHAOS_API void FXPBDStretchBiasElementConstraints::Apply(FSolverParticles& Particles, const FSolverReal Dt) const;
+template CHAOS_API void FXPBDStretchBiasElementConstraints::Apply(FSolverParticlesRange& Particles, const FSolverReal Dt) const;
 
 void FXPBDStretchBiasElementConstraints::CalculateUVStretch(const int32 ConstraintIndex, const FSolverVec3& P0, const FSolverVec3& P1, const FSolverVec3& P2, FSolverVec3& DXDu, FSolverVec3& DXDv) const
 {
@@ -706,8 +779,8 @@ static FSolverReal CalcDLambda(const FSolverReal StiffnessValue, const FSolverRe
 	return (C - Alpha * Lambda + DampingTerm) / Denom;
 }
 
-
-void FXPBDStretchBiasElementConstraints::ApplyHelper(FSolverParticles& Particles, const FSolverReal Dt, const int32 ConstraintIndex, const FSolverVec3& ExpStiffnessValue, const FSolverReal DampingRatioValue, const FSolverReal WarpScaleValue, const FSolverReal WeftScaleValue) const
+template<typename SolverParticlesOrRange>
+void FXPBDStretchBiasElementConstraints::ApplyHelper(SolverParticlesOrRange& Particles, const FSolverReal Dt, const int32 ConstraintIndex, const FSolverVec3& ExpStiffnessValue, const FSolverReal DampingRatioValue, const FSolverReal WarpScaleValue, const FSolverReal WeftScaleValue) const
 {
 	const TVec3<int32>& Constraint = Constraints[ConstraintIndex];
 	const int32 i0 = Constraint[0];
diff --git a/Engine/Source/Runtime/Experimental/Chaos/Public/Chaos/DynamicParticles.h b/Engine/Source/Runtime/Experimental/Chaos/Public/Chaos/DynamicParticles.h
--- a/Engine/Source/Runtime/Experimental/Chaos/Public/Chaos/DynamicParticles.h
+++ b/Engine/Source/Runtime/Experimental/Chaos/Public/Chaos/DynamicParticles.h
@@ -31,17 +31,22 @@ class TDynamicParticles : public TParticles<T, d>
 	const TVector<T, d>& V(const int32 Index) const { return MV[Index]; }
 	TVector<T, d>& V(const int32 Index) { return MV[Index]; }
 	const TArrayCollectionArray<TVector<T, d>>& GetV() const { return MV; }
+	TArrayCollectionArray<TVector<T, d>>& GetV() { return MV; }
 
 	const TVector<T, d>& Acceleration(const int32 Index) const { return MAcceleration[Index]; }
 	TVector<T, d>& Acceleration(const int32 Index) { return MAcceleration[Index]; }
+	const TArrayCollectionArray<TVector<T, d>>& GetAcceleration() const { return MAcceleration; }
+	TArrayCollectionArray<TVector<T, d>>& GetAcceleration() { return MAcceleration; }
 
 	const T M(const int32 Index) const { return MM[Index]; }
 	T& M(const int32 Index) { return MM[Index]; }
 	const TArrayCollectionArray<T>& GetM() const { return MM; }
+	TArrayCollectionArray<T>& GetM() { return MM; }
 
 	const T InvM(const int32 Index) const { return MInvM[Index]; }
 	T& InvM(const int32 Index) { return MInvM[Index]; }
 	const TArrayCollectionArray<T>& GetInvM() const { return MInvM; }
+	TArrayCollectionArray<T>& GetInvM() { return MInvM; }
 
   private:
 	TArrayCollectionArray<TVector<T, d>> MV, MAcceleration;
diff --git a/Engine/Source/Runtime/Experimental/Chaos/Public/Chaos/GraphColoring.h b/Engine/Source/Runtime/Experimental/Chaos/Public/Chaos/GraphColoring.h
--- a/Engine/Source/Runtime/Experimental/Chaos/Public/Chaos/GraphColoring.h
+++ b/Engine/Source/Runtime/Experimental/Chaos/Public/Chaos/GraphColoring.h
@@ -61,33 +61,54 @@ class FGraphColoring
 	};
 
   public:
+	template<typename DynamicParticlesType>
+	static TArray<TArray<int32>> ComputeGraphColoringParticlesOrRange(const TArray<TVector<int32, 2>>& Graph, const DynamicParticlesType& InParticles, const int32 GraphParticlesStart, const int32 GraphParticlesEnd);
+
 	template<typename T>
-	static TArray<TArray<int32>> ComputeGraphColoring(const TArray<TVector<int32, 2>>& Graph, const TDynamicParticles<T, 3>& InParticles, const int32 GraphParticlesStart, const int32 GraphParticlesEnd);
+	inline static TArray<TArray<int32>> ComputeGraphColoring(const TArray<TVector<int32, 2>>& Graph, const TDynamicParticles<T, 3>& InParticles, const int32 GraphParticlesStart, const int32 GraphParticlesEnd)
+	{
+		ComputeGraphColoringParticlesOrRange(Graph, InParticles, GraphParticlesStart, GraphParticlesEnd);
+	}
 	template<typename T>
 	inline static TArray<TArray<int32>> ComputeGraphColoring(const TArray<TVector<int32, 2>>& Graph, const TDynamicParticles<T, 3>& InParticles)
 	{
-		return ComputeGraphColoring(Graph, InParticles, 0, InParticles.Size());
+		return ComputeGraphColoringParticlesOrRange(Graph, InParticles, 0, InParticles.Size());
 	}
+	template<typename DynamicParticlesType>
+	static TArray<TArray<int32>> ComputeGraphColoringParticlesOrRange(const TArray<TVector<int32, 3>>& Graph, const DynamicParticlesType& InParticles, const int32 GraphParticlesStart, const int32 GraphParticlesEnd);
 	template<typename T>
-	static TArray<TArray<int32>> ComputeGraphColoring(const TArray<TVector<int32, 3>>& Graph, const TDynamicParticles<T, 3>& InParticles, const int32 GraphParticlesStart, const int32 GraphParticlesEnd);
+	inline static TArray<TArray<int32>> ComputeGraphColoring(const TArray<TVector<int32, 3>>& Graph, const TDynamicParticles<T, 3>& InParticles, const int32 GraphParticlesStart, const int32 GraphParticlesEnd)
+	{
+		return ComputeGraphColoringParticlesOrRange(Graph, InParticles, GraphParticlesStart, GraphParticlesEnd);
+	}
 	template<typename T>
 	inline static TArray<TArray<int32>> ComputeGraphColoring(const TArray<TVector<int32, 3>>& Graph, const TDynamicParticles<T, 3>& InParticles)
 	{
-		return ComputeGraphColoring(Graph, InParticles, 0, InParticles.Size());
+		return ComputeGraphColoringParticlesOrRange(Graph, InParticles, 0, InParticles.Size());
 	}
+	template<typename DynamicParticlesType>
+	static TArray<TArray<int32>> ComputeGraphColoringParticlesOrRange(const TArray<TVector<int32, 4>>& Graph, const DynamicParticlesType& InParticles, const int32 GraphParticlesStart, const int32 GraphParticlesEnd);
 	template<typename T>
-	static TArray<TArray<int32>> ComputeGraphColoring(const TArray<TVector<int32, 4>>& Graph, const TDynamicParticles<T, 3>& InParticles, const int32 GraphParticlesStart, const int32 GraphParticlesEnd);
+	inline static TArray<TArray<int32>> ComputeGraphColoring(const TArray<TVector<int32, 4>>& Graph, const TDynamicParticles<T, 3>& InParticles, const int32 GraphParticlesStart, const int32 GraphParticlesEnd)
+	{
+		return ComputeGraphColoringParticlesOrRange(Graph, InParticles, GraphParticlesStart, GraphParticlesEnd);
+	}
 	template<typename T>
 	inline static TArray<TArray<int32>> ComputeGraphColoring(const TArray<TVector<int32, 4>>& Graph, const TDynamicParticles<T, 3>& InParticles)
 	{
-		return ComputeGraphColoring(Graph, InParticles, 0, InParticles.Size());
+		return ComputeGraphColoringParticlesOrRange(Graph, InParticles, 0, InParticles.Size());
 	}
+	template<typename DynamicParticlesType>
+	static TArray<TArray<int32>> ComputeGraphColoringAllDynamicParticlesOrRange(const TArray<TVec4<int32>>& Graph, const DynamicParticlesType& InParticles, const int32 GraphParticlesStart, const int32 GraphParticlesEnd);
 	template<typename T>
-	static TArray<TArray<int32>> ComputeGraphColoringAllDynamic(const TArray<TVec4<int32>>& Graph, const Chaos::TDynamicParticles<T, 3>& InParticles, const int32 GraphParticlesStart, const int32 GraphParticlesEnd);
+	inline static TArray<TArray<int32>> ComputeGraphColoringAllDynamic(const TArray<TVec4<int32>>& Graph, const Chaos::TDynamicParticles<T, 3>& InParticles, const int32 GraphParticlesStart, const int32 GraphParticlesEnd)
+	{
+		return ComputeGraphColoringAllDynamicParticlesOrRange(Graph, InParticles, GraphParticlesStart, GraphParticlesEnd);
+	}
 	template<typename T>
 	inline static TArray<TArray<int32>> ComputeGraphColoringAllDynamic(const TArray<TVec4<int32>>& Graph, const Chaos::TDynamicParticles<T, 3>& InParticles)
 	{
-		return ComputeGraphColoring(Graph, InParticles, 0, InParticles.Size());
+		return ComputeGraphColoringAllDynamicParticlesOrRange(Graph, InParticles, 0, InParticles.Size());
 	}
 };
 
diff --git a/Engine/Source/Runtime/Experimental/Chaos/Public/Chaos/PBDAnimDriveConstraint.h b/Engine/Source/Runtime/Experimental/Chaos/Public/Chaos/PBDAnimDriveConstraint.h
--- a/Engine/Source/Runtime/Experimental/Chaos/Public/Chaos/PBDAnimDriveConstraint.h
+++ b/Engine/Source/Runtime/Experimental/Chaos/Public/Chaos/PBDAnimDriveConstraint.h
@@ -23,28 +23,6 @@ namespace Chaos::Softs
 			return IsAnimDriveStiffnessEnabled(PropertyCollection, false);
 		}
 
-		UE_DEPRECATED(5.2, "Use the other constructor supplying AnimationVelocities for correct subframe and damping behavior")
-		FPBDAnimDriveConstraint(
-			const int32 InParticleOffset,
-			const int32 InParticleCount,
-			const TArray<FSolverVec3>& InAnimationPositions,  // Use global indexation (will need adding ParticleOffset)
-			const TArray<FSolverVec3>& InOldAnimationPositions,  // Use global indexation (will need adding ParticleOffset)
-			const TConstArrayView<FRealSingle>& StiffnessMultipliers,  // Use local indexation
-			const TConstArrayView<FRealSingle>& DampingMultipliers  // Use local indexation
-		)
-			: AnimationPositions(InAnimationPositions)
-			, OldAnimationPositions_deprecated(InOldAnimationPositions)
-			, AnimationVelocities(InOldAnimationPositions) // Unused when using deprecated apply
-			, ParticleOffset(InParticleOffset)
-			, ParticleCount(InParticleCount)
-			, UseDeprecatedApply(true)
-			, Stiffness(FSolverVec2::UnitVector, StiffnessMultipliers, InParticleCount)
-			, Damping(FSolverVec2::UnitVector, DampingMultipliers, InParticleCount)
-			, AnimDriveStiffnessIndex(ForceInit)
-			, AnimDriveDampingIndex(ForceInit)
-		{
-		}
-
 		FPBDAnimDriveConstraint(
 			const int32 InParticleOffset,
 			const int32 InParticleCount,
@@ -54,11 +32,9 @@ namespace Chaos::Softs
 			const FCollectionPropertyConstFacade& PropertyCollection
 		)
 			: AnimationPositions(InAnimationPositions)
-			, OldAnimationPositions_deprecated(InAnimationVelocities) // Unused when not using apply
 			, AnimationVelocities(InAnimationVelocities)
 			, ParticleOffset(InParticleOffset)
 			, ParticleCount(InParticleCount)
-			, UseDeprecatedApply(false)
 			, Stiffness(
 				FSolverVec2(GetWeightedFloatAnimDriveStiffness(PropertyCollection, 1.f)),
 				WeightMaps.FindRef(GetAnimDriveStiffnessString(PropertyCollection, AnimDriveStiffnessName.ToString())),
@@ -83,11 +59,9 @@ namespace Chaos::Softs
 			const FCollectionPropertyConstFacade& PropertyCollection
 		)
 			: AnimationPositions(InAnimationPositions)
-			, OldAnimationPositions_deprecated(InAnimationVelocities) // Unused when not using apply
 			, AnimationVelocities(InAnimationVelocities)
 			, ParticleOffset(InParticleOffset)
 			, ParticleCount(InParticleCount)
-			, UseDeprecatedApply(false)
 			, Stiffness(
 				FSolverVec2(GetWeightedFloatAnimDriveStiffness(PropertyCollection, 1.f)),
 				StiffnessMultipliers,
@@ -111,11 +85,9 @@ namespace Chaos::Softs
 			const TConstArrayView<FRealSingle>& DampingMultipliers  // Use local indexation
 		)
 			: AnimationPositions(InAnimationPositions)
-			, OldAnimationPositions_deprecated(InAnimationVelocities) // Unused when not using apply
 			, AnimationVelocities(InAnimationVelocities)
 			, ParticleOffset(InParticleOffset)
 			, ParticleCount(InParticleCount)
-			, UseDeprecatedApply(false)
 			, Stiffness(FSolverVec2::UnitVector, StiffnessMultipliers, InParticleCount)
 			, Damping(FSolverVec2::UnitVector, DampingMultipliers, InParticleCount)
 			, AnimDriveStiffnessIndex(ForceInit)
@@ -182,7 +154,18 @@ namespace Chaos::Softs
 			Damping.ApplyPBDValues(Dt, NumIterations);
 		}
 
-		inline void Apply(FSolverParticles& InParticles, const FSolverReal Dt) const
+		TConstArrayView<FSolverVec3> GetConstArrayView(const FSolverParticles& Particles, const TArray<FSolverVec3>& Data) const
+		{
+			return TConstArrayView<FSolverVec3>(Data.GetData() + ParticleOffset, ParticleCount);
+		}
+
+		TConstArrayView<FSolverVec3> GetConstArrayView(const FSolverParticlesRange& Particles, const TArray<FSolverVec3>& Data) const
+		{
+			return Particles.GetConstArrayView(Data);
+		}
+
+		template<typename SolverParticlesOrRange>
+		inline void Apply(SolverParticlesOrRange& InParticles, const FSolverReal Dt) const
 		{
 			SCOPE_CYCLE_COUNTER(STAT_PBD_AnimDriveConstraint);
 
@@ -190,46 +173,25 @@ namespace Chaos::Softs
 			{
 				if (Damping.HasWeightMap())
 				{
-					if (UseDeprecatedApply)
+					const TConstArrayView<FSolverVec3> AnimationPositionsView = GetConstArrayView(InParticles, AnimationPositions);
+					const TConstArrayView<FSolverVec3> AnimationVelocitiesView = GetConstArrayView(InParticles, AnimationVelocities);
+					PhysicsParallelFor(ParticleCount, [this, &InParticles, &Dt, &AnimationPositionsView, &AnimationVelocitiesView](int32 Index)  // TODO: profile needed for these parallel loop based on particle count
 					{
-						PhysicsParallelFor(ParticleCount, [this, &InParticles, &Dt](int32 Index)  // TODO: profile needed for these parallel loop based on particle count
-							{
-								const FSolverReal ParticleStiffness = Stiffness[Index];
-								const FSolverReal ParticleDamping = Damping[Index];
-								ApplyHelper_Deprecated(InParticles, ParticleStiffness, ParticleDamping, Dt, Index);
-							});
-
-					}
-					else
-					{
-						PhysicsParallelFor(ParticleCount, [this, &InParticles, &Dt](int32 Index)  // TODO: profile needed for these parallel loop based on particle count
-							{
-								const FSolverReal ParticleStiffness = Stiffness[Index];
-								const FSolverReal ParticleDamping = Damping[Index];
-								ApplyHelper(InParticles, ParticleStiffness, ParticleDamping, Dt, Index);
-							});
-					}
+						const FSolverReal ParticleStiffness = Stiffness[Index];
+						const FSolverReal ParticleDamping = Damping[Index];
+						ApplyHelper(InParticles, ParticleStiffness, ParticleDamping, Dt, Index, AnimationPositionsView[Index], AnimationVelocitiesView[Index]);
+					});
 				}
 				else
 				{
 					const FSolverReal ParticleDamping = (FSolverReal)Damping;
-					if (UseDeprecatedApply)
-					{
-						PhysicsParallelFor(ParticleCount, [this, &InParticles, ParticleDamping, &Dt](int32 Index)
-							{
-								const FSolverReal ParticleStiffness = Stiffness[Index];
-								ApplyHelper_Deprecated(InParticles, ParticleStiffness, ParticleDamping, Dt, Index);
-							});
-
-					}
-					else
+					const TConstArrayView<FSolverVec3> AnimationPositionsView = GetConstArrayView(InParticles, AnimationPositions);
+					const TConstArrayView<FSolverVec3> AnimationVelocitiesView = GetConstArrayView(InParticles, AnimationVelocities);
+					PhysicsParallelFor(ParticleCount, [this, &InParticles, ParticleDamping, &Dt, &AnimationPositionsView, &AnimationVelocitiesView](int32 Index)
 					{
-						PhysicsParallelFor(ParticleCount, [this, &InParticles, ParticleDamping, &Dt](int32 Index)
-							{
-								const FSolverReal ParticleStiffness = Stiffness[Index];
-								ApplyHelper(InParticles, ParticleStiffness, ParticleDamping, Dt, Index);
-							});
-					}
+						const FSolverReal ParticleStiffness = Stiffness[Index];
+						ApplyHelper(InParticles, ParticleStiffness, ParticleDamping, Dt, Index, AnimationPositionsView[Index], AnimationVelocitiesView[Index]);
+					});
 				}
 			}
 			else
@@ -237,66 +199,32 @@ namespace Chaos::Softs
 				const FSolverReal ParticleStiffness = (FSolverReal)Stiffness;
 				if (Damping.HasWeightMap())
 				{
-					if (UseDeprecatedApply)
-					{
-						PhysicsParallelFor(ParticleCount, [this, &InParticles, &ParticleStiffness, &Dt](int32 Index)
-							{
-								const FSolverReal ParticleDamping = Damping[Index];
-								ApplyHelper_Deprecated(InParticles, ParticleStiffness, ParticleDamping, Dt, Index);
-							});
-					}
-					else
+					const TConstArrayView<FSolverVec3> AnimationPositionsView = GetConstArrayView(InParticles, AnimationPositions);
+					const TConstArrayView<FSolverVec3> AnimationVelocitiesView = GetConstArrayView(InParticles, AnimationVelocities);
+					PhysicsParallelFor(ParticleCount, [this, &InParticles, &ParticleStiffness, &Dt, &AnimationPositionsView, &AnimationVelocitiesView](int32 Index)
 					{
-						PhysicsParallelFor(ParticleCount, [this, &InParticles, &ParticleStiffness, &Dt](int32 Index)
-							{
-								const FSolverReal ParticleDamping = Damping[Index];
-								ApplyHelper(InParticles, ParticleStiffness, ParticleDamping, Dt, Index);
-							});
-					}
+						const FSolverReal ParticleDamping = Damping[Index];
+						ApplyHelper(InParticles, ParticleStiffness, ParticleDamping, Dt, Index, AnimationPositionsView[Index], AnimationVelocitiesView[Index]);
+					});
 				}
 				else
 				{
 					const FSolverReal ParticleDamping = (FSolverReal)Damping;
-					if (UseDeprecatedApply)
-					{
-						PhysicsParallelFor(ParticleCount, [this, &InParticles, &ParticleStiffness, &ParticleDamping, &Dt](int32 Index)
-							{
-								ApplyHelper_Deprecated(InParticles, ParticleStiffness, ParticleDamping, Dt, Index);
-							});
-					}
-					else
+					const TConstArrayView<FSolverVec3> AnimationPositionsView = GetConstArrayView(InParticles, AnimationPositions);
+					const TConstArrayView<FSolverVec3> AnimationVelocitiesView = GetConstArrayView(InParticles, AnimationVelocities);
+					PhysicsParallelFor(ParticleCount, [this, &InParticles, &ParticleStiffness, &ParticleDamping, &Dt, &AnimationPositionsView, &AnimationVelocitiesView](int32 Index)
 					{
-						PhysicsParallelFor(ParticleCount, [this, &InParticles, &ParticleStiffness, &ParticleDamping, &Dt](int32 Index)
-							{
-								ApplyHelper(InParticles, ParticleStiffness, ParticleDamping, Dt, Index);
-							});
-					}
+						ApplyHelper(InParticles, ParticleStiffness, ParticleDamping, Dt, Index, AnimationPositionsView[Index], AnimationVelocitiesView[Index]);
+					});
 				}
 			}
 		}
 
 	private:
-		inline void ApplyHelper(FSolverParticles& Particles, const FSolverReal InStiffness, const FSolverReal InDamping, const FSolverReal Dt, const int32 Index) const
-		{
-			const int32 ParticleIndex = ParticleOffset + Index;
-			if (Particles.InvM(ParticleIndex) == (FSolverReal)0.)
-			{
-				return;
-			}
-
-			FSolverVec3& ParticlePosition = Particles.P(ParticleIndex);
 
-			ParticlePosition -= InStiffness * (ParticlePosition - AnimationPositions[ParticleIndex]);
-
-			const FSolverVec3 ParticleDisplacement = ParticlePosition - Particles.X(ParticleIndex);
-			const FSolverVec3 AnimationDisplacement = (AnimationVelocities[ParticleIndex]) * Dt;
-			const FSolverVec3 RelativeDisplacement = ParticleDisplacement - AnimationDisplacement;
 
-			ParticlePosition -= InDamping * RelativeDisplacement;
-		}
-
-		// This method does not have the correct substepping or damping behavior, but retaining for backwards compatibility until deprecated constructor can be removed
-		inline void ApplyHelper_Deprecated(FSolverParticles& Particles, const FSolverReal InStiffness, const FSolverReal InDamping, const FSolverReal Dt, const int32 Index) const
+		template<typename SolverParticlesOrRange>
+		inline void ApplyHelper(SolverParticlesOrRange& Particles, const FSolverReal InStiffness, const FSolverReal InDamping, const FSolverReal Dt, const int32 Index, const FSolverVec3& AnimPosition, const FSolverVec3& AnimVelocity) const
 		{
 			const int32 ParticleIndex = ParticleOffset + Index;
 			if (Particles.InvM(ParticleIndex) == (FSolverReal)0.)
@@ -305,23 +233,21 @@ namespace Chaos::Softs
 			}
 
 			FSolverVec3& ParticlePosition = Particles.P(ParticleIndex);
-			const FSolverVec3& AnimationPosition = AnimationPositions[ParticleIndex];
-			const FSolverVec3& OldAnimationPosition = OldAnimationPositions_deprecated[ParticleIndex];
+
+			ParticlePosition -= InStiffness * (ParticlePosition - AnimPosition);
 
 			const FSolverVec3 ParticleDisplacement = ParticlePosition - Particles.X(ParticleIndex);
-			const FSolverVec3 AnimationDisplacement = OldAnimationPosition - AnimationPosition;
+			const FSolverVec3 AnimationDisplacement = (AnimVelocity) * Dt;
 			const FSolverVec3 RelativeDisplacement = ParticleDisplacement - AnimationDisplacement;
 
-			ParticlePosition -= InStiffness * (ParticlePosition - AnimationPosition) + InDamping * RelativeDisplacement;
+			ParticlePosition -= InDamping * RelativeDisplacement;
 		}
 
 	private:
 		const TArray<FSolverVec3>& AnimationPositions;  // Use global index (needs adding ParticleOffset)
-		const TArray<FSolverVec3>& OldAnimationPositions_deprecated;  // Use global index (needs adding ParticleOffset). Only used by ApplyHelper_Deprecated until old constructor can be removed.
 		const TArray<FSolverVec3>& AnimationVelocities;  // Use global index (needs adding ParticleOffset).
 		const int32 ParticleOffset;
 		const int32 ParticleCount;
-		const bool UseDeprecatedApply;
 
 		FPBDStiffness Stiffness;
 		FPBDStiffness Damping;
diff --git a/Engine/Source/Runtime/Experimental/Chaos/Public/Chaos/PBDAxialSpringConstraints.h b/Engine/Source/Runtime/Experimental/Chaos/Public/Chaos/PBDAxialSpringConstraints.h
--- a/Engine/Source/Runtime/Experimental/Chaos/Public/Chaos/PBDAxialSpringConstraints.h
+++ b/Engine/Source/Runtime/Experimental/Chaos/Public/Chaos/PBDAxialSpringConstraints.h
@@ -15,6 +15,22 @@ class FPBDAxialSpringConstraints : public FPBDAxialSpringConstraintsBase
 	using Base::Barys;
 
 public:
+	FPBDAxialSpringConstraints(
+		const FSolverParticlesRange& Particles,
+		const TArray<TVec3<int32>>& InConstraints,
+		const TConstArrayView<FRealSingle>& StiffnessMultipliers,
+		const FSolverVec2& InStiffness,
+		bool bTrimKinematicConstraints)
+		: Base(
+			Particles,
+			InConstraints,
+			StiffnessMultipliers,
+			InStiffness,
+			bTrimKinematicConstraints)
+	{
+		InitColor(Particles);
+	}
+
 	FPBDAxialSpringConstraints(
 		const FSolverParticles& Particles,
 		int32 InParticleOffset,
@@ -37,7 +53,8 @@ public:
 
 	virtual ~FPBDAxialSpringConstraints() override {}
 
-	CHAOS_API void Apply(FSolverParticles& InParticles, const FSolverReal Dt) const;
+	template<typename SolverParticlesOrRange>
+	CHAOS_API void Apply(SolverParticlesOrRange& InParticles, const FSolverReal Dt) const;
 
 protected:
 	using Base::Constraints;
@@ -46,8 +63,10 @@ protected:
 	using Base::ParticleCount;
 
 private:
-	CHAOS_API void InitColor(const FSolverParticles& InParticles);
-	CHAOS_API void ApplyHelper(FSolverParticles& InParticles, const FSolverReal Dt, const int32 ConstraintIndex, const FSolverReal ExpStiffnessValue) const;
+	template<typename SolverParticlesOrRange>
+	CHAOS_API void InitColor(const SolverParticlesOrRange& InParticles);
+	template<typename SolverParticlesOrRange>
+	void ApplyHelper(SolverParticlesOrRange& InParticles, const FSolverReal Dt, const int32 ConstraintIndex, const FSolverReal ExpStiffnessValue) const;
 
 	TArray<int32> ConstraintsPerColorStartIndex; // Constraints are ordered so each batch is contiguous. This is ColorNum + 1 length so it can be used as start and end.
 };
@@ -60,6 +79,21 @@ public:
 		return IsAreaSpringStiffnessEnabled(PropertyCollection, false);
 	}
 
+	FPBDAreaSpringConstraints(
+		const FSolverParticlesRange& Particles,
+		const TArray<TVec3<int32>>& InConstraints,
+		const TMap<FString, TConstArrayView<FRealSingle>>& WeightMaps,
+		const FCollectionPropertyConstFacade& PropertyCollection,
+		bool bTrimKinematicConstraints)
+		: FPBDAxialSpringConstraints(
+			Particles,
+			InConstraints,
+			WeightMaps.FindRef(GetAreaSpringStiffnessString(PropertyCollection, AreaSpringStiffnessName.ToString())),
+			FSolverVec2(GetWeightedFloatAreaSpringStiffness(PropertyCollection, 1.f)),
+			bTrimKinematicConstraints)
+		, AreaSpringStiffnessIndex(PropertyCollection)
+	{}
+
 	FPBDAreaSpringConstraints(
 		const FSolverParticles& Particles,
 		int32 InParticleOffset,
diff --git a/Engine/Source/Runtime/Experimental/Chaos/Public/Chaos/PBDAxialSpringConstraintsBase.h b/Engine/Source/Runtime/Experimental/Chaos/Public/Chaos/PBDAxialSpringConstraintsBase.h
--- a/Engine/Source/Runtime/Experimental/Chaos/Public/Chaos/PBDAxialSpringConstraintsBase.h
+++ b/Engine/Source/Runtime/Experimental/Chaos/Public/Chaos/PBDAxialSpringConstraintsBase.h
@@ -6,6 +6,7 @@
 #include "Chaos/PBDSoftsEvolutionFwd.h"
 #include "Chaos/PBDSoftsSolverParticles.h"
 #include "Chaos/PBDStiffness.h"
+#include "Chaos/SoftsSolverParticlesRange.h"
 #include "Containers/Array.h"
 
 namespace Chaos::Softs
@@ -14,6 +15,34 @@ namespace Chaos::Softs
 class FPBDAxialSpringConstraintsBase
 {
 public:
+
+	FPBDAxialSpringConstraintsBase(
+		const FSolverParticlesRange& Particles,
+		const TArray<TVec3<int32>>& InConstraints,
+		const TConstArrayView<FRealSingle>& StiffnessMultipliers,
+		const FSolverVec2& InStiffness,
+		bool bTrimKinematicConstraints,
+		FSolverReal MaxStiffness = FPBDStiffness::DefaultPBDMaxStiffness)
+		: Constraints(TrimConstraints(InConstraints,
+			[&Particles, bTrimKinematicConstraints](int32 Index0, int32 Index1, int32 Index2)
+			{
+				return bTrimKinematicConstraints && Particles.InvM(Index0) == (FSolverReal)0. && Particles.InvM(Index1) == (FSolverReal)0. && Particles.InvM(Index2) == (FSolverReal)0.;
+			}))
+		, ParticleOffset(0)
+		, ParticleCount(Particles.GetRangeSize())
+		, Stiffness(
+			InStiffness,
+			StiffnessMultipliers,
+			TConstArrayView<TVec3<int32>>(Constraints),
+			ParticleOffset,
+			ParticleCount,
+			FPBDStiffness::DefaultTableSize,
+			FPBDStiffness::DefaultParameterFitBase,
+			MaxStiffness)
+	{
+		Init(Particles);
+	}
+
 	FPBDAxialSpringConstraintsBase(
 		const FSolverParticles& Particles,
 		int32 InParticleOffset,
@@ -50,7 +79,8 @@ public:
 	void ApplyProperties(const FSolverReal Dt, const int32 NumIterations) { Stiffness.ApplyPBDValues(Dt, NumIterations); }
 
 protected:
-	inline FSolverVec3 GetDelta(const FSolverParticles& Particles, const int32 ConstraintIndex, const FSolverReal ExpStiffnessValue) const
+	template<typename SolverParticlesOrRange>
+	inline FSolverVec3 GetDelta(const SolverParticlesOrRange& Particles, const int32 ConstraintIndex, const FSolverReal ExpStiffnessValue) const
 	{
 		const TVec3<int32>& Constraint = Constraints[ConstraintIndex];
 		const int32 i1 = Constraint[0];
@@ -79,7 +109,8 @@ protected:
 	}
 
 private:
-	FSolverReal FindBary(const FSolverParticles& Particles, const int32 i1, const int32 i2, const int32 i3)
+	template<typename SolverParticlesOrRange>
+	FSolverReal FindBary(const SolverParticlesOrRange& Particles, const int32 i1, const int32 i2, const int32 i3)
 	{
 		const FSolverVec3& P1 = Particles.X(i1);
 		const FSolverVec3& P2 = Particles.X(i2);
@@ -117,7 +148,8 @@ private:
 		return TrimmedConstraints.Array();
 	}
 
-	void Init(const FSolverParticles& Particles)
+	template<typename SolverParticlesOrRange>
+	void Init(const SolverParticlesOrRange& Particles)
 	{
 		Barys.Reset(Constraints.Num());
 		Dists.Reset(Constraints.Num());
diff --git a/Engine/Source/Runtime/Experimental/Chaos/Public/Chaos/PBDBendingConstraints.h b/Engine/Source/Runtime/Experimental/Chaos/Public/Chaos/PBDBendingConstraints.h
--- a/Engine/Source/Runtime/Experimental/Chaos/Public/Chaos/PBDBendingConstraints.h
+++ b/Engine/Source/Runtime/Experimental/Chaos/Public/Chaos/PBDBendingConstraints.h
@@ -15,6 +15,33 @@ public:
 	static bool IsEnabled(const FCollectionPropertyConstFacade& PropertyCollection)
 	{
 		return IsBendingElementStiffnessEnabled(PropertyCollection, false);
+	}	
+	
+	FPBDBendingConstraints(const FSolverParticlesRange& InParticles,
+		TArray<TVec4<int32>>&& InConstraints,
+		const TMap<FString, TConstArrayView<FRealSingle>>& WeightMaps,
+		const FCollectionPropertyConstFacade& PropertyCollection,
+		bool bTrimKinematicConstraints = false)
+		: Base(
+			InParticles,
+			MoveTemp(InConstraints),
+			WeightMaps.FindRef(GetBendingElementStiffnessString(PropertyCollection, BendingElementStiffnessName.ToString())),
+			WeightMaps.FindRef(GetBucklingStiffnessString(PropertyCollection, BucklingStiffnessName.ToString())),
+			GetRestAngleMapFromCollection(WeightMaps, PropertyCollection),
+			FSolverVec2(GetWeightedFloatBendingElementStiffness(PropertyCollection, 1.f)),
+			(FSolverReal)GetBucklingRatio(PropertyCollection, 0.f),  // BucklingRatio is clamped in base class
+			FSolverVec2(GetWeightedFloatBucklingStiffness(PropertyCollection, 1.f)),
+			GetRestAngleValueFromCollection(PropertyCollection),
+			(ERestAngleConstructionType)GetRestAngleType(PropertyCollection, (int32)ERestAngleConstructionType::Use3DRestAngles),
+			bTrimKinematicConstraints)
+		, BendingElementStiffnessIndex(PropertyCollection)
+		, BucklingRatioIndex(PropertyCollection)
+		, BucklingStiffnessIndex(PropertyCollection)
+		, FlatnessRatioIndex(PropertyCollection)
+		, RestAngleIndex(PropertyCollection)
+		, RestAngleTypeIndex(PropertyCollection)
+	{
+		InitColor(InParticles);
 	}
 
 	FPBDBendingConstraints(const FSolverParticles& InParticles,
@@ -132,13 +159,16 @@ PRAGMA_ENABLE_DEPRECATION_WARNINGS
 		SetProperties(PropertyCollection, TMap<FString, TConstArrayView<FRealSingle>>());
 	}
 
-	CHAOS_API void Apply(FSolverParticles& InParticles, const FSolverReal Dt) const;
+	template<typename SolverParticlesOrRange>
+	CHAOS_API void Apply(SolverParticlesOrRange& InParticles, const FSolverReal Dt) const;
 
 	const TArray<int32>& GetConstraintsPerColorStartIndex() const { return ConstraintsPerColorStartIndex; }
 
 private:
-	CHAOS_API void InitColor(const FSolverParticles& InParticles);
-	CHAOS_API void ApplyHelper(FSolverParticles& InParticles, const FSolverReal Dt, const int32 ConstraintIndex, const FSolverReal ExpStiffnessValue, const FSolverReal ExpBucklingValue) const;
+	template<typename SolverParticlesOrRange>
+	CHAOS_API void InitColor(const SolverParticlesOrRange& InParticles);
+	template<typename SolverParticlesOrRange>
+	void ApplyHelper(SolverParticlesOrRange& InParticles, const FSolverReal Dt, const int32 ConstraintIndex, const FSolverReal ExpStiffnessValue, const FSolverReal ExpBucklingValue) const;
 
 	TConstArrayView<FRealSingle> GetRestAngleMapFromCollection(
 		const TMap<FString, TConstArrayView<FRealSingle>>& WeightMaps,
diff --git a/Engine/Source/Runtime/Experimental/Chaos/Public/Chaos/PBDBendingConstraintsBase.h b/Engine/Source/Runtime/Experimental/Chaos/Public/Chaos/PBDBendingConstraintsBase.h
--- a/Engine/Source/Runtime/Experimental/Chaos/Public/Chaos/PBDBendingConstraintsBase.h
+++ b/Engine/Source/Runtime/Experimental/Chaos/Public/Chaos/PBDBendingConstraintsBase.h
@@ -3,6 +3,7 @@
 
 #include "Chaos/PBDSoftsEvolutionFwd.h"
 #include "Chaos/PBDSoftsSolverParticles.h"
+#include "Chaos/SoftsSolverParticlesRange.h"
 #include "Chaos/PBDStiffness.h"
 #include "Chaos/ParticleRule.h"
 #include "Containers/StaticArray.h"
@@ -20,6 +21,45 @@ public:
 		ExplicitRestAngles
 	};
 
+	FPBDBendingConstraintsBase(const FSolverParticlesRange& InParticles,
+		TArray<TVec4<int32>>&& InConstraints,
+		const TConstArrayView<FRealSingle>& StiffnessMultipliers,
+		const TConstArrayView<FRealSingle>& BucklingStiffnessMultipliers,
+		const TConstArrayView<FRealSingle>& RestAngleMap,
+		const FSolverVec2& InStiffness,
+		const FSolverReal InBucklingRatio,
+		const FSolverVec2& InBucklingStiffness,
+		const FSolverVec2& RestAngleValue,
+		ERestAngleConstructionType RestAngleConstructionType,
+		bool bTrimKinematicConstraints = false,
+		FSolverReal MaxStiffness = FPBDStiffness::DefaultPBDMaxStiffness)
+		: Constraints(bTrimKinematicConstraints ? TrimKinematicConstraints(InConstraints, InParticles) : MoveTemp(InConstraints))
+		, ConstraintSharedEdges(ExtractConstraintSharedEdges(Constraints))
+		, ParticleOffset(0)
+		, ParticleCount(InParticles.GetRangeSize())
+		, Stiffness(
+			InStiffness,
+			StiffnessMultipliers,
+			TConstArrayView<TVec2<int32>>(ConstraintSharedEdges),
+			ParticleOffset,
+			ParticleCount,
+			FPBDStiffness::DefaultTableSize,
+			FPBDStiffness::DefaultParameterFitBase,
+			MaxStiffness)
+		, BucklingRatio(FMath::Clamp(InBucklingRatio, (FSolverReal)0., (FSolverReal)1.))
+		, BucklingStiffness(
+			InBucklingStiffness,
+			BucklingStiffnessMultipliers,
+			TConstArrayView<TVec2<int32>>(ConstraintSharedEdges),
+			ParticleOffset,
+			ParticleCount,
+			FPBDStiffness::DefaultTableSize,
+			FPBDStiffness::DefaultParameterFitBase,
+			MaxStiffness)
+	{
+		CalculateRestAngles(InParticles, ParticleOffset, ParticleCount, RestAngleMap, RestAngleValue, RestAngleConstructionType);
+	}
+
 	FPBDBendingConstraintsBase(const FSolverParticles& InParticles,
 		int32 InParticleOffset,
 		int32 InParticleCount,
@@ -111,7 +151,8 @@ public:
 	UE_DEPRECATED(5.1, "Use SetProperties instead.")
 	void SetStiffness(FSolverReal InStiffness) { SetProperties(FSolverVec2(InStiffness), 0.f, 1.f); }
 
-	TStaticArray<FSolverVec3, 4> GetGradients(const FSolverParticles& InParticles, const int32 i) const
+	template<typename SolverParticlesOrRange>
+	TStaticArray<FSolverVec3, 4> GetGradients(const SolverParticlesOrRange& InParticles, const int32 i) const
 	{
 		const TVec4<int32>& Constraint = Constraints[i];
 		const FSolverVec3& P1 = InParticles.P(Constraint[0]);
@@ -122,7 +163,8 @@ public:
 		return CalcGradients(P1, P2, P3, P4);
 	}
 
-	FSolverReal GetScalingFactor(const FSolverParticles& InParticles, const int32 i, const TStaticArray<FSolverVec3, 4>& Grads, const FSolverReal ExpStiffnessValue, const FSolverReal ExpBucklingValue) const
+	template<typename SolverParticlesOrRange>
+	FSolverReal GetScalingFactor(const SolverParticlesOrRange& InParticles, const int32 i, const TStaticArray<FSolverVec3, 4>& Grads, const FSolverReal ExpStiffnessValue, const FSolverReal ExpBucklingValue) const
 	{
 		const TVec4<int32>& Constraint = Constraints[i];
 		const int32 i1 = Constraint[0];
@@ -177,7 +219,8 @@ public:
 		return UE_PI - FMath::Abs(Angle) < BucklingRatio * (UE_PI - FMath::Abs(RestAngle));
 	}
 
-	void Init(const FSolverParticles& InParticles)
+	template<typename SolverParticlesOrRange>
+	void Init(const SolverParticlesOrRange& InParticles)
 	{
 		IsBuckled.SetNumUninitialized(Constraints.Num());
 		for (int32 ConstraintIndex = 0; ConstraintIndex < Constraints.Num(); ++ConstraintIndex)
@@ -244,7 +287,8 @@ private:
 		return Grads;
 	}
 
-	static TArray<TVec4<int32>> TrimKinematicConstraints(const TArray<TVec4<int32>>& InConstraints, const FSolverParticles& InParticles)
+	template<typename SolverParticlesOrRange>
+	static TArray<TVec4<int32>> TrimKinematicConstraints(const TArray<TVec4<int32>>& InConstraints, const SolverParticlesOrRange& InParticles)
 	{
 		TArray<TVec4<int32>> TrimmedConstraints;
 		TrimmedConstraints.Reserve(InConstraints.Num());
@@ -270,7 +314,8 @@ private:
 		return ExtractedEdges;
 	}
 
-	CHAOS_API void CalculateRestAngles(const FSolverParticles& InParticles,
+	template<typename SolverParticlesOrRange>
+	CHAOS_API void CalculateRestAngles(const SolverParticlesOrRange& InParticles,
 		int32 InParticleOffset,
 		int32 InParticleCount,
 		const TConstArrayView<FRealSingle>& RestAngleMap,
diff --git a/Engine/Source/Runtime/Experimental/Chaos/Public/Chaos/PBDCollisionSpringConstraints.h b/Engine/Source/Runtime/Experimental/Chaos/Public/Chaos/PBDCollisionSpringConstraints.h
--- a/Engine/Source/Runtime/Experimental/Chaos/Public/Chaos/PBDCollisionSpringConstraints.h
+++ b/Engine/Source/Runtime/Experimental/Chaos/Public/Chaos/PBDCollisionSpringConstraints.h
@@ -87,7 +87,8 @@ public:
 		}
 	}
 
-	void Apply(FSolverParticles& Particles, const FSolverReal Dt, const int32 ConstraintIndex) const
+	template<typename SolverParticlesOrRange>
+	void Apply(SolverParticlesOrRange& Particles, const FSolverReal Dt, const int32 ConstraintIndex) const
 	{
 		const int32 i = ConstraintIndex;
 		const TVector<int32, 4>& Constraint = Constraints[i];
@@ -114,7 +115,8 @@ public:
 		}
 	}
 
-	void Apply(FSolverParticles& InParticles, const FSolverReal Dt) const
+	template<typename SolverParticlesOrRange>
+	void Apply(SolverParticlesOrRange& InParticles, const FSolverReal Dt) const
 	{
 		for (int32 i = 0; i < Constraints.Num(); ++i)
 		{
diff --git a/Engine/Source/Runtime/Experimental/Chaos/Public/Chaos/PBDCollisionSpringConstraintsBase.h b/Engine/Source/Runtime/Experimental/Chaos/Public/Chaos/PBDCollisionSpringConstraintsBase.h
--- a/Engine/Source/Runtime/Experimental/Chaos/Public/Chaos/PBDCollisionSpringConstraintsBase.h
+++ b/Engine/Source/Runtime/Experimental/Chaos/Public/Chaos/PBDCollisionSpringConstraintsBase.h
@@ -40,10 +40,11 @@ public:
 	UE_DEPRECATED(5.0, "Use Init(Particles, Spatial, GIAColors) instead.")
 	CHAOS_API void Init(const FSolverParticles& Particles);
 
-	template<typename SpatialAccelerator>
-	void Init(const FSolverParticles& Particles, const SpatialAccelerator& Spatial, const TConstArrayView<FPBDTriangleMeshCollisions::FGIAColor>& VertexGIAColors, const TArray<FPBDTriangleMeshCollisions::FGIAColor>& TriangleGIAColors);
+	template<typename SpatialAccelerator, typename SolverParticlesOrRange>
+	void Init(const SolverParticlesOrRange& Particles, const SpatialAccelerator& Spatial, const TConstArrayView<FPBDTriangleMeshCollisions::FGIAColor>& VertexGIAColors, const TArray<FPBDTriangleMeshCollisions::FGIAColor>& TriangleGIAColors);
 
-	CHAOS_API FSolverVec3 GetDelta(const FSolverParticles& InParticles, const int32 i) const;
+	template<typename SolverParticlesOrRange>
+	CHAOS_API FSolverVec3 GetDelta(const SolverParticlesOrRange& InParticles, const int32 i) const;
 
 	const TArray<TVec4<int32>>& GetConstraints() const { return Constraints;  }
 	const TArray<FSolverVec3>& GetBarys() const { return Barys; }
diff --git a/Engine/Source/Runtime/Experimental/Chaos/Public/Chaos/PBDEvolution.h b/Engine/Source/Runtime/Experimental/Chaos/Public/Chaos/PBDEvolution.h
--- a/Engine/Source/Runtime/Experimental/Chaos/Public/Chaos/PBDEvolution.h
+++ b/Engine/Source/Runtime/Experimental/Chaos/Public/Chaos/PBDEvolution.h
@@ -63,6 +63,10 @@ class FPBDEvolution : public TArrayCollection
 	const FSolverParticles& Particles() const { return MParticles; }
 	FSolverParticles& Particles() { return MParticles; }
 	const TPBDActiveView<FSolverParticles>& ParticlesActiveView() { return MParticlesActiveView; }
+	// These versions just help share code with Softs::FEvolution which follows UE naming standards.
+	const FSolverParticles& GetParticles() const { return MParticles; }
+	FSolverParticles& GetParticles() { return MParticles; }
+
 
 	const TArray<uint32>& ParticleGroupIds() const { return MParticleGroupIds; }
 
diff --git a/Engine/Source/Runtime/Experimental/Chaos/Public/Chaos/PBDLongRangeConstraints.h b/Engine/Source/Runtime/Experimental/Chaos/Public/Chaos/PBDLongRangeConstraints.h
--- a/Engine/Source/Runtime/Experimental/Chaos/Public/Chaos/PBDLongRangeConstraints.h
+++ b/Engine/Source/Runtime/Experimental/Chaos/Public/Chaos/PBDLongRangeConstraints.h
@@ -21,6 +21,25 @@ public:
 		return IsTetherStiffnessEnabled(PropertyCollection, false);
 	}
 
+	FPBDLongRangeConstraints(
+		const FSolverParticlesRange& Particles,
+		const TArray<TConstArrayView<TTuple<int32, int32, FRealSingle>>>& InTethers,
+		const TMap<FString, TConstArrayView<FRealSingle>>& WeightMaps,
+		const FCollectionPropertyConstFacade& PropertyCollection,
+		FSolverReal MeshScale)
+		: FPBDLongRangeConstraintsBase(
+			Particles,
+			InTethers,
+			WeightMaps.FindRef(GetTetherStiffnessString(PropertyCollection, TetherStiffnessName.ToString())),
+			WeightMaps.FindRef(GetTetherScaleString(PropertyCollection, TetherScaleName.ToString())),
+			FSolverVec2(GetWeightedFloatTetherStiffness(PropertyCollection, 1.f)),
+			FSolverVec2(GetWeightedFloatTetherScale(PropertyCollection, 1.f)),  // Scale clamping done in constructor
+			FPBDStiffness::DefaultPBDMaxStiffness,
+			MeshScale)
+		, TetherStiffnessIndex(PropertyCollection)
+		, TetherScaleIndex(PropertyCollection)
+	{}
+
 	FPBDLongRangeConstraints(
 		const FSolverParticles& Particles,
 		const int32 InParticleOffset,
@@ -108,7 +127,8 @@ public:
 		SetProperties(PropertyCollection, TMap<FString, TConstArrayView<FRealSingle>>(), (FSolverReal)1.);
 	}
 
-	CHAOS_API void Apply(FSolverParticles& Particles, const FSolverReal Dt) const;
+	template<typename SolverParticlesOrRange>
+	CHAOS_API void Apply(SolverParticlesOrRange& Particles, const FSolverReal Dt) const;
 
 private:
 	using Base::MinTetherScale;
diff --git a/Engine/Source/Runtime/Experimental/Chaos/Public/Chaos/PBDLongRangeConstraintsBase.h b/Engine/Source/Runtime/Experimental/Chaos/Public/Chaos/PBDLongRangeConstraintsBase.h
--- a/Engine/Source/Runtime/Experimental/Chaos/Public/Chaos/PBDLongRangeConstraintsBase.h
+++ b/Engine/Source/Runtime/Experimental/Chaos/Public/Chaos/PBDLongRangeConstraintsBase.h
@@ -6,6 +6,7 @@
 #include "Chaos/PBDSoftsEvolutionFwd.h"
 #include "Chaos/PBDActiveView.h"
 #include "Chaos/PBDStiffness.h"
+#include "Chaos/SoftsSolverParticlesRange.h"
 #include "Containers/Map.h"
 #include "Containers/Set.h"
 
@@ -33,6 +34,16 @@ public:
 
 	typedef TTuple<int32, int32, FRealSingle> FTether;
 
+	CHAOS_API FPBDLongRangeConstraintsBase(
+		const FSolverParticlesRange& Particles,
+		const TArray<TConstArrayView<TTuple<int32, int32, FRealSingle>>>& InTethers,
+		const TConstArrayView<FRealSingle>& StiffnessMultipliers,
+		const TConstArrayView<FRealSingle>& ScaleMultipliers,
+		const FSolverVec2& InStiffness = FSolverVec2::UnitVector,
+		const FSolverVec2& InScale = FSolverVec2::UnitVector,
+		FSolverReal MaxStiffness = FPBDStiffness::DefaultPBDMaxStiffness,
+		FSolverReal MeshScale = (FSolverReal)1.);
+
 	CHAOS_API FPBDLongRangeConstraintsBase(
 		const FSolverParticles& Particles,
 		const int32 InParticleOffset,
@@ -111,7 +122,8 @@ protected:
 	void ApplyScale() { TetherScale.ApplyValues(); }
 
 	// Return a vector representing the amount of segment required for the tether to shrink back to its maximum target length constraint, or zero if the constraint is already met
-	inline FSolverVec3 GetDelta(const FSolverParticles& Particles, const FTether& Tether, const FSolverReal InScale) const
+	template<typename SolverParticlesOrRange>
+	inline FSolverVec3 GetDelta(const SolverParticlesOrRange& Particles, const FTether& Tether, const FSolverReal InScale) const
 	{
 		const int32 Start = GetStartParticle(Tether);
 		const int32 End = GetEndParticle(Tether);
@@ -125,7 +137,8 @@ protected:
 	};
 
 	// Return a direction and length representing the amount of segment required for the tether to shrink back to its maximum target length constraint, or zero if the constraint is already met
-	inline void GetDelta(const FSolverParticles& Particles, const FTether& Tether, const FSolverReal InScale, FSolverVec3& OutDirection, FSolverReal& OutOffset) const
+	template<typename SolverParticlesOrRange>
+	inline void GetDelta(const SolverParticlesOrRange& Particles, const FTether& Tether, const FSolverReal InScale, FSolverVec3& OutDirection, FSolverReal& OutOffset) const
 	{
 		const int32 Start = GetStartParticle(Tether);
 		const int32 End = GetEndParticle(Tether);
diff --git a/Engine/Source/Runtime/Experimental/Chaos/Public/Chaos/PBDSelfCollisionSphereConstraints.h b/Engine/Source/Runtime/Experimental/Chaos/Public/Chaos/PBDSelfCollisionSphereConstraints.h
--- a/Engine/Source/Runtime/Experimental/Chaos/Public/Chaos/PBDSelfCollisionSphereConstraints.h
+++ b/Engine/Source/Runtime/Experimental/Chaos/Public/Chaos/PBDSelfCollisionSphereConstraints.h
@@ -6,6 +6,7 @@
 #include "Chaos/CollectionPropertyFacade.h"
 #include "Chaos/PBDSoftsEvolutionFwd.h"
 #include "Chaos/PBDSoftsSolverParticles.h"
+#include "Chaos/SoftsSolverParticlesRange.h"
 #include "Containers/Set.h"
 
 namespace Chaos::Softs
@@ -23,9 +24,11 @@ namespace Chaos::Softs
 
 		virtual ~FPBDSelfCollisionSphereConstraintsBase() {}
 
-		CHAOS_API void Init(const FSolverParticles& Particles);
+		template<typename SolverParticlesOrRange>
+		CHAOS_API void Init(const SolverParticlesOrRange& Particles);
 
-		CHAOS_API void Apply(FSolverParticles& InParticles, const FSolverReal Dt) const;
+		template<typename SolverParticlesOrRange>
+		CHAOS_API void Apply(SolverParticlesOrRange& InParticles, const FSolverReal Dt) const;
 
 		const TArray<TVec2<int32>>& GetConstraints() const { return Constraints; }
 		const TSet<int32>* GetVertexSet() const { return VertexSetNoOffset; }
diff --git a/Engine/Source/Runtime/Experimental/Chaos/Public/Chaos/PBDSoftBodyCollisionConstraint.h b/Engine/Source/Runtime/Experimental/Chaos/Public/Chaos/PBDSoftBodyCollisionConstraint.h
new file mode 100644
--- /dev/null
+++ b/Engine/Source/Runtime/Experimental/Chaos/Public/Chaos/PBDSoftBodyCollisionConstraint.h
@@ -0,0 +1,115 @@
+// Copyright Epic Games, Inc. All Rights Reserved.
+#pragma once
+
+#include "Chaos/Core.h"
+#include "Chaos/PBDSoftsEvolutionFwd.h"
+#include "Chaos/PBDSoftsSolverParticles.h"
+#include "Chaos/SoftsSolverCollisionParticles.h"
+#include "Chaos/SoftsSolverCollisionParticlesRange.h"
+#include "Chaos/SoftsSolverParticlesRange.h"
+#include "Chaos/CollectionPropertyFacade.h"
+
+#if !defined(CHAOS_PER_PARTICLE_COLLISION_ISPC_ENABLED_DEFAULT)
+#define CHAOS_PER_PARTICLE_COLLISION_ISPC_ENABLED_DEFAULT 1
+#endif
+
+// Support run-time toggling on supported platforms in non-shipping configurations
+#if !INTEL_ISPC || UE_BUILD_SHIPPING
+const bool bChaos_SoftBodyCollision_ISPC_Enabled = INTEL_ISPC && CHAOS_PER_PARTICLE_COLLISION_ISPC_ENABLED_DEFAULT;
+#else
+extern CHAOS_API bool bChaos_SoftBodyCollision_ISPC_Enabled;
+#endif
+
+namespace Chaos::Softs
+{
+
+class FPBDSoftBodyCollisionConstraintBase
+{
+public:
+	FPBDSoftBodyCollisionConstraintBase(
+		const TArray<FSolverRigidTransform3>& InLastCollisionTransforms,
+		FSolverReal InCollisionThickness,
+		FSolverReal InFrictionCoefficient,
+		bool bInUseCCD,
+		TArray<bool>* InCollisionParticleCollided = nullptr,
+		TArray<FSolverVec3>* InContacts = nullptr,
+		TArray<FSolverVec3>* InNormals = nullptr,
+		TArray<FSolverReal>* InPhis = nullptr
+		)
+		: LastCollisionTransforms(InLastCollisionTransforms)
+		, CollisionThickness(InCollisionThickness)
+		, FrictionCoefficient(InFrictionCoefficient)
+		, bUseCCD(bInUseCCD)
+		, CollisionParticleCollided(InCollisionParticleCollided)
+		, Contacts(InContacts)
+		, Normals(InNormals)
+		, Phis(InPhis)
+	{}
+
+	void SetWriteDebugContacts(bool bWrite) { bWriteDebugContacts = bWrite; }
+
+	CHAOS_API void Apply(FSolverParticlesRange& Particles, const FSolverReal Dt, const TArray<FSolverCollisionParticlesRange>& CollisionParticles) const;
+
+private:
+	template<bool bLockAndWriteContacts, bool bWithFriction>
+	void ApplyInternal(FSolverParticlesRange& Particles, const FSolverReal Dt, const TArray<FSolverCollisionParticlesRange>& CollisionParticles) const;
+	template<bool bLockAndWriteContacts, bool bWithFriction>
+	void ApplyInternalCCD(FSolverParticlesRange& Particles, const FSolverReal Dt, const TArray<FSolverCollisionParticlesRange>& CollisionParticles) const;
+
+	template<bool bWithFriction>
+	void ApplyInternalISPC(FSolverParticlesRange& Particles, const FSolverReal Dt, const TArray<FSolverCollisionParticlesRange>& CollisionParticles) const;
+
+protected:
+	const TArray<FSolverRigidTransform3>& LastCollisionTransforms; // Used by CCD
+	FSolverReal CollisionThickness;
+	FSolverReal FrictionCoefficient;
+	bool bUseCCD;
+
+	/**  Used for writing debug contacts */
+	bool bWriteDebugContacts = false;
+	TArray<bool>* const CollisionParticleCollided; // Per collision particle
+	// List of contact data
+	TArray<FSolverVec3>* const Contacts;
+	TArray<FSolverVec3>* const Normals;
+	TArray<FSolverReal>* const Phis;
+	mutable FCriticalSection DebugMutex;
+};
+
+
+class FPBDSoftBodyCollisionConstraint : public FPBDSoftBodyCollisionConstraintBase
+{
+	typedef FPBDSoftBodyCollisionConstraintBase Base;
+
+public:
+	static constexpr FSolverReal DefaultCollisionThickness = (FSolverReal)1.;
+	static constexpr FSolverReal DefaultFrictionCoefficient = (FSolverReal)0.8;
+
+	FPBDSoftBodyCollisionConstraint(
+		const TArray<FSolverRigidTransform3>& InLastCollisionTransforms,
+		const FCollectionPropertyConstFacade& PropertyCollection,
+		FSolverReal InMeshScale,
+		TArray<bool>* InCollisionParticleCollided = nullptr,
+		TArray<FSolverVec3>* InContacts = nullptr,
+		TArray<FSolverVec3>* InNormals = nullptr,
+		TArray<FSolverReal>* InPhis = nullptr)
+		: Base(InLastCollisionTransforms,
+			InMeshScale * GetCollisionThickness(PropertyCollection, DefaultCollisionThickness),
+			GetFrictionCoefficient(PropertyCollection, DefaultFrictionCoefficient),
+			GetUseCCD(PropertyCollection, false),
+			InCollisionParticleCollided, InContacts, InNormals, InPhis)
+		, MeshScale(InMeshScale)
+		, CollisionThicknessIndex(PropertyCollection)
+		, FrictionCoefficientIndex(PropertyCollection)
+		, UseCCDIndex(PropertyCollection)
+	{}
+
+	CHAOS_API void SetProperties(const FCollectionPropertyConstFacade& PropertyCollection);
+
+private:
+	const FSolverReal MeshScale;
+	UE_CHAOS_DECLARE_PROPERTYCOLLECTION_NAME(CollisionThickness, float);
+	UE_CHAOS_DECLARE_PROPERTYCOLLECTION_NAME(FrictionCoefficient, float);
+	UE_CHAOS_DECLARE_PROPERTYCOLLECTION_NAME(UseCCD, bool);
+};
+
+}  // End namespace Chaos::Softs
diff --git a/Engine/Source/Runtime/Experimental/Chaos/Public/Chaos/PBDSoftsEvolutionFwd.h b/Engine/Source/Runtime/Experimental/Chaos/Public/Chaos/PBDSoftsEvolutionFwd.h
--- a/Engine/Source/Runtime/Experimental/Chaos/Public/Chaos/PBDSoftsEvolutionFwd.h
+++ b/Engine/Source/Runtime/Experimental/Chaos/Public/Chaos/PBDSoftsEvolutionFwd.h
@@ -8,6 +8,8 @@ namespace UE::Math
 	template<typename T> struct TTransform;
 }
 
+class FPerSolverFieldSystem;
+
 namespace Chaos
 {
 	// Chaos types
@@ -39,10 +41,13 @@ namespace Chaos
 		using FSolverTransform3 = UE::Math::TTransform<FSolverReal>;
 		using FSolverAABB3 = TAABB<FSolverReal, 3>;
 		class FSolverParticles;
+		class FSolverParticlesRange;
+		class FSolverCollisionParticlesRange;
 		struct FPAndInvM;
 
 		// Softs solver class
 		class FPBDEvolution;
+		class FEvolution;
 
 		// Softs solver constraint classes
 		class FPBDSpringConstraints;
@@ -70,9 +75,11 @@ namespace Chaos
 		class FPBDTriangleMeshIntersections;
 		class FPBDTriangleMeshCollisions;
 		class FPBDSelfCollisionSphereConstraints;
+		class FPBDSoftBodyCollisionConstraint;
 
 		// Softs solver forces
 		class FVelocityAndPressureField;
 		using FVelocityField UE_DEPRECATED(5.1, "Chaos::Softs::FVelocityField has been renamed FVelocityAndPressureField to match its new behavior.") = FVelocityAndPressureField;
+		class FExternalForces;
 	}  // End namespace Softs
 }  // End namespace Chaos
diff --git a/Engine/Source/Runtime/Experimental/Chaos/Public/Chaos/PBDSoftsSolverParticles.h b/Engine/Source/Runtime/Experimental/Chaos/Public/Chaos/PBDSoftsSolverParticles.h
--- a/Engine/Source/Runtime/Experimental/Chaos/Public/Chaos/PBDSoftsSolverParticles.h
+++ b/Engine/Source/Runtime/Experimental/Chaos/Public/Chaos/PBDSoftsSolverParticles.h
@@ -35,16 +35,31 @@ public:
 
 	const FSolverVec3& P(const int32 index) const { return MPAndInvM[index].P; }
 	FSolverVec3& P(const int32 index) { return MPAndInvM[index].P; }
+	// When PBD is solving, V = last time step's V and we use the latest P to estimate velocity, (VelocityEstimate = (P - X)/Dt.
+	// After PBD (e.g., during implicit timestepping CG solve), V = current velocity estimate, VPrev = last time step's velocity
+	// Warning: VPrev and Acceleration use the same arrays as Accleration is only used by PBD and VPrev is only used by the linear solver
+	const FSolverVec3& VPrev(const int32 Index) const { return Acceleration(Index); }
+	FSolverVec3& VPrev(const int32 Index) { return Acceleration(Index); }
 
 	const FPAndInvM& PAndInvM(const int32 index) const { return MPAndInvM[index]; }
 	FPAndInvM& PAndInvM(const int32 index) { return MPAndInvM[index]; }
 
 	TArrayCollectionArray<FPAndInvM>& GetPAndInvM() { return MPAndInvM; }
 	const TArrayCollectionArray<FPAndInvM>& GetPAndInvM() const { return MPAndInvM; }
+	// Warning: VPrev and Acceleration use the same arrays as Accleration is only used by PBD and VPrev is only used by the linear solver
+	TArrayCollectionArray<FSolverVec3>& GetVPrev() { return GetAcceleration(); }
+	const TArrayCollectionArray<FSolverVec3>& GetVPrev() const { return  GetAcceleration(); }
+
+	// Convenience method to match interface of FSolverParticlesRange
+	template<typename T>
+	TConstArrayView<T> GetConstArrayView(const TArray<T>& Data) const { return TConstArrayView<T>(Data); }
+	template<typename T>
+	TArrayView<T> GetArrayView(const TArray<T>& Data) { return TArrayView<T>(Data); }
 
 private:
 	// Note that InvM is copied from TDynamicParticles to here during ApplyPreSimulationTransforms (when P is updated)
 	TArrayCollectionArray<FPAndInvM> MPAndInvM;
+
 };
 
 }
diff --git a/Engine/Source/Runtime/Experimental/Chaos/Public/Chaos/PBDSphericalConstraint.h b/Engine/Source/Runtime/Experimental/Chaos/Public/Chaos/PBDSphericalConstraint.h
--- a/Engine/Source/Runtime/Experimental/Chaos/Public/Chaos/PBDSphericalConstraint.h
+++ b/Engine/Source/Runtime/Experimental/Chaos/Public/Chaos/PBDSphericalConstraint.h
@@ -4,6 +4,7 @@
 #include "Chaos/PBDSoftsEvolutionFwd.h"
 #include "Chaos/PBDSoftsSolverParticles.h"
 #include "Chaos/CollectionPropertyFacade.h"
+#include "Chaos/SoftsSolverParticlesRange.h"
 #include "ChaosStats.h"
 #include "Containers/ArrayView.h"
 #include "Containers/ContainersFwd.h"
@@ -101,7 +102,9 @@ public:
 		SetProperties(PropertyCollection, TMap<FString, TConstArrayView<FRealSingle>>(), MeshScale);
 	}
 
-	void Apply(FSolverParticles& Particles, const FSolverReal Dt) const
+
+	template<typename SolverParticlesOrRange>
+	void Apply(SolverParticlesOrRange& Particles, const FSolverReal Dt) const
 	{
 		SCOPE_CYCLE_COUNTER(STAT_PBD_Spherical);
 
@@ -142,10 +145,12 @@ public:
 	}
 
 private:
-	template<bool bHasMaxDistance>
-	void ApplyHelper(FSolverParticles& Particles, const FSolverReal Dt) const
+
+	template<bool bHasMaxDistance, typename SolverParticlesOrRange>
+	void ApplyHelper(SolverParticlesOrRange& Particles, const FSolverReal Dt) const
 	{
-		PhysicsParallelFor(ParticleCount, [this, &Particles, Dt](int32 Index)  // TODO: profile need for parallel loop based on particle count
+		const TConstArrayView<FSolverVec3> AnimationPositionsView = Particles.GetConstArrayView(AnimationPositions);
+		PhysicsParallelFor(ParticleCount, [this, &Particles, Dt, &AnimationPositionsView](int32 Index)  // TODO: profile need for parallel loop based on particle count
 		{
 			const int32 ParticleIndex = ParticleOffset + Index;
 
@@ -155,7 +160,7 @@ private:
 			}
 
 			const FSolverReal Radius = (bHasMaxDistance ? MaxDistanceBase + MaxDistanceRange * SphereRadii[Index] : MaxDistanceBase) * Scale;
-			const FSolverVec3& Center = AnimationPositions[ParticleIndex];
+			const FSolverVec3& Center = AnimationPositionsView[ParticleIndex];
 
 			const FSolverVec3 CenterToParticle = Particles.P(ParticleIndex) - Center;
 			const FSolverReal DistanceSquared = CenterToParticle.SizeSquared();
@@ -170,7 +175,8 @@ private:
 		});
 	}
 
-	CHAOS_API void ApplyHelperISPC(FSolverParticles& Particles, const FSolverReal Dt) const;
+	template<typename SolverParticlesOrRange>
+	CHAOS_API void ApplyHelperISPC(SolverParticlesOrRange& Particles, const FSolverReal Dt) const;
 
 protected:
 	const TArray<FSolverVec3>& AnimationPositions;  // Use global indexation (will need adding ParticleOffset)
@@ -292,7 +298,8 @@ public:
 	void SetEnabled(bool bInEnabled) { bEnabled = bInEnabled; }
 	bool IsEnabled() const { return bEnabled; }
 
-	void Apply(FSolverParticles& Particles, const FSolverReal Dt) const
+	template<typename SolverParticlesOrRange>
+	void Apply(SolverParticlesOrRange& Particles, const FSolverReal Dt) const
 	{
 		SCOPE_CYCLE_COUNTER(STAT_PBD_SphericalBackstop);
 
@@ -401,10 +408,13 @@ PRAGMA_ENABLE_DEPRECATION_WARNINGS
 	}
 
 private:
-	template<bool bHasBackstopDistance, bool bHasBackstopRadius>
-	void ApplyHelper(FSolverParticles& Particles, const FSolverReal Dt) const
+
+	template<bool bHasBackstopDistance, bool bHasBackstopRadius, typename SolverParticlesOrRange>
+	void ApplyHelper(SolverParticlesOrRange& Particles, const FSolverReal Dt) const
 	{
-		PhysicsParallelFor(ParticleCount, [this, &Particles, Dt](int32 Index)  // TODO: profile need for parallel loop based on particle count
+		const TConstArrayView<FSolverVec3> AnimationPositionsView = Particles.GetConstArrayView(AnimationPositions);
+		const TConstArrayView<FSolverVec3> AnimationNormalsView = Particles.GetConstArrayView(AnimationNormals);
+		PhysicsParallelFor(ParticleCount, [this, &Particles, Dt, &AnimationPositionsView, &AnimationNormalsView](int32 Index)  // TODO: profile need for parallel loop based on particle count
 		{
 			const int32 ParticleIndex = ParticleOffset + Index;
 
@@ -413,8 +423,8 @@ private:
 				return;
 			}
 
-			const FSolverVec3& AnimationPosition = AnimationPositions[ParticleIndex];
-			const FSolverVec3& AnimationNormal = AnimationNormals[ParticleIndex];
+			const FSolverVec3& AnimationPosition = AnimationPositionsView[ParticleIndex];
+			const FSolverVec3& AnimationNormal = AnimationNormalsView[ParticleIndex];
 
 			const FSolverReal SphereOffsetDistance = (bHasBackstopDistance ? BackstopDistanceBase + BackstopDistanceRange * SphereOffsetDistances[Index] : BackstopDistanceBase) * Scale;
 			const FSolverReal Radius = (bHasBackstopRadius ? BackstopRadiusBase + BackstopRadiusRange * SphereRadii[Index] : BackstopRadiusBase) * Scale;
@@ -437,10 +447,12 @@ private:
 		});
 	}
 
-	template<bool bHasBackstopDistance, bool bHasBackstopRadius>
-	void ApplyLegacyHelper(FSolverParticles& Particles, const FSolverReal Dt) const
+	template<bool bHasBackstopDistance, bool bHasBackstopRadius, typename SolverParticlesOrRange>
+	void ApplyLegacyHelper(SolverParticlesOrRange& Particles, const FSolverReal Dt) const
 	{
-		PhysicsParallelFor(ParticleCount, [this, &Particles, Dt](int32 Index)  // TODO: profile need for parallel loop based on particle count
+		const TConstArrayView<FSolverVec3> AnimationPositionsView = Particles.GetConstArrayView(AnimationPositions);
+		const TConstArrayView<FSolverVec3> AnimationNormalsView = Particles.GetConstArrayView(AnimationNormals);
+		PhysicsParallelFor(ParticleCount, [this, &Particles, Dt, &AnimationPositionsView, &AnimationNormalsView](int32 Index)  // TODO: profile need for parallel loop based on particle count
 		{
 			const int32 ParticleIndex = ParticleOffset + Index;
 
@@ -449,8 +461,8 @@ private:
 				return;
 			}
 
-			const FSolverVec3& AnimationPosition = AnimationPositions[ParticleIndex];
-			const FSolverVec3& AnimationNormal = AnimationNormals[ParticleIndex];
+			const FSolverVec3& AnimationPosition = AnimationPositionsView[ParticleIndex];
+			const FSolverVec3& AnimationNormal = AnimationNormalsView[ParticleIndex];
 
 			const FSolverReal SphereOffsetDistance = (bHasBackstopDistance ? BackstopDistanceBase + BackstopDistanceRange * SphereOffsetDistances[Index] : BackstopDistanceBase) * Scale;
 			const FSolverReal Radius = (bHasBackstopRadius ? BackstopRadiusBase + BackstopRadiusRange * SphereRadii[Index] : BackstopRadiusBase) * Scale;
@@ -473,8 +485,10 @@ private:
 		});
 	}
 
-	CHAOS_API void ApplyLegacyHelperISPC(FSolverParticles& Particles, const FSolverReal Dt) const;
-	CHAOS_API void ApplyHelperISPC(FSolverParticles& Particles, const FSolverReal Dt) const;
+	template<typename SolverParticlesOrRange>
+	CHAOS_API void ApplyLegacyHelperISPC(SolverParticlesOrRange& Particles, const FSolverReal Dt) const;
+	template<typename SolverParticlesOrRange>
+	CHAOS_API void ApplyHelperISPC(SolverParticlesOrRange& Particles, const FSolverReal Dt) const;
 
 private:
 	const TArray<FSolverVec3>& AnimationPositions;  // Positions of spheres, use global indexation (will need adding ParticleOffset)
diff --git a/Engine/Source/Runtime/Experimental/Chaos/Public/Chaos/PBDSpringConstraints.h b/Engine/Source/Runtime/Experimental/Chaos/Public/Chaos/PBDSpringConstraints.h
--- a/Engine/Source/Runtime/Experimental/Chaos/Public/Chaos/PBDSpringConstraints.h
+++ b/Engine/Source/Runtime/Experimental/Chaos/Public/Chaos/PBDSpringConstraints.h
@@ -10,6 +10,24 @@ namespace Chaos::Softs
 class FPBDSpringConstraints : public FPBDSpringConstraintsBase
 {
 public:
+	template<int32 Valence>
+	FPBDSpringConstraints(
+		const FSolverParticlesRange& Particles,
+		const TArray<TVector<int32, Valence>>& InConstraints,
+		const TConstArrayView<FRealSingle>& StiffnessMultipliers,
+		const FSolverVec2& InStiffness,
+		bool bTrimKinematicConstraints = false,
+		typename TEnableIf<Valence >= 2 && Valence <= 4>::Type* = nullptr)
+		: Base(
+			Particles,
+			InConstraints,
+			StiffnessMultipliers,
+			InStiffness,
+			bTrimKinematicConstraints)
+	{
+		InitColor(Particles);
+	}
+
 	template<int32 Valence>
 	FPBDSpringConstraints(
 		const FSolverParticles& Particles,
@@ -34,7 +52,8 @@ public:
 
 	virtual ~FPBDSpringConstraints() override {}
 
-	CHAOS_API void Apply(FSolverParticles& Particles, const FSolverReal Dt) const;
+	template<typename SolverParticlesOrRange>
+	CHAOS_API void Apply(SolverParticlesOrRange& Particles, const FSolverReal Dt) const;
 
 	const TArray<int32>& GetConstraintsPerColorStartIndex() const { return ConstraintsPerColorStartIndex; }
 
@@ -46,8 +65,10 @@ protected:
 	using Base::ParticleCount;
 
 private:
-	CHAOS_API void InitColor(const FSolverParticles& InParticles);
-	CHAOS_API void ApplyHelper(FSolverParticles& Particles, const FSolverReal Dt, const int32 ConstraintIndex, const FSolverReal ExpStiffnessValue) const;
+	template<typename SolverParticlesOrRange>
+	CHAOS_API void InitColor(const SolverParticlesOrRange& InParticles);
+	template<typename SolverParticlesOrRange>
+	void ApplyHelper(SolverParticlesOrRange& Particles, const FSolverReal Dt, const int32 ConstraintIndex, const FSolverReal ExpStiffnessValue) const;
 
 private:
 	TArray<int32> ConstraintsPerColorStartIndex; // Constraints are ordered so each batch is contiguous. This is ColorNum + 1 length so it can be used as start and end.
@@ -61,6 +82,21 @@ public:
 		return IsEdgeSpringStiffnessEnabled(PropertyCollection, false);
 	}
 
+	FPBDEdgeSpringConstraints(
+		const FSolverParticlesRange& Particles,
+		const TArray<TVec3<int32>>& InConstraints,
+		const TMap<FString, TConstArrayView<FRealSingle>>& WeightMaps,
+		const FCollectionPropertyConstFacade& PropertyCollection,
+		bool bTrimKinematicConstraints = false)
+		: FPBDSpringConstraints(
+			Particles,
+			InConstraints,
+			WeightMaps.FindRef(GetEdgeSpringStiffnessString(PropertyCollection, EdgeSpringStiffnessName.ToString())),
+			FSolverVec2(GetWeightedFloatEdgeSpringStiffness(PropertyCollection, 1.f)),
+			bTrimKinematicConstraints)
+		, EdgeSpringStiffnessIndex(PropertyCollection)
+	{}
+
 	FPBDEdgeSpringConstraints(
 		const FSolverParticles& Particles,
 		int32 ParticleOffset,
@@ -126,6 +162,21 @@ public:
 		return IsBendingSpringStiffnessEnabled(PropertyCollection, false);
 	}
 
+	FPBDBendingSpringConstraints(
+		const FSolverParticlesRange& Particles,
+		const TArray<TVec2<int32>>& InConstraints,
+		const TMap<FString, TConstArrayView<FRealSingle>>& WeightMaps,
+		const FCollectionPropertyConstFacade& PropertyCollection,
+		bool bTrimKinematicConstraints = false)
+		: FPBDSpringConstraints(
+			Particles,
+			InConstraints,
+			WeightMaps.FindRef(GetBendingSpringStiffnessString(PropertyCollection, BendingSpringStiffnessName.ToString())),
+			FSolverVec2(GetWeightedFloatBendingSpringStiffness(PropertyCollection, 1.f)),
+			bTrimKinematicConstraints)
+		, BendingSpringStiffnessIndex(PropertyCollection)
+	{}
+
 	FPBDBendingSpringConstraints(
 		const FSolverParticles& Particles,
 		int32 ParticleOffset,
diff --git a/Engine/Source/Runtime/Experimental/Chaos/Public/Chaos/PBDSpringConstraintsBase.h b/Engine/Source/Runtime/Experimental/Chaos/Public/Chaos/PBDSpringConstraintsBase.h
--- a/Engine/Source/Runtime/Experimental/Chaos/Public/Chaos/PBDSpringConstraintsBase.h
+++ b/Engine/Source/Runtime/Experimental/Chaos/Public/Chaos/PBDSpringConstraintsBase.h
@@ -5,6 +5,7 @@
 #include "Chaos/PBDSoftsEvolutionFwd.h"
 #include "Chaos/Utilities.h"
 #include "Chaos/PBDSoftsSolverParticles.h"
+#include "Chaos/SoftsSolverParticlesRange.h"
 #include "Chaos/PBDStiffness.h"
 #include "Containers/Array.h"
 #include "Templates/EnableIf.h"
@@ -15,6 +16,41 @@ namespace Chaos::Softs
 class FPBDSpringConstraintsBase
 {
 public:
+	template<int32 Valence, TEMPLATE_REQUIRES(Valence >= 2 && Valence <= 4)>
+	FPBDSpringConstraintsBase(
+		const FSolverParticlesRange& Particles,
+		const TArray<TVector<int32, Valence>>& InConstraints,
+		const TConstArrayView<FRealSingle>& StiffnessMultipliers,
+		const FSolverVec2& InStiffness,
+		bool bTrimKinematicConstraints = false,
+		FSolverReal MaxStiffness = FPBDStiffness::DefaultPBDMaxStiffness)
+		: Constraints(TrimConstraints(InConstraints,
+			[&Particles, bTrimKinematicConstraints](int32 Index0, int32 Index1)
+			{
+				return bTrimKinematicConstraints && Particles.InvM(Index0) == (FSolverReal)0. && Particles.InvM(Index1) == (FSolverReal)0.;
+			}))
+		, ParticleOffset(0)
+		, ParticleCount(Particles.GetRangeSize())
+		, Stiffness(
+			InStiffness,
+			StiffnessMultipliers,
+			TConstArrayView<TVec2<int32>>(Constraints),
+			ParticleOffset,
+			ParticleCount,
+			FPBDStiffness::DefaultTableSize,
+			FPBDStiffness::DefaultParameterFitBase,
+			MaxStiffness)
+	{
+		// Update distances
+		Dists.Reset(Constraints.Num());
+		for (const TVec2<int32>& Constraint : Constraints)
+		{
+			const FSolverVec3& P0 = Particles.X(Constraint[0]);
+			const FSolverVec3& P1 = Particles.X(Constraint[1]);
+			Dists.Add((P1 - P0).Size());
+		}
+	}
+
 	template<int32 Valence, TEMPLATE_REQUIRES(Valence >= 2 && Valence <= 4)>
 	FPBDSpringConstraintsBase(
 		const FSolverParticles& Particles,
@@ -64,7 +100,8 @@ public:
 	const TArray<TVec2<int32>>& GetConstraints() const { return Constraints; }
 
 protected:
-	FSolverVec3 GetDelta(const FSolverParticles& Particles, const int32 ConstraintIndex, const FSolverReal ExpStiffnessValue) const
+	template<typename SolverParticlesOrRange>
+	FSolverVec3 GetDelta(const SolverParticlesOrRange& Particles, const int32 ConstraintIndex, const FSolverReal ExpStiffnessValue) const
 	{
 		const auto& Constraint = Constraints[ConstraintIndex];
 		const int32 i1 = Constraint[0];
diff --git a/Engine/Source/Runtime/Experimental/Chaos/Public/Chaos/PBDTriangleMeshCollisions.h b/Engine/Source/Runtime/Experimental/Chaos/Public/Chaos/PBDTriangleMeshCollisions.h
--- a/Engine/Source/Runtime/Experimental/Chaos/Public/Chaos/PBDTriangleMeshCollisions.h
+++ b/Engine/Source/Runtime/Experimental/Chaos/Public/Chaos/PBDTriangleMeshCollisions.h
@@ -190,9 +190,11 @@ public:
 		}
 	}
 
-	CHAOS_API void Init(const FSolverParticles& Particles, const FSolverReal MinProximityQueryRadius = (FSolverReal)0.);
+	template<typename SolverParticlesOrRange>
+	CHAOS_API void Init(const SolverParticlesOrRange& Particles, const FSolverReal MinProximityQueryRadius = (FSolverReal)0.);
 
-	CHAOS_API void PostStepInit(const FSolverParticles& Particles);
+	template<typename SolverParticlesOrRange>
+	CHAOS_API void PostStepInit(const SolverParticlesOrRange& Particles);
 
 	void SetGlobalIntersectionAnalysis(bool bInGlobalIntersectionAnalysis) { bGlobalIntersectionAnalysis = bInGlobalIntersectionAnalysis; }
 	void SetContourMinimization(bool bInContourMinimization) { bContourMinimization = bInContourMinimization; }
diff --git a/Engine/Source/Runtime/Experimental/Chaos/Public/Chaos/PBDTriangleMeshIntersections.h b/Engine/Source/Runtime/Experimental/Chaos/Public/Chaos/PBDTriangleMeshIntersections.h
--- a/Engine/Source/Runtime/Experimental/Chaos/Public/Chaos/PBDTriangleMeshIntersections.h
+++ b/Engine/Source/Runtime/Experimental/Chaos/Public/Chaos/PBDTriangleMeshIntersections.h
@@ -29,7 +29,8 @@ public:
 
 	~FPBDTriangleMeshIntersections() = default;
 
-	CHAOS_API void Apply(FSolverParticles& Particles, const TArray<FPBDTriangleMeshCollisions::FContourMinimizationIntersection>& Intersections, const FSolverReal Dt) const;
+	template<typename SolverParticlesOrRange>
+	CHAOS_API void Apply(SolverParticlesOrRange& Particles, const TArray<FPBDTriangleMeshCollisions::FContourMinimizationIntersection>& Intersections, const FSolverReal Dt) const;
 
 private:
 	const FTriangleMesh& TriangleMesh;
diff --git a/Engine/Source/Runtime/Experimental/Chaos/Public/Chaos/Particles.h b/Engine/Source/Runtime/Experimental/Chaos/Public/Chaos/Particles.h
--- a/Engine/Source/Runtime/Experimental/Chaos/Public/Chaos/Particles.h
+++ b/Engine/Source/Runtime/Experimental/Chaos/Public/Chaos/Particles.h
@@ -96,7 +96,7 @@ namespace Chaos
 
 		void Resize(const int32 Num)
 		{
-			AddParticles(Num - Size());
+			ResizeHelper(Num);
 			IncrementDirtyValidation();
 		}
 
@@ -134,6 +134,11 @@ namespace Chaos
 		{
 			return MX;
 		}
+
+		TArrayCollectionArray<TVector<T, d>>& XArray()
+		{
+			return MX;
+		}
 		
 		const TVector<T, d>& X(const int32 Index) const
 		{
diff --git a/Engine/Source/Runtime/Experimental/Chaos/Public/Chaos/ParticlesRange.h b/Engine/Source/Runtime/Experimental/Chaos/Public/Chaos/ParticlesRange.h
new file mode 100644
--- /dev/null
+++ b/Engine/Source/Runtime/Experimental/Chaos/Public/Chaos/ParticlesRange.h
@@ -0,0 +1,59 @@
+// Copyright Epic Games, Inc. All Rights Reserved.
+#pragma once
+
+namespace Chaos::Softs
+{
+template<typename ParticlesType, typename = typename TEnableIf<TIsDerivedFrom<ParticlesType, TParticles<FSolverReal, 3>>::IsDerived>::Type>
+class TParticlesRange
+{
+public:
+
+	TParticlesRange() = default;
+	virtual ~TParticlesRange() = default;
+
+	TParticlesRange(ParticlesType* InParticles, int32 InOffset, int32 InRangeSize)
+		: Particles(InParticles)
+		, Offset(InOffset)
+		, RangeSize(InRangeSize)
+	{
+	}
+
+	static TParticlesRange AddParticleRange(ParticlesType& InParticles, const int32 InRangeSize)
+	{
+		const int32 Offset = (int32)InParticles.Size();
+		InParticles.AddParticles(InRangeSize);
+		return TParticlesRange(&InParticles, Offset, InRangeSize);
+	}
+
+	bool IsValid() const
+	{
+		return Particles && Offset >= 0 && ((uint32)(Offset + RangeSize) <= Particles->Size());
+	}
+
+	template<typename T>
+	TConstArrayView<T> GetConstArrayView(const TArray<T>& Array) const
+	{
+		check(Offset >= 0 && Offset + RangeSize <= Array.Num());
+		return TConstArrayView<T>(Array.GetData() + Offset, RangeSize);
+	}
+
+	template<typename T>
+	TArrayView<T> GetArrayView(TArray<T>& Array) const
+	{
+		check(Offset >= 0 && Offset + RangeSize <= Array.Num());
+		return TArrayView<T>(Array.GetData() + Offset, RangeSize);
+	}
+
+	const ParticlesType& GetParticles() const { check(Particles);  return *Particles; }
+	ParticlesType& GetParticles() { check(Particles); return *Particles; }
+	int32 GetOffset() const { return Offset; }
+	int32 GetRangeSize() const { return RangeSize; }
+	int32 Size() const { return RangeSize; } // So this has same interface as Particles
+
+protected:
+
+	ParticlesType* Particles = nullptr;
+	int32 Offset = INDEX_NONE;
+	int32 RangeSize = 0;
+};
+}
diff --git a/Engine/Source/Runtime/Experimental/Chaos/Public/Chaos/PerParticleDampVelocity.h b/Engine/Source/Runtime/Experimental/Chaos/Public/Chaos/PerParticleDampVelocity.h
--- a/Engine/Source/Runtime/Experimental/Chaos/Public/Chaos/PerParticleDampVelocity.h
+++ b/Engine/Source/Runtime/Experimental/Chaos/Public/Chaos/PerParticleDampVelocity.h
@@ -18,13 +18,19 @@ public:
 	~FPerParticleDampVelocity() {}
 
 	void UpdatePositionBasedState(const FSolverParticles& Particles, const int32 Offset, const int32 Range);
+	void UpdatePositionBasedState(const FSolverParticlesRange& Particles);
 
 	// Apply damping without first checking for kinematic particles
-	inline void ApplyFast(FSolverParticles& Particles, const FSolverReal /*Dt*/, const int32 Index) const
+	void ApplyFast(FSolverParticles& Particles, const FSolverReal /*Dt*/, const int32 Index) const
 	{
-		const FSolverVec3 R = Particles.X(Index) - Xcm;
-		const FSolverVec3 Dv = Vcm - Particles.V(Index) + FSolverVec3::CrossProduct(R, Omega);
-		Particles.V(Index) += Coefficient * Dv;
+		Apply(Particles.X(Index), Particles.V(Index));
+	}
+
+	void Apply(const FSolverVec3& X, FSolverVec3& V) const
+	{
+		const FSolverVec3 R = X - Xcm;
+		const FSolverVec3 Dv = Vcm - V + FSolverVec3::CrossProduct(R, Omega);
+		V += Coefficient * Dv;
 	}
 
 private:
diff --git a/Engine/Source/Runtime/Experimental/Chaos/Public/Chaos/SimpleGeometryParticles.h b/Engine/Source/Runtime/Experimental/Chaos/Public/Chaos/SimpleGeometryParticles.h
--- a/Engine/Source/Runtime/Experimental/Chaos/Public/Chaos/SimpleGeometryParticles.h
+++ b/Engine/Source/Runtime/Experimental/Chaos/Public/Chaos/SimpleGeometryParticles.h
@@ -48,6 +48,9 @@ namespace Chaos
 
 		FORCEINLINE const TRotation<T, d>& R(const int32 Index) const { return MR[Index]; }
 		FORCEINLINE TRotation<T, d>& R(const int32 Index) { return MR[Index]; }
+		const TArrayCollectionArray<TRotation<T, d>>& GetR() const { return MR; }
+		TArrayCollectionArray<TRotation<T, d>>& GetR() { return MR; }
+
 		FORCEINLINE const FImplicitObjectPtr& GetGeometry(const int32 Index) const { return MGeometry[Index]; }
 		void SetGeometry(const int32 Index, const FImplicitObjectPtr& InGeometry)
 		{
diff --git a/Engine/Source/Runtime/Experimental/Chaos/Public/Chaos/SoftsEvolution.h b/Engine/Source/Runtime/Experimental/Chaos/Public/Chaos/SoftsEvolution.h
new file mode 100644
--- /dev/null
+++ b/Engine/Source/Runtime/Experimental/Chaos/Public/Chaos/SoftsEvolution.h
@@ -0,0 +1,372 @@
+// Copyright Epic Games, Inc. All Rights Reserved.
+#pragma once
+
+#include "Chaos/Core.h"
+#include "Chaos/ArrayCollection.h"
+#include "Chaos/PBDSoftsEvolutionFwd.h"
+#include "Chaos/PBDSoftsSolverParticles.h"
+#include "Chaos/SoftsSolverParticlesRange.h"
+#include "Chaos/SoftsSolverCollisionParticles.h"
+#include "Chaos/SoftsSolverCollisionParticlesRange.h"
+#include "Chaos/VelocityField.h"
+
+namespace Chaos::Softs
+{
+
+/**
+ * Solver can contain multiple "Groups". Groups do not interact with each other. 
+ * They may be in different spaces. They may be solved in parallel, completely independently of each other. 
+ * The only reason why they're in the same evolution is because they share the same solver
+ * settings and step together in time. 
+ * 
+ * A Group can contain multiple "SoftBodies". SoftBodies can interact but have different 
+ * constraint rules/forces.
+ */
+class FEvolution 
+{
+public:
+
+	CHAOS_API FEvolution(const FCollectionPropertyConstFacade& Properties);
+	~FEvolution() = default;
+
+	/** Reset/empty everything.*/
+	CHAOS_API void Reset();
+
+	/** Move forward in time */
+	CHAOS_API void AdvanceOneTimeStep(const FSolverReal Dt, const FSolverReal TimeDependentIterationMultiplier);
+
+	/** Add custom collection arrays */
+	void AddGroupArray(TArrayCollectionArrayBase* Array) { Groups.AddArray(Array); }
+	void AddParticleArray(TArrayCollectionArrayBase* Array) { Particles.AddArray(Array); }
+	void AddCollisionParticleArray(TArrayCollectionArrayBase* Array) { CollisionParticles.AddArray(Array); }
+
+	const FSolverParticles& GetParticles() const { return Particles; }
+	// Giving non-const access so data can be set freely, but do not add or remove particles here. Use AddSoftBody
+	FSolverParticles& GetParticles() { return Particles; }
+
+	const TSet<uint32>& GetActiveGroups() const { return ActiveGroups; }
+	CHAOS_API int32 NumActiveParticles() const;
+
+	/** 
+	 * Add a SoftBody to a Group.
+	 * 
+	 * @return SoftBodyId
+	 */
+	CHAOS_API int32 AddSoftBody(uint32 GroupId, int32 NumParticles, bool bEnable);
+	int32 GetSoftBodyParticleNum(int32 SoftBodyId) const { return SoftBodies.ParticleRanges[SoftBodyId].GetRangeSize(); }
+	int32 GetSoftBodyGroupId(int32 SoftBodyId) const { return SoftBodies.GroupId[SoftBodyId]; }
+	CHAOS_API void SetSoftBodyProperties(int32 SoftBodyId, const FCollectionPropertyConstFacade& PropertyCollection);
+	CHAOS_API void ActivateSoftBody(int32 SoftBodyId, bool bActivate);
+	bool IsSoftBodyActive(int32 SoftBodyId) const { return SoftBodies.Active[SoftBodyId]; }
+	FSolverParticlesRange& GetSoftBodyParticles(int32 SoftBodyId) { return SoftBodies.ParticleRanges[SoftBodyId]; }
+	const FSolverParticlesRange& GetSoftBodyParticles(int32 SoftBodyId) const { return SoftBodies.ParticleRanges[SoftBodyId]; }
+	const TArray<int32>& GetGroupSoftBodies(uint32 GroupId) const { return Groups.SoftBodies[GroupId]; }
+	const TSet<int32>& GetGroupActiveSoftBodies(uint32 GroupId) const { return Groups.ActiveSoftBodies[GroupId]; }
+
+	/**
+	 * Add Collision particle range to a group.
+	 * 
+	 * @return Particle range offset (unique id for this range)
+	 */
+	CHAOS_API int32 AddCollisionParticleRange(uint32 GroupId, int32 NumParticles, bool bEnable);
+	CHAOS_API void RemoveCollisionParticleRange(int32 CollisionRangeId);
+	CHAOS_API void ActivateCollisionParticleRange(int32 CollisionRangeId, bool bEnable);
+	const TSet<int32>& GetGroupActiveCollisionParticleRanges(uint32 GroupId) const { return Groups.ActiveCollisionParticleRanges[GroupId]; }
+	CHAOS_API TArray<FSolverCollisionParticlesRange> GetActiveCollisionParticles(uint32 GroupId) const;
+
+	FSolverCollisionParticlesRange& GetCollisionParticleRange(int32 CollisionRangeId) { return CollisionRanges.ParticleRanges[CollisionRangeId]; }
+	const FSolverCollisionParticlesRange& GetCollisionParticleRange(int32 CollisionRangeId) const { return CollisionRanges.ParticleRanges[CollisionRangeId]; }
+	
+	/** Global Rules*/
+	typedef TFunction<void(FSolverParticlesRange&, const FSolverReal Dt, const FSolverReal Time)> KinematicUpdateFunc;
+	typedef TFunction<void(FSolverCollisionParticlesRange&, const FSolverReal Dt, const FSolverReal Time)> CollisionKinematicUpdateFunc;
+	void SetKinematicUpdateFunction(KinematicUpdateFunc Func) { KinematicUpdate = Func; }
+	void SetCollisionKinematicUpdateFunction(CollisionKinematicUpdateFunc Func) { CollisionKinematicUpdate = Func; }
+
+	/** Soft Body Rules. Add Ranges to allocate space for your rules. You will get back an ArrayView where you can then set the rules. */
+	typedef TFunction<void(const FSolverParticlesRange&, const FSolverReal)> PreSubstepParallelInitFunc;
+	typedef TFunction<void(FSolverParticlesRange&, const FSolverReal)> ExternalForceRuleFunc;
+	typedef TFunction<void(const FSolverParticlesRange&, const FSolverReal)> ConstraintParallelInitFunc;
+	typedef TFunction<void(FSolverParticlesRange&, const FSolverReal)> ConstraintRuleFunc;
+	typedef TFunction<void(FSolverParticlesRange&, const FSolverReal, const TArray<FSolverCollisionParticlesRange>&)> CollisionConstraintRuleFunc;
+	// Warning: Rules are allocated into shared buffers. The ArrayViews may go stale if you allocate ANY new rules.
+	// Allocate your batch of rules, then get your array views.	
+	void AllocatePreSubstepParallelInitRange(int32 SoftBodyId, int32 NumRules)
+	{
+		return AllocateRulesRange(SoftBodyId, NumRules, ConstParticleRules, SoftBodies.PreSubstepParallelInits);
+	}
+	void AllocateExternalForceRulesRange(int32 SoftBodyId, int32 NumRules)
+	{
+		return AllocateRulesRange(SoftBodyId, NumRules, ParticleRules, SoftBodies.ExternalForceRules);
+	}
+	void AllocateConstraintParallelInitsRange(int32 SoftBodyId, int32 NumRules)
+	{
+		return AllocateRulesRange(SoftBodyId, NumRules, ConstParticleRules, SoftBodies.ConstraintParallelInits);
+	}
+	void AllocatePreSubstepConstraintRulesRange(int32 SoftBodyId, int32 NumRules)
+	{
+		return AllocateRulesRange(SoftBodyId, NumRules, ParticleRules, SoftBodies.PreSubstepConstraintRules);
+	}
+	void AllocatePerIterationConstraintRulesRange(int32 SoftBodyId, int32 NumRules)
+	{
+		return AllocateRulesRange(SoftBodyId, NumRules, ParticleRules, SoftBodies.PerIterationConstraintRules);
+	}
+	void AllocatePerIterationCollisionConstraintRulesRange(int32 SoftBodyId, int32 NumRules)
+	{
+		return AllocateRulesRange(SoftBodyId, NumRules, CollisionRules, SoftBodies.PerIterationCollisionConstraintRules);
+	}
+	void AllocatePerIterationPostCollisionsConstraintRulesRange(int32 SoftBodyId, int32 NumRules)
+	{
+		return AllocateRulesRange(SoftBodyId, NumRules, ParticleRules, SoftBodies.PerIterationPostCollisionsConstraintRules);
+	}
+	void AllocatePostSubstepConstraintRulesRange(int32 SoftBodyId, int32 NumRules)
+	{
+		return AllocateRulesRange(SoftBodyId, NumRules, ParticleRules, SoftBodies.PostSubstepConstraintRules);
+	}
+
+	TArrayView<PreSubstepParallelInitFunc> GetPreSubstepParallelInitRange(int32 SoftBodyId)
+	{
+		return GetRulesRange(SoftBodyId, SoftBodies.PreSubstepParallelInits);
+	}
+	TArrayView<ExternalForceRuleFunc> GetExternalForceRulesRange(int32 SoftBodyId)
+	{
+		return GetRulesRange(SoftBodyId, SoftBodies.ExternalForceRules);
+	}
+	TArrayView<ConstraintParallelInitFunc> GetConstraintParallelInitsRange(int32 SoftBodyId)
+	{
+		return GetRulesRange(SoftBodyId, SoftBodies.ConstraintParallelInits);
+	}
+	TArrayView<ConstraintRuleFunc> GetPreSubstepConstraintRulesRange(int32 SoftBodyId)
+	{
+		return GetRulesRange(SoftBodyId, SoftBodies.PreSubstepConstraintRules);
+	}
+	TArrayView<ConstraintRuleFunc> GetPerIterationConstraintRulesRange(int32 SoftBodyId)
+	{
+		return GetRulesRange(SoftBodyId, SoftBodies.PerIterationConstraintRules);
+	}
+	TArrayView<CollisionConstraintRuleFunc> GetPerIterationCollisionConstraintRulesRange(int32 SoftBodyId)
+	{
+		return GetRulesRange(SoftBodyId, SoftBodies.PerIterationCollisionConstraintRules);
+	}
+	TArrayView<ConstraintRuleFunc> GetPerIterationPostCollisionsConstraintRulesRange(int32 SoftBodyId)
+	{
+		return GetRulesRange(SoftBodyId, SoftBodies.PerIterationPostCollisionsConstraintRules);
+	}
+	TArrayView<ConstraintRuleFunc> GetPostSubstepConstraintRulesRange(int32 SoftBodyId)
+	{
+		return GetRulesRange(SoftBodyId, SoftBodies.PostSubstepConstraintRules);
+	}
+
+	/** Solver settings */
+	FSolverReal GetTime() const { return Time; }
+	int32 GetIterations() const { return NumIterations; }
+	int32 GetMaxIterations() const { return MaxNumIterations; }
+	bool GetDisableTimeDependentNumIterations() const { return bDisableTimeDependentNumIterations; }
+	bool GetDoQuasistatics() const { return bDoQuasistatics; }
+	void SetDisableTimeDependentNumIterations(bool bDisable) { bDisableTimeDependentNumIterations = bDisable; }
+	CHAOS_API void SetSolverProperties(const FCollectionPropertyConstFacade& PropertyCollection);
+
+private:
+
+	template<typename ElementType>
+	struct TArrayRange
+	{
+		static TArrayRange AddRange(TArray<ElementType>& InArray, int32 InRangeSize)
+		{
+			TArrayRange Range;
+			Range.Offset = InArray.Num();
+			Range.Array = &InArray;
+			Range.Array->AddDefaulted(InRangeSize);
+			Range.RangeSize = InRangeSize;
+			return Range;
+		}
+
+		bool IsValid() const
+		{
+			return Array && Offset >= 0 && Offset + RangeSize <= Array->Num();
+		}
+
+		TConstArrayView<ElementType> GetConstArrayView() const
+		{
+			check(IsValid());
+			return TConstArrayView<ElementType>(Array->GetData() + Offset, RangeSize);
+		}
+
+		TArrayView<ElementType> GetArrayView()
+		{
+			check(IsValid());
+			return TArrayView<ElementType>(Array->GetData() + Offset, RangeSize);
+		}
+
+		bool IsEmpty() const { return RangeSize == 0; }
+		int32 GetRangeSize() const { return RangeSize; }
+	private:
+		TArray<ElementType>* Array = nullptr;
+		int32 Offset = INDEX_NONE;
+		int32 RangeSize = 0;
+	};
+
+	// SoftBody SOA
+	struct FSoftBodies : public TArrayCollection
+	{
+		FSoftBodies()
+		{
+			TArrayCollection::AddArray(&Active);
+			TArrayCollection::AddArray(&GroupId);
+			TArrayCollection::AddArray(&ParticleRanges);
+			TArrayCollection::AddArray(&GlobalDampings);
+			TArrayCollection::AddArray(&LocalDampings);
+			TArrayCollection::AddArray(&UsePerParticleDamping);
+			TArrayCollection::AddArray(&PreSubstepParallelInits);
+			TArrayCollection::AddArray(&ExternalForceRules);
+			TArrayCollection::AddArray(&ConstraintParallelInits);
+			TArrayCollection::AddArray(&PreSubstepConstraintRules);
+			TArrayCollection::AddArray(&PerIterationConstraintRules);
+			TArrayCollection::AddArray(&PerIterationCollisionConstraintRules);
+			TArrayCollection::AddArray(&PerIterationPostCollisionsConstraintRules);
+			TArrayCollection::AddArray(&PostSubstepConstraintRules);
+		}
+
+		void Reset()
+		{
+			ResizeHelper(0);
+		}
+
+		int32 AddSoftBody()
+		{
+			const int32 Offset = Size();
+			AddElementsHelper(1);
+			return Offset;
+		}
+
+		TArrayCollectionArray<bool> Active;
+		TArrayCollectionArray<uint32> GroupId;
+		TArrayCollectionArray<FSolverParticlesRange> ParticleRanges;
+		TArrayCollectionArray<FSolverReal> GlobalDampings;
+		TArrayCollectionArray<FSolverReal> LocalDampings;
+		TArrayCollectionArray<bool> UsePerParticleDamping;
+		TArrayCollectionArray<TArrayRange<PreSubstepParallelInitFunc>> PreSubstepParallelInits;
+		TArrayCollectionArray<TArrayRange<ExternalForceRuleFunc>> ExternalForceRules;
+		TArrayCollectionArray<TArrayRange<ConstraintParallelInitFunc>> ConstraintParallelInits;
+		TArrayCollectionArray<TArrayRange<ConstraintRuleFunc>> PreSubstepConstraintRules;
+		TArrayCollectionArray<TArrayRange<ConstraintRuleFunc>> PerIterationConstraintRules;
+		TArrayCollectionArray<TArrayRange<CollisionConstraintRuleFunc>> PerIterationCollisionConstraintRules;
+		TArrayCollectionArray<TArrayRange<ConstraintRuleFunc>> PerIterationPostCollisionsConstraintRules;
+		TArrayCollectionArray<TArrayRange<ConstraintRuleFunc>> PostSubstepConstraintRules;
+	};
+
+	// CollisionBodyRange SOA
+	struct FCollisionBodyRanges : public TArrayCollection
+	{
+		FCollisionBodyRanges()
+		{
+			TArrayCollection::AddArray(&Status);
+			TArrayCollection::AddArray(&GroupId);
+			TArrayCollection::AddArray(&ParticleRanges);
+		}
+
+		void Reset()
+		{
+			ResizeHelper(0);
+		}
+
+		int32 AddRange()
+		{
+			const int32 Offset = Size();
+			AddElementsHelper(1);
+			return Offset;
+		}
+
+		enum struct EStatus : uint8
+		{
+			Invalid = 0,
+			Active = 1,
+			Inactive = 2,
+			Free = 3 // Available for recycling
+		};
+		TArrayCollectionArray<EStatus> Status;
+		TArrayCollectionArray<uint32> GroupId;
+		TArrayCollectionArray<FSolverCollisionParticlesRange> ParticleRanges;
+	};
+
+	struct FGroups : public TArrayCollection
+	{
+		FGroups()
+		{
+			TArrayCollection::AddArray(&SoftBodies);
+			TArrayCollection::AddArray(&ActiveSoftBodies);
+			TArrayCollection::AddArray(&ActiveCollisionParticleRanges);
+		}
+		
+		void Reset()
+		{
+			ResizeHelper(0);
+		}
+
+		void AddGroupsToSize(uint32 DesiredSize)
+		{
+			if (ensure(DesiredSize >= Size()))
+			{
+				ResizeHelper((int32)DesiredSize);
+			}
+		}
+
+		TArrayCollectionArray<TArray<int32>> SoftBodies;
+		TArrayCollectionArray<TSet<int32>> ActiveSoftBodies;
+		TArrayCollectionArray<TSet<int32>> ActiveCollisionParticleRanges;
+	};
+
+	template<typename RuleFunc>
+	void AllocateRulesRange(int32 SoftBodyId, int32 NumRules, TArray<RuleFunc>& RuleArray, TArrayCollectionArray<TArrayRange<RuleFunc>>& RangeArray)
+	{
+		check(RangeArray[SoftBodyId].IsEmpty());
+		RangeArray[SoftBodyId] = TArrayRange<RuleFunc>::AddRange(RuleArray, NumRules);
+	}
+
+	template<typename RuleFunc>
+	TArrayView<RuleFunc> GetRulesRange(int32 SoftBodyId, TArrayCollectionArray<TArrayRange<RuleFunc>>& RangeArray)
+	{
+		return RangeArray[SoftBodyId].GetArrayView();
+	}
+
+	void AdvanceOneTimeStepInternal(const FSolverReal Dt, const int32 TimeDependentNumIterations, uint32 GroupId);
+
+	// Solver data
+	int32 NumIterations;
+	int32 MaxNumIterations; // Used for time-dependent iteration counts
+	bool bDisableTimeDependentNumIterations;
+	bool bDoQuasistatics;
+	FSolverReal SolverFrequency; 
+	FSolverReal Time;
+
+	// Per-Particle data
+	FSolverParticles Particles;
+	TArrayCollectionArray<FSolverReal> ParticleDampings;
+
+	// Per-Collision particle data
+	FSolverCollisionParticles CollisionParticles;
+
+	// Per-SoftBody data
+	FSoftBodies SoftBodies;
+
+	// Per-CollisionBodyRange data
+	FCollisionBodyRanges CollisionRanges;
+
+	// Collision Range free-list
+	TMap<int32, TArray<int32>> CollisionRangeFreeList; // Key = NumParticles, Value = CollisionRangeId(s)
+
+	// Per-Group data
+	FGroups Groups;
+	TSet<uint32> ActiveGroups; // Groups with at least one active softbody
+
+	KinematicUpdateFunc KinematicUpdate;
+	CollisionKinematicUpdateFunc CollisionKinematicUpdate;
+
+	// Rules that run on ParticleRanges (SoftBodies will have views into these)
+	TArray<TFunction<void(const FSolverParticlesRange&, const FSolverReal)>> ConstParticleRules;
+	TArray<TFunction<void(FSolverParticlesRange&, const FSolverReal)>> ParticleRules;
+
+	// Collision Rules
+	TArray<TFunction<void(FSolverParticlesRange&, const FSolverReal, const TArray<FSolverCollisionParticlesRange>&)>> CollisionRules;
+};
+}
\ No newline at end of file
diff --git a/Engine/Source/Runtime/Experimental/Chaos/Public/Chaos/SoftsExternalForces.h b/Engine/Source/Runtime/Experimental/Chaos/Public/Chaos/SoftsExternalForces.h
new file mode 100644
--- /dev/null
+++ b/Engine/Source/Runtime/Experimental/Chaos/Public/Chaos/SoftsExternalForces.h
@@ -0,0 +1,163 @@
+// Copyright Epic Games, Inc. All Rights Reserved.
+#pragma once
+
+#include "Chaos/PBDSoftsEvolutionFwd.h"
+#include "Chaos/PBDSoftsSolverParticles.h"
+#include "Chaos/SoftsSolverParticlesRange.h"
+#include "Chaos/CollectionPropertyFacade.h"
+
+namespace Chaos::Softs
+{
+class FExternalForcesBase
+{
+public:
+	FExternalForcesBase(const FSolverVec3& InGravity,
+		const TArray<FSolverVec3>& InNormals)
+		: Gravity(InGravity)
+		, FictitiousAngularDisplacement(0.f)
+		, ReferenceSpaceLocation(0.f)
+		, bUsePointBasedWindModel(false)
+		, PointBasedWind(0.f)
+		, LegacyWindAdaptation(0.f)
+		, Normals(InNormals)
+	{}
+
+	virtual ~FExternalForcesBase() {}
+
+	void Apply(FSolverParticlesRange& Particles, const FSolverReal Dt) const
+	{
+		FSolverVec3* const Acceleration = Particles.GetAcceleration().GetData();
+		const FSolverReal* const InvM = Particles.GetInvM().GetData();
+		const FSolverVec3* const X = Particles.XArray().GetData();
+		const FSolverVec3* const V = Particles.GetV().GetData();
+		const FSolverVec3* const N = Particles.GetConstArrayView(Normals).GetData();
+
+		const bool bHasFictitiousForces = !FictitiousAngularDisplacement.IsNearlyZero();
+		const FSolverVec3 W = FictitiousAngularDisplacement / Dt;
+
+		for (int32 Index = 0; Index < Particles.GetRangeSize(); ++Index)
+		{
+			if (InvM[Index] != (FSolverReal)0.)
+			{
+				Acceleration[Index] = Gravity;
+
+				if (bHasFictitiousForces)
+				{
+					// Centrifugal force (*InvM to get acceleration)
+					Acceleration[Index] -= FSolverVec3::CrossProduct(W, FSolverVec3::CrossProduct(W, X[Index] - ReferenceSpaceLocation));
+				}
+				if (bUsePointBasedWindModel)
+				{
+					const FSolverVec3 VelocityDelta = PointBasedWind - V[Index];
+					FSolverVec3 Direction = VelocityDelta;
+					if (Direction.Normalize())
+					{
+						// Scale by angle
+						const FReal DirectionDot = FVec3::DotProduct(Direction, N[Index]);
+						const FReal ScaleFactor = FMath::Min(1.f, FMath::Abs(DirectionDot) * LegacyWindAdaptation);
+						Acceleration[Index] += VelocityDelta * ScaleFactor;
+					}
+				}
+			}
+		}
+	}
+
+	bool UsePointBasedWindModel() const { return bUsePointBasedWindModel; }
+	const FSolverVec3& GetGravity() const { return Gravity; }
+
+protected:
+	FSolverVec3 Gravity; 
+	FSolverVec3 FictitiousAngularDisplacement;
+	FSolverVec3 ReferenceSpaceLocation;
+	bool bUsePointBasedWindModel;
+	FSolverVec3 PointBasedWind;
+	FSolverReal LegacyWindAdaptation;
+	const TArray<FSolverVec3>& Normals;
+};
+
+class FExternalForces : public FExternalForcesBase
+{
+	typedef FExternalForcesBase Base;
+public:
+	static constexpr FSolverReal DefaultGravityScale = (FSolverReal)1.f;
+	static constexpr FSolverReal DefaultGravityZOverride = (FSolverReal)-980.665f;
+	static constexpr bool bDefaultUseGravityOverride = false;
+	static constexpr FSolverReal DefaultFictitiousAngularScale = (FSolverReal)1.f;
+	static constexpr bool bDefaultUsePointBasedWindModel = false;
+
+	// We don't have the solver values at the time this is constructed. Need to SetProperties at least once after setting 
+	// data from the solver.
+	FExternalForces(
+		const TArray<FSolverVec3>& InNormals,
+		const FCollectionPropertyConstFacade& PropertyCollection)
+		: Base(
+			FSolverVec3((FSolverReal)0.f, (FSolverReal)0.f, DefaultGravityZOverride),
+			InNormals
+			)
+		, WorldGravityMultiplier((FSolverReal)1.f)
+		, SolverGravity((FSolverReal)0.f, (FSolverReal)0.f, DefaultGravityZOverride)
+		, bPerSoftBodyGravityOverrideEnabled(true)
+		, FictitiousAngularDisplacementNoScale(0.f)
+		, UseGravityOverrideIndex(PropertyCollection)
+		, GravityOverrideIndex(PropertyCollection)
+		, GravityScaleIndex(PropertyCollection)
+		, FictitiousAngularScaleIndex(PropertyCollection)
+		, UsePointBasedWindModelIndex(PropertyCollection)
+	{}
+
+	void SetWorldGravityMultiplier(FSolverReal InWorldGravityMultiplier) { WorldGravityMultiplier = InWorldGravityMultiplier; }
+	void SetSolverGravityProperties(const FSolverVec3& InSolverGravity, bool bInPerSoftBodyGravityOverrideEnabled)
+	{
+		SolverGravity = InSolverGravity;
+		bPerSoftBodyGravityOverrideEnabled = bInPerSoftBodyGravityOverrideEnabled;
+	}
+	void SetFictitiousForcesData(const FSolverVec3& InFictitiousAngularDisplacementNoScale, const FSolverVec3& InReferenceSpaceLocation)
+	{
+		FictitiousAngularDisplacementNoScale = InFictitiousAngularDisplacementNoScale;
+		ReferenceSpaceLocation = InReferenceSpaceLocation;
+	}
+	void SetSolverWind(const FSolverVec3& SolverWind, const FSolverReal InLegacyWindAdaptation)
+	{
+		PointBasedWind = SolverWind;
+		LegacyWindAdaptation = InLegacyWindAdaptation;
+	}
+
+	/** This should be called after all other data is set as it will populate the final base class values.*/
+	void SetProperties(const FCollectionPropertyConstFacade& PropertyCollection)
+	{
+		const bool bUseGravityOverride = UseGravityOverrideIndex != INDEX_NONE ? GetUseGravityOverride(PropertyCollection) : bDefaultUseGravityOverride;
+		const FSolverVec3 GravityOverride = GravityOverrideIndex != INDEX_NONE ? FSolverVec3(GetGravityOverride(PropertyCollection)) :
+			FSolverVec3(0.f, 0.f, DefaultGravityZOverride);
+		const FSolverReal GravityScale = GravityScaleIndex != INDEX_NONE ? GetGravityScale(PropertyCollection) : DefaultGravityScale;
+		CalculateGravity(bUseGravityOverride, GravityOverride, GravityScale);
+
+		const FSolverReal FictitiousAngularScale = FMath::Min((FSolverReal)2., FictitiousAngularScaleIndex != INDEX_NONE ? GetFictitiousAngularScale(PropertyCollection) : DefaultFictitiousAngularScale);
+
+		FictitiousAngularDisplacement = FictitiousAngularDisplacementNoScale * FictitiousAngularScale;
+
+		bUsePointBasedWindModel = UsePointBasedWindModelIndex != INDEX_NONE ? GetUsePointBasedWindModel(PropertyCollection) : bDefaultUsePointBasedWindModel;
+	}
+private:
+	void CalculateGravity(bool bUseGravityOverride, const FSolverVec3& GravityOverride, FSolverReal GravityScale)
+	{
+		Gravity = (bPerSoftBodyGravityOverrideEnabled && bUseGravityOverride ? GravityOverride : SolverGravity * GravityScale) * WorldGravityMultiplier;
+	}
+
+	/** Gravity */
+	// "World" level properties (this comes from a world-level CVar)
+	FSolverReal WorldGravityMultiplier;
+
+	// "Solver" level properties (set from Solver-level BP)
+	FSolverVec3 SolverGravity;
+	bool bPerSoftBodyGravityOverrideEnabled;
+
+	/** Fictitious Forces */
+	FSolverVec3 FictitiousAngularDisplacementNoScale; // Without Scale parameter applied.
+
+	UE_CHAOS_DECLARE_PROPERTYCOLLECTION_NAME(UseGravityOverride, bool);
+	UE_CHAOS_DECLARE_PROPERTYCOLLECTION_NAME(GravityOverride, FVector3f);
+	UE_CHAOS_DECLARE_PROPERTYCOLLECTION_NAME(GravityScale, float);
+	UE_CHAOS_DECLARE_PROPERTYCOLLECTION_NAME(FictitiousAngularScale, float);
+	UE_CHAOS_DECLARE_PROPERTYCOLLECTION_NAME(UsePointBasedWindModel, bool);
+};
+}
diff --git a/Engine/Source/Runtime/Experimental/Chaos/Public/Chaos/SoftsSolverCollisionParticles.h b/Engine/Source/Runtime/Experimental/Chaos/Public/Chaos/SoftsSolverCollisionParticles.h
--- a/Engine/Source/Runtime/Experimental/Chaos/Public/Chaos/SoftsSolverCollisionParticles.h
+++ b/Engine/Source/Runtime/Experimental/Chaos/Public/Chaos/SoftsSolverCollisionParticles.h
@@ -28,9 +28,13 @@ class FSolverCollisionParticles : public TSimpleGeometryParticles<FSolverReal, 3
 
 	const FSolverVec3& V(const int32 Index) const { return MV[Index]; }
 	FSolverVec3& V(const int32 Index) { return MV[Index]; }
+	const TArrayCollectionArray<FSolverVec3>& GetV() const { return MV; }
+	TArrayCollectionArray<FSolverVec3>& GetV() { return MV; }
 
 	const FSolverVec3& W(const int32 Index) const { return MW[Index]; }
 	FSolverVec3& W(const int32 Index) { return MW[Index]; }
+	const TArrayCollectionArray<FSolverVec3>& GetW() const { return MW; }
+	TArrayCollectionArray<FSolverVec3>& GetW() { return MW; }
 
 	virtual void Serialize(FChaosArchive& Ar) override
 	{
diff --git a/Engine/Source/Runtime/Experimental/Chaos/Public/Chaos/SoftsSolverCollisionParticlesRange.h b/Engine/Source/Runtime/Experimental/Chaos/Public/Chaos/SoftsSolverCollisionParticlesRange.h
new file mode 100644
--- /dev/null
+++ b/Engine/Source/Runtime/Experimental/Chaos/Public/Chaos/SoftsSolverCollisionParticlesRange.h
@@ -0,0 +1,50 @@
+// Copyright Epic Games, Inc. All Rights Reserved.
+#pragma once
+
+#include "Chaos/SoftsSolverCollisionParticles.h"
+#include "Chaos/PBDSoftsEvolutionFwd.h"
+#include "Chaos/ParticlesRange.h"
+
+namespace Chaos::Softs
+{
+class FSolverCollisionParticlesRange final : public TParticlesRange<FSolverCollisionParticles>
+{
+public:
+	FSolverCollisionParticlesRange() = default;
+	~FSolverCollisionParticlesRange() = default;
+	FSolverCollisionParticlesRange(const FSolverCollisionParticlesRange&) = default;
+	FSolverCollisionParticlesRange(FSolverCollisionParticlesRange&&) = default;
+	FSolverCollisionParticlesRange& operator=(const FSolverCollisionParticlesRange&) = default;
+	FSolverCollisionParticlesRange& operator=(FSolverCollisionParticlesRange&&) = default;
+
+	FSolverCollisionParticlesRange(TParticlesRange<FSolverCollisionParticles>&& Other)
+		:TParticlesRange<FSolverCollisionParticles>(MoveTemp(Other))
+	{}
+
+	// SolverCollisionParticles data
+	const FSolverVec3& V(const int32 Index) const { return GetParticles().V(Index + Offset); }
+	FSolverVec3& V(const int32 Index) { return GetParticles().V(Index + Offset); }
+	TConstArrayView<FSolverVec3> GetV() const { return GetConstArrayView(GetParticles().GetV()); }
+	TArrayView<FSolverVec3> GetV() { return GetArrayView(GetParticles().GetV()); }
+	const FSolverVec3& W(const int32 Index) const { return GetParticles().W(Index + Offset); }
+	FSolverVec3& W(const int32 Index) { return GetParticles().W(Index + Offset); }
+	TConstArrayView<FSolverVec3> GetW() const { return GetConstArrayView(GetParticles().GetW()); }
+	TArrayView<FSolverVec3> GetW() { return GetArrayView(GetParticles().GetW()); }
+
+	// SimpleGeometryParticles data
+	const FSolverRotation3& R(const int32 Index) const { return GetParticles().R(Index + Offset); }
+	FSolverRotation3& R(const int32 Index) { return GetParticles().R(Index + Offset); }
+	TConstArrayView<FSolverRotation3> GetR() const { return GetConstArrayView(GetParticles().GetR()); }
+	TArrayView<FSolverRotation3> GetR() { return GetArrayView(GetParticles().GetR()); }
+
+	const FImplicitObjectPtr& GetGeometry(const int32 Index) const { return GetParticles().GetGeometry(Index + Offset); }
+	void SetGeometry(const int32 Index, const FImplicitObjectPtr& InGeometry) { return GetParticles().SetGeometry(Index + Offset, InGeometry); }
+	TConstArrayView<FImplicitObjectPtr> GetAllGeometry() const { return GetConstArrayView(GetParticles().GetAllGeometry()); }
+
+	// Particles data
+	const FSolverVec3& X(const int32 Index) const { return GetParticles().X(Index + Offset); }
+	FSolverVec3& X(const int32 Index) { return GetParticles().X(Index + Offset); }
+	TConstArrayView<FSolverVec3> XArray() const { return GetConstArrayView(GetParticles().XArray()); }
+	TArrayView<FSolverVec3> XArray() { return GetArrayView(GetParticles().XArray()); }
+};
+}
diff --git a/Engine/Source/Runtime/Experimental/Chaos/Public/Chaos/SoftsSolverParticlesRange.h b/Engine/Source/Runtime/Experimental/Chaos/Public/Chaos/SoftsSolverParticlesRange.h
new file mode 100644
--- /dev/null
+++ b/Engine/Source/Runtime/Experimental/Chaos/Public/Chaos/SoftsSolverParticlesRange.h
@@ -0,0 +1,65 @@
+// Copyright Epic Games, Inc. All Rights Reserved.
+#pragma once
+
+#include "Chaos/PBDSoftsSolverParticles.h"
+#include "Chaos/PBDSoftsEvolutionFwd.h"
+#include "Chaos/ParticlesRange.h"
+
+namespace Chaos::Softs
+{
+
+class FSolverParticlesRange final : public TParticlesRange<FSolverParticles>
+{
+public:
+	FSolverParticlesRange() = default;
+	~FSolverParticlesRange() = default;
+	FSolverParticlesRange(const FSolverParticlesRange&) = default;
+	FSolverParticlesRange(FSolverParticlesRange&&) = default;
+	FSolverParticlesRange& operator=(const FSolverParticlesRange&) = default;
+	FSolverParticlesRange& operator=(FSolverParticlesRange&&) = default;
+
+	FSolverParticlesRange(FSolverParticles* InParticles, const int32 InOffset, const int32 InRange)
+		: TParticlesRange<FSolverParticles>(InParticles, InOffset, InRange)
+	{}
+
+	FSolverParticlesRange(TParticlesRange<FSolverParticles>&& Other)
+		: TParticlesRange<FSolverParticles>(MoveTemp(Other))
+	{}
+
+	// SolverParticles data
+	const FSolverVec3& P(const int32 Index) const { return GetParticles().P(Index + Offset); }
+	FSolverVec3& P(const int32 Index) { return GetParticles().P(Index + Offset); }
+	const FPAndInvM& PAndInvM(const int32 Index) const { return GetParticles().PAndInvM(Index + Offset); }
+	FPAndInvM& PAndInvM(const int32 Index) { return GetParticles().PAndInvM(Index + Offset); }
+	TConstArrayView<FPAndInvM> GetPAndInvM() const { return GetConstArrayView(GetParticles().GetPAndInvM()); }
+	TArrayView<FPAndInvM> GetPAndInvM() { return GetArrayView(GetParticles().GetPAndInvM()); }
+	const FSolverVec3& VPrev(const int32 Index) const { return GetParticles().VPrev(Index + Offset); }
+	FSolverVec3& VPrev(const int32 Index) { return GetParticles().VPrev(Index + Offset); }
+	TConstArrayView<FSolverVec3> GetVPrev() const { return GetConstArrayView(GetParticles().GetVPrev()); }
+	TArrayView<FSolverVec3> GetVPrev() { return GetArrayView(GetParticles().GetVPrev()); }
+
+	// DynamicParticles data
+	const FSolverVec3& V(const int32 Index) const { return GetParticles().V(Index + Offset); }
+	FSolverVec3& V(const int32 Index) { return GetParticles().V(Index + Offset); }
+	TConstArrayView<FSolverVec3> GetV() const { return GetConstArrayView(GetParticles().GetV()); }
+	TArrayView<FSolverVec3> GetV() { return GetArrayView(GetParticles().GetV()); }
+	const FSolverVec3& Acceleration(const int32 Index) const { return GetParticles().Acceleration(Index + Offset); }
+	FSolverVec3& Acceleration(const int32 Index) { return GetParticles().Acceleration(Index + Offset); }
+	TConstArrayView<FSolverVec3> GetAcceleration() const { return GetConstArrayView(GetParticles().GetAcceleration()); }
+	TArrayView<FSolverVec3> GetAcceleration() { return GetArrayView(GetParticles().GetAcceleration()); }
+	FSolverReal M(const int32 Index) const { return GetParticles().M(Index + Offset); }
+	FSolverReal& M(const int32 Index) { return GetParticles().M(Index + Offset); }
+	TConstArrayView<FSolverReal> GetM() const { return GetConstArrayView(GetParticles().GetM()); }
+	TArrayView<FSolverReal> GetM() { return GetArrayView(GetParticles().GetM()); }
+	FSolverReal InvM(const int32 Index) const { return GetParticles().InvM(Index + Offset); }
+	FSolverReal& InvM(const int32 Index) { return GetParticles().InvM(Index + Offset); }
+	TConstArrayView<FSolverReal> GetInvM() const { return GetConstArrayView(GetParticles().GetInvM()); }
+	TArrayView<FSolverReal> GetInvM() { return GetArrayView(GetParticles().GetInvM()); }
+
+	// Particles data
+	const FSolverVec3& X(const int32 Index) const { return GetParticles().X(Index + Offset); }
+	FSolverVec3& X(const int32 Index) { return GetParticles().X(Index + Offset); }
+	TConstArrayView<FSolverVec3> XArray() const { return GetConstArrayView(GetParticles().XArray()); }
+	TArrayView<FSolverVec3> XArray() { return GetArrayView(GetParticles().XArray()); }
+};
+}
diff --git a/Engine/Source/Runtime/Experimental/Chaos/Public/Chaos/VelocityField.h b/Engine/Source/Runtime/Experimental/Chaos/Public/Chaos/VelocityField.h
--- a/Engine/Source/Runtime/Experimental/Chaos/Public/Chaos/VelocityField.h
+++ b/Engine/Source/Runtime/Experimental/Chaos/Public/Chaos/VelocityField.h
@@ -3,6 +3,7 @@
 
 #include "Chaos/PBDSoftsEvolutionFwd.h"
 #include "Chaos/PBDSoftsSolverParticles.h"
+#include "Chaos/SoftsSolverParticlesRange.h"
 #include "Chaos/TriangleMesh.h"
 #include "Chaos/CollectionPropertyFacade.h"
 
@@ -34,6 +35,7 @@ public:
 		, LiftIndex(PropertyCollection)
 		, FluidDensityIndex(PropertyCollection)
 		, PressureIndex(PropertyCollection)
+		, WindVelocityIndex(PropertyCollection)
 	{
 		SetProperties(
 			FSolverVec2(GetWeightedFloatDrag(PropertyCollection, (FSolverReal)0.)),  // If these properties don't exist, set their values to 0, not to DefaultCoefficients!
@@ -42,6 +44,27 @@ public:
 			FSolverVec2(GetWeightedFloatPressure(PropertyCollection, (FSolverReal)0.)));
 	}
 
+	FVelocityAndPressureField(
+		const FTriangleMesh* TriangleMesh,
+		const FCollectionPropertyConstFacade& PropertyCollection,
+		const TMap<FString, TConstArrayView<FRealSingle>>& Weightmaps,
+		FSolverReal WorldScale)
+		: DragIndex(PropertyCollection)
+		, LiftIndex(PropertyCollection)
+		, FluidDensityIndex(PropertyCollection)
+		, PressureIndex(PropertyCollection)
+		, WindVelocityIndex(PropertyCollection)
+	{
+		SetGeometry(TriangleMesh);
+		SetProperties(
+			FSolverVec2(GetWeightedFloatDrag(PropertyCollection, 0.f)),  // If these properties don't exist, set their values to 0, not to DefaultCoefficients!
+			FSolverVec2(GetWeightedFloatLift(PropertyCollection, 0.f)),
+			(FSolverReal)GetFluidDensity(PropertyCollection, 0.f),
+			FSolverVec2(GetWeightedFloatPressure(PropertyCollection, 0.f)),  // These getters also initialize the property indices, so keep before SetMultipliers
+			WorldScale);
+		SetMultipliers(PropertyCollection, Weightmaps);
+	}
+
 	// Construct an uninitialized field. Mesh, properties, and velocity will have to be set for this field to be valid.
 	FVelocityAndPressureField()
 		: Offset(INDEX_NONE)
@@ -50,15 +73,18 @@ public:
 		, LiftIndex(ForceInit)
 		, FluidDensityIndex(ForceInit)
 		, PressureIndex(ForceInit)
+		, WindVelocityIndex(ForceInit)
 	{
 		SetProperties(FSolverVec2(0.), FSolverVec2(0.), (FSolverReal)0., FSolverVec2(0.));
 	}
 
 	~FVelocityAndPressureField() {}
 
-	CHAOS_API void UpdateForces(const FSolverParticles& InParticles, const FSolverReal /*Dt*/);
+	template<typename SolverParticlesOrRange>
+	CHAOS_API void UpdateForces(const SolverParticlesOrRange& InParticles, const FSolverReal /*Dt*/);
 
-	inline void Apply(FSolverParticles& InParticles, const FSolverReal Dt, const int32 Index) const
+	template<typename SolverParticlesOrRange>
+	inline void Apply(SolverParticlesOrRange& InParticles, const FSolverReal Dt, const int32 Index) const
 	{
 		checkSlow(Index >= Offset && Index < Offset + NumParticles);  // The index should always match the original triangle mesh range
 
@@ -69,12 +95,33 @@ public:
 		}
 	}
 
+	inline void Apply(FSolverParticlesRange& InParticles, const FSolverReal Dt) const
+	{
+		for (int32 Index = 0; Index < InParticles.GetRangeSize(); ++Index)
+		{
+			const TArray<int32>& ElementIndices = PointToTriangleMap[Index];
+			for (const int32 ElementIndex : ElementIndices)
+			{
+				InParticles.Acceleration(Index) += InParticles.InvM(Index) * Forces[ElementIndex];
+			}
+		}
+	}
+
+	// This version will not load WindVelocity from the config. Call SetVelocity to set it explicitly.
 	CHAOS_API void SetProperties(
 		const FCollectionPropertyConstFacade& PropertyCollection,
 		const TMap<FString, TConstArrayView<FRealSingle>>& Weightmaps,
 		FSolverReal WorldScale,
 		bool bEnableAerodynamics);
 
+	// This version will load WindVelocity from the config
+	CHAOS_API void SetPropertiesAndWind(
+		const FCollectionPropertyConstFacade& PropertyCollection,
+		const TMap<FString, TConstArrayView<FRealSingle>>& Weightmaps,
+		FSolverReal WorldScale,
+		bool bEnableAerodynamics,
+		const FSolverVec3& SolverWind);
+
 	UE_DEPRECATED(5.3, "Use SetProperties(const FCollectionPropertyConstFacade&, const TMap<FString, TConstArrayView<FRealSingle>>&, FSolverReal, bool) instead.")
 	void SetProperties(const FCollectionPropertyConstFacade& PropertyCollection, FSolverReal WorldScale)
 	{
@@ -126,7 +173,8 @@ private:
 		const TConstArrayView<FRealSingle>& LiftMultipliers,
 		const TConstArrayView<FRealSingle>& PressureMultipliers);
 
-	void UpdateField(const FSolverParticles& InParticles, int32 ElementIndex, const FSolverVec3& InVelocity, const FSolverReal Cd, const FSolverReal Cl, const FSolverReal Cp)
+	template<typename SolverParticlesOrRange>
+	void UpdateField(const SolverParticlesOrRange& InParticles, int32 ElementIndex, const FSolverVec3& InVelocity, const FSolverReal Cd, const FSolverReal Cl, const FSolverReal Cp)
 	{
 		const TVec3<int32>& Element = Elements[ElementIndex];
 
@@ -152,7 +200,8 @@ private:
 			(Cl - Cd) * VDotN * V - Cl * VSquare * N) - DoubleArea * (FSolverReal)0.5 * Cp * N; // N points in the opposite direction of the actual mesh normals
 	}
 
-	void UpdateField(const FSolverParticles& InParticles, int32 ElementIndex, const FSolverVec3& InVelocity, const FSolverReal Cd, const FSolverReal Cl, const FSolverReal Cp, const FSolverReal MaxVelocitySquared)
+	template<typename SolverParticlesOrRange>
+	void UpdateField(const SolverParticlesOrRange& InParticles, int32 ElementIndex, const FSolverVec3& InVelocity, const FSolverReal Cd, const FSolverReal Cl, const FSolverReal Cp, const FSolverReal MaxVelocitySquared)
 	{
 		checkSlow(MaxVelocitySquared > (FSolverReal)0);
 
@@ -208,6 +257,7 @@ private:
 	UE_CHAOS_DECLARE_PROPERTYCOLLECTION_NAME(Lift, float);
 	UE_CHAOS_DECLARE_PROPERTYCOLLECTION_NAME(FluidDensity, float);
 	UE_CHAOS_DECLARE_PROPERTYCOLLECTION_NAME(Pressure, float);
+	UE_CHAOS_DECLARE_PROPERTYCOLLECTION_NAME(WindVelocity, FVector3f);
 };
 
 using FVelocityField UE_DEPRECATED(5.1, "Chaos::Softs::FVelocityField has been renamed FVelocityAndPressureField to match its new behavior.") = FVelocityAndPressureField;
diff --git a/Engine/Source/Runtime/Experimental/Chaos/Public/Chaos/XPBDAnisotropicBendingConstraints.h b/Engine/Source/Runtime/Experimental/Chaos/Public/Chaos/XPBDAnisotropicBendingConstraints.h
--- a/Engine/Source/Runtime/Experimental/Chaos/Public/Chaos/XPBDAnisotropicBendingConstraints.h
+++ b/Engine/Source/Runtime/Experimental/Chaos/Public/Chaos/XPBDAnisotropicBendingConstraints.h
@@ -29,6 +29,13 @@ public:
 		return IsXPBDAnisoBendingStiffnessWarpEnabled(PropertyCollection, false);
 	}
 
+	CHAOS_API FXPBDAnisotropicBendingConstraints(const FSolverParticlesRange& InParticles,
+		const FTriangleMesh& TriangleMesh,
+		const TArray<TVec3<FVec2f>>& FaceVertexPatternPositions,
+		const TMap<FString, TConstArrayView<FRealSingle>>& WeightMaps,
+		const FCollectionPropertyConstFacade& PropertyCollection,
+		bool bTrimKinematicConstraints = false);
+
 	CHAOS_API FXPBDAnisotropicBendingConstraints(const FSolverParticles& InParticles,
 		int32 InParticleOffset,
 		int32 InParticleCount,
@@ -62,7 +69,8 @@ public:
 
 	virtual ~FXPBDAnisotropicBendingConstraints() override {}
 
-	void Init(const FSolverParticles& InParticles)
+	template<typename SolverParticlesOrRange>
+	void Init(const SolverParticlesOrRange& InParticles)
 	{ 
 		Lambdas.Reset();
 		Lambdas.AddZeroed(Constraints.Num());
@@ -91,7 +99,8 @@ public:
 		DampingRatio.ApplyValues();
 	}
 
-	CHAOS_API void Apply(FSolverParticles& Particles, const FSolverReal Dt) const;
+	template<typename SolverParticlesOrRange>
+	CHAOS_API void Apply(SolverParticlesOrRange& Particles, const FSolverReal Dt) const;
 
 	const TArray<int32>& GetConstraintsPerColorStartIndex() const { return ConstraintsPerColorStartIndex; }
 	const TArray<FSolverVec3>& GetWarpWeftBiasBaseMultipliers() const { return WarpWeftBiasBaseMultipliers; }
@@ -117,11 +126,13 @@ public:
 
 
 private:
-	CHAOS_API void InitColor(const FSolverParticles& InParticles);
-	CHAOS_API void ApplyHelper(FSolverParticles& Particles, const FSolverReal Dt, const int32 ConstraintIndex, const FSolverVec3& ExpStiffnessValues, 
+	template<typename SolverParticlesOrRange>
+	void InitColor(const SolverParticlesOrRange& InParticles);
+	template<typename SolverParticlesOrRange>
+	void ApplyHelper(SolverParticlesOrRange& Particles, const FSolverReal Dt, const int32 ConstraintIndex, const FSolverVec3& ExpStiffnessValues,
 		const FSolverVec3& ExpBucklingStiffnessValues, const FSolverReal DampingRatioValue) const;
 
-	CHAOS_API TArray<FSolverVec3> GenerateWarpWeftBiasBaseMultipliers(const TArray<TVec3<FVec2f>>& FaceVertexPatternPositions, const FTriangleMesh& TriangleMesh) const;
+	TArray<FSolverVec3> GenerateWarpWeftBiasBaseMultipliers(const TArray<TVec3<FVec2f>>& FaceVertexPatternPositions, const FTriangleMesh& TriangleMesh) const;
 
 	TConstArrayView<FRealSingle> GetRestAngleMapFromCollection(
 		const TMap<FString, TConstArrayView<FRealSingle>>& WeightMaps,
diff --git a/Engine/Source/Runtime/Experimental/Chaos/Public/Chaos/XPBDAxialSpringConstraints.h b/Engine/Source/Runtime/Experimental/Chaos/Public/Chaos/XPBDAxialSpringConstraints.h
--- a/Engine/Source/Runtime/Experimental/Chaos/Public/Chaos/XPBDAxialSpringConstraints.h
+++ b/Engine/Source/Runtime/Experimental/Chaos/Public/Chaos/XPBDAxialSpringConstraints.h
@@ -25,6 +25,23 @@ public:
 	static constexpr FSolverReal MinStiffness = (FSolverReal)1e-4; // Stiffness below this will be considered 0 since all of our calculations are actually based on 1 / stiffness.
 	static constexpr FSolverReal MaxStiffness = (FSolverReal)1e7;
 
+	FXPBDAxialSpringConstraints(
+		const FSolverParticlesRange& Particles,
+		const TArray<TVec3<int32>>& InConstraints,
+		const TConstArrayView<FRealSingle>& StiffnessMultipliers,
+		const FSolverVec2& InStiffness,
+		bool bTrimKinematicConstraints)
+		: Base(
+			Particles,
+			InConstraints,
+			StiffnessMultipliers,
+			InStiffness,
+			bTrimKinematicConstraints,
+			MaxStiffness)
+	{
+		Lambdas.Init(0.f, Constraints.Num());
+	}
+
 	FXPBDAxialSpringConstraints(
 		const FSolverParticles& Particles,
 		int32 ParticleOffset,
@@ -54,7 +71,8 @@ public:
 
 	void Init() const { for (FSolverReal& Lambda : Lambdas) { Lambda = (FSolverReal)0.; } }
 
-	void Apply(FSolverParticles& Particles, const FSolverReal Dt) const
+	template<typename SolverParticlesOrRange>
+	void Apply(SolverParticlesOrRange& Particles, const FSolverReal Dt) const
 	{
 		SCOPE_CYCLE_COUNTER(STAT_XPBD_AxialSpring);
 		if (!Stiffness.HasWeightMap())
@@ -114,7 +132,8 @@ public:
 	}
 
 private:
-	FSolverVec3 GetDelta(const FSolverParticles& Particles, const FSolverReal Dt, const int32 InConstraintIndex, const FSolverReal StiffnessValue) const
+	template<typename SolverParticlesOrRange>
+	FSolverVec3 GetDelta(const SolverParticlesOrRange& Particles, const FSolverReal Dt, const int32 InConstraintIndex, const FSolverReal StiffnessValue) const
 	{
 		const TVector<int32, 3>& Constraint = Constraints[InConstraintIndex];
 		const int32 i1 = Constraint[0];
@@ -175,6 +194,21 @@ public:
 		return IsXPBDAreaSpringStiffnessEnabled(PropertyCollection, false);
 	}
 
+	FXPBDAreaSpringConstraints(
+		const FSolverParticlesRange& Particles,
+		const TArray<TVec3<int32>>& InConstraints,
+		const TMap<FString, TConstArrayView<FRealSingle>>& WeightMaps,
+		const FCollectionPropertyConstFacade& PropertyCollection,
+		bool bTrimKinematicConstraints)
+		: FXPBDAxialSpringConstraints(
+			Particles,
+			InConstraints,
+			WeightMaps.FindRef(GetXPBDAreaSpringStiffnessString(PropertyCollection, XPBDAreaSpringStiffnessName.ToString())),
+			FSolverVec2(GetWeightedFloatXPBDAreaSpringStiffness(PropertyCollection, MaxStiffness)),
+			bTrimKinematicConstraints)
+		, XPBDAreaSpringStiffnessIndex(PropertyCollection)
+	{}
+
 	FXPBDAreaSpringConstraints(
 		const FSolverParticles& Particles,
 		int32 ParticleOffset,
diff --git a/Engine/Source/Runtime/Experimental/Chaos/Public/Chaos/XPBDBendingConstraints.h b/Engine/Source/Runtime/Experimental/Chaos/Public/Chaos/XPBDBendingConstraints.h
--- a/Engine/Source/Runtime/Experimental/Chaos/Public/Chaos/XPBDBendingConstraints.h
+++ b/Engine/Source/Runtime/Experimental/Chaos/Public/Chaos/XPBDBendingConstraints.h
@@ -30,6 +30,42 @@ public:
 		return IsXPBDBendingElementStiffnessEnabled(PropertyCollection, false);
 	}
 
+	FXPBDBendingConstraints(const FSolverParticlesRange& InParticles,
+		TArray<TVec4<int32>>&& InConstraints,
+		const TMap<FString, TConstArrayView<FRealSingle>>& WeightMaps,
+		const FCollectionPropertyConstFacade& PropertyCollection,
+		bool bTrimKinematicConstraints = false)
+		: Base(
+			InParticles,
+			MoveTemp(InConstraints),
+			WeightMaps.FindRef(GetXPBDBendingElementStiffnessString(PropertyCollection, XPBDBendingElementStiffnessName.ToString())),
+			WeightMaps.FindRef(GetXPBDBucklingStiffnessString(PropertyCollection, XPBDBucklingStiffnessName.ToString())),
+			GetRestAngleMapFromCollection(WeightMaps, PropertyCollection),
+			FSolverVec2(GetWeightedFloatXPBDBendingElementStiffness(PropertyCollection, MaxStiffness)),
+			(FSolverReal)GetXPBDBucklingRatio(PropertyCollection, 0.f),
+			FSolverVec2(GetWeightedFloatXPBDBucklingStiffness(PropertyCollection, MaxStiffness)),
+			GetRestAngleValueFromCollection(PropertyCollection),
+			(ERestAngleConstructionType)GetXPBDRestAngleType(PropertyCollection, (int32)ERestAngleConstructionType::Use3DRestAngles),
+			bTrimKinematicConstraints,
+			MaxStiffness)
+		, DampingRatio(
+			FSolverVec2(GetWeightedFloatXPBDBendingElementDamping(PropertyCollection, MinDamping)).ClampAxes(MinDamping, MaxDamping),
+			WeightMaps.FindRef(GetXPBDBendingElementDampingString(PropertyCollection, XPBDBendingElementDampingName.ToString())),
+			TConstArrayView<TVec2<int32>>(ConstraintSharedEdges),
+			ParticleOffset,
+			ParticleCount)
+		, XPBDBendingElementStiffnessIndex(PropertyCollection)
+		, XPBDBendingElementDampingIndex(PropertyCollection)
+		, XPBDBucklingRatioIndex(PropertyCollection)
+		, XPBDBucklingStiffnessIndex(PropertyCollection)
+		, XPBDFlatnessRatioIndex(PropertyCollection)
+		, XPBDRestAngleIndex(PropertyCollection)
+		, XPBDRestAngleTypeIndex(PropertyCollection)
+	{
+		Lambdas.Init((FSolverReal)0., Constraints.Num());
+		InitColor(InParticles);
+	}
+
 	FXPBDBendingConstraints(const FSolverParticles& InParticles,
 		int32 InParticleOffset,
 		int32 InParticleCount,
@@ -190,7 +226,8 @@ public:
 
 	virtual ~FXPBDBendingConstraints() override {}
 
-	void Init(const FSolverParticles& InParticles)
+	template<typename SolverParticlesOrRange>
+	void Init(const SolverParticlesOrRange& InParticles)
 	{ 
 		Lambdas.Reset();
 		Lambdas.AddZeroed(Constraints.Num());
@@ -223,7 +260,8 @@ public:
 		DampingRatio.ApplyValues();
 	}
 
-	CHAOS_API void Apply(FSolverParticles& Particles, const FSolverReal Dt) const;
+	template<typename SolverParticlesOrRange>
+	CHAOS_API void Apply(SolverParticlesOrRange& Particles, const FSolverReal Dt) const;
 
 	const TArray<int32>& GetConstraintsPerColorStartIndex() const { return ConstraintsPerColorStartIndex; }
 
@@ -248,8 +286,10 @@ public:
 
 	CHAOS_API FSolverReal ComputeTotalEnergy(const FSolverParticles& InParticles, const FSolverReal ExplicitStiffness = -1.f);
 private:
-	CHAOS_API void InitColor(const FSolverParticles& InParticles);
-	CHAOS_API void ApplyHelper(FSolverParticles& Particles, const FSolverReal Dt, const int32 ConstraintIndex, const FSolverReal ExpStiffnessValue, const FSolverReal ExpBucklingValue, const FSolverReal DampingRatioValue) const;
+	template<typename SolverParticlesOrRange>
+	CHAOS_API void InitColor(const SolverParticlesOrRange& InParticles);
+	template<typename SolverParticlesOrRange>
+	void ApplyHelper(SolverParticlesOrRange& Particles, const FSolverReal Dt, const int32 ConstraintIndex, const FSolverReal ExpStiffnessValue, const FSolverReal ExpBucklingValue, const FSolverReal DampingRatioValue) const;
 
 	TConstArrayView<FRealSingle> GetRestAngleMapFromCollection(
 		const TMap<FString, TConstArrayView<FRealSingle>>& WeightMaps,
diff --git a/Engine/Source/Runtime/Experimental/Chaos/Public/Chaos/XPBDSpringConstraints.h b/Engine/Source/Runtime/Experimental/Chaos/Public/Chaos/XPBDSpringConstraints.h
--- a/Engine/Source/Runtime/Experimental/Chaos/Public/Chaos/XPBDSpringConstraints.h
+++ b/Engine/Source/Runtime/Experimental/Chaos/Public/Chaos/XPBDSpringConstraints.h
@@ -48,6 +48,33 @@ public:
 	{
 		Lambdas.Init((FSolverReal)0., Constraints.Num());
 		InitColor(Particles);
+	}	
+	
+	template<int32 Valence, TEMPLATE_REQUIRES(Valence >= 2 && Valence <= 4)>
+	FXPBDSpringConstraints(
+		const FSolverParticlesRange& Particles,
+		const TArray<TVector<int32, Valence>>& InConstraints,
+		const TConstArrayView<FRealSingle>& StiffnessMultipliers,
+		const TConstArrayView<FRealSingle>& DampingMultipliers,
+		const FSolverVec2& InStiffness,
+		const FSolverVec2& InDampingRatio,
+		bool bTrimKinematicConstraints = false)
+		: Base(
+			Particles,
+			InConstraints,
+			StiffnessMultipliers,
+			InStiffness,
+			bTrimKinematicConstraints,
+			MaxStiffness)
+		, DampingRatio(
+			InDampingRatio.ClampAxes(MinDampingRatio, MaxDampingRatio),
+			DampingMultipliers,
+			TConstArrayView<TVec2<int32>>(Constraints),
+			ParticleOffset,
+			ParticleCount)
+	{
+		Lambdas.Init((FSolverReal)0., Constraints.Num());
+		InitColor(Particles);
 	}
 
 	template<int32 Valence, TEMPLATE_REQUIRES(Valence >= 2 && Valence <= 4)>
@@ -99,15 +126,19 @@ public:
 		DampingRatio.ApplyValues();
 	}
 
-	CHAOS_API void Apply(FSolverParticles& Particles, const FSolverReal Dt) const;
+	template<typename SolverParticlesOrRange>
+	CHAOS_API void Apply(SolverParticlesOrRange& Particles, const FSolverReal Dt) const;
 
 	const TArray<int32>& GetConstraintsPerColorStartIndex() const { return ConstraintsPerColorStartIndex; }
 
 private:
-	CHAOS_API void InitColor(const FSolverParticles& InParticles);
-	CHAOS_API void ApplyHelper(FSolverParticles& Particles, const FSolverReal Dt, const int32 ConstraintIndex, const FSolverReal StiffnessValue, const FSolverReal DampingRatioValue) const;
+	template<typename SolverParticlesOrRange>
+	CHAOS_API void InitColor(const SolverParticlesOrRange& InParticles);
+	template<typename SolverParticlesOrRange>
+	void ApplyHelper(SolverParticlesOrRange& Particles, const FSolverReal Dt, const int32 ConstraintIndex, const FSolverReal StiffnessValue, const FSolverReal DampingRatioValue) const;
 
-	FSolverVec3 GetDelta(const FSolverParticles& Particles, const FSolverReal Dt, const int32 ConstraintIndex, const FSolverReal StiffnessValue, const FSolverReal DampingRatioValue) const
+	template<typename SolverParticlesOrRange>
+	FSolverVec3 GetDelta(const SolverParticlesOrRange& Particles, const FSolverReal Dt, const int32 ConstraintIndex, const FSolverReal StiffnessValue, const FSolverReal DampingRatioValue) const
 	{
 		const TVec2<int32>& Constraint = Constraints[ConstraintIndex];
 
@@ -170,6 +201,24 @@ public:
 		return IsXPBDEdgeSpringStiffnessEnabled(PropertyCollection, false);
 	}
 
+	FXPBDEdgeSpringConstraints(
+		const FSolverParticlesRange& Particles,
+		const TArray<TVec3<int32>>& InConstraints,
+		const TMap<FString, TConstArrayView<FRealSingle>>& WeightMaps,
+		const FCollectionPropertyConstFacade& PropertyCollection,
+		bool bTrimKinematicConstraints = false)
+		: FXPBDSpringConstraints(
+			Particles,
+			InConstraints,
+			WeightMaps.FindRef(GetXPBDEdgeSpringStiffnessString(PropertyCollection, XPBDEdgeSpringStiffnessName.ToString())),
+			WeightMaps.FindRef(GetXPBDEdgeSpringDampingString(PropertyCollection, XPBDEdgeSpringDampingName.ToString())),
+			FSolverVec2(GetWeightedFloatXPBDEdgeSpringStiffness(PropertyCollection, MaxStiffness)),
+			FSolverVec2(GetWeightedFloatXPBDEdgeSpringDamping(PropertyCollection, MinDampingRatio)),
+			bTrimKinematicConstraints)
+		, XPBDEdgeSpringStiffnessIndex(PropertyCollection)
+		, XPBDEdgeSpringDampingIndex(PropertyCollection)
+	{}
+
 	FXPBDEdgeSpringConstraints(
 		const FSolverParticles& Particles,
 		int32 ParticleOffset,
@@ -247,6 +296,24 @@ public:
 		return IsXPBDBendingSpringStiffnessEnabled(PropertyCollection, false);
 	}
 
+	FXPBDBendingSpringConstraints(
+		const FSolverParticlesRange& Particles,
+		const TArray<TVec2<int32>>& InConstraints,
+		const TMap<FString, TConstArrayView<FRealSingle>>& WeightMaps,
+		const FCollectionPropertyConstFacade& PropertyCollection,
+		bool bTrimKinematicConstraints = false)
+		: FXPBDSpringConstraints(
+			Particles,
+			InConstraints,
+			WeightMaps.FindRef(GetXPBDBendingSpringStiffnessString(PropertyCollection, XPBDBendingSpringStiffnessName.ToString())),
+			WeightMaps.FindRef(GetXPBDBendingSpringDampingString(PropertyCollection, XPBDBendingSpringDampingName.ToString())),
+			FSolverVec2(GetWeightedFloatXPBDBendingSpringStiffness(PropertyCollection, MaxStiffness)),
+			FSolverVec2(GetWeightedFloatXPBDBendingSpringDamping(PropertyCollection, MinDampingRatio)),
+			bTrimKinematicConstraints)
+		, XPBDBendingSpringStiffnessIndex(PropertyCollection)
+		, XPBDBendingSpringDampingIndex(PropertyCollection)
+	{}
+
 	FXPBDBendingSpringConstraints(
 		const FSolverParticles& Particles,
 		int32 ParticleOffset,
diff --git a/Engine/Source/Runtime/Experimental/Chaos/Public/Chaos/XPBDStretchBiasElementConstraints.h b/Engine/Source/Runtime/Experimental/Chaos/Public/Chaos/XPBDStretchBiasElementConstraints.h
--- a/Engine/Source/Runtime/Experimental/Chaos/Public/Chaos/XPBDStretchBiasElementConstraints.h
+++ b/Engine/Source/Runtime/Experimental/Chaos/Public/Chaos/XPBDStretchBiasElementConstraints.h
@@ -34,6 +34,13 @@ public:
 		return IsXPBDAnisoStretchStiffnessWarpEnabled(PropertyCollection, false);
 	}
 
+	CHAOS_API FXPBDStretchBiasElementConstraints(const FSolverParticlesRange& InParticles,
+		const FTriangleMesh& TriangleMesh,
+		const TArray<TVec3<FVec2f>>& FaceVertexUVs,
+		const TMap<FString, TConstArrayView<FRealSingle>>& WeightMaps,
+		const FCollectionPropertyConstFacade& PropertyCollection,
+		bool bTrimKinematicConstraints = false);
+
 	CHAOS_API FXPBDStretchBiasElementConstraints(const FSolverParticles& InParticles,
 		int32 ParticleOffset,
 		int32 ParticleCount,
@@ -96,7 +103,8 @@ public:
 		WeftScale.ApplyValues();
 	}
 
-	CHAOS_API void Apply(FSolverParticles& Particles, const FSolverReal Dt) const;
+	template<typename SolverParticlesOrRange>
+	CHAOS_API void Apply(SolverParticlesOrRange& Particles, const FSolverReal Dt) const;
 	CHAOS_API void CalculateUVStretch(const int32 ConstraintIndex, const FSolverVec3& P0, const FSolverVec3& P1, const FSolverVec3& P2, FSolverVec3& DXDu, FSolverVec3& DXDv) const;
 	
 	const TArray<TVec3<int32>>& GetConstraints() const { return Constraints; }
@@ -130,10 +138,13 @@ public:
 	CHAOS_API void AddInternalForceDifferential(const FSolverParticles& InParticles, const TArray<TVector<FSolverReal, 3>>& DeltaParticles, TArray<TVector<FSolverReal, 3>>& ndf);
 
 private:
-	CHAOS_API void InitConstraintsAndRestData(const FSolverParticles& InParticles, const FTriangleMesh& TriangleMesh,
+	template<typename SolverParticlesOrRange>
+	void InitConstraintsAndRestData(const SolverParticlesOrRange& InParticles, const FTriangleMesh& TriangleMesh,
 		const TArray<TVec3<FVec2f>>& FaceVertexUVs, const bool bUse3dRestLengths, const bool bTrimKinematicConstraints);
-	CHAOS_API void InitColor(const FSolverParticles& InParticles);
-	CHAOS_API void ApplyHelper(FSolverParticles& Particles, const FSolverReal Dt, const int32 ConstraintIndex, const FSolverVec3& ExpStiffnessValue, const FSolverReal DampingRatioValue, const FSolverReal WarpScaleValue, const FSolverReal WeftScaleValue) const;
+	template<typename SolverParticlesOrRange>
+	void InitColor(const SolverParticlesOrRange& InParticles);
+	template<typename SolverParticlesOrRange>
+	void ApplyHelper(SolverParticlesOrRange& Particles, const FSolverReal Dt, const int32 ConstraintIndex, const FSolverVec3& ExpStiffnessValue, const FSolverReal DampingRatioValue, const FSolverReal WarpScaleValue, const FSolverReal WeftScaleValue) const;
 
 	TArray<TVec3<int32>> Constraints;
 	const int32 ParticleOffset;
