Updated ChaosDebugDrawComponent to render in the Editor world

diff --git a/Engine/Source/Runtime/Experimental/ChaosSolverEngine/Private/Chaos/ChaosDebugDrawComponent.cpp b/Engine/Source/Runtime/Experimental/ChaosSolverEngine/Private/Chaos/ChaosDebugDrawComponent.cpp
--- a/Engine/Source/Runtime/Experimental/ChaosSolverEngine/Private/Chaos/ChaosDebugDrawComponent.cpp
+++ b/Engine/Source/Runtime/Experimental/ChaosSolverEngine/Private/Chaos/ChaosDebugDrawComponent.cpp
@@ -70,7 +70,7 @@ float CommandLifeTime(const Chaos::FLatentDrawCommand& Command, const bool bIsPa
 	return Command.LifeTime;
 }
 
-void DebugDrawChaos(AActor* DebugDrawActor, const TArray<Chaos::FLatentDrawCommand>& DrawCommands, bool bIsPaused)
+void DebugDrawChaos(const AActor* DebugDrawActor, const TArray<Chaos::FLatentDrawCommand>& DrawCommands, const bool bIsPaused)
 {
 	using namespace Chaos;
 
@@ -79,13 +79,13 @@ void DebugDrawChaos(AActor* DebugDrawActor, const TArray<Chaos::FLatentDrawComma
 		return;
 	}
 
-	UWorld* World = DebugDrawActor->GetWorld();
+	const UWorld* World = DebugDrawActor->GetWorld();
 	if (World == nullptr)
 	{
 		return;
 	}
 
-	if (!World->IsGameWorld())
+	if (World->IsPreviewWorld())
 	{
 		return;
 	}
@@ -148,7 +148,7 @@ void DebugDrawChaos(AActor* DebugDrawActor, const TArray<Chaos::FLatentDrawComma
 	{
 		for (const FLatentDrawCommand& Command : DrawCommands)
 		{
-			AActor* Actor = (Command.TestBaseActor) ? Command.TestBaseActor : DebugDrawActor;
+			const AActor* Actor = (Command.TestBaseActor) ? Command.TestBaseActor : DebugDrawActor;
 
 			switch (Command.Type)
 			{
@@ -205,6 +205,8 @@ UChaosDebugDrawComponent::UChaosDebugDrawComponent()
 	PrimaryComponentTick.bCanEverTick = true;
 	PrimaryComponentTick.SetTickFunctionEnable(true);
 	PrimaryComponentTick.TickGroup = TG_PostUpdateWork;
+
+	bTickInEditor = true;
 }
 
 void UChaosDebugDrawComponent::BeginDestroy()
@@ -245,6 +247,31 @@ void UChaosDebugDrawComponent::EndPlay(const EEndPlayReason::Type EndPlayReason)
 #endif
 }
 
+#if WITH_EDITOR && CHAOS_DEBUG_DRAW
+void UChaosDebugDrawComponent::OnRegister()
+{
+	Super::OnRegister();
+	const UWorld* World = GetWorld();
+	check(World);
+	// GameWorld is handled by BeginPlay/EndPlay
+	if (!World->IsGameWorld() && !World->IsPreviewWorld())
+	{
+		Chaos::FDebugDrawQueue::GetInstance().SetConsumerActive(this, /*bConsumerActive*/true);
+	}
+}
+void UChaosDebugDrawComponent::OnUnregister()
+{
+	const UWorld* World = GetWorld();
+	check(World);
+	// GameWorld is handled by BeginPlay/EndPlay
+	if (!World->IsGameWorld() && !World->IsPreviewWorld())
+	{
+		Chaos::FDebugDrawQueue::GetInstance().SetConsumerActive(this, /*bConsumerActive*/false);
+	}
+	Super::OnUnregister();
+}
+#endif // WITH_EDITOR && CHAOS_DEBUG_DRAW
+
 void UChaosDebugDrawComponent::TickComponent(float DeltaTime, enum ELevelTick TickType, FActorComponentTickFunction *ThisTickFunction)
 {
 	using namespace Chaos;
@@ -314,10 +341,10 @@ void UChaosDebugDrawComponent::BindWorldDelegates()
 #endif
 }
 
-void UChaosDebugDrawComponent::HandlePostWorldInitialization(UWorld* World, const UWorld::InitializationValues IVS)
+void UChaosDebugDrawComponent::HandlePostWorldInitialization(UWorld* World, const UWorld::InitializationValues)
 {
 #if CHAOS_DEBUG_DRAW
-	if ((World != nullptr) && World->IsGameWorld())
+	if ((World != nullptr) && !IsRunningCommandlet() && !World->IsPreviewWorld())
 	{
 		CreateDebugDrawActor(World);
 	}
diff --git a/Engine/Source/Runtime/Experimental/ChaosSolverEngine/Public/Chaos/ChaosDebugDrawComponent.h b/Engine/Source/Runtime/Experimental/ChaosSolverEngine/Public/Chaos/ChaosDebugDrawComponent.h
--- a/Engine/Source/Runtime/Experimental/ChaosSolverEngine/Public/Chaos/ChaosDebugDrawComponent.h
+++ b/Engine/Source/Runtime/Experimental/ChaosSolverEngine/Public/Chaos/ChaosDebugDrawComponent.h
@@ -28,6 +28,12 @@ public:
 
 	static CHAOSSOLVERENGINE_API void BindWorldDelegates();
 
+protected:
+#if WITH_EDITOR && CHAOS_DEBUG_DRAW
+	virtual void OnRegister() override;
+	virtual void OnUnregister() override;
+#endif // WITH_EDITOR && CHAOS_DEBUG_DRAW
+
 private:
 	static void HandlePostWorldInitialization(UWorld* World, const UWorld::InitializationValues IVS);
 	static void CreateDebugDrawActor(UWorld* World);
