The application crashes in-game: "Chaos::TEventContainer<Chaos::FCollisionEventData>::DispatchConsumerData".

#tests Tested in PIE

diff --git a/Engine/Source/Runtime/Experimental/Chaos/Public/EventManager.h b/Engine/Source/Runtime/Experimental/Chaos/Public/EventManager.h
--- a/Engine/Source/Runtime/Experimental/Chaos/Public/EventManager.h
+++ b/Engine/Source/Runtime/Experimental/Chaos/Public/EventManager.h
@@ -307,24 +307,24 @@ namespace Chaos
 			// Only take this path if we have fewer Events than Handlers
 			if (Map && Map->Num() + HandlersNotInProxyOwnerMap.Num() < HandlerMap.Num())
 			{
-				TSet<FEventHandlerPtr> UniqueHandlers;
-				UniqueHandlers.Reserve(Map->Num());
-				for (auto& KeyValue : *Map) // Only iterating over objects that are associated with events here
+				if (!ProxyOwnerToHandlerMap.IsEmpty())
 				{
-					const IPhysicsProxyBase* Proxy = KeyValue.Get<0>();
-					const UObject* Owner = Proxy->GetOwner();
-					if (ProxyOwnerToHandlerMap.Contains(Owner))
+					TSet<FEventHandlerPtr> UniqueHandlers;
+					UniqueHandlers.Reserve(Map->Num());
+					for (auto& KeyValue : *Map) // Only iterating over objects that are associated with events here
 					{
+						const IPhysicsProxyBase* Proxy = KeyValue.Get<0>();
+						const UObject* Owner = Proxy->GetOwner();
 						for (TMultiMap<UObject*, FEventHandlerPtr>::TConstKeyIterator It = ProxyOwnerToHandlerMap.CreateConstKeyIterator(Owner); It; ++It)
 						{
 							UniqueHandlers.Add(It.Value());
 						}
 					}
-				}
 
-				for (const FEventHandlerPtr Handler : UniqueHandlers)
-				{
-					Handler->HandleEvent(Buffer);
+					for (const FEventHandlerPtr Handler : UniqueHandlers)
+					{
+						Handler->HandleEvent(Buffer);
+					}
 				}
 
 				for (const TPair<void*, FEventHandlerPtr>& Pair : HandlersNotInProxyOwnerMap)
