Chaos - fix clusterunion manager and GC conflict when we receive a ClusterUnion Add operation on the same frame as a GC break.

This is caused by the fact that the Add operation is deferred to after the GC break is processed. The Add would then add the now-empty and disabled GC root to the cluster union.

#tested tested in PIE

diff --git a/Engine/Source/Runtime/Experimental/Chaos/Private/Chaos/ClusterUnionManager.cpp b/Engine/Source/Runtime/Experimental/Chaos/Private/Chaos/ClusterUnionManager.cpp
--- a/Engine/Source/Runtime/Experimental/Chaos/Private/Chaos/ClusterUnionManager.cpp
+++ b/Engine/Source/Runtime/Experimental/Chaos/Private/Chaos/ClusterUnionManager.cpp
@@ -24,6 +24,14 @@ namespace Chaos
 			TEXT("Whether to generate intercluster edges automatically when adding to a cluster union (and remove them when removing from the cluster union).")
 		);
 
+		// @tmp: To be removed
+		bool bChaosClusterUnionDoNotAddEmptyClusters = true;
+		FAutoConsoleVariableRef CVarChaosClusterUnionDoNotAddEmptyClusters(
+			TEXT("p.Chaos.ClusterUnion.DoNotAddEmptyClusters"),
+			bChaosClusterUnionDoNotAddEmptyClusters,
+			TEXT("Gating a risky bug fix.")
+		);
+
 		FRigidTransform3 GetParticleRigidFrameInClusterUnion(FPBDRigidParticleHandle* Child, const FRigidTransform3& ClusterWorldTM)
 		{
 			FRigidTransform3 Frame = FRigidTransform3::Identity;
@@ -583,10 +591,21 @@ namespace Chaos
 				continue;
 			}
 
-			bIsSleeping &= Handle->ObjectState() == EObjectStateType::Sleeping;
-
 			if (FPBDRigidClusteredParticleHandle* ClusterHandle = Handle->CastToClustered())
 			{
+				// If this particle is a broken cluster without any children, then we should not be adding it to a ClusterUnion.
+				// This can happen on the client if we receive an AddComponentToCluster on the same frame when we receive a GC repdata
+				// that breaks all the particles and releases its children. Unfortunately the Add operation has already been queued
+				if (bChaosClusterUnionDoNotAddEmptyClusters)
+				{
+					const bool bIsClusterThatHadChildren = (ClusterHandle->GetParticleType() == EParticleType::Clustered);
+					if (bIsClusterThatHadChildren && (ClusterHandle->Parent() == nullptr) && ClusterHandle->Disabled() && (ClusterHandle->ClusterIds().NumChildren == 0))
+					{
+						UE_LOG(LogChaos, Verbose, TEXT("FClusterUnionManager::HandleAddOperation rejecting particle because it is 'broken' %s"), *ClusterHandle->GetDebugName());
+						continue;
+					}
+				}
+
 				if (bReleaseClustersFirst)
 				{
 					TSet<FPBDRigidParticleHandle*> Children = MClustering.ReleaseClusterParticles(ClusterHandle, true);
@@ -608,6 +627,8 @@ namespace Chaos
 			{
 				FinalParticlesToAdd.Add(Handle);
 			}
+
+			bIsSleeping &= Handle->ObjectState() == EObjectStateType::Sleeping;
 		}
 
 		if (FinalParticlesToAdd.IsEmpty())
