Fix bug in GraphColoring from last change.

diff --git a/Engine/Source/Runtime/Experimental/Chaos/Public/Chaos/GraphColoring.h b/Engine/Source/Runtime/Experimental/Chaos/Public/Chaos/GraphColoring.h
--- a/Engine/Source/Runtime/Experimental/Chaos/Public/Chaos/GraphColoring.h
+++ b/Engine/Source/Runtime/Experimental/Chaos/Public/Chaos/GraphColoring.h
@@ -67,7 +67,7 @@ class FGraphColoring
 	template<typename T>
 	inline static TArray<TArray<int32>> ComputeGraphColoring(const TArray<TVector<int32, 2>>& Graph, const TDynamicParticles<T, 3>& InParticles, const int32 GraphParticlesStart, const int32 GraphParticlesEnd)
 	{
-		ComputeGraphColoringParticlesOrRange(Graph, InParticles, GraphParticlesStart, GraphParticlesEnd);
+		return ComputeGraphColoringParticlesOrRange(Graph, InParticles, GraphParticlesStart, GraphParticlesEnd);
 	}
 	template<typename T>
 	inline static TArray<TArray<int32>> ComputeGraphColoring(const TArray<TVector<int32, 2>>& Graph, const TDynamicParticles<T, 3>& InParticles)
