Iris - Hook up Rotator3f and Rotator3d NetSerializers.

diff --git a/Engine/Source/Runtime/Experimental/Iris/Core/Private/Iris/ReplicationState/DefaultPropertyNetSerializerInfos.cpp b/Engine/Source/Runtime/Experimental/Iris/Core/Private/Iris/ReplicationState/DefaultPropertyNetSerializerInfos.cpp
--- a/Engine/Source/Runtime/Experimental/Iris/Core/Private/Iris/ReplicationState/DefaultPropertyNetSerializerInfos.cpp
+++ b/Engine/Source/Runtime/Experimental/Iris/Core/Private/Iris/ReplicationState/DefaultPropertyNetSerializerInfos.cpp
@@ -50,6 +50,8 @@ UE_NET_IMPLEMENT_NAMED_STRUCT_NETSERIALIZER_INFO(NAME_Vector, FVectorNetSerializ
 UE_NET_IMPLEMENT_NAMED_STRUCT_NETSERIALIZER_INFO(NAME_Vector3f, FVector3fNetSerializer);
 UE_NET_IMPLEMENT_NAMED_STRUCT_NETSERIALIZER_INFO(NAME_Vector3d, FVector3dNetSerializer);
 UE_NET_IMPLEMENT_NAMED_STRUCT_NETSERIALIZER_INFO(NAME_Rotator, FRotatorNetSerializer);
+UE_NET_IMPLEMENT_NAMED_STRUCT_NETSERIALIZER_INFO(NAME_Rotator3f, FRotator3fNetSerializer);
+UE_NET_IMPLEMENT_NAMED_STRUCT_NETSERIALIZER_INFO(NAME_Rotator3d, FRotator3dNetSerializer);
 UE_NET_IMPLEMENT_NAMED_STRUCT_NETSERIALIZER_INFO(NAME_Quat, FUnitQuatNetSerializer);
 UE_NET_IMPLEMENT_NAMED_STRUCT_NETSERIALIZER_INFO(NAME_Quat4f, FUnitQuat4fNetSerializer);
 UE_NET_IMPLEMENT_NAMED_STRUCT_NETSERIALIZER_INFO(NAME_Quat4d, FUnitQuat4dNetSerializer);
