IAS - allow uploading content w/o credentials and let endpoint decide whether proper authentication is neeed

#rnx

diff --git a/Engine/Source/Runtime/Experimental/IoStoreOnDemand/Private/IoStoreOnDemand.cpp b/Engine/Source/Runtime/Experimental/IoStoreOnDemand/Private/IoStoreOnDemand.cpp
--- a/Engine/Source/Runtime/Experimental/IoStoreOnDemand/Private/IoStoreOnDemand.cpp
+++ b/Engine/Source/Runtime/Experimental/IoStoreOnDemand/Private/IoStoreOnDemand.cpp
@@ -917,14 +917,6 @@ TIoStatusOr<FIoStoreUploadParams> FIoStoreUploadParams::Parse(const TCHAR* Comma
 		}
 	}
 
-	if (Params.AccessKey.IsEmpty() &&
-		Params.SecretKey.IsEmpty() &&
-		Params.CredentialsFile.IsEmpty() &&
-		Params.CredentialsFileKeyName.IsEmpty())
-	{
-		return FIoStatus(EIoErrorCode::InvalidParameter, TEXT("Invalid credentials"));
-	}
-	
 	if (!Params.AccessKey.IsEmpty() && Params.SecretKey.IsEmpty())
 	{
 		return FIoStatus(EIoErrorCode::InvalidParameter, TEXT("Invalid secret key"));
