Constraints: make filters public

diff --git a/Engine/Source/Runtime/Experimental/Animation/Constraints/Public/TransformConstraint.h b/Engine/Source/Runtime/Experimental/Animation/Constraints/Public/TransformConstraint.h
--- a/Engine/Source/Runtime/Experimental/Animation/Constraints/Public/TransformConstraint.h
+++ b/Engine/Source/Runtime/Experimental/Animation/Constraints/Public/TransformConstraint.h
@@ -221,7 +221,6 @@ public:
 	UPROPERTY(EditAnywhere, BlueprintReadWrite, Category="Offset", meta=(EditCondition="bMaintainOffset"))
 	FVector OffsetTranslation = FVector::ZeroVector;
 
-protected:
 	/** Defines which translation axis is constrained. */
 	UPROPERTY(BlueprintReadWrite, Category = "Axis Filter")
 	FFilterOptionPerAxis AxisFilter;
@@ -270,7 +269,6 @@ public:
 	UPROPERTY(EditAnywhere, BlueprintReadWrite, Category="Offset", meta=(EditCondition="bMaintainOffset"))
 	FQuat OffsetRotation = FQuat::Identity;
 
-protected:
 	/** Defines which rotation axis is constrained. */
 	UPROPERTY(BlueprintReadWrite, Category = "Axis Filter")
 	FFilterOptionPerAxis AxisFilter;
@@ -318,7 +316,6 @@ public:
 	UPROPERTY(EditAnywhere, BlueprintReadWrite, Category="Offset", meta=(EditCondition="bMaintainOffset"))
 	FVector OffsetScale = FVector::OneVector;
 
-protected:
 	/** Defines which scale axis is constrained. */
 	UPROPERTY(BlueprintReadWrite, Category = "Axis Filter")
 	FFilterOptionPerAxis AxisFilter;
@@ -378,7 +375,6 @@ public:
 	UPROPERTY(EditAnywhere, BlueprintReadWrite, Category="Offset", meta=(EditCondition="bMaintainOffset"))
 	FTransform OffsetTransform = FTransform::Identity;
 
-protected:
 	/** Defines whether we propagate the parent scale. */
 	UPROPERTY(EditAnywhere, BlueprintReadWrite, Category="Properties")
 	bool bScaling = false;
