ASAN crash
Remove GC proxy from stall events

#localization none
#tests test in PIE

diff --git a/Engine/Source/Runtime/Experimental/Chaos/Private/PBDRigidsSolver.cpp b/Engine/Source/Runtime/Experimental/Chaos/Private/PBDRigidsSolver.cpp
--- a/Engine/Source/Runtime/Experimental/Chaos/Private/PBDRigidsSolver.cpp
+++ b/Engine/Source/Runtime/Experimental/Chaos/Private/PBDRigidsSolver.cpp
@@ -793,11 +793,27 @@ namespace Chaos
 
 
 		EnqueueCommandImmediate([InProxy, this]()
-			{
-				GeometryCollectionPhysicsProxies_Internal.RemoveSingle(InProxy);
-				InProxy->SyncBeforeDestroy();
-				PendingDestroyGeometryCollectionPhysicsProxy.Add(InProxy);
-			});
+		{
+			// Generally need to remove stale events for particles that no longer exist
+			GetEventManager()->template ClearEvents<FCollisionEventData>(EEventType::Collision, [InProxy]
+			(FCollisionEventData& EventDataInOut)
+				{
+					const FCollisionDataArray& CollisionData = EventDataInOut.CollisionData.AllCollisionsArray;
+					if (CollisionData.Num() > 0)
+					{
+						check(InProxy);
+						const TArray<int32>* const CollisionIndices = EventDataInOut.PhysicsProxyToCollisionIndices.PhysicsProxyToIndicesMap.Find(InProxy);
+						if (CollisionIndices)
+						{
+							EventDataInOut.PhysicsProxyToCollisionIndices.PhysicsProxyToIndicesMap.Remove(InProxy);
+						}
+					}
+				});
+
+			GeometryCollectionPhysicsProxies_Internal.RemoveSingle(InProxy);
+			InProxy->SyncBeforeDestroy();
+			PendingDestroyGeometryCollectionPhysicsProxy.Add(InProxy);
+		});
 	}
 
 	void FPBDRigidsSolver::RegisterObject(FClusterUnionPhysicsProxy* Proxy)
