cluster union syncing needs to better handle cases where the GC proxy pointer can be almost immediately re-used

diff --git a/Engine/Source/Runtime/Engine/Private/PhysicsEngine/ClusterUnionComponent.cpp b/Engine/Source/Runtime/Engine/Private/PhysicsEngine/ClusterUnionComponent.cpp
--- a/Engine/Source/Runtime/Engine/Private/PhysicsEngine/ClusterUnionComponent.cpp
+++ b/Engine/Source/Runtime/Engine/Private/PhysicsEngine/ClusterUnionComponent.cpp
@@ -954,7 +954,8 @@ void UClusterUnionComponent::SyncClusterUnionFromProxy()
 		for (const Chaos::FClusterUnionChildData& ChildData : FullData.ChildParticles)
 		{
 			// Using the scene's proxy to component mapping let's us detect a component physics state was destroyed.
-			if (UPrimitiveComponent* Component = Scene->GetOwningComponent<UPrimitiveComponent>(ChildData.Proxy))
+			UPrimitiveComponent* Component = Scene->GetOwningComponent<UPrimitiveComponent>(ChildData.Proxy);
+			if (Component && reinterpret_cast<void*>(Component) == ChildData.CachedOwner)
 			{
 				FMappedComponentKey WrappedComponentKey(Component);
 				int32 PrevMappedDataNum = MappedData.Num();
diff --git a/Engine/Source/Runtime/Experimental/Chaos/Private/PhysicsProxy/ClusterUnionPhysicsProxy.cpp b/Engine/Source/Runtime/Experimental/Chaos/Private/PhysicsProxy/ClusterUnionPhysicsProxy.cpp
--- a/Engine/Source/Runtime/Experimental/Chaos/Private/PhysicsProxy/ClusterUnionPhysicsProxy.cpp
+++ b/Engine/Source/Runtime/Experimental/Chaos/Private/PhysicsProxy/ClusterUnionPhysicsProxy.cpp
@@ -417,6 +417,7 @@ namespace Chaos
 			ConvertedData.ParticleIdx = InData.ParticleIdx;
 			ConvertedData.ChildToParent = InData.ChildToParent;
 			ConvertedData.Proxy = InData.Proxy;
+			ConvertedData.CachedOwner = InData.CachedOwner;
 			ConvertedData.BoneId = InData.BoneId;
 			SyncedData_External.ChildParticles.Add(ConvertedData);
 		}
@@ -576,6 +577,7 @@ namespace Chaos
 						Data.ChildToParent = FRigidTransform3::Identity;
 					}
 					Data.Proxy = Proxy;
+					Data.CachedOwner = reinterpret_cast<void*>(Proxy->GetOwner());
 					
 					if (Proxy->GetType() == EPhysicsProxyType::GeometryCollectionType)
 					{
@@ -624,6 +626,9 @@ namespace Chaos
 			FDirtyClusterUnionParticleData ConvertedData;
 			ConvertedData.ParticleIdx = Data.ParticleIdx;
 			ConvertedData.ChildToParent = Data.ChildToParent;
+			ConvertedData.Proxy = Data.Proxy;
+			ConvertedData.CachedOwner = Data.CachedOwner;
+			ConvertedData.BoneId = Data.BoneId;
 			BufferData.ChildParticles.Add(ConvertedData);
 		}
 	}
diff --git a/Engine/Source/Runtime/Experimental/Chaos/Public/Chaos/PullPhysicsDataImp.h b/Engine/Source/Runtime/Experimental/Chaos/Public/Chaos/PullPhysicsDataImp.h
--- a/Engine/Source/Runtime/Experimental/Chaos/Public/Chaos/PullPhysicsDataImp.h
+++ b/Engine/Source/Runtime/Experimental/Chaos/Public/Chaos/PullPhysicsDataImp.h
@@ -99,6 +99,7 @@ struct FDirtyClusterUnionParticleData
 	FUniqueIdx ParticleIdx;
 	FRigidTransform3 ChildToParent;
 	IPhysicsProxyBase* Proxy = nullptr;
+	void* CachedOwner = nullptr;
 	int32 BoneId = INDEX_NONE;
 };
 
diff --git a/Engine/Source/Runtime/Experimental/Chaos/Public/PhysicsProxy/ClusterUnionPhysicsProxy.h b/Engine/Source/Runtime/Experimental/Chaos/Public/PhysicsProxy/ClusterUnionPhysicsProxy.h
--- a/Engine/Source/Runtime/Experimental/Chaos/Public/PhysicsProxy/ClusterUnionPhysicsProxy.h
+++ b/Engine/Source/Runtime/Experimental/Chaos/Public/PhysicsProxy/ClusterUnionPhysicsProxy.h
@@ -22,6 +22,7 @@ namespace Chaos
 		FUniqueIdx ParticleIdx;
 		FTransform ChildToParent;
 		IPhysicsProxyBase* Proxy = nullptr;
+		void* CachedOwner = nullptr;
 		int32 BoneId = INDEX_NONE;
 	};
 
