Iris - Hook up Rotator3f and Rotator3d NetSerializers all the way.

#rnx

diff --git a/Engine/Source/Runtime/Experimental/Iris/Core/Private/Iris/ReplicationState/DefaultPropertyNetSerializerInfos.cpp b/Engine/Source/Runtime/Experimental/Iris/Core/Private/Iris/ReplicationState/DefaultPropertyNetSerializerInfos.cpp
--- a/Engine/Source/Runtime/Experimental/Iris/Core/Private/Iris/ReplicationState/DefaultPropertyNetSerializerInfos.cpp
+++ b/Engine/Source/Runtime/Experimental/Iris/Core/Private/Iris/ReplicationState/DefaultPropertyNetSerializerInfos.cpp
@@ -188,6 +188,8 @@ void RegisterDefaultPropertyNetSerializerInfos()
 	UE_NET_REGISTER_NETSERIALIZER_INFO(NAME_Vector3f);
 	UE_NET_REGISTER_NETSERIALIZER_INFO(NAME_Vector3d);
 	UE_NET_REGISTER_NETSERIALIZER_INFO(NAME_Rotator);
+	UE_NET_REGISTER_NETSERIALIZER_INFO(NAME_Rotator3f);
+	UE_NET_REGISTER_NETSERIALIZER_INFO(NAME_Rotator3d);
 	UE_NET_REGISTER_NETSERIALIZER_INFO(NAME_Quat);
 	UE_NET_REGISTER_NETSERIALIZER_INFO(NAME_Quat4f);
 	UE_NET_REGISTER_NETSERIALIZER_INFO(NAME_Quat4d);
