[NetPhysics] Various fixes for resimulation.

- Make sure RewindData is setup before UNetworkPhysicsComponent::CreateDatasHistory is processed.
- Use the same history size in UNetworkPhysicsComponent as in RewindData
- Make "Project Settings > Physics > Enable Resimulation" setup RewindData
- Fix issue with time dilation RPC not getting applied on the client
- Fix issue with FPhysicsReplicationAsync::ResimulationReplication not triggering a resimulation due to target TickCount not always being 0 for resimulation targets
- Stop warning logs of "out of rewind data bounds" from printing before the rewind data is populated fully.
- Remove / fix various asserts that could trigger during a resimulation due to objects having changed between cached history and current state, for example an object existing currently but not in history which is intended.
- Don't allow duplicate or out of order physics replication targets
- Also raise the physics replication target alignment span via a CVar multiplier

diff --git a/Engine/Source/Runtime/Engine/Classes/PhysicsEngine/PhysicsSettings.h b/Engine/Source/Runtime/Engine/Classes/PhysicsEngine/PhysicsSettings.h
--- a/Engine/Source/Runtime/Engine/Classes/PhysicsEngine/PhysicsSettings.h
+++ b/Engine/Source/Runtime/Engine/Classes/PhysicsEngine/PhysicsSettings.h
@@ -262,7 +262,7 @@ public:
 	UFUNCTION(BlueprintCallable, Category = "Physics")
 	int32 GetPhysicsHistoryCount() const
 	{
-		return FMath::CeilToInt(0.001f * PhysicsPrediction.MaxSupportedLatencyPrediction / AsyncFixedTimeStepSize);
+		return FMath::Max<int32>(1, FMath::CeilToInt(0.001f * PhysicsPrediction.MaxSupportedLatencyPrediction / AsyncFixedTimeStepSize));
 	}
 
 	ENGINE_API virtual void PostInitProperties() override;
diff --git a/Engine/Source/Runtime/Engine/Private/PhysicsEngine/NetworkPhysicsComponent.cpp b/Engine/Source/Runtime/Engine/Private/PhysicsEngine/NetworkPhysicsComponent.cpp
--- a/Engine/Source/Runtime/Engine/Private/PhysicsEngine/NetworkPhysicsComponent.cpp
+++ b/Engine/Source/Runtime/Engine/Private/PhysicsEngine/NetworkPhysicsComponent.cpp
@@ -332,16 +332,26 @@ void UNetworkPhysicsSystem::OnWorldPostInit(UWorld* World, const UWorld::Initial
 		return;
 	}
 
-	if(UPhysicsSettings::Get()->PhysicsPrediction.bEnablePhysicsPrediction)
+	if (UPhysicsSettings::Get()->PhysicsPrediction.bEnablePhysicsPrediction)
 	{
 		if (FPhysScene* PhysScene = World->GetPhysicsScene())
 		{
-			if(Chaos::FPhysicsSolver* Solver = PhysScene->GetSolver())
+			if (Chaos::FPhysicsSolver* Solver = PhysScene->GetSolver())
 			{ 
 				if (Solver->GetRewindCallback() == nullptr)
 				{
 					Solver->SetRewindCallback(MakeUnique<FNetworkPhysicsCallback>(World));
 				}
+
+				if (UPhysicsSettings::Get()->PhysicsPrediction.bEnablePhysicsResimulation)
+				{
+					// Enable RewindData from having bEnablePhysicsResimulation set
+					if (Solver->GetRewindData() == nullptr)
+					{
+						const int32 NumFrames = UPhysicsSettings::Get()->GetPhysicsHistoryCount();
+						Solver->EnableRewindCapture(NumFrames, true);
+					}
+				}
 			}
 		}
 	}
@@ -390,11 +400,7 @@ void UNetworkPhysicsComponent::BeginPlay()
 			{
 				if (UPhysicsSettings::Get()->PhysicsPrediction.bEnablePhysicsPrediction)
 				{
-					if (Solver->GetRewindData() == nullptr)
-					{
-						const int32 NumFrames = FMath::Max<int32>(1, UPhysicsSettings::Get()->GetPhysicsHistoryCount());
-						Solver->EnableRewindCapture(NumFrames, true);
-					}
+					SetupRewindData();
 
 					if(FNetworkPhysicsCallback* SolverCallback = static_cast<FNetworkPhysicsCallback*>(Solver->GetRewindCallback()))
 					{
@@ -903,6 +909,29 @@ void UNetworkPhysicsComponent::AddDatasHistory()
 	}
 }
 
+int32 UNetworkPhysicsComponent::SetupRewindData()
+{
+	int32 NumFrames = UPhysicsSettings::Get()->GetPhysicsHistoryCount();
+
+	if (FPhysScene* PhysScene = GetWorld()->GetPhysicsScene())
+	{
+		if (Chaos::FPhysicsSolver* Solver = PhysScene->GetSolver())
+		{
+			if (UPhysicsSettings::Get()->PhysicsPrediction.bEnablePhysicsPrediction && Solver->GetRewindData() == nullptr)
+			{
+				Solver->EnableRewindCapture(NumFrames, true);
+			}
+
+			if (Chaos::FRewindData* RewindData = Solver->GetRewindData())
+			{
+				NumFrames = RewindData->Capacity();
+			}
+		}
+	}
+
+	return NumFrames;
+}
+
 
 
 
diff --git a/Engine/Source/Runtime/Engine/Private/PhysicsEngine/PhysicsReplication.cpp b/Engine/Source/Runtime/Engine/Private/PhysicsEngine/PhysicsReplication.cpp
--- a/Engine/Source/Runtime/Engine/Private/PhysicsEngine/PhysicsReplication.cpp
+++ b/Engine/Source/Runtime/Engine/Private/PhysicsEngine/PhysicsReplication.cpp
@@ -880,7 +880,7 @@ void FPhysicsReplicationAsync::UpdateAsyncTarget(const FPhysicsRepAsyncInputData
 		Target->TargetState = Input.TargetState;
 		Target->RepMode = Input.RepMode;
 		Target->FrameOffset = Input.FrameOffset;
-		Target->TickCount = bFirstTarget ? 0 : FMath::Clamp((PrevTickCount - (SendInterval > 0 ? SendInterval : PrevReceiveInterval)), -AdjustedAverageReceiveInterval, AdjustedAverageReceiveInterval);
+		Target->TickCount = 0;
 		Target->AccumulatedSleepSeconds = 0.0f;
 
 		if (Input.RepMode == EPhysicsReplicationMode::PredictiveInterpolation)
@@ -893,6 +893,10 @@ void FPhysicsReplicationAsync::UpdateAsyncTarget(const FPhysicsRepAsyncInputData
 			}
 #endif
 
+			// Set the TickCount to the physics tick offset value from where we expected this target to arrive.
+			// If the client has ticked 2 times ahead from the last target and this target is 3 ticks in front of the previous target then the TickOffset should be -1
+			Target->TickCount = bFirstTarget ? 0 : FMath::Clamp((PrevTickCount - (SendInterval > 0 ? SendInterval : PrevReceiveInterval)), -AdjustedAverageReceiveInterval, AdjustedAverageReceiveInterval);
+
 			// Cache the position we received this target at, Predictive Interpolation will alter the target state but use this as the source position for reconciliation.
 			Target->PrevPosTarget = Input.TargetState.Position;
 			Target->PrevRotTarget = Input.TargetState.Quaternion;
@@ -1320,6 +1324,16 @@ bool FPhysicsReplicationAsync::PredictiveInterpolation(Chaos::FPBDRigidParticleH
 		return false;
 	}
 
+#if !(UE_BUILD_SHIPPING || UE_BUILD_TEST)
+	if (PhysicsReplicationCVars::PredictiveInterpolationCVars::bDrawDebugTargets)
+	{
+		const FVector Offset = FVector(0.0f, 0.0f, 50.0f);
+		const FVector StartPos = Target.TargetState.Position + Offset;
+		const int32 SizeMultiplier = FMath::Clamp(Target.TickCount, -4, 30);
+		Chaos::FDebugDrawQueue::GetInstance().DrawDebugBox(StartPos, FVector(5.0f + SizeMultiplier * 0.75f, 5.0f + SizeMultiplier * 0.75f, 5.0f + SizeMultiplier * 0.75f), Target.TargetState.Quaternion, FColor::MakeRandomSeededColor(Target.ServerFrame), false, CharacterMovementCVars::NetCorrectionLifetime, 0, 1.0f);
+	}
+#endif
+
 	const bool bIsSleeping = Handle->IsSleeping();
 	const bool bCanSimulate = Handle->IsDynamic() || bIsSleeping;
 
@@ -1541,16 +1555,6 @@ bool FPhysicsReplicationAsync::PredictiveInterpolation(Chaos::FPBDRigidParticleH
 		// Cache data for next replication
 		Target.PrevPos = FVector(CurrentState.Position);
 
-#if !(UE_BUILD_SHIPPING || UE_BUILD_TEST)
-		if (PhysicsReplicationCVars::PredictiveInterpolationCVars::bDrawDebugTargets)
-		{
-			const FVector Offset = FVector(0.0f, 0.0f, 50.0f);
-			const FVector StartPos = TargetPos + Offset;
-			const int32 SizeMultiplier = FMath::Min(Target.TickCount, 100);
-			Chaos::FDebugDrawQueue::GetInstance().DrawDebugBox(StartPos, FVector(5.0f + SizeMultiplier * 0.75f, 5.0f + SizeMultiplier * 0.75f, 5.0f + SizeMultiplier * 0.75f), Target.TargetState.Quaternion, FColor::MakeRandomSeededColor(Target.ServerFrame), false, CharacterMovementCVars::NetCorrectionLifetime, 0, 1.0f);
-		}
-#endif
-
 		// --- Target Extrapolation ---
 		if (Target.TickCount <= FMath::CeilToInt(Target.ReceiveInterval * PhysicsReplicationCVars::PredictiveInterpolationCVars::ExtrapolationTimeMultiplier))
 		{
@@ -1606,10 +1610,10 @@ bool FPhysicsReplicationAsync::ResimulationReplication(Chaos::FPBDRigidParticleH
 
 	if (LocalFrame > RewindData->CurrentFrame() || LocalFrame < RewindData->GetEarliestFrame_Internal())
 	{
-		if (LocalFrame > 0)
+		if (LocalFrame > 0 && (RewindData->CurrentFrame() - RewindData->GetEarliestFrame_Internal()) == RewindData->Capacity())
 		{
-			UE_LOG(LogPhysics, Warning, TEXT("FPhysicsReplication::ApplyRigidBodyState target frame (%d) out of rewind data bounds (%d,%d)"), LocalFrame,
-				RewindData->GetEarliestFrame_Internal(), RewindData->CurrentFrame());
+			UE_LOG(LogPhysics, Warning, TEXT("FPhysicsReplication::ResimulationReplication target frame (%d) out of rewind data bounds (%d,%d)"),
+				LocalFrame,	RewindData->GetEarliestFrame_Internal(), RewindData->CurrentFrame());
 		}
 		return true;
 	}
diff --git a/Engine/Source/Runtime/Engine/Private/PlayerController.cpp b/Engine/Source/Runtime/Engine/Private/PlayerController.cpp
--- a/Engine/Source/Runtime/Engine/Private/PlayerController.cpp
+++ b/Engine/Source/Runtime/Engine/Private/PlayerController.cpp
@@ -6332,7 +6332,7 @@ void APlayerController::ClientSetupNetworkPhysicsTimestamp_Implementation(FAsync
 
 void APlayerController::ClientAckTimeDilation_Implementation(float TimeDilation, int32 ServerStep)
 {
-	if (ClientLatestTimeDilationServerStep <= ServerStep)
+	if (ServerStep <= ClientLatestTimeDilationServerStep)
 	{
 		return;
 	}
diff --git a/Engine/Source/Runtime/Engine/Public/Physics/NetworkPhysicsComponent.h b/Engine/Source/Runtime/Engine/Public/Physics/NetworkPhysicsComponent.h
--- a/Engine/Source/Runtime/Engine/Public/Physics/NetworkPhysicsComponent.h
+++ b/Engine/Source/Runtime/Engine/Public/Physics/NetworkPhysicsComponent.h
@@ -424,6 +424,9 @@ public:
 	// Add states/inputs history to rewind datas
 	ENGINE_API void AddDatasHistory();
 
+	// Enable RewindData history caching and return the history size
+	ENGINE_API int32 SetupRewindData();
+
 	// Get the datas factory that will be used for net serialization
 	TSharedPtr<Chaos::FBaseRewindHistory>& GetStatesHistory() { return StatesHistory; }
 
@@ -500,10 +503,11 @@ private:
 template<typename PhysicsTraits>
 FORCEINLINE void UNetworkPhysicsComponent::CreateDatasHistory(UActorComponent* HistoryComponent)
 {
+	const int32 NumFrames = SetupRewindData();
+
 	APlayerController* Controller = GetPlayerController();
-	const bool bIsLocalHistory = (Controller && Controller->IsLocalController());
-	const int32 NumFrames = UPhysicsSettings::Get()->GetPhysicsHistoryCount();
-	
+	const bool bIsLocalHistory = (Controller && Controller->IsLocalController()); // FIXME: The controller is null at this point, but bIsLocalHistory isn't currently used so doesn't create an issue.
+
 	InputsHistory = MakeShared<TNetRewindHistory<typename PhysicsTraits::InputsType>>(NumFrames, bIsLocalHistory);
 	StatesHistory = MakeShared<TNetRewindHistory<typename PhysicsTraits::StatesType>>(NumFrames, bIsLocalHistory);
 
diff --git a/Engine/Source/Runtime/Experimental/Chaos/Private/Chaos/Collision/ParticlePairMidPhase.cpp b/Engine/Source/Runtime/Experimental/Chaos/Private/Chaos/Collision/ParticlePairMidPhase.cpp
--- a/Engine/Source/Runtime/Experimental/Chaos/Private/Chaos/Collision/ParticlePairMidPhase.cpp
+++ b/Engine/Source/Runtime/Experimental/Chaos/Private/Chaos/Collision/ParticlePairMidPhase.cpp
@@ -684,7 +684,6 @@ namespace Chaos
 			Constraint->GetContainerCookie().bIsMultiShapePair = false;
 			Constraint->GetContainerCookie().CreationEpoch = CurrentEpoch;
 		}
-		check(!Constraint->IsEnabled());
 
 		// Copy the constraint data over the existing one (ensure we do not replace data required by the graph and the allocator/container)
 		Constraint->RestoreFrom(SourceConstraint);
diff --git a/Engine/Source/Runtime/Experimental/Chaos/Private/PBDRigidsSolver.cpp b/Engine/Source/Runtime/Experimental/Chaos/Private/PBDRigidsSolver.cpp
--- a/Engine/Source/Runtime/Experimental/Chaos/Private/PBDRigidsSolver.cpp
+++ b/Engine/Source/Runtime/Experimental/Chaos/Private/PBDRigidsSolver.cpp
@@ -999,8 +999,6 @@ namespace Chaos
 
 	void FPBDRigidsSolver::EnableRewindCapture(int32 NumFrames, bool InUseCollisionResimCache)
 	{
-		UE_LOG(LogChaos, Log, TEXT("PBDRigidsSolver::EnableRewindCapture - Starting physics data history caching for rewind / resimulation."));
-
 		//TODO: this function calls both internal and external - sort of assumed during initialization. Should decide what thread it's called on and mark it as either external or internal
 		if (MRewindData.IsValid())
 		{
@@ -1021,6 +1019,8 @@ namespace Chaos
 		}
 		
 		UpdateIsDeterministic();
+
+		UE_LOG(LogChaos, Log, TEXT("PBDRigidsSolver::EnableRewindCapture - Starting physics data history caching for rewind / resimulation. History Size: %d"), NumFramesSet);
 	}
 
 	void FPBDRigidsSolver::Reset()
@@ -1550,10 +1550,6 @@ namespace Chaos
 						// Pass until after all bodies are created. 
 						break;
 					}
-					default:
-					{
-						ensure(0 && TEXT("Unknown proxy type in physics solver."));
-					}
 				}
 			}
 		});
@@ -1565,61 +1561,60 @@ namespace Chaos
 			{
 				switch (Dirty.Proxy->GetType())
 				{
-				case EPhysicsProxyType::JointConstraintType:
-				{
-					auto JointProxy = static_cast<FJointConstraintPhysicsProxy*>(Dirty.Proxy);
-					const bool bIsNew = !JointProxy->IsInitialized();
-					if (bIsNew)
+					case EPhysicsProxyType::JointConstraintType:
 					{
-						JointConstraintPhysicsProxies_Internal.Add(JointProxy);
-						JointProxy->InitializeOnPhysicsThread(this, *Manager, DataIdx, Dirty.PropertyData);
-						JointProxy->SetInitialized(GetCurrentFrame());
-					}
+						auto JointProxy = static_cast<FJointConstraintPhysicsProxy*>(Dirty.Proxy);
+						const bool bIsNew = !JointProxy->IsInitialized();
+						if (bIsNew)
+						{
+							JointConstraintPhysicsProxies_Internal.Add(JointProxy);
+							JointProxy->InitializeOnPhysicsThread(this, *Manager, DataIdx, Dirty.PropertyData);
+							JointProxy->SetInitialized(GetCurrentFrame());
+						}
 
-					//TODO: if we support predicting creation / destruction of joints need to handle null joint case
-					if (RewindData)
-					{
-						RewindData->PushGTDirtyData(*Manager, DataIdx, Dirty, nullptr);
-					}
+						//TODO: if we support predicting creation / destruction of joints need to handle null joint case
+						if (RewindData)
+						{
+							RewindData->PushGTDirtyData(*Manager, DataIdx, Dirty, nullptr);
+						}
 				
-					JointProxy->PushStateOnPhysicsThread(this, *Manager, DataIdx, Dirty.PropertyData);
-					break;
-				}
-
-				case EPhysicsProxyType::SuspensionConstraintType:
-				{
-					auto SuspensionProxy = static_cast<FSuspensionConstraintPhysicsProxy*>(Dirty.Proxy);
-					const bool bIsNew = !SuspensionProxy->IsInitialized();
-					if (bIsNew)
-					{
-						SuspensionProxy->InitializeOnPhysicsThread(this, *Manager, DataIdx, Dirty.PropertyData);
-						SuspensionProxy->SetInitialized();
+						JointProxy->PushStateOnPhysicsThread(this, *Manager, DataIdx, Dirty.PropertyData);
+						break;
 					}
-					SuspensionProxy->PushStateOnPhysicsThread(this, *Manager, DataIdx, Dirty.PropertyData);
-					break;
-				}
 
-				case EPhysicsProxyType::CharacterGroundConstraintType:
-				{
-					auto ConstraintProxy = static_cast<FCharacterGroundConstraintProxy*>(Dirty.Proxy);
-					const bool bIsNew = !ConstraintProxy->IsInitialized();
-					if (bIsNew)
+					case EPhysicsProxyType::SuspensionConstraintType:
 					{
-						CharacterGroundConstraintProxies_Internal.Add(ConstraintProxy);
-						ConstraintProxy->InitializeOnPhysicsThread(this, *Manager, DataIdx, Dirty.PropertyData);
-						ConstraintProxy->SetInitialized(GetCurrentFrame());
+						auto SuspensionProxy = static_cast<FSuspensionConstraintPhysicsProxy*>(Dirty.Proxy);
+						const bool bIsNew = !SuspensionProxy->IsInitialized();
+						if (bIsNew)
+						{
+							SuspensionProxy->InitializeOnPhysicsThread(this, *Manager, DataIdx, Dirty.PropertyData);
+							SuspensionProxy->SetInitialized();
+						}
+						SuspensionProxy->PushStateOnPhysicsThread(this, *Manager, DataIdx, Dirty.PropertyData);
+						break;
 					}
 
-					//TODO: Support rewind for character ground constraints
-					//if (RewindData)
-					//{
-					//	RewindData->PushGTDirtyData(*Manager, DataIdx, Dirty, nullptr);
-					//}
+					case EPhysicsProxyType::CharacterGroundConstraintType:
+					{
+						auto ConstraintProxy = static_cast<FCharacterGroundConstraintProxy*>(Dirty.Proxy);
+						const bool bIsNew = !ConstraintProxy->IsInitialized();
+						if (bIsNew)
+						{
+							CharacterGroundConstraintProxies_Internal.Add(ConstraintProxy);
+							ConstraintProxy->InitializeOnPhysicsThread(this, *Manager, DataIdx, Dirty.PropertyData);
+							ConstraintProxy->SetInitialized(GetCurrentFrame());
+						}
 
-					ConstraintProxy->PushStateOnPhysicsThread(this, *Manager, DataIdx, Dirty.PropertyData);
-					break;
-				}
+						//TODO: Support rewind for character ground constraints
+						//if (RewindData)
+						//{
+						//	RewindData->PushGTDirtyData(*Manager, DataIdx, Dirty, nullptr);
+						//}
 
+						ConstraintProxy->PushStateOnPhysicsThread(this, *Manager, DataIdx, Dirty.PropertyData);
+						break;
+					}
 				}
 			}
 		});
diff --git a/Engine/Source/Runtime/Experimental/Chaos/Private/RewindData.cpp b/Engine/Source/Runtime/Experimental/Chaos/Private/RewindData.cpp
--- a/Engine/Source/Runtime/Experimental/Chaos/Private/RewindData.cpp
+++ b/Engine/Source/Runtime/Experimental/Chaos/Private/RewindData.cpp
@@ -333,8 +333,12 @@ bool FRewindData::RewindToFrame(int32 Frame)
 
 		bool bAnyChange = RewindHelper(PTParticle, bResimAsFollower, History.ParticlePositionRotation, [](auto Particle, const auto& Data) {Particle->SetXR(Data); });
 		bAnyChange |= RewindHelper(PTParticle->CastToKinematicParticle(), bResimAsFollower, History.Velocities, [](auto Particle, const auto& Data) {Particle->SetVelocities(Data); });
-		bAnyChange |= RewindHelper(PTParticle, bResimAsFollower, History.NonFrequentData, [](auto Particle, const auto& Data) {Particle->SetNonFrequentData(Data); });
-		bAnyChange |= RewindHelper(PTParticle->CastToKinematicParticle(), bResimAsFollower, History.KinematicTarget, [](auto Particle, const auto& Data) {Particle->SetKinematicTarget(Data); });
+		bAnyChange |= RewindHelper(PTParticle, bResimAsFollower, History.NonFrequentData, [this](auto Particle, const auto& Data)
+			{
+				Solver->GetEvolution()->InvalidateParticle(Particle); // Clear collision/constraints before updating NonFrequentData
+				Particle->SetNonFrequentData(Data); 
+			});
+		bAnyChange |= RewindHelper(PTParticle->CastToKinematicParticle(), bResimAsFollower, History.KinematicTarget, [](auto Particle, const auto& Data) { Particle->SetKinematicTarget(Data); });
 		bAnyChange |= RewindHelper(PTParticle->CastToRigidParticle(), bResimAsFollower, History.Dynamics, [](auto Particle, const auto& Data) {Particle->SetDynamics(Data); });
 		bAnyChange |= RewindHelper(PTParticle->CastToRigidParticle(), bResimAsFollower, History.DynamicsMisc, [this](auto Particle, const auto& Data) {Solver->SetParticleDynamicMisc(Particle, Data); });
 		bAnyChange |= RewindHelper(PTParticle->CastToRigidParticle(), bResimAsFollower, History.MassProps, [](auto Particle, const auto& Data) {Particle->SetMassProps(Data); });
diff --git a/Engine/Source/Runtime/Experimental/Chaos/Public/Chaos/EvolutionResimCache.h b/Engine/Source/Runtime/Experimental/Chaos/Public/Chaos/EvolutionResimCache.h
--- a/Engine/Source/Runtime/Experimental/Chaos/Public/Chaos/EvolutionResimCache.h
+++ b/Engine/Source/Runtime/Experimental/Chaos/Public/Chaos/EvolutionResimCache.h
@@ -54,7 +54,7 @@ namespace Chaos
 		{
 			//if this function is called it means the particle is in sync, which means we should have a cached value
 			const FPBDSolveCache* Cache = ParticleToCachedSolve.Find(Particle.UniqueIdx());
-			if(ensure(Cache))
+			if (Cache)
 			{
 				Particle.P() = Cache->P;
 				Particle.Q() = Cache->Q;
