[ChaosVD] Performance improvements

- Added a way to generate KeyFrames on load, so when we scrub skipping frames or backwards, we have way less data to process (as the delta between keyframes will be lower).
- Allow use delta frame instead of generated keyframes during playback (processing generated keyframes on recordings with +50000 particles is not cheap, so they should not be used during normal playback).
- Allow to override far clipping setting (helps with complex scenes. A proper fix for this will be fully switching to Instanced Static Meshes).
- Partial Fix for visibility button not working on the Scene Outliner.
- Added solver name and ID to the progress bar for slow operations. This is needed to no give the false impression we are loading the same data twice. For example in PIE we usually have to load particle data for Client and Server. Without this change this would be two progress bars that appear to be the same.
- Added an option to disable visualization of disabled particles

diff --git a/Engine/Plugins/Experimental/ChaosVD/Source/ChaosVD/Private/ChaosVDEditorSettings.cpp b/Engine/Plugins/Experimental/ChaosVD/Source/ChaosVD/Private/ChaosVDEditorSettings.cpp
--- a/Engine/Plugins/Experimental/ChaosVD/Source/ChaosVD/Private/ChaosVDEditorSettings.cpp
+++ b/Engine/Plugins/Experimental/ChaosVD/Source/ChaosVD/Private/ChaosVDEditorSettings.cpp
@@ -82,4 +82,11 @@ void UChaosVDEditorSettings::PostEditChangeProperty(FPropertyChangedEvent& Prope
 	{
 		ColorsSettingsChangedDelegate.Broadcast(this);
 	}
+	else if (PropertyName == GET_MEMBER_NAME_CHECKED(UChaosVDEditorSettings, FarClippingOverride))
+	{
+		FarClippingOverrideChangedDelegate.Broadcast(this);
+	}
+
+	// TODO: If we keep this object as the main setting object,
+	// we should have a single event for what changed and an enum flags that the listener could use to decide if cares about the change
 }
diff --git a/Engine/Plugins/Experimental/ChaosVD/Source/ChaosVD/Private/ChaosVDGeometryDataComponent.cpp b/Engine/Plugins/Experimental/ChaosVD/Source/ChaosVD/Private/ChaosVDGeometryDataComponent.cpp
--- a/Engine/Plugins/Experimental/ChaosVD/Source/ChaosVD/Private/ChaosVDGeometryDataComponent.cpp
+++ b/Engine/Plugins/Experimental/ChaosVD/Source/ChaosVD/Private/ChaosVDGeometryDataComponent.cpp
@@ -25,9 +25,26 @@ void FChaosVDGeometryDataComponentBase::UpdateVisibility_Internal(const FChaosVD
 		if (const UChaosVDEditorSettings* EditorSettings = GetDefault<UChaosVDEditorSettings>())
 		{
 			const EChaosVDGeometryVisibilityFlags CurrentVisibilityFlags = static_cast<EChaosVDGeometryVisibilityFlags>(EditorSettings->GeometryVisibilityFlags);
-
+			
 			bool bShouldGeometryBeVisible = false;
 
+			if (!EnumHasAnyFlags(CurrentVisibilityFlags, EChaosVDGeometryVisibilityFlags::ShowDisabledParticles))
+			{
+				//TODO: We should use IChaosVDParticleVisualizationDataProvider instead, which AChaosVDParticleActor implements already
+				// but it is not an uinterface
+				if (AChaosVDParticleActor* ParticleActor = Cast<AChaosVDParticleActor>(MeshComponent->GetOwner()))
+				{
+					if (const FChaosVDParticleDataWrapper* ParticleData = ParticleActor->GetParticleData())
+					{
+						if (ParticleData->ParticleDynamicsMisc.HasValidData() && ParticleData->ParticleDynamicsMisc.bDisabled)
+						{
+							MeshComponent->SetVisibility(bShouldGeometryBeVisible);
+							return;
+						}
+					}
+				}
+			}
+
 			// TODO: Re-visit the way we determine visibility of the meshes.
 			// Now that the options have grown and they will continue to do so, these checks are becoming hard to read and extend
 
diff --git a/Engine/Plugins/Experimental/ChaosVD/Source/ChaosVD/Private/ChaosVDParticleActor.cpp b/Engine/Plugins/Experimental/ChaosVD/Source/ChaosVD/Private/ChaosVDParticleActor.cpp
--- a/Engine/Plugins/Experimental/ChaosVD/Source/ChaosVD/Private/ChaosVDParticleActor.cpp
+++ b/Engine/Plugins/Experimental/ChaosVD/Source/ChaosVD/Private/ChaosVDParticleActor.cpp
@@ -225,6 +225,15 @@ void AChaosVDParticleActor::PostEditChangeProperty(FPropertyChangedEvent& Proper
 	}
 }
 
+#if WITH_EDITOR
+void AChaosVDParticleActor::SetIsTemporarilyHiddenInEditor(bool bIsHidden)
+{
+	const bool bShouldBeHidden = bIsActive ? bIsHidden : true;
+
+	Super::SetIsTemporarilyHiddenInEditor(bShouldBeHidden);
+}
+#endif //WITH_EDITOR
+
 void AChaosVDParticleActor::GetCollisionData(TArray<TSharedPtr<FChaosVDCollisionDataFinder>>& OutCollisionDataFound)
 {
 	if (const TArray<TSharedPtr<FChaosVDParticlePairMidPhase>>* MidPhases = GetCollisionMidPhasesArray())
@@ -333,15 +342,15 @@ void AChaosVDParticleActor::SetIsActive(bool bNewActive)
 {
 	if (bIsActive != bNewActive)
 	{
+		bIsActive = bNewActive;
 #if WITH_EDITOR
 		//TODO: We need to add support for this to our Scene Outliner
 		// This will hide the actor and disable it in the outliner but it will still be listed
 		// We need to add a way to unlist inactive particle actors without a full hierarchy rebuild, which would be too costly
 		bEditable = bNewActive;
 		bListedInSceneOutliner = bNewActive;
-		SetIsTemporarilyHiddenInEditor(!bNewActive);
+		SetIsTemporarilyHiddenInEditor(!bIsActive);
 #endif
-		bIsActive = bNewActive;
 
 		if (const TSharedPtr<FChaosVDScene> ScenePtr = SceneWeakPtr.Pin())
 		{
diff --git a/Engine/Plugins/Experimental/ChaosVD/Source/ChaosVD/Private/ChaosVDPlaybackViewportClient.cpp b/Engine/Plugins/Experimental/ChaosVD/Source/ChaosVD/Private/ChaosVDPlaybackViewportClient.cpp
--- a/Engine/Plugins/Experimental/ChaosVD/Source/ChaosVD/Private/ChaosVDPlaybackViewportClient.cpp
+++ b/Engine/Plugins/Experimental/ChaosVD/Source/ChaosVD/Private/ChaosVDPlaybackViewportClient.cpp
@@ -29,6 +29,12 @@ FChaosVDPlaybackViewportClient::FChaosVDPlaybackViewportClient(const TSharedPtr<
 	{
 		GEngine->OnActorMoving().AddRaw(this, &FChaosVDPlaybackViewportClient::HandleActorMoving);
 	}
+
+	if (UChaosVDEditorSettings* Settings = GetMutableDefault<UChaosVDEditorSettings>())
+	{
+		OverrideFarClipPlane(Settings->FarClippingOverride);
+		Settings->OnFarClippingOverrideChanged().AddRaw(this, &FChaosVDPlaybackViewportClient::HandleFarClippingOverrideSettingsChanged);
+	}
 }
 
 FChaosVDPlaybackViewportClient::~FChaosVDPlaybackViewportClient()
@@ -45,6 +51,11 @@ FChaosVDPlaybackViewportClient::~FChaosVDPlaybackViewportClient()
 	{
 		GEngine->OnActorMoving().RemoveAll(this);
 	}
+	
+	if (UChaosVDEditorSettings* Settings = GetMutableDefault<UChaosVDEditorSettings>())
+	{
+		Settings->OnFarClippingOverrideChanged().RemoveAll(this);
+	}
 }
 
 void FChaosVDPlaybackViewportClient::ProcessClick(FSceneView& View, HHitProxy* HitProxy, FKey Key, EInputEvent Event, uint32 HitX, uint32 HitY)
@@ -127,6 +138,14 @@ void FChaosVDPlaybackViewportClient::HandleActorMoving(AActor* MovedActor) const
 	}
 }
 
+void FChaosVDPlaybackViewportClient::HandleFarClippingOverrideSettingsChanged(UChaosVDEditorSettings* SettingsObject)
+{
+	if (SettingsObject)
+	{
+		OverrideFarClipPlane(SettingsObject->FarClippingOverride);
+	}
+}
+
 void FChaosVDPlaybackViewportClient::TrackActor(AActor* ActorToTrack, EChaosVDActorTrackingMode TrackingMode)
 {
 	if (!ActorToTrack)
diff --git a/Engine/Plugins/Experimental/ChaosVD/Source/ChaosVD/Private/ChaosVDScene.cpp b/Engine/Plugins/Experimental/ChaosVD/Source/ChaosVD/Private/ChaosVDScene.cpp
--- a/Engine/Plugins/Experimental/ChaosVD/Source/ChaosVD/Private/ChaosVDScene.cpp
+++ b/Engine/Plugins/Experimental/ChaosVD/Source/ChaosVD/Private/ChaosVDScene.cpp
@@ -127,7 +127,9 @@ void FChaosVDScene::UpdateFromRecordedStepData(const int32 SolverID, const FStri
 	{
 		constexpr float AmountOfWork = 1.0f;
 		const float PercentagePerElement = 1.0f / InRecordedStepData.RecordedParticlesData.Num();
-		FScopedSlowTask UpdatingSceneSlowTask(AmountOfWork, LOCTEXT("ProcessingParticleData", "Processing Particle Data..."));
+
+		const FText ProgressBarTitle = FText::Format(FTextFormat(LOCTEXT("ProcessingParticleData", "Processing Particle Data for {0} Solver with ID {1} ...")), FText::FromString(SolverName), FText::AsNumber(SolverID));
+		FScopedSlowTask UpdatingSceneSlowTask(AmountOfWork, ProgressBarTitle);
 		UpdatingSceneSlowTask.MakeDialogDelayed(ChaosVDSceneUIOptions::DelayToShowProgressDialogThreshold, ChaosVDSceneUIOptions::bShowCancelButton, ChaosVDSceneUIOptions::bAllowInPIE);
 	
 		// Go over existing Particle VD Instances and update them or create them if needed 
diff --git a/Engine/Plugins/Experimental/ChaosVD/Source/ChaosVD/Public/ChaosVDEditorSettings.h b/Engine/Plugins/Experimental/ChaosVD/Source/ChaosVD/Public/ChaosVDEditorSettings.h
--- a/Engine/Plugins/Experimental/ChaosVD/Source/ChaosVD/Public/ChaosVDEditorSettings.h
+++ b/Engine/Plugins/Experimental/ChaosVD/Source/ChaosVD/Public/ChaosVDEditorSettings.h
@@ -42,6 +42,7 @@ enum class EChaosVDGeometryVisibilityFlags : uint8
 	Simple = 1 << 3,
 	Complex = 1 << 4,
 	ShowHeightfields = 1 << 5, // Selecting this will show heightfields even if complex is not selected
+	ShowDisabledParticles = 1 << 6,
 };
 ENUM_CLASS_FLAGS(EChaosVDGeometryVisibilityFlags)
 
@@ -166,6 +167,9 @@ class UChaosVDEditorSettings : public UObject
 	GENERATED_BODY()
 public:
 
+	UPROPERTY(EditAnywhere, Category = "Viewport")
+	float FarClippingOverride = 10000.0f;
+
 	UPROPERTY(EditAnywhere, Category = "Viewport Visualization Flags", meta = (Bitmask, BitmaskEnum = "/Script/ChaosVD.EChaosVDParticleDataVisualizationFlags"))
 	uint32 GlobalParticleDataVisualizationFlags = 0;
 	
@@ -220,10 +224,13 @@ public:
 
 	FChaosVDSettingChaged& OnColorSettingsChanged() { return ColorsSettingsChangedDelegate; }
 
+	FChaosVDSettingChaged& OnFarClippingOverrideChanged() { return FarClippingOverrideChangedDelegate; }
+
 	TSharedPtr<FName> SelectedTrackedTransformName;
 	TSharedPtr<FName> SelectedTrackedLocationName;
 
 protected:
 	FChaosVDSettingChaged VisibilitySettingsChangedDelegate;
 	FChaosVDSettingChaged ColorsSettingsChangedDelegate;
+	FChaosVDSettingChaged FarClippingOverrideChangedDelegate;
 };
diff --git a/Engine/Plugins/Experimental/ChaosVD/Source/ChaosVD/Public/ChaosVDGeometryBuilder.h b/Engine/Plugins/Experimental/ChaosVD/Source/ChaosVD/Public/ChaosVDGeometryBuilder.h
--- a/Engine/Plugins/Experimental/ChaosVD/Source/ChaosVD/Public/ChaosVDGeometryBuilder.h
+++ b/Engine/Plugins/Experimental/ChaosVD/Source/ChaosVD/Public/ChaosVDGeometryBuilder.h
@@ -247,13 +247,14 @@ void FChaosVDGeometryBuilder::CreateMeshComponentsFromImplicit(const Chaos::FImp
 	
 	if (InnerType == ImplicitObjectType::Union)
 	{
-		const FImplicitObjectUnion* Union = InImplicitObject->template GetObject<FImplicitObjectUnion>();
-
-		for (int i = 0; i < Union->GetObjects().Num(); ++i)
+		if (const FImplicitObjectUnion* Union = InImplicitObject->template AsA<FImplicitObjectUnion>())
 		{
-			const FImplicitObjectPtr& UnionImplicit = Union->GetObjects()[i];
+			for (int i = 0; i < Union->GetObjects().Num(); ++i)
+			{
+				const FImplicitObjectPtr& UnionImplicit = Union->GetObjects()[i];
 
-			CreateMeshComponentsFromImplicit<MeshType, ComponentType>(UnionImplicit.GetReference(), Owner, OutMeshComponents, Transform, i, DesiredLODCount);	
+				CreateMeshComponentsFromImplicit<MeshType, ComponentType>(UnionImplicit.GetReference(), Owner, OutMeshComponents, Transform, i, DesiredLODCount);	
+			}
 		}
 
 		return;
diff --git a/Engine/Plugins/Experimental/ChaosVD/Source/ChaosVD/Public/ChaosVDParticleActor.h b/Engine/Plugins/Experimental/ChaosVD/Source/ChaosVD/Public/ChaosVDParticleActor.h
--- a/Engine/Plugins/Experimental/ChaosVD/Source/ChaosVD/Public/ChaosVDParticleActor.h
+++ b/Engine/Plugins/Experimental/ChaosVD/Source/ChaosVD/Public/ChaosVDParticleActor.h
@@ -66,6 +66,7 @@ public:
 #if WITH_EDITOR
 	virtual bool IsSelectedInEditor() const override;
 	virtual void PostEditChangeProperty(FPropertyChangedEvent& PropertyChangedEvent) override;
+	virtual void SetIsTemporarilyHiddenInEditor(bool bIsHidden) override;
 #endif
 
 	void UpdateGeometryComponentsVisibility();
diff --git a/Engine/Plugins/Experimental/ChaosVD/Source/ChaosVD/Public/ChaosVDPlaybackViewportClient.h b/Engine/Plugins/Experimental/ChaosVD/Source/ChaosVD/Public/ChaosVDPlaybackViewportClient.h
--- a/Engine/Plugins/Experimental/ChaosVD/Source/ChaosVD/Public/ChaosVDPlaybackViewportClient.h
+++ b/Engine/Plugins/Experimental/ChaosVD/Source/ChaosVD/Public/ChaosVDPlaybackViewportClient.h
@@ -3,6 +3,7 @@
 #pragma once
 #include "LevelEditorViewport.h"
 
+class UChaosVDEditorSettings;
 struct FChaosVDGameFrameData;
 class FChaosVDScene;
 enum class EChaosVDActorTrackingMode;
@@ -34,6 +35,7 @@ public:
 private:
 	void HandleObjectFocused(UObject* FocusedObject);
 	void HandleActorMoving(AActor* MovedActor) const;
+	void HandleFarClippingOverrideSettingsChanged(UChaosVDEditorSettings* SettingsObject);
 	
 	FDelegateHandle ObjectFocusedDelegateHandle;
 	UWorld* CVDWorld;
diff --git a/Engine/Source/Runtime/Experimental/ChaosVDData/Private/ChaosVDRecording.cpp b/Engine/Source/Runtime/Experimental/ChaosVDData/Private/ChaosVDRecording.cpp
--- a/Engine/Source/Runtime/Experimental/ChaosVDData/Private/ChaosVDRecording.cpp
+++ b/Engine/Source/Runtime/Experimental/ChaosVDData/Private/ChaosVDRecording.cpp
@@ -56,12 +56,29 @@ FString FChaosVDRecording::GetSolverName_AssumedLocked(int32 SolverID)
 	return DefaultName;
 }
 
-FChaosVDSolverFrameData* FChaosVDRecording::GetSolverFrameData_AssumesLocked(const int32 SolverID, const int32 FrameNumber)
+FChaosVDSolverFrameData* FChaosVDRecording::GetSolverFrameData_AssumesLocked(const int32 SolverID, const int32 FrameNumber, bool bKeyFrameOnly)
 {
 	if (TArray<FChaosVDSolverFrameData>* SolverFrames = RecordedFramesDataPerSolver.Find(SolverID))
 	{
 		//TODO: Find a safer way of do this. If someone stores this ptr bad things will happen
-		return SolverFrames->IsValidIndex(FrameNumber) ? &(*SolverFrames)[FrameNumber] : nullptr;
+		if (FChaosVDSolverFrameData* FoundFrame = SolverFrames->IsValidIndex(FrameNumber) ? &(*SolverFrames)[FrameNumber] : nullptr)
+		{
+			if (!FoundFrame->bIsKeyFrame && bKeyFrameOnly)
+			{
+				if (TMap<int32, FChaosVDSolverFrameData>* GenerateSolverFrameByNumber = GeneratedKeyFrameDataPerSolver.Find(SolverID))
+				{
+					if (FChaosVDSolverFrameData* GeneratedKeyFrame = GenerateSolverFrameByNumber->Find(FrameNumber))
+					{
+						return GeneratedKeyFrame;
+					}
+				}
+				ensureMsgf(false, TEXT("Failed to find generated KeyFrame [%d] for Solver [%d]"), FrameNumber, SolverID);
+			}
+			else
+			{
+				return FoundFrame;
+			}			
+		}
 	}
 
 	return nullptr;
@@ -172,6 +189,11 @@ int32 FChaosVDRecording::GetLowestGameFrameAtSolverFrameNumber(int32 SolverID, i
 void FChaosVDRecording::AddKeyFrameNumberForSolver(const int32 SolverID, int32 FrameNumber)
 {
 	FWriteScopeLock WriteLock(RecordingDataLock);
+	AddKeyFrameNumberForSolver_AssumesLocked(SolverID, FrameNumber);
+}
+
+void FChaosVDRecording::AddKeyFrameNumberForSolver_AssumesLocked(int32 SolverID, int32 FrameNumber)
+{
 	if (TArray<int32>* KeyFrameNumber = RecordedKeyFramesNumberPerSolver.Find(SolverID))
 	{
 		KeyFrameNumber->Add(FrameNumber);
@@ -182,30 +204,79 @@ void FChaosVDRecording::AddKeyFrameNumberForSolver(const int32 SolverID, int32 F
 	}
 }
 
+void FChaosVDRecording::GenerateAndStoreKeyframeForSolver_AssumesLocked(const int32 SolverID, int32 CurrentFrameNumber, const int32 LastKeyFrameNumber)
+{
+	FChaosVDSolverFrameData GeneratedKeyFrame;
+	CollapseSolverFramesRange_AssumesLocked(SolverID, LastKeyFrameNumber, CurrentFrameNumber, GeneratedKeyFrame);
+
+	// We don't replace an existing delta frame with a generated keyframe because processing keyframes during playback is expensive
+	// So we keep the generated keyframes on its own map, so we can access them when needed
+	// (usually when we are skipping frames and we need to collapse frame data from the closest keyframe)
+	if (TMap<int32, FChaosVDSolverFrameData>* GenerateSolverFrameByNumber = GeneratedKeyFrameDataPerSolver.Find(SolverID))
+	{
+		GenerateSolverFrameByNumber->Add(CurrentFrameNumber, GeneratedKeyFrame);
+	}
+	else
+	{
+		TMap<int32, FChaosVDSolverFrameData> GeneratedSolverFramesByNumber;
+		GeneratedSolverFramesByNumber.Add(CurrentFrameNumber, GeneratedKeyFrame);
+		GeneratedKeyFrameDataPerSolver.Add(SolverID, GeneratedSolverFramesByNumber);
+	}
+}
+
+
 void FChaosVDRecording::AddFrameForSolver(const int32 SolverID, FChaosVDSolverFrameData&& InFrameData)
 {
-	int32 FrameNumber;
-	const bool bIsKeyFrame = InFrameData.bIsKeyFrame;
+	int32 CurrentFrameNumber;
+	bool bIsKeyFrame = InFrameData.bIsKeyFrame;
 
 	{
 		FWriteScopeLock WriteLock(RecordingDataLock);
 
+		auto FindLastKeyFrameNumberForSolver = [this](int32 SolverID)-> int32
+		{
+			if (TArray<int32>* KeyFrameNumber = RecordedKeyFramesNumberPerSolver.Find(SolverID))
+			{
+				return KeyFrameNumber->Num() ? KeyFrameNumber->Last() : INDEX_NONE;
+			}
+			return INDEX_NONE;
+		};
+
 		if (TArray<FChaosVDSolverFrameData>* SolverFrames = RecordedFramesDataPerSolver.Find(SolverID))
 		{
-			FrameNumber = SolverFrames->Num();
+			CurrentFrameNumber = SolverFrames->Num();
 
-			SolverFrames->Add(MoveTemp(InFrameData));	
+			SolverFrames->Add(MoveTemp(InFrameData));
+
+			if (!bIsKeyFrame)
+			{
+				// If not a keyframe, see if we should generate a to keyframe for the frame number we just added.
+				// This greatly reduces the cost during playback when we are skipping more than one frame or going backwards because with more keyframe
+				// we have less data to process on the process "Play from last key frame", needed in such situations.
+				const int32 LastKeyFrameNumber = FindLastKeyFrameNumberForSolver(SolverID);
+				if (LastKeyFrameNumber != INDEX_NONE)
+				{
+					constexpr int32 MaxDeltaBetweenKeyframes = 5;
+					const int32 FrameDiffSinceLastKeyframe = FMath::Abs(CurrentFrameNumber - LastKeyFrameNumber);
+
+					if (FrameDiffSinceLastKeyframe > MaxDeltaBetweenKeyframes)
+					{
+						GenerateAndStoreKeyframeForSolver_AssumesLocked(SolverID, CurrentFrameNumber, LastKeyFrameNumber);
+						AddKeyFrameNumberForSolver_AssumesLocked(SolverID, CurrentFrameNumber);
+					}
+				}
+			}
 		}
 		else
 		{
-			FrameNumber = 0;
+			CurrentFrameNumber = 0;
 			RecordedFramesDataPerSolver.Add(SolverID, {}).Emplace(MoveTemp(InFrameData));
 		}
 	}
 
 	if (bIsKeyFrame)
 	{
-		AddKeyFrameNumberForSolver(SolverID, FrameNumber);
+		AddKeyFrameNumberForSolver(SolverID, CurrentFrameNumber);
 	}
 
 	OnRecordingUpdated().Broadcast();
@@ -301,7 +372,8 @@ void FChaosVDRecording::CollapseSolverFramesRange_AssumesLocked(int32 SolverID,
 	
 	for (int32 CurrentFrameNumber = StartFrame; CurrentFrameNumber <= EndFrame; CurrentFrameNumber++)
 	{
-		if (const FChaosVDSolverFrameData* SolverFrameData = GetSolverFrameData_AssumesLocked(SolverID, CurrentFrameNumber))
+		const bool bRequestingKeyFrameOnly = CurrentFrameNumber == StartFrame;
+		if (const FChaosVDSolverFrameData* SolverFrameData = GetSolverFrameData_AssumesLocked(SolverID, CurrentFrameNumber, bRequestingKeyFrameOnly))
 		{
 			OutCollapsedFrameData.ParticlesDestroyedIDs.Append(SolverFrameData->ParticlesDestroyedIDs);
 
diff --git a/Engine/Source/Runtime/Experimental/ChaosVDData/Public/ChaosVDRecording.h b/Engine/Source/Runtime/Experimental/ChaosVDData/Public/ChaosVDRecording.h
--- a/Engine/Source/Runtime/Experimental/ChaosVDData/Public/ChaosVDRecording.h
+++ b/Engine/Source/Runtime/Experimental/ChaosVDData/Public/ChaosVDRecording.h
@@ -131,9 +131,10 @@ struct CHAOSVDDATA_API FChaosVDRecording
 	 * Return a ptr to the existing solver frame data from the specified ID and Frame number
 	 * @param SolverID ID of the solver
 	 * @param FrameNumber Frame number
+	 * @param bKeyFrameOnly True if we should return a keyframe (real or generated) for the provided frame number if available or nothing
 	 * @return Ptr to the existing solver frame data from the specified ID and Frame number - It is a ptr to the array element, Do not store
 	 */
-	FChaosVDSolverFrameData* GetSolverFrameData_AssumesLocked(int32 SolverID, int32 FrameNumber);
+	FChaosVDSolverFrameData* GetSolverFrameData_AssumesLocked(int32 SolverID, int32 FrameNumber, bool bKeyFrameOnly = false);
 	
 	/**
 	 * Return a ptr to the existing solver frame data from the specified ID and Frame number
@@ -265,8 +266,11 @@ protected:
 	 * These are used when scrubbing to make sure the visualization is in sync with what was recorded
 	 */
 	void AddKeyFrameNumberForSolver(int32 SolverID, int32 FrameNumber);
-	
+	void AddKeyFrameNumberForSolver_AssumesLocked(int32 SolverID, int32 FrameNumber);
+	void GenerateAndStoreKeyframeForSolver_AssumesLocked(int32 SolverID, int32 CurrentFrameNumber, int32 LastKeyFrameNumber);
+
 	TMap<int32, TArray<FChaosVDSolverFrameData>> RecordedFramesDataPerSolver;
+	TMap<int32, TMap<int32, FChaosVDSolverFrameData>> GeneratedKeyFrameDataPerSolver;
 	TMap<int32, TArray<int32>> RecordedKeyFramesNumberPerSolver;
 	TArray<FChaosVDGameFrameData> GameFrames;
 	FChaosVDRecordingUpdated RecordingUpdatedDelegate;
