Response is destroyed once content has arrived

#rnx

diff --git a/Engine/Source/Runtime/Experimental/IoStoreOnDemand/Private/CoreHttp/Http.cpp b/Engine/Source/Runtime/Experimental/IoStoreOnDemand/Private/CoreHttp/Http.cpp
--- a/Engine/Source/Runtime/Experimental/IoStoreOnDemand/Private/CoreHttp/Http.cpp
+++ b/Engine/Source/Runtime/Experimental/IoStoreOnDemand/Private/CoreHttp/Http.cpp
@@ -1947,7 +1947,7 @@ uint32 FTicketStatus::GetIndex() const
 ////////////////////////////////////////////////////////////////////////////////
 FResponse& FTicketStatus::GetResponse() const
 {
-	check(GetId() <= EId::Content);
+	check(GetId() < EId::Content);
 	const auto* Activity = (FActivity*)this;
 	return *(FResponse*)Activity;
 }
