Removed unused member

#rnx

diff --git a/Engine/Source/Runtime/Experimental/IoStoreOnDemand/Private/OnDemandIoDispatcherBackend.cpp b/Engine/Source/Runtime/Experimental/IoStoreOnDemand/Private/OnDemandIoDispatcherBackend.cpp
--- a/Engine/Source/Runtime/Experimental/IoStoreOnDemand/Private/OnDemandIoDispatcherBackend.cpp
+++ b/Engine/Source/Runtime/Experimental/IoStoreOnDemand/Private/OnDemandIoDispatcherBackend.cpp
@@ -855,7 +855,6 @@ struct FChunkRequest
 	FIoRequestImpl* RequestHead;
 	FIoRequestImpl* RequestTail;
 	FIoBuffer Chunk;
-	UE::Tasks::FTask DecodeTask;
 	uint64 StartTime;
 	int32 Priority;
 	uint16 RequestCount;
