Additional bounds check when loading journal

#rnx

diff --git a/Engine/Source/Runtime/Experimental/IoStoreOnDemand/Private/JournaledCache.cpp b/Engine/Source/Runtime/Experimental/IoStoreOnDemand/Private/JournaledCache.cpp
--- a/Engine/Source/Runtime/Experimental/IoStoreOnDemand/Private/JournaledCache.cpp
+++ b/Engine/Source/Runtime/Experimental/IoStoreOnDemand/Private/JournaledCache.cpp
@@ -1155,10 +1155,18 @@ static bool LoadCache(FDiskCache& DiskCache)
 	{
 		const FPhraseDesc& Holm = Paragraphs[i].Phrase[0];
 		uint32 EntryCount = Paragraphs[i].EntryCount;
-		// The last entry may be a padded entry which should be skipped. See FDiskPhrase::ClosePhrase for details.
-		if (Holm.Entries[EntryCount-1].Key == 0)
+
+		const FDataEntry* LastEntry = Holm.Entries + (EntryCount - 1);
+		if (IsOob(LastEntry))
+		{
+			return false;
+		}
+
+		// The last entry may be a padded entry which should be skipped. See
+		// FDiskPhrase::ClosePhrase for details.
+		if (LastEntry->Key == 0)
 		{
-			check(Holm.Entries[EntryCount-1].Size == 0);
+			check(LastEntry->Size == 0);
 			--EntryCount;
 		}
 		MappedItems += EntryCount;
