Iris - Dispatch all received ordered unreliable attachments after processing as many reliable and ordered unreliable in order as possible. This expedites unreliable attachment processing.

diff --git a/Engine/Source/Runtime/Experimental/Iris/Core/Private/Iris/ReplicationSystem/AttachmentReplication.cpp b/Engine/Source/Runtime/Experimental/Iris/Core/Private/Iris/ReplicationSystem/AttachmentReplication.cpp
--- a/Engine/Source/Runtime/Experimental/Iris/Core/Private/Iris/ReplicationSystem/AttachmentReplication.cpp
+++ b/Engine/Source/Runtime/Experimental/Iris/Core/Private/Iris/ReplicationSystem/AttachmentReplication.cpp
@@ -309,7 +309,7 @@ EAttachmentWriteStatus FNetObjectAttachmentSendQueue::Serialize(FNetSerializatio
 	uint32 SerializedReliableCount = 0;
 	if (bCanSendReliableAttachments)
 	{
-		UE_NET_TRACE_SCOPE(Reliable, Writer, Context.GetTraceCollector(), ENetTraceVerbosity::Trace);
+		UE_NET_TRACE_SCOPE(Ordered, Writer, Context.GetTraceCollector(), ENetTraceVerbosity::Trace);
 		SerializedReliableCount = SerializeReliable(Context, RefHandle, ReplicationRecord.ReliableReplicationRecord);
 		// If we couldn't fit any reliable attachments then don't even try unreliable
 		if (SerializedReliableCount == 0)
@@ -674,9 +674,23 @@ public:
 		}
 	}
 
-	const TRefCountPtr<FNetBlob>* Peek() const
+	const TRefCountPtr<FNetBlob>* Peek()
 	{
-		return &Queue.Peek();
+		for (SIZE_T It = 0, EndIt = Queue.Count(); It < EndIt; ++It)
+		{
+			// Peek at head. We will return or pop the entry.
+			TRefCountPtr<FNetBlob>& Blob = Queue.Poke();
+			if (Blob.GetRefCount() > 0)
+			{
+				return &Blob;
+			}
+			else
+			{
+				Queue.Pop();
+			}
+		}
+
+		return nullptr;
 	}
 
 	void Pop()
@@ -684,6 +698,21 @@ public:
 		return Queue.Pop();
 	}
 
+	void DequeueUnreliable(TArray<TRefCountPtr<FNetBlob>>& Unreliable)
+	{
+		for (SIZE_T It = 0, EndIt = Queue.Count(); It < EndIt; ++It)
+		{
+			TRefCountPtr<FNetBlob>& RefCntBlob = Queue.PokeAtOffset(It);
+			if (RefCntBlob.GetRefCount() > 0)
+			{
+				if (const bool bIsUnReliable = !EnumHasAnyFlags(RefCntBlob.GetReference()->GetCreationInfo().Flags, ENetBlobFlags::Reliable))
+				{
+					Unreliable.Emplace(MoveTemp(RefCntBlob));
+				}
+			}
+		}
+	}
+
 private:
 	TUniquePtr<FNetBlobAssembler> NetBlobAssembler;
 	TResizableCircularQueue<TRefCountPtr<FNetBlob>> Queue;
@@ -720,7 +749,7 @@ bool FNetObjectAttachmentReceiveQueue::HasUnprocessed() const
 	return !UnreliableQueue.IsEmpty() || HasDeferredProcessingQueueUnprocessed();
 }
 
-const TRefCountPtr<FNetBlob>* FNetObjectAttachmentReceiveQueue::PeekReliable() const
+const TRefCountPtr<FNetBlob>* FNetObjectAttachmentReceiveQueue::PeekReliable()
 {
 	if (IsDeferredProcessingQueueEmpty())
 	{
@@ -754,6 +783,19 @@ void FNetObjectAttachmentReceiveQueue::PopUnreliable()
 	UnreliableQueue.Pop();
 }
 
+void FNetObjectAttachmentReceiveQueue::GetOrderedUnreliable(TArray<TRefCountPtr<FNetBlob>>& OrderedUnreliable)
+{
+	if (DeferredProcessingQueue != nullptr)
+	{
+		DeferredProcessingQueue->DequeueUnreliable(OrderedUnreliable);
+	}
+	
+	if (ReliableQueue != nullptr)
+	{
+		ReliableQueue->DequeueUnreliable(OrderedUnreliable);
+	}
+}
+
 void FNetObjectAttachmentReceiveQueue::SetUnreliableQueueCapacity(uint32 QueueCapacity)
 {
 	MaxUnreliableCount = QueueCapacity;
@@ -805,7 +847,7 @@ void FNetObjectAttachmentReceiveQueue::Deserialize(FNetSerializationContext& Con
 
 	if (bHasReliableAttachments)
 	{
-		UE_NET_TRACE_SCOPE(Reliable, Reader, Context.GetTraceCollector(), ENetTraceVerbosity::Trace);
+		UE_NET_TRACE_SCOPE(Ordered, Reader, Context.GetTraceCollector(), ENetTraceVerbosity::Trace);
 		DeserializeReliable(Context, RefHandle);
 		if (Context.HasErrorOrOverflow())
 		{
diff --git a/Engine/Source/Runtime/Experimental/Iris/Core/Private/Iris/ReplicationSystem/AttachmentReplication.h b/Engine/Source/Runtime/Experimental/Iris/Core/Private/Iris/ReplicationSystem/AttachmentReplication.h
--- a/Engine/Source/Runtime/Experimental/Iris/Core/Private/Iris/ReplicationSystem/AttachmentReplication.h
+++ b/Engine/Source/Runtime/Experimental/Iris/Core/Private/Iris/ReplicationSystem/AttachmentReplication.h
@@ -166,12 +166,14 @@ public:
 	bool IsSafeToDestroy() const;
 	bool HasUnprocessed() const;
 
-	const TRefCountPtr<FNetBlob>* PeekReliable() const;
+	const TRefCountPtr<FNetBlob>* PeekReliable();
 	void PopReliable();
 
 	const TRefCountPtr<FNetBlob>* PeekUnreliable() const;
 	void PopUnreliable();
 
+	void GetOrderedUnreliable(TArray<TRefCountPtr<FNetBlob>>& OrderedUnreliable);
+
 	void SetUnreliableQueueCapacity(uint32 QueueCapacity);
 
 private:
diff --git a/Engine/Source/Runtime/Experimental/Iris/Core/Private/Iris/ReplicationSystem/NetBlob/NetBlob.cpp b/Engine/Source/Runtime/Experimental/Iris/Core/Private/Iris/ReplicationSystem/NetBlob/NetBlob.cpp
--- a/Engine/Source/Runtime/Experimental/Iris/Core/Private/Iris/ReplicationSystem/NetBlob/NetBlob.cpp
+++ b/Engine/Source/Runtime/Experimental/Iris/Core/Private/Iris/ReplicationSystem/NetBlob/NetBlob.cpp
@@ -46,12 +46,14 @@ void FNetBlob::SerializeCreationInfo(FNetSerializationContext& Context, const FN
 {
 	FNetBitStreamWriter* Writer = Context.GetBitStreamWriter();
 
-	// Currently we only need to serialize the blob type as reliability is expected to be handled externally.
 	Writer->WriteBits((CreationInfo.Type == 0 ? 0U : 1U), 1U);
 	if (CreationInfo.Type != 0)
 	{
 		Writer->WriteBits(CreationInfo.Type, 7U);
 	}
+
+	// Retain Reliable flag
+	Writer->WriteBits(EnumHasAnyFlags(CreationInfo.Flags, ENetBlobFlags::Reliable) ? 1U : 0U, 1U);
 }
 
 void FNetBlob::DeserializeCreationInfo(FNetSerializationContext& Context, FNetBlobCreationInfo& OutCreationInfo)
@@ -63,8 +65,11 @@ void FNetBlob::DeserializeCreationInfo(FNetSerializationContext& Context, FNetBl
 		Type = Reader->ReadBits(7U);
 	}
 
+	const uint32 Reliable = Reader->ReadBits(1);
+	ENetBlobFlags Flags = (Reliable ? ENetBlobFlags::Reliable : ENetBlobFlags::None);
+
 	OutCreationInfo.Type = Type;
-	OutCreationInfo.Flags = ENetBlobFlags::None;
+	OutCreationInfo.Flags = Flags;
 }
 
 void FNetBlob::SerializeWithObject(FNetSerializationContext& Context, FNetRefHandle RefHandle) const
diff --git a/Engine/Source/Runtime/Experimental/Iris/Core/Private/Iris/ReplicationSystem/NetBlob/ReliableNetBlobQueue.cpp b/Engine/Source/Runtime/Experimental/Iris/Core/Private/Iris/ReplicationSystem/NetBlob/ReliableNetBlobQueue.cpp
--- a/Engine/Source/Runtime/Experimental/Iris/Core/Private/Iris/ReplicationSystem/NetBlob/ReliableNetBlobQueue.cpp
+++ b/Engine/Source/Runtime/Experimental/Iris/Core/Private/Iris/ReplicationSystem/NetBlob/ReliableNetBlobQueue.cpp
@@ -342,6 +342,35 @@ void FReliableNetBlobQueue::Pop()
 	// $TODO. For a memory optimization one can change the storage implementation and free the memory here if everything is acked.
 }
 
+
+void FReliableNetBlobQueue::DequeueUnreliable(TArray<TRefCountPtr<FNetBlob>>& Unreliable)
+{
+	for (uint32 Seq = FirstSeq; Seq < LastSeq; ++Seq)
+	{
+		const uint32 Index = SequenceToIndex(Seq);
+		if (!IsIndexAcked(Index))
+		{
+			continue;
+		}
+
+		TRefCountPtr<FNetBlob>& RefCntBlob = NetBlobs[Index];
+		if (RefCntBlob.GetRefCount() > 0)
+		{
+			if (const bool bIsUnReliable = !EnumHasAnyFlags(RefCntBlob.GetReference()->GetCreationInfo().Flags, ENetBlobFlags::Reliable))
+			{
+				Unreliable.Emplace(MoveTemp(RefCntBlob));
+
+				// Advance the ack window if possible.
+				if (Seq == FirstSeq)
+				{
+					ClearIndexIsAcked(Index);
+					++FirstSeq;
+				}
+			}
+		}
+	}
+}
+
 void FReliableNetBlobQueue::ProcessPacketDeliveryStatus(EPacketDeliveryStatus Status, const FReliableNetBlobQueue::FReplicationRecord& Record)
 {
 	switch (Status)
diff --git a/Engine/Source/Runtime/Experimental/Iris/Core/Private/Iris/ReplicationSystem/ReplicationReader.cpp b/Engine/Source/Runtime/Experimental/Iris/Core/Private/Iris/ReplicationSystem/ReplicationReader.cpp
--- a/Engine/Source/Runtime/Experimental/Iris/Core/Private/Iris/ReplicationSystem/ReplicationReader.cpp
+++ b/Engine/Source/Runtime/Experimental/Iris/Core/Private/Iris/ReplicationSystem/ReplicationReader.cpp
@@ -2103,9 +2103,9 @@ void FReplicationReader::ResolveAndDispatchAttachments(FNetSerializationContext&
 	const uint32 InternalIndex = ReplicationInfo->InternalIndex;
 
 	/**
-	* This code path handles all cases where the initial state has already been applied. An object can have multiple entries in ObjectsPendingResolve.
-	* Reliable attachments will be delivered if they can be resolved or if CVarDelayUnmappedRPCs is <= 0
-	*/
+	 * This code path handles all cases where the initial state has already been applied. An object can have multiple entries in ObjectsPendingResolve.
+	 * Reliable attachments will be dispatched if they can be resolved or if CVarDelayUnmappedRPCs is <= 0. Unreliable but ordered attachments will always be dispatched.
+	 */
 	bool bHasUnresolvedReferences = false;
 	const ENetObjectAttachmentType AttachmentType = (IsObjectIndexForOOBAttachment(InternalIndex) ? ENetObjectAttachmentType::OutOfBand : ENetObjectAttachmentType::Normal);
 	if (FNetObjectAttachmentReceiveQueue* AttachmentQueue = Attachments.GetQueue(AttachmentType, InternalIndex))
@@ -2153,6 +2153,21 @@ void FReplicationReader::ResolveAndDispatchAttachments(FNetSerializationContext&
 					return;
 				}
 			}
+
+			// Dispatch remaining ordered unreliable attachments. The side effects of for example blocking move RPCs are significant.
+			{
+				TArray<TRefCountPtr<FNetBlob>> OrderedUnreliableAttachments;
+				OrderedUnreliableAttachments.Reserve(16);
+				AttachmentQueue->GetOrderedUnreliable(OrderedUnreliableAttachments);
+				for (TRefCountPtr<FNetBlob>& Attachment : OrderedUnreliableAttachments)
+				{
+					NetBlobHandlerManager->OnNetBlobReceived(Context, Attachment);
+					if (Context.HasError())
+					{
+						return;
+					}
+				}
+			}
 		}
 
 		if (bDispatchUnreliableAttachments)
diff --git a/Engine/Source/Runtime/Experimental/Iris/Core/Public/Iris/ReplicationSystem/NetBlob/ReliableNetBlobQueue.h b/Engine/Source/Runtime/Experimental/Iris/Core/Public/Iris/ReplicationSystem/NetBlob/ReliableNetBlobQueue.h
--- a/Engine/Source/Runtime/Experimental/Iris/Core/Public/Iris/ReplicationSystem/NetBlob/ReliableNetBlobQueue.h
+++ b/Engine/Source/Runtime/Experimental/Iris/Core/Public/Iris/ReplicationSystem/NetBlob/ReliableNetBlobQueue.h
@@ -86,6 +86,9 @@ public:
 	/** On the receiving end this will remove the next blob to be processed from the queue. Call after processing the blob returned from Peek(). */
 	IRISCORE_API void Pop();
 
+	/** On the receiving end this will move all received unreliable NetBlobs to the array and release them from the queue. This breaks the ordering guarantees provided by using Peek and Pop. Reliable NetBlobs are unaffected by this operation. */
+	IRISCORE_API void DequeueUnreliable(TArray<TRefCountPtr<FNetBlob>>& Unreliable);
+
 	/**
 	 * Serializes as many blobs as possible using their respective SerializeWithObject() method. It is assumed the NetRefHandle will be 
 	 * reconstructed somehow on the receiving end and passed to DeserializeWithObject().
