Fix Cluster Union initial anchored state.

Problem fixed: The initial state of a cluster in non-anchored and simulating. On the server everything works fine if you set the anchored state after adding all components to the cluster and checking for support because it all happens on the same tick. But on the client we may receive the components and anchored state change over multiple frames. If anthing interacts with the client-side cluster after components are added but before the anchored flag replciates, it will move and later freeze in its new location when the anchored state is received.

diff --git a/Engine/Source/Runtime/Engine/Classes/PhysicsEngine/ClusterUnionComponent.h b/Engine/Source/Runtime/Engine/Classes/PhysicsEngine/ClusterUnionComponent.h
--- a/Engine/Source/Runtime/Engine/Classes/PhysicsEngine/ClusterUnionComponent.h
+++ b/Engine/Source/Runtime/Engine/Classes/PhysicsEngine/ClusterUnionComponent.h
@@ -127,7 +127,7 @@ struct FClusterUnionReplicatedData
 	uint8 ObjectState = 0;
 
 	UPROPERTY()
-	bool bIsAnchored = false;
+	bool bIsAnchored = true;
 
 	FClusterUnionReplicatedData& operator=(const FClusterUnionReplicatedData&) = default;
 	bool operator==(const FClusterUnionReplicatedData&) const = default;
diff --git a/Engine/Source/Runtime/Experimental/Chaos/Public/PhysicsProxy/ClusterUnionPhysicsProxy.h b/Engine/Source/Runtime/Experimental/Chaos/Public/PhysicsProxy/ClusterUnionPhysicsProxy.h
--- a/Engine/Source/Runtime/Experimental/Chaos/Public/PhysicsProxy/ClusterUnionPhysicsProxy.h
+++ b/Engine/Source/Runtime/Experimental/Chaos/Public/PhysicsProxy/ClusterUnionPhysicsProxy.h
@@ -44,10 +44,10 @@ namespace Chaos
 	struct FClusterUnionSyncedData
 	{
 		// Whether the cluster is anchored or not.
-		bool bIsAnchored;
+		bool bIsAnchored = true;
 
 		// Whether we assigned new geometry from the PT in PullFromPhysicsState
-		bool bDidSyncGeometry;
+		bool bDidSyncGeometry = false;
 
 		// Data on every child particle in the cluster union.
 		TArray<FClusterUnionChildData> ChildParticles;
