Chaos - fix joint force reporting (when projection was enabled it was getting zeroed)

diff --git a/Engine/Source/Runtime/Experimental/Chaos/Private/Chaos/Joint/PBDJointCachedSolverGaussSeidel.cpp b/Engine/Source/Runtime/Experimental/Chaos/Private/Chaos/Joint/PBDJointCachedSolverGaussSeidel.cpp
--- a/Engine/Source/Runtime/Experimental/Chaos/Private/Chaos/Joint/PBDJointCachedSolverGaussSeidel.cpp
+++ b/Engine/Source/Runtime/Experimental/Chaos/Private/Chaos/Joint/PBDJointCachedSolverGaussSeidel.cpp
@@ -8,8 +8,6 @@
 #include "Chaos/Utilities.h"
 
 
-//UE_DISABLE_OPTIMIZATION
-
 namespace Chaos
 {
 	
@@ -221,8 +219,9 @@ void FPBDJointCachedSolver::Init(
 	UpdateMass1(ConditionedInvMs[1], ConditionedInvILs[1]);
 
 	// Cache all the informations for the position and rotation constraints
-	InitPositionConstraints(Dt, SolverSettings, JointSettings);
-	InitRotationConstraints(Dt, SolverSettings, JointSettings);
+	const bool bResetLambdas = true;	// zero accumulators on full init
+	InitPositionConstraints(Dt, SolverSettings, JointSettings, bResetLambdas);
+	InitRotationConstraints(Dt, SolverSettings, JointSettings, bResetLambdas);
 
 	InitPositionDrives(Dt, SolverSettings, JointSettings);
 	InitRotationDrives(Dt, SolverSettings, JointSettings);
@@ -256,14 +255,17 @@ void FPBDJointCachedSolver::InitProjection(
 		InvMs[1] = Body1().InvM();
 		InvIs[1] = FMatrix33::FromDiagonal(FVec3(Body1().InvILocal().GetMin()));
 
+		// We are reusing the constraints for projection but we don't want to reset the accumulated corrections
+		const bool bResetLambdas = false;
+
 		if(bHasLinearProjection)
 		{
-			InitPositionConstraints(Dt, SolverSettings, JointSettings);
+			InitPositionConstraints(Dt, SolverSettings, JointSettings, bResetLambdas);
 		}
 
 		if(bHasAngularProjection)
 		{
-			InitRotationConstraints(Dt, SolverSettings, JointSettings);
+			InitRotationConstraints(Dt, SolverSettings, JointSettings, bResetLambdas);
 		}
 	}
 }
@@ -370,7 +372,8 @@ FORCEINLINE bool ExtractLinearMotion( const FPBDJointSettings& JointSettings,
 void FPBDJointCachedSolver::InitPositionConstraints(
 	const FReal Dt,
 	const FPBDJointSolverSettings& SolverSettings,
-	const FPBDJointSettings& JointSettings)
+	const FPBDJointSettings& JointSettings,
+	const bool bResetLambdas)
 {
 	PositionConstraints.bValidDatas[0] = false;
 	PositionConstraints.bValidDatas[1] = false;
@@ -388,7 +391,8 @@ void FPBDJointCachedSolver::InitPositionConstraints(
 			FPBDJointUtilities::GetSoftLinearLimitEnabled(SolverSettings, JointSettings),
 			FPBDJointUtilities::GetSoftLinearStiffness(SolverSettings, JointSettings),
 			FPBDJointUtilities::GetSoftLinearDamping(SolverSettings, JointSettings),
-			FPBDJointUtilities::GetLinearStiffness(SolverSettings, JointSettings));
+			FPBDJointUtilities::GetLinearStiffness(SolverSettings, JointSettings),
+			bResetLambdas);
 	}
 
 	const TVec3<EJointMotionType>& LinearMotion = JointSettings.LinearMotionTypes;
@@ -828,7 +832,8 @@ FORCEINLINE bool ExtractAngularMotion( const FPBDJointSettings& JointSettings,
 void FPBDJointCachedSolver::InitRotationConstraints(
 	const FReal Dt,
 	const FPBDJointSolverSettings& SolverSettings,
-	const FPBDJointSettings& JointSettings)
+	const FPBDJointSettings& JointSettings,
+	const bool bResetLambdas)
 {
 	RotationConstraints.bValidDatas[0] = false;
 	RotationConstraints.bValidDatas[1] = false;
@@ -847,17 +852,20 @@ void FPBDJointCachedSolver::InitRotationConstraints(
 	RotationConstraints.InitDatas(TW,FPBDJointUtilities::GetSoftTwistLimitEnabled(SolverSettings, JointSettings) && !bAngularLocked[TW],
 					FPBDJointUtilities::GetSoftTwistStiffness(SolverSettings, JointSettings),
 					FPBDJointUtilities::GetSoftTwistDamping(SolverSettings, JointSettings),
-					FPBDJointUtilities::GetTwistStiffness(SolverSettings, JointSettings));
+					FPBDJointUtilities::GetTwistStiffness(SolverSettings, JointSettings),
+					bResetLambdas);
 
 	RotationConstraints.InitDatas(S1,FPBDJointUtilities::GetSoftSwingLimitEnabled(SolverSettings, JointSettings) && !bAngularLocked[S1],
 					FPBDJointUtilities::GetSoftSwingStiffness(SolverSettings, JointSettings),
 					FPBDJointUtilities::GetSoftSwingDamping(SolverSettings, JointSettings),
-					FPBDJointUtilities::GetSwingStiffness(SolverSettings, JointSettings));
+					FPBDJointUtilities::GetSwingStiffness(SolverSettings, JointSettings),
+					bResetLambdas);
 
 	RotationConstraints.InitDatas(S2, FPBDJointUtilities::GetSoftSwingLimitEnabled(SolverSettings, JointSettings) && !bAngularLocked[S2],
 					FPBDJointUtilities::GetSoftSwingStiffness(SolverSettings, JointSettings),
 					FPBDJointUtilities::GetSoftSwingDamping(SolverSettings, JointSettings),
-					FPBDJointUtilities::GetSwingStiffness(SolverSettings, JointSettings));
+					FPBDJointUtilities::GetSwingStiffness(SolverSettings, JointSettings),
+					bResetLambdas);
 
 	const FVec3 Twist0 = ConnectorRs[0] * FJointConstants::TwistAxis();
 	const FVec3 Twist1 = ConnectorRs[1] * FJointConstants::TwistAxis();
@@ -1979,7 +1987,8 @@ void FAxisConstraintDatas::InitDatas(
 	const bool bHasSoftLimits,
 	const FReal SoftStiffness,
 	const FReal SoftDamping,
-	const FReal HardStiffness)
+	const FReal HardStiffness,
+	const bool bResetLambdas)
 {
 	bSoftLimit[ConstraintIndex] = bHasSoftLimits;
 	ConstraintHardStiffness[ConstraintIndex] = HardStiffness;
@@ -1990,10 +1999,13 @@ void FAxisConstraintDatas::InitDatas(
 	SettingsSoftDamping[ConstraintIndex] = SoftDamping;
 	bValidDatas[ConstraintIndex] = false;
 	bLimitsCheck[ConstraintIndex] = true;
-	ConstraintLambda = FVec3::Zero();
-	ConstraintLambdaVelocity = FVec3::Zero();
-	ConstraintLimits = FVec3::Zero();
 	MotionType[ConstraintIndex] = EJointMotionType::Free;
+	if (bResetLambdas)
+	{
+		ConstraintLambda = FVec3::Zero();
+		ConstraintLambdaVelocity = FVec3::Zero();
+		ConstraintLimits = FVec3::Zero();
+	}
 }
 
 void FAxisConstraintDatas::UpdateDatas(
@@ -2028,8 +2040,6 @@ void FAxisConstraintDatas::UpdateMass(
 	ConstraintDRAxis[ConstraintIndex][0] = DatasIA0;
 	ConstraintDRAxis[ConstraintIndex][1] = -DatasIA1;
 	ConstraintHardIM[ConstraintIndex] = DatasIM;
-	ConstraintLambda = FVec3::Zero();
-	ConstraintLambdaVelocity = FVec3::Zero();
 
 	if(bSoftLimit[ConstraintIndex])
 	{
diff --git a/Engine/Source/Runtime/Experimental/Chaos/Public/Chaos/Joint/PBDJointCachedSolverGaussSeidel.h b/Engine/Source/Runtime/Experimental/Chaos/Public/Chaos/Joint/PBDJointCachedSolverGaussSeidel.h
--- a/Engine/Source/Runtime/Experimental/Chaos/Public/Chaos/Joint/PBDJointCachedSolverGaussSeidel.h
+++ b/Engine/Source/Runtime/Experimental/Chaos/Public/Chaos/Joint/PBDJointCachedSolverGaussSeidel.h
@@ -22,7 +22,8 @@ struct FAxisConstraintDatas
 		const bool bHasSoftLimits,
 		const FReal SoftStiffness,
 		const FReal SoftDamping,
-		const FReal HardStiffness);
+		const FReal HardStiffness,
+		const bool bResetLambdas = true);
 
 	/** Initialize the max force ( force is modified for acceleration mode and converted to positional impulse limit) */
 	void SetMaxForce(
@@ -362,7 +363,8 @@ struct FAxisConstraintDatas
 		void InitPositionConstraints(
 			const FReal Dt,
 			const FPBDJointSolverSettings& SolverSettings,
-			const FPBDJointSettings& JointSettings);
+			const FPBDJointSettings& JointSettings,
+			const bool bResetLambdas);
 
 		void InitPositionDatasMass(
 			FAxisConstraintDatas& PositionDatas,
@@ -439,7 +441,8 @@ struct FAxisConstraintDatas
 		void InitRotationConstraints(
 			const FReal Dt,
 			const FPBDJointSolverSettings& SolverSettings,
-			const FPBDJointSettings& JointSettings);
+			const FPBDJointSettings& JointSettings,
+			const bool bResetLambdas);
 		
 		void CorrectAxisAngleConstraint(
 			const FPBDJointSettings& JointSettings,
