Drop the cache if loading fails significantly

#rnx

diff --git a/Engine/Source/Runtime/Experimental/IoStoreOnDemand/Private/JournaledCache.cpp b/Engine/Source/Runtime/Experimental/IoStoreOnDemand/Private/JournaledCache.cpp
--- a/Engine/Source/Runtime/Experimental/IoStoreOnDemand/Private/JournaledCache.cpp
+++ b/Engine/Source/Runtime/Experimental/IoStoreOnDemand/Private/JournaledCache.cpp
@@ -55,7 +55,7 @@ TRACE_DECLARE_INT_COUNTER(IasOpCount,   TEXT("Ias/CacheOpCount"));
 TRACE_DECLARE_INT_COUNTER(IasReadCursor,TEXT("Ias/CacheReadCursor"));
 
 ////////////////////////////////////////////////////////////////////////////////
-static bool LoadCache(class FDiskCache&);
+static int32 LoadCache(class FDiskCache&);
 
 ////////////////////////////////////////////////////////////////////////////////
 enum class EAilments
@@ -390,7 +390,7 @@ public:
 	uint32					GetPreviousPartial() const { return PreviousPartial; };
 
 private:
-	friend bool				LoadCache(FDiskCache&);
+	friend int32			LoadCache(FDiskCache&);
 	void					GetPath(TStringBuilder<64>& Out);
 	void					OpenJrnFile();
 	static uint32			HashBytes(const uint8* Data, uint32 Size, uint32 Seed);
@@ -626,7 +626,7 @@ private:
 	};
 	static_assert(sizeof(FMapEntry) == sizeof(uint64));
 
-	friend bool				LoadCache(FDiskCache&);
+	friend int32			LoadCache(FDiskCache&);
 	void					OpenDataFile();
 	using					FDataMap = TMap<uint64, FMapEntry>;
 	void					Spam();
@@ -933,7 +933,7 @@ uint32 FDiskCache::DebugVisit(void* Param, FDebugCacheEntry::Callback* Callback)
 // {{{1 loader .................................................................
 
 ////////////////////////////////////////////////////////////////////////////////
-static bool LoadCache(FDiskCache& DiskCache)
+static int32 LoadCache(FDiskCache& DiskCache)
 {
 	FWriteScopeLock _(DiskCache.Lock);
 
@@ -947,7 +947,7 @@ static bool LoadCache(FDiskCache& DiskCache)
 		DataSize = uint32(Handle->Size());
 		if (DataSize == 0)
 		{
-			return false;
+			return 0;
 		}
 
 		Data = TUniquePtr<uint8[]>(new uint8[DataSize]);
@@ -960,7 +960,7 @@ static bool LoadCache(FDiskCache& DiskCache)
 
 	if (DataSize == 0)
 	{
-		return false;
+		return 0;
 	}
 
 	UE_LOG(LogIas, VeryVerbose, TEXT("JournaledCache: %u byte journal found"), DataSize);
@@ -1055,7 +1055,7 @@ static bool LoadCache(FDiskCache& DiskCache)
 
 	if (Paragraphs.IsEmpty())
 	{
-		return false;
+		return -1;
 	}
 
 	auto LessWithWrap = [] (
@@ -1095,7 +1095,7 @@ static bool LoadCache(FDiskCache& DiskCache)
 	if (!DiskCache.DataHandle.IsValid())
 	{
 		UE_LOG(LogIas, VeryVerbose, TEXT("JournaledCache: unable to open '%s'"), *DiskCache.BinPath);
-		return false;
+		return -1;
 	}
 
 	IFileHandle* File = DiskCache.DataHandle.Get();
@@ -1105,8 +1105,7 @@ static bool LoadCache(FDiskCache& DiskCache)
 			TEXT("JournaledCache: Dropping - existing cache to bi; %llu/%llu"),
 			uint64(File->Size()), DiskCache.MaxDataSize
 		);
-		DiskCache.Drop();
-		return false;
+		return -1;
 	}
 
 	// Detect data writes that are newer than any journal flushes.
@@ -1159,7 +1158,7 @@ static bool LoadCache(FDiskCache& DiskCache)
 		const FDataEntry* LastEntry = Holm.Entries + (EntryCount - 1);
 		if (IsOob(LastEntry))
 		{
-			return false;
+			return -1;
 		}
 
 		// The last entry may be a padded entry which should be skipped. See
@@ -1190,8 +1189,7 @@ static bool LoadCache(FDiskCache& DiskCache)
 			TEXT("JournaledCache: Journal exceeds given size - dropping; %u/%u"),
 			DataSize, Journal.MaxSize
 		);
-		Journal.Drop();
-		return false;
+		return -1;
 	}
 
 	// Prime the disk-cache's state
@@ -1202,12 +1200,10 @@ static bool LoadCache(FDiskCache& DiskCache)
 			TEXT("JournaledCache: Dropping - DataCursor too big; %llu/%llu"),
 			DiskCache.DataCursor, DiskCache.MaxDataSize
 		);
-		DiskCache.Drop();
-		return false;
+		return -1;
 	}
 
-	// DiskCache.Spam();
-	return true;
+	return 1;
 }
 
 // {{{1 cache ..................................................................
@@ -1275,7 +1271,13 @@ uint32 FCache::GetAilments() const
 ////////////////////////////////////////////////////////////////////////////////
 bool FCache::Load()
 {
-	return LoadCache(DiskCache);
+	int32 Result = LoadCache(DiskCache);
+	if (Result < 0)
+	{
+		Drop();
+	}
+	
+	return Result > 0;
 }
 
 ////////////////////////////////////////////////////////////////////////////////
