CoreHttp - enable platform override for setting non blocking socket option

diff --git a/Engine/Source/Runtime/Experimental/IoStoreOnDemand/Private/CoreHttp/Http.cpp b/Engine/Source/Runtime/Experimental/IoStoreOnDemand/Private/CoreHttp/Http.cpp
--- a/Engine/Source/Runtime/Experimental/IoStoreOnDemand/Private/CoreHttp/Http.cpp
+++ b/Engine/Source/Runtime/Experimental/IoStoreOnDemand/Private/CoreHttp/Http.cpp
@@ -901,8 +901,9 @@ int32 FSocket::Recv(char* Dest, uint32 Size)
 bool FSocket::SetBlocking(bool bBlocking)
 {
 	bool bSuccess = false;
-
-#if PLATFORM_MICROSOFT
+#if defined(IAS_HTTP_HAS_NONBLOCK_IMPL)
+	bSuccess = SetNonBlockingSocket(Socket);
+#elif PLATFORM_MICROSOFT
 	unsigned long NonBlockingMode = 1;
 	if (ioctlsocket(Socket, FIONBIO, &NonBlockingMode) != SOCKET_ERROR)
 	{
