IAS can now be provided with a fallback cdn url to be used if the distributed endpoint cannot be contacted.

#rnx

- If a user is unable to contact the distributed endpoint IAS will continue to retry the connection until it works. However it is possible that some users might never be able to contact the endpoint even if their network connection is okay. This means they would never be able to download IAS data and end up with poor visuals even though they should be able to access the data.
- To fix this we have added the ability to add a fallback url to the IoStoreOnDemand.ini which the user will try connecting to after a certain number of failed attempts at the distributed endpoint.
- To set the url add "-FallbackUrl=..." when uploading IAS data via UnrealPak.exe, although note that this will only be used if you are also providing "-DistributionUrl=..."
- The url can be overriden if needed by setting the cvar "ias.DistributedEndpointFallbackUrl" on startup
- The cvar "ias.DistributedEndpointAttemptCount" controls how many attempts the user should make at resolving the distributed endpoint before trying the fallback url. If the url is not set then this value is ignored and we will continue with the current behavior of retrying as many times as it takes.

diff --git a/Engine/Source/Developer/IoStoreUtilities/Private/IoStoreUtilities.cpp b/Engine/Source/Developer/IoStoreUtilities/Private/IoStoreUtilities.cpp
--- a/Engine/Source/Developer/IoStoreUtilities/Private/IoStoreUtilities.cpp
+++ b/Engine/Source/Developer/IoStoreUtilities/Private/IoStoreUtilities.cpp
@@ -9252,6 +9252,13 @@ bool UploadIoStoreContainerFiles(const UE::IO::IAS::FIoStoreUploadParams& Upload
 		if (FParse::Value(FCommandLine::Get(), TEXT("DistributionUrl="), DistributionUrl))
 		{
 			Sb << TEXT("DistributionUrl=\"") << DistributionUrl << TEXT("\"\r\n");
+
+			FString FallbackUrl;
+			if (FParse::Value(FCommandLine::Get(), TEXT("FallbackUrl="), FallbackUrl))
+			{
+				Sb << TEXT("FallbackUrl=\"") << FallbackUrl << TEXT("\"\r\n");
+			}
+			
 			Sb << TEXT("TocPath=\"") << UploadResult.TocPath << TEXT("\"\r\n");
 		}
 		else
diff --git a/Engine/Source/Runtime/Experimental/IoStoreOnDemand/Private/IoStoreOnDemand.cpp b/Engine/Source/Runtime/Experimental/IoStoreOnDemand/Private/IoStoreOnDemand.cpp
--- a/Engine/Source/Runtime/Experimental/IoStoreOnDemand/Private/IoStoreOnDemand.cpp
+++ b/Engine/Source/Runtime/Experimental/IoStoreOnDemand/Private/IoStoreOnDemand.cpp
@@ -57,6 +57,14 @@ static FAutoConsoleVariableRef CVar_SuspendSystemEnabled(
 	ECVF_ReadOnly
 );
 
+/** Temp cvar to allow the fallback url to be hotfixed in case of problems */
+static FString GDistributedEndpointFallbackUrl;
+static FAutoConsoleVariableRef CVar_DistributedEndpointFallbackUrl(
+	TEXT("ias.DistributedEndpointFallbackUrl"),
+	GDistributedEndpointFallbackUrl,
+	TEXT("CDN url to be used if a distributed endpoint cannot be reached (overrides IoStoreOnDemand.ini)")
+);
+
 ////////////////////////////////////////////////////////////////////////////////
 static int64 ParseSizeParam(FStringView Value)
 {
@@ -139,6 +147,16 @@ static bool TryParseConfigContent(const FString& ConfigContent, const FString& C
 	Config.ProcessInputFileContents(ConfigContent, ConfigFileName);
 
 	Config.GetString(TEXT("Endpoint"), TEXT("DistributionUrl"), OutEndpoint.DistributionUrl);
+	if (!OutEndpoint.DistributionUrl.IsEmpty())
+	{
+		Config.GetString(TEXT("Endpoint"), TEXT("FallbackUrl"), OutEndpoint.FallbackUrl);
+
+		if (!GDistributedEndpointFallbackUrl.IsEmpty())
+		{
+			OutEndpoint.FallbackUrl = GDistributedEndpointFallbackUrl;
+		}
+	}
+	
 	Config.GetArray(TEXT("Endpoint"), TEXT("ServiceUrl"), OutEndpoint.ServiceUrls);
 	Config.GetString(TEXT("Endpoint"), TEXT("TocPath"), OutEndpoint.TocPath);
 
diff --git a/Engine/Source/Runtime/Experimental/IoStoreOnDemand/Private/OnDemandIoDispatcherBackend.cpp b/Engine/Source/Runtime/Experimental/IoStoreOnDemand/Private/OnDemandIoDispatcherBackend.cpp
--- a/Engine/Source/Runtime/Experimental/IoStoreOnDemand/Private/OnDemandIoDispatcherBackend.cpp
+++ b/Engine/Source/Runtime/Experimental/IoStoreOnDemand/Private/OnDemandIoDispatcherBackend.cpp
@@ -229,6 +229,13 @@ static FAutoConsoleVariableRef CVar_DistributedEndpointRetryWaitTime(
 	TEXT("How long to wait (in seconds) after failing to resolve a distributed endpoint before retrying")
 );
 
+static int32 GDistributedEndpointAttemptCount = 5;
+static FAutoConsoleVariableRef CVar_DistributedEndpointAttemptCount(
+	TEXT("ias.DistributedEndpointAttemptCount"),
+	GDistributedEndpointAttemptCount,
+	TEXT("Number of times we should try to resolve a distributed endpoint befor eusing the fallback url (if there is one)")
+);
+
 #if !UE_BUILD_SHIPPING
 static FAutoConsoleCommand CVar_IasAbandonCache(
 	TEXT("Ias.AbandonCache"),
@@ -1576,7 +1583,7 @@ private:
 	FIoStatus ApplyGeneratedOnDemandToc(const FString& CdnUrl, const FString& TocPath);
 	FIoStatus DownloadoadOnDemandToc(const FString& CdnUrl, const FString& TocPath);
 
-	bool ResolveDistributedEndpoint(const FString& Url);
+	bool ResolveDistributedEndpoint(const FDistributedEndpointUrl& Url);
 	void InitializePrimaryEndpoint();
 
 	/** Mode to control which types of .iochunktoc are flush */
@@ -1616,7 +1623,7 @@ private:
 	FOnDemandIoBackendStats Stats;
 	FBackendStatus BackendStatus;
 	FAvailableEps AvailableEps;
-	FString DistributionUrl;
+	FDistributedEndpointUrl DistributionUrl;
 	FEventRef DistributedEndpointEvent;
 
 	mutable FRWLock Lock;
@@ -2092,11 +2099,11 @@ FIoStatus FOnDemandIoBackend::DownloadoadOnDemandToc(const FString& CdnUrl, cons
 	}
 }
 
-bool FOnDemandIoBackend::ResolveDistributedEndpoint(const FString& DistributedEndpointUrl)
+bool FOnDemandIoBackend::ResolveDistributedEndpoint(const FDistributedEndpointUrl& DistributedEndpointUrl)
 {
 	TRACE_CPUPROFILER_EVENT_SCOPE(IasBackend::ResolveDistributedEndpoint);
 
-	check(!DistributedEndpointUrl.IsEmpty());
+	check(DistributedEndpointUrl.IsValid());
 
 	// We need to resolve the end point in this method which occurs after the config system has initialized
 	// rather than in ::Mount which can occur before that.
@@ -2104,12 +2111,14 @@ bool FOnDemandIoBackend::ResolveDistributedEndpoint(const FString& DistributedEn
 	// to resolve and the OnDemand system will not recover.
 	check(GConfig->IsReadyForUse());
 
+	int32 NumAttempts = 0;
+
 	while (!bStopRequested)
 	{
 		TArray<FString> ServiceUrls;
 
 		FDistributionEndpoints Resolver;
-		FDistributionEndpoints::EResult Result = Resolver.ResolveEndpoints(DistributedEndpointUrl, ServiceUrls, *DistributedEndpointEvent.Get());
+		FDistributionEndpoints::EResult Result = Resolver.ResolveEndpoints(DistributedEndpointUrl.EndpointUrl, ServiceUrls, *DistributedEndpointEvent.Get());
 		if (Result == FDistributionEndpoints::EResult::Success)
 		{
 			FWriteScopeLock _(Lock);
@@ -2121,6 +2130,17 @@ bool FOnDemandIoBackend::ResolveDistributedEndpoint(const FString& DistributedEn
 			return true;
 		}
 
+		if (DistributedEndpointUrl.HasFallbackUrl() && ++NumAttempts == GDistributedEndpointAttemptCount)
+		{
+			FString FallbackUrl = DistributedEndpointUrl.FallbackUrl.Replace(TEXT("https"), TEXT("http"));
+			UE_LOG(LogIas, Warning, TEXT("Failed to resolve the distributed endpoint %d times. Fallback CDN '%s' will be used instead"), GDistributedEndpointAttemptCount , *FallbackUrl);
+			
+			FWriteScopeLock _(Lock);
+			AvailableEps.Urls.Emplace(MoveTemp(FallbackUrl));
+		
+			return true;
+		}
+
 		if (!bStopRequested)
 		{
 			const uint32 WaitTime = GDistributedEndpointRetryWaitTime >= 0 ? (static_cast<uint32>(GDistributedEndpointRetryWaitTime) * 1000) : MAX_uint32;
@@ -2249,9 +2269,9 @@ void FOnDemandIoBackend::Mount(const FOnDemandEndpoint& Endpoint)
 			{
 				if (AvailableEps.HasCurrent() == false)
 				{
-					if (DistributionUrl.IsEmpty())
+					if (!DistributionUrl.IsValid())
 					{
-						DistributionUrl = Endpoint.DistributionUrl;
+						DistributionUrl = { Endpoint.DistributionUrl, Endpoint.FallbackUrl };
 					}
 					DeferredTocs.Add(FTocParams{Endpoint.TocPath, Endpoint.bForceTocDownload});
 					return;
@@ -2478,13 +2498,16 @@ bool FOnDemandIoBackend::SetupHttpThread()
 
 	// Note that the following method will block until the distributed endpoint (if we have one) has
 	// been resolved as the thread cannot function without it.
-	if (!DistributionUrl.IsEmpty())
+	if (DistributionUrl.IsValid())
 	{
-		if (!ResolveDistributedEndpoint(DistributionUrl))
+		const bool bResult = ResolveDistributedEndpoint(DistributionUrl);
+		DistributionUrl.Reset();
+
+		if (!bResult)
 		{
 			return false;
 		}
-		DistributionUrl.Reset();
+		
 	}
 
 	InitializePrimaryEndpoint();
diff --git a/Engine/Source/Runtime/Experimental/IoStoreOnDemand/Private/OnDemandIoDispatcherBackend.h b/Engine/Source/Runtime/Experimental/IoStoreOnDemand/Private/OnDemandIoDispatcherBackend.h
--- a/Engine/Source/Runtime/Experimental/IoStoreOnDemand/Private/OnDemandIoDispatcherBackend.h
+++ b/Engine/Source/Runtime/Experimental/IoStoreOnDemand/Private/OnDemandIoDispatcherBackend.h
@@ -17,10 +17,34 @@ enum class EOnDemandEndpointType
 	ZEN
 };
 
+struct FDistributedEndpointUrl
+{
+	FString EndpointUrl;
+	FString FallbackUrl;
+
+	bool IsValid() const
+	{
+		return !EndpointUrl.IsEmpty();
+	}
+
+	bool HasFallbackUrl() const
+	{
+		return !FallbackUrl.IsEmpty();
+	}
+
+	void Reset()
+	{
+		EndpointUrl.Empty();
+		FallbackUrl.Empty();
+	}
+};
+
 struct FOnDemandEndpoint
 {
 	EOnDemandEndpointType EndpointType;
 	FString DistributionUrl;
+	FString FallbackUrl;
+
 	TArray<FString> ServiceUrls;
 	FString TocPath;
 
