Constraints: give the constraint manager a chance to register constraints

In the context of loading pre 5.4 maps, we have to make sure that we don't remove constraint with valid handles.
"valid" means that the components are valid but not necessary bound (control rig).

diff --git a/Engine/Plugins/Animation/ControlRig/Source/ControlRig/Private/Constraints/ControlRigTransformableHandle.cpp b/Engine/Plugins/Animation/ControlRig/Source/ControlRig/Private/Constraints/ControlRigTransformableHandle.cpp
--- a/Engine/Plugins/Animation/ControlRig/Source/ControlRig/Private/Constraints/ControlRigTransformableHandle.cpp
+++ b/Engine/Plugins/Animation/ControlRig/Source/ControlRig/Private/Constraints/ControlRigTransformableHandle.cpp
@@ -31,23 +31,26 @@ void UTransformableControlHandle::PostLoad()
 	RegisterDelegates();
 }
 
-bool UTransformableControlHandle::IsValid() const
+bool UTransformableControlHandle::IsValid(const bool bDeepCheck) const
 {
 	if (!ControlRig.IsValid() || ControlName == NAME_None)
 	{
 		return false;
 	}
 
-	const USceneComponent* BoundComponent = GetBoundComponent();
-	if (!BoundComponent)
+	const FRigControlElement* ControlElement = ControlRig->FindControl(ControlName);
+	if (!ControlElement)
 	{
 		return false;
 	}
 
-	const FRigControlElement* ControlElement = ControlRig->FindControl(ControlName);
-	if (!ControlElement)
+	if (bDeepCheck)
 	{
-		return false;
+		const USceneComponent* BoundComponent = GetBoundComponent();
+		if (!BoundComponent)
+		{
+			return false;
+		}
 	}
 	
 	return true;
diff --git a/Engine/Plugins/Animation/ControlRig/Source/ControlRig/Public/Constraints/ControlRigTransformableHandle.h b/Engine/Plugins/Animation/ControlRig/Source/ControlRig/Public/Constraints/ControlRigTransformableHandle.h
--- a/Engine/Plugins/Animation/ControlRig/Source/ControlRig/Public/Constraints/ControlRigTransformableHandle.h
+++ b/Engine/Plugins/Animation/ControlRig/Source/ControlRig/Public/Constraints/ControlRigTransformableHandle.h
@@ -31,7 +31,7 @@ public:
 	virtual void PostLoad() override;
 	
 	/** Sanity check to ensure that ControlRig and ControlName are safe to use. */
-	virtual bool IsValid() const override;
+	virtual bool IsValid(const bool bDeepCheck = true) const override;
 
 	/** Sets the global transform of the control. */
 	virtual void SetGlobalTransform(const FTransform& InGlobal) const override;
diff --git a/Engine/Source/Runtime/Experimental/Animation/Constraints/Private/ConstraintsActor.cpp b/Engine/Source/Runtime/Experimental/Animation/Constraints/Private/ConstraintsActor.cpp
--- a/Engine/Source/Runtime/Experimental/Animation/Constraints/Private/ConstraintsActor.cpp
+++ b/Engine/Source/Runtime/Experimental/Animation/Constraints/Private/ConstraintsActor.cpp
@@ -70,14 +70,14 @@ void AConstraintsActor::RegisterConstraintsTickFunctions() const
 	{
 		if (ConstraintsManager)
 		{
-			// remove invalid pointers
-			ConstraintsManager->Constraints.RemoveAll( [](const TObjectPtr<UTickableConstraint>& InConstraint)
+			// remove invalid constraints (meaning that both handle's components are invalid, unbound controls are valid)
+			static constexpr bool bDeepValidityCheck = false;
+			ConstraintsManager->Constraints.RemoveAll([](const TObjectPtr<UTickableConstraint>& InConstraint)
 			{
-				return !IsValid(InConstraint) || !InConstraint->IsValid();
+				return !IsValid(InConstraint) || !InConstraint->IsValid(bDeepValidityCheck);
 			});
 			
-			// ensure registration, we now do this by adding it to the controller (which registers and sets it with the
-			//subsystem
+			// ensure registration, we now do this by adding it to the controller (which registers and sets it with the subsystem)
 			FConstraintsManagerController& Controller = FConstraintsManagerController::Get(GetWorld());
 			for (const TObjectPtr<UTickableConstraint>& ConstPtr: ConstraintsManager->Constraints)
 			{
diff --git a/Engine/Source/Runtime/Experimental/Animation/Constraints/Private/TransformConstraint.cpp b/Engine/Source/Runtime/Experimental/Animation/Constraints/Private/TransformConstraint.cpp
--- a/Engine/Source/Runtime/Experimental/Animation/Constraints/Private/TransformConstraint.cpp
+++ b/Engine/Source/Runtime/Experimental/Animation/Constraints/Private/TransformConstraint.cpp
@@ -619,10 +619,13 @@ void UTickableTransformConstraint::OnHandleModified(UTransformableHandle* InHand
 	}
 }
 
-bool UTickableTransformConstraint::IsValid() const
+bool UTickableTransformConstraint::IsValid(const bool bDeepCheck) const
 {
-	return (bValid && ::IsValid(ChildTRSHandle) && ChildTRSHandle->IsValid()
-		&& ::IsValid(ParentTRSHandle) && ParentTRSHandle->IsValid());
+	const bool bAreHandlesValid =
+		::IsValid(ChildTRSHandle) && ChildTRSHandle->IsValid(bDeepCheck) &&
+		::IsValid(ParentTRSHandle) && ParentTRSHandle->IsValid(bDeepCheck); 
+	
+	return bDeepCheck ? bAreHandlesValid && bValid : bAreHandlesValid;
 }
 
 bool UTickableTransformConstraint::IsFullyActive() const
diff --git a/Engine/Source/Runtime/Experimental/Animation/Constraints/Private/TransformableHandle.cpp b/Engine/Source/Runtime/Experimental/Animation/Constraints/Private/TransformableHandle.cpp
--- a/Engine/Source/Runtime/Experimental/Animation/Constraints/Private/TransformableHandle.cpp
+++ b/Engine/Source/Runtime/Experimental/Animation/Constraints/Private/TransformableHandle.cpp
@@ -61,7 +61,7 @@ void UTransformableComponentHandle::PostLoad()
 	RegisterDelegates();
 }
 
-bool UTransformableComponentHandle::IsValid() const
+bool UTransformableComponentHandle::IsValid(const bool bDeepCheck) const
 {
 	return Component.IsValid();
 }
diff --git a/Engine/Source/Runtime/Experimental/Animation/Constraints/Public/ConstraintsManager.h b/Engine/Source/Runtime/Experimental/Animation/Constraints/Public/ConstraintsManager.h
--- a/Engine/Source/Runtime/Experimental/Animation/Constraints/Public/ConstraintsManager.h
+++ b/Engine/Source/Runtime/Experimental/Animation/Constraints/Public/ConstraintsManager.h
@@ -84,7 +84,7 @@ public:
 	CONSTRAINTS_API virtual void Evaluate(bool bTickHandlesAlso = false) const;
 
 	/** Whether or not it's valid for example it may not be fully loaded, missing handles or set to not be valid*/
-	virtual bool IsValid() const PURE_VIRTUAL(IsValid, return false;);
+	virtual bool IsValid(const bool bDeepCheck = true) const PURE_VIRTUAL(IsValid, return false;);
 	/** Sets the Active value and enable/disable the tick function. */
 	CONSTRAINTS_API virtual void SetActive(const bool bIsActive);
 
diff --git a/Engine/Source/Runtime/Experimental/Animation/Constraints/Public/TransformConstraint.h b/Engine/Source/Runtime/Experimental/Animation/Constraints/Public/TransformConstraint.h
--- a/Engine/Source/Runtime/Experimental/Animation/Constraints/Public/TransformConstraint.h
+++ b/Engine/Source/Runtime/Experimental/Animation/Constraints/Public/TransformConstraint.h
@@ -47,7 +47,7 @@ public:
 	CONSTRAINTS_API virtual void ResolveBoundObjects(FMovieSceneSequenceID LocalSequenceID, IMovieScenePlayer& Player,UObject* SubObject) override;
 
 	/** Whether or not it's valid for example it may not be fully loaded*/
-	virtual bool IsValid() const override;
+	virtual bool IsValid(const bool bDeepCheck = true) const override;
 
 	/** If Active and the handles and targets are valid, and tick function is registered*/
 	CONSTRAINTS_API virtual bool IsFullyActive() const override;
diff --git a/Engine/Source/Runtime/Experimental/Animation/Constraints/Public/TransformableHandle.h b/Engine/Source/Runtime/Experimental/Animation/Constraints/Public/TransformableHandle.h
--- a/Engine/Source/Runtime/Experimental/Animation/Constraints/Public/TransformableHandle.h
+++ b/Engine/Source/Runtime/Experimental/Animation/Constraints/Public/TransformableHandle.h
@@ -47,8 +47,13 @@ public:
 	
 	CONSTRAINTS_API virtual ~UTransformableHandle();
 	
-	/** Sanity check to ensure the handle is safe to use. */
-	CONSTRAINTS_API virtual bool IsValid() const PURE_VIRTUAL(IsValid, return false;);
+	/** Sanity check to ensure the handle is safe to use.
+	 * @param bDeepCheck to check that the transformable object it wraps is valid AND can be transformed. Default is true.
+	 * Some handles, such as control handles, will only be bound to a skeletal mesh once the level sequence has been opened, but their control rig
+	 * pointer and control name are valid, so they must be fully loaded so that the constraint can be updated later. bDeepCheck = false will be used
+	 * for that purpose.
+	 */
+	CONSTRAINTS_API virtual bool IsValid(const bool bDeepCheck = true) const PURE_VIRTUAL(IsValid, return false;);
 	
 	/** Sets the global transform of the underlying transformable object. */
 	CONSTRAINTS_API virtual void SetGlobalTransform(const FTransform& InGlobal) const PURE_VIRTUAL(SetGlobalTransform, );
@@ -62,7 +67,7 @@ public:
 	/** If true it contains objects bound to an external system, like sequencer so we don't do certain things, like remove constraints when they don't resolve*/
 	CONSTRAINTS_API virtual bool HasBoundObjects() const;
 
-	/** Resolve the bound objects so that any object it references are resovled and correctly set up*/
+	/** Resolve the bound objects so that any object it references are resolved and correctly set up*/
 	CONSTRAINTS_API virtual void ResolveBoundObjects(FMovieSceneSequenceID LocalSequenceID, IMovieScenePlayer& Player, UObject* SubObject = nullptr) PURE_VIRTUAL(ResolveBoundObjects);
 
 	/** Make a duplicate of myself with this outer*/
@@ -139,7 +144,7 @@ public:
 	CONSTRAINTS_API virtual void PostLoad() override;
 	
 	/** Sanity check to ensure that Component. */
-	CONSTRAINTS_API virtual bool IsValid() const override;
+	CONSTRAINTS_API virtual bool IsValid(const bool bDeepCheck = true) const override;
 	
 	/** Sets the global transform of Component. */
 	CONSTRAINTS_API virtual void SetGlobalTransform(const FTransform& InGlobal) const override;
