Early out when parsing headers

#rnx

diff --git a/Engine/Source/Runtime/Experimental/IoStoreOnDemand/Private/CoreHttp/Http.cpp b/Engine/Source/Runtime/Experimental/IoStoreOnDemand/Private/CoreHttp/Http.cpp
--- a/Engine/Source/Runtime/Experimental/IoStoreOnDemand/Private/CoreHttp/Http.cpp
+++ b/Engine/Source/Runtime/Experimental/IoStoreOnDemand/Private/CoreHttp/Http.cpp
@@ -2149,11 +2149,12 @@ static int32 DoRecvMessage(FActivity* Activity, FSocket& Socket)
 		Internal.MessageLength - Internal.Offsets.Headers - 2 // "-2" trims off '\r\n' that signals end of headers
 	);
 
+	int32 Count = 2;
 	bool IsKeepAlive = true;
 	int32 ContentLength = -1;
 	EnumerateHeaders(
 		Headers,
-		[&ContentLength, &IsKeepAlive] (FAnsiStringView Name, FAnsiStringView Value)
+		[&ContentLength, &IsKeepAlive, &Count] (FAnsiStringView Name, FAnsiStringView Value)
 		{
 			// todo; may need smarter value handling; ;/, separated options & key-value pairs (ex. in rfc2068)
 
@@ -2163,16 +2164,16 @@ static int32 DoRecvMessage(FActivity* Activity, FSocket& Socket)
 			if (Name.Equals("Content-Length", ESearchCase::IgnoreCase))
 			{
 				ContentLength = int32(CrudeToInt(Value));
-				return true;
+				Count--;
 			}
 
 			else if (Name.Equals("Connection", ESearchCase::IgnoreCase))
 			{
 				IsKeepAlive = !Value.Equals("close");
-				return true;
+				Count--;
 			}
 
-			return true;
+			return Count > 0;
 		}
 	);
 
