More modern test server and test updates

#rnx

diff --git a/Engine/Source/Runtime/Experimental/IoStoreOnDemand/Private/CoreHttp/Http.cpp b/Engine/Source/Runtime/Experimental/IoStoreOnDemand/Private/CoreHttp/Http.cpp
--- a/Engine/Source/Runtime/Experimental/IoStoreOnDemand/Private/CoreHttp/Http.cpp
+++ b/Engine/Source/Runtime/Experimental/IoStoreOnDemand/Private/CoreHttp/Http.cpp
@@ -3856,7 +3856,7 @@ IOSTOREONDEMAND_API void IasHttpTest(const ANSICHAR* TestHost="localhost")
 		FConnectionPool Pool(Params);
 		for (int32 j = 0; j < i; ++j)
 		{
-			FRequest Request = Loop.Get(BuildUrl("/data"), Pool);
+			FRequest Request = Loop.Get("/data", Pool);
 			Loop.Send(MoveTemp(Request), HashSink);
 		}
 		WaitForLoopIdle();
@@ -3882,13 +3882,21 @@ IOSTOREONDEMAND_API void IasHttpTest(const ANSICHAR* TestHost="localhost")
 			check(IsFailTimeout == bExpectFailTimeout);
 		};
 
+		auto ErrorSink = [] (const FTicketStatus& Status)
+		{
+			check(Status.GetId() == FTicketStatus::EId::Error);
+		};
+
 		FConnectionPool::FParams Params;
 		Params.SetHostFromUrl(BuildUrl("", 9494));
 		FConnectionPool Pool(Params);
 
 		FEventLoop Loop2;
 		Loop2.Send(Loop2.Get("/data?stall", Pool), Sink);
-		Loop2.Send(Loop2.Get("/data", Pool), HashSink);
+
+		// Requests are pipelined. The second one will get went during the stall so
+		// we expect it to fail. The subsequent ones are expected to succeed.
+		Loop2.Send(Loop2.Get("/data", Pool), ErrorSink);
 		Loop2.Send(Loop2.Get("/data", Pool), HashSink);
 		Loop2.Send(Loop2.Get("/data", Pool), HashSink);
 
diff --git a/Engine/Source/Runtime/Experimental/IoStoreOnDemand/Private/CoreHttp/TestServer.py b/Engine/Source/Runtime/Experimental/IoStoreOnDemand/Private/CoreHttp/TestServer.py
--- a/Engine/Source/Runtime/Experimental/IoStoreOnDemand/Private/CoreHttp/TestServer.py
+++ b/Engine/Source/Runtime/Experimental/IoStoreOnDemand/Private/CoreHttp/TestServer.py
@@ -1,158 +1,220 @@
 import io
 import os
 import re
+import time
 import base64
-import random
 import socket
-import asyncio
+import random
 import threading
-from flask import request, Flask, Response
-
-os.chdir("c:/")
-
-#-------------------------------------------------------------------------------
-async def proxy_impl(loop, client):
-    msg = b""
-    while not msg.endswith(b"\r\n\r\n"):
-        msg += await loop.sock_recv(client, 2048)
-
-    mode = "="
-    min_trunc = 0
+import http.server
+import http.client
+
+# {{{1 proxied .................................................................
+
+def intercept(fd):
+    line = fd.readline()
+    headers = http.client.parse_headers(fd)
+    return line, headers
+
+def make_preamble(line, headers):
+    msg = line
+    for k, v in headers.items():
+        msg += (k + ": " + v + "\r\n").encode()
+    msg += b"\r\n"
+    return msg
+
+def proxy_impl(client, httpd):
+    # get request
+    line, headers = intercept(client)
+    if not (line or headers):
+        return False
+    close = (headers.get("Connection", "").lower() == "close")
+    msg = make_preamble(line, headers)
+    httpd.write(msg)
+    httpd.flush()
+
+    # establish behaviour
     disconnect = False
-    stall = False
-    tamper = 0
-
-    line = next(iter(io.BytesIO(msg)))
     if b"?disconnect HTTP" in line:
-        mode = "D"
         disconnect = True
 
+    stall = False
     if b"?stall HTTP" in line:
-        mode = "S"
-        disconnect = True
-        stall = True
-        min_trunc = len(msg) + 16
+        stall = disconnect = True
 
+    tamper = 0
     if m := re.search(b"\?tamper=(\d+) HTTP", line):
-        mode = "T"
         tamper = int(m.group(1))
         tamper = tamper / 100.0
 
-    send_time = 1.0 + (random.random() * 0.5)
+    # get repsonse
+    line, headers = intercept(httpd)
+    close = close or (headers.get("Connection", "").lower() == "close")
+    content_len = int(headers["Content-Length"])
+
+    # get data to retransmit
+    data = make_preamble(line, headers)
+    data += httpd.read(content_len)
+    data_size = len(data)
+
+    if stall:
+        data = data[:-1]
+    elif disconnect:
+        trunc = int(data_size * random.random())
+        data = data[:trunc]
+
+    if tamper > 0:
+        data = bytearray(data)
+        for i in range(data_size):
+            c = data[i] if random.random() > tamper else (int(random.random() * 0x4567) & 0xff)
+            data[i] = c
+
+
+    # retransmit
+    slowly = not (stall or tamper > 0)
+    send_time = (0.75 + (random.random() * 0.75)) if slowly else 0
+    while data:
+        percent = 0.02 + (random.random() * 0.08)
+        send_size = max(int(data_size * percent), 1)
+        piece = data[:send_size]
+        data = data[send_size:]
+
+        client.write(piece)
+        client.flush()
+
+        time.sleep(send_time * percent)
+
+    # we're done
+    if stall:
+        time.sleep(2)
+
+    return not (close or disconnect)
+
+def proxy_loop(client):
+    client_fd = client.makefile("rwb")
     with socket.socket(socket.AF_INET, socket.SOCK_STREAM) as httpd:
         httpd.connect(("127.0.0.1", 9493))
-        await loop.sock_sendall(httpd, msg)
-        data = bytes()
-        while x := await loop.sock_recv(httpd, 128 << 10):
-            data += x
-
-        data_size = len(data)
-
-        if disconnect:
-            trunc = max(int(data_size * random.random()), min_trunc)
-            if trunc >= data_size:
-                trunc = data_size - 1
-            print("trun:", data_size, "->", trunc)
-            data = data[:trunc]
-
-        if tamper:
-            data = bytearray(data)
-            for i in range(len(data)):
-                c = data[i] if random.random() > tamper else (int(random.random() * 0x4567) & 0xff)
-                data[i] = c
-
+        httpd_fd = httpd.makefile("rwb")
         try:
-            while data:
-                percent = 0.02 + (random.random() * 0.08)
-                send_size = max(int(data_size * percent), 1)
-                print("%s snd:%5d tot:%5d slp:%5f %%:%5f" % (mode, send_size, data_size, send_time * percent, percent), end="\r")
-                await loop.sock_sendall(client, data[:send_size])
-                data = data[send_size:]
-                if not tamper:
-                    await asyncio.sleep(send_time * percent)
+            while proxy_impl(client_fd, httpd_fd):
+                pass
         except ConnectionError:
             pass
-        print("")
-
-
-        if stall:
-            print("...stalling", end="")
-            await asyncio.sleep(2)
-            print("")
-
     client.close()
 
-#-------------------------------------------------------------------------------
-async def proxy_loop():
-    loop = asyncio.get_event_loop()
+def proxy():
     with socket.socket(socket.AF_INET, socket.SOCK_STREAM) as sock:
         sock.bind(("", 9494))
         sock.listen()
         sock.setblocking(True)
         while True:
-            client, address = await loop.sock_accept(sock)
-            loop.create_task(proxy_impl(loop, client))
-
-#-------------------------------------------------------------------------------
-def proxy():
-    asyncio.run(proxy_loop())
-
-thread = threading.Thread(target=proxy)
-thread.start()
-
+            client, address = sock.accept()
+            try:
+                threading.Thread(target=proxy_loop, args=(client,), daemon=True).start()
+            except (ConnectionResetError, ConnectionAbortedError):
+                pass
 
+# {{{1 httpd ...................................................................
 
-#-------------------------------------------------------------------------------
-app = Flask(__name__)
 payload_data = random.randbytes(2 << 20)
 
-#-------------------------------------------------------------------------------
[REDACTED_EMAIL]("/test")
-def test():
-    return "<html><head>test-server<title></title></head><body>hello world</body></html>"
-
-#-------------------------------------------------------------------------------
[REDACTED_EMAIL]("/seed/<int:value>")
-def seed(value):
+def http_seed(handler, value):
+    handler.send_response(200)
+    handler.send_header("Content-Length", 0)
+    handler.end_headers()
     random.seed(value)
-    return { "_" : "ok" }
 
-#-------------------------------------------------------------------------------
[REDACTED_EMAIL]("/data", methods=("GET", "POST"))
[REDACTED_EMAIL]("/data/<int:payload_size>", methods=("GET", "POST"))
[REDACTED_EMAIL]("/data/[a-zA-Z]+", methods=("GET", "POST"))
-def get(payload_size=-1):
-    if request.method == "POST":
-        abort(400)
+def http_data(handler, payload_size):
+    handler.send_response(200)
 
+    mega_size = int(random.random() * (4 << 10))
+    for i in range(1024):
+        key = "X-MegaHeader-%04d" % i
+        value = payload_data[:int(random.random() * 64)]
+        value = base64.b64encode(value)
+        handler.send_header(key, value.decode())
+        mega_size -= len(key) + len(value) + 4
+        if mega_size <= 0:
+            break
+
+    payload_size = int(payload_size)
     if payload_size <= 0:
         payload_size = int(random.random() * (8 << 10)) + 16
-
     payload_size = min(payload_size, len(payload_data))
+
     payload = payload_data[-payload_size:]
 
     payload_hash = 0x493
     for c in payload:
         payload_hash = ((payload_hash + c) * 0x493) & 0xffffffff
+    handler.send_header("X-TestServer-Hash", payload_hash)
 
-    mega_headers = {}
-    mega_size = int(random.random() * (4 << 10))
-    for i in range(1024):
-        key = "X-MegaHeader-%04d" % i
-        value = payload_data[:int(random.random() * 64)]
-        value = base64.b64encode(value)
-        mega_headers[key] = value.decode()
+    handler.send_header("Content-Length", len(payload))
+    handler.send_header("Content-Type", "application/octet-stream")
+    handler.end_headers()
+    handler.wfile.write(payload)
 
-        mega_size -= len(key) + len(value) + 4
-        if mega_size <= 0:
-            break
 
-    return (payload, {
-        "X-TestServer-Hash" : payload_hash,
-        **mega_headers,
-        "Content-Type" : "application/octet-stream",
-    })
 
-app.run(host="0.0.0.0", port=9493)
+class Handler(http.server.BaseHTTPRequestHandler):
+    def _preroll(self):
+        self.protocol_version = "HTTP/1.1"
+
+        conn_value = self.headers.get("Connection", "").lower()
+        self.close_connection = (conn_value == "close")
+
+        parts = [x for x in self.path.split("/") if x]
+        if len(parts) >= 1:
+            return parts
+
+        self.send_error(404)
+
+    def do_HEAD(self):
+        parts = self._preroll()
+        if not parts:
+            return
+
+        self.send_response(200)
+        self.send_header("Content-Length", 0)
+        self.end_headers()
+
+    def end_headers(self):
+        if self.close_connection:
+            self.send_header("Connection", "close")
+        super().end_headers()
+
+    def do_GET(self):
+        parts = self._preroll()
+        if not parts:
+            return self.send_error(404)
+
+        if query := parts[-1].split("?"):
+            parts[-1] = query[0]
+
+        if parts[0] == "data":
+            size = -1
+            if len(parts) > 1:
+                size = int(parts[1])
+            return http_data(self, size)
+
+        if parts[0] == "seed":
+            size = int(parts[1])
+            return http_seed(self, size)
+
+        return self.send_error(404, f"not found '{self.path}'")
+
+
+
+# {{{1 main ....................................................................
+
+def main():
+    os.chdir("c:/")
+
+    proxy_thread = threading.Thread(target=proxy, daemon=True)
+    proxy_thread.start()
+
+    server = http.server.ThreadingHTTPServer(("", 9493), Handler)
+    server.serve_forever()
 
-# vim: et
+if __name__ == "__main__":
+    main()
