IAS - check project persistent download directory is available before creating cache

diff --git a/Engine/Source/Runtime/Experimental/IoStoreOnDemand/Private/IoStoreOnDemand.cpp b/Engine/Source/Runtime/Experimental/IoStoreOnDemand/Private/IoStoreOnDemand.cpp
--- a/Engine/Source/Runtime/Experimental/IoStoreOnDemand/Private/IoStoreOnDemand.cpp
+++ b/Engine/Source/Runtime/Experimental/IoStoreOnDemand/Private/IoStoreOnDemand.cpp
@@ -2225,8 +2225,11 @@ void FIoStoreOnDemandModule::InitializeInternal()
 	CacheConfig.DropCache = DeferredAbandonCache.Get(CacheConfig.DropCache);
 	if (CacheConfig.DiskQuota > 0)
 	{
-		FString CacheDir = FPaths::ProjectPersistentDownloadDir();
-		Cache = MakeIasCache(*CacheDir, CacheConfig);
+		if (FPaths::HasProjectPersistentDownloadDir())
+		{
+			FString CacheDir = FPaths::ProjectPersistentDownloadDir();
+			Cache = MakeIasCache(*CacheDir, CacheConfig);
+		}
 	}
 	if (!Cache.IsValid())
 	{
