Chaos physics bug fix: When using a rotated query object for overlap queries, the rotation was not accounted for when calculating transforms
#[fyi] Benn.Gallagher, Chris.Caulfield

diff --git a/Engine/Source/Programs/HeadlessChaos/Private/HeadlessChaosTestOverlap.cpp b/Engine/Source/Programs/HeadlessChaos/Private/HeadlessChaosTestOverlap.cpp
--- a/Engine/Source/Programs/HeadlessChaos/Private/HeadlessChaosTestOverlap.cpp
+++ b/Engine/Source/Programs/HeadlessChaos/Private/HeadlessChaosTestOverlap.cpp
@@ -1065,6 +1065,52 @@ namespace ChaosTest
 					EXPECT_EQ(bResult, bResultMTD);
 				}
 			}
+
+			{
+				// Regression test: Box with centre not at the local origin
+								
+				// Rotated UnScaled box
+				const TBox<FReal, 3> BoxUnscaled = TBox<FReal, 3>({ -5.0f, -6.0f, 0.0f }, { 5.0f, 6.0f, 10.0f });
+				{
+					FRigidTransform3 QueryTM(FVec3(0, 0.0, 11.0), FQuat{ 1, 0, 0, 0 }); // Pi rotation around x axis
+					bool bResult = TriangleMesh->OverlapGeom(BoxUnscaled, QueryTM, 0.0, nullptr);
+					EXPECT_EQ(bResult, true);
+					bool bResultMTD = TriangleMesh->OverlapGeom(BoxUnscaled, QueryTM, 0.0, &MTDInfo);
+					EXPECT_EQ(bResult, bResultMTD);
+				}
+
+				// Rotated Scaled box
+				FBoxPtr BoxSafe(new TBox<FReal, 3>(FVec3(-10.0, -10.0, 0.0), FVec3(10.0, 10.0, 10.0)));
+				FVec3 TriMeshScale = { 5.0f, 6.0f, 10.0f };
+				FVec3 BoxScale = { 5.0f, 6.0f, 10.0f };
+				TImplicitObjectScaled<TBox<FReal, 3>> ScaledBox = TImplicitObjectScaled<TBox<FReal, 3>>(BoxSafe, BoxScale);
+				{
+					FRigidTransform3 QueryTM(FVec3(0, 0.0, 110.0), FQuat{1, 0, 0, 0}); // Pi rotation around x axis
+					bool bResult = TriangleMesh->OverlapGeom(ScaledBox, QueryTM, 0.0, nullptr, TriMeshScale);
+					EXPECT_EQ(bResult, true);
+					bool bResultMTD = TriangleMesh->OverlapGeom(ScaledBox, QueryTM, 0.0, &MTDInfo, TriMeshScale);
+					EXPECT_EQ(bResult, bResultMTD);					
+				}
+
+				// Unrotated unscaled box
+				{
+					FRigidTransform3 QueryTM(FVec3(0, 0.0, 11.0), FQuat::Identity);
+					bool bResult = TriangleMesh->OverlapGeom(BoxUnscaled, QueryTM, 0.0, nullptr);
+					EXPECT_EQ(bResult, false);
+					bool bResultMTD = TriangleMesh->OverlapGeom(BoxUnscaled, QueryTM, 0.0, &MTDInfo);
+					EXPECT_EQ(bResult, bResultMTD);
+				}
+
+				// Unrotated scaled box
+				{
+					FRigidTransform3 QueryTM(FVec3(0, 0.0, 110.0), FQuat::Identity);
+					bool bResult = TriangleMesh->OverlapGeom(ScaledBox, QueryTM, 0.0, nullptr, TriMeshScale);
+					EXPECT_EQ(bResult, false);
+					bool bResultMTD = TriangleMesh->OverlapGeom(ScaledBox, QueryTM, 0.0, &MTDInfo, TriMeshScale);
+					EXPECT_EQ(bResult, bResultMTD);
+				}
+			}
+
 			{
 				// Non uniform test with box not being a cube
 				FBoxPtr BigBoxSafe( new TBox<FReal, 3>(FVec3(-1.0, -5.0, -1.0), FVec3(1.0, 5.0, 1.0)));
diff --git a/Engine/Source/Runtime/Experimental/Chaos/Private/Chaos/Collision/TriangleOverlap.cpp b/Engine/Source/Runtime/Experimental/Chaos/Private/Chaos/Collision/TriangleOverlap.cpp
--- a/Engine/Source/Runtime/Experimental/Chaos/Private/Chaos/Collision/TriangleOverlap.cpp
+++ b/Engine/Source/Runtime/Experimental/Chaos/Private/Chaos/Collision/TriangleOverlap.cpp
@@ -311,7 +311,7 @@ namespace Chaos
 	FBoxSimd::FBoxSimd(const FRigidTransform3& WorldScaleQueryTM, const TBox<FReal, 3>& InQueryGeom)
 	{
 		FRigidTransform3 Transform(WorldScaleQueryTM);
-		Transform.SetTranslation(Transform.GetTranslation() + InQueryGeom.GetCenter());
+		Transform.SetTranslation(Transform.TransformPositionNoScale(InQueryGeom.GetCenter()));
 		FVec3f HalfExtentsf = InQueryGeom.Extents() * 0.5;
 		Initialize(Transform, HalfExtentsf);
 	}
@@ -319,7 +319,7 @@ namespace Chaos
 	FBoxSimd::FBoxSimd(const FRigidTransform3& QueryTM, const TImplicitObjectScaled< TBox<FReal, 3> >& QueryGeom)
 	{
 		FRigidTransform3 Transform(QueryTM);
-		Transform.TransformPosition(QueryGeom.GetUnscaledObject()->GetCenter() * QueryGeom.GetScale());
+		Transform.SetTranslation(Transform.TransformPositionNoScale(QueryGeom.GetUnscaledObject()->GetCenter() * QueryGeom.GetScale()));
 		FVec3f HalfExtentsf = QueryGeom.GetUnscaledObject()->Extents() * 0.5 * (QueryGeom.GetScale() + UE_SMALL_NUMBER);
 		Initialize(Transform, HalfExtentsf);
 	}
@@ -533,27 +533,21 @@ namespace Chaos
 		HalfExtents = VectorLoadFloat3(&HalfExtentsf.X);
 	}
 
-	FAABBSimd::FAABBSimd(const FRigidTransform3& WorldScaleQueryTM, const TBox<FReal, 3>& InQueryGeom)
+	FAABBSimd::FAABBSimd(const FVec3& Translation, const TBox<FReal, 3>& InQueryGeom)
 	{
-		TRigidTransform<FReal, 3> Transform(WorldScaleQueryTM);
-		Transform.SetTranslation(Transform.GetTranslation() + InQueryGeom.GetCenter());
 		FVec3f HalfExtentsf = InQueryGeom.Extents() * 0.5;
-		Initialize(Transform, HalfExtentsf);
+		Initialize(Translation + InQueryGeom.GetCenter(), HalfExtentsf);
 	}
 
-	FAABBSimd::FAABBSimd(const FRigidTransform3& QueryTM, const TImplicitObjectScaled< TBox<FReal, 3> >& QueryGeom)
+	FAABBSimd::FAABBSimd(const FVec3& Translation, const TImplicitObjectScaled< TBox<FReal, 3> >& QueryGeom)
 	{
-		TRigidTransform<FReal, 3> Transform(QueryTM);
-		Transform.TransformPosition(QueryGeom.GetUnscaledObject()->GetCenter() * QueryGeom.GetScale());
 		FVec3f HalfExtentsf = QueryGeom.GetUnscaledObject()->Extents() * 0.5 * (QueryGeom.GetScale() + UE_SMALL_NUMBER);
-		Initialize(Transform, HalfExtentsf);
+		Initialize(Translation + QueryGeom.GetUnscaledObject()->GetCenter() * QueryGeom.GetScale(), HalfExtentsf);
 	}
 
-	FORCEINLINE void FAABBSimd::Initialize(const FRigidTransform3& Transform, const FVec3f& HalfExtentsf)
+	FORCEINLINE void FAABBSimd::Initialize(const FVec3& Translation, const FVec3f& HalfExtentsf)
 	{
-		FVec3 Translation = Transform.GetTranslation();
 		Position = MakeVectorRegisterFloatFromDouble(VectorLoadDouble3(&Translation.X));
-
 		HalfExtents = VectorLoadFloat3(&HalfExtentsf.X);
 	}
 
diff --git a/Engine/Source/Runtime/Experimental/Chaos/Private/Chaos/Collision/TriangleOverlap.h b/Engine/Source/Runtime/Experimental/Chaos/Private/Chaos/Collision/TriangleOverlap.h
--- a/Engine/Source/Runtime/Experimental/Chaos/Private/Chaos/Collision/TriangleOverlap.h
+++ b/Engine/Source/Runtime/Experimental/Chaos/Private/Chaos/Collision/TriangleOverlap.h
@@ -49,12 +49,12 @@ namespace Chaos
 		FAABBSimd();
 		FAABBSimd(const TBox<FReal, 3>& InQueryGeom);
 		FAABBSimd(const TImplicitObjectScaled< TBox<FReal, 3>, true >& QueryGeom);
-		FAABBSimd(const FRigidTransform3& WorldScaleQueryTM, const TBox<FReal, 3>& InQueryGeom);
-		FAABBSimd(const FRigidTransform3& QueryTM, const TImplicitObjectScaled< TBox<FReal, 3>, true >& QueryGeom);
+		FAABBSimd(const FVec3& Translation, const TBox<FReal, 3>& InQueryGeom);
+		FAABBSimd(const FVec3& Translation, const TImplicitObjectScaled< TBox<FReal, 3>, true >& QueryGeom);
 		bool OverlapTriangle(const VectorRegister4Float& A, const VectorRegister4Float& B, const VectorRegister4Float& C) const;
 
 	private:
-		void Initialize(const FRigidTransform3& Transform, const FVec3f& HalfExtentsf);
+		void Initialize(const FVec3& Translation, const FVec3f& HalfExtentsf);
 		bool ComputeEdgeOverlap(const VectorRegister4Float& TriangleEdge, const VectorRegister4Float& TriangleVertex, const VectorRegister4Float& Centroid) const;
 
 		VectorRegister4Float Position;
diff --git a/Engine/Source/Runtime/Experimental/Chaos/Private/Chaos/TriangleMeshImplicitObject.cpp b/Engine/Source/Runtime/Experimental/Chaos/Private/Chaos/TriangleMeshImplicitObject.cpp
--- a/Engine/Source/Runtime/Experimental/Chaos/Private/Chaos/TriangleMeshImplicitObject.cpp
+++ b/Engine/Source/Runtime/Experimental/Chaos/Private/Chaos/TriangleMeshImplicitObject.cpp
@@ -651,7 +651,7 @@ struct FTriangleMeshOverlapVisitorNoMTD<TBox<FReal, 3>>
 	{
 		if (bIsAxisAligned)
 		{
-			Box.AABBSimd = FAABBSimd(WorldScaleQueryTM, InQueryGeom);
+			Box.AABBSimd = FAABBSimd(WorldScaleQueryTM.GetTranslation(), InQueryGeom);
 		}
 		else
 		{
@@ -716,7 +716,7 @@ struct FTriangleMeshOverlapVisitorNoMTD<TImplicitObjectScaled<TBox<FReal, 3>>>
 	{
 		if (bIsAxisAligned)
 		{
-			Box.AABBSimd = FAABBSimd(WorldScaleQueryTM, InQueryGeom);
+			Box.AABBSimd = FAABBSimd(WorldScaleQueryTM.GetTranslation(), InQueryGeom);
 		}
 		else
 		{
@@ -1135,7 +1135,7 @@ bool FTriangleMeshImplicitObject::OverlapGeomImp(const QueryGeomType& QueryGeom,
 			FAABB3 GeometryAABB = WorldScaleQueryGeom.BoundingBox();
 			GeometryAABB.ThickenSymmetrically(FVec3(Thickness));
 			// The geometry is not necessarily in the center of its translation. (cf: a capsule can be create with two points away from its own transform)
-			TriMeshToGeomNoScale.SetTranslation(TriMeshToGeomNoScale.GetTranslation() + GeometryAABB.Center());
+			TriMeshToGeomNoScale.SetTranslation(TriMeshToGeomNoScale.TransformPositionNoScale(GeometryAABB.Center()));
 			Private::FOBBVectorized QueryObb(TriMeshToGeomNoScale, (GeometryAABB.Max() - GeometryAABB.Min()) * 0.5, InvTriMeshScale);
 			TRigidTransform<FReal, 3> WorldScaleQueryTM;
 			ScaleTransformHelper(TriMeshScale, QueryTM, WorldScaleQueryTM);
