Fix accessing deleted proxy when dispatching events

diff --git a/Engine/Source/Runtime/Experimental/Chaos/Private/EventDefaults.cpp b/Engine/Source/Runtime/Experimental/Chaos/Private/EventDefaults.cpp
--- a/Engine/Source/Runtime/Experimental/Chaos/Private/EventDefaults.cpp
+++ b/Engine/Source/Runtime/Experimental/Chaos/Private/EventDefaults.cpp
@@ -339,10 +339,11 @@ namespace Chaos
 				for (int32 Idx = 0; Idx < AllClusterBreakings.Num(); ++Idx)
 				{					
 					const FBreakingData& ClusterBreaking = AllClusterBreakings[Idx];
-					if (!SolverBreakingEventFilter->Enabled() || SolverBreakingEventFilter->Pass(ClusterBreaking))
+					IPhysicsProxyBase* Proxy = AllClusterBreakings[Idx].Proxy;
+					if (!Proxy->GetMarkedDeleted() && (!SolverBreakingEventFilter->Enabled() || SolverBreakingEventFilter->Pass(ClusterBreaking)))
 					{
 						const int32 NewIndex = FilteredBreakingDataArray.Emplace(ClusterBreaking);
-						FilteredBreakingIndicesByPhysicsProxy.FindOrAdd(ClusterBreaking.Proxy).Add(FEventManager::EncodeCollisionIndex(NewIndex, false));
+						FilteredBreakingIndicesByPhysicsProxy.FindOrAdd(Proxy).Add(FEventManager::EncodeCollisionIndex(NewIndex, false));
 						BreakingEventData.BreakingData.bHasGlobalEvent |= (ClusterBreaking.EmitterFlag & EventEmitterFlag::GlobalDispatcher) != 0;
 					}
 				}
@@ -518,16 +519,18 @@ namespace Chaos
 				
 				for (int32 Idx = 0; Idx < AllRemovalsArray.Num(); ++Idx)
 				{
-					FRemovalData RemovalData;
-					RemovalData.Location = AllRemovalsArray[Idx].Location;
-					RemovalData.Mass = AllRemovalsArray[Idx].Mass;
-					RemovalData.Proxy = AllRemovalsArray[Idx].Proxy;
-					RemovalData.BoundingBox = AllRemovalsArray[Idx].BoundingBox;
-
-					int32 NewIdx = AllRemovalDataArray.Add(FRemovalData());
-					FRemovalData& RemovalDataArrayItem = AllRemovalDataArray[NewIdx];
-					RemovalDataArrayItem = RemovalData;
-					AllRemovalIndicesByPhysicsProxy.FindOrAdd(RemovalData.Proxy).Add(FEventManager::EncodeCollisionIndex(NewIdx, false));
+					IPhysicsProxyBase* Proxy = AllRemovalsArray[Idx].Proxy;
+					if (!Proxy->GetMarkedDeleted())
+					{
+						FRemovalData RemovalData;
+						RemovalData.Location = AllRemovalsArray[Idx].Location;
+						RemovalData.Mass = AllRemovalsArray[Idx].Mass;
+						RemovalData.Proxy = Proxy;
+						RemovalData.BoundingBox = AllRemovalsArray[Idx].BoundingBox;
+
+						const int32 NewIdx = AllRemovalDataArray.Add(RemovalData);
+						AllRemovalIndicesByPhysicsProxy.FindOrAdd(Proxy).Add(FEventManager::EncodeCollisionIndex(NewIdx, false));
+					}
 				}
 			});
 	}
diff --git a/Engine/Source/Runtime/Experimental/Chaos/Private/PBDRigidsSolver.cpp b/Engine/Source/Runtime/Experimental/Chaos/Private/PBDRigidsSolver.cpp
--- a/Engine/Source/Runtime/Experimental/Chaos/Private/PBDRigidsSolver.cpp
+++ b/Engine/Source/Runtime/Experimental/Chaos/Private/PBDRigidsSolver.cpp
@@ -794,22 +794,6 @@ namespace Chaos
 
 		EnqueueCommandImmediate([InProxy, this]()
 		{
-			// Generally need to remove stale events for particles that no longer exist
-			GetEventManager()->template ClearEvents<FCollisionEventData>(EEventType::Collision, [InProxy]
-			(FCollisionEventData& EventDataInOut)
-				{
-					const FCollisionDataArray& CollisionData = EventDataInOut.CollisionData.AllCollisionsArray;
-					if (CollisionData.Num() > 0)
-					{
-						check(InProxy);
-						const TArray<int32>* const CollisionIndices = EventDataInOut.PhysicsProxyToCollisionIndices.PhysicsProxyToIndicesMap.Find(InProxy);
-						if (CollisionIndices)
-						{
-							EventDataInOut.PhysicsProxyToCollisionIndices.PhysicsProxyToIndicesMap.Remove(InProxy);
-						}
-					}
-				});
-
 			GeometryCollectionPhysicsProxies_Internal.RemoveSingle(InProxy);
 			InProxy->SyncBeforeDestroy();
 			PendingDestroyGeometryCollectionPhysicsProxy.Add(InProxy);
