Small coding standard fix that failed to merge down after being cherry picked to parent stream and altered there.

diff --git a/Engine/Source/Runtime/Experimental/GeometryCollectionEngine/Private/GeometryCollection/GeometryCollectionComponent.cpp b/Engine/Source/Runtime/Experimental/GeometryCollectionEngine/Private/GeometryCollection/GeometryCollectionComponent.cpp
--- a/Engine/Source/Runtime/Experimental/GeometryCollectionEngine/Private/GeometryCollection/GeometryCollectionComponent.cpp
+++ b/Engine/Source/Runtime/Experimental/GeometryCollectionEngine/Private/GeometryCollection/GeometryCollectionComponent.cpp
@@ -4363,7 +4363,8 @@ void UGeometryCollectionComponent::UpdateRemovalIfNeeded()
 
 void UGeometryCollectionComponent::RequestUpdateRepData()
 {
-	if (FPhysScene* PhysScene = GetInnerChaosScene(); PhysScene && PhysicsProxy)
+	FPhysScene* PhysScene = GetInnerChaosScene();
+	if (PhysScene && PhysicsProxy) 
 	{
 		if (Chaos::FPBDRigidsSolver* Solver = PhysicsProxy->GetSolver<Chaos::FPBDRigidsSolver>())
 		{
