Iris - Map unmapped dynamic objects.

diff --git a/Engine/Source/Runtime/Experimental/Iris/Core/Private/Iris/ReplicationSystem/ReplicationReader.cpp b/Engine/Source/Runtime/Experimental/Iris/Core/Private/Iris/ReplicationSystem/ReplicationReader.cpp
--- a/Engine/Source/Runtime/Experimental/Iris/Core/Private/Iris/ReplicationSystem/ReplicationReader.cpp
+++ b/Engine/Source/Runtime/Experimental/Iris/Core/Private/Iris/ReplicationSystem/ReplicationReader.cpp
@@ -45,15 +45,13 @@ static bool bExecuteReliableRPCsBeforeApplyState = true;
 static FAutoConsoleVariableRef CVarExecuteReliableRPCsBeforeApplyState(
 		TEXT("net.Iris.ExecuteReliableRPCsBeforeApplyState"),
 		bExecuteReliableRPCsBeforeApplyState,
-		TEXT("If true and Iris runs in backwards compatibility mode then reliable RPCs will be executed before we apply state data on the target object unless we first need to spawn the object."
-		));
+		TEXT("If true and Iris runs in backwards compatibility mode then reliable RPCs will be executed before we apply state data on the target object unless we first need to spawn the object."));
 
 static bool bDeferEndReplication = true;
 static FAutoConsoleVariableRef CVarDeferEndReplication(
 	TEXT("net.Iris.DeferEndReplication"),
 	bDeferEndReplication,
-	TEXT("bDeferEndReplication if true calls to EndReplication will be defered until after we have applied statedata. Default is true."
-	));
+	TEXT("bDeferEndReplication if true calls to EndReplication will be defered until after we have applied statedata. Default is true."));
 
 static bool bDispatchUnresolvedPreviouslyReceivedChanges = false;
 static FAutoConsoleVariableRef CvarDispatchUnresolvedPreviouslyReceivedChanges(
@@ -61,6 +59,18 @@ static FAutoConsoleVariableRef CvarDispatchUnresolvedPreviouslyReceivedChanges(
 	bDispatchUnresolvedPreviouslyReceivedChanges,
 	TEXT("Whether to include previously received changes with unresolved object references to data received this frame when applying state data. This can call rep notify functions to be called despite being unchanged. Default is false."));
 
+static bool bRemapDynamicObjects = true;
+static FAutoConsoleVariableRef CvarRemapDynamicObjects(
+		TEXT("net.Iris.RemapDynamicObjects"),
+		bRemapDynamicObjects,
+		TEXT("Allow remapping of dynamic objects on the receiving end. This allows properties previously pointing to a particular object to be updated if the object is re-created. Default is true."));
+
+static bool bResolvedObjectsDispatchDebugging = false;
+static FAutoConsoleVariableRef CvarResolvedObjectsDispatchDebugging(
+	TEXT("net.Iris.ResolvedObjectsDispatchDebugging"),
+	bResolvedObjectsDispatchDebugging,
+	TEXT("Debug logging of resolved object state dispatching. Default is false."));
+
 static const FName NetError_FailedToFindAttachmentQueue("Failed to find attachment queue");
 
 class FResolveAndCollectUnresolvedAndResolvedReferenceCollector
@@ -1091,7 +1101,7 @@ void FReplicationReader::UpdateObjectReferenceTracking(FReplicatedObjectInfo* Re
 	}
 
 	// Update tracking for resolved dynamic references
-#if 0
+	if (bRemapDynamicObjects)
 	{
 		// Try to avoid dynamic allocations during the update of the ResolvedDynamicObjectReferences.
 		ReplicationInfo->ResolvedDynamicObjectReferences.Reserve(ReplicationInfo->ResolvedDynamicObjectReferences.Num() + NewMappedDynamicReferences.Num());
@@ -1138,7 +1148,7 @@ void FReplicationReader::UpdateObjectReferenceTracking(FReplicatedObjectInfo* Re
 			{
 				// Remove from tracking
 				ResolvedDynamicHandleToDependents.RemoveSingle(Handle, OwnerInternalIndex);
-				UE_LOG(LogIris, Verbose, TEXT("FReplicationReader::UpdateObjectReferenceTracking Removing resolved dynamic reference %s for %s"), ToCStr(Handle.ToString()), ToCStr(NetRefHandleManager->GetNetHandleFromInternalIndex(OwnerInternalIndex).ToString()));
+				//UE_LOG(LogIris, Verbose, TEXT("FReplicationReader::UpdateObjectReferenceTracking Removing resolved dynamic reference %s for %s"), ToCStr(Handle.ToString()), ToCStr(NetRefHandleManager->GetNetRefHandleFromInternalIndex(OwnerInternalIndex).ToString()));
 			}
 		}
 
@@ -1149,11 +1159,10 @@ void FReplicationReader::UpdateObjectReferenceTracking(FReplicatedObjectInfo* Re
 			{
 				// Add to tracking
 				ResolvedDynamicHandleToDependents.Add(Handle, OwnerInternalIndex);
-				UE_LOG(LogIris, Verbose, TEXT("FReplicationReader::UpdateObjectReferenceTracking Adding resolved dynamic reference %s for %s"), ToCStr(Handle.ToString()), ToCStr(NetRefHandleManager->GetNetHandleFromInternalIndex(OwnerInternalIndex).ToString()));
+				//UE_LOG(LogIris, Verbose, TEXT("FReplicationReader::UpdateObjectReferenceTracking Adding resolved dynamic reference %s for %s"), ToCStr(Handle.ToString()), ToCStr(NetRefHandleManager->GetNetRefHandleFromInternalIndex(OwnerInternalIndex).ToString()));
 			}
 		}
 	}
-#endif
 }
 
 void FReplicationReader::RemoveUnresolvedObjectReferenceInReplicationInfo(FReplicatedObjectInfo* ReplicationInfo, FNetRefHandle Handle)
@@ -1300,6 +1309,7 @@ void FReplicationReader::ResolveAndDispatchUnresolvedReferencesForObject(FNetSer
 	if (bObjectHasReferences)
 	{
 		const FNetRefHandleManager::FReplicatedObjectData& ObjectData = NetRefHandleManager->GetReplicatedObjectDataNoCheck(ReplicationInfo->InternalIndex);
+		UE_LOG(LogIris, VeryVerbose, TEXT("ResolveAndDispatchUnresolvedReferencesForObject %s RefHandle %s"), ObjectData.Protocol->DebugName->Name, ToCStr(ObjectData.RefHandle.ToString()));
 		const uint32 ChangeMaskBitCount = ReplicationInfo->ChangeMaskBitCount;
 		
 		// Try to resolve references and collect unresolved references
@@ -1361,6 +1371,30 @@ void FReplicationReader::ResolveAndDispatchUnresolvedReferencesForObject(FNetSer
 				Params.SrcObjectStateBuffer = ObjectData.ReceiveStateBuffer;
 				Params.bHasUnresolvedInitReferences = ReplicationInfo->bHasUnresolvedInitialReferences;
 
+				if (bResolvedObjectsDispatchDebugging && UE_LOG_ACTIVE(LogIris, VeryVerbose))
+				{
+					uint32 CurrentChangeMaskBitOffset = 0;
+					for (const FReplicationStateDescriptor* StateDescriptor : MakeArrayView(ObjectData.Protocol->ReplicationStateDescriptors, ObjectData.Protocol->ReplicationStateCount))
+					{
+						if (ResolvedChangeMask.IsAnyBitSet(CurrentChangeMaskBitOffset, StateDescriptor->ChangeMaskBitCount))
+						{
+							for (uint32 MemberIt = 0, MemberEndIt = StateDescriptor->MemberCount; MemberIt != MemberEndIt; ++MemberIt)
+							{
+								const FReplicationStateMemberChangeMaskDescriptor* MemberChangeMaskDesc = StateDescriptor->MemberChangeMaskDescriptors + MemberIt;
+								if (ResolvedChangeMask.IsAnyBitSet(CurrentChangeMaskBitOffset + MemberChangeMaskDesc->BitOffset, MemberChangeMaskDesc->BitCount))
+								{
+									if (const FProperty* MemberProperty = StateDescriptor->MemberProperties[MemberIt])
+									{
+										UE_LOG(LogIris, VeryVerbose, TEXT("ResolvedChangeMask State %s Property %s"), ToCStr(StateDescriptor->DebugName->Name), ToCStr(MemberProperty->GetName()));
+									}
+								}
+							}
+						}
+
+						CurrentChangeMaskBitOffset += StateDescriptor->ChangeMaskBitCount;
+					}
+				}
+
 				FReplicationInstanceOperations::DequantizeAndApply(Context, Params);
 			}
 		}
