Rebuild cluster geometry if the GC transform is scaled at runtime

diff --git a/Engine/Source/Runtime/Experimental/Chaos/Private/Chaos/PBDRigidClusteringAlgo.cpp b/Engine/Source/Runtime/Experimental/Chaos/Private/Chaos/PBDRigidClusteringAlgo.cpp
--- a/Engine/Source/Runtime/Experimental/Chaos/Private/Chaos/PBDRigidClusteringAlgo.cpp
+++ b/Engine/Source/Runtime/Experimental/Chaos/Private/Chaos/PBDRigidClusteringAlgo.cpp
@@ -207,6 +207,59 @@ namespace Chaos
 	DECLARE_CYCLE_STAT(TEXT("TPBDRigidClustering<>::UpdateGeometry_CopyPoints"), STAT_UpdateGeometry_CopyPoints, STATGROUP_Chaos);
 	DECLARE_CYCLE_STAT(TEXT("TPBDRigidClustering<>::UpdateGeometry_PointsBVH"), STAT_UpdateGeometry_PointsBVH, STATGROUP_Chaos);
 
+	void BuildScaledGeometry(Chaos::FPBDRigidClusteredParticleHandle* ParticleHandle, const Chaos::FImplicitObjectPtr& ProxyGeometry, const FVector& WorldScale)
+	{
+		if(ParticleHandle && ProxyGeometry.IsValid())
+		{ 
+			auto DeepCopyImplicit = [&WorldScale](const Chaos::FImplicitObjectPtr& ImplicitToCopy) -> Chaos::FImplicitObjectPtr
+			{
+				if (WorldScale.Equals(FVector::OneVector))
+				{
+					return ImplicitToCopy->DeepCopyGeometry();
+				}
+				else
+				{
+					return ImplicitToCopy->DeepCopyGeometryWithScale(WorldScale);
+				}
+			};
+
+			Chaos::EImplicitObjectType GeometryType = ProxyGeometry->GetType();
+			// Don't copy if it is not a level set and scale is one
+			if (GeometryType != Chaos::ImplicitObjectType::LevelSet && WorldScale.Equals(FVector::OneVector))
+			{
+				ParticleHandle->SetGeometry(ProxyGeometry);
+			}
+			else
+			{
+				ParticleHandle->SetGeometry(DeepCopyImplicit(ProxyGeometry));
+			}
+		}
+	}
+
+	void UpdateCollisionFlags(Chaos::FPBDRigidClusteredParticleHandle* ParticleHandle, const bool bUseParticleImplicit)
+	{
+		if (FImplicitObjectPtr ImplicitGeometry = ParticleHandle->GetGeometry())
+		{
+			// if children are ignore analytic and this is a dynamic geom, mark it too. todo(ocohen): clean this up
+			if (bUseParticleImplicit)
+			{
+				ImplicitGeometry->SetDoCollide(false);
+			}
+
+			ParticleHandle->SetHasBounds(true);
+			ParticleHandle->SetLocalBounds(ImplicitGeometry->BoundingBox());
+
+			if (const FImplicitObjectUnion* ImplicitUnion = ImplicitGeometry->GetObject<FImplicitObjectUnion>())
+			{
+				const_cast<FImplicitObjectUnion*>(ImplicitUnion)->SetAllowBVH(true);
+			}
+			else if (const FImplicitObjectUnion* ImplicitUnionClustered = ImplicitGeometry->GetObject<FImplicitObjectUnionClustered>())
+			{
+				const_cast<FImplicitObjectUnion*>(ImplicitUnionClustered)->SetAllowBVH(true);
+			}
+		}
+	}
+
 	void
 	UpdateGeometry(
 		Chaos::FPBDRigidClusteredParticleHandle* Parent,
@@ -313,31 +366,7 @@ namespace Chaos
 		// by failing artificially this condition thmake sure we create a FImplicitObjectUnionClustered for this particle 
 		if (ProxyGeometry)
 		{
-			const FVector Scale = Parameters.Scale;
-			auto DeepCopyImplicit = [&Scale](const Chaos::FImplicitObjectPtr& ImplicitToCopy) -> Chaos::FImplicitObjectPtr
-			{
-				if (Scale.Equals(FVector::OneVector))
-				{
-					return ImplicitToCopy->DeepCopyGeometry();
-				}
-				else
-				{
-					return ImplicitToCopy->DeepCopyGeometryWithScale(Scale);
-				}
-			};
-			//ensureMsgf(false, TEXT("Checking usage with proxy"));
-			//@coverage {production}
-
-			Chaos::EImplicitObjectType GeometryType = ProxyGeometry->GetType();
-			// Don't copy if it is not a level set and scale is one
-			if (GeometryType != Chaos::ImplicitObjectType::LevelSet && Scale.Equals(FVector::OneVector))
-			{
-				Parent->SetGeometry(ProxyGeometry);
-			}
-			else
-			{
-				Parent->SetGeometry(DeepCopyImplicit(ProxyGeometry));
-			}
+			BuildScaledGeometry(Parent, ProxyGeometry, Parameters.Scale);
 		}
 		else if (Objects.Num() == 0)
 		{
@@ -404,12 +433,7 @@ namespace Chaos
 			}
 		}
 
-		//if children are ignore analytic and this is a dynamic geom, mark it too. todo(ocohen): clean this up
-		if (bUseParticleImplicit && Parent->GetGeometry())
-		{
-			Parent->GetGeometry()->SetDoCollide(false);
-		}
-
+		
 		if (Parameters.CollisionParticles)
 		{
 			SCOPE_CYCLE_COUNTER(STAT_UpdateGeometry_CopyPoints);
@@ -434,24 +458,10 @@ namespace Chaos
 				Parent->CollisionParticles()->UpdateAccelerationStructures();
 			}
 		}
+		const Chaos::FRigidTransform3 Xf(Parent->X(), Parent->R());
+		Parent->UpdateWorldSpaceState(Xf, FVec3(0));
 
-		if (FImplicitObjectPtr Implicit = Parent->GetGeometry())
-		{
-			// strange hacked initilization that seems misplaced and ill thought
-			Parent->SetHasBounds(true);
-			Parent->SetLocalBounds(Implicit->BoundingBox());
-			const Chaos::FRigidTransform3 Xf(Parent->X(), Parent->R());
-			Parent->UpdateWorldSpaceState(Xf, FVec3(0));
-
-			if (const FImplicitObjectUnion* ImplicitUnion = Implicit->GetObject<FImplicitObjectUnion>())
-			{
-				const_cast<FImplicitObjectUnion*>(ImplicitUnion)->SetAllowBVH(true);
-			}
-			else if (const FImplicitObjectUnion* ImplicitUnionClustered = Implicit->GetObject<FImplicitObjectUnionClustered>())
-			{
-				const_cast<FImplicitObjectUnion*>(ImplicitUnionClustered)->SetAllowBVH(true);
-			}
-		}
+		UpdateCollisionFlags(Parent, bUseParticleImplicit);
 	
 		// Update filter data on new shapes
 		const FRigidClustering::FRigidHandleArray& ChildrenArray = ChildrenMap[Parent];
diff --git a/Engine/Source/Runtime/Experimental/Chaos/Private/PhysicsProxy/GeometryCollectionPhysicsProxy.cpp b/Engine/Source/Runtime/Experimental/Chaos/Private/PhysicsProxy/GeometryCollectionPhysicsProxy.cpp
--- a/Engine/Source/Runtime/Experimental/Chaos/Private/PhysicsProxy/GeometryCollectionPhysicsProxy.cpp
+++ b/Engine/Source/Runtime/Experimental/Chaos/Private/PhysicsProxy/GeometryCollectionPhysicsProxy.cpp
@@ -144,6 +144,12 @@ FAutoConsoleVariableRef CVarGeometryCollectionUseRootBrokenFlag(
 	bGeometryCollectionUseRootBrokenFlag,
 	TEXT("If enabled, check if the root transform is broken in the proxy and disable the GT particle if so. Should be enabled - cvar is a failsafe to revert behaviour"));
 
+bool bGeometryCollectionScaleClusterGeometry = true;
+FAutoConsoleVariableRef CVarGeometryCollectionScaleClusterGeometry(
+	TEXT("p.GeometryCollection.ScaleClusterGeometry"),
+	bGeometryCollectionScaleClusterGeometry,
+	TEXT("If enabled, update the cluster geometry if the scale has changed"));
+
 enum EOverrideGCCollisionSetupForTraces
 {
 	GCCSFT_Property   = -1,  // Default: do what property says
@@ -3198,12 +3204,28 @@ void FGeometryCollectionPhysicsProxy::SetWorldTransform_External(const FTransfor
 
 		ExecuteOnPhysicsThread(*this,
 			[this, WorldTransform]()
-			{
+			{	
 				SetWorldTransform_Internal(WorldTransform);
 			});
 	}
 }
 
+void FGeometryCollectionPhysicsProxy::ScaleClusterGeometry_Internal(const FVector& WorldScale)
+{
+	const int32 NumTransforms = PhysicsThreadCollection.NumElements(FGeometryCollection::TransformGroup);
+	for (int32 TransformGroupIndex = 0; TransformGroupIndex < NumTransforms; ++TransformGroupIndex)
+	{
+		if (Chaos::FPBDRigidClusteredParticleHandle* ParticleHandle = SolverParticleHandles[TransformGroupIndex])
+		{
+			// Scale the geometry if necessary
+			BuildScaledGeometry(ParticleHandle, ParticleHandle->GetGeometry(), WorldScale);
+
+			// Update acceleration structure, bounds and collision flags
+			UpdateCollisionFlags(ParticleHandle, false);
+		}
+	}
+}
+
 void FGeometryCollectionPhysicsProxy::SetWorldTransform_Internal(const FTransform& InWorldTransform)
 {
 	using namespace Chaos;
@@ -3288,6 +3310,14 @@ void FGeometryCollectionPhysicsProxy::SetWorldTransform_Internal(const FTransfor
 			ClusterUnionManager.UpdateClusterUnionParticlesChildToParent(ClusterUnionIndex, DeferredClusterUnionParticleUpdates, DeferredClusterUnionChildToParentUpdates, false);
 		}
 	}
+	if (bGeometryCollectionScaleClusterGeometry)
+	{
+		const FVector ScaleRatio = Parameters.WorldTransform.GetScale3D() / Parameters.PrevWorldTransform.GetScale3D();
+		if (!ScaleRatio.Equals(FVector::OneVector))
+		{
+			ScaleClusterGeometry_Internal(ScaleRatio);
+		}
+	}
 }
 
 DECLARE_CYCLE_STAT(TEXT("FGeometryCollectionPhysicsProxy::SetUseStaticMeshCollisionForTraces_External"), STAT_SetUseStaticMeshCollisionForTraces_External, STATGROUP_Chaos);
diff --git a/Engine/Source/Runtime/Experimental/Chaos/Public/Chaos/ImplicitObjectScaled.h b/Engine/Source/Runtime/Experimental/Chaos/Public/Chaos/ImplicitObjectScaled.h
--- a/Engine/Source/Runtime/Experimental/Chaos/Public/Chaos/ImplicitObjectScaled.h
+++ b/Engine/Source/Runtime/Experimental/Chaos/Public/Chaos/ImplicitObjectScaled.h
@@ -1009,7 +1009,7 @@ public:
 	virtual Chaos::FImplicitObjectPtr CopyGeometryWithScale(const FVec3& Scale) const override
 	{
 		TImplicitObjectScaled<TConcrete, bInstanced>* Obj = CopyHelper(this);
-		Obj->SetScale(Scale);
+		Obj->SetScale(Obj->GetScale() * Scale);
 		return Chaos::FImplicitObjectPtr(Obj);
 	}
 
diff --git a/Engine/Source/Runtime/Experimental/Chaos/Public/Chaos/PBDRigidClusteringAlgo.h b/Engine/Source/Runtime/Experimental/Chaos/Public/Chaos/PBDRigidClusteringAlgo.h
--- a/Engine/Source/Runtime/Experimental/Chaos/Public/Chaos/PBDRigidClusteringAlgo.h
+++ b/Engine/Source/Runtime/Experimental/Chaos/Public/Chaos/PBDRigidClusteringAlgo.h
@@ -42,6 +42,12 @@ namespace Chaos
 		const FRigidClustering::FClusterMap& ChildrenMap,
 		const Chaos::FImplicitObjectPtr& ProxyGeometry,
 		const FClusterCreationParameters& Parameters);
+	
+	// Scale proxy geometry and set the result to the particle handle
+	void CHAOS_API BuildScaledGeometry(Chaos::FPBDRigidClusteredParticleHandle* ParticleHandle, const Chaos::FImplicitObjectPtr& ProxyGeometry, const FVector& WorldScale);
+
+	// Update the collision acceleration structure and flags
+	void CHAOS_API UpdateCollisionFlags(Chaos::FPBDRigidClusteredParticleHandle* ParticleHandle, const bool bUseParticleImplicit);
 
 	UE_DEPRECATED(5.4, "Use UpdateGeometry with FImplicitObjectPtr instead")
 	FORCEINLINE void CHAOS_API UpdateGeometry(
diff --git a/Engine/Source/Runtime/Experimental/Chaos/Public/PhysicsProxy/GeometryCollectionPhysicsProxy.h b/Engine/Source/Runtime/Experimental/Chaos/Public/PhysicsProxy/GeometryCollectionPhysicsProxy.h
--- a/Engine/Source/Runtime/Experimental/Chaos/Public/PhysicsProxy/GeometryCollectionPhysicsProxy.h
+++ b/Engine/Source/Runtime/Experimental/Chaos/Public/PhysicsProxy/GeometryCollectionPhysicsProxy.h
@@ -557,6 +557,8 @@ protected:
 
 	CHAOS_API void UpdateDamageThreshold_Internal();
 
+	/** Scale the cluster particles geometry (creates if necessary an additional TImplicitObjectScaled object into the implicits hierarchy) */
+	CHAOS_API void ScaleClusterGeometry_Internal(const FVector& WorldScale);
 
 	CHAOS_API void SetWorldTransform_Internal(const FTransform& WorldTransform);
 	CHAOS_API void SetFilterData_Internal(const FCollisionFilterData& NewSimFilter, const FCollisionFilterData& NewQueryFilter);
diff --git a/Engine/Source/Runtime/Experimental/GeometryCollectionEngine/Private/GeometryCollection/GeometryCollectionISMPoolComponent.cpp b/Engine/Source/Runtime/Experimental/GeometryCollectionEngine/Private/GeometryCollection/GeometryCollectionISMPoolComponent.cpp
--- a/Engine/Source/Runtime/Experimental/GeometryCollectionEngine/Private/GeometryCollection/GeometryCollectionISMPoolComponent.cpp
+++ b/Engine/Source/Runtime/Experimental/GeometryCollectionEngine/Private/GeometryCollection/GeometryCollectionISMPoolComponent.cpp
@@ -494,26 +494,27 @@ void FGeometryCollectionISMPool::UpdateAbsoluteTransforms(const FTransform& Base
 		const bool bReverseCulling = (GcIsm.MeshInstance.Desc.Flags & FISMComponentDescription::ReverseCulling) != 0;
 		check(GcIsm.MeshInstance.Desc.Position == FVector::ZeroVector);
 		
-		UInstancedStaticMeshComponent* Ism = GcIsm.ISMComponent;
-
-		if(bReverseCulling)
-		{
-			// As in InitISM we need to apply the inverted X scale for reverse culling.
-			// Just copy the transform and set an inverted scale to apply to the ISM
-			FVector BaseScale = BaseTransform.GetScale3D();
-			BaseScale.X = -BaseScale.X;
-			FTransform Flipped = BaseTransform;
-			Flipped.SetScale3D(BaseScale);
+		if(UInstancedStaticMeshComponent* Ism = GcIsm.ISMComponent)
+		{ 
+			if(bReverseCulling)
+			{
+				// As in InitISM we need to apply the inverted X scale for reverse culling.
+				// Just copy the transform and set an inverted scale to apply to the ISM
+				FVector BaseScale = BaseTransform.GetScale3D();
+				BaseScale.X = -BaseScale.X;
+				FTransform Flipped = BaseTransform;
+				Flipped.SetScale3D(BaseScale);
 			
-			Ism->SetComponentToWorld(Flipped);
-		}
-		else
-		{
-			Ism->SetComponentToWorld(BaseTransform);
-		}
+				Ism->SetComponentToWorld(Flipped);
+			}
+			else
+			{
+				Ism->SetComponentToWorld(BaseTransform);
+			}
 
-		Ism->UpdateComponentTransform(UpdateTransformFlags | EUpdateTransformFlags::SkipPhysicsUpdate, Teleport);
-		Ism->MarkRenderTransformDirty();
+			Ism->UpdateComponentTransform(UpdateTransformFlags | EUpdateTransformFlags::SkipPhysicsUpdate, Teleport);
+			Ism->MarkRenderTransformDirty();
+		}
 	}
 }
 
