Remove endpoint priming option (-PrimeEndPoint) from UnrealPak

#rnx

- Although this priming did help in certain cases it is not really robust enough for most CDN implementations so removing it so that people do not try to use it thinking that it will help them.

diff --git a/Engine/Source/Developer/IoStoreUtilities/Private/IoStoreUtilities.cpp b/Engine/Source/Developer/IoStoreUtilities/Private/IoStoreUtilities.cpp
--- a/Engine/Source/Developer/IoStoreUtilities/Private/IoStoreUtilities.cpp
+++ b/Engine/Source/Developer/IoStoreUtilities/Private/IoStoreUtilities.cpp
@@ -9456,18 +9456,3 @@ bool ListOnDemandTocs()
 
 	return Status.IsOk();
 }
-
-bool PrimeEndPoint(FStringView IoStoreOnDemandIniPath)
-{
-	FIoStatus Status = UE::IO::IAS::PrimeEndPoint(IoStoreOnDemandIniPath);
-	if (Status.IsOk())
-	{
-		return true;
-	}
-	else
-	{
-		UE_LOG(LogIoStore, Error, TEXT("PrimeEndPoint command failed due to: %s"), *Status.ToString());
-		return false;
-	}
-}
-
diff --git a/Engine/Source/Developer/IoStoreUtilities/Public/IoStoreUtilities.h b/Engine/Source/Developer/IoStoreUtilities/Public/IoStoreUtilities.h
--- a/Engine/Source/Developer/IoStoreUtilities/Public/IoStoreUtilities.h
+++ b/Engine/Source/Developer/IoStoreUtilities/Public/IoStoreUtilities.h
@@ -63,11 +63,4 @@ UE_API bool DownloadIoStoreContainerFiles(const TCHAR* TocPath);
 
 UE_API bool ListOnDemandTocs();
 
-/**
- * Attempts to connect to the endpoint in the provided ini file, download all
- * available tocs and then download all chunks in those tocs. This will
- * pre-warm the CDN and allow testing closer to the end users experience.
- */
-UE_API bool PrimeEndPoint(FStringView IoStoreOnDemandIniPath);
-
 #undef UE_API
diff --git a/Engine/Source/Developer/PakFileUtilities/Private/PakFileUtilities.cpp b/Engine/Source/Developer/PakFileUtilities/Private/PakFileUtilities.cpp
--- a/Engine/Source/Developer/PakFileUtilities/Private/PakFileUtilities.cpp
+++ b/Engine/Source/Developer/PakFileUtilities/Private/PakFileUtilities.cpp
@@ -5469,11 +5469,6 @@ bool ExecuteUnrealPak(const TCHAR* CmdLine)
 				return DownloadIoStoreContainerFiles(*IoStoreArg) == 0;
 			}
 
-			if (FParse::Value(CmdLine, TEXT("-PrimeEndPoint="), IoStoreArg))
-			{
-				return PrimeEndPoint(*IoStoreArg);
-			}
-
 			if (FParse::Param(CmdLine, TEXT("ListTocs")))
 			{
 				return ListOnDemandTocs();
diff --git a/Engine/Source/Runtime/Experimental/IoStoreOnDemand/Private/IoStoreOnDemand.cpp b/Engine/Source/Runtime/Experimental/IoStoreOnDemand/Private/IoStoreOnDemand.cpp
--- a/Engine/Source/Runtime/Experimental/IoStoreOnDemand/Private/IoStoreOnDemand.cpp
+++ b/Engine/Source/Runtime/Experimental/IoStoreOnDemand/Private/IoStoreOnDemand.cpp
@@ -21,7 +21,6 @@
 #include "Modules/ModuleManager.h"
 #include "OnDemandHttpClient.h"
 #include "OnDemandIoDispatcherBackend.h"
-#include "PrimeEndpoint.h"
 #include "Serialization/Archive.h"
 #include "Serialization/CompactBinarySerialization.h"
 #include "Serialization/CompactBinaryWriter.h"
@@ -2089,12 +2088,6 @@ FIoStatus ListTocs(const FIoStoreListTocsParams& Params)
 	return FIoStatus::Ok;
 }
 
-////////////////////////////////////////////////////////////////////////////////
-FIoStatus PrimeEndPoint(FStringView IoStoreOnDemandIniPath)
-{
-	return UE::IO::IAS::PrimeEndpointInternal(FString(IoStoreOnDemandIniPath));
-}
-
 #endif // (PLATFORM_DESKTOP && (IS_PROGRAM || WITH_EDITOR))
 
 
diff --git a/Engine/Source/Runtime/Experimental/IoStoreOnDemand/Private/PrimeEndpoint.cpp b/Engine/Source/Runtime/Experimental/IoStoreOnDemand/Private/PrimeEndpoint.cpp
deleted file mode 100644
diff --git a/Engine/Source/Runtime/Experimental/IoStoreOnDemand/Private/PrimeEndpoint.h b/Engine/Source/Runtime/Experimental/IoStoreOnDemand/Private/PrimeEndpoint.h
deleted file mode 100644
diff --git a/Engine/Source/Runtime/Experimental/IoStoreOnDemand/Public/IO/IoStoreOnDemand.h b/Engine/Source/Runtime/Experimental/IoStoreOnDemand/Public/IO/IoStoreOnDemand.h
--- a/Engine/Source/Runtime/Experimental/IoStoreOnDemand/Public/IO/IoStoreOnDemand.h
+++ b/Engine/Source/Runtime/Experimental/IoStoreOnDemand/Public/IO/IoStoreOnDemand.h
@@ -268,8 +268,6 @@ struct FIoStoreListTocsParams
 
 UE_API FIoStatus ListTocs(const FIoStoreListTocsParams& Params);
 
-////////////////////////////////////////////////////////////////////////////////
-UE_API FIoStatus PrimeEndPoint(FStringView IoStoreOnDemandIniPath);
 #endif // (IS_PROGRAM || WITH_EDITOR)
 
 class IOnDemandIoDispatcherBackend;
