Chaos - WIP support for iteration counts greater thn 127 (not used yet)

diff --git a/Engine/Source/Runtime/Experimental/Chaos/Public/Chaos/Evolution/IterationSettings.h b/Engine/Source/Runtime/Experimental/Chaos/Public/Chaos/Evolution/IterationSettings.h
--- a/Engine/Source/Runtime/Experimental/Chaos/Public/Chaos/Evolution/IterationSettings.h
+++ b/Engine/Source/Runtime/Experimental/Chaos/Public/Chaos/Evolution/IterationSettings.h
@@ -12,34 +12,45 @@ namespace Chaos
 		/**
 		 * Iteration counts for use by bodies and constraints. 
 		 * A simulation island will use the maximum number of iterations of any body or constraint in the island.
+		 * Iteration counts may be -1 which indicates that the configuration setting should be used.
 		*/
-		class FIterationSettings
+		template<typename TIndexType> 
+		class TIterationSettings
 		{
 		private:
-			using FIndexType = int8;
+			using FIndexType = TIndexType;
 			FIndexType NumPositionIterations;
 			FIndexType NumVelocityIterations;
 			FIndexType NumProjectionIterations;
 
 		public:
 			static const int32 InvalidIterations = INDEX_NONE;
-			static const int32 MaxPositionIterations = std::numeric_limits<FIndexType>::max();
-			static const int32 MaxVelocityIterations = std::numeric_limits<FIndexType>::max();
-			static const int32 MaxProjectionIterations = std::numeric_limits<FIndexType>::max();
+			static const int32 MaxIterations = std::numeric_limits<FIndexType>::max();
 
-			static FIterationSettings Empty()
+			static TIterationSettings<FIndexType> Empty()
 			{
-				return FIterationSettings(InvalidIterations, InvalidIterations, InvalidIterations);
+				return TIterationSettings<FIndexType>(InvalidIterations, InvalidIterations, InvalidIterations);
 			}
 
-			FIterationSettings()
+			// Merge (take the max) of the two iteration settings objects. Returned object has index type of the first parameter.
+			template<typename TOtherIndex>
+			static TIterationSettings<FIndexType> Merge(const TIterationSettings<FIndexType>& L, const TIterationSettings<TOtherIndex>& R)
+			{
+				return TIterationSettings<FIndexType>(
+					FMath::Max((int32)L.NumPositionIterations, (int32)R.NumPositionIterations),
+					FMath::Max((int32)L.NumVelocityIterations, (int32)R.NumVelocityIterations),
+					FMath::Max((int32)L.NumProjectionIterations, (int32)R.NumProjectionIterations)
+				);
+			}
+
+			TIterationSettings<FIndexType>()
 				: NumPositionIterations(InvalidIterations)
 				, NumVelocityIterations(InvalidIterations)
 				, NumProjectionIterations(InvalidIterations)
 			{
 			}
 
-			FIterationSettings(const int32 InNumPositionIterations, const int32 InNumVelocityIterations, const int32 InNumProjectionInterations)
+			TIterationSettings<FIndexType>(const int32 InNumPositionIterations, const int32 InNumVelocityIterations, const int32 InNumProjectionInterations)
 			{
 				SetNumPositionIterations(InNumPositionIterations);
 				SetNumVelocityIterations(InNumVelocityIterations);
@@ -50,18 +61,15 @@ namespace Chaos
 			int32 GetNumVelocityIterations() const { return NumVelocityIterations; }
 			int32 GetNumProjectionIterations() const { return NumProjectionIterations; }
 
-			void SetNumPositionIterations(const int32 InNum) { NumPositionIterations = FIndexType(FMath::Clamp(InNum, InvalidIterations, MaxPositionIterations)); }
-			void SetNumVelocityIterations(const int32 InNum) { NumVelocityIterations = FIndexType(FMath::Clamp(InNum, InvalidIterations, MaxVelocityIterations)); }
-			void SetNumProjectionIterations(const int32 InNum) { NumProjectionIterations = FIndexType(FMath::Clamp(InNum, InvalidIterations, MaxProjectionIterations)); }
-
-			static FIterationSettings Merge(const FIterationSettings& L, const FIterationSettings& R)
-			{
-				FIterationSettings Result;
-				Result.NumPositionIterations = FMath::Max(L.NumPositionIterations, R.NumPositionIterations);
-				Result.NumVelocityIterations = FMath::Max(L.NumVelocityIterations, R.NumVelocityIterations);
-				Result.NumProjectionIterations = FMath::Max(L.NumProjectionIterations, R.NumProjectionIterations);
-				return Result;
-			}
+			void SetNumPositionIterations(const int32 InNum) { NumPositionIterations = FIndexType(FMath::Clamp(InNum, InvalidIterations, MaxIterations)); }
+			void SetNumVelocityIterations(const int32 InNum) { NumVelocityIterations = FIndexType(FMath::Clamp(InNum, InvalidIterations, MaxIterations)); }
+			void SetNumProjectionIterations(const int32 InNum) { NumProjectionIterations = FIndexType(FMath::Clamp(InNum, InvalidIterations, MaxIterations)); }
 		};
+
+		// NOTE: only use signed types (or modify how we handle TIterationSettings::InvalidIterations above)
+		using FIterationSettings8 = TIterationSettings<int8>;
+		using FIterationSettings16 = TIterationSettings<int16>;
+
+		using FIterationSettings = FIterationSettings8;
 	}
 }
