Chaos - debugdraw - add option to show original collision even if there is an optimized convex

diff --git a/Engine/Source/Runtime/Experimental/Chaos/Private/Chaos/ChaosDebugDraw.cpp b/Engine/Source/Runtime/Experimental/Chaos/Private/Chaos/ChaosDebugDraw.cpp
--- a/Engine/Source/Runtime/Experimental/Chaos/Private/Chaos/ChaosDebugDraw.cpp
+++ b/Engine/Source/Runtime/Experimental/Chaos/Private/Chaos/ChaosDebugDraw.cpp
@@ -91,6 +91,9 @@ namespace Chaos
 		FAutoConsoleVariableRef CVarChaosDebugDrawCoreShapes(TEXT("p.Chaos.DebugDraw.ShowCoreShapes"), bChaosDebugDebugDrawCoreShapes, TEXT("Whether to show the core (margin-reduced) shape where applicable"));
 		FAutoConsoleVariableRef CVarChaosDebugDrawExactShapes(TEXT("p.Chaos.DebugDraw.ShowExactCoreShapes"), bChaosDebugDebugDrawExactCoreShapes, TEXT("Whether to show the exact core shape. NOTE: Extremely expensive and should only be used on a small scene with a couple convex shapes in it"));
 
+		bool bChaosDebugDebugDrawShowOptimizedConvexes = false;
+		FAutoConsoleVariableRef CVarChaosDebugDrawShowOptimizedConvexes(TEXT("p.Chaos.DebugDraw.ShowOptimizedConvexes"), bChaosDebugDebugDrawShowOptimizedConvexes, TEXT(""));
+
 		bool bChaosDebugDebugDrawIslands = true;
 		FAutoConsoleVariableRef CVarChaosDebugDrawIslands(TEXT("p.Chaos.DebugDraw.ShowIslands"), bChaosDebugDebugDrawIslands, TEXT("Whether to show the iosland boxes when drawing islands (if you want only the contact graph)"));
 
@@ -896,6 +899,11 @@ namespace Chaos
 		
 		bool DrawConvexOptimizerShapes(const FVec3& P, const FRotation3& Q, const FGeometryParticleHandle* Particle, const FColor& InColor, const FChaosDebugDrawSettings& Settings)
 		{
+			if (!bChaosDebugDebugDrawShowOptimizedConvexes)
+			{
+				return false;
+			}
+
 			if(auto ClusteredParticle = Particle->CastToClustered())
 			{
 				const TPimplPtr<Private::FConvexOptimizer>& ConvexOptimizer = ClusteredParticle->ConvexOptimizer();
