Chaos - add cvar to enable/disable the one-way interaction system just in case

diff --git a/Engine/Source/Runtime/Experimental/Chaos/Private/Chaos/Collision/PBDCollisionConstraint.cpp b/Engine/Source/Runtime/Experimental/Chaos/Private/Chaos/Collision/PBDCollisionConstraint.cpp
--- a/Engine/Source/Runtime/Experimental/Chaos/Private/Chaos/Collision/PBDCollisionConstraint.cpp
+++ b/Engine/Source/Runtime/Experimental/Chaos/Private/Chaos/Collision/PBDCollisionConstraint.cpp
@@ -54,6 +54,10 @@ namespace Chaos
 	FRealSingle Chaos_Collision_Stiffness = -1.0f;
 	FAutoConsoleVariableRef CVarChaos_Collision_Stiffness(TEXT("p.Chaos.Collision.Stiffness"), Chaos_Collision_Stiffness, TEXT("Override the collision solver stiffness (if >= 0)"));
 
+	// Enable/Disable the one-way interaction system for collisions
+	bool bChaos_Collision_EnableOneWayInteraction = true;
+	FAutoConsoleVariableRef CVarChaos_Collision_EnableOneWayInteraction(TEXT("p.Chaos.Collision.EnableOneWayInteraction"), bChaos_Collision_EnableOneWayInteraction, TEXT("Whether the one-way interaction flag is respected in collision constraints"));
+
 	// The stiffness used for collision between one-way objects and other dynamic objects. The stiffness against kinematics and statics is not affected. A value of less than 1
 	// will treat collisions with the static scene as "harder" than collisions with dynamics which will prevent one-way objects from getting squeezed out of the world when a dynamic
 	// lands on top of them. Instead, they will stay on top of the terrain and sink into the dynamic.
@@ -474,7 +478,7 @@ namespace Chaos
 
 	void FPBDCollisionConstraint::UpdateMassScales()
 	{
-		if (Flags.bIsOneWayInteraction)
+		if (Flags.bIsOneWayInteraction && bChaos_Collision_EnableOneWayInteraction)
 		{
 			const bool bOneWay0 = FConstGenericParticleHandle(GetParticle0())->OneWayInteraction();
 			const bool bOneWay1 = FConstGenericParticleHandle(GetParticle1())->OneWayInteraction();
