Chaos - handle zero inertia

diff --git a/Engine/Source/Runtime/Experimental/Chaos/Private/Chaos/MassConditioning.cpp b/Engine/Source/Runtime/Experimental/Chaos/Private/Chaos/MassConditioning.cpp
--- a/Engine/Source/Runtime/Experimental/Chaos/Private/Chaos/MassConditioning.cpp
+++ b/Engine/Source/Runtime/Experimental/Chaos/Private/Chaos/MassConditioning.cpp
@@ -10,6 +10,7 @@ namespace Chaos
 	FVec3f CalculateInertiaConditioning(const FRealSingle InvM, const FVec3f& InvI, const FVec3f& ConstraintExtents, const FRealSingle MaxDistance, const FRealSingle MaxRotationRatio, const FRealSingle MaxInvInertiaComponentRatio)
 	{
 		check(InvM > 0);
+		check(!InvI.IsZero());
 
 		FVec3f InvInertiaScale = FVec3f(1);
 
@@ -108,7 +109,7 @@ namespace Chaos
 
 	FVec3f CalculateParticleInertiaConditioning(const FPBDRigidParticleHandle* Rigid, const FRealSingle MaxDistance, const FRealSingle MaxRotationRatio, const FRealSingle MaxInvInertiaComponentRatio)
 	{
-		if (Rigid->InvM() == 0)
+		if (FMath::IsNearlyZero(Rigid->InvM()) || Rigid->InvI().IsNearlyZero())
 		{
 			return FVec3f(1);
 		}
