Iris - Update world locations for spatially filtered objects that are assumed to be able to move.

diff --git a/Engine/Plugins/Runtime/ReplicationSystemTestPlugin/Source/Private/Tests/ReplicationSystem/Filtering/TestGridFilter.cpp b/Engine/Plugins/Runtime/ReplicationSystemTestPlugin/Source/Private/Tests/ReplicationSystem/Filtering/TestGridFilter.cpp
--- a/Engine/Plugins/Runtime/ReplicationSystemTestPlugin/Source/Private/Tests/ReplicationSystem/Filtering/TestGridFilter.cpp
+++ b/Engine/Plugins/Runtime/ReplicationSystemTestPlugin/Source/Private/Tests/ReplicationSystem/Filtering/TestGridFilter.cpp
@@ -111,7 +111,7 @@ UE_NET_TEST_FIXTURE(FTestGridFilterFixture, TestWorldLocGridFilter)
 	UReplicatedTestObject* ServerObjectCulled = Server->CreateObject(Params);
 	UReplicatedTestObject* ServerObjectVeryFar = Server->CreateObject(Params);
 
-	const UNetObjectGridFilterConfig* DefaultGridConfig = Cast<UNetObjectGridFilterConfig>(UNetObjectGridFilterConfig::StaticClass()->GetDefaultObject());
+	const UNetObjectGridFilterConfig* DefaultGridConfig = GetDefault<UNetObjectGridFilterConfig>();
 
 	struct FObjectLoc 
 	{
@@ -164,6 +164,70 @@ UE_NET_TEST_FIXTURE(FTestGridFilterFixture, TestWorldLocGridFilter)
 	Server->DestroyObject(ServerObjectVeryFar);
 }
 
+// Test world loc filter
+UE_NET_TEST_FIXTURE(FTestGridFilterFixture, TestWorldLocationIsFrequentlyUpdatedForNonDormantObject)
+{
+	struct FObjectLoc
+	{
+		FVector Loc;
+		float CullDistance;
+	};
+	
+	TMap<const UObject*, FObjectLoc> ObjectLocs;
+
+	Server->GetReplicationBridge()->SetExternalWorldLocationUpdateFunctor([&](FNetRefHandle NetHandle, const UObject* ReplicatedObject, FVector& OutLocation, float& OutCullDistance)
+		{
+			const FObjectLoc& ObjectLoc = ObjectLocs[ReplicatedObject];
+			OutLocation = ObjectLoc.Loc;
+			OutCullDistance = ObjectLoc.CullDistance;
+		});
+
+	// Add client (view location at 0,0,0)
+	FReplicationSystemTestClient* Client = CreateClient();
+
+	// Spawn object with WorldLocation's on server
+	UObjectReplicationBridge::FCreateNetRefHandleParams Params;
+	Params.bNeedsWorldLocationUpdate = true;
+	Params.bIsDormant = false;
+	Params.bUseClassConfigDynamicFilter = false;
+	Params.bUseExplicitDynamicFilter = true;
+	Params.ExplicitDynamicFilterName = FName("NetObjectGridWorldLocFilter");
+
+	UReplicatedTestObject* ServerObject = Server->CreateObject(Params);
+
+	// Visible objects
+	ObjectLocs.Add(ServerObject, FObjectLoc{ FVector::ZeroVector, 1500.f });
+
+	// Send and deliver packet
+	Server->UpdateAndSend({ Client });
+
+	// Verify objects has been created
+	UE_NET_ASSERT_NE(Client->GetReplicationBridge()->GetReplicatedObject(ServerObject->NetRefHandle), nullptr);
+
+	const UNetObjectGridFilterConfig* DefaultGridConfig = GetDefault<UNetObjectGridFilterConfig>();
+
+	// Update the world location to a cell the object shouldn't previously have been touching, without marking the object dirty
+	ObjectLocs.Emplace(ServerObject, FObjectLoc{ FVector(DefaultGridConfig->CellSizeX + 1600.0f,DefaultGridConfig->CellSizeY + 1600.0f, 0), 1500.f });
+
+	// Send and deliver packet
+	for (uint32 LoopIt = 0, LoopEndIt = DefaultGridConfig->ViewPosRelevancyFrameCount; LoopIt <= LoopEndIt; ++LoopIt)
+	{
+		Server->UpdateAndSend({ Client });
+	}
+
+	// Object should now have been destroyed.
+	UE_NET_ASSERT_EQ(Client->GetReplicationBridge()->GetReplicatedObject(ServerObject->NetRefHandle), nullptr);
+
+	// Return object to origin.
+	ObjectLocs.Emplace(ServerObject, FObjectLoc{ FVector::ZeroVector, 1500.f });
+
+	// Send and deliver packet
+	Server->UpdateAndSend({ Client });
+
+	// Object should have been re-created
+	UE_NET_ASSERT_NE(Client->GetReplicationBridge()->GetReplicatedObject(ServerObject->NetRefHandle), nullptr);
+}
+
 // Test fragment data filter
 UE_NET_TEST_FIXTURE(FTestGridFilterFixture, TestFragmentFilter)
 {
@@ -173,8 +237,7 @@ UE_NET_TEST_FIXTURE(FTestGridFilterFixture, TestFragmentFilter)
 	UTestLocationFragmentFilteringObject* ServerObjectCulled = Server->CreateObject<UTestLocationFragmentFilteringObject>();
 	UTestLocationFragmentFilteringObject* ServerObjectVeryFar = Server->CreateObject<UTestLocationFragmentFilteringObject>();
 
-	const UNetObjectGridFilterConfig* DefaultGridConfig = Cast<UNetObjectGridFilterConfig>(UNetObjectGridFilterConfig::StaticClass()->GetDefaultObject());
-
+	const UNetObjectGridFilterConfig* DefaultGridConfig = GetDefault<UNetObjectGridFilterConfig>();
 
 	ServerObjectZero->WorldLocation = FVector::ZeroVector;
 	ServerObjectZero->NetCullDistanceSquared = 1500.f * 1500.f;
diff --git a/Engine/Source/Runtime/Engine/Private/Net/Iris/ReplicationSystem/ActorReplicationBridge.cpp b/Engine/Source/Runtime/Engine/Private/Net/Iris/ReplicationSystem/ActorReplicationBridge.cpp
--- a/Engine/Source/Runtime/Engine/Private/Net/Iris/ReplicationSystem/ActorReplicationBridge.cpp
+++ b/Engine/Source/Runtime/Engine/Private/Net/Iris/ReplicationSystem/ActorReplicationBridge.cpp
@@ -247,6 +247,7 @@ UE::Net::FNetRefHandle UActorReplicationBridge::BeginReplication(AActor* Actor,
 	Super::FCreateNetRefHandleParams CreateNetRefHandleParams = UObjectReplicationBridge::DefaultCreateNetRefHandleParams;
 	CreateNetRefHandleParams.bNeedsPreUpdate = 1U;
 	CreateNetRefHandleParams.bNeedsWorldLocationUpdate = 1U;
+	CreateNetRefHandleParams.bIsDormant = Actor->NetDormancy > DORM_Awake;
 	CreateNetRefHandleParams.StaticPriority = (Actor->bAlwaysRelevant || Actor->bOnlyRelevantToOwner) ? Actor->NetPriority : 0.0f;
 	CreateNetRefHandleParams.PollFrequency = Actor->NetUpdateFrequency;
 
diff --git a/Engine/Source/Runtime/Experimental/Iris/Core/Private/Iris/ReplicationSystem/Filtering/NetObjectFilter.cpp b/Engine/Source/Runtime/Experimental/Iris/Core/Private/Iris/ReplicationSystem/Filtering/NetObjectFilter.cpp
--- a/Engine/Source/Runtime/Experimental/Iris/Core/Private/Iris/ReplicationSystem/Filtering/NetObjectFilter.cpp
+++ b/Engine/Source/Runtime/Experimental/Iris/Core/Private/Iris/ReplicationSystem/Filtering/NetObjectFilter.cpp
@@ -11,6 +11,11 @@ FNetObjectFilteringParams::FNetObjectFilteringParams(const UE::Net::FNetBitArray
 {
 }
 
+FNetObjectPreFilteringParams::FNetObjectPreFilteringParams(const UE::Net::FNetBitArrayView InFilteredObjects)
+: FilteredObjects(InFilteredObjects)
+{
+}
+
 UNetObjectFilter::UNetObjectFilter()
 {
 }
diff --git a/Engine/Source/Runtime/Experimental/Iris/Core/Private/Iris/ReplicationSystem/Filtering/NetObjectGridFilter.cpp b/Engine/Source/Runtime/Experimental/Iris/Core/Private/Iris/ReplicationSystem/Filtering/NetObjectGridFilter.cpp
--- a/Engine/Source/Runtime/Experimental/Iris/Core/Private/Iris/ReplicationSystem/Filtering/NetObjectGridFilter.cpp
+++ b/Engine/Source/Runtime/Experimental/Iris/Core/Private/Iris/ReplicationSystem/Filtering/NetObjectGridFilter.cpp
@@ -11,6 +11,8 @@
 
 void UNetObjectGridFilter::OnInit(FNetObjectFilterInitParams& Params)
 {
+	AddFilterTraits(ENetFilterTraits::Spatial);
+
 	Config = TStrongObjectPtr<UNetObjectGridFilterConfig>(CastChecked<UNetObjectGridFilterConfig>(Params.Config));
 	checkf(Config.IsValid(), TEXT("Need config to operate."));
 
@@ -73,18 +75,6 @@ void UNetObjectGridFilter::RemoveObject(uint32 ObjectIndex, const FNetObjectFilt
 	FreeObjectInfo(InfoIndex);
 }
 
-void UNetObjectGridFilter::UpdateObjects(FNetObjectFilterUpdateParams& Params)
-{
-	for (SIZE_T ObjectIt = 0, ObjectEndIt = Params.ObjectCount; ObjectIt != ObjectEndIt; ++ObjectIt)
-	{
-		const uint32 ObjectIndex = Params.ObjectIndices[ObjectIt];
-
-		const FObjectLocationInfo& ObjectLocationInfo = static_cast<const FObjectLocationInfo&>(Params.FilteringInfos[ObjectIndex]);
-		const UE::Net::FReplicationInstanceProtocol* InstanceProtocol = Params.InstanceProtocols ? Params.InstanceProtocols[ObjectIt] : nullptr;
-		UpdateCellInfoForObject(ObjectLocationInfo, InstanceProtocol);
-	}
-}
-
 void UNetObjectGridFilter::PreFilter(FNetObjectPreFilteringParams&)
 {
 	++FrameIndex;
@@ -417,6 +407,26 @@ void UNetObjectGridWorldLocFilter::OnInit(FNetObjectFilterInitParams& Params)
 	WorldLocations = &Params.ReplicationSystem->GetWorldLocations();
 }
 
+void UNetObjectGridWorldLocFilter::UpdateObjects(FNetObjectFilterUpdateParams&)
+{
+}
+
+void UNetObjectGridWorldLocFilter::PreFilter(FNetObjectPreFilteringParams& Params)
+{
+	Super::PreFilter(Params);
+
+	// Update logic performed here in order to not rely on any object being dirtied.
+	auto UpdateCells = [this, &Params](uint32 ObjectIndex)
+		{
+			const FObjectLocationInfo& ObjectLocationInfo = static_cast<const FObjectLocationInfo&>(Params.FilteringInfos[ObjectIndex]);
+			this->UpdateCellInfoForObject(ObjectLocationInfo, static_cast<const UE::Net::FReplicationInstanceProtocol*>(nullptr));
+		};
+
+	// Update cell info for all objects that have moved.
+	UE::Net::FNetBitArrayView ObjectsWithDirtyWorldLocations = WorldLocations->GetObjectsWithDirtyInfo();
+	UE::Net::FNetBitArrayView::ForAllSetBits(Params.FilteredObjects, ObjectsWithDirtyWorldLocations, UE::Net::FNetBitArrayBase::AndOp, UpdateCells);
+}
+
 void UNetObjectGridWorldLocFilter::UpdateObjectInfo(UNetObjectGridFilter::FPerObjectInfo& PerObjectInfo, const UNetObjectGridFilter::FObjectLocationInfo& ObjectLocationInfo, const UE::Net::FReplicationInstanceProtocol* InstanceProtocol)
 {
 	check(ObjectLocationInfo.IsUsingWorldLocations());
@@ -450,6 +460,17 @@ void UNetObjectGridFragmentLocFilter::OnInit(FNetObjectFilterInitParams& InitPar
 	SetupFilterType(ENetFilterType::PostPoll_FragmentBased);
 }
 
+void UNetObjectGridFragmentLocFilter::UpdateObjects(FNetObjectFilterUpdateParams& Params)
+{
+	for (SIZE_T ObjectIt = 0, ObjectEndIt = Params.ObjectCount; ObjectIt != ObjectEndIt; ++ObjectIt)
+	{
+		const uint32 ObjectIndex = Params.ObjectIndices[ObjectIt];
+		const FObjectLocationInfo& ObjectLocationInfo = static_cast<const FObjectLocationInfo&>(Params.FilteringInfos[ObjectIndex]);
+		const UE::Net::FReplicationInstanceProtocol* InstanceProtocol = Params.InstanceProtocols ? Params.InstanceProtocols[ObjectIt] : nullptr;
+		UpdateCellInfoForObject(ObjectLocationInfo, InstanceProtocol);
+	}
+}
+
 bool UNetObjectGridFragmentLocFilter::BuildObjectInfo(uint32 ObjectIndex, FNetObjectFilterAddObjectParams& Params)
 {
 	UE::Net::FRepTagFindInfo WorldLocationTagInfo;
diff --git a/Engine/Source/Runtime/Experimental/Iris/Core/Private/Iris/ReplicationSystem/Filtering/ReplicationFiltering.cpp b/Engine/Source/Runtime/Experimental/Iris/Core/Private/Iris/ReplicationSystem/Filtering/ReplicationFiltering.cpp
--- a/Engine/Source/Runtime/Experimental/Iris/Core/Private/Iris/ReplicationSystem/Filtering/ReplicationFiltering.cpp
+++ b/Engine/Source/Runtime/Experimental/Iris/Core/Private/Iris/ReplicationSystem/Filtering/ReplicationFiltering.cpp
@@ -523,6 +523,19 @@ bool FReplicationFiltering::SetFilter(FInternalNetRefIndex ObjectIndex, FNetObje
 	return false;
 }
 
+bool FReplicationFiltering::IsUsingSpatialFilter(FInternalNetRefIndex ObjectIndex) const
+{
+	const uint8 DynamicFilterIndex = ObjectIndexToDynamicFilterIndex[ObjectIndex];
+	if (DynamicFilterIndex == InvalidDynamicFilterIndex)
+	{
+		return false;
+	}
+
+	const FFilterInfo& FilterInfo = DynamicFilterInfos[DynamicFilterIndex];
+	const bool bIsUsingSpatialFilter = EnumHasAnyFlags(FilterInfo.Filter->GetFilterTraits(), ENetFilterTraits::Spatial);
+	return bIsUsingSpatialFilter;
+}
+
 bool FReplicationFiltering::SetConnectionFilter(FInternalNetRefIndex ObjectIndex, const FNetBitArrayView& ConnectionIndices, ENetFilterStatus ReplicationStatus)
 {
 	// If we have some sort of other filtering set then remove that. From now on we want to use connection filtering.
@@ -1242,20 +1255,17 @@ void FReplicationFiltering::PreUpdateDynamicFiltering(ENetFilterType FilterType)
 
 	// Give filters a chance to prepare for filtering. It's only called if any object has the filter set.
 	{
-		FNetObjectPreFilteringParams PreFilteringParams;
-		PreFilteringParams.ValidConnections = MakeNetBitArrayView(ValidConnections);
-
 		for (FFilterInfo& Info : DynamicFilterInfos)
 		{
-			if (Info.ObjectCount == 0U)
+			if (Info.ObjectCount == 0U || Info.Type != FilterType)
 			{
 				continue;
 			}
 
-			if (Info.Type == FilterType)
-			{
-				Info.Filter->PreFilter(PreFilteringParams);
-			}
+			FNetObjectPreFilteringParams PreFilteringParams(MakeNetBitArrayView(Info.FilteredObjects));
+			PreFilteringParams.FilteringInfos = MakeArrayView(NetObjectFilteringInfos);
+			PreFilteringParams.ValidConnections = MakeNetBitArrayView(ValidConnections);
+			Info.Filter->PreFilter(PreFilteringParams);
 		}
 	}
 }
diff --git a/Engine/Source/Runtime/Experimental/Iris/Core/Private/Iris/ReplicationSystem/Filtering/ReplicationFiltering.h b/Engine/Source/Runtime/Experimental/Iris/Core/Private/Iris/ReplicationSystem/Filtering/ReplicationFiltering.h
--- a/Engine/Source/Runtime/Experimental/Iris/Core/Private/Iris/ReplicationSystem/Filtering/ReplicationFiltering.h
+++ b/Engine/Source/Runtime/Experimental/Iris/Core/Private/Iris/ReplicationSystem/Filtering/ReplicationFiltering.h
@@ -86,6 +86,7 @@ public:
 	uint32 GetOwningConnection(FInternalNetRefIndex ObjectIndex) const { return !bHasDirtyOwner ? ObjectIndexToOwningConnection[ObjectIndex] : GetOwningConnectionIfDirty(ObjectIndex); }
 
 	bool SetFilter(FInternalNetRefIndex ObjectIndex, FNetObjectFilterHandle Filter);
+	bool IsUsingSpatialFilter(FInternalNetRefIndex ObjectIndex) const;
 
 	// Set whether an object is allowed to be replicated to certain connections or not.
 	bool SetConnectionFilter(FInternalNetRefIndex ObjectIndex, const FNetBitArrayView& ConnectionIndices, ENetFilterStatus ReplicationStatus);
diff --git a/Engine/Source/Runtime/Experimental/Iris/Core/Private/Iris/ReplicationSystem/ObjectReplicationBridge.cpp b/Engine/Source/Runtime/Experimental/Iris/Core/Private/Iris/ReplicationSystem/ObjectReplicationBridge.cpp
--- a/Engine/Source/Runtime/Experimental/Iris/Core/Private/Iris/ReplicationSystem/ObjectReplicationBridge.cpp
+++ b/Engine/Source/Runtime/Experimental/Iris/Core/Private/Iris/ReplicationSystem/ObjectReplicationBridge.cpp
@@ -91,8 +91,10 @@ UObjectReplicationBridge::FCreateNetRefHandleParams UObjectReplicationBridge::De
 {
 	.bCanReceive=false, 
 	.bNeedsPreUpdate=false,
-	.bNeedsWorldLocationUpdate=false, 
+	.bNeedsWorldLocationUpdate=false,
+	.bIsDormant=false,
 	.bUseClassConfigDynamicFilter=true,
+	.bUseExplicitDynamicFilter=false,
 	.StaticPriority=0.0f, 
 	.PollFrequency=0.0f,
 };
@@ -378,7 +380,13 @@ UE::Net::FNetRefHandle UObjectReplicationBridge::BeginReplication(UObject* Insta
 				ReplicationSystem->SetDeltaCompressionStatus(RefHandle, ENetObjectDeltaCompressionStatus::Allow);
 			}
 
-			// Release instance protocol from the uniquePtr as it is now successfully bound to the handle
+			// Spatially filtered non-dormant objects requires frequent world location updates. Expecting a better solution that instead of us polling will inform us when locations change, UE-193004.
+			if (Params.bNeedsWorldLocationUpdate && !Params.bIsDormant)
+			{
+				OptionallySetObjectRequiresFrequentWorldLocationUpdate(RefHandle, true);
+			}
+
+			// Release instance protocol from the UniquePtr as it is now successfully bound to the handle
 			(void)InstanceProtocol.Release();
 
 			return RefHandle;
@@ -1127,11 +1135,11 @@ void UObjectReplicationBridge::UpdateInstancesWorldLocation()
 		}
 	};
 
-	// $IRIS TODO: This code assumes users are calling MarkDirty whenever an actor changes location. Need to add a location tracker to ensure if that is not the case.
+	// Objects marked as dirty or that have requested frequent world location updates will be updated. Failing to do either when the location has changed will result in replication issues when using spatial filters such as the NetObjectGridFilter.
 	FDirtyObjectsAccessor DirtyObjectsAccessor(ReplicationSystemInternal->GetDirtyNetObjectTracker());
 	const FNetBitArrayView DirtyObjectsThisFrame = DirtyObjectsAccessor.GetDirtyNetObjects();
-	DirtyObjectsThisFrame.ForAllSetBits(UpdateInstanceWorldLocation);
-
+	const FNetBitArrayView ObjectsRequiringFrequentUpdates = WorldLocations.GetObjectsRequiringFrequentWorldLocationUpdate();
+	FNetBitArrayView::ForAllSetBits(DirtyObjectsThisFrame, ObjectsRequiringFrequentUpdates, FNetBitArrayBase::OrOp, UpdateInstanceWorldLocation);
 }
 
 void UObjectReplicationBridge::SetPollWithObject(FNetRefHandle ObjectToPollWithHandle, FNetRefHandle ObjectHandle)
@@ -1158,6 +1166,7 @@ bool UObjectReplicationBridge::GetObjectWantsToBeDormant(FNetRefHandle Handle) c
 
 void UObjectReplicationBridge::SetObjectWantsToBeDormant(FNetRefHandle Handle, bool bWantsToBeDormant)
 {
+	using namespace UE::Net;
 	using namespace UE::Net::Private;
 
 	FReplicationSystemInternal* ReplicationSystemInternal = GetReplicationSystem()->GetReplicationSystemInternal();
@@ -1165,7 +1174,7 @@ void UObjectReplicationBridge::SetObjectWantsToBeDormant(FNetRefHandle Handle, b
 
 	if (const FInternalNetRefIndex InternalObjectIndex = LocalNetRefHandleManager.GetInternalIndex(Handle))
 	{
-		UE::Net::FNetBitArray& WantToBeDormantObjects = LocalNetRefHandleManager.GetWantToBeDormantInternalIndices();
+		FNetBitArray& WantToBeDormantObjects = LocalNetRefHandleManager.GetWantToBeDormantInternalIndices();
 
 		// Update pending dormancy status
 		WantToBeDormantObjects.SetBitValue(InternalObjectIndex, bWantsToBeDormant);
@@ -1184,6 +1193,9 @@ void UObjectReplicationBridge::SetObjectWantsToBeDormant(FNetRefHandle Handle, b
 		{
 			WantToBeDormantObjects.SetBitValue(SubObjectInternalIndex, bWantsToBeDormant);
 		}
+
+		// Request frequent world location updates for non-dormant spatially filtered objects.
+		OptionallySetObjectRequiresFrequentWorldLocationUpdate(Handle, !bWantsToBeDormant);
 	}
 }
 
@@ -1422,6 +1434,24 @@ bool UObjectReplicationBridge::IsClassCritical(const UClass* Class)
 	return false;
 }
 
+void UObjectReplicationBridge::OptionallySetObjectRequiresFrequentWorldLocationUpdate(FNetRefHandle RefHandle, bool bDesiresFrequentWorldLocationUpdate)
+{
+	using namespace UE::Net;
+	using namespace UE::Net::Private;
+
+	FReplicationSystemInternal* ReplicationSystemInternal = GetReplicationSystem()->GetReplicationSystemInternal();
+	const FNetRefHandleManager& LocalNetRefHandleManager = ReplicationSystemInternal->GetNetRefHandleManager();
+	const FInternalNetRefIndex InternalObjectIndex = LocalNetRefHandleManager.GetInternalIndex(RefHandle);
+	FWorldLocations& WorldLocations = ReplicationSystemInternal->GetWorldLocations();
+	// When this function is called due to dormancy changes we don't know that the object requires world location updates at all. Checking if it has world location info is how we find that out.
+	if (WorldLocations.HasInfoForObject(InternalObjectIndex))
+	{
+		const FReplicationFiltering& Filtering = ReplicationSystemInternal->GetFiltering();
+		const bool bRequireFrequentWorldLocationUpdates = bDesiresFrequentWorldLocationUpdate && Filtering.IsUsingSpatialFilter(InternalObjectIndex);
+		WorldLocations.SetObjectRequiresFrequentWorldLocationUpdate(InternalObjectIndex, bRequireFrequentWorldLocationUpdates);
+	}
+}
+
 int32 UObjectReplicationBridge::GetTypeStatsIndex(const UClass* Class)
 {
 	using namespace UE::Net;
diff --git a/Engine/Source/Runtime/Experimental/Iris/Core/Private/Iris/ReplicationSystem/ReplicationSystem.cpp b/Engine/Source/Runtime/Experimental/Iris/Core/Private/Iris/ReplicationSystem/ReplicationSystem.cpp
--- a/Engine/Source/Runtime/Experimental/Iris/Core/Private/Iris/ReplicationSystem/ReplicationSystem.cpp
+++ b/Engine/Source/Runtime/Experimental/Iris/Core/Private/Iris/ReplicationSystem/ReplicationSystem.cpp
@@ -309,6 +309,10 @@ public:
 	{
 		IRIS_PROFILER_SCOPE(FReplicationSystem_UpdateWorldLocations);
 
+		// Reset dirty object info before updating.
+		FWorldLocations& WorldLocations = ReplicationSystemInternal.GetWorldLocations();
+		WorldLocations.ResetObjectsWithDirtyInfo();
+
 		ReplicationSystemInternal.GetReplicationBridge()->CallUpdateInstancesWorldLocation();
 	}
 
diff --git a/Engine/Source/Runtime/Experimental/Iris/Core/Private/Iris/ReplicationSystem/WorldLocations.cpp b/Engine/Source/Runtime/Experimental/Iris/Core/Private/Iris/ReplicationSystem/WorldLocations.cpp
--- a/Engine/Source/Runtime/Experimental/Iris/Core/Private/Iris/ReplicationSystem/WorldLocations.cpp
+++ b/Engine/Source/Runtime/Experimental/Iris/Core/Private/Iris/ReplicationSystem/WorldLocations.cpp
@@ -9,6 +9,8 @@ namespace UE::Net
 void FWorldLocations::Init(const FWorldLocationsInitParams& InitParams)
 {
 	ValidInfoIndexes.Init(InitParams.MaxObjectCount);
+	ObjectsWithDirtyInfo.Init(InitParams.MaxObjectCount);
+	ObjectsRequiringFrequentWorldLocationUpdate.Init(InitParams.MaxObjectCount);
 }
 
 void FWorldLocations::InitObjectInfoCache(uint32 ObjectIndex)
@@ -34,18 +36,33 @@ void FWorldLocations::InitObjectInfoCache(uint32 ObjectIndex)
 void FWorldLocations::RemoveObjectInfoCache(uint32 ObjectIndex)
 {
 	ValidInfoIndexes.ClearBit(ObjectIndex);
+	ObjectsWithDirtyInfo.ClearBit(ObjectIndex);
+	ObjectsRequiringFrequentWorldLocationUpdate.ClearBit(ObjectIndex);
 }
 
 void FWorldLocations::SetObjectInfo(uint32 ObjectIndex, const FWorldLocations::FObjectInfo& ObjectInfo)
 {
 	checkSlow(ValidInfoIndexes.IsBitSet(ObjectIndex));
-	StoredObjectInfo[ObjectIndex] = ObjectInfo;
+	FObjectInfo& TargetObjectInfo = StoredObjectInfo[ObjectIndex];
+	const bool bHasInfoChanged = ObjectsWithDirtyInfo.GetBit(ObjectIndex) || TargetObjectInfo.WorldLocation != ObjectInfo.WorldLocation || TargetObjectInfo.CullDistance != ObjectInfo.CullDistance;
+	TargetObjectInfo = ObjectInfo;
+
+	ObjectsWithDirtyInfo.SetBitValue(ObjectIndex, bHasInfoChanged);
 }
 
 void FWorldLocations::UpdateWorldLocation(uint32 ObjectIndex, const FVector& WorldLocation)
 {
 	checkSlow(ValidInfoIndexes.GetBit(ObjectIndex));
-	StoredObjectInfo[ObjectIndex].WorldLocation = WorldLocation;
+	FVector& TargetWorldLocation = StoredObjectInfo[ObjectIndex].WorldLocation;
+	TargetWorldLocation = WorldLocation;
+	const bool bHasInfoChanged = ObjectsWithDirtyInfo.GetBit(ObjectIndex) || TargetWorldLocation != WorldLocation;
+	
+	ObjectsWithDirtyInfo.SetBitValue(ObjectIndex, bHasInfoChanged);
+}
+
+void FWorldLocations::ResetObjectsWithDirtyInfo()
+{
+	ObjectsWithDirtyInfo.Reset();
 }
 
 }
diff --git a/Engine/Source/Runtime/Experimental/Iris/Core/Public/Iris/ReplicationSystem/Filtering/NetObjectFilter.h b/Engine/Source/Runtime/Experimental/Iris/Core/Public/Iris/ReplicationSystem/Filtering/NetObjectFilter.h
--- a/Engine/Source/Runtime/Experimental/Iris/Core/Public/Iris/ReplicationSystem/Filtering/NetObjectFilter.h
+++ b/Engine/Source/Runtime/Experimental/Iris/Core/Public/Iris/ReplicationSystem/Filtering/NetObjectFilter.h
@@ -56,10 +56,7 @@ struct FNetObjectFilteringParams
 	/** FilteringInfos for all objects. Index using the set bit indices in FilteredObjects. */
 	const FNetObjectFilteringInfo* FilteringInfos = nullptr;
 
-	/** 
-	* State buffers for all objects. Index using the set bit indices in FilteredObjects. 
-	*/
-	
+	/** State buffers for all objects. Index using the set bit indices in FilteredObjects. */
 	uint8 *const* StateBuffers = nullptr;
 
 	/** ID of the connection that the filtering applies to. */
@@ -74,8 +71,16 @@ struct FNetObjectFilteringParams
  */
 struct FNetObjectPreFilteringParams
 {
+	FNetObjectPreFilteringParams(const UE::Net::FNetBitArrayView InFilteredObjects);
+
 	// The IDs of all valid connections.
 	UE::Net::FNetBitArrayView ValidConnections;
+
+	/** The indices of the objects that have this filter set. The indices of set bits correspond to the object indices. */
+	const UE::Net::FNetBitArrayView FilteredObjects;
+
+	/** FilteringInfos for all objects. Index using the set bit indices in FilteredObjects. */
+	TArrayView<const FNetObjectFilteringInfo> FilteringInfos;
 };
 
 /**
@@ -113,6 +118,13 @@ enum class ENetFilterType : uint8
 	PostPoll_FragmentBased,
 };
 
+enum class ENetFilterTraits : uint8
+{
+	None = 0,
+	Spatial = 1,
+};
+ENUM_CLASS_FLAGS(ENetFilterTraits);
+
 /**
  * Base class for filter specific configuration.
  * @see FNetObjectFilterDefinition
@@ -232,6 +244,9 @@ public:
 	/** Returns what type of filter it is. Default is to cull and be executed before dirty state copying. */
 	ENetFilterType GetFilterType() const { return FilterType; }
 
+	/** Returns the filter's traits. */
+	ENetFilterTraits GetFilterTraits() const { return FilterTraits; }
+
 protected:
 	IRISCORE_API UNetObjectFilter();
 
@@ -247,6 +262,12 @@ protected:
 	/** Directly set when you want your dynamic filter to be executed. */
 	void SetupFilterType(ENetFilterType NewFilterType) { FilterType = NewFilterType; }
 
+	/** Adds traits. */
+	void AddFilterTraits(ENetFilterTraits Traits);
+
+	/** Sets the traits specified by TraitsMask to Traits. */
+	void SetFilterTraits(ENetFilterTraits Traits, ENetFilterTraits TraitsMask);
+
 private:
 	class FFilterInfo
 	{
@@ -266,6 +287,7 @@ private:
 	};
 
 	ENetFilterType FilterType = ENetFilterType::PrePoll_Raw;
+	ENetFilterTraits FilterTraits = ENetFilterTraits::None;
 	FFilterInfo FilterInfo;
 };
 
@@ -284,3 +306,13 @@ inline bool UNetObjectFilter::FFilterInfo::IsAddedToFilter(uint32 ObjectIndex) c
 	return ObjectIndex < FilteredObjects.GetNumBits() && FilteredObjects.IsBitSet(ObjectIndex);
 }
 
+inline void UNetObjectFilter::AddFilterTraits(ENetFilterTraits Traits)
+{
+	FilterTraits |= Traits;
+}
+
+inline void UNetObjectFilter::SetFilterTraits(ENetFilterTraits Traits, ENetFilterTraits TraitsMask)
+{
+	const ENetFilterTraits NewFilterTraits = (FilterTraits & ~TraitsMask) | (Traits & TraitsMask);
+	FilterTraits = NewFilterTraits;
+}
diff --git a/Engine/Source/Runtime/Experimental/Iris/Core/Public/Iris/ReplicationSystem/Filtering/NetObjectGridFilter.h b/Engine/Source/Runtime/Experimental/Iris/Core/Public/Iris/ReplicationSystem/Filtering/NetObjectGridFilter.h
--- a/Engine/Source/Runtime/Experimental/Iris/Core/Public/Iris/ReplicationSystem/Filtering/NetObjectGridFilter.h
+++ b/Engine/Source/Runtime/Experimental/Iris/Core/Public/Iris/ReplicationSystem/Filtering/NetObjectGridFilter.h
@@ -49,7 +49,7 @@ public:
 	FVector MaxPos = {+0.5f*2097152.0f, +0.5f*2097152.0f, +0.5f*2097152.0f};
 };
 
-UCLASS()
+UCLASS(abstract)
 class UNetObjectGridFilter : public UNetObjectFilter
 {
 	GENERATED_BODY()
@@ -61,7 +61,6 @@ protected:
 	IRISCORE_API virtual void RemoveConnection(uint32 ConnectionId) override;
 	IRISCORE_API virtual bool AddObject(uint32 ObjectIndex, FNetObjectFilterAddObjectParams&) override;
 	IRISCORE_API virtual void RemoveObject(uint32 ObjectIndex, const FNetObjectFilteringInfo&) override;
-	IRISCORE_API virtual void UpdateObjects(FNetObjectFilterUpdateParams&) override;
 	IRISCORE_API virtual void PreFilter(FNetObjectPreFilteringParams&) override;
 	IRISCORE_API virtual void Filter(FNetObjectFilteringParams&) override;
 
@@ -200,6 +199,8 @@ class UNetObjectGridWorldLocFilter : public UNetObjectGridFilter
 protected:
 
 	virtual void OnInit(FNetObjectFilterInitParams&) override;
+	virtual void UpdateObjects(FNetObjectFilterUpdateParams&) override;
+	virtual void PreFilter(FNetObjectPreFilteringParams&) override;
 	virtual void UpdateObjectInfo(FPerObjectInfo& PerObjectInfo, const UNetObjectGridFilter::FObjectLocationInfo& ObjectLocationInfo, const UE::Net::FReplicationInstanceProtocol* InstanceProtocol) override;
 	virtual bool BuildObjectInfo(uint32 ObjectIndex, FNetObjectFilterAddObjectParams& Params) override;
 
@@ -221,6 +222,7 @@ class UNetObjectGridFragmentLocFilter : public UNetObjectGridFilter
 protected:
 
 	virtual void OnInit(FNetObjectFilterInitParams&) override;
+	virtual void UpdateObjects(FNetObjectFilterUpdateParams&) override;
 	virtual void UpdateObjectInfo(FPerObjectInfo& PerObjectInfo, const UNetObjectGridFilter::FObjectLocationInfo& ObjectLocationInfo, const UE::Net::FReplicationInstanceProtocol* InstanceProtocol) override;
 	virtual bool BuildObjectInfo(uint32 ObjectIndex, FNetObjectFilterAddObjectParams& Params) override;
 	virtual void OnObjectRemoved(uint32 ObjectIndex) override;
diff --git a/Engine/Source/Runtime/Experimental/Iris/Core/Public/Iris/ReplicationSystem/ObjectReplicationBridge.h b/Engine/Source/Runtime/Experimental/Iris/Core/Public/Iris/ReplicationSystem/ObjectReplicationBridge.h
--- a/Engine/Source/Runtime/Experimental/Iris/Core/Public/Iris/ReplicationSystem/ObjectReplicationBridge.h
+++ b/Engine/Source/Runtime/Experimental/Iris/Core/Public/Iris/ReplicationSystem/ObjectReplicationBridge.h
@@ -52,6 +52,9 @@ public:
 		bool bNeedsPreUpdate = false;
 		bool bNeedsWorldLocationUpdate = false;
 
+		/** Whether the object is dormant or not */
+		bool bIsDormant = false;
+
 		/** When true we ask the class config if a dynamic filter was assigned to this class or one of it's parent inherited class. */
 		bool bUseClassConfigDynamicFilter = false;
 
@@ -343,6 +346,9 @@ private:
 	/** Returns true if instances of this class should be delta compressed */
 	bool ShouldClassBeDeltaCompressed(const UClass* Class);
 
+	/** Marks a spatially filtered object as requiring or not requiring frequent world location updates independent of it having dirty replicated properties */
+	void OptionallySetObjectRequiresFrequentWorldLocationUpdate(FNetRefHandle RefHandle, bool bDesiresFrequentWorldLocationUpdate);
+
 	/** Returns the TypeStatsIndex this class should use */
 	int32 GetTypeStatsIndex(const UClass* Class);
 
diff --git a/Engine/Source/Runtime/Experimental/Iris/Core/Public/Iris/ReplicationSystem/WorldLocations.h b/Engine/Source/Runtime/Experimental/Iris/Core/Public/Iris/ReplicationSystem/WorldLocations.h
--- a/Engine/Source/Runtime/Experimental/Iris/Core/Public/Iris/ReplicationSystem/WorldLocations.h
+++ b/Engine/Source/Runtime/Experimental/Iris/Core/Public/Iris/ReplicationSystem/WorldLocations.h
@@ -41,7 +41,6 @@ public:
 	};
 
 public:
-
 	// Internal API
 	void UpdateWorldLocation(uint32 ObjectIndex, const FVector& WorldLocation);
 
@@ -51,6 +50,20 @@ public:
 	void SetObjectInfo(uint32 ObjectIndex, const FObjectInfo& ObjectInfo);
     const FObjectInfo& GetObjectInfo(uint32 ObjectIndex) const;
 
+	/**
+	 * Objects are not necessarily marked as dirty just because they're moving, such as objects attached to other objects. If such objects are spatially filtered they need to update their world locations in order for replication to work as expected.
+	 * Use SetObjectRequiresFrequentWorldLocationUpdate to force frequent world location update on an object.
+	 */
+	void SetObjectRequiresFrequentWorldLocationUpdate(uint32 ObjectIndex, bool bRequiresFrequentUpdate);
+
+	/** Returns whether an object requires frequent world location updates. */
+	bool GetObjectRequiresFrequentWorldLocationUpdate(uint32 ObjectIndex) const;
+
+	FNetBitArrayView GetObjectsRequiringFrequentWorldLocationUpdate() const;
+
+	void ResetObjectsWithDirtyInfo();
+	FNetBitArrayView GetObjectsWithDirtyInfo() const;
+
 private:
 	enum : uint32
 	{
@@ -59,6 +72,10 @@ private:
 
 	/** Set bits indicate that we have stored information for this internal object index */
 	FNetBitArray ValidInfoIndexes;
+	/** Set bits indicate that the world location or net cull distance has changed since last update */
+	FNetBitArray ObjectsWithDirtyInfo;
+	/** Set bits indicate that the object requires frequent world location updates */
+	FNetBitArray ObjectsRequiringFrequentWorldLocationUpdate;
 
 	TChunkedArray<FObjectInfo, BytesPerLocationChunk> StoredObjectInfo;
 };
@@ -103,5 +120,24 @@ inline float FWorldLocations::GetCullDistance(uint32 ObjectIndex) const
 	}
 }
 
+inline void FWorldLocations::SetObjectRequiresFrequentWorldLocationUpdate(uint32 ObjectIndex, bool bRequiresFrequentUpdate)
+{
+	ObjectsRequiringFrequentWorldLocationUpdate.SetBitValue(ObjectIndex, ValidInfoIndexes.GetBit(ObjectIndex) && bRequiresFrequentUpdate);
+}
+
+inline bool FWorldLocations::GetObjectRequiresFrequentWorldLocationUpdate(uint32 ObjectIndex) const
+{
+	return ObjectsRequiringFrequentWorldLocationUpdate.GetBit(ObjectIndex);
+}
+
+inline FNetBitArrayView FWorldLocations::GetObjectsRequiringFrequentWorldLocationUpdate() const
+{
+	return MakeNetBitArrayView(ObjectsRequiringFrequentWorldLocationUpdate);
 }
 
+inline FNetBitArrayView FWorldLocations::GetObjectsWithDirtyInfo() const
+{
+	return MakeNetBitArrayView(ObjectsWithDirtyInfo);
+}
+
+}
