Fix C4702
#rnx

diff --git a/Engine/Source/Runtime/Experimental/Chaos/Private/Chaos/WeightedLatticeImplicitObject.cpp b/Engine/Source/Runtime/Experimental/Chaos/Private/Chaos/WeightedLatticeImplicitObject.cpp
--- a/Engine/Source/Runtime/Experimental/Chaos/Private/Chaos/WeightedLatticeImplicitObject.cpp
+++ b/Engine/Source/Runtime/Experimental/Chaos/Private/Chaos/WeightedLatticeImplicitObject.cpp
@@ -620,7 +620,6 @@ FImplicitObjectPtr TWeightedLatticeImplicitObject<TConcrete>::CopyGeometry() con
 	{
 		FImplicitObjectPtr CopiedShape = Object->CopyGeometry();
 		return new TWeightedLatticeImplicitObject<TConcrete>(reinterpret_cast<ObjectType&&>(CopiedShape), *this);
-		return FImplicitObjectPtr(CopyHelper(this));
 	}
 	else
 	{
