Fixing mismatched types. FReal was being used instead of FSolverReal when working with FSolverReal variables.

diff --git a/Engine/Source/Runtime/Experimental/Chaos/Public/Chaos/SoftsExternalForces.h b/Engine/Source/Runtime/Experimental/Chaos/Public/Chaos/SoftsExternalForces.h
--- a/Engine/Source/Runtime/Experimental/Chaos/Public/Chaos/SoftsExternalForces.h
+++ b/Engine/Source/Runtime/Experimental/Chaos/Public/Chaos/SoftsExternalForces.h
@@ -54,8 +54,8 @@ public:
 					if (Direction.Normalize())
 					{
 						// Scale by angle
-						const FReal DirectionDot = FVec3::DotProduct(Direction, N[Index]);
-						const FReal ScaleFactor = FMath::Min(1.f, FMath::Abs(DirectionDot) * LegacyWindAdaptation);
+						const FSolverReal DirectionDot = FSolverVec3::DotProduct(Direction, N[Index]);
+						const FSolverReal ScaleFactor = FMath::Min(1.f, FMath::Abs(DirectionDot) * LegacyWindAdaptation);
 						Acceleration[Index] += VelocityDelta * ScaleFactor;
 					}
 				}
