Geometry collection : fix anchored particle being set to kinematic wrongly causing them to not be released properly when under strain

diff --git a/Engine/Source/Runtime/Experimental/Chaos/Private/PhysicsProxy/GeometryCollectionPhysicsProxy.cpp b/Engine/Source/Runtime/Experimental/Chaos/Private/PhysicsProxy/GeometryCollectionPhysicsProxy.cpp
--- a/Engine/Source/Runtime/Experimental/Chaos/Private/PhysicsProxy/GeometryCollectionPhysicsProxy.cpp
+++ b/Engine/Source/Runtime/Experimental/Chaos/Private/PhysicsProxy/GeometryCollectionPhysicsProxy.cpp
@@ -1264,10 +1264,6 @@ void FGeometryCollectionPhysicsProxy::InitializeBodiesPT(Chaos::FPBDRigidsSolver
 				// Mass space -> Composed parent space -> world
 				const Chaos::FReal ScaledMass = AdjustMassForScale(Mass[TransformGroupIndex]);
 				const Chaos::FVec3f ScaledInertia = AdjustInertiaForScale((Chaos::FVec3f)InertiaTensor[TransformGroupIndex]);
-				const uint8 AdjustedDynamicState = 
-					bIsAnchored 
-					? static_cast<uint8>(Chaos::EObjectStateType::Kinematic)
-					: static_cast<uint8>(DynamicState[TransformGroupIndex]);
 
 				PopulateSimulatedParticle(
 					Handle,
@@ -1279,7 +1275,7 @@ void FGeometryCollectionPhysicsProxy::InitializeBodiesPT(Chaos::FPBDRigidsSolver
 					ScaledMass,
 					ScaledInertia,
 					Transforms[TransformGroupIndex],
-					AdjustedDynamicState,
+					DynamicState[TransformGroupIndex],
 					static_cast<int16>(Parameters.CollisionGroup),
 					CollisionParticlesPerObjectFraction);
 
