Iris
* Added ability to request to print additional information on the remote connection when a deserialization error causes a forced disconnection

#rnx

diff --git a/Engine/Source/Runtime/Engine/Private/DataChannel.cpp b/Engine/Source/Runtime/Engine/Private/DataChannel.cpp
--- a/Engine/Source/Runtime/Engine/Private/DataChannel.cpp
+++ b/Engine/Source/Runtime/Engine/Private/DataChannel.cpp
@@ -1646,6 +1646,7 @@ IMPLEMENT_CONTROL_CHANNEL_MESSAGE(BeaconJoin);
 IMPLEMENT_CONTROL_CHANNEL_MESSAGE(BeaconAssignGUID);
 IMPLEMENT_CONTROL_CHANNEL_MESSAGE(BeaconNetGUIDAck);
 IMPLEMENT_CONTROL_CHANNEL_MESSAGE(IrisProtocolMismatch);
+IMPLEMENT_CONTROL_CHANNEL_MESSAGE(IrisNetRefHandleError);
 
 void UControlChannel::Init( UNetConnection* InConnection, int32 InChannelIndex, EChannelCreateFlags CreateFlags )
 {
@@ -1863,6 +1864,20 @@ void UControlChannel::ReceivedBunch( FInBunch& Bunch )
 				{
 					IrisRepSystem->ReportProtocolMismatch(NetRefHandleId, Connection->GetConnectionId());
 				}
+#endif
+			}
+		}
+		else if (MessageType == NMT_IrisNetRefHandleError)
+		{
+			uint32 ErrorType = 0; //TBD
+			uint64 NetRefHandleId = 0;
+			if (FNetControlMessage<NMT_IrisNetRefHandleError>::Receive(Bunch, ErrorType, NetRefHandleId))
+			{
+#if UE_WITH_IRIS
+				if (UReplicationSystem* IrisRepSystem = Connection->Driver->GetReplicationSystem())
+				{
+					IrisRepSystem->ReportErrorWithNetRefHandle(ErrorType, NetRefHandleId, Connection->GetConnectionId());
+				}
 #endif
 			}
 		}
@@ -1946,6 +1961,9 @@ void UControlChannel::ReceivedBunch( FInBunch& Bunch )
 				case NMT_IrisProtocolMismatch:
 					FNetControlMessage<NMT_IrisProtocolMismatch>::Discard(Bunch);
 					break;
+				case NMT_IrisNetRefHandleError:
+					FNetControlMessage<NMT_IrisNetRefHandleError>::Discard(Bunch);
+					break;
 				default:
 					// if this fails, a case is missing above for an implemented message type
 					// or the connection is being sent potentially malformed packets
diff --git a/Engine/Source/Runtime/Engine/Private/Net/Experimental/Iris/DataStreamChannel.cpp b/Engine/Source/Runtime/Engine/Private/Net/Experimental/Iris/DataStreamChannel.cpp
--- a/Engine/Source/Runtime/Engine/Private/Net/Experimental/Iris/DataStreamChannel.cpp
+++ b/Engine/Source/Runtime/Engine/Private/Net/Experimental/Iris/DataStreamChannel.cpp
@@ -19,7 +19,11 @@
 #include "Iris/Core/IrisProfiler.h"
 #include "Iris/Core/IrisMemoryTracker.h"
 #include "Iris/Core/IrisLog.h"
+
 #include "Net/Core/Trace/NetTrace.h"
+#include "Net/Core/Connection/NetResult.h"
+#include "Net/DataChannel.h"
+
 #include "PacketHandler.h"
 #include "ProfilingDebugging/CsvProfiler.h"
 // IWYU pragma: end_keep
@@ -127,6 +131,15 @@ void UDataStreamChannel::ReceivedBunch(FInBunch& Bunch)
 	// If receiving was unsuccessful set bunch in error
 	if (SerializationContext.HasErrorOrOverflow())
 	{
+		if (SerializationContext.GetErrorHandleContext().IsValid())
+		{
+			TNetResult<ENetCloseResult> NetResult(ENetCloseResult::IrisNetRefHandleError, FString::Printf(TEXT("IrisNetRefHandleError=%s"), *SerializationContext.GetErrorHandleContext().ToString()));
+			AddToChainResultPtr(Bunch.ExtendedError, MoveTemp(NetResult));
+
+			uint32 ErrorType = 0; //TBD
+			uint64 RawHandleId = SerializationContext.GetErrorHandleContext().GetId();
+			FNetControlMessage<NMT_IrisNetRefHandleError>::Send(Connection, ErrorType, RawHandleId);
+		}
 		Bunch.SetError();
 	}
 
diff --git a/Engine/Source/Runtime/Engine/Private/Net/Iris/ReplicationSystem/ActorReplicationBridge.cpp b/Engine/Source/Runtime/Engine/Private/Net/Iris/ReplicationSystem/ActorReplicationBridge.cpp
--- a/Engine/Source/Runtime/Engine/Private/Net/Iris/ReplicationSystem/ActorReplicationBridge.cpp
+++ b/Engine/Source/Runtime/Engine/Private/Net/Iris/ReplicationSystem/ActorReplicationBridge.cpp
@@ -1201,6 +1201,21 @@ void UActorReplicationBridge::OnProtocolMismatchReported(FNetRefHandle RefHandle
 	}
 }
 
+void UActorReplicationBridge::ReportErrorWithNetRefHandle(uint32 ErrorType, FNetRefHandle RefHandle, uint32 ConnectionId)
+{
+	if (NetDriver)
+	{
+		if (UNetConnection* ClientConnection = NetDriver->GetConnectionById(ConnectionId))
+		{
+			uint64 RawHandleId = RefHandle.GetId();
+			FNetControlMessage<NMT_IrisNetRefHandleError>::Send(ClientConnection, ErrorType, RawHandleId);
+		}
+		else
+		{
+			UE_LOG(LogIrisBridge, Error, TEXT("UActorReplicationBridge::ReportErrorWithNetRefHandle could not find Connection for id:%u"), ConnectionId);
+		}
+	}
+}
 #else //!UE_WITH_IRIS
 
 UActorReplicationBridge::UActorReplicationBridge() = default;
diff --git a/Engine/Source/Runtime/Engine/Private/NetConnection.cpp b/Engine/Source/Runtime/Engine/Private/NetConnection.cpp
--- a/Engine/Source/Runtime/Engine/Private/NetConnection.cpp
+++ b/Engine/Source/Runtime/Engine/Private/NetConnection.cpp
@@ -3725,7 +3725,7 @@ void UNetConnection::DispatchPacket( FBitReader& Reader, int32 PacketId, bool& b
 					// If Iris reports errors they are unrecoverable.
 					if (Driver->GetReplicationSystem() != nullptr)
 					{
-						ensureAlwaysMsgf(false, TEXT("Received corrupted packet data with SequenceId: %d from server."), InPacketId);
+						ensureMsgf(false, TEXT("Received corrupted packet data. Iris cannot recover from this."));
 						Close(AddToAndConsumeChainResultPtr(Bunch.ExtendedError, ENetCloseResult::CorruptData));
 						return;
 					}
diff --git a/Engine/Source/Runtime/Engine/Public/Net/DataChannel.h b/Engine/Source/Runtime/Engine/Public/Net/DataChannel.h
--- a/Engine/Source/Runtime/Engine/Public/Net/DataChannel.h
+++ b/Engine/Source/Runtime/Engine/Public/Net/DataChannel.h
@@ -194,4 +194,6 @@ DEFINE_CONTROL_CHANNEL_MESSAGE(BeaconWelcome, 25); // server tells client they'r
 DEFINE_CONTROL_CHANNEL_MESSAGE(BeaconJoin, 26, FString, FUniqueNetIdRepl);  // server tries to create beacon type requested by client, sends NetGUID for actor sync
 DEFINE_CONTROL_CHANNEL_MESSAGE(BeaconAssignGUID, 27, FNetworkGUID); // client assigns NetGUID from server to beacon actor, sends NetGUIDAck
 DEFINE_CONTROL_CHANNEL_MESSAGE(BeaconNetGUIDAck, 28, FString); // server received NetGUIDAck from client, connection established successfully
+
 DEFINE_CONTROL_CHANNEL_MESSAGE(IrisProtocolMismatch, 29, uint64); // client has a different protocol hash from the server
+DEFINE_CONTROL_CHANNEL_MESSAGE(IrisNetRefHandleError, 30, uint32, uint64); // a specific handle caused an error and we want the remote connection to log all information it has on it
diff --git a/Engine/Source/Runtime/Engine/Public/Net/Iris/ReplicationSystem/ActorReplicationBridge.h b/Engine/Source/Runtime/Engine/Public/Net/Iris/ReplicationSystem/ActorReplicationBridge.h
--- a/Engine/Source/Runtime/Engine/Public/Net/Iris/ReplicationSystem/ActorReplicationBridge.h
+++ b/Engine/Source/Runtime/Engine/Public/Net/Iris/ReplicationSystem/ActorReplicationBridge.h
@@ -77,9 +77,12 @@ public:
 
 	/** Stop replicating an ActorComponent and its associated SubObjects. */
 	ENGINE_API void EndReplicationForActorComponent(UActorComponent* ActorComponent, EEndReplicationFlags EndReplicationFlags = EEndReplicationFlags::None);
-	
+
 	/** Get object reference packagemap. Used in special cases where serialization hasn't been converted to use NetSerializers.  */
 	UIrisObjectReferencePackageMap* GetObjectReferencePackageMap() const { return ObjectReferencePackageMap; }
+
+	/** Tell the remote connection that we detected a reading error with a specific replicated object */
+	ENGINE_API virtual void ReportErrorWithNetRefHandle(uint32 ErrorType, FNetRefHandle RefHandle, uint32 ConnectionId) override;
 	
 	using UObjectReplicationBridge::EndReplication;
 
diff --git a/Engine/Source/Runtime/Experimental/Iris/Core/Private/Iris/ReplicationSystem/ObjectReplicationBridge.cpp b/Engine/Source/Runtime/Experimental/Iris/Core/Private/Iris/ReplicationSystem/ObjectReplicationBridge.cpp
--- a/Engine/Source/Runtime/Experimental/Iris/Core/Private/Iris/ReplicationSystem/ObjectReplicationBridge.cpp
+++ b/Engine/Source/Runtime/Experimental/Iris/Core/Private/Iris/ReplicationSystem/ObjectReplicationBridge.cpp
@@ -1919,8 +1919,6 @@ void UObjectReplicationBridge::OnProtocolMismatchReported(FNetRefHandle RefHandl
 
 	UE_LOG(LogIris, Error, TEXT("OnProtocolMismatchReported from client:%u when instancing %s. CDO:%s ReplicatedObject:%s"), ConnectionId, *RefHandle.ToString(), *GetNameSafe(ObjArchetype), *GetNameSafe(ObjInstance));
 
-	//$IRIS TODO: Tell ActorBridge so he can choose to disconnect the client if the actor was critical.
-
 	if (UE_LOG_ACTIVE(LogIris, Error))
 	{
 		const FNetRefHandleManager::FReplicatedObjectData& ObjectData = NetRefHandleManager->GetReplicatedObjectDataNoCheck(ObjectInternalIndex);
@@ -1954,3 +1952,44 @@ void UObjectReplicationBridge::OnProtocolMismatchReported(FNetRefHandle RefHandl
 		UE::Net::Private::ObjectBridgeDebugging::RemoteProtocolMismatchDetected(ReplicationSystem, 0 /*TODO: Local ConnectionId*/, Fragments, ObjArchetype, ObjInstance);
 	}
 }
+
+void UObjectReplicationBridge::OnErrorWithNetRefHandleReported(uint32 ErrorType, FNetRefHandle RefHandle, uint32 ConnectionId)
+{
+	using namespace UE::Net;
+	using namespace UE::Net::Private;
+
+	// Ensure at the end so the log contains all the relevant information
+	ON_SCOPE_EXIT
+	{
+		ensureMsgf(false, TEXT("NetRefHandle error(%u) reported. Look at the log for important information on the object tied to the handle."), ErrorType);
+	};
+
+	const FInternalNetRefIndex ObjectInternalIndex = NetRefHandleManager->GetInternalIndex(RefHandle);
+	if (ObjectInternalIndex == FNetRefHandleManager::InvalidInternalIndex)
+	{
+		UE_LOG(LogIris, Warning, TEXT("OnErrorWithNetRefHandleReported(%u) from Connection:%u for %s but object has no InternalIndex."), ErrorType, ConnectionId, *RefHandle.ToString());
+		return;
+	}
+
+	UObject* ObjInstance = NetRefHandleManager->GetReplicatedObjectInstance(ObjectInternalIndex);
+	const FNetRefHandleManager::FReplicatedObjectData& ObjData = NetRefHandleManager->GetReplicatedObjectData(ObjectInternalIndex);
+
+	if (ObjData.IsSubObject())
+	{
+		UObject* RootObjInstance = NetRefHandleManager->GetReplicatedObjectInstance(ObjData.SubObjectRootIndex);
+		const FNetRefHandle RootObjNetHandle = NetRefHandleManager->GetNetRefHandleFromInternalIndex(ObjData.SubObjectRootIndex);
+		
+		UE_LOG(LogIris, Error, TEXT("OnErrorWithNetRefHandleReported(%u) from client:%u. %s maps to SubObject: %s owned by RootObject: %s using %s"), 
+			ErrorType, ConnectionId, 
+			*RefHandle.ToString(), *GetNameSafe(ObjInstance),
+			*GetNameSafe(RootObjInstance), *RootObjNetHandle.ToString()
+		);
+	}
+	else
+	{
+		UE_LOG(LogIris, Error, TEXT("OnErrorWithNetRefHandleReported(%u) from client:%u. %s maps to RootObject: %s"), 
+			ErrorType, ConnectionId, 
+			*RefHandle.ToString(), *GetNameSafe(ObjInstance)
+		);
+	}
+}
diff --git a/Engine/Source/Runtime/Experimental/Iris/Core/Private/Iris/ReplicationSystem/ReplicationReader.cpp b/Engine/Source/Runtime/Experimental/Iris/Core/Private/Iris/ReplicationSystem/ReplicationReader.cpp
--- a/Engine/Source/Runtime/Experimental/Iris/Core/Private/Iris/ReplicationSystem/ReplicationReader.cpp
+++ b/Engine/Source/Runtime/Experimental/Iris/Core/Private/Iris/ReplicationSystem/ReplicationReader.cpp
@@ -573,7 +573,7 @@ uint32 FReplicationReader::ReadObjectsInBatch(FNetSerializationContext& Context,
 		++ReadObjectCount;
 	}
 
-	ensureAlways(Reader.GetPosBits() <= BatchEndBitPosition);
+	ensure(Reader.GetPosBits() <= BatchEndBitPosition);
 
 	// ReadSubObjects 
 	while (Reader.GetPosBits() < BatchEndBitPosition)
@@ -726,6 +726,9 @@ uint32 FReplicationReader::ReadObjectBatch(FNetSerializationContext& Context)
 		{
 			if (Context.GetError() == GNetError_BrokenNetHandle)
 			{
+				const uint32 ErrorType = 1; //TBD
+				ReplicationBridge->ReportErrorWithNetRefHandle(ErrorType, IncompleteHandle, Parameters.ConnectionId);
+ 
 				// $TODO: Report this to the server so it knows that the state of data in the batch is unknown
 
 				// Log error and try to recover, if get more incoming data for an object in the broken state we will skip it.
@@ -756,12 +759,7 @@ void FReplicationReader::ReadObjectInBatch(FNetSerializationContext& Context, FN
 {
 	FNetBitStreamReader& Reader = *Context.GetBitStreamReader();
 
-	// If we are reading 
-	FNetRefHandle IncompleteHandle = BatchHandle;
-	if (bIsSubObject)
-	{
-		IncompleteHandle = ReadNetRefHandleId(Context, Reader);
-	}
+	const FNetRefHandle IncompleteHandle = !bIsSubObject ? BatchHandle : ReadNetRefHandleId(Context, Reader);
 	
 	// Read replicated destroy header if necessary. We don't know the internal index yet so can't do the more appropriate check IsObjectIndexForOOBAttachment.
 	const bool bReadReplicatedDestroyHeader = IncompleteHandle.IsValid();
@@ -1017,6 +1015,7 @@ void FReplicationReader::ReadObjectInBatch(FNetSerializationContext& Context, FN
 ErrorHandling:
 	if (bHasErrors)
 	{
+		Context.SetErrorHandleContext(IncompleteHandle);
 		UE_LOG_REPLICATIONREADER_ERROR(TEXT("FReplicationReader::ReadObject Failed to read replicated object with %s. Error '%s'."), *IncompleteHandle.ToString(), (Context.HasError() ? ToCStr(Context.GetError().ToString()) : TEXT("BitStream Overflow")));
 	}
 }
@@ -1769,7 +1768,8 @@ void FReplicationReader::ReadObjects(FNetSerializationContext& Context, uint32 O
 		--ObjectBatchCountToRead;
 	}
 
-	ensureAlwaysMsgf(!Context.HasErrorOrOverflow(), TEXT("Overflow: %c Error: %s Bit stream bits left: %u position: %u"), TEXT("YN")[Context.HasError()], ToCStr(Context.GetError().ToString()), Reader.GetBitsLeft(), Reader.GetPosBits());
+	UE_CLOG(Context.HasErrorOrOverflow(), LogIris, Error, TEXT("Overflow: %c Error: %s Bit stream bits left: %u position: %u"), TEXT("YN")[Context.HasError()], ToCStr(Context.GetError().ToString()), Reader.GetBitsLeft(), Reader.GetPosBits())
+	ensure(!Context.HasErrorOrOverflow());
 }
 
 void FReplicationReader::ProcessHugeObjectAttachment(FNetSerializationContext& Context, const TRefCountPtr<FNetBlob>& Attachment)
@@ -1860,7 +1860,9 @@ bool FReplicationReader::EnqueueEndReplication(FPendingBatchData* PendingBatchDa
 
 	if (Writer.IsOverflown())
 	{
-		ensureAlwaysMsgf(false, TEXT("Failed to EnqueueEndReplication for %s, Should never occur unless size of NetRefHandle has been increased."), *NetRefHandleToEndReplication.ToString());
+		UE_LOG(LogIris, Error, TEXT("Failed to EnqueueEndReplication for %s, Should never occur unless size of NetRefHandle has been increased."), *NetRefHandleToEndReplication.ToString());
+		ensure(false);
+
 		return false;
 	}
 
diff --git a/Engine/Source/Runtime/Experimental/Iris/Core/Private/Iris/ReplicationSystem/ReplicationSystem.cpp b/Engine/Source/Runtime/Experimental/Iris/Core/Private/Iris/ReplicationSystem/ReplicationSystem.cpp
--- a/Engine/Source/Runtime/Experimental/Iris/Core/Private/Iris/ReplicationSystem/ReplicationSystem.cpp
+++ b/Engine/Source/Runtime/Experimental/Iris/Core/Private/Iris/ReplicationSystem/ReplicationSystem.cpp
@@ -951,6 +951,8 @@ bool UReplicationSystem::SetRPCSendPolicyFlags(const UFunction* Function, UE::Ne
 		ensureAlwaysMsgf(false, TEXT("ENetObjectAttachmentSendPolicyFlags::SendImmediate is not allowed to use on Reliable RPC: %s"), *GetNameSafe(Function));
 		return false;
 	}
+
+	UE_LOG(LogIris, Verbose, TEXT("SetRPCSendPolicyFlags %s::%s => %s "), *GetNameSafe(Function->GetOuterUClass()), *GetNameSafe(Function), LexToString(SendFlags));
 		
 	Impl->AttachmentSendPolicyFlags.Add(FObjectKey(Function), SendFlags);
 	return true;
@@ -1627,11 +1629,19 @@ float UReplicationSystem::GetCullDistanceSqrOverride(FNetRefHandle Handle, float
 void UReplicationSystem::ReportProtocolMismatch(uint64 NetRefHandleId, uint32 ConnectionId)
 {
 	using namespace UE::Net::Private;
-	const FNetRefHandle NetRefHandle = FNetRefHandleManager::MakeNetRefHandle(NetRefHandleId, GetId()/*DON'T SUBMIT REVIEW NOTE: This ok or should it be incomplete?*/ );
+	const FNetRefHandle NetRefHandle = FNetRefHandleManager::MakeNetRefHandle(NetRefHandleId, GetId());
 
 	Impl->ReplicationSystemInternal.GetReplicationBridge()->OnProtocolMismatchReported(NetRefHandle, ConnectionId);
 }
 
+void UReplicationSystem::ReportErrorWithNetRefHandle(uint32 ErrorType, uint64 NetRefHandleId, uint32 ConnectionId)
+{
+	using namespace UE::Net::Private;
+	const FNetRefHandle NetRefHandle = FNetRefHandleManager::MakeNetRefHandle(NetRefHandleId, GetId());
+
+	Impl->ReplicationSystemInternal.GetReplicationBridge()->OnErrorWithNetRefHandleReported(ErrorType, NetRefHandle, ConnectionId);
+}
+
 
 #pragma region ReplicationSystemFactory
 
diff --git a/Engine/Source/Runtime/Experimental/Iris/Core/Public/Iris/ReplicationSystem/ObjectReplicationBridge.h b/Engine/Source/Runtime/Experimental/Iris/Core/Public/Iris/ReplicationSystem/ObjectReplicationBridge.h
--- a/Engine/Source/Runtime/Experimental/Iris/Core/Public/Iris/ReplicationSystem/ObjectReplicationBridge.h
+++ b/Engine/Source/Runtime/Experimental/Iris/Core/Public/Iris/ReplicationSystem/ObjectReplicationBridge.h
@@ -204,7 +204,8 @@ protected:
 	IRISCORE_API virtual void DetachInstanceFromRemote(FNetRefHandle Handle, EReplicationBridgeDestroyInstanceReason DestroyReason, EReplicationBridgeDestroyInstanceFlags DestroyFlags) override;
 	IRISCORE_API virtual void DetachInstance(FNetRefHandle Handle) override;
 	IRISCORE_API virtual void OnProtocolMismatchReported(FNetRefHandle RefHandle, uint32 ConnectionId) override;
-
+	IRISCORE_API virtual void OnErrorWithNetRefHandleReported(uint32 ErrorType, FNetRefHandle RefHandle, uint32 ConnectionId) override;
+	
 protected:
 	/**
 	* $IRIS TODO:
diff --git a/Engine/Source/Runtime/Experimental/Iris/Core/Public/Iris/ReplicationSystem/ReplicationBridge.h b/Engine/Source/Runtime/Experimental/Iris/Core/Public/Iris/ReplicationSystem/ReplicationBridge.h
--- a/Engine/Source/Runtime/Experimental/Iris/Core/Public/Iris/ReplicationSystem/ReplicationBridge.h
+++ b/Engine/Source/Runtime/Experimental/Iris/Core/Public/Iris/ReplicationSystem/ReplicationBridge.h
@@ -259,6 +259,12 @@ protected:
 	/** Called when a remote connection detected a protocol mismatch when trying to instantiate the NetRefHandle replicated object. */
 	IRISCORE_API virtual void OnProtocolMismatchReported(FNetRefHandle RefHandle, uint32 ConnectionId) {}
 
+	/** Called when a remote connection has a critical error caused by a specific NetRefHandle */
+	IRISCORE_API virtual void OnErrorWithNetRefHandleReported(uint32 ErrorType, FNetRefHandle RefHandle, uint32 ConnectionId) {}
+
+	/** Tell the remote connection that we detected a reading error with a specific replicated object */
+	IRISCORE_API virtual void ReportErrorWithNetRefHandle(uint32 ErrorType, FNetRefHandle RefHandle, uint32 ConnectionId) {}
+
 private:
 
 	// Internal operations invoked by ReplicationSystem/ReplicationWriter
diff --git a/Engine/Source/Runtime/Experimental/Iris/Core/Public/Iris/ReplicationSystem/ReplicationSystem.h b/Engine/Source/Runtime/Experimental/Iris/Core/Public/Iris/ReplicationSystem/ReplicationSystem.h
--- a/Engine/Source/Runtime/Experimental/Iris/Core/Public/Iris/ReplicationSystem/ReplicationSystem.h
+++ b/Engine/Source/Runtime/Experimental/Iris/Core/Public/Iris/ReplicationSystem/ReplicationSystem.h
@@ -587,6 +587,9 @@ public:
 	/** Called when a connection finds a protocol divergence when instantiating a replicated object. */
 	IRISCORE_API void ReportProtocolMismatch(uint64 NetRefHandleId, uint32 ConnectionId);
 
+	/** Called when a connection reports a critical error with a netrefhandle object */
+	IRISCORE_API void ReportErrorWithNetRefHandle(uint32 ErrorType, uint64 NetRefHandleId, uint32 ConnectionId);
+
 public:
 	// For internal use and not exported.
 	UE::Net::Private::FReplicationSystemInternal* GetReplicationSystemInternal();
diff --git a/Engine/Source/Runtime/Experimental/Iris/Core/Public/Iris/ReplicationSystem/ReplicationSystemTypes.h b/Engine/Source/Runtime/Experimental/Iris/Core/Public/Iris/ReplicationSystem/ReplicationSystemTypes.h
--- a/Engine/Source/Runtime/Experimental/Iris/Core/Public/Iris/ReplicationSystem/ReplicationSystemTypes.h
+++ b/Engine/Source/Runtime/Experimental/Iris/Core/Public/Iris/ReplicationSystem/ReplicationSystemTypes.h
@@ -58,4 +58,16 @@ enum class ENetObjectAttachmentSendPolicyFlags : uint32
 
 ENUM_CLASS_FLAGS(ENetObjectAttachmentSendPolicyFlags);
 
+inline const TCHAR* LexToString(ENetObjectAttachmentSendPolicyFlags SendFlags)
+{
+	switch (SendFlags)
+	{
+	case ENetObjectAttachmentSendPolicyFlags::None: return TEXT("None");
+	case ENetObjectAttachmentSendPolicyFlags::ScheduleAsOOB: return TEXT("ScheduleAsOOB");
+	case ENetObjectAttachmentSendPolicyFlags::SendInPostTickDispatch: return TEXT("SendInPostTickDispatch");
+	case ENetObjectAttachmentSendPolicyFlags::SendImmediate: return TEXT("SendImmediate");
+	default: ensure(false); return TEXT("Missing");
+	}
+}
+
 }
diff --git a/Engine/Source/Runtime/Experimental/Iris/Core/Public/Iris/Serialization/NetErrorContext.h b/Engine/Source/Runtime/Experimental/Iris/Core/Public/Iris/Serialization/NetErrorContext.h
--- a/Engine/Source/Runtime/Experimental/Iris/Core/Public/Iris/Serialization/NetErrorContext.h
+++ b/Engine/Source/Runtime/Experimental/Iris/Core/Public/Iris/Serialization/NetErrorContext.h
@@ -4,6 +4,8 @@
 
 #include "UObject/NameTypes.h"
 
+#include "Iris/ReplicationSystem/NetRefHandle.h"
+
 namespace UE::Net
 {
 
@@ -19,12 +21,21 @@ IRISCORE_API extern const FName GNetError_InternalError;
 class FNetErrorContext
 {
 public:
+
 	bool HasError() const;
+	
 	/** If an error has already been set calling this function again will be a no-op. */
 	IRISCORE_API void SetError(const FName Error);
+
 	FName GetError() const { return Error; }
 
+	void SetObjectHandle(const FNetRefHandle& InObjectHandle) { ObjectHandle = InObjectHandle; }
+
+	const FNetRefHandle& GetObjectHandle() const { return ObjectHandle; }
+
 private:
+
+	FNetRefHandle ObjectHandle;
 	FName Error;
 };
 
diff --git a/Engine/Source/Runtime/Experimental/Iris/Core/Public/Iris/Serialization/NetSerializationContext.h b/Engine/Source/Runtime/Experimental/Iris/Core/Public/Iris/Serialization/NetSerializationContext.h
--- a/Engine/Source/Runtime/Experimental/Iris/Core/Public/Iris/Serialization/NetSerializationContext.h
+++ b/Engine/Source/Runtime/Experimental/Iris/Core/Public/Iris/Serialization/NetSerializationContext.h
@@ -12,6 +12,7 @@ namespace UE::Net
 	class FNetBitArrayView;
 	class FNetBitStreamReader;
 	class FNetBitStreamWriter;
+
 	namespace Private
 	{
 		class FInternalNetSerializationContext;
@@ -42,10 +43,15 @@ public:
 
 	bool HasError() const { return ErrorContext.HasError(); }
 	bool HasErrorOrOverflow() const;
+	
 	/** If an error has already been set calling this function again will be a no-op, if bDoOverFlow is true, the function will also mark the current bitstream as overflown */
 	void SetError(const FName Error, bool bDoOverFlow = true);
 	FName GetError() const { return ErrorContext.GetError(); }
 
+	/** Store extra information regarding the object that triggered an error. */
+	void SetErrorHandleContext(const FNetRefHandle& HandleContext);
+	const FNetRefHandle& GetErrorHandleContext() const { return ErrorContext.GetObjectHandle(); }
+
 	/** There are cases where an error is handled and reported where we want to stay calm, reset the error context and carry on */
 	void ResetErrorContext() { ErrorContext = FNetErrorContext(); }
 
@@ -90,16 +96,21 @@ private:
 	IRISCORE_API void SetBitStreamOverflow();
 
 	FNetErrorContext ErrorContext;
-	FNetBitStreamReader* BitStreamReader;
-	FNetBitStreamWriter* BitStreamWriter;
-	FNetTraceCollector* TraceCollector;
-	Private::FInternalNetSerializationContext* InternalContext;
-	Private::FNetExportContext* ExportContext;
-	Private::FNetStatsContext* NetStatsContext;
-	const FNetBitArrayView* ChangeMask;
-	INetBlobReceiver* NetBlobReceiver;
-	uint32 LocalConnectionId;
-	int32 PacketId;
+
+	FNetBitStreamReader* BitStreamReader = nullptr;
+	FNetBitStreamWriter* BitStreamWriter = nullptr;
+	FNetTraceCollector* TraceCollector = nullptr;
+
+	Private::FInternalNetSerializationContext* InternalContext = nullptr;
+	Private::FNetExportContext* ExportContext = nullptr;
+	Private::FNetStatsContext* NetStatsContext = nullptr;
+
+	const FNetBitArrayView* ChangeMask = nullptr;
+	INetBlobReceiver* NetBlobReceiver = nullptr ;
+
+	uint32 LocalConnectionId = 0;
+	int32 PacketId = -1;
+
 	/** Set when replicated objects send their very first state. */
 	uint32 bIsInitState : 1;
 	/** Set only when dealing with a default state. */
@@ -110,14 +121,6 @@ private:
 inline FNetSerializationContext::FNetSerializationContext(FNetBitStreamReader* InBitStreamReader, FNetBitStreamWriter* InBitStreamWriter)
 : BitStreamReader(InBitStreamReader)
 , BitStreamWriter(InBitStreamWriter)
-, TraceCollector(nullptr)
-, InternalContext(nullptr)
-, ExportContext(nullptr)
-, NetStatsContext(nullptr)
-, ChangeMask(nullptr)
-, NetBlobReceiver(nullptr)
-, LocalConnectionId(0)
-, PacketId(-1)
 , bIsInitState(0)
 , bIsInitializingDefaultState(0)
 {
@@ -167,6 +170,11 @@ inline void FNetSerializationContext::SetError(const FName Error, bool bDoOverFl
 	ErrorContext.SetError(Error);
 }
 
+inline void FNetSerializationContext::SetErrorHandleContext(const FNetRefHandle& HandleContext)
+{
+	ErrorContext.SetObjectHandle(HandleContext);
+}
+
 inline bool FNetSerializationContext::HasErrorOrOverflow() const
 {
 	return ErrorContext.HasError() || IsBitStreamOverflown();
diff --git a/Engine/Source/Runtime/Net/Core/Public/Net/Core/Connection/NetCloseResult.h b/Engine/Source/Runtime/Net/Core/Public/Net/Core/Connection/NetCloseResult.h
--- a/Engine/Source/Runtime/Net/Core/Public/Net/Core/Connection/NetCloseResult.h
+++ b/Engine/Source/Runtime/Net/Core/Public/Net/Core/Connection/NetCloseResult.h
@@ -377,6 +377,8 @@ enum class ENetCloseResult : uint32
 	/** Protocol mismatch in Iris preventing a critical object instantiation */
 	IrisProtocolMismatch,
 
+	/** When a specific handle caused a reading error */
+	IrisNetRefHandleError,
 
 	/** Fault Handler Error result types */
 
