Chaos - particle names and logging utils

- Add some useful logging string conversions.
- Make sure that all (or most anyway) particles have a debug name.
- Add a name to the evolution so that we can tell which is client and server in the debugger and for logging.

diff --git a/Engine/Source/Runtime/Engine/Classes/PhysicsEngine/PhysicsLogUtil.h b/Engine/Source/Runtime/Engine/Classes/PhysicsEngine/PhysicsLogUtil.h
new file mode 100644
--- /dev/null
+++ b/Engine/Source/Runtime/Engine/Classes/PhysicsEngine/PhysicsLogUtil.h
@@ -0,0 +1,17 @@
+// Copyright Epic Games, Inc. All Rights Reserved.
+#pragma once
+
+#include "CoreTypes.h"
+
+class AActor;
+
+class FString;
+class UActorComponent;
+
+namespace PhysicsLogUtil
+{
+	ENGINE_API FString MakeComponentNetIDString(const UActorComponent* Component);
+	ENGINE_API FString MakeComponentNameString(const UActorComponent* Component);
+	ENGINE_API FString MakeActorNameString(const UActorComponent* ActorComponent);
+	ENGINE_API FString MakeActorNameString(const AActor* Actor);
+}
diff --git a/Engine/Source/Runtime/Engine/Private/PhysicsEngine/BodyInstance.cpp b/Engine/Source/Runtime/Engine/Private/PhysicsEngine/BodyInstance.cpp
--- a/Engine/Source/Runtime/Engine/Private/PhysicsEngine/BodyInstance.cpp
+++ b/Engine/Source/Runtime/Engine/Private/PhysicsEngine/BodyInstance.cpp
@@ -1102,13 +1102,7 @@ TSharedPtr<TArray<ANSICHAR>> GetDebugDebugName(const UPrimitiveComponent* Primit
 
 	if (PrimitiveComp)
 	{
-#if WITH_EDITOR
-		if (PrimitiveComp->GetOwner())
-		{
-			DebugName += FString::Printf(TEXT("Actor: '%s' "), *PrimitiveComp->GetOwner()->GetActorLabel(false));
-		}
-#endif
-		DebugName += FString::Printf(TEXT("Component: '%s' "), *PrimitiveComp->GetPathName());
+		DebugName += FString::Printf(TEXT("%s %s "), *AActor::GetDebugName(PrimitiveComp->GetOwner()), *PrimitiveComp->GetName());
 	}
 
 	if (BodySetup->BoneName != NAME_None)
diff --git a/Engine/Source/Runtime/Engine/Private/PhysicsEngine/ClusterUnionComponent.cpp b/Engine/Source/Runtime/Engine/Private/PhysicsEngine/ClusterUnionComponent.cpp
--- a/Engine/Source/Runtime/Engine/Private/PhysicsEngine/ClusterUnionComponent.cpp
+++ b/Engine/Source/Runtime/Engine/Private/PhysicsEngine/ClusterUnionComponent.cpp
@@ -705,8 +705,13 @@ void UClusterUnionComponent::OnCreatePhysicsState()
 	InitData.bCheckConnectivity = bHasAuthority;
 	InitData.bGenerateConnectivityEdges = bHasAuthority;
 
+#if CHAOS_DEBUG_NAME
+	InitData.DebugName = MakeShared<FString>(FString::Printf(TEXT("%s %s"), *AActor::GetDebugName(GetOwner()), *GetName()));
+#endif
+
 	PhysicsProxy = new Chaos::FClusterUnionPhysicsProxy{ this, Parameters, InitData };
 	PhysicsProxy->Initialize_External();
+
 	if (FPhysScene_Chaos* Scene = GetChaosScene())
 	{
 		Scene->AddObject(this, PhysicsProxy);
diff --git a/Engine/Source/Runtime/Engine/Private/PhysicsEngine/PhysicsLogUtil.cpp b/Engine/Source/Runtime/Engine/Private/PhysicsEngine/PhysicsLogUtil.cpp
new file mode 100644
--- /dev/null
+++ b/Engine/Source/Runtime/Engine/Private/PhysicsEngine/PhysicsLogUtil.cpp
@@ -0,0 +1,79 @@
+// Copyright Epic Games, Inc. All Rights Reserved.
+#include "PhysicsEngine/PhysicsLogUtil.h"
+#include "Components/ActorComponent.h"
+#include "Containers/UnrealString.h"
+#include "Engine/NetDriver.h"
+#include "Engine/PackageMapClient.h"
+#include "Engine/World.h"
+#include "GameFramework/Actor.h"
+#include "Misc/NetworkGuid.h"
+
+namespace PhysicsLogUtil
+{
+	// A simplified net mode string to make visually paring logs easier
+	FString MakeCompactNetModeString(ENetMode NetMode)
+	{
+		switch (NetMode)
+		{
+			case NM_DedicatedServer:
+			case NM_ListenServer:
+				return TEXT("Server");
+			case NM_Client:
+				return TEXT("Client");
+		}
+		return ToString(NetMode);
+	}
+
+	FString MakeComponentNetIDString(const UActorComponent* Component)
+	{
+		if (Component != nullptr)
+		{
+			if (const AActor* Actor = Component->GetOwner())
+			{
+				if (const UNetDriver* NetDriver = Actor->GetNetDriver())
+				{
+					if (const TSharedPtr<FNetGUIDCache>& NetGUIDCache = NetDriver->GuidCache)
+					{
+						// Network data available - return net mode, net ID, actor and component name
+						const ENetMode NetMode = Actor->GetNetMode();
+						const FNetworkGUID NetGUID = NetGUIDCache->GetNetGUID(Component);
+						return FString::Printf(TEXT("%s %-6s"), *MakeCompactNetModeString(NetMode), *NetGUID.ToString());
+					}
+				}
+			}
+		}
+
+		return FString(TEXT("NULL"));
+	}
+
+	FString MakeActorNameString(const AActor* Actor)
+	{
+		return AActor::GetDebugName(Actor);
+	}
+
+	FString MakeActorNameString(const UActorComponent* ActorComponent)
+	{
+		if (ActorComponent != nullptr)
+		{
+			return AActor::GetDebugName(ActorComponent->GetOwner());
+		}
+
+		return FString(TEXT("NULL"));
+	}
+
+	FString MakeComponentNameString(const UActorComponent* Component)
+	{
+		if (Component != nullptr)
+		{
+			if (const AActor* Actor = Component->GetOwner())
+			{
+				return FString::Printf(TEXT("%s %s"), *AActor::GetDebugName(Actor), *Component->GetName());
+			}
+
+			// No actor - return component name
+			return Component->GetName();
+		}
+
+		return FString(TEXT("NULL"));
+	}
+}
diff --git a/Engine/Source/Runtime/Experimental/Chaos/Private/Chaos/ChaosDebugName.cpp b/Engine/Source/Runtime/Experimental/Chaos/Private/Chaos/ChaosDebugName.cpp
new file mode 100644
--- /dev/null
+++ b/Engine/Source/Runtime/Experimental/Chaos/Private/Chaos/ChaosDebugName.cpp
@@ -0,0 +1,35 @@
+// Copyright Epic Games, Inc. All Rights Reserved.
+#include "Chaos/ChaosDebugName.h"
+#include "CoreTypes.h"
+
+namespace Chaos
+{
+	FString FSharedDebugName::DefaultName = FString(TEXT("NoName"));
+
+#if CHAOS_DEBUG_NAME
+	FSharedDebugName::FSharedDebugName(const FString& S)
+		: Name(TSharedPtr<FString, ESPMode::ThreadSafe>(new FString(S)))
+	{
+	}
+
+	FSharedDebugName::FSharedDebugName(FString&& S)
+		: Name(TSharedPtr<FString, ESPMode::ThreadSafe>(new FString(MoveTemp(S))))
+	{
+	}
+
+	inline bool FSharedDebugName::IsValid() const
+	{ 
+		return Name.IsValid();
+	}
+
+	const FString& FSharedDebugName::Value() const
+	{
+		if (Name.IsValid())
+		{
+			return *Name.Get();
+		}
+		return DefaultName;
+	}
+
+#endif
+}
\ No newline at end of file
diff --git a/Engine/Source/Runtime/Experimental/Chaos/Private/Chaos/ChaosLogUtil.cpp b/Engine/Source/Runtime/Experimental/Chaos/Private/Chaos/ChaosLogUtil.cpp
new file mode 100644
--- /dev/null
+++ b/Engine/Source/Runtime/Experimental/Chaos/Private/Chaos/ChaosLogUtil.cpp
@@ -0,0 +1,40 @@
+// Copyright Epic Games, Inc. All Rights Reserved.
+#include "Chaos/ChaosLogUtil.h"
+#include "CoreTypes.h"
+
+namespace Chaos
+{
+	FString ToString(const EObjectStateType ObjectState)
+	{
+		static FString StateStrings[] =
+		{
+			TEXT("Uninitialized(0)"),
+			TEXT("Sleeping(1)"),
+			TEXT("Kinematic(2)"),
+			TEXT("Static(3)"),
+			TEXT("Dynamic(4)"),
+		};
+
+		const bool bIsValid = ((int32)ObjectState >= 0) && ((int32)ObjectState < (int32)EObjectStateType::Count);
+		if (ensure(bIsValid))
+		{
+			return StateStrings[(int32)ObjectState];
+		}
+
+		return FString::Printf(TEXT("Invalid(%d)"), (int32)ObjectState);
+	}
+
+	FString ToString(const TBitArray<>& BitArray)
+	{
+		FString S;
+		S.Reserve(BitArray.Num());
+
+		for (int32 BitIndex = 0; BitIndex < BitArray.Num(); ++BitIndex)
+		{
+			S.Append(BitArray[BitIndex] ? TEXT("1") : TEXT("0"));
+		}
+
+		return S;
+	}
+
+}
\ No newline at end of file
diff --git a/Engine/Source/Runtime/Experimental/Chaos/Private/Chaos/Framework/PhysicsSolverBase.cpp b/Engine/Source/Runtime/Experimental/Chaos/Private/Chaos/Framework/PhysicsSolverBase.cpp
--- a/Engine/Source/Runtime/Experimental/Chaos/Private/Chaos/Framework/PhysicsSolverBase.cpp
+++ b/Engine/Source/Runtime/Experimental/Chaos/Private/Chaos/Framework/PhysicsSolverBase.cpp
@@ -214,6 +214,9 @@ namespace Chaos
 	FPhysicsSolverBase::FPhysicsSolverBase(const EMultiBufferMode BufferingModeIn,const EThreadingModeTemp InThreadingMode,UObject* InOwner, Chaos::FReal InAsyncDt)
 		: BufferMode(BufferingModeIn)
 		, ThreadingMode(!!GSingleThreadedPhysics ? EThreadingModeTemp::SingleThread : InThreadingMode)
+#if CHAOS_DEBUG_NAME
+		, DebugName(NAME_None)
+#endif
 		, PullResultsManager(MakeUnique<FChaosResultsManager>(MarshallingManager))
 		, PendingSpatialOperations_External(MakeUnique<FPendingSpatialDataQueue>())
 		, bUseCollisionResimCache(false)
@@ -243,6 +246,23 @@ namespace Chaos
 		}
 	}
 
+#if CHAOS_DEBUG_NAME
+	void FPhysicsSolverBase::SetDebugName(const FName& Name)
+	{
+		DebugName = Name;
+		OnDebugNameChanged();
+	}
+#endif
+
+	FName FPhysicsSolverBase::GetDebugName() const
+	{
+#if CHAOS_DEBUG_NAME
+		return DebugName;
+#else
+		return NAME_None;
+#endif
+	}
+
 	void FPhysicsSolverBase::EnableAsyncMode(FReal FixedDt)
 	{
 		AsyncDt = FixedDt;
diff --git a/Engine/Source/Runtime/Experimental/Chaos/Private/Chaos/GeometryParticles.cpp b/Engine/Source/Runtime/Experimental/Chaos/Private/Chaos/GeometryParticles.cpp
--- a/Engine/Source/Runtime/Experimental/Chaos/Private/Chaos/GeometryParticles.cpp
+++ b/Engine/Source/Runtime/Experimental/Chaos/Private/Chaos/GeometryParticles.cpp
@@ -17,11 +17,7 @@ namespace Chaos
 	namespace Private
 	{
 		// The name shown for particoles that have not had their DebugName set
-#if CHAOS_DEBUG_NAME
-		FString EmptyParticleName = TEXT("<NotNamed>");
-#else
-		FString EmptyParticleName = TEXT("<NotSupported>");
-#endif
+		FString EmptyParticleName = TEXT("<NoName>");
 	}
 
 	FShapeOrShapesArray::FShapeOrShapesArray(const FGeometryParticleHandle* Particle)
diff --git a/Engine/Source/Runtime/Experimental/Chaos/Private/PBDRigidsSolver.cpp b/Engine/Source/Runtime/Experimental/Chaos/Private/PBDRigidsSolver.cpp
--- a/Engine/Source/Runtime/Experimental/Chaos/Private/PBDRigidsSolver.cpp
+++ b/Engine/Source/Runtime/Experimental/Chaos/Private/PBDRigidsSolver.cpp
@@ -618,6 +618,13 @@ namespace Chaos
 			});
 	}
 
+#if CHAOS_DEBUG_NAME
+	void FPBDRigidsSolver::OnDebugNameChanged()
+	{
+		MEvolution->SetName(GetDebugName().ToString());
+	}
+#endif
+
 	FRealSingle MaxBoundsForTree = (FRealSingle)10000;
 	FAutoConsoleVariableRef CVarMaxBoundsForTree(
 		TEXT("p.MaxBoundsForTree"),
diff --git a/Engine/Source/Runtime/Experimental/Chaos/Private/PhysicsProxy/ClusterUnionPhysicsProxy.cpp b/Engine/Source/Runtime/Experimental/Chaos/Private/PhysicsProxy/ClusterUnionPhysicsProxy.cpp
--- a/Engine/Source/Runtime/Experimental/Chaos/Private/PhysicsProxy/ClusterUnionPhysicsProxy.cpp
+++ b/Engine/Source/Runtime/Experimental/Chaos/Private/PhysicsProxy/ClusterUnionPhysicsProxy.cpp
@@ -81,6 +81,10 @@ namespace Chaos
 		Particle_External = FExternalParticle::CreateParticle();
 		check(Particle_External != nullptr);
 
+#if CHAOS_DEBUG_NAME
+		Particle_External->SetDebugName(InitData.DebugName);
+#endif
+
 		Particle_External->SetProxy(this);
 		Particle_External->SetUserData(InitData.UserData);
 		Particle_External->SetX(InitData.InitialTransform.GetTranslation());
@@ -135,8 +139,11 @@ namespace Chaos
 			Particle_Internal->SetCenterOfMass(FVector3f::ZeroVector);
 			Particle_Internal->SetRotationOfMass(FQuat::Identity);
 
-#if !(UE_BUILD_SHIPPING || UE_BUILD_TEST)
-			Particle_Internal->SetDebugName(MakeShared<FString, ESPMode::ThreadSafe>(FString::Printf(TEXT("%s"), *GetOwner()->GetName())));
+#if CHAOS_DEBUG_NAME
+			if (Particle_External.IsValid())
+			{
+				Particle_Internal->SetDebugName(Particle_External->DebugName());
+			}
 #endif
 
 			// For explicit cluster unions (i.e. cluster unions created via a cluster union component rather than the cluster group index), the cluster union itself is responsible for initializing
diff --git a/Engine/Source/Runtime/Experimental/Chaos/Private/PhysicsProxy/GeometryCollectionPhysicsProxy.cpp b/Engine/Source/Runtime/Experimental/Chaos/Private/PhysicsProxy/GeometryCollectionPhysicsProxy.cpp
--- a/Engine/Source/Runtime/Experimental/Chaos/Private/PhysicsProxy/GeometryCollectionPhysicsProxy.cpp
+++ b/Engine/Source/Runtime/Experimental/Chaos/Private/PhysicsProxy/GeometryCollectionPhysicsProxy.cpp
@@ -728,6 +728,10 @@ void FGeometryCollectionPhysicsProxy::CreateGTParticles(const TBitArray<>& Effec
 				GTParticlesToTransformGroupIndex.Add(P, Index);
 				GTParticles[Index]->SetUniqueIdx(UniqueIdxs[Index]);
 
+#if CHAOS_DEBUG_NAME
+				P->SetDebugName(MakeShared<FString, ESPMode::ThreadSafe>(FString::Printf(TEXT("%s-%d"), *Parameters.Name, Index)));
+#endif
+
 				const float ScaledMass = AdjustMassForScale(Mass[Index]);
 
 				// Note that this transform must match the physics thread transform computation for initialization.
@@ -790,10 +794,6 @@ void FGeometryCollectionPhysicsProxy::CreateGTParticles(const TBitArray<>& Effec
 
 				const bool bIsOneWayInteraction = (Parameters.OneWayInteractionLevel >= 0) && (Level[Index] >= Parameters.OneWayInteractionLevel);
 				P->SetOneWayInteraction(bIsOneWayInteraction);
-
-#if !(UE_BUILD_SHIPPING || UE_BUILD_TEST)
-				P->SetDebugName(MakeShared<FString, ESPMode::ThreadSafe>(FString::Printf(TEXT("%s-%d"), *Parameters.Name, Index)));
-#endif
 			}
 		}
 		// this step is necessary for Phase 2 where we need to walk back the hierarchy from children to parent 
@@ -1488,6 +1488,11 @@ void FGeometryCollectionPhysicsProxy::InitializeBodiesPT(Chaos::FPBDRigidsSolver
 						Handle->GTGeometryParticle() = GTParticle;
 					}
 
+#if CHAOS_DEBUG_NAME
+					TSharedPtr<FString, ESPMode::ThreadSafe> ParticleName = MakeShared<FString, ESPMode::ThreadSafe>(FString::Printf(TEXT("%s-%d"), *Parameters.Name, TransformGroupIndex));
+					Handle->SetDebugName(ParticleName);
+#endif
+
 					int32 RigidChildrenIdx = 0;
 					for(const int32 ChildTransformIndex : RigidChildrenTransformGroupIndex)
 					{
@@ -1635,10 +1640,6 @@ void FGeometryCollectionPhysicsProxy::InitializeBodiesPT(Chaos::FPBDRigidsSolver
 				Handle->SetLinearEtherDrag(Parameters.LinearDamping);
 				Handle->SetAngularEtherDrag(Parameters.AngularDamping);
 				Handle->SetInitialOverlapDepenetrationVelocity(Parameters.InitialOverlapDepenetrationVelocity);
-
-#if !(UE_BUILD_SHIPPING || UE_BUILD_TEST)
-				Handle->SetDebugName(MakeShared<FString, ESPMode::ThreadSafe>(FString::Printf(TEXT("%s-%d"), *Parameters.Name, HandleIndex)));
-#endif
 			}
 
 #if !(UE_BUILD_SHIPPING || UE_BUILD_TEST)
diff --git a/Engine/Source/Runtime/Experimental/Chaos/Public/Chaos/ChaosDebugName.h b/Engine/Source/Runtime/Experimental/Chaos/Public/Chaos/ChaosDebugName.h
new file mode 100644
--- /dev/null
+++ b/Engine/Source/Runtime/Experimental/Chaos/Public/Chaos/ChaosDebugName.h
@@ -0,0 +1,43 @@
+// Copyright Epic Games, Inc. All Rights Reserved.
+#pragma once
+
+#include "Containers/UnrealString.h"
+#include "Templates/SharedPointer.h"
+
+#ifndef CHAOS_DEBUG_NAME
+#define CHAOS_DEBUG_NAME 0
+#endif
+
+namespace Chaos
+{
+	/**
+	* A wrapper around shared pointer to a string that compiles away in shipping buiilds,
+	* but Get() always returns a valid string reference to make writing logs a little easier
+	*/
+	CHAOS_API class FSharedDebugName
+	{
+	public:
+		FSharedDebugName() = default;
+		FSharedDebugName(const FSharedDebugName& Other) = default;
+		FSharedDebugName(FSharedDebugName&& Other) = default;
+
+		FSharedDebugName(const FString& S);
+		FSharedDebugName(FString&& S);
+
+		bool IsValid() const;
+		const FString& Value() const;
+
+	private:
+#if CHAOS_DEBUG_NAME
+		TSharedPtr<FString, ESPMode::ThreadSafe> Name;
+#endif
+		static FString DefaultName;
+	};
+
+#if !CHAOS_DEBUG_NAME
+	inline FSharedDebugName::FSharedDebugName(const FString& S) {}
+	inline FSharedDebugName::FSharedDebugName(FString&& S) {}
+	inline bool FSharedDebugName::IsValid() const { return false; }
+	inline const FString& FSharedDebugName::Value() const { return DefaultName; }
+#endif
+}
\ No newline at end of file
diff --git a/Engine/Source/Runtime/Experimental/Chaos/Public/Chaos/ChaosLogUtil.h b/Engine/Source/Runtime/Experimental/Chaos/Public/Chaos/ChaosLogUtil.h
new file mode 100644
--- /dev/null
+++ b/Engine/Source/Runtime/Experimental/Chaos/Public/Chaos/ChaosLogUtil.h
@@ -0,0 +1,12 @@
+// Copyright Epic Games, Inc. All Rights Reserved.
+#pragma once
+
+#include "Chaos/Particle/ObjectState.h"
+
+class FString;
+
+namespace Chaos
+{
+	CHAOS_API FString ToString(const EObjectStateType ObjectState);
+	CHAOS_API FString ToString(const TBitArray<>& BitArray);
+}
\ No newline at end of file
diff --git a/Engine/Source/Runtime/Experimental/Chaos/Public/Chaos/Framework/PhysicsSolverBase.h b/Engine/Source/Runtime/Experimental/Chaos/Public/Chaos/Framework/PhysicsSolverBase.h
--- a/Engine/Source/Runtime/Experimental/Chaos/Public/Chaos/Framework/PhysicsSolverBase.h
+++ b/Engine/Source/Runtime/Experimental/Chaos/Public/Chaos/Framework/PhysicsSolverBase.h
@@ -508,16 +508,10 @@ namespace Chaos
 		CHAOS_API FGraphEventRef AdvanceAndDispatch_External(FReal InDt);
 
 #if CHAOS_DEBUG_NAME
-		void SetDebugName(const FName& Name)
-		{
-			DebugName = Name;
-		}
-
-		const FName& GetDebugName() const
-		{
-			return DebugName;
-		}
+		CHAOS_API void SetDebugName(const FName& Name);
 #endif
+		CHAOS_API FName GetDebugName() const;
+
 
 		//Tells us if we're on the frozen game thread. This is needed for knowing which data to read/write to
 		//The IsInGameThread check is so that other threads (e.g audio thread) which might be running queries in parallel will continue to use the correct interpolated GT data
@@ -745,35 +739,37 @@ namespace Chaos
 		virtual void SetExternalTimestampConsumed_Internal(const int32 Timestamp) = 0;
 
 #if CHAOS_DEBUG_NAME
+		virtual void OnDebugNameChanged() {}
+
 		FName DebugName;
 #endif
 
-	FChaosMarshallingManager MarshallingManager;
-	TUniquePtr<FChaosResultsManager> PullResultsManager;	//must come after MarshallingManager since it knows about MarshallingManager
+		FChaosMarshallingManager MarshallingManager;
+		TUniquePtr<FChaosResultsManager> PullResultsManager;	//must come after MarshallingManager since it knows about MarshallingManager
 
-	// The spatial operations not yet consumed by the internal sim. Use this to ensure any GT operations are seen immediately
-	TUniquePtr<FPendingSpatialDataQueue> PendingSpatialOperations_External;
+		// The spatial operations not yet consumed by the internal sim. Use this to ensure any GT operations are seen immediately
+		TUniquePtr<FPendingSpatialDataQueue> PendingSpatialOperations_External;
 
-	TArray<ISimCallbackObject*> SimCallbackObjects;
-	TArray<ISimCallbackObject*> MidPhaseModifiers;
-	TArray<ISimCallbackObject*> CCDModifiers;
-	TArray<ISimCallbackObject*> StrainModifiers;
-	TArray<ISimCallbackObject*> ContactModifiers;
-	TArray<ISimCallbackObject*> RegistrationWatchers;
-	TArray<ISimCallbackObject*> UnregistrationWatchers;
-	TArray<ISimCallbackObject*> PhysicsObjectUnregistrationWatchers;
+		TArray<ISimCallbackObject*> SimCallbackObjects;
+		TArray<ISimCallbackObject*> MidPhaseModifiers;
+		TArray<ISimCallbackObject*> CCDModifiers;
+		TArray<ISimCallbackObject*> StrainModifiers;
+		TArray<ISimCallbackObject*> ContactModifiers;
+		TArray<ISimCallbackObject*> RegistrationWatchers;
+		TArray<ISimCallbackObject*> UnregistrationWatchers;
+		TArray<ISimCallbackObject*> PhysicsObjectUnregistrationWatchers;
 
-	TUniquePtr<FRewindData> MRewindData;
-	TUniquePtr<IRewindCallback> MRewindCallback;
+		TUniquePtr<FRewindData> MRewindData;
+		TUniquePtr<IRewindCallback> MRewindCallback;
 
-	bool bUseCollisionResimCache;
+		bool bUseCollisionResimCache;
 
-	FGraphEventRef PendingTasks;
+		FGraphEventRef PendingTasks;
 
-	bool bSolverHasFrozenGameThreadCallbacks = false;
-	bool bGameThreadFrozen = false;
-	FReal MLastDt = FReal(0);
-	FReal MTime = FReal(0);
+		bool bSolverHasFrozenGameThreadCallbacks = false;
+		bool bGameThreadFrozen = false;
+		FReal MLastDt = FReal(0);
+		FReal MTime = FReal(0);
 
 	private:
 
diff --git a/Engine/Source/Runtime/Experimental/Chaos/Public/Chaos/PBDRigidsEvolution.h b/Engine/Source/Runtime/Experimental/Chaos/Public/Chaos/PBDRigidsEvolution.h
--- a/Engine/Source/Runtime/Experimental/Chaos/Public/Chaos/PBDRigidsEvolution.h
+++ b/Engine/Source/Runtime/Experimental/Chaos/Public/Chaos/PBDRigidsEvolution.h
@@ -949,6 +949,9 @@ public:
 	const TArray<FRemovalData>& GetAllRemovals() { return MAllRemovals; }
 	void ResetAllRemovals() { MAllRemovals.Reset(); }
 
+	void SetName(const FString& InName) { EvolutionName = InName; }
+	const FString& GetName() const { return EvolutionName; }
+
 protected:
 	int32 NumConstraints() const
 	{
@@ -1142,6 +1145,9 @@ protected:
 	bool bIsResim = false; 
 	bool bIsReset = false;
 
+	// Useful name for debugging. E.g., Indicates whether we are on client or server
+	FString EvolutionName;
+
 #if CHAOS_EVOLUTION_COLLISION_TESTMODE
 	// Test Mode for Collision issues (resets particle positions every tick for repeatable testing)
 	CHAOS_API void TestModeParticleDisabled(FGeometryParticleHandle* Particle);
diff --git a/Engine/Source/Runtime/Experimental/Chaos/Public/Chaos/Particle/ObjectState.h b/Engine/Source/Runtime/Experimental/Chaos/Public/Chaos/Particle/ObjectState.h
new file mode 100644
--- /dev/null
+++ b/Engine/Source/Runtime/Experimental/Chaos/Public/Chaos/Particle/ObjectState.h
@@ -0,0 +1,19 @@
+// Copyright Epic Games, Inc. All Rights Reserved.
+#pragma once
+
+#include "Chaos/Core.h"
+
+
+namespace Chaos
+{
+	enum class EObjectStateType : int8
+	{
+		Uninitialized = 0,
+		Sleeping = 1,
+		Kinematic = 2,
+		Static = 3,
+		Dynamic = 4,
+
+		Count
+	};
+}
\ No newline at end of file
diff --git a/Engine/Source/Runtime/Experimental/Chaos/Public/Chaos/ParticleHandle.h b/Engine/Source/Runtime/Experimental/Chaos/Public/Chaos/ParticleHandle.h
--- a/Engine/Source/Runtime/Experimental/Chaos/Public/Chaos/ParticleHandle.h
+++ b/Engine/Source/Runtime/Experimental/Chaos/Public/Chaos/ParticleHandle.h
@@ -30,6 +30,8 @@ namespace Chaos
 	namespace Private
 	{
 		class FConvexOptimizer;
+
+		CHAOS_API extern FString EmptyParticleName;
 	}
 
 	namespace CVars
@@ -537,15 +539,15 @@ public:
 	
 	void SetNonFrequentData(const FParticleNonFrequentData& InData)
 	{
+#if CHAOS_DEBUG_NAME
+		SetDebugName(InData.DebugName());
+#endif
+
 		SetGeometry(Chaos::FImplicitObjectPtr(InData.GetGeometry()));
 		SetUniqueIdx(InData.UniqueIdx());
 		SetSpatialIdx(InData.SpatialIdx());
 		SetResimType(InData.ResimType());
 		SetParticleID(InData.ParticleID());
-
-#if CHAOS_DEBUG_NAME
-		SetDebugName(InData.DebugName());
-#endif
 	}
 
 	bool HasCollision() const { return GeometryParticles->HasCollision(ParticleIdx); }
@@ -2678,6 +2680,18 @@ public:
 	}
 #endif
 
+	const FString& GetDebugName() const
+	{
+#if CHAOS_DEBUG_NAME
+		const TSharedPtr<FString, ESPMode::ThreadSafe>& DebugNamePtr = DebugName();
+		if (DebugNamePtr.IsValid())
+		{
+			return *(DebugNamePtr.Get());
+		}
+#endif
+		return Private::EmptyParticleName;
+	}
+
 	void MergeShapesArray(FShapesArray&& InShapesArray)
 	{
 		MergeShapeInstances(reinterpret_cast<FShapeInstanceProxyArray&&>(InShapesArray));
diff --git a/Engine/Source/Runtime/Experimental/Chaos/Public/Chaos/Particles.h b/Engine/Source/Runtime/Experimental/Chaos/Public/Chaos/Particles.h
--- a/Engine/Source/Runtime/Experimental/Chaos/Public/Chaos/Particles.h
+++ b/Engine/Source/Runtime/Experimental/Chaos/Public/Chaos/Particles.h
@@ -4,6 +4,7 @@
 #include "Chaos/ArrayCollection.h"
 #include "Chaos/ArrayCollectionArray.h"
 #include "Chaos/Core.h"
+#include "Chaos/Particle/ObjectState.h"
 #include "Chaos/Vector.h"
 #include "ChaosArchive.h"
 #include "HAL/LowLevelMemTracker.h"
@@ -214,17 +215,6 @@ namespace Chaos
 		return InParticles.GetTypeHash();
 	}
 	
-	enum class EObjectStateType: int8
-	{
-		Uninitialized = 0,
-		Sleeping = 1,
-		Kinematic = 2,
-		Static = 3,
-		Dynamic = 4,
-
-		Count
-	};
-
 	enum class EChaosCollisionTraceFlag : int8
 	{
 		/** Use project physics settings (DefaultShapeComplexity) */
diff --git a/Engine/Source/Runtime/Experimental/Chaos/Public/PBDRigidsSolver.h b/Engine/Source/Runtime/Experimental/Chaos/Public/PBDRigidsSolver.h
--- a/Engine/Source/Runtime/Experimental/Chaos/Public/PBDRigidsSolver.h
+++ b/Engine/Source/Runtime/Experimental/Chaos/Public/PBDRigidsSolver.h
@@ -340,6 +340,12 @@ namespace Chaos
 		// Apply callbacks internally 
 		CHAOS_API virtual void ApplyCallbacks_Internal() override;
 
+	protected:
+
+#if CHAOS_DEBUG_NAME
+		virtual void OnDebugNameChanged() override final;
+#endif
+
 	private:
 
 		/**/
diff --git a/Engine/Source/Runtime/Experimental/Chaos/Public/PhysicsProxy/ClusterUnionPhysicsProxy.h b/Engine/Source/Runtime/Experimental/Chaos/Public/PhysicsProxy/ClusterUnionPhysicsProxy.h
--- a/Engine/Source/Runtime/Experimental/Chaos/Public/PhysicsProxy/ClusterUnionPhysicsProxy.h
+++ b/Engine/Source/Runtime/Experimental/Chaos/Public/PhysicsProxy/ClusterUnionPhysicsProxy.h
@@ -36,6 +36,9 @@ namespace Chaos
 		bool bUnbreakable = false;
 		bool bGenerateConnectivityEdges = true;
 		int32 GravityGroupOverride = INDEX_NONE;
+#if CHAOS_DEBUG_NAME
+		TSharedPtr<FString, ESPMode::ThreadSafe> DebugName;
+#endif
 	};
 
 	/**
