Chaos - fix mispaced CHAOS_API in class

diff --git a/Engine/Source/Runtime/Experimental/Chaos/Public/Chaos/ChaosDebugName.h b/Engine/Source/Runtime/Experimental/Chaos/Public/Chaos/ChaosDebugName.h
--- a/Engine/Source/Runtime/Experimental/Chaos/Public/Chaos/ChaosDebugName.h
+++ b/Engine/Source/Runtime/Experimental/Chaos/Public/Chaos/ChaosDebugName.h
@@ -14,7 +14,7 @@ namespace Chaos
 	* A wrapper around shared pointer to a string that compiles away in shipping buiilds,
 	* but Get() always returns a valid string reference to make writing logs a little easier
 	*/
-	CHAOS_API class FSharedDebugName
+	class CHAOS_API FSharedDebugName
 	{
 	public:
 		FSharedDebugName() = default;
