Dataflow : Bounds Facade
- Updated to allow for bounds around collections without a transform group.

diff --git a/Engine/Plugins/Experimental/Dataflow/Source/DataflowEnginePlugin/Private/Dataflow/DataflowComponent.cpp b/Engine/Plugins/Experimental/Dataflow/Source/DataflowEnginePlugin/Private/Dataflow/DataflowComponent.cpp
--- a/Engine/Plugins/Experimental/Dataflow/Source/DataflowEnginePlugin/Private/Dataflow/DataflowComponent.cpp
+++ b/Engine/Plugins/Experimental/Dataflow/Source/DataflowEnginePlugin/Private/Dataflow/DataflowComponent.cpp
@@ -73,7 +73,9 @@ void UDataflowComponent::TickComponent(float DeltaTime, enum ELevelTick TickType
 			}
 		}
 
+		bBoundsNeedsUpdate = true;
 		UpdateLocalBounds();
+
 		bUpdateRender = false;
 		bNeedsSceneProxyUpdate = true;
 	}
@@ -95,7 +97,12 @@ void UDataflowComponent::UpdateLocalBounds()
 {
 	if (bBoundsNeedsUpdate)
 	{
-		GeometryCollection::Facades::FBoundsFacade(RenderCollection).UpdateBoundingBox();
+		GeometryCollection::Facades::FBoundsFacade BoundsFacade(RenderCollection);
+		BoundsFacade.DefineSchema();
+
+		BoundsFacade.UpdateBoundingBox();
+		BoundingBox = BoundsFacade.GetBoundingBoxInCollectionSpace();
+
 		bBoundsNeedsUpdate = false;
 	}
 }
diff --git a/Engine/Source/Runtime/Experimental/Chaos/Private/GeometryCollection/Facades/CollectionBoundsFacade.cpp b/Engine/Source/Runtime/Experimental/Chaos/Private/GeometryCollection/Facades/CollectionBoundsFacade.cpp
--- a/Engine/Source/Runtime/Experimental/Chaos/Private/GeometryCollection/Facades/CollectionBoundsFacade.cpp
+++ b/Engine/Source/Runtime/Experimental/Chaos/Private/GeometryCollection/Facades/CollectionBoundsFacade.cpp
@@ -19,7 +19,8 @@ namespace GeometryCollection::Facades
 		, VertexAttribute(InCollection,"Vertex", FGeometryCollection::VerticesGroup, FGeometryCollection::VerticesGroup)
 		, BoneMapAttribute(InCollection,"BoneMap", FGeometryCollection::VerticesGroup)
 		, TransformToGeometryIndexAttribute(InCollection,"TransformToGeometryIndex", FTransformCollection::TransformGroup)
-		, ParentAttribute(InCollection, "Parent", FTransformCollection::TransformGroup)
+		, VertexStartAttribute(InCollection, "VertexStart", FGeometryCollection::GeometryGroup)
+		, VertexCountAttribute(InCollection, "VertexCount", FGeometryCollection::GeometryGroup)
 	{}
 
 	FBoundsFacade::FBoundsFacade(const FManagedArrayCollection& InCollection)
@@ -29,7 +30,8 @@ namespace GeometryCollection::Facades
 		, VertexAttribute(InCollection, "Vertex", FGeometryCollection::VerticesGroup, FGeometryCollection::VerticesGroup)
 		, BoneMapAttribute(InCollection, "BoneMap", FGeometryCollection::VerticesGroup)
 		, TransformToGeometryIndexAttribute(InCollection, "TransformToGeometryIndex", FTransformCollection::TransformGroup)
-		, ParentAttribute(InCollection, "Parent", FTransformCollection::TransformGroup)
+		, VertexStartAttribute(InCollection, "VertexStart", FGeometryCollection::GeometryGroup)
+		, VertexCountAttribute(InCollection, "VertexCount", FGeometryCollection::GeometryGroup)
 	{}
 
 	//
@@ -40,42 +42,58 @@ namespace GeometryCollection::Facades
 	{
 		check(!IsConst());
 		BoundingBoxAttribute.Add();
-		VertexAttribute.Add();
-		BoneMapAttribute.Add();
-		TransformToGeometryIndexAttribute.Add();
-		ParentAttribute.Add();
 	}
 
 	bool FBoundsFacade::IsValid() const
 	{
-		return BoundingBoxAttribute.IsValid() && VertexAttribute.IsValid() 
-			&& BoneMapAttribute.IsValid() && TransformToGeometryIndexAttribute.IsValid()
-			&& ParentAttribute.IsValid();
+		return BoundingBoxAttribute.IsValid();
 	}
 
 
-	void FBoundsFacade::UpdateBoundingBox(bool bSkipCheck)
+	void FBoundsFacade::UpdateBoundingBox()
 	{
-		check(!IsConst());
-
-		if (!bSkipCheck || !IsValid())
+		if (BoundingBoxAttribute.IsValid())
 		{
-			return;
+			bool bTranformBased = VertexAttribute.IsValid() && TransformToGeometryIndexAttribute.Num() && BoneMapAttribute.IsValid();
+
+			bool bVertexBased = VertexAttribute.IsValid() && VertexStartAttribute.Num() && VertexCountAttribute.IsValid();
+
+			if (bTranformBased)
+			{
+				UpdateTransformBasedBoundingBox();
+			}
+			else if (bVertexBased)
+			{
+				UpdateVertexBasedBoundingBox();
+			}
+			else
+			{
+				TManagedArray<FBox>& BoundingBox = BoundingBoxAttribute.Modify();
+				for (int32 bdx = 0; bdx < BoundingBox.Num(); bdx++)
+				{
+					BoundingBox[bdx].Init();
+				}
+			}
 		}
+	}
 
+	void FBoundsFacade::UpdateTransformBasedBoundingBox()
+	{
 		TManagedArray<FBox>& BoundingBox = BoundingBoxAttribute.Modify();
 		const TManagedArray<FVector3f>& Vertex = VertexAttribute.Get();
 		const TManagedArray<int32>& BoneMap = BoneMapAttribute.Get();
 		const TManagedArray<int32>& TransformToGeometryIndex = TransformToGeometryIndexAttribute.Get();
 
-		if (BoundingBox.Num())
+		for (int32 bdx = 0; bdx < BoundingBox.Num(); bdx++)
 		{
-			// Initialize BoundingBox
-			for (int32 Idx = 0; Idx < BoundingBox.Num(); ++Idx)
-			{
-				BoundingBox[Idx].Init();
-			}
+			BoundingBox[bdx].Init();
+		}
 
+		// Use the mapping stored from the vertices to the transforms to generate a bounding box
+		// relative to transform origin. 
+
+		if (BoundingBox.Num())
+		{
 			// Compute BoundingBox
 			for (int32 Idx = 0; Idx < Vertex.Num(); ++Idx)
 			{
@@ -85,6 +103,32 @@ namespace GeometryCollection::Facades
 		}
 	}
 
+	void FBoundsFacade::UpdateVertexBasedBoundingBox()
+	{
+		TManagedArray<FBox>& BoundingBox = BoundingBoxAttribute.Modify();
+		const TManagedArray<FVector3f>& Vertex = VertexAttribute.Get();
+		const TManagedArray<int32>& VertexStart = VertexStartAttribute.Get();
+		const TManagedArray<int32>& VertexCount = VertexCountAttribute.Get();
+
+		for (int32 bdx = 0; bdx < BoundingBox.Num(); bdx++)
+		{
+			BoundingBox[bdx].Init();
+		}
+
+		// Use the mapping stored from the geometry to the vertices to generate a bounding box.
+		// This configuration might not have an assiocated transform.
+
+		for (int32 Gdx = 0; Gdx < BoundingBox.Num(); Gdx++)
+		{
+			// Compute BoundingBox
+			int32 VertexEnd = VertexStart[Gdx] + VertexCount[Gdx];
+			for (int32 Vdx = VertexStart[Gdx]; Vdx < VertexEnd; ++Vdx)
+			{
+				BoundingBox[Gdx] += FVector(Vertex[Vdx]);
+			}
+		}
+	}
+
 	TArray<FVector> FBoundsFacade::GetCentroids() const
 	{
 		TArray<FVector> Centroids;
@@ -110,17 +154,28 @@ namespace GeometryCollection::Facades
 		if (IsValid())
 		{
 			GeometryCollection::Facades::FCollectionTransformFacade TransformFacade(ConstCollection);
-			TArray<FTransform> CollectionSpaceTransforms = TransformFacade.ComputeCollectionSpaceTransforms();
+			if (TransformFacade.IsValid())
+			{
+				TArray<FTransform> CollectionSpaceTransforms = TransformFacade.ComputeCollectionSpaceTransforms();
 
-			for (int32 TransformIndex = 0; TransformIndex < CollectionSpaceTransforms.Num(); ++TransformIndex)
+				for (int32 TransformIndex = 0; TransformIndex < CollectionSpaceTransforms.Num(); ++TransformIndex)
+				{
+					const int32 GeoIndex = TransformToGeometryIndexAttribute[TransformIndex];
+					if (BoundingBoxAttribute.IsValidIndex(GeoIndex))
+					{
+						const FBox& GeoBoundingBox = BoundingBoxAttribute[GeoIndex];
+						const FTransform CollectionSpaceTransform = CollectionSpaceTransforms[TransformIndex];
+						const FBox BoundingBoxInCollectionSpace = GeoBoundingBox.TransformBy(CollectionSpaceTransform);
+						BoundingBox += BoundingBoxInCollectionSpace;
+					}
+				}
+			}
+			else
 			{
-				const int32 GeoIndex = TransformToGeometryIndexAttribute[TransformIndex];
-				if (BoundingBoxAttribute.IsValidIndex(GeoIndex))
+				const TManagedArray<FBox>& BoundingBoxes = BoundingBoxAttribute.Get();
+				for (int32 Idx = 0; Idx < BoundingBoxes.Num(); ++Idx)
 				{
-					const FBox& GeoBoundingBox = BoundingBoxAttribute[GeoIndex];
-					const FTransform CollectionSpaceTransform = CollectionSpaceTransforms[TransformIndex];
-					const FBox BoundingBoxInCollectionSpace = GeoBoundingBox.TransformBy(CollectionSpaceTransform);
-					BoundingBox += BoundingBoxInCollectionSpace;
+					BoundingBox += BoundingBoxes[Idx];
 				}
 			}
 		}
diff --git a/Engine/Source/Runtime/Experimental/Chaos/Public/GeometryCollection/Facades/CollectionBoundsFacade.h b/Engine/Source/Runtime/Experimental/Chaos/Public/GeometryCollection/Facades/CollectionBoundsFacade.h
--- a/Engine/Source/Runtime/Experimental/Chaos/Public/GeometryCollection/Facades/CollectionBoundsFacade.h
+++ b/Engine/Source/Runtime/Experimental/Chaos/Public/GeometryCollection/Facades/CollectionBoundsFacade.h
@@ -39,7 +39,7 @@ namespace GeometryCollection::Facades
 		CHAOS_API bool IsValid() const;		
 
 		/** UpdateBoundingBox */
-		CHAOS_API void UpdateBoundingBox(bool bSkipCheck = false);
+		CHAOS_API void UpdateBoundingBox();
 
 		/** BoundingBox access */
 		const TManagedArray< FBox >& GetBoundingBoxes() const { return BoundingBoxAttribute.Get(); }
@@ -55,6 +55,13 @@ namespace GeometryCollection::Facades
 
 		/** TransformToGeometryIndex Access */
 		const TManagedArray< int32 >& GetTransformToGeometryIndex() const { return TransformToGeometryIndexAttribute.Get(); }
+	protected:
+
+		/** Transform based bounds evaluation, where the vertices are nested within a transform */
+		void UpdateTransformBasedBoundingBox();
+
+		/** Vertex based bounds evaluation, where the vertices are NOT nested within a transform */
+		void UpdateVertexBasedBoundingBox();
 
 	private:
 		const FManagedArrayCollection& ConstCollection;
@@ -64,7 +71,8 @@ namespace GeometryCollection::Facades
 		TManagedArrayAccessor<FVector3f>	VertexAttribute;
 		TManagedArrayAccessor<int32>		BoneMapAttribute;
 		TManagedArrayAccessor<int32>		TransformToGeometryIndexAttribute;
-		TManagedArrayAccessor<int32>		ParentAttribute;
+		TManagedArrayAccessor<int32>		VertexStartAttribute;
+		TManagedArrayAccessor<int32>		VertexCountAttribute;
 	};
 
 }
