Chaos Cloth: Add new Anisotropic Stretch Spring constraint type.

This constraint has XPBD edge and axial springs, both of which have spatially varying stiffnesses based on warp/weft/bias stiffness controls. The same controls are shared for edge and axial springs. They also allow for dynamic restlength modification via the WarpScale and WeftScale parameters. All and all, these have the same types of parameters as the UV-based XPBDStretchBiasElement constraints, but are simpler because they're just mass springs.

diff --git a/Engine/Plugins/Experimental/ChaosCloth/Source/Chaos/Private/ChaosCloth/ChaosClothConstraints.cpp b/Engine/Plugins/Experimental/ChaosCloth/Source/Chaos/Private/ChaosCloth/ChaosClothConstraints.cpp
--- a/Engine/Plugins/Experimental/ChaosCloth/Source/Chaos/Private/ChaosCloth/ChaosClothConstraints.cpp
+++ b/Engine/Plugins/Experimental/ChaosCloth/Source/Chaos/Private/ChaosCloth/ChaosClothConstraints.cpp
@@ -5,6 +5,7 @@
 #include "Chaos/PBDSpringConstraints.h"
 #include "Chaos/XPBDSpringConstraints.h"
 #include "Chaos/XPBDStretchBiasElementConstraints.h"
+#include "Chaos/XPBDAnisotropicSpringConstraints.h"
 #include "Chaos/PBDBendingConstraints.h"
 #include "Chaos/XPBDBendingConstraints.h"
 #include "Chaos/XPBDAnisotropicBendingConstraints.h"
@@ -100,6 +101,233 @@ bool bEnableGS = false;
 
 #endif
 
+class FClothConstraints::FRuleCreator
+{
+public:
+	explicit FRuleCreator(FClothConstraints* InConstraints)
+		: Constraints(InConstraints)
+	{
+		check(Constraints);
+		Constraints->Evolution->AllocatePreSubstepParallelInitRange(Constraints->ParticleRangeId, Constraints->NumPreSubstepInits);
+		Constraints->Evolution->AllocatePBDExternalForceRulesRange(Constraints->ParticleRangeId, Constraints->NumExternalForceRules);
+		Constraints->Evolution->AllocatePostInitialGuessParallelInitRange(Constraints->ParticleRangeId, Constraints->NumConstraintInits);
+		Constraints->Evolution->AllocatePreSubstepConstraintRulesRange(Constraints->ParticleRangeId, Constraints->NumPreSubstepConstraintRules);
+		Constraints->Evolution->AllocatePerIterationPBDConstraintRulesRange(Constraints->ParticleRangeId, Constraints->NumConstraintRules);
+		Constraints->Evolution->AllocatePerIterationCollisionPBDConstraintRulesRange(Constraints->ParticleRangeId, Constraints->NumCollisionConstraintRules);
+		Constraints->Evolution->AllocatePerIterationPostCollisionsPBDConstraintRulesRange(Constraints->ParticleRangeId, Constraints->NumPostCollisionConstraintRules);
+		Constraints->Evolution->AllocateUpdateLinearSystemRulesRange(Constraints->ParticleRangeId, Constraints->NumUpdateLinearSystemRules);
+		Constraints->Evolution->AllocateUpdateLinearSystemCollisionsRulesRange(Constraints->ParticleRangeId, Constraints->NumUpdateLinearSystemCollisionsRules);
+		Constraints->Evolution->AllocatePostSubstepConstraintRulesRange(Constraints->ParticleRangeId, Constraints->NumPostprocessingConstraintRules);
+
+		PreSubstepParallelInits = Constraints->Evolution->GetPreSubstepParallelInitRange(Constraints->ParticleRangeId);
+		ExternalForceRules = Constraints->Evolution->GetPBDExternalForceRulesRange(Constraints->ParticleRangeId);
+		PostInitialGuessParallelInits = Constraints->Evolution->GetPostInitialGuessParallelInitRange(Constraints->ParticleRangeId);
+		PreSubstepConstraintRules = Constraints->Evolution->GetPreSubstepConstraintRulesRange(Constraints->ParticleRangeId);
+		PerIterationConstraintRules = Constraints->Evolution->GetPerIterationPBDConstraintRulesRange(Constraints->ParticleRangeId);
+		PerIterationCollisionConstraintRules = Constraints->Evolution->GetPerIterationCollisionPBDConstraintRulesRange(Constraints->ParticleRangeId);
+		PerIterationPostCollisionsConstraintRules = Constraints->Evolution->GetPerIterationPostCollisionsPBDConstraintRulesRange(Constraints->ParticleRangeId);
+		UpdateLinearSystemRules = Constraints->Evolution->GetUpdateLinearSystemRulesRange(Constraints->ParticleRangeId);
+		UpdateLinearSystemCollisionsRules = Constraints->Evolution->GetUpdateLinearSystemCollisionsRulesRange(Constraints->ParticleRangeId);
+		PostSubstepConstraintRules = Constraints->Evolution->GetPostSubstepConstraintRulesRange(Constraints->ParticleRangeId);
+	}
+
+	~FRuleCreator()
+	{
+		if (Constraints)
+		{
+			check(PreSubstepInitsIndex == Constraints->NumPreSubstepInits);
+			check(ExternalForceRulesIndex == Constraints->NumExternalForceRules);
+			check(PostInitialGuessInitsIndex == Constraints->NumConstraintInits);
+			check(PreSubstepConstraintRulesIndex == Constraints->NumPreSubstepConstraintRules);
+			check(ConstraintRuleIndex == Constraints->NumConstraintRules);
+			check(CollisionConstraintRulesIndex == Constraints->NumCollisionConstraintRules);
+			check(PostCollisionConstraintRulesIndex == Constraints->NumPostCollisionConstraintRules);
+			check(UpdateLinearSystemRulesIndex == Constraints->NumUpdateLinearSystemRules);
+			check(UpdateLinearSystemCollisionsRulesIndex == Constraints->NumUpdateLinearSystemCollisionsRules);
+			check(PostSubstepConstraintRulesIndex == Constraints->NumPostprocessingConstraintRules);
+		}
+	}
+
+	void PreSubstepParallelInitRule(Softs::FEvolution::ParallelInitFunc&& Rule)
+	{
+		PreSubstepParallelInits[PreSubstepInitsIndex++] = MoveTemp(Rule);
+	}
+
+	void AddExternalForceRule(Softs::FEvolution::PBDConstraintRuleFunc&& Rule)
+	{
+		ExternalForceRules[ExternalForceRulesIndex++] = MoveTemp(Rule);
+	}
+
+	void AddPostInitialGuessParallelInitRule(Softs::FEvolution::ParallelInitFunc&& Rule)
+	{
+		PostInitialGuessParallelInits[PostInitialGuessInitsIndex++] = MoveTemp(Rule);
+	}
+
+	void AddPreSubstepConstraintRule(Softs::FEvolution::ConstraintRuleFunc&& Rule)
+	{
+		PreSubstepConstraintRules[PreSubstepConstraintRulesIndex++] = MoveTemp(Rule);
+	}
+
+	void AddPerIterationPBDConstraintRule(Softs::FEvolution::PBDConstraintRuleFunc&& Rule)
+	{
+		PerIterationConstraintRules[ConstraintRuleIndex++] = MoveTemp(Rule);
+	}
+
+	void AddPerIterationCollisionPBDConstraintRule(Softs::FEvolution::PBDCollisionConstraintRuleFunc&& Rule)
+	{
+		PerIterationCollisionConstraintRules[CollisionConstraintRulesIndex++] = MoveTemp(Rule);
+	}
+
+	void AddPerIterationPostCollisionsPBDConstraintRule(Softs::FEvolution::PBDConstraintRuleFunc&& Rule)
+	{
+		PerIterationPostCollisionsConstraintRules[PostCollisionConstraintRulesIndex++] = MoveTemp(Rule);
+	}
+
+	void AddUpdateLinearSystemRule(Softs::FEvolution::UpdateLinearSystemFunc&& Rule)
+	{
+		UpdateLinearSystemRules[UpdateLinearSystemRulesIndex++] = MoveTemp(Rule);
+	}
+
+	void AddUpdateLinearSystemCollisionsRule(Softs::FEvolution::UpdateLinearSystemCollisionsFunc&& Rule)
+	{
+		UpdateLinearSystemCollisionsRules[UpdateLinearSystemCollisionsRulesIndex++] = MoveTemp(Rule);
+	}
+
+	void AddPostSubstepConstraintRule(Softs::FEvolution::ConstraintRuleFunc&& Rule)
+	{
+		PostSubstepConstraintRules[PostSubstepConstraintRulesIndex++] = MoveTemp(Rule);
+	}
+
+	template<typename ConstraintType>
+	void AddExternalForceRule_Apply(ConstraintType* const Constraint)
+	{
+		AddExternalForceRule(
+			[Constraint](Softs::FSolverParticlesRange& Particles, const Softs::FSolverReal Dt)
+		{
+			Constraint->Apply(Particles, Dt);
+		});
+	}
+
+	template<typename ConstraintType>
+	void AddUpdateLinearSystemRule_UpdateLinearSystem(ConstraintType* const Constraint)
+	{
+		AddUpdateLinearSystemRule([Constraint](const Softs::FSolverParticlesRange& Particles, const FSolverReal Dt, Softs::FEvolutionLinearSystem& LinearSystem)
+		{
+			Constraint->UpdateLinearSystem(Particles, Dt, LinearSystem);
+		});
+	}
+
+	template<bool bInitParticles = false, typename ConstraintType>
+	void AddXPBDParallelInitRule(ConstraintType* const Constraint)
+	{
+		if constexpr (bInitParticles)
+		{
+			AddPostInitialGuessParallelInitRule(
+				[Constraint, Evolution = Constraints->Evolution](const Softs::FSolverParticlesRange& Particles, const Softs::FSolverReal Dt, const Softs::ESolverMode SolverMode)
+			{
+				Constraint->ApplyProperties(Dt, Evolution->GetIterations());
+				if ((SolverMode & Softs::ESolverMode::PBD) != Softs::ESolverMode::None)\
+				{
+					Constraint->Init(Particles);
+				}
+			});
+		}
+		else
+		{
+			AddPostInitialGuessParallelInitRule(
+				[Constraint, Evolution = Constraints->Evolution](const Softs::FSolverParticlesRange& Particles, const Softs::FSolverReal Dt, const Softs::ESolverMode SolverMode)
+			{
+				Constraint->ApplyProperties(Dt, Evolution->GetIterations());
+				if ((SolverMode & Softs::ESolverMode::PBD) != Softs::ESolverMode::None)\
+				{
+					Constraint->Init();
+				}
+			});
+		}
+	}
+
+	template<typename ConstraintType>
+	void AddPBDParallelInitRule(ConstraintType* const Constraint)
+	{
+		AddPostInitialGuessParallelInitRule(
+			[Constraint, Evolution = Constraints->Evolution](const Softs::FSolverParticlesRange& Particles, const Softs::FSolverReal Dt, const Softs::ESolverMode SolverMode)
+		{
+			if ((SolverMode & Softs::ESolverMode::PBD) != Softs::ESolverMode::None)
+			{
+				Constraint->ApplyProperties(Dt, Evolution->GetIterations()); 
+			}
+		});
+	}
+
+	template<bool bPostCollisions, typename ConstraintType>
+	void AddPerIterationPBDConstraintRule_Apply(ConstraintType* const Constraint)
+	{
+		if constexpr (bPostCollisions)
+		{
+			AddPerIterationPostCollisionsPBDConstraintRule(
+				[Constraint](Softs::FSolverParticlesRange& Particles, const Softs::FSolverReal Dt)
+			{
+				Constraint->Apply(Particles, Dt);
+			});
+		}
+		else
+		{
+			AddPerIterationPBDConstraintRule(
+				[Constraint](Softs::FSolverParticlesRange& Particles, const Softs::FSolverReal Dt)
+			{
+				Constraint->Apply(Particles, Dt);
+			});
+		}
+	}
+
+	template<typename ConstraintType>
+	void AddPerIterationCollisionPBDConstraintRule_Apply(ConstraintType* const Constraint)
+	{
+		AddPerIterationCollisionPBDConstraintRule(
+			[Constraint](Softs::FSolverParticlesRange& Particles, const Softs::FSolverReal Dt, const TArray<Softs::FSolverCollisionParticlesRange>& CollisionParticles)
+		{
+			Constraint->Apply(Particles, Dt, CollisionParticles);
+		}
+		);
+	}
+
+	template<typename ConstraintType>
+	void AddUpdateLinearSystemCollisionsRule_UpdateLinearSystem(ConstraintType* const Constraint)
+	{
+		AddUpdateLinearSystemCollisionsRule(
+			[Constraint](const Softs::FSolverParticlesRange& Particles, const Softs::FSolverReal Dt, const TArray<Softs::FSolverCollisionParticlesRange>& CollisionParticles, Softs::FEvolutionLinearSystem& LinearSystem)
+		{
+			Constraint->UpdateLinearSystem(Particles, Dt, CollisionParticles, LinearSystem);
+		}
+		);
+	}
+
+private:
+
+	FClothConstraints* const Constraints = nullptr;
+	TArrayView<Softs::FEvolution::ParallelInitFunc> PreSubstepParallelInits;
+	TArrayView<Softs::FEvolution::PBDConstraintRuleFunc> ExternalForceRules;
+	TArrayView<Softs::FEvolution::ParallelInitFunc> PostInitialGuessParallelInits;
+	TArrayView<Softs::FEvolution::ConstraintRuleFunc> PreSubstepConstraintRules;
+	TArrayView<Softs::FEvolution::PBDConstraintRuleFunc> PerIterationConstraintRules;
+	TArrayView<Softs::FEvolution::PBDCollisionConstraintRuleFunc> PerIterationCollisionConstraintRules;
+	TArrayView<Softs::FEvolution::PBDConstraintRuleFunc> PerIterationPostCollisionsConstraintRules;
+	TArrayView<Softs::FEvolution::UpdateLinearSystemFunc> UpdateLinearSystemRules;
+	TArrayView<Softs::FEvolution::UpdateLinearSystemCollisionsFunc> UpdateLinearSystemCollisionsRules;
+	TArrayView<Softs::FEvolution::ConstraintRuleFunc> PostSubstepConstraintRules;
+
+	int32 PreSubstepInitsIndex = 0;
+	int32 ExternalForceRulesIndex = 0;
+	int32 PostInitialGuessInitsIndex = 0;
+	int32 PreSubstepConstraintRulesIndex = 0;
+	int32 ConstraintRuleIndex = 0;
+	int32 CollisionConstraintRulesIndex = 0;
+	int32 PostCollisionConstraintRulesIndex = 0;
+	int32 UpdateLinearSystemRulesIndex = 0;
+	int32 UpdateLinearSystemCollisionsRulesIndex = 0;
+	int32 PostSubstepConstraintRulesIndex = 0;
+};
+
 FClothConstraints::FClothConstraints()
 	: Evolution(nullptr), PBDEvolution(nullptr)
 	, AnimationPositions(nullptr)
@@ -791,6 +1019,36 @@ void FClothConstraints::CreateStretchConstraints(
 		++NumConstraintInits;  // Uses init to update the property tables
 		++NumConstraintRules;
 	}
+	if (PatternData && PatternData->PatternPositions.Num() && Softs::FXPBDAnisotropicSpringConstraints::IsEnabled(ConfigProperties))
+	{
+		if (Evolution)
+		{
+			XAnisoSpringConstraints = MakeShared<Softs::FXPBDAnisotropicSpringConstraints>(
+				Evolution->GetSoftBodyParticles(ParticleRangeId),
+				TriangleMesh,
+				PatternData->WeldedFaceVertexPatternPositions,
+				WeightMaps,
+				ConfigProperties,
+				/*bTrimKinematicConstraints =*/ true);
+			++NumUpdateLinearSystemRules;
+		}
+		else
+		{
+			XAnisoSpringConstraints = MakeShared<Softs::FXPBDAnisotropicSpringConstraints>(
+				PBDEvolution->Particles(),
+				ParticleOffset,
+				NumParticles,
+				TriangleMesh,
+				PatternData->WeldedFaceVertexPatternPositions,
+				WeightMaps,
+				ConfigProperties,
+				/*bTrimKinematicConstraints =*/ true);
+		}
+
+		++NumConstraintInits;  // Uses init to update the property tables
+		++NumConstraintRules;
+		++NumConstraintRules; // 2X because edge and axial spring applies are split
+	}
 }
 
 void FClothConstraints::CreateBendingConstraints(
@@ -1178,83 +1436,28 @@ void FClothConstraints::CreateCollisionConstraint(
 
 void FClothConstraints::CreateForceBasedRules()
 {
-	Evolution->AllocatePreSubstepParallelInitRange(ParticleRangeId, NumPreSubstepInits);
-	Evolution->AllocatePBDExternalForceRulesRange(ParticleRangeId, NumExternalForceRules);
-	Evolution->AllocatePostInitialGuessParallelInitRange(ParticleRangeId, NumConstraintInits);
-	Evolution->AllocatePreSubstepConstraintRulesRange(ParticleRangeId, NumPreSubstepConstraintRules);
-	Evolution->AllocatePerIterationPBDConstraintRulesRange(ParticleRangeId, NumConstraintRules);
-	Evolution->AllocatePerIterationCollisionPBDConstraintRulesRange(ParticleRangeId, NumCollisionConstraintRules);
-	Evolution->AllocatePerIterationPostCollisionsPBDConstraintRulesRange(ParticleRangeId, NumPostCollisionConstraintRules);
-	Evolution->AllocateUpdateLinearSystemRulesRange(ParticleRangeId, NumUpdateLinearSystemRules);
-	Evolution->AllocateUpdateLinearSystemCollisionsRulesRange(ParticleRangeId, NumUpdateLinearSystemCollisionsRules);
-	Evolution->AllocatePostSubstepConstraintRulesRange(ParticleRangeId, NumPostprocessingConstraintRules);
-
-	TArrayView<Softs::FEvolution::ParallelInitFunc> PreSubstepParallelInits =
-		Evolution->GetPreSubstepParallelInitRange(ParticleRangeId);
-	TArrayView<Softs::FEvolution::PBDConstraintRuleFunc> ExternalForceRules =
-		Evolution->GetPBDExternalForceRulesRange(ParticleRangeId);
-	TArrayView<Softs::FEvolution::ParallelInitFunc> PostInitialGuessParallelInits =
-		Evolution->GetPostInitialGuessParallelInitRange(ParticleRangeId);
-	TArrayView<Softs::FEvolution::ConstraintRuleFunc> PreSubstepConstraintRules =
-		Evolution->GetPreSubstepConstraintRulesRange(ParticleRangeId);
-	TArrayView<Softs::FEvolution::PBDConstraintRuleFunc> PerIterationConstraintRules =
-		Evolution->GetPerIterationPBDConstraintRulesRange(ParticleRangeId);
-	TArrayView<Softs::FEvolution::PBDCollisionConstraintRuleFunc> PerIterationCollisionConstraintRules =
-		Evolution->GetPerIterationCollisionPBDConstraintRulesRange(ParticleRangeId);
-	TArrayView<Softs::FEvolution::PBDConstraintRuleFunc> PerIterationPostCollisionsConstraintRules =
-		Evolution->GetPerIterationPostCollisionsPBDConstraintRulesRange(ParticleRangeId);
-	TArrayView<Softs::FEvolution::UpdateLinearSystemFunc> UpdateLinearSystemRules =
-		Evolution->GetUpdateLinearSystemRulesRange(ParticleRangeId);
-	TArrayView<Softs::FEvolution::UpdateLinearSystemCollisionsFunc> UpdateLinearSystemCollisionsRules =
-		Evolution->GetUpdateLinearSystemCollisionsRulesRange(ParticleRangeId);
-	TArrayView<Softs::FEvolution::ConstraintRuleFunc> PostSubstepConstraintRules =
-		Evolution->GetPostSubstepConstraintRulesRange(ParticleRangeId);
-
-	int32 PreSubstepInitsIndex = 0;
-	int32 ExternalForceRulesIndex = 0;
-	int32 PostInitialGuessInitsIndex = 0;
-	int32 PreSubstepConstraintRulesIndex = 0;
-	int32 ConstraintRuleIndex = 0;
-	int32 CollisionConstraintRulesIndex = 0;
-	int32 PostCollisionConstraintRulesIndex = 0;
-	int32 UpdateLinearSystemRulesIndex = 0;
-	int32 UpdateLinearSystemCollisionsRulesIndex = 0;
-	int32 PostprocessingConstraintRulesIndex = 0;
+	FRuleCreator RuleCreator(this);
 
 	if (ExternalForces)
 	{
-		ExternalForceRules[ExternalForceRulesIndex++] =
-			[this](Softs::FSolverParticlesRange& Particles, const Softs::FSolverReal Dt)
-		{
-			ExternalForces->Apply(Particles, Dt);
-		};
-		UpdateLinearSystemRules[UpdateLinearSystemRulesIndex++] =
-			[this](const Softs::FSolverParticlesRange& Particles, const FSolverReal Dt, Softs::FEvolutionLinearSystem& LinearSystem)
-		{
-			ExternalForces->UpdateLinearSystem(Particles, Dt, LinearSystem);
-		};
+		RuleCreator.AddExternalForceRule_Apply(ExternalForces.Get());
+		RuleCreator.AddUpdateLinearSystemRule_UpdateLinearSystem(ExternalForces.Get());
 	}
-
 	if (VelocityAndPressureField)
 	{
-		PreSubstepParallelInits[PreSubstepInitsIndex++] =
+		RuleCreator.PreSubstepParallelInitRule(
 			[this](const Softs::FSolverParticlesRange& Particles, const Softs::FSolverReal Dt, const Softs::ESolverMode SolverMode)
 		{
 			VelocityAndPressureField->UpdateForces(Particles, Dt);
-		};
-
-		ExternalForceRules[ExternalForceRulesIndex++] =
-			[this](Softs::FSolverParticlesRange& Particles, const Softs::FSolverReal Dt)
-		{
-			VelocityAndPressureField->Apply(Particles, Dt);
-		};
+		});
+		RuleCreator.AddExternalForceRule_Apply(VelocityAndPressureField.Get());
 
 		// TODO Linear System
 	}
 
 	if (PerSolverField)
 	{
-		ExternalForceRules[ExternalForceRulesIndex++] =
+		RuleCreator.AddExternalForceRule(
 			[this](Softs::FSolverParticlesRange& Particles, const Softs::FSolverReal Dt)
 		{
 			const TArray<FVector>& LinearVelocities = PerSolverField->GetOutputResults(EFieldCommandOutputType::LinearVelocity);
@@ -1284,173 +1487,155 @@ void FClothConstraints::CreateForceBasedRules()
 					}
 				}
 			}
-		};
+		});
 
 		// TODO: Linear System
 	}
 
-#define UE_CHAOS_XPBD_INIT_VOID(Constraint)\
-	PostInitialGuessParallelInits[PostInitialGuessInitsIndex++] = \
-		[this](const Softs::FSolverParticlesRange& /*Particles*/, const Softs::FSolverReal Dt, const Softs::ESolverMode SolverMode)\
-	{\
-		Constraint->ApplyProperties(Dt, Evolution->GetIterations()); \
-		if ((SolverMode & Softs::ESolverMode::PBD) != Softs::ESolverMode::None)\
-		{\
-			Constraint->Init(); \
-		}\
-	};
-#define UE_CHAOS_XPBD_INIT_PARTICLES(Constraint)\
-	PostInitialGuessParallelInits[PostInitialGuessInitsIndex++] = \
-		[this](const Softs::FSolverParticlesRange& Particles, const Softs::FSolverReal Dt, const Softs::ESolverMode SolverMode)\
-	{\
-		Constraint->ApplyProperties(Dt, Evolution->GetIterations()); \
-		if ((SolverMode & Softs::ESolverMode::PBD) != Softs::ESolverMode::None)\
-		{\
-			Constraint->Init(Particles); \
-		}\
-	};
-#define UE_CHAOS_PBD_INIT(Constraint)\
-	PostInitialGuessParallelInits[PostInitialGuessInitsIndex++] = \
-		[this](const Softs::FSolverParticlesRange& /*Particles*/, const Softs::FSolverReal Dt, const Softs::ESolverMode SolverMode)\
-	{\
-		if ((SolverMode & Softs::ESolverMode::PBD) != Softs::ESolverMode::None)\
-		{\
-			Constraint->ApplyProperties(Dt, Evolution->GetIterations()); \
-		}\
-	};
-#define UE_CHAOS_APPLY_CONSTRAINT(Constraint)\
-	PerIterationConstraintRules[ConstraintRuleIndex++] = \
-		[this](Softs::FSolverParticlesRange& Particles, const Softs::FSolverReal Dt)\
-	{\
-		Constraint->Apply(Particles, Dt); \
-	}; 
-#define UE_CHAOS_UPDATE_LINEAR_SYSTEM(Constraint)\
-	UpdateLinearSystemRules[UpdateLinearSystemRulesIndex++] =\
-		[this](const Softs::FSolverParticlesRange& Particles, const FSolverReal Dt, Softs::FEvolutionLinearSystem& LinearSystem)\
-	{\
-		Constraint->UpdateLinearSystem(Particles, Dt, LinearSystem);\
-	};
+	if (XStretchBiasConstraints)
+	{
+		constexpr bool bInitParticles = false;
+		RuleCreator.AddXPBDParallelInitRule<bInitParticles>(XStretchBiasConstraints.Get());
+		constexpr bool bPostCollisions = false;
+		RuleCreator.AddPerIterationPBDConstraintRule_Apply<bPostCollisions>(XStretchBiasConstraints.Get());
+	}
+	if (XEdgeConstraints)
+	{
+		constexpr bool bInitParticles = false;
+		RuleCreator.AddXPBDParallelInitRule<bInitParticles>(XEdgeConstraints.Get());
+		constexpr bool bPostCollisions = false;
+		RuleCreator.AddPerIterationPBDConstraintRule_Apply<bPostCollisions>(XEdgeConstraints.Get());
+		RuleCreator.AddUpdateLinearSystemRule_UpdateLinearSystem(XEdgeConstraints.Get());
+	}
+	if (EdgeConstraints)
+	{
+		RuleCreator.AddPBDParallelInitRule(EdgeConstraints.Get());
+		constexpr bool bPostCollisions = false;
+		RuleCreator.AddPerIterationPBDConstraintRule_Apply<bPostCollisions>(EdgeConstraints.Get());
+	}
+	if (XAnisoSpringConstraints)
+	{
+		constexpr bool bInitParticles = false;
+		RuleCreator.AddXPBDParallelInitRule<bInitParticles>(XAnisoSpringConstraints.Get());
+		constexpr bool bPostCollisions = false;
+		RuleCreator.AddPerIterationPBDConstraintRule_Apply<bPostCollisions>(&XAnisoSpringConstraints->GetEdgeConstraints());
+		RuleCreator.AddUpdateLinearSystemRule(
+			[this](const Softs::FSolverParticlesRange& Particles, const FSolverReal Dt, Softs::FEvolutionLinearSystem& LinearSystem)
+		{
+			XAnisoSpringConstraints->GetEdgeConstraints().UpdateLinearSystem(Particles, Dt, LinearSystem);
+			XAnisoSpringConstraints->GetAxialConstraints().UpdateLinearSystem(Particles, Dt, LinearSystem);
+		});
+	}
 
-#define UE_CHAOS_ADD_XPBD_CONSTRAINT_NO_LINEAR_SYSTEM(Constraint) \
-	if (Constraint)\
-	{\
-		UE_CHAOS_XPBD_INIT_VOID(Constraint)\
-		UE_CHAOS_APPLY_CONSTRAINT(Constraint)\
-	}
-
-#define UE_CHAOS_ADD_XPBD_CONSTRAINT_INIT_PARTICLES_NO_LINEAR_SYSTEM(Constraint) \
-	if (Constraint)\
-	{\
-		UE_CHAOS_XPBD_INIT_PARTICLES(Constraint)\
-		UE_CHAOS_APPLY_CONSTRAINT(Constraint)\
-	}
-
-#define UE_CHAOS_ADD_XPBD_CONSTRAINT(Constraint) \
-	if (Constraint)\
-	{\
-		UE_CHAOS_XPBD_INIT_VOID(Constraint)\
-		UE_CHAOS_APPLY_CONSTRAINT(Constraint)\
-		UE_CHAOS_UPDATE_LINEAR_SYSTEM(Constraint)\
-	}
-
-#define UE_CHAOS_ADD_PBD_CONSTRAINT(Constraint) \
-	if (Constraint)\
-	{\
-		UE_CHAOS_PBD_INIT(Constraint)\
-		UE_CHAOS_APPLY_CONSTRAINT(Constraint)\
-	}
-
-#define UE_CHAOS_ADD_PBD_CONSTRAINT_NO_PROPERTIES(Constraint) \
-	if (Constraint)\
-	{\
-		UE_CHAOS_APPLY_CONSTRAINT(Constraint)\
-	}
-	
-	UE_CHAOS_ADD_XPBD_CONSTRAINT_NO_LINEAR_SYSTEM(XStretchBiasConstraints);
-	UE_CHAOS_ADD_XPBD_CONSTRAINT(XEdgeConstraints);
-	UE_CHAOS_ADD_PBD_CONSTRAINT(EdgeConstraints);
-	UE_CHAOS_ADD_XPBD_CONSTRAINT(XBendingConstraints);
-	UE_CHAOS_ADD_PBD_CONSTRAINT(BendingConstraints);
-	UE_CHAOS_ADD_PBD_CONSTRAINT(BendingElementConstraints);
-	UE_CHAOS_ADD_XPBD_CONSTRAINT_INIT_PARTICLES_NO_LINEAR_SYSTEM(XBendingElementConstraints);
-	UE_CHAOS_ADD_XPBD_CONSTRAINT_INIT_PARTICLES_NO_LINEAR_SYSTEM(XAnisoBendingElementConstraints);
-	UE_CHAOS_ADD_XPBD_CONSTRAINT_NO_LINEAR_SYSTEM(XAreaConstraints);
-	UE_CHAOS_ADD_PBD_CONSTRAINT(AreaConstraints);
-	UE_CHAOS_ADD_PBD_CONSTRAINT_NO_PROPERTIES(MaximumDistanceConstraints);
-	UE_CHAOS_ADD_PBD_CONSTRAINT_NO_PROPERTIES(BackstopConstraints);
-	UE_CHAOS_ADD_PBD_CONSTRAINT(AnimDriveConstraints);
-
-#undef UE_CHAOS_ADD_PBD_CONSTRAINT_NO_PROPERTIES
-#undef UE_CHAOS_ADD_PBD_CONSTRAINT
-#undef UE_CHAOS_ADD_XPBD_CONSTRAINT
-#undef UE_CHAOS_ADD_XPBD_CONSTRAINT_INIT_PARTICLES_NO_LINEAR_SYSTEM
-#undef UE_CHAOS_ADD_XPBD_CONSTRAINT_NO_LINEAR_SYSTEM
-#undef UE_CHAOS_UPDATE_LINEAR_SYSTEM
-#undef UE_CHAOS_APPLY_CONSTRAINT
-#undef UE_CHAOS_PBD_INIT
-#undef UE_CHAOS_XPBD_INIT_PARTICLES
-#undef UE_CHAOS_XPBD_INIT_VOID
+	if (XBendingConstraints)
+	{
+		constexpr bool bInitParticles = false;
+		RuleCreator.AddXPBDParallelInitRule<bInitParticles>(XBendingConstraints.Get());
+		constexpr bool bPostCollisions = false;
+		RuleCreator.AddPerIterationPBDConstraintRule_Apply<bPostCollisions>(XBendingConstraints.Get());
+		RuleCreator.AddUpdateLinearSystemRule_UpdateLinearSystem(XBendingConstraints.Get());
+	}
+	if (BendingConstraints)
+	{
+		RuleCreator.AddPBDParallelInitRule(BendingConstraints.Get());
+		constexpr bool bPostCollisions = false;
+		RuleCreator.AddPerIterationPBDConstraintRule_Apply<bPostCollisions>(BendingConstraints.Get());
+	}
+	if (BendingElementConstraints)
+	{
+		RuleCreator.AddPBDParallelInitRule(BendingElementConstraints.Get());
+		constexpr bool bPostCollisions = false;
+		RuleCreator.AddPerIterationPBDConstraintRule_Apply<bPostCollisions>(BendingElementConstraints.Get());
+	}
+	if (XBendingElementConstraints)
+	{
+		constexpr bool bInitParticles = true;
+		RuleCreator.AddXPBDParallelInitRule<bInitParticles>(XBendingElementConstraints.Get());
+		constexpr bool bPostCollisions = false;
+		RuleCreator.AddPerIterationPBDConstraintRule_Apply<bPostCollisions>(XBendingElementConstraints.Get());
+	}
+	if (XAnisoBendingElementConstraints)
+	{
+		constexpr bool bInitParticles = true;
+		RuleCreator.AddXPBDParallelInitRule<bInitParticles>(XAnisoBendingElementConstraints.Get());
+		constexpr bool bPostCollisions = false;
+		RuleCreator.AddPerIterationPBDConstraintRule_Apply<bPostCollisions>(XAnisoBendingElementConstraints.Get());
+	}
+	if (XAreaConstraints)
+	{
+		constexpr bool bInitParticles = false;
+		RuleCreator.AddXPBDParallelInitRule<bInitParticles>(XAreaConstraints.Get());
+		constexpr bool bPostCollisions = false;
+		RuleCreator.AddPerIterationPBDConstraintRule_Apply<bPostCollisions>(XAreaConstraints.Get());
+	}
+	if (AreaConstraints)
+	{
+		RuleCreator.AddPBDParallelInitRule(AreaConstraints.Get());
+		constexpr bool bPostCollisions = false;
+		RuleCreator.AddPerIterationPBDConstraintRule_Apply<bPostCollisions>(AreaConstraints.Get());
+	}
+	if (XAnisoSpringConstraints)
+	{
+		constexpr bool bPostCollisions = false;
+		RuleCreator.AddPerIterationPBDConstraintRule_Apply<bPostCollisions>(&XAnisoSpringConstraints->GetAxialConstraints());
+	}
+	if (MaximumDistanceConstraints)
+	{
+		constexpr bool bPostCollisions = false;
+		RuleCreator.AddPerIterationPBDConstraintRule_Apply<bPostCollisions>(MaximumDistanceConstraints.Get());
+	}
+	if (BackstopConstraints)
+	{
+		constexpr bool bPostCollisions = false;
+		RuleCreator.AddPerIterationPBDConstraintRule_Apply<bPostCollisions>(BackstopConstraints.Get());
+	}
+	if (AnimDriveConstraints)
+	{
+		RuleCreator.AddPBDParallelInitRule(AnimDriveConstraints.Get());
+		constexpr bool bPostCollisions = false;
+		RuleCreator.AddPerIterationPBDConstraintRule_Apply<bPostCollisions>(AnimDriveConstraints.Get());
+	}
 
 	if (CollisionConstraint)
 	{
-		PerIterationCollisionConstraintRules[CollisionConstraintRulesIndex++] =
-			[this](Softs::FSolverParticlesRange& Particles, const Softs::FSolverReal Dt, const TArray<Softs::FSolverCollisionParticlesRange>& CollisionParticles)
-		{
-			CollisionConstraint->Apply(Particles, Dt, CollisionParticles);
-		};
-
-		UpdateLinearSystemCollisionsRules[UpdateLinearSystemCollisionsRulesIndex++] =
-			[this](const Softs::FSolverParticlesRange& Particles, const Softs::FSolverReal Dt, const TArray<Softs::FSolverCollisionParticlesRange>& CollisionParticles, Softs::FEvolutionLinearSystem& LinearSystem)
-		{
-			CollisionConstraint->UpdateLinearSystem(Particles, Dt, CollisionParticles, LinearSystem);
-		};
+		RuleCreator.AddPerIterationCollisionPBDConstraintRule_Apply(CollisionConstraint.Get());
+		RuleCreator.AddUpdateLinearSystemCollisionsRule_UpdateLinearSystem(CollisionConstraint.Get());
 	}
 	if (SelfCollisionInit && SelfCollisionConstraints)
 	{
-		PostInitialGuessParallelInits[PostInitialGuessInitsIndex++] =
+		RuleCreator.AddPostInitialGuessParallelInitRule(
 			[this](const Softs::FSolverParticlesRange& Particles, const Softs::FSolverReal Dt, const Softs::ESolverMode SolverMode)
 		{
 			// Thickness * 2 to account for collision radius for both particles
 			SelfCollisionInit->Init(Particles, SelfCollisionConstraints->GetThickness() * (Softs::FSolverReal)2.f);
 			SelfCollisionConstraints->Init(Particles, SelfCollisionInit->GetSpatialHash(), SelfCollisionInit->GetVertexGIAColors(), SelfCollisionInit->GetTriangleGIAColors());
-		};
+		});
 
-		PerIterationPostCollisionsConstraintRules[PostCollisionConstraintRulesIndex++] =
-			[this](Softs::FSolverParticlesRange& Particles, const Softs::FSolverReal Dt)
-		{
-			SelfCollisionConstraints->Apply(Particles, Dt);
-		};
-
-		UpdateLinearSystemRules[UpdateLinearSystemRulesIndex++] =
-			[this](const Softs::FSolverParticlesRange& Particles, const FSolverReal Dt, Softs::FEvolutionLinearSystem& LinearSystem)
-		{
-			SelfCollisionConstraints->UpdateLinearSystem(Particles, Dt, LinearSystem);
-		};
+		constexpr bool bPostCollisions = true;
+		RuleCreator.AddPerIterationPBDConstraintRule_Apply<bPostCollisions>(SelfCollisionConstraints.Get());
+		RuleCreator.AddUpdateLinearSystemRule_UpdateLinearSystem(SelfCollisionConstraints.Get());
 	}
 	if (SelfCollisionSphereConstraints)
 	{
-		PostInitialGuessParallelInits[PostInitialGuessInitsIndex++] =
+		RuleCreator.AddPostInitialGuessParallelInitRule(
 			[this](const Softs::FSolverParticlesRange& Particles, const Softs::FSolverReal Dt, const Softs::ESolverMode SolverMode)
 		{
 			SelfCollisionSphereConstraints->Init(Particles);
-		};
+		});
 
-		PerIterationPostCollisionsConstraintRules[PostCollisionConstraintRulesIndex++] =
-			[this](Softs::FSolverParticlesRange& Particles, const Softs::FSolverReal Dt)
-		{
-			SelfCollisionSphereConstraints->Apply(Particles, Dt);
-		};
+		constexpr bool bPostCollisions = true;
+		RuleCreator.AddPerIterationPBDConstraintRule_Apply<bPostCollisions>(SelfCollisionSphereConstraints.Get());
 	}
 
 	if (SelfCollisionInit && SelfIntersectionConstraints)
 	{
-		PreSubstepConstraintRules[PreSubstepConstraintRulesIndex++] =
+		RuleCreator.AddPreSubstepConstraintRule(
 			[this](Softs::FSolverParticlesRange& Particles, const Softs::FSolverReal Dt, const Softs::ESolverMode SolverMode)
 		{
 			SelfIntersectionConstraints->Apply(Particles, SelfCollisionInit->GetContourMinimizationIntersections(), Dt);
-		};
+		});
 
-		PostSubstepConstraintRules[PostprocessingConstraintRulesIndex++] =
+		RuleCreator.AddPostSubstepConstraintRule(
 			[this](Softs::FSolverParticlesRange& Particles, const Softs::FSolverReal Dt, const Softs::ESolverMode SolverMode)
 		{
 			const int32 NumContourIterations = SelfCollisionInit->GetNumContourMinimizationPostSteps();
@@ -1459,11 +1644,11 @@ void FClothConstraints::CreateForceBasedRules()
 				SelfCollisionInit->PostStepInit(Particles);
 				SelfIntersectionConstraints->Apply(Particles, SelfCollisionInit->GetPostStepContourMinimizationIntersections(), Dt);
 			}
-		};
+		});
 	}
 	if (LongRangeConstraints)
 	{
-		PreSubstepConstraintRules[PreSubstepConstraintRulesIndex++] =
+		RuleCreator.AddPreSubstepConstraintRule(
 			[this](Softs::FSolverParticlesRange& Particles, const Softs::FSolverReal Dt, const Softs::ESolverMode SolverMode)
 		{
 			if ((SolverMode & Softs::ESolverMode::PBD) != Softs::ESolverMode::None)
@@ -1473,19 +1658,8 @@ void FClothConstraints::CreateForceBasedRules()
 				LongRangeConstraints->ApplyProperties(Dt, NumLRAIterations);
 				LongRangeConstraints->Apply(Particles, Dt);  // Run the LRA constraint only once per timestep
 			}
-		};
+		});
 	}
-
-	check(PreSubstepInitsIndex == NumPreSubstepInits);
-	check(ExternalForceRulesIndex == NumExternalForceRules);
-	check(PostInitialGuessInitsIndex == NumConstraintInits)
-	check(PreSubstepConstraintRulesIndex == NumPreSubstepConstraintRules);
-	check(ConstraintRuleIndex == NumConstraintRules);
-	check(CollisionConstraintRulesIndex == NumCollisionConstraintRules);
-	check(PostCollisionConstraintRulesIndex == NumPostCollisionConstraintRules);
-	check(UpdateLinearSystemRulesIndex == NumUpdateLinearSystemRules);
-	check(UpdateLinearSystemCollisionsRulesIndex == NumUpdateLinearSystemCollisionsRules);
-	check(PostprocessingConstraintRulesIndex == NumPostprocessingConstraintRules);
 }
 
 
@@ -1660,6 +1834,20 @@ void FClothConstraints::CreatePBDRules()
 				EdgeConstraints->Apply(Particles, Dt);
 			};
 	}
+	if (XAnisoSpringConstraints)
+	{
+		ConstraintInits[ConstraintInitIndex++] =
+			[this](Softs::FSolverParticles& /*Particles*/, const Softs::FSolverReal Dt)
+		{
+			XAnisoSpringConstraints->Init();
+			XAnisoSpringConstraints->ApplyProperties(Dt, PBDEvolution->GetIterations());
+		};
+		ConstraintRules[ConstraintRuleIndex++] =
+			[this](Softs::FSolverParticles& Particles, const Softs::FSolverReal Dt)
+		{
+			XAnisoSpringConstraints->GetEdgeConstraints().Apply(Particles, Dt);
+		};
+	}
 	if (XBendingConstraints)
 	{
 		ConstraintInits[ConstraintInitIndex++] =
@@ -1756,6 +1944,14 @@ void FClothConstraints::CreatePBDRules()
 				AreaConstraints->Apply(Particles, Dt);
 			};
 	}
+	if (XAnisoSpringConstraints)
+	{
+		ConstraintRules[ConstraintRuleIndex++] =
+			[this](Softs::FSolverParticles& Particles, const Softs::FSolverReal Dt)
+		{
+			XAnisoSpringConstraints->GetAxialConstraints().Apply(Particles, Dt);
+		};
+	}
 	if (MaximumDistanceConstraints)
 	{
 		ConstraintRules[ConstraintRuleIndex++] =
@@ -1903,6 +2099,10 @@ void FClothConstraints::Update(
 	{
 		XStretchBiasConstraints->SetProperties(ConfigProperties, WeightMaps);
 	}
+	if (XAnisoSpringConstraints)
+	{
+		XAnisoSpringConstraints->SetProperties(ConfigProperties, WeightMaps);
+	}
 	if (BendingConstraints)
 	{
 		BendingConstraints->SetProperties(ConfigProperties, WeightMaps);
diff --git a/Engine/Plugins/Experimental/ChaosCloth/Source/Chaos/Private/ChaosCloth/ChaosClothVisualization.cpp b/Engine/Plugins/Experimental/ChaosCloth/Source/Chaos/Private/ChaosCloth/ChaosClothVisualization.cpp
--- a/Engine/Plugins/Experimental/ChaosCloth/Source/Chaos/Private/ChaosCloth/ChaosClothVisualization.cpp
+++ b/Engine/Plugins/Experimental/ChaosCloth/Source/Chaos/Private/ChaosCloth/ChaosClothVisualization.cpp
@@ -26,6 +26,7 @@
 #include "Chaos/XPBDBendingConstraints.h"
 #include "Chaos/XPBDSpringConstraints.h"
 #include "Chaos/XPBDAnisotropicBendingConstraints.h"
+#include "Chaos/XPBDAnisotropicSpringConstraints.h"
 #include "Chaos/XPBDStretchBiasElementConstraints.h"
 #include "Chaos/WeightedLatticeImplicitObject.h"
 #include "DynamicMeshBuilder.h"
@@ -75,6 +76,14 @@ static FAutoConsoleVariableRef CVarClothVizStretchBiasDrawRangeMax(TEXT("p.Chaos
 static bool bStretchBiasDrawOutOfRange = true;
 static FAutoConsoleVariableRef CVarClothVizStretchBiasDrawOutOfRange(TEXT("p.ChaosClothVisualization.StretchBiasDrawOutOfRange"), bStretchBiasDrawOutOfRange, TEXT("Draw out of range elements (When drawing warp/weft stretch)"));
 
+enum class EAnisoSpringDrawMode : int
+{
+	ParallelGraphColor = 0,
+	Anisotropy = 1,
+};
+static int32 AnisoSpringDrawMode = (int32)EAnisoSpringDrawMode::ParallelGraphColor;
+static FAutoConsoleVariableRef CVarClothVizAnisoSpringDrawMode(TEXT("p.ChaosClothVisualization.AnisoSpringDrawMode"), AnisoSpringDrawMode, TEXT("Stretch draw mode, 0 = Parallel graph color, 1 = Anisotropy"));
+
 static FString WeightMapName = "";
 static FAutoConsoleVariableRef CVarClothVizWeightMapName(TEXT("p.ChaosClothVisualization.WeightMapName"), WeightMapName, TEXT("Weight map name to be visualized"));
 
@@ -1355,13 +1364,11 @@ static FAutoConsoleVariableRef CVarClothVizWeightMapName(TEXT("p.ChaosClothVisua
 			}
 		}
 	}
-
-	template<typename SpringConstraintType>
-	static void DrawSpringConstraintColors(FPrimitiveDrawInterface* PDI, const TConstArrayView<::Chaos::Softs::FSolverVec3>& Positions, const ::Chaos::FVec3& LocalSpaceLocation, const SpringConstraintType* const SpringConstraints)
+	template<typename SpringConstraintType, typename GetEndPointsFuncType>
+	static void DrawSpringConstraintColors(FPrimitiveDrawInterface* PDI, const ::Chaos::FVec3& LocalSpaceLocation, const SpringConstraintType* const SpringConstraints, GetEndPointsFuncType GetEndPoints)
 	{
 		check(SpringConstraints);
 
-		const auto& Constraints = SpringConstraints->GetConstraints(); // auto because constraints can be TVec2<int> or TVec4<int>, but we just care about first two for drawing here either way
 		const TArray<int32>& ConstraintsPerColorStartIndex = SpringConstraints->GetConstraintsPerColorStartIndex();
 		if (ConstraintsPerColorStartIndex.Num() > 1)
 		{
@@ -1381,26 +1388,41 @@ static FAutoConsoleVariableRef CVarClothVizWeightMapName(TEXT("p.ChaosClothVisua
 				for (int32 ConstraintIndex = ColorStart; ConstraintIndex < ColorEnd; ++ConstraintIndex)
 				{
 					// Draw line
-					const auto& Constraint = Constraints[ConstraintIndex];
-					const FVec3 Pos0 = FVec3(Positions[Constraint[0]]) + LocalSpaceLocation;
-					const FVec3 Pos1 = FVec3(Positions[Constraint[1]]) + LocalSpaceLocation;
+					Softs::FSolverVec3 P1, P2;
+					GetEndPoints(ConstraintIndex, P1, P2);
+
+					const FVec3 Pos0 = FVec3(P1) + LocalSpaceLocation;
+					const FVec3 Pos1 = FVec3(P2) + LocalSpaceLocation;
 					DrawLine(PDI, Pos0, Pos1, DrawColor);
 				}
 			}
 		}
 		else
 		{
-			for (const auto& Constraint : Constraints)
+			for (int32 ConstraintIndex = 0; ConstraintIndex < SpringConstraints->GetConstraints().Num(); ++ConstraintIndex)
 			{
 				// Draw line
-				const FVec3 Pos0 = FVec3(Positions[Constraint[0]]) + LocalSpaceLocation;
-				const FVec3 Pos1 = FVec3(Positions[Constraint[1]]) + LocalSpaceLocation;
+				Softs::FSolverVec3 P1, P2;
+				GetEndPoints(ConstraintIndex, P1, P2);
+
+				const FVec3 Pos0 = FVec3(P1) + LocalSpaceLocation;
+				const FVec3 Pos1 = FVec3(P2) + LocalSpaceLocation;
 
 				DrawLine(PDI, Pos0, Pos1, FLinearColor::Black);
 			}
 		}
 	}
 
+	template<typename SpringConstraintType>
+	static void DrawSpringConstraintColors(FPrimitiveDrawInterface* PDI, const TConstArrayView<::Chaos::Softs::FSolverVec3>& Positions, const ::Chaos::FVec3& LocalSpaceLocation, const SpringConstraintType* const SpringConstraints)
+	{
+		DrawSpringConstraintColors(PDI, LocalSpaceLocation, SpringConstraints, [&SpringConstraints, &Positions](const int32 ConstraintIndex, Softs::FSolverVec3& P1, Softs::FSolverVec3& P2)
+		{
+			P1 = Positions[SpringConstraints->GetConstraints()[ConstraintIndex][0]];
+			P2 = Positions[SpringConstraints->GetConstraints()[ConstraintIndex][1]];
+		});
+	}
+
 	static void DrawStretchBiasConstraints_ParallelGraphColor(FPrimitiveDrawInterface* PDI, const TConstArrayView<::Chaos::Softs::FSolverVec3>& Positions, const ::Chaos::FVec3& LocalSpaceLocation, const FMaterialRenderProxy* MaterialRenderProxy, const Softs::FXPBDStretchBiasElementConstraints* const SpringConstraints)
 	{
 #if WITH_EDITOR
@@ -1588,6 +1610,55 @@ static FAutoConsoleVariableRef CVarClothVizWeightMapName(TEXT("p.ChaosClothVisua
 #endif
 	}
 
+	template<typename ConstraintType>
+	static void DrawEdgeAnisotropy(FPrimitiveDrawInterface* PDI, const TConstArrayView<::Chaos::Softs::FSolverVec3>& Positions, const ::Chaos::FVec3& LocalSpaceLocation, const ConstraintType* const BendingConstraints)
+	{
+		const auto& Constraints = BendingConstraints->GetConstraints(); // auto because this could be Vec2 or Vec4, but we always just care about first two indices
+		const TArray<Softs::FSolverVec3>& WarpWeftBiasBaseMultipliers = BendingConstraints->GetWarpWeftBiasBaseMultipliers();
+		for (int32 ConstraintIndex = 0; ConstraintIndex < Constraints.Num(); ++ConstraintIndex)
+		{
+			const Softs::FSolverVec3& P1 = Positions[Constraints[ConstraintIndex][0]];
+			const Softs::FSolverVec3& P2 = Positions[Constraints[ConstraintIndex][1]];
+			const Softs::FSolverVec3& Multiplier = WarpWeftBiasBaseMultipliers[ConstraintIndex];
+
+			const FVec3 Pos0 = FVec3(P1) + LocalSpaceLocation;
+			const FVec3 Pos1 = FVec3(P2) + LocalSpaceLocation;
+			DrawLine(PDI, Pos0, Pos1, FLinearColor(Multiplier[0], Multiplier[1], Multiplier[2]));
+		}
+	}
+
+	template<typename ConstraintType>
+	static void DrawAxialSpringAnisotropy(FPrimitiveDrawInterface* PDI, const TConstArrayView<::Chaos::Softs::FSolverVec3>& Positions, const ::Chaos::FVec3& LocalSpaceLocation, const ConstraintType* const AxialConstraints)
+	{
+		const TArray<TVec3<int32>>& Constraints = AxialConstraints->GetConstraints();
+		const TArray<Softs::FSolverVec3>& WarpWeftBiasBaseMultipliers = AxialConstraints->GetWarpWeftBiasBaseMultipliers();
+		const TArray<Softs::FSolverReal>& Barys = AxialConstraints->GetBarys();
+		for (int32 ConstraintIndex = 0; ConstraintIndex < Constraints.Num(); ++ConstraintIndex)
+		{
+			const Softs::FSolverVec3& P1 = Positions[Constraints[ConstraintIndex][0]];
+			const Softs::FSolverVec3& P2 = Positions[Constraints[ConstraintIndex][1]];
+			const Softs::FSolverVec3& P3 = Positions[Constraints[ConstraintIndex][2]];
+			const Softs::FSolverVec3 P = Barys[ConstraintIndex] * P2 + ((FSolverReal)1. - Barys[ConstraintIndex]) * P3;
+			const Softs::FSolverVec3& Multiplier = WarpWeftBiasBaseMultipliers[ConstraintIndex];
+
+			const FVec3 Pos0 = FVec3(P1) + LocalSpaceLocation;
+			const FVec3 Pos1 = FVec3(P) + LocalSpaceLocation;
+			DrawLine(PDI, Pos0, Pos1, FLinearColor(Multiplier[0], Multiplier[1], Multiplier[2]));
+		}
+	}
+
+	template<typename SpringConstraintType>
+	static void DrawAxialSpringConstraintColors(FPrimitiveDrawInterface* PDI, const TConstArrayView<::Chaos::Softs::FSolverVec3>& Positions, const ::Chaos::FVec3& LocalSpaceLocation, const SpringConstraintType* const SpringConstraints)
+	{
+		DrawSpringConstraintColors(PDI, LocalSpaceLocation, SpringConstraints, [&SpringConstraints, &Positions](const int32 ConstraintIndex, Softs::FSolverVec3& P1, Softs::FSolverVec3& P2)
+		{
+			P1 = Positions[SpringConstraints->GetConstraints()[ConstraintIndex][0]];
+			const Softs::FSolverReal Bary = SpringConstraints->GetBarys()[ConstraintIndex];
+			
+			P2 = Bary * Positions[SpringConstraints->GetConstraints()[ConstraintIndex][1]] + ((FSolverReal)1. - Bary)* Positions[SpringConstraints->GetConstraints()[ConstraintIndex][2]];
+		});
+	}
+
 	void FClothVisualization::DrawEdgeConstraint(FPrimitiveDrawInterface* PDI) const
 	{
 		if (!Solver)
@@ -1621,6 +1692,22 @@ static FAutoConsoleVariableRef CVarClothVizWeightMapName(TEXT("p.ChaosClothVisua
 				DrawSpringConstraintColors(PDI, Positions, LocalSpaceLocation, EdgeConstraints);
 			}
 
+			if (const Softs::FXPBDAnisotropicSpringConstraints* const AnisoSpringConstraints = ClothConstraints.GetXAnisoSpringConstraints().Get())
+			{
+				switch ((Private::EAnisoSpringDrawMode)Private::AnisoSpringDrawMode)
+				{
+				case Private::EAnisoSpringDrawMode::Anisotropy:
+					DrawEdgeAnisotropy(PDI, Positions, LocalSpaceLocation, &AnisoSpringConstraints->GetEdgeConstraints());
+					DrawAxialSpringAnisotropy(PDI, Positions, LocalSpaceLocation, &AnisoSpringConstraints->GetAxialConstraints());
+					break;
+				case Private::EAnisoSpringDrawMode::ParallelGraphColor: // fallthrough
+				default:
+					DrawSpringConstraintColors(PDI, Positions, LocalSpaceLocation, &AnisoSpringConstraints->GetEdgeConstraints());
+					DrawAxialSpringConstraintColors(PDI, Positions, LocalSpaceLocation, &AnisoSpringConstraints->GetAxialConstraints());
+					break;
+				}
+			}
+
 			if (const Softs::FXPBDStretchBiasElementConstraints* const StretchConstraints = ClothConstraints.GetXStretchBiasConstraints().Get())
 			{
 				const FMaterialRenderProxy* MaterialRenderProxy =
@@ -1666,21 +1753,8 @@ static FAutoConsoleVariableRef CVarClothVizWeightMapName(TEXT("p.ChaosClothVisua
 		}
 	}
 
-	static void DrawBendingElementAnisotropy(FPrimitiveDrawInterface* PDI, const TConstArrayView<::Chaos::Softs::FSolverVec3>& Positions, const ::Chaos::FVec3& LocalSpaceLocation, const ::Chaos::Softs::FXPBDAnisotropicBendingConstraints* const BendingConstraints)
-	{
-		const TArray<TVec4<int32>>& Constraints = BendingConstraints->GetConstraints();
-		const TArray<Softs::FSolverVec3>& WarpWeftBiasBaseMultipliers = BendingConstraints->GetWarpWeftBiasBaseMultipliers();
-		for (int32 ConstraintIndex = 0; ConstraintIndex < Constraints.Num(); ++ConstraintIndex)
-		{
-			const Softs::FSolverVec3& P1 = Positions[Constraints[ConstraintIndex][0]];
-			const Softs::FSolverVec3& P2 = Positions[Constraints[ConstraintIndex][1]];
-			const Softs::FSolverVec3& Multiplier = WarpWeftBiasBaseMultipliers[ConstraintIndex];
 
-			const FVec3 Pos0 = FVec3(P1) + LocalSpaceLocation;
-			const FVec3 Pos1 = FVec3(P2) + LocalSpaceLocation;
-			DrawLine(PDI, Pos0, Pos1, FLinearColor(Multiplier[0], Multiplier[1], Multiplier[2]));
-		}
-	}
+
 	static void DrawBendingElementRestAngle(FPrimitiveDrawInterface* PDI, const TConstArrayView<::Chaos::Softs::FSolverVec3>& Positions, const ::Chaos::FVec3& LocalSpaceLocation, const ::Chaos::Softs::FPBDBendingConstraintsBase* const BendingConstraints)
 	{
 		const TArray<TVec4<int32>>& Constraints = BendingConstraints->GetConstraints();
@@ -1777,7 +1851,7 @@ static FAutoConsoleVariableRef CVarClothVizWeightMapName(TEXT("p.ChaosClothVisua
 					DrawSpringConstraintColors(PDI, Positions, LocalSpaceLocation, BendingConstraints);
 					break;
 				case Private::EBendingDrawMode::Anisotropy:
-					DrawBendingElementAnisotropy(PDI, Positions, LocalSpaceLocation, BendingConstraints);
+					DrawEdgeAnisotropy(PDI, Positions, LocalSpaceLocation, BendingConstraints);
 					break;
 				case Private::EBendingDrawMode::BuckleStatus:
 				default:
diff --git a/Engine/Plugins/Experimental/ChaosCloth/Source/Chaos/Public/ChaosCloth/ChaosClothConstraints.h b/Engine/Plugins/Experimental/ChaosCloth/Source/Chaos/Public/ChaosCloth/ChaosClothConstraints.h
--- a/Engine/Plugins/Experimental/ChaosCloth/Source/Chaos/Public/ChaosCloth/ChaosClothConstraints.h
+++ b/Engine/Plugins/Experimental/ChaosCloth/Source/Chaos/Public/ChaosCloth/ChaosClothConstraints.h
@@ -121,6 +121,7 @@ PRAGMA_ENABLE_DEPRECATION_WARNINGS
 		const TSharedPtr<Softs::FPBDEdgeSpringConstraints>& GetEdgeSpringConstraints() const { return EdgeConstraints; }
 		const TSharedPtr<Softs::FXPBDEdgeSpringConstraints>& GetXEdgeSpringConstraints() const { return XEdgeConstraints; }
 		const TSharedPtr<Softs::FXPBDStretchBiasElementConstraints>& GetXStretchBiasConstraints() const { return XStretchBiasConstraints; }
+		const TSharedPtr<Softs::FXPBDAnisotropicSpringConstraints>& GetXAnisoSpringConstraints() const { return XAnisoSpringConstraints; }
 		const TSharedPtr<Softs::FPBDBendingSpringConstraints>& GetBendingSpringConstraints() const { return BendingConstraints; }
 		const TSharedPtr<Softs::FXPBDBendingSpringConstraints>& GetXBendingSpringConstraints() const { return XBendingConstraints; }
 		const TSharedPtr<Softs::FPBDBendingConstraints>& GetBendingElementConstraints() const { return BendingElementConstraints; }
@@ -202,6 +203,7 @@ PRAGMA_ENABLE_DEPRECATION_WARNINGS
 		TSharedPtr<Softs::FPBDEdgeSpringConstraints> EdgeConstraints;
 		TSharedPtr<Softs::FXPBDEdgeSpringConstraints> XEdgeConstraints;
 		TSharedPtr<Softs::FXPBDStretchBiasElementConstraints> XStretchBiasConstraints;
+		TSharedPtr<Softs::FXPBDAnisotropicSpringConstraints> XAnisoSpringConstraints;
 		TSharedPtr<Softs::FPBDBendingSpringConstraints> BendingConstraints;
 		TSharedPtr<Softs::FXPBDBendingSpringConstraints> XBendingConstraints;
 		TSharedPtr<Softs::FPBDBendingConstraints> BendingElementConstraints;
@@ -267,5 +269,7 @@ PRAGMA_ENABLE_DEPRECATION_WARNINGS
 		int32 PostCollisionConstraintRuleOffset;
 		int32 PostprocessingConstraintRuleOffset;
 		//~ End PBD solver only fields
+
+		class FRuleCreator;
 	};
 } // namespace Chaos
diff --git a/Engine/Plugins/Experimental/ChaosClothAssetEditor/Source/ChaosClothAssetDataflowNodes/Private/ChaosClothAsset/ChaosClothAssetDataflowNodesModule.cpp b/Engine/Plugins/Experimental/ChaosClothAssetEditor/Source/ChaosClothAssetDataflowNodes/Private/ChaosClothAsset/ChaosClothAssetDataflowNodesModule.cpp
--- a/Engine/Plugins/Experimental/ChaosClothAssetEditor/Source/ChaosClothAssetDataflowNodes/Private/ChaosClothAsset/ChaosClothAssetDataflowNodesModule.cpp
+++ b/Engine/Plugins/Experimental/ChaosClothAssetEditor/Source/ChaosClothAssetDataflowNodes/Private/ChaosClothAsset/ChaosClothAssetDataflowNodesModule.cpp
@@ -35,6 +35,7 @@
 #include "ChaosClothAsset/SimulationVelocityScaleConfigNode.h"
 #include "ChaosClothAsset/SimulationXPBDAreaSpringConfigNode.h"
 #include "ChaosClothAsset/SimulationXPBDAnisoBendingConfigNode.h"
+#include "ChaosClothAsset/SimulationXPBDAnisoSpringConfigNode.h"
 #include "ChaosClothAsset/SimulationXPBDAnisoStretchConfigNode.h"
 #include "ChaosClothAsset/SimulationXPBDBendingElementConfigNode.h"
 #include "ChaosClothAsset/SimulationXPBDBendingSpringConfigNode.h"
@@ -87,6 +88,7 @@ namespace UE::Chaos::ClothAsset
 			DATAFLOW_NODE_REGISTER_CREATION_FACTORY(FChaosClothAssetSimulationSolverConfigNode);
 			DATAFLOW_NODE_REGISTER_CREATION_FACTORY(FChaosClothAssetSimulationVelocityScaleConfigNode);
 			DATAFLOW_NODE_REGISTER_CREATION_FACTORY(FChaosClothAssetSimulationXPBDAnisoBendingConfigNode);
+			DATAFLOW_NODE_REGISTER_CREATION_FACTORY(FChaosClothAssetSimulationXPBDAnisoSpringConfigNode);
 			DATAFLOW_NODE_REGISTER_CREATION_FACTORY(FChaosClothAssetSimulationXPBDAnisoStretchConfigNode);
 			DATAFLOW_NODE_REGISTER_CREATION_FACTORY(FChaosClothAssetSimulationXPBDAreaSpringConfigNode);
 			DATAFLOW_NODE_REGISTER_CREATION_FACTORY(FChaosClothAssetSimulationXPBDBendingElementConfigNode);
diff --git a/Engine/Plugins/Experimental/ChaosClothAssetEditor/Source/ChaosClothAssetDataflowNodes/Private/ChaosClothAsset/SimulationXPBDAnisoSpringConfigNode.cpp b/Engine/Plugins/Experimental/ChaosClothAssetEditor/Source/ChaosClothAssetDataflowNodes/Private/ChaosClothAsset/SimulationXPBDAnisoSpringConfigNode.cpp
new file mode 100644
--- /dev/null
+++ b/Engine/Plugins/Experimental/ChaosClothAssetEditor/Source/ChaosClothAssetDataflowNodes/Private/ChaosClothAsset/SimulationXPBDAnisoSpringConfigNode.cpp
@@ -0,0 +1,30 @@
+// Copyright Epic Games, Inc. All Rights Reserved.
+
+#include "ChaosClothAsset/SimulationXPBDAnisoSpringConfigNode.h"
+#include "ChaosClothAsset/SimulationBaseConfigNodePrivate.h"
+#include "Chaos/CollectionPropertyFacade.h"
+
+#include UE_INLINE_GENERATED_CPP_BY_NAME(SimulationXPBDAnisoSpringConfigNode)
+
+FChaosClothAssetSimulationXPBDAnisoSpringConfigNode::FChaosClothAssetSimulationXPBDAnisoSpringConfigNode(const Dataflow::FNodeParameters& InParam, FGuid InGuid)
+	: FChaosClothAssetSimulationBaseConfigNode(InParam, InGuid)
+{
+	RegisterCollectionConnections();
+	RegisterInputConnection(&XPBDAnisoSpringStiffnessWarp.WeightMap);
+	RegisterInputConnection(&XPBDAnisoSpringStiffnessWeft.WeightMap);
+	RegisterInputConnection(&XPBDAnisoSpringStiffnessBias.WeightMap);
+	RegisterInputConnection(&XPBDAnisoSpringDamping.WeightMap);
+	RegisterInputConnection(&XPBDAnisoSpringWarpScale.WeightMap);
+	RegisterInputConnection(&XPBDAnisoSpringWeftScale.WeightMap);
+}
+
+void FChaosClothAssetSimulationXPBDAnisoSpringConfigNode::AddProperties(Dataflow::FContext& Context, ::Chaos::Softs::FCollectionPropertyMutableFacade& Properties) const
+{
+	UE_CHAOS_CLOTHASSET_SIMULATIONCONFIG_SETPROPERTYBOOL(XPBDAnisoSpringUse3dRestLengths);
+	UE_CHAOS_CLOTHASSET_SIMULATIONCONFIG_SETPROPERTYWEIGHTED(XPBDAnisoSpringStiffnessWarp);
+	UE_CHAOS_CLOTHASSET_SIMULATIONCONFIG_SETPROPERTYWEIGHTED(XPBDAnisoSpringStiffnessWeft);
+	UE_CHAOS_CLOTHASSET_SIMULATIONCONFIG_SETPROPERTYWEIGHTED(XPBDAnisoSpringStiffnessBias);
+	UE_CHAOS_CLOTHASSET_SIMULATIONCONFIG_SETPROPERTYWEIGHTED(XPBDAnisoSpringDamping);
+	UE_CHAOS_CLOTHASSET_SIMULATIONCONFIG_SETPROPERTYWEIGHTED(XPBDAnisoSpringWarpScale);
+	UE_CHAOS_CLOTHASSET_SIMULATIONCONFIG_SETPROPERTYWEIGHTED(XPBDAnisoSpringWeftScale);
+}
diff --git a/Engine/Plugins/Experimental/ChaosClothAssetEditor/Source/ChaosClothAssetDataflowNodes/Public/ChaosClothAsset/SimulationXPBDAnisoSpringConfigNode.h b/Engine/Plugins/Experimental/ChaosClothAssetEditor/Source/ChaosClothAssetDataflowNodes/Public/ChaosClothAsset/SimulationXPBDAnisoSpringConfigNode.h
new file mode 100644
--- /dev/null
+++ b/Engine/Plugins/Experimental/ChaosClothAssetEditor/Source/ChaosClothAssetDataflowNodes/Public/ChaosClothAsset/SimulationXPBDAnisoSpringConfigNode.h
@@ -0,0 +1,79 @@
+// Copyright Epic Games, Inc. All Rights Reserved.
+
+#pragma once 
+
+#include "ChaosClothAsset/SimulationBaseConfigNode.h"
+#include "ChaosClothAsset/WeightedValue.h"
+#include "SimulationXPBDAnisoSpringConfigNode.generated.h"
+
+/** XPBD anisotropic spring constraint property configuration node. */
+USTRUCT(Meta = (DataflowCloth))
+struct FChaosClothAssetSimulationXPBDAnisoSpringConfigNode : public FChaosClothAssetSimulationBaseConfigNode
+{
+	GENERATED_USTRUCT_BODY()
+	DATAFLOW_NODE_DEFINE_INTERNAL(FChaosClothAssetSimulationXPBDAnisoSpringConfigNode, "SimulationXPBDAnisoSpringConfig", "Cloth", "Cloth Simulation XPBD Anisotropic Spring Config")
+
+public:
+	/** Whether to use the 3D draped space as rest lengths, or use the 2D pattern space instead. */
+	UPROPERTY(EditAnywhere, Category = "XPBDAnisoSpring Properties")
+	bool bXPBDAnisoSpringUse3dRestLengths = true;
+
+	/**
+	 * The stiffness of the stretch constraints in the warp (vertical) direction.
+	 * If a valid weight map is found with the given Weight Map name, then both Low and High values
+	 * are interpolated with the per particle weight to make the final value used for the simulation.
+	 * Otherwise all particles are considered to have a zero weight, and only the Low value is meaningful.
+	 */
+	UPROPERTY(EditAnywhere, Category = "XPBDAnisoSpring Properties", Meta = (UIMin = "0", UIMax = "10000", ClampMin = "0", ClampMax = "10000000"))
+	FChaosClothAssetWeightedValue XPBDAnisoSpringStiffnessWarp = { true, 100.f, 100.f, TEXT("XPBDAnisoSpringStiffnessWarp") };
+
+	/**
+	 * The stiffness of the stretch constraints in the weft (horizontal) direction.
+	 * If a valid weight map is found with the given Weight Map name, then both Low and High values
+	 * are interpolated with the per particle weight to make the final value used for the simulation.
+	 * Otherwise all particles are considered to have a zero weight, and only the Low value is meaningful.
+	 */
+	UPROPERTY(EditAnywhere, Category = "XPBDAnisoSpring Properties", Meta = (UIMin = "0", UIMax = "10000", ClampMin = "0", ClampMax = "10000000"))
+	FChaosClothAssetWeightedValue XPBDAnisoSpringStiffnessWeft = { true, 100.f, 100.f, TEXT("XPBDAnisoSpringStiffnessWeft") };
+
+	/**
+	 * The stiffness of the stretch constraints in the bias (diagonal) direction.
+	 * If a valid weight map is found with the given Weight Map name, then both Low and High values
+	 * are interpolated with the per particle weight to make the final value used for the simulation.
+	 * Otherwise all particles are considered to have a zero weight, and only the Low value is meaningful.
+	 */
+	UPROPERTY(EditAnywhere, Category = "XPBDAnisoSpring Properties", Meta = (UIMin = "0", UIMax = "10000", ClampMin = "0", ClampMax = "10000000"))
+	FChaosClothAssetWeightedValue XPBDAnisoSpringStiffnessBias = { true, 100.f, 100.f, TEXT("XPBDAnisoSpringStiffnessBias") };
+
+	/**
+	 * The damping of the stretch constraints, relative to critical damping.
+	 * If a valid weight map is found with the given Weight Map name, then both Low and High values
+	 * are interpolated with the per particle weight to make the final value used for the simulation.
+	 * Otherwise all particles are considered to have a zero weight, and only the Low value is meaningful.
+	 */
+	UPROPERTY(EditAnywhere, Category = "XPBDAnisoSpring Properties", Meta = (UIMin = "0", UIMax = "10", ClampMin = "0", ClampMax = "1000"))
+	FChaosClothAssetWeightedValue XPBDAnisoSpringDamping = { true, 1.f, 1.f, TEXT("XPBDAnisoSpringDamping") };
+
+	/**
+	 * The scale of the stretch constraints at rest in the warp (vertical) direction.
+	 * If a valid weight map is found with the given Weight Map name, then both Low and High values
+	 * are interpolated with the per particle weight to make the final value used for the simulation.
+	 * Otherwise all particles are considered to have a zero weight, and only the Low value is meaningful.
+	 */
+	UPROPERTY(EditAnywhere, Category = "XPBDAnisoSpring Properties", Meta = (UIMin = "0", UIMax = "10", ClampMin = "0", ClampMax = "10000000"))
+	FChaosClothAssetWeightedValue XPBDAnisoSpringWarpScale = { true, 1.f, 1.f, TEXT("XPBDAnisoSpringWarpScale") };
+
+	/**
+	 * The scale of the stretch constraints at rest in the weft (horizontal) direction.
+	 * If a valid weight map is found with the given Weight Map name, then both Low and High values
+	 * are interpolated with the per particle weight to make the final value used for the simulation.
+	 * Otherwise all particles are considered to have a zero weight, and only the Low value is meaningful.
+	 */
+	UPROPERTY(EditAnywhere, Category = "XPBDAnisoSpring Properties", Meta = (UIMin = "0", UIMax = "10", ClampMin = "0", ClampMax = "10000000"))
+	FChaosClothAssetWeightedValue XPBDAnisoSpringWeftScale = { true, 1.f, 1.f, TEXT("XPBDAnisoStretchWeftSpringScale") };
+
+	FChaosClothAssetSimulationXPBDAnisoSpringConfigNode(const Dataflow::FNodeParameters& InParam, FGuid InGuid = FGuid::NewGuid());
+
+private:
+	virtual void AddProperties(Dataflow::FContext& Context, ::Chaos::Softs::FCollectionPropertyMutableFacade& Properties) const override;
+};
diff --git a/Engine/Source/Runtime/Experimental/Chaos/Private/Chaos/SoftsSpring.isph b/Engine/Source/Runtime/Experimental/Chaos/Private/Chaos/SoftsSpring.isph
--- a/Engine/Source/Runtime/Experimental/Chaos/Private/Chaos/SoftsSpring.isph
+++ b/Engine/Source/Runtime/Experimental/Chaos/Private/Chaos/SoftsSpring.isph
@@ -36,6 +36,27 @@ inline FVector3f GetXPBDSpringDelta(const FVector3f& P1, const FVector3f& P2, co
 	return VectorSelect((InvM2 == FLOAT_ZERO && InvM1 == FLOAT_ZERO), FloatZeroVector, ComputedDelta);
 }
 
+template <typename StiffnessFloatType>
+inline FVector3f GetXPBDAxialSpringDelta(const FVector3f& P1, const FVector3f& P2, const FVector3f& P3, const varying float InvM1, const varying float InvM2, const varying float InvM3,
+	 const uniform float Dt, const varying float Bary, const varying float Dist, const varying float Lambda, const StiffnessFloatType Stiffness, const uniform float MinStiffness, float& DLambda)
+{
+	FVector3f P = (P2 - P3) * Bary + P3;
+	FVector3f Direction = P1 - P;
+	const float Distance = SafeNormalize(Direction);
+	const varying float Offset = Distance - Dist;
+	
+	const varying float AlphaInv = Stiffness * Dt * Dt;
+	
+	const float CombinedInvMass = (1.f - Bary) * InvM3 + Bary * InvM2 + InvM1;
+
+	const float ComputedDLambda = (AlphaInv * Offset - Lambda) / (AlphaInv * CombinedInvMass + (uniform float)1.f);
+	DLambda = select(Stiffness < MinStiffness, 0.f, ComputedDLambda);
+
+	const FVector3f ComputedDelta = DLambda * Direction;
+
+	return VectorSelect((CombinedInvMass == FLOAT_ZERO), FloatZeroVector, ComputedDelta);
+}
+
 template <typename StiffnessFloatType>
 inline FVector3f GetXPBDSpringDeltaWithDamping(const FVector3f& P1, const FVector3f& P2, const FVector3f& X1, const FVector3f& X2,
 	const varying float InvM1, const varying float InvM2,
@@ -60,4 +81,32 @@ inline FVector3f GetXPBDSpringDeltaWithDamping(const FVector3f& P1, const FVecto
 
 	return VectorSelect((InvM2 == FLOAT_ZERO && InvM1 == FLOAT_ZERO), FloatZeroVector, ComputedDelta);
 }
+
+template <typename StiffnessFloatType>
+inline FVector3f GetXPBDAxialSpringDeltaWithDamping(const FVector3f& P1, const FVector3f& P2, const FVector3f& P3, 
+	const FVector3f& X1, const FVector3f& X2, const FVector3f& X3,
+	const varying float InvM1, const varying float InvM2, const varying float InvM3,
+	const uniform float Dt, const varying float Bary, const varying float Dist, const varying float Lambda, 
+	const StiffnessFloatType Stiffness, const StiffnessFloatType DampingRatio, const uniform float MinStiffness, float& DLambda)
+{
+	FVector3f P = (P2 - P3) * Bary + P3;
+	FVector3f Direction = P1 - P;
+	const float Distance = SafeNormalize(Direction);
+	const varying float Offset = Distance - Dist;
+	const varying float OneOverDist = select(Dist< FLOAT_SMALL_NUMBER, 1.f, 1.f / Dist);
+	
+	const varying float AlphaInv = Stiffness * Dt * Dt;
+	
+	const float CombinedInvMass = (1.f - Bary) * InvM3 + Bary * InvM2 + InvM1;
+	const varying float Damping = (uniform float)2.f * DampingRatio * sqrt(Stiffness / CombinedInvMass) * OneOverDist;
+	const varying float BetaDt = Damping * Dt;
+	const varying FVector3f RelativeVelocityTimesDt = P1 - X1 - P2 + X2;
+
+	const float ComputedDLambda = (AlphaInv * Offset - Lambda + BetaDt * VectorDot(Direction, RelativeVelocityTimesDt)) / ((AlphaInv + BetaDt)*CombinedInvMass + (uniform float)1.f);
+	DLambda = select(Stiffness < MinStiffness, 0.f, ComputedDLambda);
+
+	const FVector3f ComputedDelta = DLambda * Direction;
+
+	return VectorSelect((CombinedInvMass == FLOAT_ZERO), FloatZeroVector, ComputedDelta);
+}
 #endif
diff --git a/Engine/Source/Runtime/Experimental/Chaos/Private/Chaos/XPBDAnisotropicSpringConstraints.cpp b/Engine/Source/Runtime/Experimental/Chaos/Private/Chaos/XPBDAnisotropicSpringConstraints.cpp
new file mode 100644
--- /dev/null
+++ b/Engine/Source/Runtime/Experimental/Chaos/Private/Chaos/XPBDAnisotropicSpringConstraints.cpp
@@ -0,0 +1,1419 @@
+// Copyright Epic Games, Inc. All Rights Reserved.
+#include "Chaos/XPBDAnisotropicSpringConstraints.h"
+#include "Chaos/XPBDSpringConstraints.h"
+#include "Chaos/GraphColoring.h"
+#include "Chaos/Framework/Parallel.h"
+#include "Chaos/SoftsSpring.h"
+#include "Chaos/TriangleMesh.h"
+#include "ChaosStats.h"
+
+#if INTEL_ISPC
+#include "XPBDAnisotropicSpringConstraints.ispc.generated.h"
+#endif
+
+DECLARE_CYCLE_STAT(TEXT("Chaos XPBD Anisotropic Spring Constraint"), STAT_XPBD_AnisoSpring, STATGROUP_Chaos);
+
+#if INTEL_ISPC && !UE_BUILD_SHIPPING
+static_assert(sizeof(ispc::FVector4f) == sizeof(Chaos::Softs::FPAndInvM), "sizeof(ispc::FVector4f) != sizeof(Chaos::Softs::FPAndInvM");
+static_assert(sizeof(ispc::FVector3f) == sizeof(Chaos::Softs::FSolverVec3), "sizeof(ispc::FVector3f) == sizeof(Chaos::Softs::FSolverVec3)");
+static_assert(sizeof(ispc::FIntVector2) == sizeof(Chaos::TVec2<int32>), "sizeof(ispc::FIntVector2) != sizeof(Chaos::TVec2<int32>)");
+#endif
+
+namespace Chaos::Softs {
+
+// @todo(chaos): the parallel threshold (or decision to run parallel) should probably be owned by the solver and passed to the constraint container
+extern int32 Chaos_XPBDSpring_ParallelConstraintCount;
+
+namespace Private
+{
+static void UpdateDists(const TArray<FSolverReal>& BaseDists, const TArray<FSolverVec2>& WarpWeftScaleBaseMultipliers, const FPBDWeightMap& WarpScale, const FPBDWeightMap& WeftScale, TArray<FSolverReal>& Dists)
+{
+	const bool WarpScaleHasWeightMap = WarpScale.HasWeightMap();
+	const bool WeftScaleHasWeightMap = WeftScale.HasWeightMap();
+	const FSolverReal WarpScaleNoMap = (FSolverReal)WarpScale;
+	const FSolverReal WeftScaleNoMap = (FSolverReal)WeftScale;
+	for (int32 ConstraintIndex = 0; ConstraintIndex < BaseDists.Num(); ++ConstraintIndex)
+	{
+		const FSolverReal WarpScaleValue = WarpScaleHasWeightMap ? WarpScale[ConstraintIndex] : WarpScaleNoMap;
+		const FSolverReal WeftScaleValue = WeftScaleHasWeightMap ? WeftScale[ConstraintIndex] : WeftScaleNoMap;
+
+		Dists[ConstraintIndex] = BaseDists[ConstraintIndex] * FMath::Sqrt(FMath::Square(WeftScaleValue * WarpWeftScaleBaseMultipliers[ConstraintIndex][0]) + FMath::Square(WarpScaleValue * WarpWeftScaleBaseMultipliers[ConstraintIndex][1]));
+	}
+}
+}
+
+FXPBDAnisotropicEdgeSpringConstraints::FXPBDAnisotropicEdgeSpringConstraints(
+	const FSolverParticlesRange& Particles,
+	const FTriangleMesh& TriangleMesh,
+	const TArray<TVec3<FVec2f>>& FaceVertexPatternPositions,
+	bool bUse3dRestLengths,
+	const TConstArrayView<FRealSingle>& StiffnessWarpMultipliers,
+	const TConstArrayView<FRealSingle>& StiffnessWeftMultipliers,
+	const TConstArrayView<FRealSingle>& StiffnessBiasMultipliers,
+	const TConstArrayView<FRealSingle>& DampingMultipliers,
+	const TConstArrayView<FRealSingle>& WarpScaleMultipliers,
+	const TConstArrayView<FRealSingle>& WeftScaleMultipliers,
+	const FSolverVec2& InStiffnessWarp,
+	const FSolverVec2& InStiffnessWeft,
+	const FSolverVec2& InStiffnessBias,
+	const FSolverVec2& InDampingRatio,
+	const FSolverVec2& InWarpScale,
+	const FSolverVec2& InWeftScale,
+	bool bTrimKinematicConstraints)
+	: Base(Particles,
+		TriangleMesh.GetElements(),
+		StiffnessWarpMultipliers,
+		InStiffnessWarp,
+		bTrimKinematicConstraints,
+		MaxStiffness)
+	, StiffnessWeft(InStiffnessWeft,
+		StiffnessWeftMultipliers,
+		TConstArrayView<TVec2<int32>>(Constraints),
+		ParticleOffset,
+		ParticleCount,
+		FPBDStiffness::DefaultTableSize,
+		FPBDStiffness::DefaultParameterFitBase,
+		MaxStiffness)
+	, StiffnessBias(InStiffnessBias,
+		StiffnessBiasMultipliers,
+		TConstArrayView<TVec2<int32>>(Constraints),
+		ParticleOffset,
+		ParticleCount,
+		FPBDStiffness::DefaultTableSize,
+		FPBDStiffness::DefaultParameterFitBase,
+		MaxStiffness)
+	, DampingRatio(
+		InDampingRatio.ClampAxes(MinDampingRatio, MaxDampingRatio),
+		DampingMultipliers,
+		TConstArrayView<TVec2<int32>>(Constraints),
+		ParticleOffset,
+		ParticleCount)
+	, WarpScale(
+		InWarpScale.ClampAxes(MinWarpWeftScale, MaxWarpWeftScale),
+		WarpScaleMultipliers,
+		TConstArrayView<TVec2<int32>>(Constraints),
+		ParticleOffset,
+		ParticleCount)
+	, WeftScale(
+		InWeftScale.ClampAxes(MinWarpWeftScale, MaxWarpWeftScale),
+		WeftScaleMultipliers,
+		TConstArrayView<TVec2<int32>>(Constraints),
+		ParticleOffset,
+		ParticleCount)
+{
+	Init(); // Lambda initialization
+	InitFromPatternData(bUse3dRestLengths, FaceVertexPatternPositions, TriangleMesh);
+	InitColor(Particles);
+}
+
+FXPBDAnisotropicEdgeSpringConstraints::FXPBDAnisotropicEdgeSpringConstraints(
+	const FSolverParticles& Particles,
+	int32 InParticleOffset,
+	int32 InParticleCount,
+	const FTriangleMesh& TriangleMesh,
+	const TArray<TVec3<FVec2f>>& FaceVertexPatternPositions,
+	bool bUse3dRestLengths,
+	const TConstArrayView<FRealSingle>& StiffnessWarpMultipliers,
+	const TConstArrayView<FRealSingle>& StiffnessWeftMultipliers,
+	const TConstArrayView<FRealSingle>& StiffnessBiasMultipliers,
+	const TConstArrayView<FRealSingle>& DampingMultipliers,
+	const TConstArrayView<FRealSingle>& WarpScaleMultipliers,
+	const TConstArrayView<FRealSingle>& WeftScaleMultipliers,
+	const FSolverVec2& InStiffnessWarp,
+	const FSolverVec2& InStiffnessWeft,
+	const FSolverVec2& InStiffnessBias,
+	const FSolverVec2& InDampingRatio,
+	const FSolverVec2& InWarpScale,
+	const FSolverVec2& InWeftScale,
+	bool bTrimKinematicConstraints)
+	: Base(Particles,
+		InParticleOffset,
+		InParticleCount,
+		TriangleMesh.GetElements(),
+		StiffnessWarpMultipliers,
+		InStiffnessWarp,
+		bTrimKinematicConstraints,
+		MaxStiffness)
+	, StiffnessWeft(InStiffnessWeft,
+		StiffnessWeftMultipliers,
+		TConstArrayView<TVec2<int32>>(Constraints),
+		ParticleOffset,
+		ParticleCount,
+		FPBDStiffness::DefaultTableSize,
+		FPBDStiffness::DefaultParameterFitBase,
+		MaxStiffness)
+	, StiffnessBias(InStiffnessBias,
+		StiffnessBiasMultipliers,
+		TConstArrayView<TVec2<int32>>(Constraints),
+		ParticleOffset,
+		ParticleCount,
+		FPBDStiffness::DefaultTableSize,
+		FPBDStiffness::DefaultParameterFitBase,
+		MaxStiffness)
+	, DampingRatio(
+		InDampingRatio.ClampAxes(MinDampingRatio, MaxDampingRatio),
+		DampingMultipliers,
+		TConstArrayView<TVec2<int32>>(Constraints),
+		ParticleOffset,
+		ParticleCount)
+	, WarpScale(
+		InWarpScale.ClampAxes(MinWarpWeftScale, MaxWarpWeftScale),
+		WarpScaleMultipliers,
+		TConstArrayView<TVec2<int32>>(Constraints),
+		ParticleOffset,
+		ParticleCount)
+	, WeftScale(
+		InWeftScale.ClampAxes(MinWarpWeftScale, MaxWarpWeftScale),
+		WeftScaleMultipliers,
+		TConstArrayView<TVec2<int32>>(Constraints),
+		ParticleOffset,
+		ParticleCount)
+{
+	Init(); // Lambda initialization
+	InitFromPatternData(bUse3dRestLengths, FaceVertexPatternPositions, TriangleMesh);
+	InitColor(Particles);
+}
+
+void FXPBDAnisotropicEdgeSpringConstraints::InitFromPatternData(bool bUse3dRestLengths, const TArray<TVec3<FVec2f>>& FaceVertexPatternPositions, const FTriangleMesh& TriangleMesh)
+{
+	// Calculate multipliers per face. Average the multipliers for faces that share the same edge.
+	struct FMultipliers
+	{
+		FSolverVec3 WarpWeftBias;
+		FSolverVec2 WarpWeftScale;
+		FSolverReal Length;
+	};
+	TMap<TVec2<int32> /*Edge*/, TArray<FMultipliers>> EdgeToMultiplierMap;
+	auto SortedEdge = [](int32 P0, int32 P1) { return P0 <= P1 ? TVec2<int32>(P0, P1) : TVec2<int32>(P1, P0); };
+	auto Multiplier = [](const FVec2f& UV0, const FVec2f& UV1)
+	{
+		// Input UVs: X = Weft, Y = Warp direction
+		// Transform to X = Weft, Y = Warp, Z = Bias for WarpWeftBias multiplier
+		// 
+		const FSolverVec2 UVDiff = UV1 - UV0;
+		const FSolverVec2 UVDiffAbs = UVDiff.GetAbs();
+		FSolverVec3 UVDiffTransformed;
+		if (UVDiffAbs.X > UVDiffAbs.Y)
+		{
+			// Stiffness is blend between weft and bias
+			UVDiffTransformed = FSolverVec3(UVDiffAbs.X - UVDiffAbs.Y, (FSolverReal)0.f, UVDiffAbs.Y);
+		}
+		else
+		{
+			// Stiffness is blend between warp and bias
+			UVDiffTransformed = FSolverVec3((FSolverReal)0.f, UVDiffAbs.Y - UVDiffAbs.X, UVDiffAbs.X);
+		}
+		const FSolverReal Denom3 = UVDiffTransformed.X + UVDiffTransformed.Y + UVDiffTransformed.Z; // These are used as weights
+		const FSolverReal UVLength = UVDiffAbs.Length();
+		return FMultipliers
+		{
+			Denom3 > UE_SMALL_NUMBER ? UVDiffTransformed / Denom3 : FSolverVec3(0.f, 1.f, 0.f), // Default to Warp if zero length
+			UVLength > UE_SMALL_NUMBER ? UVDiffAbs / UVLength : FSolverVec2(UE_INV_SQRT_2,UE_INV_SQRT_2), // Default to equally scaling warp and weft directions if zero length
+			UVLength
+		};
+	};
+
+	const TArray<TVec3<int32>>& Elements = TriangleMesh.GetElements();
+	for (int32 ElemIdx = 0; ElemIdx < Elements.Num(); ++ElemIdx)
+	{
+		const TVec3<int32>& Element = Elements[ElemIdx];
+		const TVec3<FVec2f>& UVs = FaceVertexPatternPositions[ElemIdx];
+		EdgeToMultiplierMap.FindOrAdd(SortedEdge(Element[0], Element[1])).Add(Multiplier(UVs[0], UVs[1]));
+		EdgeToMultiplierMap.FindOrAdd(SortedEdge(Element[1], Element[2])).Add(Multiplier(UVs[1], UVs[2]));
+		EdgeToMultiplierMap.FindOrAdd(SortedEdge(Element[2], Element[0])).Add(Multiplier(UVs[2], UVs[0]));
+	}
+
+	if (bUse3dRestLengths)
+	{
+		BaseDists = Dists; // Use Dists calculated by base class using 3D positions
+	}
+	else
+	{
+		// We'll set them from the computed 2D lengths above
+		BaseDists.SetNumZeroed(Constraints.Num());
+	}
+	WarpWeftBiasBaseMultipliers.SetNumZeroed(Constraints.Num());
+	WarpWeftScaleBaseMultipliers.SetNumZeroed(Constraints.Num());
+	for (int32 ConstraintIdx = 0; ConstraintIdx < Constraints.Num(); ++ConstraintIdx)
+	{
+		const TArray<FMultipliers>& EdgeMultipliers = EdgeToMultiplierMap.FindChecked(SortedEdge(Constraints[ConstraintIdx][0], Constraints[ConstraintIdx][1]));
+		check(EdgeMultipliers.Num());
+		for (const FMultipliers& EdgeMult : EdgeMultipliers)
+		{
+			WarpWeftBiasBaseMultipliers[ConstraintIdx] += EdgeMult.WarpWeftBias;
+			WarpWeftScaleBaseMultipliers[ConstraintIdx] += EdgeMult.WarpWeftScale;
+			if (!bUse3dRestLengths)
+			{
+				BaseDists[ConstraintIdx] += EdgeMult.Length;
+			}
+		}
+		WarpWeftBiasBaseMultipliers[ConstraintIdx] /= (FSolverReal)EdgeMultipliers.Num();
+		WarpWeftScaleBaseMultipliers[ConstraintIdx].Normalize();
+		if (!bUse3dRestLengths)
+		{
+			BaseDists[ConstraintIdx] /= (FSolverReal)EdgeMultipliers.Num();
+			Dists[ConstraintIdx] = BaseDists[ConstraintIdx];
+		}
+	}
+}
+
+template<typename SolverParticlesOrRange>
+void FXPBDAnisotropicEdgeSpringConstraints::InitColor(const SolverParticlesOrRange& InParticles)
+{
+	// In dev builds we always color so we can tune the system without restarting. See Apply()
+#if UE_BUILD_SHIPPING || UE_BUILD_TEST
+	if (Constraints.Num() > Chaos_XPBDSpring_ParallelConstraintCount)
+#endif
+	{
+		const TArray<TArray<int32>> ConstraintsPerColor = FGraphColoring::ComputeGraphColoringParticlesOrRange(Constraints, InParticles, ParticleOffset, ParticleOffset + ParticleCount);
+
+		// Reorder constraints based on color so each array in ConstraintsPerColor contains contiguous elements.
+		TArray<TVec2<int32>> ReorderedConstraints;
+		TArray<FSolverReal> ReorderedDists;
+		TArray<FSolverReal> ReorderedBaseDists;
+		TArray<FSolverVec3> ReorderedWarpWeftBiasBaseMultipliers;
+		TArray<FSolverVec2> ReorderedWarpWeftScaleBaseMultipliers;
+		TArray<int32> OrigToReorderedIndices; // used to reorder stiffness indices
+		ReorderedConstraints.SetNumUninitialized(Constraints.Num());
+		ReorderedDists.SetNumUninitialized(Constraints.Num());
+		ReorderedBaseDists.SetNumUninitialized(Constraints.Num());
+		ReorderedWarpWeftBiasBaseMultipliers.SetNumUninitialized(Constraints.Num());
+		ReorderedWarpWeftScaleBaseMultipliers.SetNumUninitialized(Constraints.Num());
+		OrigToReorderedIndices.SetNumUninitialized(Constraints.Num());
+
+		ConstraintsPerColorStartIndex.Reset(ConstraintsPerColor.Num() + 1);
+
+		int32 ReorderedIndex = 0;
+		for (const TArray<int32>& ConstraintsBatch : ConstraintsPerColor)
+		{
+			ConstraintsPerColorStartIndex.Add(ReorderedIndex);
+			for (const int32& BatchConstraint : ConstraintsBatch)
+			{
+				const int32 OrigIndex = BatchConstraint;
+				ReorderedConstraints[ReorderedIndex] = Constraints[OrigIndex];
+				ReorderedDists[ReorderedIndex] = Dists[OrigIndex];
+				ReorderedBaseDists[ReorderedIndex] = BaseDists[OrigIndex];
+				ReorderedWarpWeftBiasBaseMultipliers[ReorderedIndex] = WarpWeftBiasBaseMultipliers[OrigIndex];
+				ReorderedWarpWeftScaleBaseMultipliers[ReorderedIndex] = WarpWeftScaleBaseMultipliers[OrigIndex];
+				OrigToReorderedIndices[OrigIndex] = ReorderedIndex;
+				++ReorderedIndex;
+			}
+		}
+		ConstraintsPerColorStartIndex.Add(ReorderedIndex);
+
+		Constraints = MoveTemp(ReorderedConstraints);
+		Dists = MoveTemp(ReorderedDists);
+		BaseDists = MoveTemp(ReorderedBaseDists);
+		WarpWeftBiasBaseMultipliers = MoveTemp(ReorderedWarpWeftBiasBaseMultipliers);
+		WarpWeftScaleBaseMultipliers = MoveTemp(ReorderedWarpWeftScaleBaseMultipliers);
+		Stiffness.ReorderIndices(OrigToReorderedIndices);
+		StiffnessWeft.ReorderIndices(OrigToReorderedIndices);
+		StiffnessBias.ReorderIndices(OrigToReorderedIndices);
+		DampingRatio.ReorderIndices(OrigToReorderedIndices);
+		WarpScale.ReorderIndices(OrigToReorderedIndices);
+		WeftScale.ReorderIndices(OrigToReorderedIndices);
+	}
+}
+
+void FXPBDAnisotropicEdgeSpringConstraints::UpdateDists()
+{
+	Private::UpdateDists(BaseDists, WarpWeftScaleBaseMultipliers, WarpScale, WeftScale, Dists);
+}
+
+template<typename SolverParticlesOrRange>
+void FXPBDAnisotropicEdgeSpringConstraints::ApplyHelper(SolverParticlesOrRange& Particles, const FSolverReal Dt, const int32 ConstraintIndex, const FSolverVec3& ExpStiffnessValues, const FSolverReal DampingRatioValue) const
+{
+	const TVec2<int32>& Constraint = Constraints[ConstraintIndex];
+	const int32 Index1 = Constraint[0];
+	const int32 Index2 = Constraint[1];
+
+	const FSolverVec3& WarpWeftBiasBaseMultiplier = WarpWeftBiasBaseMultipliers[ConstraintIndex];
+	const FSolverReal FinalStiffnessValue = WarpWeftBiasBaseMultiplier.Dot(ExpStiffnessValues);
+
+	const FSolverVec3 Delta = Spring::GetXPBDSpringDelta(Particles, Dt, Constraint, Dists[ConstraintIndex], Lambdas[ConstraintIndex],
+		FinalStiffnessValue, MinStiffness, DampingRatioValue);
+
+	if (Particles.InvM(Index1) > (FSolverReal)0.)
+	{
+		Particles.P(Index1) -= Particles.InvM(Index1) * Delta;
+	}
+	if (Particles.InvM(Index2) > (FSolverReal)0.)
+	{
+		Particles.P(Index2) += Particles.InvM(Index2) * Delta;
+	}
+}
+
+template<typename SolverParticlesOrRange>
+void FXPBDAnisotropicEdgeSpringConstraints::Apply(SolverParticlesOrRange& Particles, const FSolverReal Dt) const
+{
+	TRACE_CPUPROFILER_EVENT_SCOPE(FXPBDAnisotropicEdgeSpringConstraints_Apply);
+	SCOPE_CYCLE_COUNTER(STAT_XPBD_AnisoSpring);
+
+	const bool StiffnessHasWeightMap = Stiffness.HasWeightMap();
+	const bool StiffnessWeftHasWeightMap = StiffnessWeft.HasWeightMap();
+	const bool StiffnessBiasHasWeightMap = StiffnessBias.HasWeightMap();
+	const bool DampingHasWeightMap = DampingRatio.HasWeightMap();
+
+	if ((ConstraintsPerColorStartIndex.Num() > 1) && (Constraints.Num() > Chaos_XPBDSpring_ParallelConstraintCount))
+	{
+		const int32 ConstraintColorNum = ConstraintsPerColorStartIndex.Num() - 1;
+		if (!StiffnessHasWeightMap && !StiffnessWeftHasWeightMap && !StiffnessBiasHasWeightMap && !DampingHasWeightMap)
+		{
+			const FSolverVec3 ExpStiffnessValue((FSolverReal)StiffnessWeft, (FSolverReal)Stiffness, (FSolverReal)StiffnessBias);
+			const FSolverReal DampingRatioValue = (FSolverReal)DampingRatio;
+			if (ExpStiffnessValue.Max() < MinStiffness)
+			{
+				return;
+			}
+
+#if INTEL_ISPC
+			if (bRealTypeCompatibleWithISPC && bChaos_XPBDSpring_ISPC_Enabled)
+			{
+				if (DampingRatioValue > 0)
+				{
+					for (int32 ConstraintColorIndex = 0; ConstraintColorIndex < ConstraintColorNum; ++ConstraintColorIndex)
+					{
+						const int32 ColorStart = ConstraintsPerColorStartIndex[ConstraintColorIndex];
+						const int32 ColorSize = ConstraintsPerColorStartIndex[ConstraintColorIndex + 1] - ColorStart;
+						ispc::ApplyXPBDAnisoSpringConstraintsWithDamping(
+							(ispc::FVector4f*)Particles.GetPAndInvM().GetData(),
+							(const ispc::FVector3f*)Particles.XArray().GetData(),
+							(ispc::FIntVector2*)&Constraints.GetData()[ColorStart],
+							&Dists.GetData()[ColorStart],
+							(const ispc::FVector3f*)&WarpWeftBiasBaseMultipliers.GetData()[ColorStart],
+							&Lambdas.GetData()[ColorStart],
+							Dt,
+							reinterpret_cast<const ispc::FVector3f&>(ExpStiffnessValue),
+							DampingRatioValue,
+							ColorSize);
+					}
+				}
+				else
+				{
+					for (int32 ConstraintColorIndex = 0; ConstraintColorIndex < ConstraintColorNum; ++ConstraintColorIndex)
+					{
+						const int32 ColorStart = ConstraintsPerColorStartIndex[ConstraintColorIndex];
+						const int32 ColorSize = ConstraintsPerColorStartIndex[ConstraintColorIndex + 1] - ColorStart;
+						ispc::ApplyXPBDAnisoSpringConstraints(
+							(ispc::FVector4f*)Particles.GetPAndInvM().GetData(),
+							(ispc::FIntVector2*)&Constraints.GetData()[ColorStart],
+							&Dists.GetData()[ColorStart],
+							(const ispc::FVector3f*)&WarpWeftBiasBaseMultipliers.GetData()[ColorStart],
+							&Lambdas.GetData()[ColorStart],
+							Dt,
+							reinterpret_cast<const ispc::FVector3f&>(ExpStiffnessValue),
+							ColorSize);
+					}
+				}
+			}
+			else
+#endif
+			{
+				for (int32 ConstraintColorIndex = 0; ConstraintColorIndex < ConstraintColorNum; ++ConstraintColorIndex)
+				{
+					const int32 ColorStart = ConstraintsPerColorStartIndex[ConstraintColorIndex];
+					const int32 ColorSize = ConstraintsPerColorStartIndex[ConstraintColorIndex + 1] - ColorStart;
+					PhysicsParallelFor(ColorSize, [this, &Particles, Dt, ColorStart, &ExpStiffnessValue, DampingRatioValue](const int32 Index)
+					{
+						const int32 ConstraintIndex = ColorStart + Index;
+						ApplyHelper(Particles, Dt, ConstraintIndex, ExpStiffnessValue, DampingRatioValue);
+					});
+				}
+			}
+		}
+		else  // Has weight maps
+		{
+#if INTEL_ISPC
+			if (bRealTypeCompatibleWithISPC && bChaos_XPBDSpring_ISPC_Enabled)
+			{
+				if (DampingHasWeightMap || (FSolverReal)DampingRatio > 0)
+				{
+					for (int32 ConstraintColorIndex = 0; ConstraintColorIndex < ConstraintColorNum; ++ConstraintColorIndex)
+					{
+						const int32 ColorStart = ConstraintsPerColorStartIndex[ConstraintColorIndex];
+						const int32 ColorSize = ConstraintsPerColorStartIndex[ConstraintColorIndex + 1] - ColorStart;
+						ispc::ApplyXPBDAnisoSpringConstraintsWithDampingAndWeightMaps(
+							(ispc::FVector4f*)Particles.GetPAndInvM().GetData(),
+							(const ispc::FVector3f*)Particles.XArray().GetData(),
+							(ispc::FIntVector2*)&Constraints.GetData()[ColorStart],
+							&Dists.GetData()[ColorStart],
+							(const ispc::FVector3f*)&WarpWeftBiasBaseMultipliers.GetData()[ColorStart],
+							&Lambdas.GetData()[ColorStart],
+							Dt,
+							MinStiffness,
+							StiffnessHasWeightMap,
+							StiffnessHasWeightMap ? &Stiffness.GetIndices().GetData()[ColorStart] : nullptr,
+							&Stiffness.GetTable().GetData()[0],
+							StiffnessWeftHasWeightMap,
+							StiffnessWeftHasWeightMap ? &StiffnessWeft.GetIndices().GetData()[ColorStart] : nullptr,
+							&StiffnessWeft.GetTable().GetData()[0],
+							StiffnessBiasHasWeightMap,
+							StiffnessBiasHasWeightMap ? &StiffnessBias.GetIndices().GetData()[ColorStart] : nullptr,
+							&StiffnessBias.GetTable().GetData()[0],
+							DampingHasWeightMap,
+							DampingHasWeightMap ? &DampingRatio.GetIndices().GetData()[ColorStart] : nullptr,
+							&DampingRatio.GetTable().GetData()[0],
+							ColorSize);
+					}
+				}
+				else
+				{
+					for (int32 ConstraintColorIndex = 0; ConstraintColorIndex < ConstraintColorNum; ++ConstraintColorIndex)
+					{
+						const int32 ColorStart = ConstraintsPerColorStartIndex[ConstraintColorIndex];
+						const int32 ColorSize = ConstraintsPerColorStartIndex[ConstraintColorIndex + 1] - ColorStart;
+						ispc::ApplyXPBDAnisoSpringConstraintsWithWeightMaps(
+							(ispc::FVector4f*)Particles.GetPAndInvM().GetData(),
+							(ispc::FIntVector2*)&Constraints.GetData()[ColorStart],
+							&Dists.GetData()[ColorStart],
+							(const ispc::FVector3f*)&WarpWeftBiasBaseMultipliers.GetData()[ColorStart],
+							&Lambdas.GetData()[ColorStart],
+							Dt,
+							MinStiffness,
+							StiffnessHasWeightMap,
+							&Stiffness.GetIndices().GetData()[ColorStart],
+							&Stiffness.GetTable().GetData()[0],
+							StiffnessWeftHasWeightMap,
+							StiffnessWeftHasWeightMap ? &StiffnessWeft.GetIndices().GetData()[ColorStart] : nullptr,
+							&StiffnessWeft.GetTable().GetData()[0],
+							StiffnessBiasHasWeightMap,
+							StiffnessBiasHasWeightMap ? &StiffnessBias.GetIndices().GetData()[ColorStart] : nullptr,
+							&StiffnessBias.GetTable().GetData()[0],
+							ColorSize);
+					}
+				}
+			}
+			else
+#endif
+			{
+				const FSolverReal StiffnessNoMap = (FSolverReal)Stiffness;
+				const FSolverReal StiffnessWeftNoMap = (FSolverReal)StiffnessWeft;
+				const FSolverReal StiffnessBiasNoMap = (FSolverReal)StiffnessBias;
+				const FSolverReal DampingNoMap = (FSolverReal)DampingRatio;
+				for (int32 ConstraintColorIndex = 0; ConstraintColorIndex < ConstraintColorNum; ++ConstraintColorIndex)
+				{
+					const int32 ColorStart = ConstraintsPerColorStartIndex[ConstraintColorIndex];
+					const int32 ColorSize = ConstraintsPerColorStartIndex[ConstraintColorIndex + 1] - ColorStart;
+					PhysicsParallelFor(ColorSize, [this, &Particles, Dt, ColorStart, StiffnessHasWeightMap, StiffnessNoMap, StiffnessWeftHasWeightMap, StiffnessWeftNoMap, StiffnessBiasHasWeightMap, 
+						StiffnessBiasNoMap, DampingHasWeightMap, DampingNoMap](const int32 Index)
+					{
+						const int32 ConstraintIndex = ColorStart + Index;
+						const FSolverReal ExpStiffnessValue = StiffnessHasWeightMap ? Stiffness[ConstraintIndex] : StiffnessNoMap;
+						const FSolverReal ExpStiffnessWeftValue = StiffnessWeftHasWeightMap ? StiffnessWeft[ConstraintIndex] : StiffnessWeftNoMap;
+						const FSolverReal ExpStiffnessBiasValue = StiffnessBiasHasWeightMap ? StiffnessBias[ConstraintIndex] : StiffnessBiasNoMap;
+						const FSolverReal DampingRatioValue = DampingHasWeightMap ? DampingRatio[ConstraintIndex] : DampingNoMap;
+						ApplyHelper(Particles, Dt, ConstraintIndex, FSolverVec3(ExpStiffnessWeftValue, ExpStiffnessValue, ExpStiffnessBiasValue), DampingRatioValue);
+					});
+				}
+			}
+		}
+	}
+	else
+	{
+		if (!StiffnessHasWeightMap && !StiffnessWeftHasWeightMap && !StiffnessBiasHasWeightMap && !DampingHasWeightMap)
+		{
+			const FSolverVec3 ExpStiffnessValue((FSolverReal)StiffnessWeft, (FSolverReal)Stiffness, (FSolverReal)StiffnessBias);
+			const FSolverReal DampingRatioValue = (FSolverReal)DampingRatio;
+			if (ExpStiffnessValue.Max() < MinStiffness)
+			{
+				return;
+			}
+			for (int32 ConstraintIndex = 0; ConstraintIndex < Constraints.Num(); ++ConstraintIndex)
+			{
+				ApplyHelper(Particles, Dt, ConstraintIndex, ExpStiffnessValue, DampingRatioValue);
+			}
+		}
+		else
+		{
+			const FSolverReal StiffnessNoMap = (FSolverReal)Stiffness;
+			const FSolverReal StiffnessWeftNoMap = (FSolverReal)StiffnessWeft;
+			const FSolverReal StiffnessBiasNoMap = (FSolverReal)StiffnessBias;
+			const FSolverReal DampingNoMap = (FSolverReal)DampingRatio;
+			for (int32 ConstraintIndex = 0; ConstraintIndex < Constraints.Num(); ++ConstraintIndex)
+			{
+				const FSolverReal ExpStiffnessValue = StiffnessHasWeightMap ? Stiffness[ConstraintIndex] : StiffnessNoMap;
+				const FSolverReal ExpStiffnessWeftValue = StiffnessWeftHasWeightMap ? StiffnessWeft[ConstraintIndex] : StiffnessWeftNoMap;
+				const FSolverReal ExpStiffnessBiasValue = StiffnessBiasHasWeightMap ? StiffnessBias[ConstraintIndex] : StiffnessBiasNoMap;
+				const FSolverReal DampingRatioValue = DampingHasWeightMap ? DampingRatio[ConstraintIndex] : DampingNoMap;
+				ApplyHelper(Particles, Dt, ConstraintIndex, FSolverVec3(ExpStiffnessWeftValue, ExpStiffnessValue, ExpStiffnessBiasValue), DampingRatioValue);
+			}
+		}
+	}
+}
+template CHAOS_API void FXPBDAnisotropicEdgeSpringConstraints::Apply(FSolverParticles& Particles, const FSolverReal Dt) const;
+template CHAOS_API void FXPBDAnisotropicEdgeSpringConstraints::Apply(FSolverParticlesRange& Particles, const FSolverReal Dt) const;
+
+void FXPBDAnisotropicEdgeSpringConstraints::UpdateLinearSystem(const FSolverParticlesRange& Particles, const FSolverReal Dt, FEvolutionLinearSystem& LinearSystem) const
+{
+	TRACE_CPUPROFILER_EVENT_SCOPE(FXPBDAnisotropicEdgeSpringConstraints_UpdateLinearSystem);
+	LinearSystem.ReserveForParallelAdd(Constraints.Num() * 2, Constraints.Num());
+
+	const bool StiffnessHasWeightMap = Stiffness.HasWeightMap();
+	const bool StiffnessWeftHasWeightMap = StiffnessWeft.HasWeightMap();
+	const bool StiffnessBiasHasWeightMap = StiffnessBias.HasWeightMap();
+	const bool DampingHasWeightMap = DampingRatio.HasWeightMap();
+	if ((ConstraintsPerColorStartIndex.Num() > 1) && (Constraints.Num() > Chaos_XPBDSpring_ParallelConstraintCount))
+	{
+		const int32 ConstraintColorNum = ConstraintsPerColorStartIndex.Num() - 1;
+		if (!StiffnessHasWeightMap && !StiffnessWeftHasWeightMap && !StiffnessBiasHasWeightMap && !DampingHasWeightMap)
+		{
+			const FSolverVec3 ExpStiffnessValue((FSolverReal)StiffnessWeft, (FSolverReal)Stiffness, (FSolverReal)StiffnessBias);
+			const FSolverReal DampingRatioValue = (FSolverReal)DampingRatio;
+			if (ExpStiffnessValue.Max() < MinStiffness)
+			{
+				return;
+			}
+			for (int32 ConstraintColorIndex = 0; ConstraintColorIndex < ConstraintColorNum; ++ConstraintColorIndex)
+			{
+				const int32 ColorStart = ConstraintsPerColorStartIndex[ConstraintColorIndex];
+				const int32 ColorSize = ConstraintsPerColorStartIndex[ConstraintColorIndex + 1] - ColorStart;
+				PhysicsParallelFor(ColorSize, [this, &Particles, Dt, &LinearSystem, ColorStart, ExpStiffnessValue, DampingRatioValue](const int32 Index)
+				{
+					const int32 ConstraintIndex = ColorStart + Index;
+					const FSolverVec3& WarpWeftBiasBaseMultiplier = WarpWeftBiasBaseMultipliers[ConstraintIndex];
+					const FSolverReal FinalStiffnessValue = WarpWeftBiasBaseMultiplier.Dot(ExpStiffnessValue);
+					Spring::UpdateSpringLinearSystem(Particles, Dt, Constraints[ConstraintIndex], Dists[ConstraintIndex], FinalStiffnessValue, MinStiffness, DampingRatioValue, LinearSystem);
+				});
+			}
+		}
+		else // Has weight maps
+		{
+			const FSolverReal StiffnessNoMap = (FSolverReal)Stiffness;
+			const FSolverReal StiffnessWeftNoMap = (FSolverReal)StiffnessWeft;
+			const FSolverReal StiffnessBiasNoMap = (FSolverReal)StiffnessBias;
+			const FSolverReal DampingNoMap = (FSolverReal)DampingRatio;
+			for (int32 ConstraintColorIndex = 0; ConstraintColorIndex < ConstraintColorNum; ++ConstraintColorIndex)
+			{
+				const int32 ColorStart = ConstraintsPerColorStartIndex[ConstraintColorIndex];
+				const int32 ColorSize = ConstraintsPerColorStartIndex[ConstraintColorIndex + 1] - ColorStart;
+				PhysicsParallelFor(ColorSize, [this, &Particles, Dt, &LinearSystem, ColorStart, StiffnessHasWeightMap, StiffnessNoMap, StiffnessWeftHasWeightMap, StiffnessWeftNoMap, StiffnessBiasHasWeightMap,
+					StiffnessBiasNoMap, DampingHasWeightMap, DampingNoMap](const int32 Index)
+				{
+					const int32 ConstraintIndex = ColorStart + Index;
+					const FSolverReal ExpStiffnessValue = StiffnessHasWeightMap ? Stiffness[ConstraintIndex] : StiffnessNoMap;
+					const FSolverReal ExpStiffnessWeftValue = StiffnessWeftHasWeightMap ? StiffnessWeft[ConstraintIndex] : StiffnessWeftNoMap;
+					const FSolverReal ExpStiffnessBiasValue = StiffnessBiasHasWeightMap ? StiffnessBias[ConstraintIndex] : StiffnessBiasNoMap;
+					const FSolverReal DampingRatioValue = DampingHasWeightMap ? DampingRatio[ConstraintIndex] : DampingNoMap;
+					const FSolverVec3& WarpWeftBiasBaseMultiplier = WarpWeftBiasBaseMultipliers[ConstraintIndex];
+					const FSolverReal FinalStiffnessValue = WarpWeftBiasBaseMultiplier.Dot(FSolverVec3(ExpStiffnessWeftValue, ExpStiffnessValue, ExpStiffnessBiasValue));
+					Spring::UpdateSpringLinearSystem(Particles, Dt, Constraints[ConstraintIndex], Dists[ConstraintIndex], FinalStiffnessValue, MinStiffness, DampingRatioValue, LinearSystem);
+				});
+			}
+		}
+	}
+	else
+	{
+		if (!StiffnessHasWeightMap && !StiffnessWeftHasWeightMap && !StiffnessBiasHasWeightMap && !DampingHasWeightMap)
+		{
+			const FSolverVec3 ExpStiffnessValue((FSolverReal)StiffnessWeft, (FSolverReal)Stiffness, (FSolverReal)StiffnessBias);
+			const FSolverReal DampingRatioValue = (FSolverReal)DampingRatio;
+			if (ExpStiffnessValue.Max() < MinStiffness)
+			{
+				return;
+			}
+			for (int32 ConstraintIndex = 0; ConstraintIndex < Constraints.Num(); ++ConstraintIndex)
+			{
+				const FSolverVec3& WarpWeftBiasBaseMultiplier = WarpWeftBiasBaseMultipliers[ConstraintIndex];
+				const FSolverReal FinalStiffnessValue = WarpWeftBiasBaseMultiplier.Dot(ExpStiffnessValue);
+				Spring::UpdateSpringLinearSystem(Particles, Dt, Constraints[ConstraintIndex], Dists[ConstraintIndex], FinalStiffnessValue, MinStiffness, DampingRatioValue, LinearSystem);
+			}
+		}
+		else
+		{
+			const FSolverReal StiffnessNoMap = (FSolverReal)Stiffness;
+			const FSolverReal StiffnessWeftNoMap = (FSolverReal)StiffnessWeft;
+			const FSolverReal StiffnessBiasNoMap = (FSolverReal)StiffnessBias;
+			const FSolverReal DampingNoMap = (FSolverReal)DampingRatio;
+			for (int32 ConstraintIndex = 0; ConstraintIndex < Constraints.Num(); ++ConstraintIndex)
+			{
+				const FSolverReal ExpStiffnessValue = StiffnessHasWeightMap ? Stiffness[ConstraintIndex] : StiffnessNoMap;
+				const FSolverReal ExpStiffnessWeftValue = StiffnessWeftHasWeightMap ? StiffnessWeft[ConstraintIndex] : StiffnessWeftNoMap;
+				const FSolverReal ExpStiffnessBiasValue = StiffnessBiasHasWeightMap ? StiffnessBias[ConstraintIndex] : StiffnessBiasNoMap;
+				const FSolverReal DampingRatioValue = DampingHasWeightMap ? DampingRatio[ConstraintIndex] : DampingNoMap;
+				const FSolverVec3& WarpWeftBiasBaseMultiplier = WarpWeftBiasBaseMultipliers[ConstraintIndex];
+				const FSolverReal FinalStiffnessValue = WarpWeftBiasBaseMultiplier.Dot(FSolverVec3(ExpStiffnessWeftValue, ExpStiffnessValue, ExpStiffnessBiasValue));
+				Spring::UpdateSpringLinearSystem(Particles, Dt, Constraints[ConstraintIndex], Dists[ConstraintIndex], FinalStiffnessValue, MinStiffness, DampingRatioValue, LinearSystem);
+			}
+		}
+	}
+}
+
+FXPBDAnisotropicAxialSpringConstraints::FXPBDAnisotropicAxialSpringConstraints(
+	const FSolverParticlesRange& Particles,
+	const FTriangleMesh& TriangleMesh,
+	const TArray<TVec3<FVec2f>>& FaceVertexPatternPositions,
+	bool bUse3dRestLengths,
+	const TConstArrayView<FRealSingle>& StiffnessWarpMultipliers,
+	const TConstArrayView<FRealSingle>& StiffnessWeftMultipliers,
+	const TConstArrayView<FRealSingle>& StiffnessBiasMultipliers,
+	const TConstArrayView<FRealSingle>& DampingMultipliers,
+	const TConstArrayView<FRealSingle>& WarpScaleMultipliers,
+	const TConstArrayView<FRealSingle>& WeftScaleMultipliers,
+	const FSolverVec2& InStiffnessWarp,
+	const FSolverVec2& InStiffnessWeft,
+	const FSolverVec2& InStiffnessBias,
+	const FSolverVec2& InDampingRatio,
+	const FSolverVec2& InWarpScale,
+	const FSolverVec2& InWeftScale,
+	bool bTrimKinematicConstraints)
+	: Base(Particles,
+		TriangleMesh.GetElements(),
+		StiffnessWarpMultipliers,
+		InStiffnessWarp,
+		bTrimKinematicConstraints,
+		MaxStiffness)
+	, StiffnessWeft(InStiffnessWeft,
+		StiffnessWeftMultipliers,
+		TConstArrayView<TVec3<int32>>(Constraints),
+		ParticleOffset,
+		ParticleCount,
+		FPBDStiffness::DefaultTableSize,
+		FPBDStiffness::DefaultParameterFitBase,
+		MaxStiffness)
+	, StiffnessBias(InStiffnessBias,
+		StiffnessBiasMultipliers,
+		TConstArrayView<TVec3<int32>>(Constraints),
+		ParticleOffset,
+		ParticleCount,
+		FPBDStiffness::DefaultTableSize,
+		FPBDStiffness::DefaultParameterFitBase,
+		MaxStiffness)
+	, DampingRatio(
+		InDampingRatio.ClampAxes(MinDampingRatio, MaxDampingRatio),
+		DampingMultipliers,
+		TConstArrayView<TVec3<int32>>(Constraints),
+		ParticleOffset,
+		ParticleCount)
+	, WarpScale(
+		InWarpScale.ClampAxes(MinWarpWeftScale, MaxWarpWeftScale),
+		WarpScaleMultipliers,
+		TConstArrayView<TVec3<int32>>(Constraints),
+		ParticleOffset,
+		ParticleCount)
+	, WeftScale(
+		InWeftScale.ClampAxes(MinWarpWeftScale, MaxWarpWeftScale),
+		WeftScaleMultipliers,
+		TConstArrayView<TVec3<int32>>(Constraints),
+		ParticleOffset,
+		ParticleCount)
+{
+	Init(); // Lambda initialization
+	InitFromPatternData(bUse3dRestLengths, FaceVertexPatternPositions, TriangleMesh);
+	InitColor(Particles);
+}
+
+FXPBDAnisotropicAxialSpringConstraints::FXPBDAnisotropicAxialSpringConstraints(
+	const FSolverParticles& Particles,
+	int32 InParticleOffset,
+	int32 InParticleCount,
+	const FTriangleMesh& TriangleMesh,
+	const TArray<TVec3<FVec2f>>& FaceVertexPatternPositions,
+	bool bUse3dRestLengths,
+	const TConstArrayView<FRealSingle>& StiffnessWarpMultipliers,
+	const TConstArrayView<FRealSingle>& StiffnessWeftMultipliers,
+	const TConstArrayView<FRealSingle>& StiffnessBiasMultipliers,
+	const TConstArrayView<FRealSingle>& DampingMultipliers,
+	const TConstArrayView<FRealSingle>& WarpScaleMultipliers,
+	const TConstArrayView<FRealSingle>& WeftScaleMultipliers,
+	const FSolverVec2& InStiffnessWarp,
+	const FSolverVec2& InStiffnessWeft,
+	const FSolverVec2& InStiffnessBias,
+	const FSolverVec2& InDampingRatio,
+	const FSolverVec2& InWarpScale,
+	const FSolverVec2& InWeftScale,
+	bool bTrimKinematicConstraints)
+	: Base(Particles,
+		InParticleOffset,
+		InParticleCount,
+		TriangleMesh.GetElements(),
+		StiffnessWarpMultipliers,
+		InStiffnessWarp,
+		bTrimKinematicConstraints,
+		MaxStiffness)
+	, StiffnessWeft(InStiffnessWeft,
+		StiffnessWeftMultipliers,
+		TConstArrayView<TVec3<int32>>(Constraints),
+		ParticleOffset,
+		ParticleCount,
+		FPBDStiffness::DefaultTableSize,
+		FPBDStiffness::DefaultParameterFitBase,
+		MaxStiffness)
+	, StiffnessBias(InStiffnessBias,
+		StiffnessBiasMultipliers,
+		TConstArrayView<TVec3<int32>>(Constraints),
+		ParticleOffset,
+		ParticleCount,
+		FPBDStiffness::DefaultTableSize,
+		FPBDStiffness::DefaultParameterFitBase,
+		MaxStiffness)
+	, DampingRatio(
+		InDampingRatio.ClampAxes(MinDampingRatio, MaxDampingRatio),
+		DampingMultipliers,
+		TConstArrayView<TVec3<int32>>(Constraints),
+		ParticleOffset,
+		ParticleCount)
+	, WarpScale(
+		InWarpScale.ClampAxes(MinWarpWeftScale, MaxWarpWeftScale),
+		WarpScaleMultipliers,
+		TConstArrayView<TVec3<int32>>(Constraints),
+		ParticleOffset,
+		ParticleCount)
+	, WeftScale(
+		InWeftScale.ClampAxes(MinWarpWeftScale, MaxWarpWeftScale),
+		WeftScaleMultipliers,
+		TConstArrayView<TVec3<int32>>(Constraints),
+		ParticleOffset,
+		ParticleCount)
+{
+	Init(); // Lambda initialization
+	InitFromPatternData(bUse3dRestLengths, FaceVertexPatternPositions, TriangleMesh);
+	InitColor(Particles);
+}
+
+void FXPBDAnisotropicAxialSpringConstraints::InitFromPatternData(bool bUse3dRestLengths, const TArray<TVec3<FVec2f>>& FaceVertexPatternPositions, const FTriangleMesh& TriangleMesh)
+{
+	if (bUse3dRestLengths)
+	{
+		BaseDists = Dists; // Use Dists calculated by base class using 3D positions
+	}
+	else
+	{
+		// We'll set them from computed 2D lengths
+		BaseDists.SetNumUninitialized(Constraints.Num());
+	}
+	WarpWeftBiasBaseMultipliers.SetNumUninitialized(Constraints.Num());
+	WarpWeftScaleBaseMultipliers.SetNumUninitialized(Constraints.Num());
+
+	TConstArrayView<TArray<int32>> PointToTriangleMap = TriangleMesh.GetPointToTriangleMap();
+	auto SortedElement = [](const TVec3<int32>& Element)
+	{
+		return TVec3<int32>(Element.Min(), Element.Mid(), Element.Max());
+	};
+	const TArray<TVec3<int32>>& Elements = TriangleMesh.GetElements();
+
+	for (int32 ConstraintIdx = 0; ConstraintIdx < Constraints.Num(); ++ConstraintIdx)
+	{
+		const TVec3<int32>& Constraint = Constraints[ConstraintIdx];
+		// Find corresponding triangle in original mesh (constraints get reordered by Base class init)
+		const TArray<int32>& PossibleTriangles = PointToTriangleMap[Constraint[0]];
+		const TVec3<int32> SortedConstraint = SortedElement(Constraint);
+		int32 TriangleIndex = INDEX_NONE;
+		for (const int32 Triangle : PossibleTriangles)
+		{
+			if (SortedConstraint == SortedElement(Elements[Triangle]))
+			{
+				TriangleIndex = Triangle;
+				break;
+			}
+		}
+		check(TriangleIndex != INDEX_NONE);
+
+		const TVec3<int32>& Element = Elements[TriangleIndex];
+		auto MatchIndex = [&Constraint, &Element](const int32 ConstraintAxis)
+		{
+			check(Element[0] == Constraint[ConstraintAxis] || Element[1] == Constraint[ConstraintAxis] || Element[2] == Constraint[ConstraintAxis]);
+			return Element[0] == Constraint[ConstraintAxis] ? 0 : Element[1] == Constraint[ConstraintAxis] ? 1 : 2;
+		};
+		const int32 FaceIndex1 = MatchIndex(0);
+		const int32 FaceIndex2 = MatchIndex(1);
+		const int32 FaceIndex3 = MatchIndex(2);
+
+		const FVec2f& UV1 = FaceVertexPatternPositions[TriangleIndex][FaceIndex1];
+		const FVec2f& UV2 = FaceVertexPatternPositions[TriangleIndex][FaceIndex2];
+		const FVec2f& UV3 = FaceVertexPatternPositions[TriangleIndex][FaceIndex3];
+
+		const FVec2f UV = (UV2 - UV3) * Barys[ConstraintIdx] + UV3;
+		const FSolverVec2 UVDiff = UV - UV1;
+		const FSolverVec2 UVDiffAbs = UVDiff.GetAbs();
+		FSolverVec3 UVDiffTransformed;
+		if (UVDiffAbs.X > UVDiffAbs.Y)
+		{
+			// Stiffness is blend between weft and bias
+			UVDiffTransformed = FSolverVec3(UVDiffAbs.X - UVDiffAbs.Y, (FSolverReal)0.f, UVDiffAbs.Y);
+		}
+		else
+		{
+			// Stiffness is blend between warp and bias
+			UVDiffTransformed = FSolverVec3((FSolverReal)0.f, UVDiffAbs.Y - UVDiffAbs.X, UVDiffAbs.X);
+		}
+		const FSolverReal Denom3 = UVDiffTransformed.X + UVDiffTransformed.Y + UVDiffTransformed.Z; // These are used as weights
+		const FSolverReal UVLength = UVDiffAbs.Length();
+		WarpWeftBiasBaseMultipliers[ConstraintIdx] = Denom3 > UE_SMALL_NUMBER ? UVDiffTransformed / Denom3 : FSolverVec3(0.f, 1.f, 0.f); // Default to Warp if zero length
+		WarpWeftScaleBaseMultipliers[ConstraintIdx] = UVLength > UE_SMALL_NUMBER ? UVDiffAbs / UVLength : FSolverVec2(UE_INV_SQRT_2, UE_INV_SQRT_2); // Default to equally scaling warp and weft directions if zero length
+
+		if (!bUse3dRestLengths)
+		{
+			BaseDists[ConstraintIdx] = UVLength;
+			Dists[ConstraintIdx] = UVLength;
+		}
+	}
+}
+
+template<typename SolverParticlesOrRange>
+void FXPBDAnisotropicAxialSpringConstraints::InitColor(const SolverParticlesOrRange& InParticles)
+{
+	// In dev builds we always color so we can tune the system without restarting. See Apply()
+#if UE_BUILD_SHIPPING || UE_BUILD_TEST
+	if (Constraints.Num() > Chaos_XPBDSpring_ParallelConstraintCount)
+#endif
+	{
+		const TArray<TArray<int32>> ConstraintsPerColor = FGraphColoring::ComputeGraphColoringParticlesOrRange(Constraints, InParticles, ParticleOffset, ParticleOffset + ParticleCount);
+
+		// Reorder constraints based on color so each array in ConstraintsPerColor contains contiguous elements.
+		TArray<TVec3<int32>> ReorderedConstraints;
+		TArray<FSolverReal> ReorderedBarys;
+		TArray<FSolverReal> ReorderedDists;
+		TArray<FSolverReal> ReorderedBaseDists;
+		TArray<FSolverVec3> ReorderedWarpWeftBiasBaseMultipliers;
+		TArray<FSolverVec2> ReorderedWarpWeftScaleBaseMultipliers;
+		TArray<int32> OrigToReorderedIndices; // used to reorder stiffness indices
+		ReorderedConstraints.SetNumUninitialized(Constraints.Num());
+		ReorderedBarys.SetNumUninitialized(Constraints.Num());
+		ReorderedDists.SetNumUninitialized(Constraints.Num());
+		ReorderedBaseDists.SetNumUninitialized(Constraints.Num());
+		ReorderedWarpWeftBiasBaseMultipliers.SetNumUninitialized(Constraints.Num());
+		ReorderedWarpWeftScaleBaseMultipliers.SetNumUninitialized(Constraints.Num());
+		OrigToReorderedIndices.SetNumUninitialized(Constraints.Num());
+
+		ConstraintsPerColorStartIndex.Reset(ConstraintsPerColor.Num() + 1);
+
+		int32 ReorderedIndex = 0;
+		for (const TArray<int32>& ConstraintsBatch : ConstraintsPerColor)
+		{
+			ConstraintsPerColorStartIndex.Add(ReorderedIndex);
+			for (const int32& BatchConstraint : ConstraintsBatch)
+			{
+				const int32 OrigIndex = BatchConstraint;
+				ReorderedConstraints[ReorderedIndex] = Constraints[OrigIndex];
+				ReorderedBarys[ReorderedIndex] = Barys[OrigIndex];
+				ReorderedDists[ReorderedIndex] = Dists[OrigIndex];
+				ReorderedBaseDists[ReorderedIndex] = BaseDists[OrigIndex];
+				ReorderedWarpWeftBiasBaseMultipliers[ReorderedIndex] = WarpWeftBiasBaseMultipliers[OrigIndex];
+				ReorderedWarpWeftScaleBaseMultipliers[ReorderedIndex] = WarpWeftScaleBaseMultipliers[OrigIndex];
+				OrigToReorderedIndices[OrigIndex] = ReorderedIndex;
+				++ReorderedIndex;
+			}
+		}
+		ConstraintsPerColorStartIndex.Add(ReorderedIndex);
+
+		Constraints = MoveTemp(ReorderedConstraints);
+		Barys = MoveTemp(ReorderedBarys);
+		Dists = MoveTemp(ReorderedDists);
+		BaseDists = MoveTemp(ReorderedBaseDists);
+		WarpWeftBiasBaseMultipliers = MoveTemp(ReorderedWarpWeftBiasBaseMultipliers);
+		WarpWeftScaleBaseMultipliers = MoveTemp(ReorderedWarpWeftScaleBaseMultipliers);
+		Stiffness.ReorderIndices(OrigToReorderedIndices);
+		StiffnessWeft.ReorderIndices(OrigToReorderedIndices);
+		StiffnessBias.ReorderIndices(OrigToReorderedIndices);
+		DampingRatio.ReorderIndices(OrigToReorderedIndices);
+		WarpScale.ReorderIndices(OrigToReorderedIndices);
+		WeftScale.ReorderIndices(OrigToReorderedIndices);
+	}
+}
+
+void FXPBDAnisotropicAxialSpringConstraints::UpdateDists()
+{
+	Private::UpdateDists(BaseDists, WarpWeftScaleBaseMultipliers, WarpScale, WeftScale, Dists);
+}
+
+template<typename SolverParticlesOrRange>
+void FXPBDAnisotropicAxialSpringConstraints::ApplyHelper(SolverParticlesOrRange& Particles, const FSolverReal Dt, const int32 ConstraintIndex, const FSolverVec3& ExpStiffnessValues, const FSolverReal DampingRatioValue) const
+{
+	const TVec3<int32>& Constraint = Constraints[ConstraintIndex];
+	const int32 Index1 = Constraint[0];
+	const int32 Index2 = Constraint[1];
+	const int32 Index3 = Constraint[2];
+
+	const FSolverVec3& WarpWeftBiasBaseMultiplier = WarpWeftBiasBaseMultipliers[ConstraintIndex];
+	const FSolverReal FinalStiffnessValue = WarpWeftBiasBaseMultiplier.Dot(ExpStiffnessValues);
+
+	const FSolverVec3 Delta = Spring::GetXPBDAxialSpringDelta(Particles, Dt, Constraint, Barys[ConstraintIndex],  Dists[ConstraintIndex], Lambdas[ConstraintIndex],
+		FinalStiffnessValue, MinStiffness, DampingRatioValue);
+
+	if (Particles.InvM(Index1) > (FSolverReal)0.)
+	{
+		Particles.P(Index1) -= Particles.InvM(Index1) * Delta;
+	}
+	if (Particles.InvM(Index2) > (FSolverReal)0.)
+	{
+		Particles.P(Index2) += Particles.InvM(Index2) * Barys[ConstraintIndex] * Delta;
+	}
+	if (Particles.InvM(Index3) > (FSolverReal)0.)
+	{
+		Particles.P(Index3) += Particles.InvM(Index3) * ((FSolverReal)1. - Barys[ConstraintIndex]) * Delta;
+	}
+}
+
+template<typename SolverParticlesOrRange>
+void FXPBDAnisotropicAxialSpringConstraints::Apply(SolverParticlesOrRange& Particles, const FSolverReal Dt) const
+{
+	TRACE_CPUPROFILER_EVENT_SCOPE(FXPBDAnisotropicAxialSpringConstraints_Apply);
+	SCOPE_CYCLE_COUNTER(STAT_XPBD_AnisoSpring);
+
+	const bool StiffnessHasWeightMap = Stiffness.HasWeightMap();
+	const bool StiffnessWeftHasWeightMap = StiffnessWeft.HasWeightMap();
+	const bool StiffnessBiasHasWeightMap = StiffnessBias.HasWeightMap();
+	const bool DampingHasWeightMap = DampingRatio.HasWeightMap();
+
+	if ((ConstraintsPerColorStartIndex.Num() > 1) && (Constraints.Num() > Chaos_XPBDSpring_ParallelConstraintCount))
+	{
+		const int32 ConstraintColorNum = ConstraintsPerColorStartIndex.Num() - 1;
+		if (!StiffnessHasWeightMap && !StiffnessWeftHasWeightMap && !StiffnessBiasHasWeightMap && !DampingHasWeightMap)
+		{
+			const FSolverVec3 ExpStiffnessValue((FSolverReal)StiffnessWeft, (FSolverReal)Stiffness, (FSolverReal)StiffnessBias);
+			const FSolverReal DampingRatioValue = (FSolverReal)DampingRatio;
+			if (ExpStiffnessValue.Max() < MinStiffness)
+			{
+				return;
+			}
+
+#if INTEL_ISPC
+			if (bRealTypeCompatibleWithISPC && bChaos_XPBDSpring_ISPC_Enabled)
+			{
+				if (DampingRatioValue > 0)
+				{
+					for (int32 ConstraintColorIndex = 0; ConstraintColorIndex < ConstraintColorNum; ++ConstraintColorIndex)
+					{
+						const int32 ColorStart = ConstraintsPerColorStartIndex[ConstraintColorIndex];
+						const int32 ColorSize = ConstraintsPerColorStartIndex[ConstraintColorIndex + 1] - ColorStart;
+						ispc::ApplyXPBDAnisoAxialSpringConstraintsWithDamping(
+							(ispc::FVector4f*)Particles.GetPAndInvM().GetData(),
+							(const ispc::FVector3f*)Particles.XArray().GetData(),
+							(ispc::FIntVector*)&Constraints.GetData()[ColorStart],
+							&Barys.GetData()[ColorStart],
+							&Dists.GetData()[ColorStart],
+							(const ispc::FVector3f*)&WarpWeftBiasBaseMultipliers.GetData()[ColorStart],
+							&Lambdas.GetData()[ColorStart],
+							Dt,
+							reinterpret_cast<const ispc::FVector3f&>(ExpStiffnessValue),
+							DampingRatioValue,
+							ColorSize);
+					}
+				}
+				else
+				{
+					for (int32 ConstraintColorIndex = 0; ConstraintColorIndex < ConstraintColorNum; ++ConstraintColorIndex)
+					{
+						const int32 ColorStart = ConstraintsPerColorStartIndex[ConstraintColorIndex];
+						const int32 ColorSize = ConstraintsPerColorStartIndex[ConstraintColorIndex + 1] - ColorStart;
+						ispc::ApplyXPBDAnisoAxialSpringConstraints(
+							(ispc::FVector4f*)Particles.GetPAndInvM().GetData(),
+							(ispc::FIntVector*)&Constraints.GetData()[ColorStart],
+							&Barys.GetData()[ColorStart],
+							&Dists.GetData()[ColorStart],
+							(const ispc::FVector3f*)&WarpWeftBiasBaseMultipliers.GetData()[ColorStart],
+							&Lambdas.GetData()[ColorStart],
+							Dt,
+							reinterpret_cast<const ispc::FVector3f&>(ExpStiffnessValue),
+							ColorSize);
+					}
+				}
+			}
+			else
+#endif
+			{
+				for (int32 ConstraintColorIndex = 0; ConstraintColorIndex < ConstraintColorNum; ++ConstraintColorIndex)
+				{
+					const int32 ColorStart = ConstraintsPerColorStartIndex[ConstraintColorIndex];
+					const int32 ColorSize = ConstraintsPerColorStartIndex[ConstraintColorIndex + 1] - ColorStart;
+					PhysicsParallelFor(ColorSize, [this, &Particles, Dt, ColorStart, &ExpStiffnessValue, DampingRatioValue](const int32 Index)
+					{
+						const int32 ConstraintIndex = ColorStart + Index;
+					ApplyHelper(Particles, Dt, ConstraintIndex, ExpStiffnessValue, DampingRatioValue);
+					});
+				}
+			}
+		}
+		else  // Has weight maps
+		{
+#if INTEL_ISPC
+			if (bRealTypeCompatibleWithISPC && bChaos_XPBDSpring_ISPC_Enabled)
+			{
+				if (DampingHasWeightMap || (FSolverReal)DampingRatio > 0)
+				{
+					for (int32 ConstraintColorIndex = 0; ConstraintColorIndex < ConstraintColorNum; ++ConstraintColorIndex)
+					{
+						const int32 ColorStart = ConstraintsPerColorStartIndex[ConstraintColorIndex];
+						const int32 ColorSize = ConstraintsPerColorStartIndex[ConstraintColorIndex + 1] - ColorStart;
+						ispc::ApplyXPBDAnisoAxialSpringConstraintsWithDampingAndWeightMaps(
+							(ispc::FVector4f*)Particles.GetPAndInvM().GetData(),
+							(const ispc::FVector3f*)Particles.XArray().GetData(),
+							(ispc::FIntVector*)&Constraints.GetData()[ColorStart],
+							&Barys.GetData()[ColorStart],
+							&Dists.GetData()[ColorStart],
+							(const ispc::FVector3f*)&WarpWeftBiasBaseMultipliers.GetData()[ColorStart],
+							&Lambdas.GetData()[ColorStart],
+							Dt,
+							MinStiffness,
+							StiffnessHasWeightMap,
+							StiffnessHasWeightMap ? &Stiffness.GetIndices().GetData()[ColorStart] : nullptr,
+							&Stiffness.GetTable().GetData()[0],
+							StiffnessWeftHasWeightMap,
+							StiffnessWeftHasWeightMap ? &StiffnessWeft.GetIndices().GetData()[ColorStart] : nullptr,
+							&StiffnessWeft.GetTable().GetData()[0],
+							StiffnessBiasHasWeightMap,
+							StiffnessBiasHasWeightMap ? &StiffnessBias.GetIndices().GetData()[ColorStart] : nullptr,
+							&StiffnessBias.GetTable().GetData()[0],
+							DampingHasWeightMap,
+							DampingHasWeightMap ? &DampingRatio.GetIndices().GetData()[ColorStart] : nullptr,
+							&DampingRatio.GetTable().GetData()[0],
+							ColorSize);
+					}
+				}
+				else
+				{
+					for (int32 ConstraintColorIndex = 0; ConstraintColorIndex < ConstraintColorNum; ++ConstraintColorIndex)
+					{
+						const int32 ColorStart = ConstraintsPerColorStartIndex[ConstraintColorIndex];
+						const int32 ColorSize = ConstraintsPerColorStartIndex[ConstraintColorIndex + 1] - ColorStart;
+						ispc::ApplyXPBDAnisoAxialSpringConstraintsWithWeightMaps(
+							(ispc::FVector4f*)Particles.GetPAndInvM().GetData(),
+							(ispc::FIntVector*)&Constraints.GetData()[ColorStart],
+							&Barys.GetData()[ColorStart],
+							&Dists.GetData()[ColorStart],
+							(const ispc::FVector3f*)&WarpWeftBiasBaseMultipliers.GetData()[ColorStart],
+							&Lambdas.GetData()[ColorStart],
+							Dt,
+							MinStiffness,
+							StiffnessHasWeightMap,
+							&Stiffness.GetIndices().GetData()[ColorStart],
+							&Stiffness.GetTable().GetData()[0],
+							StiffnessWeftHasWeightMap,
+							StiffnessWeftHasWeightMap ? &StiffnessWeft.GetIndices().GetData()[ColorStart] : nullptr,
+							&StiffnessWeft.GetTable().GetData()[0],
+							StiffnessBiasHasWeightMap,
+							StiffnessBiasHasWeightMap ? &StiffnessBias.GetIndices().GetData()[ColorStart] : nullptr,
+							&StiffnessBias.GetTable().GetData()[0],
+							ColorSize);
+					}
+				}
+			}
+			else
+#endif
+			{
+				const FSolverReal StiffnessNoMap = (FSolverReal)Stiffness;
+				const FSolverReal StiffnessWeftNoMap = (FSolverReal)StiffnessWeft;
+				const FSolverReal StiffnessBiasNoMap = (FSolverReal)StiffnessBias;
+				const FSolverReal DampingNoMap = (FSolverReal)DampingRatio;
+				for (int32 ConstraintColorIndex = 0; ConstraintColorIndex < ConstraintColorNum; ++ConstraintColorIndex)
+				{
+					const int32 ColorStart = ConstraintsPerColorStartIndex[ConstraintColorIndex];
+					const int32 ColorSize = ConstraintsPerColorStartIndex[ConstraintColorIndex + 1] - ColorStart;
+					PhysicsParallelFor(ColorSize, [this, &Particles, Dt, ColorStart, StiffnessHasWeightMap, StiffnessNoMap, StiffnessWeftHasWeightMap, StiffnessWeftNoMap, StiffnessBiasHasWeightMap,
+						StiffnessBiasNoMap, DampingHasWeightMap, DampingNoMap](const int32 Index)
+					{
+						const int32 ConstraintIndex = ColorStart + Index;
+						const FSolverReal ExpStiffnessValue = StiffnessHasWeightMap ? Stiffness[ConstraintIndex] : StiffnessNoMap;
+						const FSolverReal ExpStiffnessWeftValue = StiffnessWeftHasWeightMap ? StiffnessWeft[ConstraintIndex] : StiffnessWeftNoMap;
+						const FSolverReal ExpStiffnessBiasValue = StiffnessBiasHasWeightMap ? StiffnessBias[ConstraintIndex] : StiffnessBiasNoMap;
+						const FSolverReal DampingRatioValue = DampingHasWeightMap ? DampingRatio[ConstraintIndex] : DampingNoMap;
+						ApplyHelper(Particles, Dt, ConstraintIndex, FSolverVec3(ExpStiffnessWeftValue, ExpStiffnessValue, ExpStiffnessBiasValue), DampingRatioValue);
+					});
+				}
+			}
+		}
+	}
+	else
+	{
+		if (!StiffnessHasWeightMap && !StiffnessWeftHasWeightMap && !StiffnessBiasHasWeightMap && !DampingHasWeightMap)
+		{
+			const FSolverVec3 ExpStiffnessValue((FSolverReal)StiffnessWeft, (FSolverReal)Stiffness, (FSolverReal)StiffnessBias);
+			const FSolverReal DampingRatioValue = (FSolverReal)DampingRatio;
+			if (ExpStiffnessValue.Max() < MinStiffness)
+			{
+				return;
+			}
+			for (int32 ConstraintIndex = 0; ConstraintIndex < Constraints.Num(); ++ConstraintIndex)
+			{
+				ApplyHelper(Particles, Dt, ConstraintIndex, ExpStiffnessValue, DampingRatioValue);
+			}
+		}
+		else
+		{
+			const FSolverReal StiffnessNoMap = (FSolverReal)Stiffness;
+			const FSolverReal StiffnessWeftNoMap = (FSolverReal)StiffnessWeft;
+			const FSolverReal StiffnessBiasNoMap = (FSolverReal)StiffnessBias;
+			const FSolverReal DampingNoMap = (FSolverReal)DampingRatio;
+			for (int32 ConstraintIndex = 0; ConstraintIndex < Constraints.Num(); ++ConstraintIndex)
+			{
+				const FSolverReal ExpStiffnessValue = StiffnessHasWeightMap ? Stiffness[ConstraintIndex] : StiffnessNoMap;
+				const FSolverReal ExpStiffnessWeftValue = StiffnessWeftHasWeightMap ? StiffnessWeft[ConstraintIndex] : StiffnessWeftNoMap;
+				const FSolverReal ExpStiffnessBiasValue = StiffnessBiasHasWeightMap ? StiffnessBias[ConstraintIndex] : StiffnessBiasNoMap;
+				const FSolverReal DampingRatioValue = DampingHasWeightMap ? DampingRatio[ConstraintIndex] : DampingNoMap;
+				ApplyHelper(Particles, Dt, ConstraintIndex, FSolverVec3(ExpStiffnessWeftValue, ExpStiffnessValue, ExpStiffnessBiasValue), DampingRatioValue);
+			}
+		}
+	}
+}
+template CHAOS_API void FXPBDAnisotropicAxialSpringConstraints::Apply(FSolverParticles& Particles, const FSolverReal Dt) const;
+template CHAOS_API void FXPBDAnisotropicAxialSpringConstraints::Apply(FSolverParticlesRange& Particles, const FSolverReal Dt) const;
+
+void FXPBDAnisotropicAxialSpringConstraints::UpdateLinearSystem(const FSolverParticlesRange& Particles, const FSolverReal Dt, FEvolutionLinearSystem& LinearSystem) const
+{
+	TRACE_CPUPROFILER_EVENT_SCOPE(FXPBDAnisotropicAxialSpringConstraints_UpdateLinearSystem);
+	LinearSystem.ReserveForParallelAdd(Constraints.Num() * 3, Constraints.Num() * 2);
+
+	const bool StiffnessHasWeightMap = Stiffness.HasWeightMap();
+	const bool StiffnessWeftHasWeightMap = StiffnessWeft.HasWeightMap();
+	const bool StiffnessBiasHasWeightMap = StiffnessBias.HasWeightMap();
+	const bool DampingHasWeightMap = DampingRatio.HasWeightMap();
+	if ((ConstraintsPerColorStartIndex.Num() > 1) && (Constraints.Num() > Chaos_XPBDSpring_ParallelConstraintCount))
+	{
+		const int32 ConstraintColorNum = ConstraintsPerColorStartIndex.Num() - 1;
+		if (!StiffnessHasWeightMap && !StiffnessWeftHasWeightMap && !StiffnessBiasHasWeightMap && !DampingHasWeightMap)
+		{
+			const FSolverVec3 ExpStiffnessValue((FSolverReal)StiffnessWeft, (FSolverReal)Stiffness, (FSolverReal)StiffnessBias);
+			const FSolverReal DampingRatioValue = (FSolverReal)DampingRatio;
+			if (ExpStiffnessValue.Max() < MinStiffness)
+			{
+				return;
+			}
+			for (int32 ConstraintColorIndex = 0; ConstraintColorIndex < ConstraintColorNum; ++ConstraintColorIndex)
+			{
+				const int32 ColorStart = ConstraintsPerColorStartIndex[ConstraintColorIndex];
+				const int32 ColorSize = ConstraintsPerColorStartIndex[ConstraintColorIndex + 1] - ColorStart;
+				PhysicsParallelFor(ColorSize, [this, &Particles, Dt, &LinearSystem, ColorStart, ExpStiffnessValue, DampingRatioValue](const int32 Index)
+				{
+					const int32 ConstraintIndex = ColorStart + Index;
+					const FSolverVec3& WarpWeftBiasBaseMultiplier = WarpWeftBiasBaseMultipliers[ConstraintIndex];
+					const FSolverReal FinalStiffnessValue = WarpWeftBiasBaseMultiplier.Dot(ExpStiffnessValue);
+					Spring::UpdateAxialSpringLinearSystem(Particles, Dt, Constraints[ConstraintIndex], Barys[ConstraintIndex], Dists[ConstraintIndex], FinalStiffnessValue, MinStiffness, DampingRatioValue, LinearSystem);
+				});
+			}
+		}
+		else // Has weight maps
+		{
+			const FSolverReal StiffnessNoMap = (FSolverReal)Stiffness;
+			const FSolverReal StiffnessWeftNoMap = (FSolverReal)StiffnessWeft;
+			const FSolverReal StiffnessBiasNoMap = (FSolverReal)StiffnessBias;
+			const FSolverReal DampingNoMap = (FSolverReal)DampingRatio;
+			for (int32 ConstraintColorIndex = 0; ConstraintColorIndex < ConstraintColorNum; ++ConstraintColorIndex)
+			{
+				const int32 ColorStart = ConstraintsPerColorStartIndex[ConstraintColorIndex];
+				const int32 ColorSize = ConstraintsPerColorStartIndex[ConstraintColorIndex + 1] - ColorStart;
+				PhysicsParallelFor(ColorSize, [this, &Particles, Dt, &LinearSystem, ColorStart, StiffnessHasWeightMap, StiffnessNoMap, StiffnessWeftHasWeightMap, StiffnessWeftNoMap, StiffnessBiasHasWeightMap,
+					StiffnessBiasNoMap, DampingHasWeightMap, DampingNoMap](const int32 Index)
+				{
+					const int32 ConstraintIndex = ColorStart + Index;
+				const FSolverReal ExpStiffnessValue = StiffnessHasWeightMap ? Stiffness[ConstraintIndex] : StiffnessNoMap;
+				const FSolverReal ExpStiffnessWeftValue = StiffnessWeftHasWeightMap ? StiffnessWeft[ConstraintIndex] : StiffnessWeftNoMap;
+				const FSolverReal ExpStiffnessBiasValue = StiffnessBiasHasWeightMap ? StiffnessBias[ConstraintIndex] : StiffnessBiasNoMap;
+				const FSolverReal DampingRatioValue = DampingHasWeightMap ? DampingRatio[ConstraintIndex] : DampingNoMap;
+				const FSolverVec3& WarpWeftBiasBaseMultiplier = WarpWeftBiasBaseMultipliers[ConstraintIndex];
+				const FSolverReal FinalStiffnessValue = WarpWeftBiasBaseMultiplier.Dot(FSolverVec3(ExpStiffnessWeftValue, ExpStiffnessValue, ExpStiffnessBiasValue));
+				Spring::UpdateAxialSpringLinearSystem(Particles, Dt, Constraints[ConstraintIndex], Barys[ConstraintIndex], Dists[ConstraintIndex], FinalStiffnessValue, MinStiffness, DampingRatioValue, LinearSystem);
+				});
+			}
+		}
+	}
+	else
+	{
+		if (!StiffnessHasWeightMap && !StiffnessWeftHasWeightMap && !StiffnessBiasHasWeightMap && !DampingHasWeightMap)
+		{
+			const FSolverVec3 ExpStiffnessValue((FSolverReal)StiffnessWeft, (FSolverReal)Stiffness, (FSolverReal)StiffnessBias);
+			const FSolverReal DampingRatioValue = (FSolverReal)DampingRatio;
+			if (ExpStiffnessValue.Max() < MinStiffness)
+			{
+				return;
+			}
+			for (int32 ConstraintIndex = 0; ConstraintIndex < Constraints.Num(); ++ConstraintIndex)
+			{
+				const FSolverVec3& WarpWeftBiasBaseMultiplier = WarpWeftBiasBaseMultipliers[ConstraintIndex];
+				const FSolverReal FinalStiffnessValue = WarpWeftBiasBaseMultiplier.Dot(ExpStiffnessValue);
+				Spring::UpdateAxialSpringLinearSystem(Particles, Dt, Constraints[ConstraintIndex], Barys[ConstraintIndex], Dists[ConstraintIndex], FinalStiffnessValue, MinStiffness, DampingRatioValue, LinearSystem);
+			}
+		}
+		else
+		{
+			const FSolverReal StiffnessNoMap = (FSolverReal)Stiffness;
+			const FSolverReal StiffnessWeftNoMap = (FSolverReal)StiffnessWeft;
+			const FSolverReal StiffnessBiasNoMap = (FSolverReal)StiffnessBias;
+			const FSolverReal DampingNoMap = (FSolverReal)DampingRatio;
+			for (int32 ConstraintIndex = 0; ConstraintIndex < Constraints.Num(); ++ConstraintIndex)
+			{
+				const FSolverReal ExpStiffnessValue = StiffnessHasWeightMap ? Stiffness[ConstraintIndex] : StiffnessNoMap;
+				const FSolverReal ExpStiffnessWeftValue = StiffnessWeftHasWeightMap ? StiffnessWeft[ConstraintIndex] : StiffnessWeftNoMap;
+				const FSolverReal ExpStiffnessBiasValue = StiffnessBiasHasWeightMap ? StiffnessBias[ConstraintIndex] : StiffnessBiasNoMap;
+				const FSolverReal DampingRatioValue = DampingHasWeightMap ? DampingRatio[ConstraintIndex] : DampingNoMap;
+				const FSolverVec3& WarpWeftBiasBaseMultiplier = WarpWeftBiasBaseMultipliers[ConstraintIndex];
+				const FSolverReal FinalStiffnessValue = WarpWeftBiasBaseMultiplier.Dot(FSolverVec3(ExpStiffnessWeftValue, ExpStiffnessValue, ExpStiffnessBiasValue));
+				Spring::UpdateAxialSpringLinearSystem(Particles, Dt, Constraints[ConstraintIndex], Barys[ConstraintIndex], Dists[ConstraintIndex], FinalStiffnessValue, MinStiffness, DampingRatioValue, LinearSystem);
+			}
+		}
+	}
+}
+
+void FXPBDAnisotropicSpringConstraints::SetProperties(
+	const FCollectionPropertyConstFacade& PropertyCollection,
+	const TMap<FString, TConstArrayView<FRealSingle>>& WeightMaps)
+{
+	if (IsXPBDAnisoSpringStiffnessWarpMutable(PropertyCollection))
+	{
+		const FSolverVec2 WeightedValue(GetWeightedFloatXPBDAnisoSpringStiffnessWarp(PropertyCollection));
+		if (IsXPBDAnisoSpringStiffnessWarpStringDirty(PropertyCollection))
+		{
+			const FString& WeightMapName = GetXPBDAnisoSpringStiffnessWarpString(PropertyCollection);
+			EdgeConstraints.Stiffness = FPBDStiffness(
+				WeightedValue,
+				WeightMaps.FindRef(WeightMapName),
+				TConstArrayView<TVec2<int32>>(EdgeConstraints.Constraints),
+				EdgeConstraints.ParticleOffset,
+				EdgeConstraints.ParticleCount,
+				FPBDStiffness::DefaultTableSize,
+				FPBDStiffness::DefaultParameterFitBase,
+				FXPBDAnisotropicEdgeSpringConstraints::MaxStiffness);
+			AxialConstraints.Stiffness = FPBDStiffness(
+				WeightedValue,
+				WeightMaps.FindRef(WeightMapName),
+				TConstArrayView<TVec3<int32>>(AxialConstraints.Constraints),
+				AxialConstraints.ParticleOffset,
+				AxialConstraints.ParticleCount,
+				FPBDStiffness::DefaultTableSize,
+				FPBDStiffness::DefaultParameterFitBase,
+				FXPBDAnisotropicEdgeSpringConstraints::MaxStiffness);
+		}
+		else
+		{
+			EdgeConstraints.Stiffness.SetWeightedValue(WeightedValue, FXPBDAnisotropicEdgeSpringConstraints::MaxStiffness);
+			AxialConstraints.Stiffness.SetWeightedValue(WeightedValue, FXPBDAnisotropicAxialSpringConstraints::MaxStiffness);
+		}
+	}
+	if (IsXPBDAnisoSpringStiffnessWeftMutable(PropertyCollection))
+	{
+		const FSolverVec2 WeightedValue(GetWeightedFloatXPBDAnisoSpringStiffnessWeft(PropertyCollection));
+		if (IsXPBDAnisoSpringStiffnessWeftStringDirty(PropertyCollection))
+		{
+			const FString& WeightMapName = GetXPBDAnisoSpringStiffnessWeftString(PropertyCollection);
+			EdgeConstraints.StiffnessWeft = FPBDStiffness(
+				WeightedValue,
+				WeightMaps.FindRef(WeightMapName),
+				TConstArrayView<TVec2<int32>>(EdgeConstraints.Constraints),
+				EdgeConstraints.ParticleOffset,
+				EdgeConstraints.ParticleCount,
+				FPBDStiffness::DefaultTableSize,
+				FPBDStiffness::DefaultParameterFitBase,
+				FXPBDAnisotropicEdgeSpringConstraints::MaxStiffness);
+			AxialConstraints.StiffnessWeft = FPBDStiffness(
+				WeightedValue,
+				WeightMaps.FindRef(WeightMapName),
+				TConstArrayView<TVec3<int32>>(AxialConstraints.Constraints),
+				AxialConstraints.ParticleOffset,
+				AxialConstraints.ParticleCount,
+				FPBDStiffness::DefaultTableSize,
+				FPBDStiffness::DefaultParameterFitBase,
+				FXPBDAnisotropicEdgeSpringConstraints::MaxStiffness);
+		}
+		else
+		{
+			EdgeConstraints.StiffnessWeft.SetWeightedValue(WeightedValue, FXPBDAnisotropicEdgeSpringConstraints::MaxStiffness);
+			AxialConstraints.StiffnessWeft.SetWeightedValue(WeightedValue, FXPBDAnisotropicAxialSpringConstraints::MaxStiffness);
+		}
+	}
+	if (IsXPBDAnisoSpringStiffnessBiasMutable(PropertyCollection))
+	{
+		const FSolverVec2 WeightedValue(GetWeightedFloatXPBDAnisoSpringStiffnessBias(PropertyCollection));
+		if (IsXPBDAnisoSpringStiffnessBiasStringDirty(PropertyCollection))
+		{
+			const FString& WeightMapName = GetXPBDAnisoSpringStiffnessBiasString(PropertyCollection);
+			EdgeConstraints.StiffnessBias = FPBDStiffness(
+				WeightedValue,
+				WeightMaps.FindRef(WeightMapName),
+				TConstArrayView<TVec2<int32>>(EdgeConstraints.Constraints),
+				EdgeConstraints.ParticleOffset,
+				EdgeConstraints.ParticleCount,
+				FPBDStiffness::DefaultTableSize,
+				FPBDStiffness::DefaultParameterFitBase,
+				FXPBDAnisotropicEdgeSpringConstraints::MaxStiffness);
+			AxialConstraints.StiffnessBias = FPBDStiffness(
+				WeightedValue,
+				WeightMaps.FindRef(WeightMapName),
+				TConstArrayView<TVec3<int32>>(AxialConstraints.Constraints),
+				AxialConstraints.ParticleOffset,
+				AxialConstraints.ParticleCount,
+				FPBDStiffness::DefaultTableSize,
+				FPBDStiffness::DefaultParameterFitBase,
+				FXPBDAnisotropicEdgeSpringConstraints::MaxStiffness);
+		}
+		else
+		{
+			EdgeConstraints.StiffnessBias.SetWeightedValue(WeightedValue, FXPBDAnisotropicEdgeSpringConstraints::MaxStiffness);
+			AxialConstraints.StiffnessBias.SetWeightedValue(WeightedValue, FXPBDAnisotropicAxialSpringConstraints::MaxStiffness);
+		}
+	}
+	if (IsXPBDAnisoSpringDampingMutable(PropertyCollection))
+	{
+		const FSolverVec2 WeightedValue = FSolverVec2(GetWeightedFloatXPBDAnisoSpringDamping(PropertyCollection)).ClampAxes(FXPBDAnisotropicEdgeSpringConstraints::MinDampingRatio, FXPBDAnisotropicEdgeSpringConstraints::MaxDampingRatio);
+		if (IsXPBDAnisoSpringDampingStringDirty(PropertyCollection))
+		{
+			const FString& WeightMapName = GetXPBDAnisoSpringDampingString(PropertyCollection);
+			AxialConstraints.DampingRatio = FPBDWeightMap(
+				WeightedValue,
+				WeightMaps.FindRef(WeightMapName),
+				TConstArrayView<TVec3<int32>>(AxialConstraints.Constraints),
+				AxialConstraints.ParticleOffset,
+				AxialConstraints.ParticleCount);
+			AxialConstraints.DampingRatio = FPBDWeightMap(
+				WeightedValue,
+				WeightMaps.FindRef(WeightMapName),
+				TConstArrayView<TVec2<int32>>(EdgeConstraints.Constraints),
+				EdgeConstraints.ParticleOffset,
+				EdgeConstraints.ParticleCount);
+		}
+		else
+		{
+			EdgeConstraints.DampingRatio.SetWeightedValue(WeightedValue);
+			AxialConstraints.DampingRatio.SetWeightedValue(WeightedValue);
+		}
+	}
+	if (IsXPBDAnisoSpringWarpScaleMutable(PropertyCollection))
+	{
+		const FSolverVec2 WeightedValue = FSolverVec2(GetWeightedFloatXPBDAnisoSpringWarpScale(PropertyCollection)).ClampAxes(FXPBDAnisotropicEdgeSpringConstraints::MinWarpWeftScale, FXPBDAnisotropicEdgeSpringConstraints::MaxWarpWeftScale);
+		if (IsXPBDAnisoSpringWarpScaleStringDirty(PropertyCollection))
+		{
+			const FString& WeightMapName = GetXPBDAnisoSpringWarpScaleString(PropertyCollection);
+			EdgeConstraints.WarpScale = FPBDWeightMap(
+				WeightedValue,
+				WeightMaps.FindRef(WeightMapName),
+				TConstArrayView<TVec2<int32>>(EdgeConstraints.Constraints),
+				EdgeConstraints.ParticleOffset,
+				EdgeConstraints.ParticleCount);
+			AxialConstraints.WarpScale = FPBDWeightMap(
+				WeightedValue,
+				WeightMaps.FindRef(WeightMapName),
+				TConstArrayView<TVec3<int32>>(AxialConstraints.Constraints),
+				AxialConstraints.ParticleOffset,
+				AxialConstraints.ParticleCount);
+		}
+		else
+		{
+			EdgeConstraints.WarpScale.SetWeightedValue(WeightedValue);
+			AxialConstraints.WarpScale.SetWeightedValue(WeightedValue);
+		}
+	}
+	if (IsXPBDAnisoSpringWeftScaleMutable(PropertyCollection))
+	{
+		const FSolverVec2 WeightedValue = FSolverVec2(GetWeightedFloatXPBDAnisoSpringWeftScale(PropertyCollection)).ClampAxes(FXPBDAnisotropicEdgeSpringConstraints::MinWarpWeftScale, FXPBDAnisotropicEdgeSpringConstraints::MaxWarpWeftScale);
+		if (IsXPBDAnisoSpringWeftScaleStringDirty(PropertyCollection))
+		{
+			const FString& WeightMapName = GetXPBDAnisoSpringWeftScaleString(PropertyCollection);
+			EdgeConstraints.WeftScale = FPBDWeightMap(
+				WeightedValue,
+				WeightMaps.FindRef(WeightMapName),
+				TConstArrayView<TVec2<int32>>(EdgeConstraints.Constraints),
+				EdgeConstraints.ParticleOffset,
+				EdgeConstraints.ParticleCount);
+			AxialConstraints.WeftScale = FPBDWeightMap(
+				WeightedValue,
+				WeightMaps.FindRef(WeightMapName),
+				TConstArrayView<TVec3<int32>>(AxialConstraints.Constraints),
+				AxialConstraints.ParticleOffset,
+				AxialConstraints.ParticleCount);
+		}
+		else
+		{
+			EdgeConstraints.WeftScale.SetWeightedValue(WeightedValue);
+			AxialConstraints.WeftScale.SetWeightedValue(WeightedValue);
+		}
+	}
+}
+}  // End namespace Chaos::Softs
diff --git a/Engine/Source/Runtime/Experimental/Chaos/Private/Chaos/XPBDAnisotropicSpringConstraints.ispc b/Engine/Source/Runtime/Experimental/Chaos/Private/Chaos/XPBDAnisotropicSpringConstraints.ispc
new file mode 100644
--- /dev/null
+++ b/Engine/Source/Runtime/Experimental/Chaos/Private/Chaos/XPBDAnisotropicSpringConstraints.ispc
@@ -0,0 +1,518 @@
+// Copyright Epic Games, Inc. All Rights Reserved.
+
+#define EXPLICIT_VECTOR4 1
+
+#include "Math/Vector.isph"
+#include "Chaos/PBDSofts.isph"
+#include "Chaos/SoftsSpring.isph"
+
+export void ApplyXPBDAnisoSpringConstraints(uniform FVector4f PandInvM[],
+									const uniform FIntVector2 Constraints[],
+									const uniform float Dists[],
+									const uniform FVector3f WarpWeftBiasBaseMultiplers[],
+									uniform float Lambdas[],
+									const uniform float Dt,
+									const uniform FVector3f& Stiffness,
+									const uniform int32 NumConstraints)
+{
+	foreach(Index = 0 ... NumConstraints)
+	{
+		const varying FIntVector2 Constraint = VectorLoad(&Constraints[extract(Index,0)]);
+		const varying int32 Index1 = Constraint.V[0];
+		const varying int32 Index2 = Constraint.V[1];
+
+		const varying FVector4f PandInvM1 = VectorGather(&PandInvM[Index1]);
+		const varying FVector4f PandInvM2 = VectorGather(&PandInvM[Index2]);
+
+		varying FVector3f P1, P2;
+		varying float InvM1, InvM2;
+		UnzipPandInvM(PandInvM1, P1, InvM1);
+		UnzipPandInvM(PandInvM2, P2, InvM2);
+		
+		const varying float Dist = Dists[Index];
+		const varying float Lambda = Lambdas[Index];
+		const varying FVector3f WarpWeftBiasBaseMultiplier = VectorGather(&WarpWeftBiasBaseMultiplers[Index]);
+		const varying float StiffnessValue = VectorDot(WarpWeftBiasBaseMultiplier, Stiffness);
+
+		varying float DLambda;
+		const varying FVector3f Delta = GetXPBDSpringDelta(P1, P2, InvM1, InvM2, Dt, Dist, Lambda, StiffnessValue, FLOAT_SMALL_NUMBER, DLambda);
+
+		Lambdas[Index] = Lambda + DLambda;
+
+		if (InvM1 > FLOAT_ZERO)
+		{
+			VectorScatter(&PandInvM[Index1], SetVector4( P1 - (InvM1 * Delta), InvM1 ));
+		}
+		if (InvM2 > FLOAT_ZERO)
+		{
+			VectorScatter(&PandInvM[Index2], SetVector4( P2 + (InvM2 * Delta), InvM2 ));
+		}
+	}
+}
+
+export void ApplyXPBDAnisoSpringConstraintsWithDamping(uniform FVector4f PandInvM[],
+												const uniform FVector3f X[],
+												const uniform FIntVector2 Constraints[],
+												const uniform float Dists[],
+												const uniform FVector3f WarpWeftBiasBaseMultiplers[],
+												uniform float Lambdas[],
+												const uniform float Dt,
+												const uniform FVector3f& Stiffness,
+												const uniform float DampingRatio,
+												const uniform int32 NumConstraints)
+{
+	foreach(Index = 0 ... NumConstraints)
+	{
+		const varying FIntVector2 Constraint = VectorLoad(&Constraints[extract(Index,0)]);
+		const varying int32 Index1 = Constraint.V[0];
+		const varying int32 Index2 = Constraint.V[1];
+
+		const varying FVector4f PandInvM1 = VectorGather(&PandInvM[Index1]);
+		const varying FVector4f PandInvM2 = VectorGather(&PandInvM[Index2]);
+
+		varying FVector3f P1, P2;
+		varying float InvM1, InvM2;
+		UnzipPandInvM(PandInvM1, P1, InvM1);
+		UnzipPandInvM(PandInvM2, P2, InvM2);
+
+		const varying FVector3f X1 = VectorGather(&X[Index1]);
+		const varying FVector3f X2 = VectorGather(&X[Index2]);
+
+		const varying float Dist = Dists[Index];
+		const varying float Lambda = Lambdas[Index];
+		const varying FVector3f WarpWeftBiasBaseMultiplier = VectorGather(&WarpWeftBiasBaseMultiplers[Index]);
+		const varying float StiffnessValue = VectorDot(WarpWeftBiasBaseMultiplier, Stiffness);
+
+		varying float DLambda;
+		const FVector3f Delta = GetXPBDSpringDeltaWithDamping(P1, P2, X1, X2, InvM1, InvM2, Dt, Dist, Lambda, StiffnessValue, DampingRatio, FLOAT_SMALL_NUMBER, DLambda);
+
+		Lambdas[Index] = Lambda + DLambda;
+
+		if (InvM1 > FLOAT_ZERO)
+		{
+			VectorScatter(&PandInvM[Index1], SetVector4( P1 - (InvM1 * Delta), InvM1 ));
+		}
+		if (InvM2 > FLOAT_ZERO)
+		{
+			VectorScatter(&PandInvM[Index2], SetVector4( P2 + (InvM2 * Delta), InvM2 ));
+		}
+	}
+}
+
+export void ApplyXPBDAnisoSpringConstraintsWithWeightMaps(uniform FVector4f PandInvM[],
+									const uniform FIntVector2 Constraints[],
+									const uniform float Dists[],
+									const uniform FVector3f WarpWeftBiasBaseMultiplers[],
+									uniform float Lambdas[],
+									const uniform float Dt,
+									const uniform float MinStiffness,
+									const uniform bool StiffnessHasMap,
+									const uniform uint8 StiffnessIndices[],
+									const uniform float StiffnessTable[],
+									const uniform bool StiffnessWeftHasMap,
+									const uniform uint8 StiffnessWeftIndices[],
+									const uniform float StiffnessWeftTable[],
+									const uniform bool StiffnessBiasHasMap,
+									const uniform uint8 StiffnessBiasIndices[],
+									const uniform float StiffnessBiasTable[],
+									const uniform int32 NumConstraints)
+{
+
+	foreach(Index = 0 ... NumConstraints)
+	{
+		const varying FIntVector2 Constraint = VectorLoad(&Constraints[extract(Index,0)]);
+		const varying int32 Index1 = Constraint.V[0];
+		const varying int32 Index2 = Constraint.V[1];
+		
+		const varying int8 StiffnessIndex = StiffnessHasMap ? StiffnessIndices[Index] : 0;
+		const varying int8 StiffnessWeftIndex = StiffnessWeftHasMap ? StiffnessWeftIndices[Index] : 0;
+		const varying int8 StiffnessBiasIndex = StiffnessBiasHasMap ? StiffnessBiasIndices[Index] : 0;
+
+		const varying FVector4f PandInvM1 = VectorGather(&PandInvM[Index1]);
+		const varying FVector4f PandInvM2 = VectorGather(&PandInvM[Index2]);
+
+		varying FVector3f P1, P2;
+		varying float InvM1, InvM2;
+		UnzipPandInvM(PandInvM1, P1, InvM1);
+		UnzipPandInvM(PandInvM2, P2, InvM2);
+
+		const varying float Dist = Dists[Index];
+		const varying float Lambda = Lambdas[Index];
+		
+		#pragma ignore warning(perf)
+		const varying float Stiffness = StiffnessTable[StiffnessIndex];
+		#pragma ignore warning(perf)
+		const varying float StiffnessWeft = StiffnessWeftTable[StiffnessWeftIndex];
+		#pragma ignore warning(perf)
+		const varying float StiffnessBias = StiffnessBiasTable[StiffnessBiasIndex];
+		
+		const varying FVector3f StiffnessVector = {{StiffnessWeft, Stiffness, StiffnessBias}};
+		const varying FVector3f WarpWeftBiasBaseMultiplier = VectorGather(&WarpWeftBiasBaseMultiplers[Index]);
+		const varying float StiffnessValue = VectorDot(WarpWeftBiasBaseMultiplier, StiffnessVector);
+
+		varying float DLambda;
+		const varying FVector3f Delta = GetXPBDSpringDelta(P1, P2, InvM1, InvM2, Dt, Dist, Lambda, StiffnessValue, MinStiffness, DLambda);
+
+		Lambdas[Index] = Lambda + DLambda;
+
+		if (InvM1 > FLOAT_ZERO)
+		{
+			VectorScatter(&PandInvM[Index1], SetVector4( P1 - (InvM1 * Delta), InvM1 ));
+		}
+		if (InvM2 > FLOAT_ZERO)
+		{
+			VectorScatter(&PandInvM[Index2], SetVector4( P2 + (InvM2 * Delta), InvM2 ));
+		}
+	}
+}
+
+export void ApplyXPBDAnisoSpringConstraintsWithDampingAndWeightMaps(uniform FVector4f PandInvM[],
+															const uniform FVector3f X[],
+															const uniform FIntVector2 Constraints[],
+															const uniform float Dists[],
+															const uniform FVector3f WarpWeftBiasBaseMultiplers[],
+															uniform float Lambdas[],
+															const uniform float Dt,
+															const uniform float MinStiffness,
+															const uniform bool StiffnessHasMap,
+															const uniform uint8 StiffnessIndices[],
+															const uniform float StiffnessTable[],
+															const uniform bool StiffnessWeftHasMap,
+															const uniform uint8 StiffnessWeftIndices[],
+															const uniform float StiffnessWeftTable[],
+															const uniform bool StiffnessBiasHasMap,
+															const uniform uint8 StiffnessBiasIndices[],
+															const uniform float StiffnessBiasTable[],
+															const uniform bool DampingHasMap,
+															const uniform uint8 DampingIndices[],
+															const uniform float DampingTable[],
+															const uniform int32 NumConstraints)
+{
+
+	foreach(Index = 0 ... NumConstraints)
+	{
+		const varying FIntVector2 Constraint = VectorLoad(&Constraints[extract(Index,0)]);
+		const varying int32 Index1 = Constraint.V[0];
+		const varying int32 Index2 = Constraint.V[1];
+		const varying int8 StiffnessIndex = StiffnessHasMap ? StiffnessIndices[Index] : 0;
+		const varying int8 StiffnessWeftIndex = StiffnessWeftHasMap ? StiffnessWeftIndices[Index] : 0;
+		const varying int8 StiffnessBiasIndex = StiffnessBiasHasMap ? StiffnessBiasIndices[Index] : 0;
+		const varying int8 DampingIndex = DampingHasMap ? DampingIndices[Index] : 0;
+
+		const varying FVector4f PandInvM1 = VectorGather(&PandInvM[Index1]);
+		const varying FVector4f PandInvM2 = VectorGather(&PandInvM[Index2]);
+
+		varying FVector3f P1, P2;
+		varying float InvM1, InvM2;
+		UnzipPandInvM(PandInvM1, P1, InvM1);
+		UnzipPandInvM(PandInvM2, P2, InvM2);
+
+		const varying FVector3f X1 = VectorGather(&X[Index1]);
+		const varying FVector3f X2 = VectorGather(&X[Index2]);
+
+		const varying float Dist = Dists[Index];
+		const varying float Lambda = Lambdas[Index];
+
+		#pragma ignore warning(perf)
+		const varying float Stiffness = StiffnessTable[StiffnessIndex];
+		#pragma ignore warning(perf)
+		const varying float StiffnessWeft = StiffnessWeftTable[StiffnessWeftIndex];
+		#pragma ignore warning(perf)
+		const varying float StiffnessBias = StiffnessBiasTable[StiffnessBiasIndex];
+		#pragma ignore warning(perf)
+		const varying float DampingRatio = DampingTable[DampingIndex];
+
+		const varying FVector3f StiffnessVector = {{StiffnessWeft, Stiffness, StiffnessBias}};
+		const varying FVector3f WarpWeftBiasBaseMultiplier = VectorGather(&WarpWeftBiasBaseMultiplers[Index]);
+		const varying float StiffnessValue = VectorDot(WarpWeftBiasBaseMultiplier, StiffnessVector);
+
+		varying float DLambda;
+		const FVector3f Delta = GetXPBDSpringDeltaWithDamping(P1, P2, X1, X2, InvM1, InvM2, Dt, Dist, Lambda, StiffnessValue, DampingRatio, MinStiffness, DLambda);
+
+		Lambdas[Index] = Lambda + DLambda;
+
+		if (InvM1 > FLOAT_ZERO)
+		{
+			VectorScatter(&PandInvM[Index1], SetVector4( P1 - (InvM1 * Delta), InvM1 ));
+		}
+		if (InvM2 > FLOAT_ZERO)
+		{
+			VectorScatter(&PandInvM[Index2], SetVector4( P2 + (InvM2 * Delta), InvM2 ));
+		}
+	}
+}
+
+export void ApplyXPBDAnisoAxialSpringConstraints(uniform FVector4f PandInvM[],
+									const uniform FIntVector Constraints[],
+									const uniform float Barys[],
+									const uniform float Dists[],
+									const uniform FVector3f WarpWeftBiasBaseMultiplers[],
+									uniform float Lambdas[],
+									const uniform float Dt,
+									const uniform FVector3f& Stiffness,
+									const uniform int32 NumConstraints)
+{
+	foreach(Index = 0 ... NumConstraints)
+	{
+		const varying FIntVector Constraint = VectorLoad(&Constraints[extract(Index,0)]);
+		const varying int32 Index1 = Constraint.V[0];
+		const varying int32 Index2 = Constraint.V[1];
+		const varying int32 Index3 = Constraint.V[2];
+
+		const varying FVector4f PandInvM1 = VectorGather(&PandInvM[Index1]);
+		const varying FVector4f PandInvM2 = VectorGather(&PandInvM[Index2]);
+		const varying FVector4f PandInvM3 = VectorGather(&PandInvM[Index3]);
+
+		varying FVector3f P1, P2, P3;
+		varying float InvM1, InvM2, InvM3;
+		UnzipPandInvM(PandInvM1, P1, InvM1);
+		UnzipPandInvM(PandInvM2, P2, InvM2);
+		UnzipPandInvM(PandInvM3, P3, InvM3);
+		
+		const varying float Bary = Barys[Index];
+		const varying float Dist = Dists[Index];
+		const varying float Lambda = Lambdas[Index];
+		const varying FVector3f WarpWeftBiasBaseMultiplier = VectorGather(&WarpWeftBiasBaseMultiplers[Index]);
+		const varying float StiffnessValue = VectorDot(WarpWeftBiasBaseMultiplier, Stiffness);
+
+		varying float DLambda;
+		const varying FVector3f Delta = GetXPBDAxialSpringDelta(P1, P2, P3, InvM1, InvM2, InvM3, Dt, Bary, Dist, Lambda, StiffnessValue, FLOAT_SMALL_NUMBER, DLambda);
+
+		Lambdas[Index] = Lambda + DLambda;
+
+		if (InvM1 > FLOAT_ZERO)
+		{
+			VectorScatter(&PandInvM[Index1], SetVector4( P1 - (InvM1 * Delta), InvM1 ));
+		}
+		if (InvM2 > FLOAT_ZERO)
+		{
+			VectorScatter(&PandInvM[Index2], SetVector4( P2 + (InvM2 * Bary * Delta), InvM2 ));
+		}
+		if (InvM3 > FLOAT_ZERO)
+		{
+			VectorScatter(&PandInvM[Index3], SetVector4( P3 + (InvM3 * (1.f - Bary) * Delta), InvM3 ));
+		}
+	}
+}
+
+export void ApplyXPBDAnisoAxialSpringConstraintsWithDamping(uniform FVector4f PandInvM[],
+												const uniform FVector3f X[],
+												const uniform FIntVector Constraints[],
+												const uniform float Barys[],
+												const uniform float Dists[],
+												const uniform FVector3f WarpWeftBiasBaseMultiplers[],
+												uniform float Lambdas[],
+												const uniform float Dt,
+												const uniform FVector3f& Stiffness,
+												const uniform float DampingRatio,
+												const uniform int32 NumConstraints)
+{
+	foreach(Index = 0 ... NumConstraints)
+	{
+		const varying FIntVector Constraint = VectorLoad(&Constraints[extract(Index,0)]);
+		const varying int32 Index1 = Constraint.V[0];
+		const varying int32 Index2 = Constraint.V[1];
+		const varying int32 Index3 = Constraint.V[2];
+		
+		const varying FVector4f PandInvM1 = VectorGather(&PandInvM[Index1]);
+		const varying FVector4f PandInvM2 = VectorGather(&PandInvM[Index2]);
+		const varying FVector4f PandInvM3 = VectorGather(&PandInvM[Index3]);
+
+		varying FVector3f P1, P2, P3;
+		varying float InvM1, InvM2, InvM3;
+		UnzipPandInvM(PandInvM1, P1, InvM1);
+		UnzipPandInvM(PandInvM2, P2, InvM2);
+		UnzipPandInvM(PandInvM3, P3, InvM3);
+
+		const varying FVector3f X1 = VectorGather(&X[Index1]);
+		const varying FVector3f X2 = VectorGather(&X[Index2]);
+		const varying FVector3f X3 = VectorGather(&X[Index3]);
+		
+		const varying float Bary = Barys[Index];
+		const varying float Dist = Dists[Index];
+		const varying float Lambda = Lambdas[Index];
+		const varying FVector3f WarpWeftBiasBaseMultiplier = VectorGather(&WarpWeftBiasBaseMultiplers[Index]);
+		const varying float StiffnessValue = VectorDot(WarpWeftBiasBaseMultiplier, Stiffness);
+
+		varying float DLambda;
+		const FVector3f Delta = GetXPBDAxialSpringDeltaWithDamping(P1, P2, P3, X1, X2, X3, InvM1, InvM2, InvM3, Dt, Bary, Dist, Lambda, StiffnessValue, DampingRatio, FLOAT_SMALL_NUMBER, DLambda);
+
+		Lambdas[Index] = Lambda + DLambda;
+
+		if (InvM1 > FLOAT_ZERO)
+		{
+			VectorScatter(&PandInvM[Index1], SetVector4( P1 - (InvM1 * Delta), InvM1 ));
+		}
+		if (InvM2 > FLOAT_ZERO)
+		{
+			VectorScatter(&PandInvM[Index2], SetVector4( P2 + (InvM2 * Bary * Delta), InvM2 ));
+		}
+		if (InvM3 > FLOAT_ZERO)
+		{
+			VectorScatter(&PandInvM[Index3], SetVector4( P3 + (InvM3 * (1.f - Bary) * Delta), InvM3 ));
+		}
+	}
+}
+
+export void ApplyXPBDAnisoAxialSpringConstraintsWithWeightMaps(uniform FVector4f PandInvM[],
+									const uniform FIntVector Constraints[],
+									const uniform float Barys[],
+									const uniform float Dists[],
+									const uniform FVector3f WarpWeftBiasBaseMultiplers[],
+									uniform float Lambdas[],
+									const uniform float Dt,
+									const uniform float MinStiffness,
+									const uniform bool StiffnessHasMap,
+									const uniform uint8 StiffnessIndices[],
+									const uniform float StiffnessTable[],
+									const uniform bool StiffnessWeftHasMap,
+									const uniform uint8 StiffnessWeftIndices[],
+									const uniform float StiffnessWeftTable[],
+									const uniform bool StiffnessBiasHasMap,
+									const uniform uint8 StiffnessBiasIndices[],
+									const uniform float StiffnessBiasTable[],
+									const uniform int32 NumConstraints)
+{
+
+	foreach(Index = 0 ... NumConstraints)
+	{
+		const varying FIntVector Constraint = VectorLoad(&Constraints[extract(Index,0)]);
+		const varying int32 Index1 = Constraint.V[0];
+		const varying int32 Index2 = Constraint.V[1];
+		const varying int32 Index3 = Constraint.V[2];
+		
+		const varying int8 StiffnessIndex = StiffnessHasMap ? StiffnessIndices[Index] : 0;
+		const varying int8 StiffnessWeftIndex = StiffnessWeftHasMap ? StiffnessWeftIndices[Index] : 0;
+		const varying int8 StiffnessBiasIndex = StiffnessBiasHasMap ? StiffnessBiasIndices[Index] : 0;
+		
+		const varying FVector4f PandInvM1 = VectorGather(&PandInvM[Index1]);
+		const varying FVector4f PandInvM2 = VectorGather(&PandInvM[Index2]);
+		const varying FVector4f PandInvM3 = VectorGather(&PandInvM[Index3]);
+
+		varying FVector3f P1, P2, P3;
+		varying float InvM1, InvM2, InvM3;
+		UnzipPandInvM(PandInvM1, P1, InvM1);
+		UnzipPandInvM(PandInvM2, P2, InvM2);
+		UnzipPandInvM(PandInvM3, P3, InvM3);
+		
+		const varying float Bary = Barys[Index];
+		const varying float Dist = Dists[Index];
+		const varying float Lambda = Lambdas[Index];
+		
+		#pragma ignore warning(perf)
+		const varying float Stiffness = StiffnessTable[StiffnessIndex];
+		#pragma ignore warning(perf)
+		const varying float StiffnessWeft = StiffnessWeftTable[StiffnessWeftIndex];
+		#pragma ignore warning(perf)
+		const varying float StiffnessBias = StiffnessBiasTable[StiffnessBiasIndex];
+		
+		const varying FVector3f StiffnessVector = {{StiffnessWeft, Stiffness, StiffnessBias}};
+		const varying FVector3f WarpWeftBiasBaseMultiplier = VectorGather(&WarpWeftBiasBaseMultiplers[Index]);
+		const varying float StiffnessValue = VectorDot(WarpWeftBiasBaseMultiplier, StiffnessVector);
+
+		varying float DLambda;
+		const varying FVector3f Delta = GetXPBDAxialSpringDelta(P1, P2, P3, InvM1, InvM2, InvM3, Dt, Bary, Dist, Lambda, StiffnessValue, FLOAT_SMALL_NUMBER, DLambda);
+
+		Lambdas[Index] = Lambda + DLambda;
+
+		if (InvM1 > FLOAT_ZERO)
+		{
+			VectorScatter(&PandInvM[Index1], SetVector4( P1 - (InvM1 * Delta), InvM1 ));
+		}
+		if (InvM2 > FLOAT_ZERO)
+		{
+			VectorScatter(&PandInvM[Index2], SetVector4( P2 + (InvM2 * Bary * Delta), InvM2 ));
+		}
+		if (InvM3 > FLOAT_ZERO)
+		{
+			VectorScatter(&PandInvM[Index3], SetVector4( P3 + (InvM3 * (1.f - Bary) * Delta), InvM3 ));
+		}
+	}
+}
+
+export void ApplyXPBDAnisoAxialSpringConstraintsWithDampingAndWeightMaps(uniform FVector4f PandInvM[],
+															const uniform FVector3f X[],
+															const uniform FIntVector Constraints[],
+															const uniform float Barys[],
+															const uniform float Dists[],
+															const uniform FVector3f WarpWeftBiasBaseMultiplers[],
+															uniform float Lambdas[],
+															const uniform float Dt,
+															const uniform float MinStiffness,
+															const uniform bool StiffnessHasMap,
+															const uniform uint8 StiffnessIndices[],
+															const uniform float StiffnessTable[],
+															const uniform bool StiffnessWeftHasMap,
+															const uniform uint8 StiffnessWeftIndices[],
+															const uniform float StiffnessWeftTable[],
+															const uniform bool StiffnessBiasHasMap,
+															const uniform uint8 StiffnessBiasIndices[],
+															const uniform float StiffnessBiasTable[],
+															const uniform bool DampingHasMap,
+															const uniform uint8 DampingIndices[],
+															const uniform float DampingTable[],
+															const uniform int32 NumConstraints)
+{
+
+	foreach(Index = 0 ... NumConstraints)
+	{
+		const varying FIntVector Constraint = VectorLoad(&Constraints[extract(Index,0)]);
+		const varying int32 Index1 = Constraint.V[0];
+		const varying int32 Index2 = Constraint.V[1];
+		const varying int32 Index3 = Constraint.V[2];
+
+		const varying int8 StiffnessIndex = StiffnessHasMap ? StiffnessIndices[Index] : 0;
+		const varying int8 StiffnessWeftIndex = StiffnessWeftHasMap ? StiffnessWeftIndices[Index] : 0;
+		const varying int8 StiffnessBiasIndex = StiffnessBiasHasMap ? StiffnessBiasIndices[Index] : 0;
+		const varying int8 DampingIndex = DampingHasMap ? DampingIndices[Index] : 0;
+		
+		const varying FVector4f PandInvM1 = VectorGather(&PandInvM[Index1]);
+		const varying FVector4f PandInvM2 = VectorGather(&PandInvM[Index2]);
+		const varying FVector4f PandInvM3 = VectorGather(&PandInvM[Index3]);
+
+		varying FVector3f P1, P2, P3;
+		varying float InvM1, InvM2, InvM3;
+		UnzipPandInvM(PandInvM1, P1, InvM1);
+		UnzipPandInvM(PandInvM2, P2, InvM2);
+		UnzipPandInvM(PandInvM3, P3, InvM3);
+		
+		const varying FVector3f X1 = VectorGather(&X[Index1]);
+		const varying FVector3f X2 = VectorGather(&X[Index2]);
+		const varying FVector3f X3 = VectorGather(&X[Index3]);
+		
+		const varying float Bary = Barys[Index];
+		const varying float Dist = Dists[Index];
+		const varying float Lambda = Lambdas[Index];
+
+		#pragma ignore warning(perf)
+		const varying float Stiffness = StiffnessTable[StiffnessIndex];
+		#pragma ignore warning(perf)
+		const varying float StiffnessWeft = StiffnessWeftTable[StiffnessWeftIndex];
+		#pragma ignore warning(perf)
+		const varying float StiffnessBias = StiffnessBiasTable[StiffnessBiasIndex];
+		#pragma ignore warning(perf)
+		const varying float DampingRatio = DampingTable[DampingIndex];
+
+		const varying FVector3f StiffnessVector = {{StiffnessWeft, Stiffness, StiffnessBias}};
+		const varying FVector3f WarpWeftBiasBaseMultiplier = VectorGather(&WarpWeftBiasBaseMultiplers[Index]);
+		const varying float StiffnessValue = VectorDot(WarpWeftBiasBaseMultiplier, StiffnessVector);
+
+		varying float DLambda;
+		const FVector3f Delta = GetXPBDAxialSpringDeltaWithDamping(P1, P2, P3, X1, X2, X3, InvM1, InvM2, InvM3, Dt, Bary, Dist, Lambda, StiffnessValue, DampingRatio, FLOAT_SMALL_NUMBER, DLambda);
+
+		Lambdas[Index] = Lambda + DLambda;
+
+		if (InvM1 > FLOAT_ZERO)
+		{
+			VectorScatter(&PandInvM[Index1], SetVector4( P1 - (InvM1 * Delta), InvM1 ));
+		}
+		if (InvM2 > FLOAT_ZERO)
+		{
+			VectorScatter(&PandInvM[Index2], SetVector4( P2 + (InvM2 * Bary * Delta), InvM2 ));
+		}
+		if (InvM3 > FLOAT_ZERO)
+		{
+			VectorScatter(&PandInvM[Index3], SetVector4( P3 + (InvM3 * (1.f - Bary) * Delta), InvM3 ));
+		}
+	}
+}
diff --git a/Engine/Source/Runtime/Experimental/Chaos/Private/Chaos/XPBDSpringConstraints.cpp b/Engine/Source/Runtime/Experimental/Chaos/Private/Chaos/XPBDSpringConstraints.cpp
--- a/Engine/Source/Runtime/Experimental/Chaos/Private/Chaos/XPBDSpringConstraints.cpp
+++ b/Engine/Source/Runtime/Experimental/Chaos/Private/Chaos/XPBDSpringConstraints.cpp
@@ -26,7 +26,7 @@ FAutoConsoleVariableRef CVarChaosXPBDSpringISPCEnabled(TEXT("p.Chaos.XPBDSpring.
 namespace Chaos::Softs {
 
 // @todo(chaos): the parallel threshold (or decision to run parallel) should probably be owned by the solver and passed to the constraint container
-static int32 Chaos_XPBDSpring_ParallelConstraintCount = 100;
+int32 Chaos_XPBDSpring_ParallelConstraintCount = 100;
 #if !UE_BUILD_SHIPPING
 FAutoConsoleVariableRef CVarChaosXPBDSpringParallelConstraintCount(TEXT("p.Chaos.XPBDSpring.ParallelConstraintCount"), Chaos_XPBDSpring_ParallelConstraintCount, TEXT("If we have more constraints than this, use parallel-for in Apply."));
 #endif
diff --git a/Engine/Source/Runtime/Experimental/Chaos/Public/Chaos/PBDAxialSpringConstraintsBase.h b/Engine/Source/Runtime/Experimental/Chaos/Public/Chaos/PBDAxialSpringConstraintsBase.h
--- a/Engine/Source/Runtime/Experimental/Chaos/Public/Chaos/PBDAxialSpringConstraintsBase.h
+++ b/Engine/Source/Runtime/Experimental/Chaos/Public/Chaos/PBDAxialSpringConstraintsBase.h
@@ -78,6 +78,9 @@ public:
 
 	void ApplyProperties(const FSolverReal Dt, const int32 NumIterations) { Stiffness.ApplyPBDValues(Dt, NumIterations); }
 
+	const TArray<TVec3<int32>>& GetConstraints() const { return Constraints; }
+	const TArray<FSolverReal>& GetBarys() const { return Barys; }
+
 protected:
 	template<typename SolverParticlesOrRange>
 	inline FSolverVec3 GetDelta(const SolverParticlesOrRange& Particles, const int32 ConstraintIndex, const FSolverReal ExpStiffnessValue) const
diff --git a/Engine/Source/Runtime/Experimental/Chaos/Public/Chaos/PBDSoftsEvolutionFwd.h b/Engine/Source/Runtime/Experimental/Chaos/Public/Chaos/PBDSoftsEvolutionFwd.h
--- a/Engine/Source/Runtime/Experimental/Chaos/Public/Chaos/PBDSoftsEvolutionFwd.h
+++ b/Engine/Source/Runtime/Experimental/Chaos/Public/Chaos/PBDSoftsEvolutionFwd.h
@@ -53,6 +53,7 @@ namespace Chaos
 		class FPBDSpringConstraints;
 		class FPBDEdgeSpringConstraints;
 		class FXPBDStretchBiasElementConstraints;
+		class FXPBDAnisotropicSpringConstraints;
 		class FPBDBendingSpringConstraints;
 		class FXPBDSpringConstraints;
 		class FXPBDEdgeSpringConstraints;
diff --git a/Engine/Source/Runtime/Experimental/Chaos/Public/Chaos/PBDWeightMap.h b/Engine/Source/Runtime/Experimental/Chaos/Public/Chaos/PBDWeightMap.h
--- a/Engine/Source/Runtime/Experimental/Chaos/Public/Chaos/PBDWeightMap.h
+++ b/Engine/Source/Runtime/Experimental/Chaos/Public/Chaos/PBDWeightMap.h
@@ -59,7 +59,7 @@ public:
 	 * The weight map table only gets updated after ApplyValues is called.
 	 * Low and high values are clamped between [0,1]
 	 */
-	void SetWeightedValue(const FSolverVec2& InWeightedValue) { WeightedValue = InWeightedValue; bIsDirty = true; }
+	void SetWeightedValue(const FSolverVec2& InWeightedValue) { bIsDirty = WeightedValue != InWeightedValue; WeightedValue = InWeightedValue;}
 
 	/**
 	 * Set the low and high values of the weight map.
@@ -67,7 +67,7 @@ public:
 	 * Low and high values are not clamped. Commonly used for XPBD Stiffness values which are not [0,1]
 	 */
 	UE_DEPRECATED(5.2, "Use SetWeightedValue.")
-	void SetWeightedValueUnclamped(const FSolverVec2& InWeightedValue) { WeightedValue = InWeightedValue; bIsDirty = true; }
+	void SetWeightedValueUnclamped(const FSolverVec2& InWeightedValue) { bIsDirty = WeightedValue != InWeightedValue; WeightedValue = InWeightedValue;}
 
 	/**
 	 * Return the low and high values set for this weight map.
@@ -76,7 +76,7 @@ public:
 	const FSolverVec2& GetWeightedValue() const { return WeightedValue; }
 
 	/** Update the weight map table with the current simulation parameters. */
-	void ApplyValues() { ApplyValues([](FSolverReal Value)->FSolverReal { return Value; }); }
+	void ApplyValues(bool* bOutUpdated = nullptr) { ApplyValues([](FSolverReal Value)->FSolverReal { return Value; }, bOutUpdated); }
 
 	/**
 	 * Lookup for the exponential weighted value at the specified weight map index.
@@ -104,7 +104,7 @@ public:
 
 protected:
 	template<typename FunctorType>
-	inline void ApplyValues(FunctorType&& MappingFunction);
+	inline void ApplyValues(FunctorType&& MappingFunction, bool* bOutUpdated = nullptr);
 
 	TArray<uint8> Indices; // Per particle/constraints array of index to the stiffness table
 	TArray<FSolverReal> Table;  // Fixed lookup table of stiffness values, use uint8 indexation
@@ -192,9 +192,13 @@ inline FPBDWeightMap::FPBDWeightMap(
 }
 
 template<typename FunctorType>
-inline void FPBDWeightMap::ApplyValues(FunctorType&& MappingFunction)
+inline void FPBDWeightMap::ApplyValues(FunctorType&& MappingFunction, bool* bOutUpdated)
 {
 	SCOPE_CYCLE_COUNTER(STAT_PBD_WeightMapApplyValues);
+	if (bOutUpdated)
+	{
+		*bOutUpdated = true;
+	}
 	if (bIsDirty)
 	{
 		const FSolverReal Offset = WeightedValue[0];
diff --git a/Engine/Source/Runtime/Experimental/Chaos/Public/Chaos/SoftsSpring.h b/Engine/Source/Runtime/Experimental/Chaos/Public/Chaos/SoftsSpring.h
--- a/Engine/Source/Runtime/Experimental/Chaos/Public/Chaos/SoftsSpring.h
+++ b/Engine/Source/Runtime/Experimental/Chaos/Public/Chaos/SoftsSpring.h
@@ -47,6 +47,49 @@ FSolverVec3 GetXPBDSpringDelta(const SolverParticlesOrRange& Particles, const FS
 	return Delta;
 }
 
+template<typename SolverParticlesOrRange>
+FSolverVec3 GetXPBDAxialSpringDelta(const SolverParticlesOrRange& Particles, const FSolverReal Dt,
+	const TVec3<int32>& Constraint, const FSolverReal Bary, const FSolverReal RestLength, FSolverReal& Lambda,
+	const FSolverReal StiffnessValue, const FSolverReal MinStiffness, const FSolverReal DampingRatioValue)
+{
+	const int32 Index1 = Constraint[0];
+	const int32 Index2 = Constraint[1];
+	const int32 Index3 = Constraint[2];
+
+	const FSolverReal PInvMass = Particles.InvM(Index3) * ((FSolverReal)1. - Bary) + Particles.InvM(Index2) * Bary;
+	if (StiffnessValue < MinStiffness || (Particles.InvM(Index2) == (FSolverReal)0. && PInvMass == (FSolverReal)0.))
+	{
+		return FSolverVec3((FSolverReal)0.);
+	}
+
+	const FSolverReal CombinedInvMass = PInvMass + Particles.InvM(Index1);
+	const FSolverReal Damping = DampingRatioValue * 2.f * FMath::Sqrt(StiffnessValue / CombinedInvMass) * (RestLength > UE_SMALL_NUMBER ? (FSolverReal)1. / RestLength : (FSolverReal)1.);
+
+	const FSolverVec3& P1 = Particles.P(Index1);
+	const FSolverVec3& P2 = Particles.P(Index2);
+	const FSolverVec3& P3 = Particles.P(Index3);
+	const FSolverVec3 P = (P2 - P3) * Bary + P3;
+	FSolverVec3 Direction = P1 - P;
+	const FSolverReal Distance = Direction.SafeNormalize();
+	const FSolverReal Offset = Distance - RestLength;
+
+	const FSolverVec3& X1 = Particles.X(Index1);
+	const FSolverVec3& X2 = Particles.X(Index2);
+	const FSolverVec3& X3 = Particles.X(Index3);
+	const FSolverVec3 X = (X2 - X3) * Bary + X3;
+
+	const FSolverVec3 RelativeVelocityTimesDt = P1 - X1 - P + X;
+
+	const FSolverReal AlphaInv = StiffnessValue * Dt * Dt;
+	const FSolverReal BetaDt = Damping * Dt;
+
+	const FSolverReal DLambda = (AlphaInv * Offset - Lambda + BetaDt * FSolverVec3::DotProduct(Direction, RelativeVelocityTimesDt)) / ((AlphaInv + BetaDt) * CombinedInvMass + (FSolverReal)1.);
+	const FSolverVec3 Delta = DLambda * Direction;
+	Lambda += DLambda;
+
+	return Delta;
+}
+
 inline void UpdateSpringLinearSystem(const FSolverParticlesRange& Particles, const FSolverReal Dt,
 	const TVec2<int32>& Constraint, const FSolverReal RestLength,
 	const FSolverReal StiffnessValue, const FSolverReal MinStiffness, const FSolverReal DampingRatioValue,
@@ -114,5 +157,97 @@ inline void UpdateSpringLinearSystem(const FSolverParticlesRange& Particles, con
 		LinearSystem.AddSymmetricForceDerivative(Particles, &Df1Dx2, &Df1Dv2, Index2, Index1, Dt);
 	}
 }
+
+inline void UpdateAxialSpringLinearSystem(const FSolverParticlesRange& Particles, const FSolverReal Dt,
+	const TVec3<int32>& Constraint, const FSolverReal Bary, const FSolverReal RestLength,
+	const FSolverReal StiffnessValue, const FSolverReal MinStiffness, const FSolverReal DampingRatioValue,
+	FEvolutionLinearSystem& LinearSystem)
+{
+	const int32 Index1 = Constraint[0];
+	const int32 Index2 = Constraint[1];
+	const int32 Index3 = Constraint[2];
+
+	const FSolverReal PInvMass = Particles.InvM(Index3) * ((FSolverReal)1. - Bary) + Particles.InvM(Index2) * Bary;
+	if (StiffnessValue < MinStiffness || (Particles.InvM(Index2) == (FSolverReal)0. && PInvMass == (FSolverReal)0.))
+	{
+		return;
+	}
+
+	const FSolverVec3& P1 = Particles.P(Index1);
+	const FSolverVec3& P2 = Particles.P(Index2);
+	const FSolverVec3& P3 = Particles.P(Index3);
+	const FSolverVec3 P = Bary * P2 + ((FSolverReal)1. - Bary) * P3;
+	FSolverVec3 Direction = P1 - P;
+	const FSolverReal Distance = Direction.SafeNormalize();
+	if (Distance < UE_SMALL_NUMBER)
+	{
+		// We can't calculate a direction if distance is zero. Just skip
+		return;
+	}
+	const FSolverReal CombinedInvMass = PInvMass + Particles.InvM(Index1);
+
+	const FSolverReal Damping = DampingRatioValue * 2.f * FMath::Sqrt(StiffnessValue / CombinedInvMass) * (RestLength > UE_SMALL_NUMBER ? (FSolverReal)1. / RestLength : (FSolverReal)1.);
+
+	const FSolverVec3& V1 = Particles.V(Index1);
+	const FSolverVec3& V2 = Particles.V(Index2);
+	const FSolverVec3& V3 = Particles.V(Index3);
+	const FSolverVec3 V = Bary * V2 + ((FSolverReal)1. - Bary) * V3;
+
+	const FSolverVec3 RelVel = V1 - V;
+	const FSolverReal CDot = FSolverVec3::DotProduct(Direction, RelVel);
+
+	const FSolverReal Offset = Distance - RestLength; // C
+	// const FSolverVec3 GradC1 = Direction;
+	// const FSolverVec3 GradC2 = - Bary * GradC1;
+	// const FSolverVec3 GradC3 = - (1 - Bary) * GradC1;
+	const FSolverReal Scalar = -(StiffnessValue * Offset + Damping * CDot);
+	const FSolverVec3 Force1 = Scalar * Direction;
+	// const FSolverVec3 Force2 = -Bary * Force1
+	// const FSolverVec3 Force3 = -(1 - Bary) * Force1
+
+	const FSolverMatrix33 DirDirT = FSolverMatrix33::OuterProduct(Direction, Direction); // = GradC1 GradC1^T
+	const FSolverMatrix33 Hess11 = (FSolverReal)1. / Distance * (FSolverMatrix33::Identity - DirDirT);
+
+	// Df1Dx1 = -Stiffness * (DirDirT + Offset * Hess11) - Damping * CDot * Hess11
+	const FSolverReal ScalarModified = LinearSystem.RequiresSPDForceDerivatives() ? FMath::Min(Scalar, 0) : Scalar;
+	const FSolverMatrix33 Df1Dx1 = -StiffnessValue * DirDirT + ScalarModified * Hess11;
+	const FSolverMatrix33 Df1Dv1 = -Damping * DirDirT;
+
+	// DfiDxj = Bi * Bj * Df1Dx1, where B1 = 1, B2 = -Bary, B3 = -(1-Bary)
+	// Same for DfiDvj
+	const FSolverReal B3(Bary - (FSolverReal)1.);
+
+	if (Particles.InvM(Index1) != (FSolverReal)0.)
+	{
+		LinearSystem.AddForce(Particles, Force1, Index1, Dt);
+		LinearSystem.AddSymmetricForceDerivative(Particles, &Df1Dx1, &Df1Dv1, Index1, Index1, Dt);
+	}
+	if (Particles.InvM(Index2) != (FSolverReal)0.)
+	{
+		const FSolverMatrix33 Df2Dx2 = (Bary * Bary) * Df1Dx1;
+		const FSolverMatrix33 Df2Dv2 = (Bary * Bary) * Df1Dv1;
+		LinearSystem.AddForce(Particles, -Bary * Force1, Index2, Dt);
+		LinearSystem.AddSymmetricForceDerivative(Particles, &Df2Dx2, &Df2Dv2, Index2, Index2, Dt);
+	}
+	if (Particles.InvM(Index3) != (FSolverReal)0.)
+	{
+		const FSolverMatrix33 Df3Dx3 = (B3 * B3) * Df1Dx1;
+		const FSolverMatrix33 Df3Dv3 = (B3 * B3) * Df1Dv1;
+		LinearSystem.AddForce(Particles, -B3 * Force1, Index3, Dt);
+		LinearSystem.AddSymmetricForceDerivative(Particles, &Df3Dx3, &Df3Dv3, Index3, Index3, Dt);
+	}
+	if (Particles.InvM(Index1) != (FSolverReal)0. || Particles.InvM(Index2) != (FSolverReal)0.)
+	{
+		const FSolverMatrix33 Df1Dx2 = -Bary * Df1Dx1;
+		const FSolverMatrix33 Df1Dv2 = -Bary * Df1Dv1;
+		LinearSystem.AddSymmetricForceDerivative(Particles, &Df1Dx2, &Df1Dv2, Index1, Index2, Dt);
+	}
+	if (Particles.InvM(Index1) != (FSolverReal)0. || Particles.InvM(Index3) != (FSolverReal)0.)
+	{
+		const FSolverMatrix33 Df1Dx3 = -B3 * Df1Dx1;
+		const FSolverMatrix33 Df1Dv3 = -B3 * Df1Dv1;
+		LinearSystem.AddSymmetricForceDerivative(Particles, &Df1Dx3, &Df1Dv3, Index1, Index3, Dt);
+	}
+}
 }
 }
diff --git a/Engine/Source/Runtime/Experimental/Chaos/Public/Chaos/XPBDAnisotropicSpringConstraints.h b/Engine/Source/Runtime/Experimental/Chaos/Public/Chaos/XPBDAnisotropicSpringConstraints.h
new file mode 100644
--- /dev/null
+++ b/Engine/Source/Runtime/Experimental/Chaos/Public/Chaos/XPBDAnisotropicSpringConstraints.h
@@ -0,0 +1,416 @@
+// Copyright Epic Games, Inc. All Rights Reserved.
+#pragma once
+
+#include "Chaos/CollectionPropertyFacade.h"
+#include "Chaos/PBDAxialSpringConstraintsBase.h"
+#include "Chaos/PBDSpringConstraintsBase.h"
+#include "Chaos/PBDWeightMap.h"
+
+namespace Chaos
+{
+class FTriangleMesh;
+}
+
+namespace Chaos::Softs
+{
+
+// Edge and axial springs. Spatially varying stiffnesses and restlengths controlled by Warp/Weft/Bias.
+// Both types of springs controlled by the same input parameters.
+class FXPBDAnisotropicEdgeSpringConstraints : public FPBDSpringConstraintsBase
+{
+	typedef FPBDSpringConstraintsBase Base;
+
+public:
+	// Stiffness is in kg cm /s^2
+	static constexpr FSolverReal MinStiffness = (FSolverReal)UE_SMALL_NUMBER; // Stiffness below this will be considered 0 
+	static constexpr FSolverReal MaxStiffness = (FSolverReal)1e9;
+	static constexpr FSolverReal MinDampingRatio = (FSolverReal)0.;
+	static constexpr FSolverReal MaxDampingRatio = (FSolverReal)1000.;
+	static constexpr FSolverReal MinWarpWeftScale = (FSolverReal)0.;
+	static constexpr FSolverReal MaxWarpWeftScale = (FSolverReal)1e7; // No particular reason for this number. Just can't imagine wanting something bigger?
+
+	CHAOS_API FXPBDAnisotropicEdgeSpringConstraints(
+		const FSolverParticlesRange& Particles,
+		const FTriangleMesh& TriangleMesh,
+		const TArray<TVec3<FVec2f>>& FaceVertexPatternPositions,
+		bool bUse3dRestLengths,
+		const TConstArrayView<FRealSingle>& StiffnessWarpMultipliers,
+		const TConstArrayView<FRealSingle>& StiffnessWeftMultipliers,
+		const TConstArrayView<FRealSingle>& StiffnessBiasMultipliers,
+		const TConstArrayView<FRealSingle>& DampingMultipliers,
+		const TConstArrayView<FRealSingle>& WarpScaleMultipliers,
+		const TConstArrayView<FRealSingle>& WeftScaleMultipliers,
+		const FSolverVec2& InStiffnessWarp,
+		const FSolverVec2& InStiffnessWeft,
+		const FSolverVec2& InStiffnessBias,
+		const FSolverVec2& InDampingRatio,
+		const FSolverVec2& InWarpScale,
+		const FSolverVec2& InWeftScale,
+		bool bTrimKinematicConstraints = false);
+
+	CHAOS_API FXPBDAnisotropicEdgeSpringConstraints(
+		const FSolverParticles& InParticles,
+		int32 InParticleOffset,
+		int32 InParticleCount,
+		const FTriangleMesh& TriangleMesh,
+		const TArray<TVec3<FVec2f>>& FaceVertexPatternPositions,
+		bool bUse3dRestLengths,
+		const TConstArrayView<FRealSingle>& StiffnessWarpMultipliers,
+		const TConstArrayView<FRealSingle>& StiffnessWeftMultipliers,
+		const TConstArrayView<FRealSingle>& StiffnessBiasMultipliers,
+		const TConstArrayView<FRealSingle>& DampingMultipliers,
+		const TConstArrayView<FRealSingle>& WarpScaleMultipliers,
+		const TConstArrayView<FRealSingle>& WeftScaleMultipliers,
+		const FSolverVec2& InStiffnessWarp,
+		const FSolverVec2& InStiffnessWeft,
+		const FSolverVec2& InStiffnessBias,
+		const FSolverVec2& InDampingRatio,
+		const FSolverVec2& InWarpScale,
+		const FSolverVec2& InWeftScale,
+		bool bTrimKinematicConstraints = false);
+
+	virtual ~FXPBDAnisotropicEdgeSpringConstraints() override {}
+
+	void Init()
+	{
+		Lambdas.Reset();
+		Lambdas.AddZeroed(Constraints.Num());
+	}
+
+	// Update stiffness table, as well as the simulation stiffness exponent
+	void ApplyProperties(const FSolverReal /*Dt*/, const int32 /*NumIterations*/)
+	{
+		Stiffness.ApplyXPBDValues(MaxStiffness);
+		StiffnessWeft.ApplyXPBDValues(MaxStiffness);
+		StiffnessBias.ApplyXPBDValues(MaxStiffness);
+		DampingRatio.ApplyValues();
+		bool bWarpScaleChanged = false;
+		WarpScale.ApplyValues(&bWarpScaleChanged);
+		bool bWeftScaleChanged = false;
+		WeftScale.ApplyValues(&bWeftScaleChanged);
+		if (bWarpScaleChanged || bWeftScaleChanged)
+		{
+			// Need to update distances
+			UpdateDists();
+		}
+	}
+
+	template<typename SolverParticlesOrRange>
+	CHAOS_API void Apply(SolverParticlesOrRange& Particles, const FSolverReal Dt) const;
+
+	CHAOS_API void UpdateLinearSystem(const FSolverParticlesRange& Particles, const FSolverReal Dt, FEvolutionLinearSystem& LinearSystem) const;
+
+	const TArray<int32>& GetConstraintsPerColorStartIndex() const { return ConstraintsPerColorStartIndex; }
+	const TArray<FSolverVec3>& GetWarpWeftBiasBaseMultipliers() const { return WarpWeftBiasBaseMultipliers; }
+
+private:
+	template<typename SolverParticlesOrRange>
+	void InitColor(const SolverParticlesOrRange& InParticles);
+
+	void InitFromPatternData(bool bUse3dRestLengths, const TArray<TVec3<FVec2f>>& FaceVertexPatternPositions, const FTriangleMesh& TriangleMesh);
+
+	CHAOS_API void UpdateDists();
+	
+	template<typename SolverParticlesOrRange>
+	void ApplyHelper(SolverParticlesOrRange& Particles, const FSolverReal Dt, const int32 ConstraintIndex, const FSolverVec3& StiffnessValue, const FSolverReal DampingRatioValue) const;
+
+	using Base::Constraints;
+	using Base::ParticleOffset;
+	using Base::ParticleCount;
+	using Base::Dists; // These will be updated when Warp/Weft Scales change.
+	using Base::Stiffness; // Warp stiffness
+
+	FPBDStiffness StiffnessWeft;
+	FPBDStiffness StiffnessBias;
+	FPBDWeightMap DampingRatio;
+	FPBDWeightMap WarpScale;
+	FPBDWeightMap WeftScale;
+
+	mutable TArray<FSolverReal> Lambdas;
+	TArray<FSolverReal> BaseDists; // Without Warp/Weft Scale applied
+	TArray<FSolverVec3> WarpWeftBiasBaseMultipliers;
+	TArray<FSolverVec2> WarpWeftScaleBaseMultipliers;
+	TArray<int32> ConstraintsPerColorStartIndex; // Constraints are ordered so each batch is contiguous. This is ColorNum + 1 length so it can be used as start and end.
+
+	friend class FXPBDAnisotropicSpringConstraints;
+};
+
+class FXPBDAnisotropicAxialSpringConstraints : public FPBDAxialSpringConstraintsBase
+{
+	typedef FPBDAxialSpringConstraintsBase Base;
+
+public:
+	// Stiffness is in kg cm /s^2
+	static constexpr FSolverReal MinStiffness = (FSolverReal)UE_SMALL_NUMBER; // Stiffness below this will be considered 0 
+	static constexpr FSolverReal MaxStiffness = (FSolverReal)1e9;
+	static constexpr FSolverReal MinDampingRatio = (FSolverReal)0.;
+	static constexpr FSolverReal MaxDampingRatio = (FSolverReal)1000.;
+	static constexpr FSolverReal MinWarpWeftScale = (FSolverReal)0.;
+	static constexpr FSolverReal MaxWarpWeftScale = (FSolverReal)1e7; // No particular reason for this number. Just can't imagine wanting something bigger?	
+	
+	CHAOS_API FXPBDAnisotropicAxialSpringConstraints(
+	const FSolverParticlesRange& Particles,
+		const FTriangleMesh& TriangleMesh,
+		const TArray<TVec3<FVec2f>>& FaceVertexPatternPositions,
+		bool bUse3dRestLengths,
+		const TConstArrayView<FRealSingle>& StiffnessWarpMultipliers,
+		const TConstArrayView<FRealSingle>& StiffnessWeftMultipliers,
+		const TConstArrayView<FRealSingle>& StiffnessBiasMultipliers,
+		const TConstArrayView<FRealSingle>& DampingMultipliers,
+		const TConstArrayView<FRealSingle>& WarpScaleMultipliers,
+		const TConstArrayView<FRealSingle>& WeftScaleMultipliers,
+		const FSolverVec2& InStiffnessWarp,
+		const FSolverVec2& InStiffnessWeft,
+		const FSolverVec2& InStiffnessBias,
+		const FSolverVec2& InDampingRatio,
+		const FSolverVec2& InWarpScale,
+		const FSolverVec2& InWeftScale,
+		bool bTrimKinematicConstraints = false);
+
+		CHAOS_API FXPBDAnisotropicAxialSpringConstraints(
+			const FSolverParticles& InParticles,
+			int32 InParticleOffset,
+			int32 InParticleCount,
+			const FTriangleMesh& TriangleMesh,
+			const TArray<TVec3<FVec2f>>& FaceVertexPatternPositions,
+			bool bUse3dRestLengths,
+			const TConstArrayView<FRealSingle>& StiffnessWarpMultipliers,
+			const TConstArrayView<FRealSingle>& StiffnessWeftMultipliers,
+			const TConstArrayView<FRealSingle>& StiffnessBiasMultipliers,
+			const TConstArrayView<FRealSingle>& DampingMultipliers,
+			const TConstArrayView<FRealSingle>& WarpScaleMultipliers,
+			const TConstArrayView<FRealSingle>& WeftScaleMultipliers,
+			const FSolverVec2& InStiffnessWarp,
+			const FSolverVec2& InStiffnessWeft,
+			const FSolverVec2& InStiffnessBias,
+			const FSolverVec2& InDampingRatio,
+			const FSolverVec2& InWarpScale,
+			const FSolverVec2& InWeftScale,
+			bool bTrimKinematicConstraints = false);
+
+		virtual ~FXPBDAnisotropicAxialSpringConstraints() override {}
+
+		void Init()
+		{
+			Lambdas.Reset();
+			Lambdas.AddZeroed(Constraints.Num());
+		}
+
+		// Update stiffness table, as well as the simulation stiffness exponent
+		void ApplyProperties(const FSolverReal /*Dt*/, const int32 /*NumIterations*/)
+		{
+			Stiffness.ApplyXPBDValues(MaxStiffness);
+			StiffnessWeft.ApplyXPBDValues(MaxStiffness);
+			StiffnessBias.ApplyXPBDValues(MaxStiffness);
+			DampingRatio.ApplyValues();
+			bool bWarpScaleChanged = false;
+			WarpScale.ApplyValues(&bWarpScaleChanged);
+			bool bWeftScaleChanged = false;
+			WeftScale.ApplyValues(&bWeftScaleChanged);
+			if (bWarpScaleChanged || bWeftScaleChanged)
+			{
+				// Need to update distances
+				UpdateDists();
+			}
+		}
+
+		template<typename SolverParticlesOrRange>
+		CHAOS_API void Apply(SolverParticlesOrRange& Particles, const FSolverReal Dt) const;
+
+		CHAOS_API void UpdateLinearSystem(const FSolverParticlesRange& Particles, const FSolverReal Dt, FEvolutionLinearSystem& LinearSystem) const;
+
+		const TArray<int32>& GetConstraintsPerColorStartIndex() const { return ConstraintsPerColorStartIndex; }
+		const TArray<FSolverVec3>& GetWarpWeftBiasBaseMultipliers() const { return WarpWeftBiasBaseMultipliers; }
+
+private:
+	template<typename SolverParticlesOrRange>
+	void InitColor(const SolverParticlesOrRange& InParticles);
+
+	void InitFromPatternData(bool bUse3dRestLengths, const TArray<TVec3<FVec2f>>& FaceVertexPatternPositions, const FTriangleMesh& TriangleMesh);
+
+	CHAOS_API void UpdateDists();
+
+	template<typename SolverParticlesOrRange>
+	void ApplyHelper(SolverParticlesOrRange& Particles, const FSolverReal Dt, const int32 ConstraintIndex, const FSolverVec3& StiffnessValue, const FSolverReal DampingRatioValue) const;
+
+	using Base::Constraints;
+	using Base::ParticleOffset;
+	using Base::ParticleCount;
+	using Base::Barys;
+	using Base::Dists; // These will be updated when Warp/Weft Scales change.
+	using Base::Stiffness; // Warp stiffness
+
+	FPBDStiffness StiffnessWeft;
+	FPBDStiffness StiffnessBias;
+	FPBDWeightMap DampingRatio;
+	FPBDWeightMap WarpScale;
+	FPBDWeightMap WeftScale;
+
+	mutable TArray<FSolverReal> Lambdas;
+	TArray<FSolverReal> BaseDists; // Without Warp/Weft Scale applied
+	TArray<FSolverVec3> WarpWeftBiasBaseMultipliers;
+	TArray<FSolverVec2> WarpWeftScaleBaseMultipliers;
+	TArray<int32> ConstraintsPerColorStartIndex; // Constraints are ordered so each batch is contiguous. This is ColorNum + 1 length so it can be used as start and end.
+
+	friend class FXPBDAnisotropicSpringConstraints;
+
+};
+
+class FXPBDAnisotropicSpringConstraints final
+{
+public:
+	static constexpr FSolverReal DefaultStiffness = (FSolverReal)100.;
+	static constexpr FSolverReal DefaultDamping = (FSolverReal)1.;
+	static constexpr bool bDefaultUse3dRestLengths = true;
+	static constexpr FSolverReal DefaultWarpWeftScale = (FSolverReal)1.;
+
+	static bool IsEnabled(const FCollectionPropertyConstFacade& PropertyCollection)
+	{
+		return IsXPBDAnisoSpringStiffnessWarpEnabled(PropertyCollection, false);
+	}
+
+	FXPBDAnisotropicSpringConstraints(
+		const FSolverParticlesRange& Particles,
+		const FTriangleMesh& TriangleMesh,
+		const TArray<TVec3<FVec2f>>& FaceVertexPatternPositions,
+		const TMap<FString, TConstArrayView<FRealSingle>>& WeightMaps,
+		const FCollectionPropertyConstFacade& PropertyCollection,
+		bool bTrimKinematicConstraints = false)
+		: EdgeConstraints(
+			Particles,
+			TriangleMesh,
+			FaceVertexPatternPositions,
+			GetXPBDAnisoSpringUse3dRestLengths(PropertyCollection, bDefaultUse3dRestLengths),
+			WeightMaps.FindRef(GetXPBDAnisoSpringStiffnessWarpString(PropertyCollection, XPBDAnisoSpringStiffnessWarpName.ToString())),
+			WeightMaps.FindRef(GetXPBDAnisoSpringStiffnessWeftString(PropertyCollection, XPBDAnisoSpringStiffnessWeftName.ToString())),
+			WeightMaps.FindRef(GetXPBDAnisoSpringStiffnessBiasString(PropertyCollection, XPBDAnisoSpringStiffnessBiasName.ToString())),
+			WeightMaps.FindRef(GetXPBDAnisoSpringDampingString(PropertyCollection, XPBDAnisoSpringDampingName.ToString())),
+			WeightMaps.FindRef(GetXPBDAnisoSpringWarpScaleString(PropertyCollection, XPBDAnisoSpringWarpScaleName.ToString())),
+			WeightMaps.FindRef(GetXPBDAnisoSpringWeftScaleString(PropertyCollection, XPBDAnisoSpringWeftScaleName.ToString())),
+			FSolverVec2(GetWeightedFloatXPBDAnisoSpringStiffnessWarp(PropertyCollection, DefaultStiffness)),
+			FSolverVec2(GetWeightedFloatXPBDAnisoSpringStiffnessWeft(PropertyCollection, DefaultStiffness)),
+			FSolverVec2(GetWeightedFloatXPBDAnisoSpringStiffnessBias(PropertyCollection, DefaultStiffness)),
+			FSolverVec2(GetWeightedFloatXPBDAnisoSpringDamping(PropertyCollection, DefaultDamping)),
+			FSolverVec2(GetWeightedFloatXPBDAnisoSpringWarpScale(PropertyCollection, DefaultWarpWeftScale)),
+			FSolverVec2(GetWeightedFloatXPBDAnisoSpringWeftScale(PropertyCollection, DefaultWarpWeftScale)),
+			bTrimKinematicConstraints)
+		, AxialConstraints(
+			Particles,
+			TriangleMesh,
+			FaceVertexPatternPositions,
+			GetXPBDAnisoSpringUse3dRestLengths(PropertyCollection, bDefaultUse3dRestLengths),
+			WeightMaps.FindRef(GetXPBDAnisoSpringStiffnessWarpString(PropertyCollection, XPBDAnisoSpringStiffnessWarpName.ToString())),
+			WeightMaps.FindRef(GetXPBDAnisoSpringStiffnessWeftString(PropertyCollection, XPBDAnisoSpringStiffnessWeftName.ToString())),
+			WeightMaps.FindRef(GetXPBDAnisoSpringStiffnessBiasString(PropertyCollection, XPBDAnisoSpringStiffnessBiasName.ToString())),
+			WeightMaps.FindRef(GetXPBDAnisoSpringDampingString(PropertyCollection, XPBDAnisoSpringDampingName.ToString())),
+			WeightMaps.FindRef(GetXPBDAnisoSpringWarpScaleString(PropertyCollection, XPBDAnisoSpringWarpScaleName.ToString())),
+			WeightMaps.FindRef(GetXPBDAnisoSpringWeftScaleString(PropertyCollection, XPBDAnisoSpringWeftScaleName.ToString())),
+			FSolverVec2(GetWeightedFloatXPBDAnisoSpringStiffnessWarp(PropertyCollection, DefaultStiffness)),
+			FSolverVec2(GetWeightedFloatXPBDAnisoSpringStiffnessWeft(PropertyCollection, DefaultStiffness)),
+			FSolverVec2(GetWeightedFloatXPBDAnisoSpringStiffnessBias(PropertyCollection, DefaultStiffness)),
+			FSolverVec2(GetWeightedFloatXPBDAnisoSpringDamping(PropertyCollection, DefaultDamping)),
+			FSolverVec2(GetWeightedFloatXPBDAnisoSpringWarpScale(PropertyCollection, DefaultWarpWeftScale)),
+			FSolverVec2(GetWeightedFloatXPBDAnisoSpringWeftScale(PropertyCollection, DefaultWarpWeftScale)),
+			bTrimKinematicConstraints)
+		, XPBDAnisoSpringUse3dRestLengthsIndex(PropertyCollection)
+		, XPBDAnisoSpringStiffnessWarpIndex(PropertyCollection)
+		, XPBDAnisoSpringStiffnessWeftIndex(PropertyCollection)
+		, XPBDAnisoSpringStiffnessBiasIndex(PropertyCollection)
+		, XPBDAnisoSpringDampingIndex(PropertyCollection)
+		, XPBDAnisoSpringWarpScaleIndex(PropertyCollection)
+		, XPBDAnisoSpringWeftScaleIndex(PropertyCollection)
+	{}
+
+	FXPBDAnisotropicSpringConstraints(
+		const FSolverParticles& InParticles,
+		int32 InParticleOffset,
+		int32 InParticleCount,
+		const FTriangleMesh& TriangleMesh,
+		const TArray<TVec3<FVec2f>>& FaceVertexPatternPositions,
+		const TMap<FString, TConstArrayView<FRealSingle>>& WeightMaps,
+		const FCollectionPropertyConstFacade& PropertyCollection,
+		bool bTrimKinematicConstraints = false)
+		: EdgeConstraints(
+			InParticles,
+			InParticleOffset,
+			InParticleCount,
+			TriangleMesh,
+			FaceVertexPatternPositions,
+			GetXPBDAnisoSpringUse3dRestLengths(PropertyCollection, bDefaultUse3dRestLengths),
+			WeightMaps.FindRef(GetXPBDAnisoSpringStiffnessWarpString(PropertyCollection, XPBDAnisoSpringStiffnessWarpName.ToString())),
+			WeightMaps.FindRef(GetXPBDAnisoSpringStiffnessWeftString(PropertyCollection, XPBDAnisoSpringStiffnessWeftName.ToString())),
+			WeightMaps.FindRef(GetXPBDAnisoSpringStiffnessBiasString(PropertyCollection, XPBDAnisoSpringStiffnessBiasName.ToString())),
+			WeightMaps.FindRef(GetXPBDAnisoSpringDampingString(PropertyCollection, XPBDAnisoSpringDampingName.ToString())),
+			WeightMaps.FindRef(GetXPBDAnisoSpringWarpScaleString(PropertyCollection, XPBDAnisoSpringWarpScaleName.ToString())),
+			WeightMaps.FindRef(GetXPBDAnisoSpringWeftScaleString(PropertyCollection, XPBDAnisoSpringWeftScaleName.ToString())),
+			FSolverVec2(GetWeightedFloatXPBDAnisoSpringStiffnessWarp(PropertyCollection, DefaultStiffness)),
+			FSolverVec2(GetWeightedFloatXPBDAnisoSpringStiffnessWeft(PropertyCollection, DefaultStiffness)),
+			FSolverVec2(GetWeightedFloatXPBDAnisoSpringStiffnessBias(PropertyCollection, DefaultStiffness)),
+			FSolverVec2(GetWeightedFloatXPBDAnisoSpringDamping(PropertyCollection, DefaultDamping)),
+			FSolverVec2(GetWeightedFloatXPBDAnisoSpringWarpScale(PropertyCollection, DefaultWarpWeftScale)),
+			FSolverVec2(GetWeightedFloatXPBDAnisoSpringWeftScale(PropertyCollection, DefaultWarpWeftScale)),
+			bTrimKinematicConstraints)
+		, AxialConstraints(
+			InParticles,
+			InParticleOffset,
+			InParticleCount,
+			TriangleMesh,
+			FaceVertexPatternPositions,
+			GetXPBDAnisoSpringUse3dRestLengths(PropertyCollection, bDefaultUse3dRestLengths),
+			WeightMaps.FindRef(GetXPBDAnisoSpringStiffnessWarpString(PropertyCollection, XPBDAnisoSpringStiffnessWarpName.ToString())),
+			WeightMaps.FindRef(GetXPBDAnisoSpringStiffnessWeftString(PropertyCollection, XPBDAnisoSpringStiffnessWeftName.ToString())),
+			WeightMaps.FindRef(GetXPBDAnisoSpringStiffnessBiasString(PropertyCollection, XPBDAnisoSpringStiffnessBiasName.ToString())),
+			WeightMaps.FindRef(GetXPBDAnisoSpringDampingString(PropertyCollection, XPBDAnisoSpringDampingName.ToString())),
+			WeightMaps.FindRef(GetXPBDAnisoSpringWarpScaleString(PropertyCollection, XPBDAnisoSpringWarpScaleName.ToString())),
+			WeightMaps.FindRef(GetXPBDAnisoSpringWeftScaleString(PropertyCollection, XPBDAnisoSpringWeftScaleName.ToString())),
+			FSolverVec2(GetWeightedFloatXPBDAnisoSpringStiffnessWarp(PropertyCollection, DefaultStiffness)),
+			FSolverVec2(GetWeightedFloatXPBDAnisoSpringStiffnessWeft(PropertyCollection, DefaultStiffness)),
+			FSolverVec2(GetWeightedFloatXPBDAnisoSpringStiffnessBias(PropertyCollection, DefaultStiffness)),
+			FSolverVec2(GetWeightedFloatXPBDAnisoSpringDamping(PropertyCollection, DefaultDamping)),
+			FSolverVec2(GetWeightedFloatXPBDAnisoSpringWarpScale(PropertyCollection, DefaultWarpWeftScale)),
+			FSolverVec2(GetWeightedFloatXPBDAnisoSpringWeftScale(PropertyCollection, DefaultWarpWeftScale)),
+			bTrimKinematicConstraints)
+		, XPBDAnisoSpringUse3dRestLengthsIndex(PropertyCollection)
+		, XPBDAnisoSpringStiffnessWarpIndex(PropertyCollection)
+		, XPBDAnisoSpringStiffnessWeftIndex(PropertyCollection)
+		, XPBDAnisoSpringStiffnessBiasIndex(PropertyCollection)
+		, XPBDAnisoSpringDampingIndex(PropertyCollection)
+		, XPBDAnisoSpringWarpScaleIndex(PropertyCollection)
+		, XPBDAnisoSpringWeftScaleIndex(PropertyCollection)
+	{}
+
+	CHAOS_API void SetProperties(
+		const FCollectionPropertyConstFacade& PropertyCollection,
+		const TMap<FString, TConstArrayView<FRealSingle>>& WeightMaps);
+
+	// Update stiffness table, as well as the simulation stiffness exponent
+	void ApplyProperties(const FSolverReal Dt, const int32 NumIterations)
+	{
+		EdgeConstraints.ApplyProperties(Dt, NumIterations);
+		AxialConstraints.ApplyProperties(Dt, NumIterations);
+	}
+
+	void Init()
+	{
+		EdgeConstraints.Init();
+		AxialConstraints.Init();
+	}
+
+	const FXPBDAnisotropicEdgeSpringConstraints& GetEdgeConstraints() const { return EdgeConstraints; }
+	const FXPBDAnisotropicAxialSpringConstraints& GetAxialConstraints() const { return AxialConstraints; }
+
+private:
+	FXPBDAnisotropicEdgeSpringConstraints EdgeConstraints;
+	FXPBDAnisotropicAxialSpringConstraints AxialConstraints;
+
+	UE_CHAOS_DECLARE_PROPERTYCOLLECTION_NAME(XPBDAnisoSpringUse3dRestLengths, bool);
+	UE_CHAOS_DECLARE_PROPERTYCOLLECTION_NAME(XPBDAnisoSpringStiffnessWarp, float);
+	UE_CHAOS_DECLARE_PROPERTYCOLLECTION_NAME(XPBDAnisoSpringStiffnessWeft, float);
+	UE_CHAOS_DECLARE_PROPERTYCOLLECTION_NAME(XPBDAnisoSpringStiffnessBias, float);
+	UE_CHAOS_DECLARE_PROPERTYCOLLECTION_NAME(XPBDAnisoSpringDamping, float);
+	UE_CHAOS_DECLARE_PROPERTYCOLLECTION_NAME(XPBDAnisoSpringWarpScale, float);
+	UE_CHAOS_DECLARE_PROPERTYCOLLECTION_NAME(XPBDAnisoSpringWeftScale, float);
+};
+
+}  // End namespace Chaos::Softs
+
