Iris
- Added code to trap cases where we call EndReplication from client code where it shouldnt be called
- Fixed issue with client calling EndReplication on components when it should not.

diff --git a/Engine/Source/Runtime/Engine/Classes/Components/ActorComponent.h b/Engine/Source/Runtime/Engine/Classes/Components/ActorComponent.h
--- a/Engine/Source/Runtime/Engine/Classes/Components/ActorComponent.h
+++ b/Engine/Source/Runtime/Engine/Classes/Components/ActorComponent.h
@@ -638,10 +638,10 @@ public:
 	/** Register all replication fragments */
 	ENGINE_API virtual void RegisterReplicationFragments(UE::Net::FFragmentRegistrationContext& Context, UE::Net::EFragmentRegistrationFlags RegistrationFlags) override;
 	
-	/** Called when we want to start replicating this component */
+	/** Called when we want to start replicating this component, should not be called explicitly.*/
 	ENGINE_API virtual void BeginReplication();
 
-	/** Tell component to end replication */
+	/** Tell component to end replication, should not be called explicitly. */
 	ENGINE_API virtual void EndReplication();
 #endif // UE_WITH_IRIS
 
diff --git a/Engine/Source/Runtime/Engine/Private/Actor.cpp b/Engine/Source/Runtime/Engine/Private/Actor.cpp
--- a/Engine/Source/Runtime/Engine/Private/Actor.cpp
+++ b/Engine/Source/Runtime/Engine/Private/Actor.cpp
@@ -3347,14 +3347,17 @@ void AActor::UpdateReplicatedComponent(UActorComponent* Component)
 	}
 
 #if UE_WITH_IRIS
-	if (NetCondition != COND_Never && HasActorBegunPlay())
+	if (HasAuthority())
 	{
-		// Begin replication and set NetCondition, if component already is replicated the NetCondition will be updated
-		Component->BeginReplication();
-	}
-	else if (bWasRemovedFromReplicatedComponents || NetCondition == COND_Never)
-	{
-		Component->EndReplication();
+		if (NetCondition != COND_Never && HasActorBegunPlay())
+		{
+			// Begin replication and set NetCondition, if component already is replicated the NetCondition will be updated
+			Component->BeginReplication();
+		}
+		else if (bWasRemovedFromReplicatedComponents || NetCondition == COND_Never)
+		{
+			Component->EndReplication();
+		}
 	}
 #endif // UE_WITH_IRIS
 }
@@ -3370,10 +3373,13 @@ void AActor::UpdateAllReplicatedComponents()
 	{
 		RepComponentInfo.NetCondition = COND_Never;
 #if UE_WITH_IRIS
-		// Need to end replication for all components that should no longer replicate
-		if (RepComponentInfo.Component && !RepComponentInfo.Component->GetIsReplicated())
+		if (HasAuthority())
 		{
-			RepComponentInfo.Component->EndReplication();
+			// Need to end replication for all components that should no longer replicate
+			if (RepComponentInfo.Component && !RepComponentInfo.Component->GetIsReplicated())
+			{
+				RepComponentInfo.Component->EndReplication();
+			}
 		}
 #endif
 	}
@@ -3397,15 +3403,18 @@ void AActor::UpdateAllReplicatedComponents()
 				}
 
 #if UE_WITH_IRIS
-	            if (NetCondition != COND_Never)
+				if (HasAuthority())
 				{
-					// Begin replication and set NetCondition, if component already is replicated the NetCondition will be updated
-					Component->BeginReplication();
+					if (NetCondition != COND_Never)
+					{
+						// Begin replication and set NetCondition, if component already is replicated the NetCondition will be updated
+						Component->BeginReplication();
+					}
+					else
+					{
+						Component->EndReplication();
+					}
 				}
-				else
-				{
-					Component->EndReplication();
-                }
 #endif
 			}
 		}
diff --git a/Engine/Source/Runtime/Engine/Private/ActorReplication.cpp b/Engine/Source/Runtime/Engine/Private/ActorReplication.cpp
--- a/Engine/Source/Runtime/Engine/Private/ActorReplication.cpp
+++ b/Engine/Source/Runtime/Engine/Private/ActorReplication.cpp
@@ -681,7 +681,10 @@ void AActor::RemoveReplicatedComponent(UActorComponent* Component)
 	{
 		ReplicatedComponentsInfo.RemoveAtSwap(Index);
 #if UE_WITH_IRIS
-		Component->EndReplication();
+		if (HasAuthority())
+		{
+			Component->EndReplication();
+		}
 #endif
 	}
 }
@@ -732,7 +735,7 @@ void AActor::RemoveReplicatedSubObject(UObject* SubObject)
 {
 	const bool bWasRemoved = RemoveReplicatedSubObjectFromList(SubObject);
 #if UE_WITH_IRIS
-	if (bWasRemoved)
+	if (bWasRemoved && HasAuthority())
 	{
 		UE::Net::FReplicationSystemUtil::EndReplicationForActorSubObject(this, SubObject);
 	}
@@ -931,7 +934,7 @@ void AActor::RemoveActorComponentReplicatedSubObject(UActorComponent* OwnerCompo
 	const bool bWasRemoved = RemoveActorComponentReplicatedSubObjectFromList(OwnerComponent, SubObject);
 	
 #if UE_WITH_IRIS
-	if (bWasRemoved)
+	if (bWasRemoved && HasAuthority())
 	{
 		UE::Net::FReplicationSystemUtil::EndReplicationForActorComponentSubObject(OwnerComponent, SubObject);
 	}
diff --git a/Engine/Source/Runtime/Engine/Private/Components/ActorComponent.cpp b/Engine/Source/Runtime/Engine/Private/Components/ActorComponent.cpp
--- a/Engine/Source/Runtime/Engine/Private/Components/ActorComponent.cpp
+++ b/Engine/Source/Runtime/Engine/Private/Components/ActorComponent.cpp
@@ -1121,6 +1121,13 @@ void UActorComponent::EndPlay(const EEndPlayReason::Type EndPlayReason)
 
 	check(bHasBegunPlay);
 
+#if UE_WITH_IRIS
+	if (EndPlayReason != EEndPlayReason::EndPlayInEditor && EndPlayReason != EEndPlayReason::Quit)
+	{
+		EndReplication();
+	}
+#endif
+
 	// If we're in the process of being garbage collected it is unsafe to call out to blueprints
 	if (!HasAnyFlags(RF_BeginDestroyed) && !IsUnreachable() && (GetClass()->HasAnyClassFlags(CLASS_CompiledFromBlueprint) || !GetClass()->HasAnyClassFlags(CLASS_Native)))
 	{
diff --git a/Engine/Source/Runtime/Engine/Private/Net/Iris/ReplicationSystem/ActorReplicationBridge.cpp b/Engine/Source/Runtime/Engine/Private/Net/Iris/ReplicationSystem/ActorReplicationBridge.cpp
--- a/Engine/Source/Runtime/Engine/Private/Net/Iris/ReplicationSystem/ActorReplicationBridge.cpp
+++ b/Engine/Source/Runtime/Engine/Private/Net/Iris/ReplicationSystem/ActorReplicationBridge.cpp
@@ -432,6 +432,12 @@ void UActorReplicationBridge::EndReplication(AActor* Actor, EEndPlayReason::Type
 		{ 
 			Flags |= EEndReplicationFlags::Destroy | EEndReplicationFlags::DestroyNetHandle | EEndReplicationFlags::ClearNetPushId;
 		}
+
+		// If we are shutting down we do not need to validate that we are not detaching remote instances by accident.
+		if (EndPlayReason == EEndPlayReason::EndPlayInEditor || EndPlayReason != EEndPlayReason::Quit)
+		{
+			Flags |= EEndReplicationFlags::SkipPendingEndReplicationValidation;
+		}
 					
 		const bool bShouldCreateDestructionInfo = RefHandle.IsStatic() && bShouldDestroyObject && GetReplicationSystem()->IsServer();
 		if (bShouldCreateDestructionInfo)
diff --git a/Engine/Source/Runtime/Experimental/Iris/Core/Private/Iris/ReplicationSystem/NetRefHandleManager.cpp b/Engine/Source/Runtime/Experimental/Iris/Core/Private/Iris/ReplicationSystem/NetRefHandleManager.cpp
--- a/Engine/Source/Runtime/Experimental/Iris/Core/Private/Iris/ReplicationSystem/NetRefHandleManager.cpp
+++ b/Engine/Source/Runtime/Experimental/Iris/Core/Private/Iris/ReplicationSystem/NetRefHandleManager.cpp
@@ -151,6 +151,24 @@ const FReplicationInstanceProtocol* FNetRefHandleManager::DetachInstanceProtocol
 	{
 		FReplicatedObjectData& Data = ReplicatedObjectData[InternalIndex];
 		const FReplicationInstanceProtocol* InstanceProtocol = Data.InstanceProtocol;
+
+		// Detect if someone unexpectedly tries to detach replicated instance
+		const bool bDetachedInstancedIsPendingEndReplication = ReplicatedObjectData[InternalIndex].bPendingEndReplication == 1U;
+		if (!bDetachedInstancedIsPendingEndReplication)
+		{
+			ensureMsgf(false, TEXT("DetachInstanceProtocol for %s (%s) from client"), *Data.RefHandle.ToString(), Data.Protocol->DebugName->Name);
+
+			if (Data.SubObjectRootIndex != InvalidInternalIndex)
+			{
+				FReplicatedObjectData& RootObjectData = ReplicatedObjectData[Data.SubObjectRootIndex];
+				UE_LOG(LogIris, Warning, TEXT("FNetRefHandleManager::DetachInstanceProtocol - DetachInstanceProtocol %s (%s Root:%s) from client"), *Data.RefHandle.ToString(), Data.Protocol->DebugName->Name, RootObjectData.Protocol->DebugName->Name);
+			}
+			else
+			{
+				UE_LOG(LogIris, Warning, TEXT("FNetRefHandleManager::DetachInstanceProtocol - DetachInstanceProtocol %s (%s) from client"), *Data.RefHandle.ToString(), Data.Protocol->DebugName->Name);
+			}
+		}
+		
 		Data.InstanceProtocol = nullptr;
 		ReplicatedInstances[InternalIndex] = nullptr;
 		
diff --git a/Engine/Source/Runtime/Experimental/Iris/Core/Private/Iris/ReplicationSystem/NetRefHandleManager.h b/Engine/Source/Runtime/Experimental/Iris/Core/Private/Iris/ReplicationSystem/NetRefHandleManager.h
--- a/Engine/Source/Runtime/Experimental/Iris/Core/Private/Iris/ReplicationSystem/NetRefHandleManager.h
+++ b/Engine/Source/Runtime/Experimental/Iris/Core/Private/Iris/ReplicationSystem/NetRefHandleManager.h
@@ -101,6 +101,7 @@ public:
 				uint32 bAllowDestroyInstanceFromRemote : 1U;
 				uint32 bNeedsFullCopyAndQuantize : 1U;
 				uint32 bWantsFullPoll : 1U;
+				uint32 bPendingEndReplication : 1U;
 			};
 		};
 	
@@ -173,7 +174,10 @@ public:
 	/** Verify a handle against internal handle. A handle is valid if it matches internal storage. */
 	inline bool IsValidNetRefHandle(FNetRefHandle Handle) const;
 
-	/** Returns true if the handle is for a local replicated object. */
+	/** Returns true if the InternaIndex belongs to a replicated object owned by the local peer. */
+	inline bool IsLocal(FInternalNetRefIndex InternalIndex) const;
+
+	/** Returns true if the handle belongs to a replicated object owned by the local peer. */
 	inline bool IsLocalNetRefHandle(FNetRefHandle Handle) const;
 
 	/** Returns true if the handle is for a remote replicated object. */
@@ -448,15 +452,14 @@ bool FNetRefHandleManager::IsValidNetRefHandle(FNetRefHandle Handle) const
 	return RefHandleToInternalIndex.Contains(Handle);
 }
 
-bool FNetRefHandleManager::IsLocalNetRefHandle(FNetRefHandle Handle) const
+bool FNetRefHandleManager::IsLocal(FInternalNetRefIndex InternalIndex) const
 {
-	if (const FInternalNetRefIndex InternalIndex = GetInternalIndex(Handle))
-	{
-		// For the time being only replicated objects owned by this peer has a state buffer
-		return ReplicatedObjectStateBuffers[InternalIndex] != nullptr;
-	}
+	return (InternalIndex != InvalidInternalIndex) && (ReplicatedObjectStateBuffers[InternalIndex] != nullptr);
+}
 
-	return false;
+bool FNetRefHandleManager::IsLocalNetRefHandle(FNetRefHandle Handle) const
+{
+	return IsLocal(GetInternalIndex(Handle));
 }
 
 bool FNetRefHandleManager::IsRemoteNetRefHandle(FNetRefHandle Handle) const
diff --git a/Engine/Source/Runtime/Experimental/Iris/Core/Private/Iris/ReplicationSystem/ReplicationBridge.cpp b/Engine/Source/Runtime/Experimental/Iris/Core/Private/Iris/ReplicationSystem/ReplicationBridge.cpp
--- a/Engine/Source/Runtime/Experimental/Iris/Core/Private/Iris/ReplicationSystem/ReplicationBridge.cpp
+++ b/Engine/Source/Runtime/Experimental/Iris/Core/Private/Iris/ReplicationSystem/ReplicationBridge.cpp
@@ -129,6 +129,7 @@ void UReplicationBridge::DetachSubObjectInstancesFromRemote(FNetRefHandle OwnerH
 			FNetRefHandleManager::FReplicatedObjectData& SubObjectData = NetRefHandleManager->GetReplicatedObjectDataNoCheck(SubObjectInternalIndex);
 			const FNetRefHandle SubObjectHandle = SubObjectData.RefHandle;
 			SubObjectData.bTearOff = (DestroyReason == EReplicationBridgeDestroyInstanceReason::TearOff);
+			SubObjectData.bPendingEndReplication = 1U;
 
 			EReplicationBridgeDestroyInstanceFlags SubObjectDestroyFlags = DestroyFlags;
 			// The subobject is allowed to be destroyed if both the owner and the subobject allows it.
@@ -150,6 +151,7 @@ void UReplicationBridge::DestroyNetObjectFromRemote(FNetRefHandle Handle, ERepli
 		FInternalNetRefIndex OwnerInternalIndex = NetRefHandleManager->GetInternalIndex(Handle);
 		FNetRefHandleManager::FReplicatedObjectData& ObjectData = NetRefHandleManager->GetReplicatedObjectDataNoCheck(OwnerInternalIndex);
 		ObjectData.bTearOff = (DestroyReason == EReplicationBridgeDestroyInstanceReason::TearOff);
+		ObjectData.bPendingEndReplication = 1U;
 
 		// if the a subobject owner is to be destroyed we want to detach all subobjects before doing so to ensure we execute expected callbacks
 		// We keep tracking them internally
@@ -486,13 +488,15 @@ void UReplicationBridge::InternalDestroySubObjects(FNetRefHandle OwnerHandle, EE
 	{
 		for (FInternalNetRefIndex SubObjectInternalIndex : NetRefHandleManager->GetChildSubObjects(OwnerInternalIndex))
 		{
-			const FNetRefHandleManager::FReplicatedObjectData& SubObjectData = NetRefHandleManager->GetReplicatedObjectData(SubObjectInternalIndex);
+			FNetRefHandleManager::FReplicatedObjectData& SubObjectData = NetRefHandleManager->GetReplicatedObjectDataNoCheck(SubObjectInternalIndex);
+
 			const FNetRefHandle SubObjectHandle = SubObjectData.RefHandle;
 			const bool bDestroySubObjectWithOwner = SubObjectData.bDestroySubObjectWithOwner;
 				
 			// Tag subobject for destroy. The check against the scope is needed since the subobjects array might contain subobjects already pending destroy.
 			if (bDestroySubObjectWithOwner && NetRefHandleManager->IsScopableIndex(SubObjectInternalIndex))
 			{
+				SubObjectData.bPendingEndReplication = 1U;
 				UE_LOG_REPLICATIONBRIDGE(Verbose, TEXT("InternalDestroySubObjects %s - SubObject %s"), *OwnerHandle.ToString(), *SubObjectHandle.ToString());
 				DestroyLocalNetHandle(SubObjectHandle, Flags);
 			}
@@ -502,12 +506,15 @@ void UReplicationBridge::InternalDestroySubObjects(FNetRefHandle OwnerHandle, EE
 
 void UReplicationBridge::EndReplication(FNetRefHandle Handle, EEndReplicationFlags EndReplicationFlags, FEndReplicationParameters* Parameters)
 {
+	using namespace UE::Net::Private;
+
 	if (!IsReplicatedHandle(Handle))
 	{
 		return;
 	}
 
-	if (NetRefHandleManager->IsLocalNetRefHandle(Handle))
+	const FInternalNetRefIndex InternalReplicationIndex = NetRefHandleManager->GetInternalIndex(Handle);
+	if (NetRefHandleManager->IsLocal(InternalReplicationIndex))
 	{
 		if (EnumHasAnyFlags(EndReplicationFlags, EEndReplicationFlags::TearOff))
 		{
@@ -537,11 +544,16 @@ void UReplicationBridge::EndReplication(FNetRefHandle Handle, EEndReplicationFla
 				InternalFlushStateData(Handle);
 			}
 
+			NetRefHandleManager->GetReplicatedObjectDataNoCheck(InternalReplicationIndex).bPendingEndReplication = 1U;
 			DestroyLocalNetHandle(Handle, EndReplicationFlags);	
 		}
 	}
 	else
 	{
+		if (InternalReplicationIndex != FNetRefHandleManager::InvalidInternalIndex && EnumHasAnyFlags(EndReplicationFlags, EEndReplicationFlags::SkipPendingEndReplicationValidation))
+		{
+			NetRefHandleManager->GetReplicatedObjectDataNoCheck(InternalReplicationIndex).bPendingEndReplication = 1U;
+		}
 		// If we get a call to end replication on the client, we need to detach the instance as it might be garbage collected
 		InternalDetachInstanceFromNetRefHandle(Handle);
 	}
@@ -626,6 +638,8 @@ void UReplicationBridge::UpdateHandlesPendingTearOff()
 			// Immediate tear-off or object that no longer are referenced by any connections are destroyed
 			if (NetRefHandleManager->GetNetObjectRefCount(ObjectInternalIndex) == 0U || Info.bIsImmediate)
 			{
+				FNetRefHandleManager::FReplicatedObjectData& ObjectData = NetRefHandleManager->GetReplicatedObjectDataNoCheck(ObjectInternalIndex);
+				ObjectData.bPendingEndReplication = 1U;
 				DestroyLocalNetHandle(Info.Handle, Info.DestroyFlags);
 			}
 			else
diff --git a/Engine/Source/Runtime/Experimental/Iris/Core/Public/Iris/ReplicationSystem/ReplicationBridge.h b/Engine/Source/Runtime/Experimental/Iris/Core/Public/Iris/ReplicationSystem/ReplicationBridge.h
--- a/Engine/Source/Runtime/Experimental/Iris/Core/Public/Iris/ReplicationSystem/ReplicationBridge.h
+++ b/Engine/Source/Runtime/Experimental/Iris/Core/Public/Iris/ReplicationSystem/ReplicationBridge.h
@@ -74,6 +74,8 @@ enum class EEndReplicationFlags : uint32
 	DestroyNetHandle					= Flush << 1U,
 	/** Clear net push ID to prevent this object and its subobjects from being marked as dirty in the networking system. This should only be done if the object should not be replicated by any other replication system. */
 	ClearNetPushId						= DestroyNetHandle << 1U,
+	/** Skip bPendingEndReplication Validation, In some cases we want to allow detaching instance from replicated object on clients, such as when shutting down */
+	SkipPendingEndReplicationValidation = ClearNetPushId << 1U,
 };
 ENUM_CLASS_FLAGS(EEndReplicationFlags);
 
