Constraints: clean invalid constraints on post GC
    - invalidate constraint arrays on post GC to ensure that we don't keep invalid constraints

diff --git a/Engine/Source/Runtime/Experimental/Animation/Constraints/Private/ConstraintSubsystem.cpp b/Engine/Source/Runtime/Experimental/Animation/Constraints/Private/ConstraintSubsystem.cpp
--- a/Engine/Source/Runtime/Experimental/Animation/Constraints/Private/ConstraintSubsystem.cpp
+++ b/Engine/Source/Runtime/Experimental/Animation/Constraints/Private/ConstraintSubsystem.cpp
@@ -4,12 +4,12 @@
 #include "Containers/Ticker.h"
 #include "Engine/Engine.h"
 #include "ConstraintsManager.h"
-#include "TransformConstraint.h"
 #include "Misc/CoreDelegates.h"
 
 //needs to be static to avoid system getting deleted with dangling handles.
 FDelegateHandle UConstraintSubsystem::OnWorldInitHandle;
 FDelegateHandle UConstraintSubsystem::OnWorldCleanupHandle;
+FDelegateHandle UConstraintSubsystem::OnPostGarbageCollectHandle;
 
 UConstraintSubsystem::UConstraintSubsystem()
 {
@@ -35,6 +35,7 @@ void UConstraintSubsystem::RegisterWorldDelegates()
 {
 	OnWorldInitHandle = FWorldDelegates::OnPreWorldInitialization.AddStatic(&UConstraintSubsystem::OnWorldInit);
 	OnWorldCleanupHandle = FWorldDelegates::OnWorldCleanup.AddStatic(&UConstraintSubsystem::OnWorldCleanup);
+	OnPostGarbageCollectHandle = FCoreUObjectDelegates::GetPostGarbageCollect().AddStatic(&UConstraintSubsystem::OnPostGarbageCollect);
 	
 	FCoreDelegates::OnPostEngineInit.RemoveAll(this);
 }
@@ -46,8 +47,11 @@ void UConstraintSubsystem::Deinitialize()
 		ConstraintsInWorld[Index].RemoveConstraints(ConstraintsInWorld[Index].World);
 	}
 	ConstraintsInWorld.Reset();
+
 	FWorldDelegates::OnPreWorldInitialization.Remove(OnWorldInitHandle);
 	FWorldDelegates::OnWorldCleanup.Remove(OnWorldCleanupHandle);
+	FCoreUObjectDelegates::GetPostGarbageCollect().Remove(OnPostGarbageCollectHandle);
+
 	Super::Deinitialize();
 }
 
@@ -62,6 +66,11 @@ UConstraintSubsystem* UConstraintSubsystem::Get()
 
 const FConstraintsInWorld* UConstraintSubsystem::ConstraintsInWorldFind(UWorld* InWorld) const
 {
+	if (bNeedsCleanup)
+	{
+		CleanupInvalidConstraints();
+	}
+	
 	for (const FConstraintsInWorld& CInW : ConstraintsInWorld)
 	{
 		if (CInW.World.Get() == InWorld)
@@ -74,6 +83,11 @@ const FConstraintsInWorld* UConstraintSubsystem::ConstraintsInWorldFind(UWorld*
 
 FConstraintsInWorld* UConstraintSubsystem::ConstraintsInWorldFind(UWorld* InWorld) 
 {
+	if (bNeedsCleanup)
+	{
+		CleanupInvalidConstraints();
+	}
+	
 	for (FConstraintsInWorld& CInW : ConstraintsInWorld)
 	{
 		if (CInW.World.Get() == InWorld)
@@ -86,6 +100,11 @@ FConstraintsInWorld* UConstraintSubsystem::ConstraintsInWorldFind(UWorld* InWorl
 
 FConstraintsInWorld& UConstraintSubsystem::ConstraintsInWorldFindOrAdd(UWorld* InWorld)
 {
+	if (bNeedsCleanup)
+	{
+		CleanupInvalidConstraints();
+	}
+	
 	for (FConstraintsInWorld& CInW : ConstraintsInWorld)
 	{
 		if (CInW.World.Get() == InWorld)
@@ -98,6 +117,7 @@ FConstraintsInWorld& UConstraintSubsystem::ConstraintsInWorldFindOrAdd(UWorld* I
 	int32 Index = ConstraintsInWorld.Add(NewCInW);
 	return ConstraintsInWorld[Index];
 }
+
 TArray<TWeakObjectPtr<UTickableConstraint>> UConstraintSubsystem::GetConstraints(UWorld* InWorld) const
 {
 	static const TArray< TWeakObjectPtr<UTickableConstraint> > DummyArray;
@@ -187,6 +207,11 @@ bool UConstraintSubsystem::HasConstraint(UWorld* InWorld, UTickableConstraint* I
 	return  Constraints.Contains(InConstraint);
 }
 
+void UConstraintSubsystem::InvalidateConstraints()
+{
+	bNeedsCleanup = true;
+}
+
 #if WITH_EDITOR
 
 void UConstraintSubsystem::PostEditUndo()
@@ -224,6 +249,29 @@ void UConstraintSubsystem::OnWorldCleanup(UWorld* InWorld, bool bSessionEnded, b
 	}
 }
 
+void UConstraintSubsystem::OnPostGarbageCollect()
+{
+	if (UConstraintSubsystem* System = Get())
+	{
+		System->InvalidateConstraints();
+	}
+}
+
+void UConstraintSubsystem::CleanupInvalidConstraints() const
+{
+	if (UConstraintSubsystem* System = Get())
+	{
+		for (FConstraintsInWorld& WorldConstraints: System->ConstraintsInWorld)
+		{
+			WorldConstraints.Constraints.RemoveAll( [](const TWeakObjectPtr<UTickableConstraint>& InConstraint)
+			{
+				return !InConstraint.IsValid() || InConstraint.IsStale();
+			});
+		}
+		bNeedsCleanup = false;
+	}
+}
+
 /*************************************
 *FConstraintsInWorld
 *************************************/
diff --git a/Engine/Source/Runtime/Experimental/Animation/Constraints/Private/ConstraintSubsystem.h b/Engine/Source/Runtime/Experimental/Animation/Constraints/Private/ConstraintSubsystem.h
--- a/Engine/Source/Runtime/Experimental/Animation/Constraints/Private/ConstraintSubsystem.h
+++ b/Engine/Source/Runtime/Experimental/Animation/Constraints/Private/ConstraintSubsystem.h
@@ -51,17 +51,14 @@ public:
 		FConstraintTickFunction* InFunctionToTickBefore,
 		FConstraintTickFunction* InFunctionToTickAfter);
 
-	static void OnWorldInit(UWorld* World, const UWorld::InitializationValues IVS);
-	static void OnWorldCleanup(UWorld* World, bool bSessionEnded, bool bCleanupResources);
-
-public:
-
 	CONSTRAINTS_API TArray<TWeakObjectPtr<UTickableConstraint>> GetConstraints(UWorld* InWorld) const;
 	CONSTRAINTS_API const TArray<TWeakObjectPtr<UTickableConstraint>>& GetConstraintsArray(UWorld* InWorld) const;
 
 	CONSTRAINTS_API void AddConstraint(UWorld* InWorld, UTickableConstraint* InConstraint);
 	CONSTRAINTS_API void RemoveConstraint(UWorld* InWorld, UTickableConstraint* InConstraint, bool bDoNoCompensate);
 	CONSTRAINTS_API bool HasConstraint(UWorld* InWorld, UTickableConstraint* InConstraint);
+	
+	void InvalidateConstraints();
 
 #if WITH_EDITOR
 	virtual void PostEditUndo() override;
@@ -81,11 +78,22 @@ private:
 	
 	//handles for handling world creation init and teardown
 	static FDelegateHandle OnWorldInitHandle;
+	static void OnWorldInit(UWorld* World, const UWorld::InitializationValues IVS);
+	
 	static FDelegateHandle OnWorldCleanupHandle;
+	static void OnWorldCleanup(UWorld* World, bool bSessionEnded, bool bCleanupResources);
+
+	//handle for handling GC
+	static FDelegateHandle OnPostGarbageCollectHandle;
+	static void OnPostGarbageCollect();
+	
 	void RegisterWorldDelegates();
 
+	mutable bool bNeedsCleanup = false;
+
 	const FConstraintsInWorld* ConstraintsInWorldFind(UWorld* InWorld) const;
 	FConstraintsInWorld* ConstraintsInWorldFind(UWorld* InWorld);
 	FConstraintsInWorld& ConstraintsInWorldFindOrAdd(UWorld* InWorld);
 
+	void CleanupInvalidConstraints() const;
 };
