Fix race condition when creating geometry collection scene proxy for nanite. Moved some state caching from the constructor, which isn't guaranteed to be called after render resources are initialized, to CreateRenderThreadResources which is.

diff --git a/Engine/Source/Runtime/Experimental/GeometryCollectionEngine/Private/GeometryCollection/GeometryCollectionSceneProxy.cpp b/Engine/Source/Runtime/Experimental/GeometryCollectionEngine/Private/GeometryCollection/GeometryCollectionSceneProxy.cpp
--- a/Engine/Source/Runtime/Experimental/GeometryCollectionEngine/Private/GeometryCollection/GeometryCollectionSceneProxy.cpp
+++ b/Engine/Source/Runtime/Experimental/GeometryCollectionEngine/Private/GeometryCollection/GeometryCollectionSceneProxy.cpp
@@ -993,11 +993,7 @@ FNaniteGeometryCollectionSceneProxy::FNaniteGeometryCollectionSceneProxy(UGeomet
 	// Nanite requires GPUScene
 	checkSlow(UseGPUScene(GMaxRHIShaderPlatform, GetScene().GetFeatureLevel()));
 	checkSlow(DoesPlatformSupportNanite(GMaxRHIShaderPlatform));
-
-	// Should have valid Nanite data at this point.
-	check(GeometryCollection->HasNaniteData());
-	NaniteResourceID = GeometryCollection->GetNaniteResourceID();
-	NaniteHierarchyOffset = GeometryCollection->GetNaniteHierarchyOffset();
+	checkSlow(GeometryCollection->HasNaniteData());
 
 	MaterialRelevance = Component->GetMaterialRelevance(Component->GetScene()->GetFeatureLevel());
 
@@ -1144,6 +1140,14 @@ FNaniteGeometryCollectionSceneProxy::FNaniteGeometryCollectionSceneProxy(UGeomet
 	SetDynamicData_RenderThread(DynamicData, Component->GetRenderMatrix());
 }
 
+void FNaniteGeometryCollectionSceneProxy::CreateRenderThreadResources(FRHICommandListBase& RHICmdList)
+{
+	// Should have valid Nanite data at this point.
+	NaniteResourceID = GeometryCollection->GetNaniteResourceID();
+	NaniteHierarchyOffset = GeometryCollection->GetNaniteHierarchyOffset();
+	check(NaniteResourceID != INDEX_NONE && NaniteHierarchyOffset != INDEX_NONE);
+}
+
 SIZE_T FNaniteGeometryCollectionSceneProxy::GetTypeHash() const
 {
 	static size_t UniquePointer;
diff --git a/Engine/Source/Runtime/Experimental/GeometryCollectionEngine/Private/GeometryCollection/GeometryCollectionSceneProxy.h b/Engine/Source/Runtime/Experimental/GeometryCollectionEngine/Private/GeometryCollection/GeometryCollectionSceneProxy.h
--- a/Engine/Source/Runtime/Experimental/GeometryCollectionEngine/Private/GeometryCollection/GeometryCollectionSceneProxy.h
+++ b/Engine/Source/Runtime/Experimental/GeometryCollectionEngine/Private/GeometryCollection/GeometryCollectionSceneProxy.h
@@ -334,6 +334,7 @@ public:
 
 public:
 	// FPrimitiveSceneProxy interface.
+	virtual void CreateRenderThreadResources(FRHICommandListBase& RHICmdList) override;
 	virtual SIZE_T GetTypeHash() const override;
 	virtual FPrimitiveViewRelevance	GetViewRelevance(const FSceneView* View) const override;
 #if WITH_EDITOR
