Reduce array size to store only required particles.
Several array in the Geometry Collection are used to store particle. According to the max simulation level, not all particle are created.
Now array have their own size, the particle size and the transform size in the proxy is now distinguished.
This change save 55% of memory in initialization.
It also avoid to process all transforms, we now iterate only through particles which should be way faster.

diff --git a/Engine/Plugins/Experimental/ChaosCaching/Source/ChaosCaching/Private/Chaos/Adapters/GeometryCollectionComponentCacheAdapter.cpp b/Engine/Plugins/Experimental/ChaosCaching/Source/ChaosCaching/Private/Chaos/Adapters/GeometryCollectionComponentCacheAdapter.cpp
--- a/Engine/Plugins/Experimental/ChaosCaching/Source/ChaosCaching/Private/Chaos/Adapters/GeometryCollectionComponentCacheAdapter.cpp
+++ b/Engine/Plugins/Experimental/ChaosCaching/Source/ChaosCaching/Private/Chaos/Adapters/GeometryCollectionComponentCacheAdapter.cpp
@@ -135,7 +135,7 @@ namespace Chaos
 
 		for(int32 TransformIndex = 0; TransformIndex < NumTransforms; ++TransformIndex)
 		{
-			FClusterParticle* Handle = Proxy->GetParticles()[TransformIndex];
+			FClusterParticle* Handle = Proxy->GetParticle_Internal(TransformIndex);
 			
 			if(Handle)
 			{
@@ -170,7 +170,7 @@ namespace Chaos
 		{
 			if (TransformIndex > INDEX_NONE)
 			{
-				FClusterParticle* Handle = Proxy->GetParticles()[TransformIndex];
+				FClusterParticle* Handle = Proxy->GetParticle_Internal(TransformIndex);
 				if (Handle)
 				{
 					OutFrame.PendingParticleData.AddDefaulted();
@@ -263,7 +263,6 @@ namespace Chaos
 
 		FGeometryDynamicCollection&      Collection       = Proxy->GetPhysicsCollection();
 		const TManagedArray<FTransform>& MassToLocal      = RestCollection->GetAttribute<FTransform>("MassToLocal", FGeometryCollection::TransformGroup);
-		TArray<FClusterParticle*>        Particles        = Proxy->GetParticles();
 
 		FCacheEvaluationContext Context(TickRecord);
 		Context.bEvaluateTransform = true;
@@ -285,32 +284,27 @@ namespace Chaos
 			{
 				if(FEnableStateEvent* Event = Handle.Get<FEnableStateEvent>())
 				{
-					if(Particles.IsValidIndex(Event->Index))
+					if(Chaos::FPBDRigidClusteredParticleHandle* ChildParticle = Proxy->GetParticle_Internal(Event->Index))
 					{
-						Chaos::FPBDRigidClusteredParticleHandle* ChildParticle = Particles[Event->Index];
-						
-						if (ChildParticle)
+						if (ChildParticle->ObjectState() != EObjectStateType::Kinematic)
 						{
-							if (ChildParticle->ObjectState() != EObjectStateType::Kinematic)
-							{
-								// If a field or other external actor set the particle to static or dynamic we no longer apply the cache
-								continue;
-							}
+							// If a field or other external actor set the particle to static or dynamic we no longer apply the cache
+							continue;
+						}
 
-							if (FRigidParticle* ClusterParent = ChildParticle->ClusterIds().Id)
-							{
-								if (FClusterParticle* Parent = ClusterParent->CastToClustered())
-								{
-									TArray<FRigidParticle*>& Cluster = NewClusters.FindOrAdd(Parent);
-									Cluster.Add(ChildParticle);
-								}
-							}
-							else
+						if (FRigidParticle* ClusterParent = ChildParticle->ClusterIds().Id)
+						{
+							if (FClusterParticle* Parent = ClusterParent->CastToClustered())
 							{
-								// This is a cluster parent
-								ChildParticle->SetDisabled(!Event->bEnable);
+								TArray<FRigidParticle*>& Cluster = NewClusters.FindOrAdd(Parent);
+								Cluster.Add(ChildParticle);
 							}
 						}
+						else
+						{
+							// This is a cluster parent
+							ChildParticle->SetDisabled(!Event->bEnable);
+						}
 					}
 				}
 			}
@@ -355,41 +349,36 @@ namespace Chaos
 			{
 				if (FBreakingEvent* Event = Handle.Get<FBreakingEvent>())
 				{
-					if (Particles.IsValidIndex(Event->Index))
+					if (Chaos::FPBDRigidClusteredParticleHandle* Particle = Proxy->GetParticle_Internal(Event->Index))
 					{
-						Chaos::FPBDRigidClusteredParticleHandle* Particle = Particles[Event->Index];
-
-						if (Particle)
+						if (Particle->ObjectState() != EObjectStateType::Kinematic)
 						{
-							if (Particle->ObjectState() != EObjectStateType::Kinematic)
-							{
-								// If a field or other external actor set the particle to static or dynamic we no longer apply the cache
-								continue;
-							}
+							// If a field or other external actor set the particle to static or dynamic we no longer apply the cache
+							continue;
+						}
 
-							FBreakingData CachedBreak;
-							CachedBreak.Proxy = Proxy;
-							CachedBreak.Location = ComponentToWorld.TransformPosition(Event->Location);
-							CachedBreak.Velocity = ComponentToWorld.TransformVector(Event->Velocity);
-							CachedBreak.AngularVelocity = Event->AngularVelocity;
-							CachedBreak.Mass = Event->Mass;
-							CachedBreak.BoundingBox = TAABB<FReal, 3>(Event->BoundingBoxMin, Event->BoundingBoxMax);
-							CachedBreak.BoundingBox = CachedBreak.BoundingBox.TransformedAABB(ComponentToWorld);
-							CachedBreak.SetEmitterFlag(Proxy->GetSimParameters().bGenerateBreakingData, Proxy->GetSimParameters().bGenerateGlobalBreakingData);
-
-							if (!SolverBreakingEventFilter->Enabled() || SolverBreakingEventFilter->Pass(CachedBreak))
-							{
-								float TimeStamp = Solver->GetSolverTime();
-								Solver->GetEventManager()->AddEvent<FBreakingEventData>(EEventType::Breaking, [&CachedBreak, TimeStamp](FBreakingEventData& BreakingEventData)
+						FBreakingData CachedBreak;
+						CachedBreak.Proxy = Proxy;
+						CachedBreak.Location = ComponentToWorld.TransformPosition(Event->Location);
+						CachedBreak.Velocity = ComponentToWorld.TransformVector(Event->Velocity);
+						CachedBreak.AngularVelocity = Event->AngularVelocity;
+						CachedBreak.Mass = Event->Mass;
+						CachedBreak.BoundingBox = TAABB<FReal, 3>(Event->BoundingBoxMin, Event->BoundingBoxMax);
+						CachedBreak.BoundingBox = CachedBreak.BoundingBox.TransformedAABB(ComponentToWorld);
+						CachedBreak.SetEmitterFlag(Proxy->GetSimParameters().bGenerateBreakingData, Proxy->GetSimParameters().bGenerateGlobalBreakingData);
+
+						if (!SolverBreakingEventFilter->Enabled() || SolverBreakingEventFilter->Pass(CachedBreak))
+						{
+							float TimeStamp = Solver->GetSolverTime();
+							Solver->GetEventManager()->AddEvent<FBreakingEventData>(EEventType::Breaking, [&CachedBreak, TimeStamp](FBreakingEventData& BreakingEventData)
+								{
+									if (BreakingEventData.BreakingData.TimeCreated != TimeStamp)
 									{
-										if (BreakingEventData.BreakingData.TimeCreated != TimeStamp)
-										{
-											BreakingEventData.BreakingData.AllBreakingsArray.Reset();
-											BreakingEventData.BreakingData.TimeCreated = TimeStamp;
-										}
-										BreakingEventData.BreakingData.AllBreakingsArray.Add(CachedBreak);
-									});
-							}
+										BreakingEventData.BreakingData.AllBreakingsArray.Reset();
+										BreakingEventData.BreakingData.TimeCreated = TimeStamp;
+									}
+									BreakingEventData.BreakingData.AllBreakingsArray.Add(CachedBreak);
+								});
 						}
 					}
 				}
@@ -404,39 +393,34 @@ namespace Chaos
 			{
 				if (FTrailingEvent* Event = Handle.Get<FTrailingEvent>())
 				{
-					if (Particles.IsValidIndex(Event->Index))
+					if (Chaos::FPBDRigidClusteredParticleHandle* Particle = Proxy->GetParticle_Internal(Event->Index))
 					{
-						Chaos::FPBDRigidClusteredParticleHandle* Particle = Particles[Event->Index];
-
-						if (Particle)
+						if (Particle->ObjectState() != EObjectStateType::Kinematic)
 						{
-							if (Particle->ObjectState() != EObjectStateType::Kinematic)
-							{
-								// If a field or other external actor set the particle to static or dynamic we no longer apply the cache
-								continue;
-							}
+							// If a field or other external actor set the particle to static or dynamic we no longer apply the cache
+							continue;
+						}
 
-							FTrailingData CachedTrail;
-							CachedTrail.Proxy = Proxy;
-							CachedTrail.Location = ComponentToWorld.TransformPosition(Event->Location);
-							CachedTrail.Velocity = ComponentToWorld.TransformVector(Event->Velocity);
-							CachedTrail.AngularVelocity = Event->AngularVelocity;
-							CachedTrail.BoundingBox = TAABB<FReal, 3>(Event->BoundingBoxMin, Event->BoundingBoxMax);
-							CachedTrail.BoundingBox = CachedTrail.BoundingBox.TransformedAABB(ComponentToWorld);
+						FTrailingData CachedTrail;
+						CachedTrail.Proxy = Proxy;
+						CachedTrail.Location = ComponentToWorld.TransformPosition(Event->Location);
+						CachedTrail.Velocity = ComponentToWorld.TransformVector(Event->Velocity);
+						CachedTrail.AngularVelocity = Event->AngularVelocity;
+						CachedTrail.BoundingBox = TAABB<FReal, 3>(Event->BoundingBoxMin, Event->BoundingBoxMax);
+						CachedTrail.BoundingBox = CachedTrail.BoundingBox.TransformedAABB(ComponentToWorld);
 
-							if (!SolverTrailingEventFilter->Enabled() || SolverTrailingEventFilter->Pass(CachedTrail))
-							{
-								float TimeStamp = Solver->GetSolverTime();
-								Solver->GetEventManager()->AddEvent<FTrailingEventData>(EEventType::Trailing, [&CachedTrail , TimeStamp](FTrailingEventData& TrailingEventData)
+						if (!SolverTrailingEventFilter->Enabled() || SolverTrailingEventFilter->Pass(CachedTrail))
+						{
+							float TimeStamp = Solver->GetSolverTime();
+							Solver->GetEventManager()->AddEvent<FTrailingEventData>(EEventType::Trailing, [&CachedTrail , TimeStamp](FTrailingEventData& TrailingEventData)
+								{
+									if (TrailingEventData.TrailingData.TimeCreated != TimeStamp)
 									{
-										if (TrailingEventData.TrailingData.TimeCreated != TimeStamp)
-										{
-											TrailingEventData.TrailingData.AllTrailingsArray.Reset();
-											TrailingEventData.TrailingData.TimeCreated = TimeStamp;
-										}
-										TrailingEventData.TrailingData.AllTrailingsArray.Add(CachedTrail);
-									});
-							}
+										TrailingEventData.TrailingData.AllTrailingsArray.Reset();
+										TrailingEventData.TrailingData.TimeCreated = TimeStamp;
+									}
+									TrailingEventData.TrailingData.AllTrailingsArray.Add(CachedTrail);
+								});
 						}
 					}
 				}
@@ -451,7 +435,6 @@ namespace Chaos
 			{
 				if (FCollisionEvent* Event = Handle.Get<FCollisionEvent>())
 				{
-					
 					FCollidingData CachedCollision;
 					CachedCollision.Location = ComponentToWorld.TransformPosition(Event->Location);
 					CachedCollision.AccumulatedImpulse = ComponentToWorld.TransformVector(Event->AccumulatedImpulse);
@@ -493,11 +476,9 @@ namespace Chaos
 			const int32      ParticleIndex      = EvaluatedResult.ParticleIndices[Index];
 			const FTransform EvaluatedTransform = EvaluatedResult.Transform[Index];
 
-			if(Particles.IsValidIndex(ParticleIndex))
+			if (Chaos::FPBDRigidClusteredParticleHandle* Handle = Proxy->GetParticle_Internal(ParticleIndex))
 			{
-				Chaos::FPBDRigidClusteredParticleHandle* Handle = Particles[ParticleIndex];
-
-				if(!Handle || Handle->ObjectState() != EObjectStateType::Kinematic)
+				if(Handle->ObjectState() != EObjectStateType::Kinematic)
 				{
 					// If a field or other external actor set the particle to static or dynamic we no longer apply the cache
 					continue;
diff --git a/Engine/Plugins/Experimental/ChaosModularVehicle/Source/ChaosModularVehicleEngine/Private/ChaosModularVehicle/ModularVehicleComponent.cpp b/Engine/Plugins/Experimental/ChaosModularVehicle/Source/ChaosModularVehicleEngine/Private/ChaosModularVehicle/ModularVehicleComponent.cpp
--- a/Engine/Plugins/Experimental/ChaosModularVehicle/Source/ChaosModularVehicleEngine/Private/ChaosModularVehicle/ModularVehicleComponent.cpp
+++ b/Engine/Plugins/Experimental/ChaosModularVehicle/Source/ChaosModularVehicleEngine/Private/ChaosModularVehicle/ModularVehicleComponent.cpp
@@ -62,11 +62,11 @@ Chaos::FPBDRigidClusteredParticleHandle* UModularVehicleComponent::GetChassisPar
 
 	if (FGeometryCollectionPhysicsProxy* Proxy = GetPhysicsProxy())
 	{
-		const TArray<Chaos::FPBDRigidClusteredParticleHandle*>& Clusters = Proxy->GetSolverClusterHandles();
+		Chaos::FPBDRigidClusteredParticleHandle* Cluster = Proxy->GetSolverClusterHandle_Internal(0);
 
-		if (Clusters.Num() > 0 && !Clusters[0]->Disabled())
+		if (Cluster != nullptr && !Cluster->Disabled())
 		{
-			Particle = Clusters[0];
+			Particle = Cluster;
 		}
 	}
 	return Particle;
@@ -221,8 +221,8 @@ bool UModularVehicleComponent::ProcessRepData(const float DeltaTime, const float
 								check(RewindData);
 								check(!RewindData->IsResim());
 
-								const TArray<Chaos::FPBDRigidClusteredParticleHandle*>& Clusters = Proxy->GetSolverClusterHandles();
-								auto* ChassisHandle = Clusters[ChassisIdx];
+								Chaos::FPBDRigidClusteredParticleHandle* ChassisHandle = Proxy->GetSolverClusterHandle_Internal(ChassisIdx);
+
 								const FGeometryCollectionClusterRep& RepCluster = RepData.Clusters[ChassisIdx];
 
 								// Need the local frame in order to access correct history data
diff --git a/Engine/Plugins/Experimental/ChaosModularVehicle/Source/ChaosModularVehicleEngine/Private/ChaosModularVehicle/ModularVehicleSimulationCU.cpp b/Engine/Plugins/Experimental/ChaosModularVehicle/Source/ChaosModularVehicleEngine/Private/ChaosModularVehicle/ModularVehicleSimulationCU.cpp
--- a/Engine/Plugins/Experimental/ChaosModularVehicle/Source/ChaosModularVehicleEngine/Private/ChaosModularVehicle/ModularVehicleSimulationCU.cpp
+++ b/Engine/Plugins/Experimental/ChaosModularVehicle/Source/ChaosModularVehicleEngine/Private/ChaosModularVehicle/ModularVehicleSimulationCU.cpp
@@ -363,8 +363,7 @@ void FModularVehicleSimulationCU::ApplyDeferredForces(FGeometryCollectionPhysics
 					const TManagedArray<FTransform>& CollectionMassToLocal = Rest->GetGeometryCollection()->GetAttribute<FTransform>(TEXT("MassToLocal"), FTransformCollection::TransformGroup);
 
 					SimModuleTree->AccessDeferredForces().Apply(
-						Proxy->GetSolverParticleHandles(),
-						Proxy->GetSolverClusterHandles(),
+						Proxy,
 						Rest->GetGeometryCollection()->Transform,
 						CollectionMassToLocal,
 						Rest->GetGeometryCollection()->Parent);
diff --git a/Engine/Plugins/Experimental/ChaosModularVehicle/Source/ChaosModularVehicleEngine/Private/ChaosModularVehicle/ModularVehicleSimulationGC.cpp b/Engine/Plugins/Experimental/ChaosModularVehicle/Source/ChaosModularVehicleEngine/Private/ChaosModularVehicle/ModularVehicleSimulationGC.cpp
--- a/Engine/Plugins/Experimental/ChaosModularVehicle/Source/ChaosModularVehicleEngine/Private/ChaosModularVehicle/ModularVehicleSimulationGC.cpp
+++ b/Engine/Plugins/Experimental/ChaosModularVehicle/Source/ChaosModularVehicleEngine/Private/ChaosModularVehicle/ModularVehicleSimulationGC.cpp
@@ -143,25 +143,21 @@ void FModularVehicleSimulationGC::PerformAdditionalSimWork(UWorld* InWorld, cons
 	Chaos::EnsureIsInPhysicsThreadContext();
 
 	FGeometryDynamicCollection& GeometryCollection = Proxy->GetPhysicsCollection();
-	const TArray<Chaos::FPBDRigidClusteredParticleHandle*>& Clusters = Proxy->GetSolverClusterHandles();
-	const TArray<Chaos::FPBDRigidClusteredParticleHandle*>& Particles = Proxy->GetSolverParticleHandles();
 
-	if (Clusters.Num() > 0)
+	if (Proxy->GetUnorderedParticles_Internal().Num() > 0)
 	{
 		const TArray<Chaos::FSimModuleTree::FSimModuleNode>& ModuleArray = SimModuleTree->GetSimulationModuleTree();
-
 		for (const Chaos::FSimModuleTree::FSimModuleNode& Node : ModuleArray)
 		{
-			if (!Node.SimModule->IsEnabled() || !Particles[Node.SimModule->GetTransformIndex()]->Disabled())
-			{
-				Node.SimModule->SetStateFlags(Chaos::eSimModuleState::Disabled);
-				continue;
-			}
-
-			check(Node.SimModule->GetTransformIndex() < Clusters.Num());
-			const Chaos::FPBDRigidClusteredParticleHandle* Particle = Clusters[Node.SimModule->GetTransformIndex()];
-			if (Particle)
+			const Chaos::FPBDRigidClusteredParticleHandle* Cluster = Proxy->GetParticle_Internal(Node.SimModule->GetTransformIndex());
+			const Chaos::FPBDRigidClusteredParticleHandle* Particle = Proxy->GetSolverClusterHandle_Internal(Node.SimModule->GetTransformIndex());
+			if (Cluster != nullptr && Particle != nullptr)
 			{
+				if (!Node.SimModule->IsEnabled() || !Cluster->Disabled())
+				{
+					Node.SimModule->SetStateFlags(Chaos::eSimModuleState::Disabled);
+					continue;
+				}
 				const FTransform BodyTransform(Particle->R(), Particle->X());
 
 				// #TODO: cheating just now to get it working - do we pass in one or one per particle, as it will be different if it has split into fragments
@@ -300,8 +296,7 @@ void FModularVehicleSimulationGC::ApplyDeferredForces(FGeometryCollectionPhysics
 					const TManagedArray<FTransform>& CollectionMassToLocal = Rest->GetGeometryCollection()->GetAttribute<FTransform>(TEXT("MassToLocal"), FTransformCollection::TransformGroup);
 
 					SimModuleTree->AccessDeferredForces().Apply(
-						Proxy->GetSolverParticleHandles(),
-						Proxy->GetSolverClusterHandles(),
+						Proxy,
 						Rest->GetGeometryCollection()->Transform,
 						CollectionMassToLocal,
 						Rest->GetGeometryCollection()->Parent);
diff --git a/Engine/Source/Runtime/Engine/Private/PhysicsEngine/Experimental/PhysScene_Chaos.cpp b/Engine/Source/Runtime/Engine/Private/PhysicsEngine/Experimental/PhysScene_Chaos.cpp
--- a/Engine/Source/Runtime/Engine/Private/PhysicsEngine/Experimental/PhysScene_Chaos.cpp
+++ b/Engine/Source/Runtime/Engine/Private/PhysicsEngine/Experimental/PhysScene_Chaos.cpp
@@ -784,7 +784,7 @@ void FPhysScene_Chaos::RemoveObject(FGeometryCollectionPhysicsProxy* InObject)
 {
 	Chaos::FPhysicsSolver* Solver = InObject->GetSolver<Chaos::FPhysicsSolver>();
 
-	for (TUniquePtr<Chaos::TPBDRigidParticle<Chaos::FReal, 3>>& GTParticleUnique : InObject->GetExternalParticles())
+	for (TUniquePtr<Chaos::TPBDRigidParticle<Chaos::FReal, 3>>& GTParticleUnique : InObject->GetUnorderedParticles_External())
 	{
 		Chaos::TPBDRigidParticle<Chaos::FReal, 3>* GTParticle = GTParticleUnique.Get();
 		if (GTParticle != nullptr)
@@ -2190,14 +2190,13 @@ static void UpdateAccelerationStructureFromGeometryCollectionProxy(FGeometryColl
 		return Chaos::FAABB3(Particle.X(), Particle.X());
 	};
 
-	const TArray<TUniquePtr<Chaos::FPBDRigidParticle>>& GTParticles = Proxy.GetExternalParticles();
+	const TArray<TUniquePtr<Chaos::FPBDRigidParticle>>& UnorderedGTParticles = Proxy.GetUnorderedParticles_External();
 	if (bIsParentProxyNull)
 	{
-		for (const TUniquePtr<Chaos::FPBDRigidParticle>& GTParticle : GTParticles)
+		for (const TUniquePtr<Chaos::FPBDRigidParticle>& GTParticle : UnorderedGTParticles)
 		{
 			if (GTParticle)
 			{
-
 				const Chaos::FSpatialAccelerationIdx SpatialIndex = GTParticle->SpatialIdx();
 
 				// It's possible to be an enabled particle and not qualify for query collisions if the GC particle has been replication abandoned. 
@@ -2224,7 +2223,7 @@ static void UpdateAccelerationStructureFromGeometryCollectionProxy(FGeometryColl
 	{
 		// if we have a parent proxy ( like attached to a cluster union) then none of the handles should be in the acceleration structure 
 		// todo : right now we are sending all the particles, but we should be able to get away with active + direct children that could reduce the overall cost 
-		for (const TUniquePtr<Chaos::FPBDRigidParticle>& GTParticle : GTParticles)
+		for (const TUniquePtr<Chaos::FPBDRigidParticle>& GTParticle : UnorderedGTParticles)
 		{
 			if (GTParticle)
 			{
diff --git a/Engine/Source/Runtime/Experimental/Chaos/Private/Chaos/StrainModification.cpp b/Engine/Source/Runtime/Experimental/Chaos/Private/Chaos/StrainModification.cpp
--- a/Engine/Source/Runtime/Experimental/Chaos/Private/Chaos/StrainModification.cpp
+++ b/Engine/Source/Runtime/Experimental/Chaos/Private/Chaos/StrainModification.cpp
@@ -71,7 +71,7 @@ static void ForEachRootChildParticle(const Chaos::FStrainedProxyAndRoot& ProxyAn
 		{
 			for (int32 RestChildIdx : *RestChildren)
 			{
-				Chaos::FPBDRigidClusteredParticleHandle* ChildHandle = ProxyGC->GetSolverParticleHandles()[RestChildIdx];
+				Chaos::FPBDRigidClusteredParticleHandle* ChildHandle = ProxyGC->GetParticle_Internal(RestChildIdx);
 				if ((ChildHandle == nullptr) || (ChildHandle->Parent() == nullptr)) { continue; }
 
 				Func(ChildHandle);
@@ -86,7 +86,7 @@ int32 Chaos::FStrainedProxyModifier::GetNumRestBreakables() const
 	{
 		return RestChildren->Num();
 	}
-	return 1; // cluser union case
+	return 1; // cluster union case
 }
 
 int32 Chaos::FStrainedProxyModifier::GetNumBreakingStrains(const bool bDoubleCount, const uint8 StrainTypes) const
@@ -228,10 +228,10 @@ Chaos::FStrainedProxyRange::FStrainedProxyRange(Chaos::FRigidClustering& InRigid
 			{
 				FSimulationParameters& Parameters = ProxyGC->GetSimParameters();
 				const int32 RootIndex = Parameters.InitialRootIndex;
-				TArray<Chaos::FPBDRigidClusteredParticleHandle*>& ParticleHandles = ProxyGC->GetSolverParticleHandles();
-				if (ParticleHandles.IsValidIndex(RootIndex))
+				Chaos::FPBDRigidClusteredParticleHandle* ParticleHandle = ProxyGC->GetParticle_Internal(RootIndex);
+				if (ParticleHandle != nullptr)
 				{
-					if (Cluster != ParticleHandles[RootIndex])
+					if (Cluster != ParticleHandle)
 					{
 						return;
 					}
diff --git a/Engine/Source/Runtime/Experimental/Chaos/Private/PBDRigidsSolver.cpp b/Engine/Source/Runtime/Experimental/Chaos/Private/PBDRigidsSolver.cpp
--- a/Engine/Source/Runtime/Experimental/Chaos/Private/PBDRigidsSolver.cpp
+++ b/Engine/Source/Runtime/Experimental/Chaos/Private/PBDRigidsSolver.cpp
@@ -772,7 +772,7 @@ namespace Chaos
 		InProxy->NewData(); // Buffers data on the proxy.
 
 		// Need to immediately add to the SQ to make sure that the GC is in the external SQ even after the SQ's flip.
-		for (const TUniquePtr<FPBDRigidParticle>& Particle : InProxy->GetExternalParticles())
+		for (const TUniquePtr<FPBDRigidParticle>& Particle : InProxy->GetUnorderedParticles_External())
 		{
 			if (Particle && !Particle->Disabled())
 			{
diff --git a/Engine/Source/Runtime/Experimental/Chaos/Private/PhysicsProxy/GeometryCollectionPhysicsProxy.cpp b/Engine/Source/Runtime/Experimental/Chaos/Private/PhysicsProxy/GeometryCollectionPhysicsProxy.cpp
--- a/Engine/Source/Runtime/Experimental/Chaos/Private/PhysicsProxy/GeometryCollectionPhysicsProxy.cpp
+++ b/Engine/Source/Runtime/Experimental/Chaos/Private/PhysicsProxy/GeometryCollectionPhysicsProxy.cpp
@@ -505,7 +505,7 @@ FGeometryCollectionPhysicsProxy::FGeometryCollectionPhysicsProxy(
 	const Chaos::EMultiBufferMode BufferMode)
 	: Base(InOwner)
 	, Parameters(SimulationParameters)
-	, NumParticles(INDEX_NONE)
+	, NumTransforms(INDEX_NONE)
 	, NumEffectiveParticles(INDEX_NONE)
 	, BaseParticleIndex(INDEX_NONE)
 	, IsObjectDynamic(false)
@@ -563,31 +563,40 @@ void FGeometryCollectionPhysicsProxy::Initialize(Chaos::FPBDRigidsEvolutionBase
 
 	InitializeDynamicCollection(DynamicCollection, *Parameters.RestCollection, Parameters);
 
-	NumParticles = DynamicCollection.NumElements(FGeometryCollection::TransformGroup);
+	NumTransforms = DynamicCollection.NumElements(FGeometryCollection::TransformGroup);
 	BaseParticleIndex = 0; // Are we always zero indexed now?
-	SolverClusterID.Init(nullptr, NumParticles);
-	SolverClusterHandles.Init(nullptr, NumParticles);
-	SolverParticleHandles.Init(nullptr, NumParticles);
-	GTParticles.SetNum(NumParticles);
-	UniqueIdxs.SetNum(NumParticles);
 	
 	TBitArray<> EffectiveParticles;
-	NumEffectiveParticles = CalculateEffectiveParticles(DynamicCollection, NumParticles, Parameters.MaxSimulatedLevel, Parameters.EnableClustering, GetOwner(), EffectiveParticles);
+	NumEffectiveParticles = CalculateEffectiveParticles(DynamicCollection, NumTransforms, Parameters.MaxSimulatedLevel, Parameters.EnableClustering, GetOwner(), EffectiveParticles);
 
+	SolverClusterID.Init(nullptr, NumEffectiveParticles);
+	SolverClusterHandles.Init(nullptr, NumEffectiveParticles);
+	SolverParticleHandles.Init(nullptr, NumEffectiveParticles);
+	GTParticles.SetNum(NumEffectiveParticles);
+	UniqueIdxs.SetNum(NumEffectiveParticles);
+	
+	FromParticleToTransformIndex.SetNum(NumEffectiveParticles);
+	FromTransformToParticleIndex.SetNum(NumTransforms);
 	PhysicsObjects.Empty();
-	PhysicsObjects.Reserve(NumParticles);
-	for (int32 Index = 0; Index < NumParticles; ++Index)
+	PhysicsObjects.Reserve(NumEffectiveParticles);
+	int32 RedirectIndex = 0;
+	for (int32 Index = 0; Index < NumTransforms; ++Index)
 	{
 		if (EffectiveParticles[Index])
 		{
+			FromTransformToParticleIndex[Index] = RedirectIndex;
+			FromParticleToTransformIndex[RedirectIndex] = Index;
 			PhysicsObjects.Emplace(Chaos::FPhysicsObjectFactory::CreatePhysicsObject(this, Index, FName(Parameters.RestCollection->BoneName[Index])));
+			RedirectIndex++;
 		}
 		else
 		{
-			PhysicsObjects.Emplace(nullptr);
+			FromTransformToParticleIndex[Index] = INDEX_NONE;
 		}
 	}
 
+	check(RedirectIndex == NumEffectiveParticles);
+
 	// we need to make sure the world transform is kept up to date on the game thread 
 	WorldTransform_External = Parameters.WorldTransform;
 	PreviousWorldTransform_External = WorldTransform_External;
@@ -612,7 +621,7 @@ void FGeometryCollectionPhysicsProxy::Initialize(Chaos::FPBDRigidsEvolutionBase
 	}
 
 	// Initialise GT/External particles
-	const int32 NumTransforms = DynamicCollection.GetNumTransforms();
+	check(NumTransforms == DynamicCollection.GetNumTransforms());
 	// make sure we copy the anchored information over to the physics thread collection
 	const Chaos::Facades::FCollectionAnchoringFacade DynamicCollectionAnchoringFacade(DynamicCollection);
 	Chaos::Facades::FCollectionAnchoringFacade PhysicsThreadCollectionAnchoringFacade(PhysicsThreadCollection);
@@ -622,7 +631,7 @@ void FGeometryCollectionPhysicsProxy::Initialize(Chaos::FPBDRigidsEvolutionBase
 	const TManagedArray<float>& Mass = Parameters.RestCollection->GetAttribute<float>(MassAttributeName, FTransformCollection::TransformGroup);
 
 	TManagedArray<Chaos::FImplicitObjectPtr>& Implicits = GameThreadCollection.ModifyAttribute<Chaos::FImplicitObjectPtr>(FGeometryDynamicCollection::ImplicitsAttribute, FTransformCollection::TransformGroup);
-	if (ensure(NumTransforms == Implicits.Num() && NumTransforms == GTParticles.Num())) // Implicits are in the transform group so this invariant should always hold
+	if (ensure(NumTransforms == Implicits.Num() && NumEffectiveParticles == GTParticles.Num())) // Implicits are in the transform group so this invariant should always hold
 	{
 		constexpr bool bInitializationTime = true;
 		bHasBuiltGeometryOnGT = bBuildGeometryForChildren;
@@ -668,13 +677,14 @@ void FGeometryCollectionPhysicsProxy::Initialize(Chaos::FPBDRigidsEvolutionBase
 		if (Parameters.EnableClustering)
 		{
 			// make sure we set Activate the right way when clustering is enabled ( only root should be enabled at start ) 
-			for (int32 TransformIndex = 0; TransformIndex < NumTransforms; ++TransformIndex)
+			for (int32 ParticleIndex = 0; ParticleIndex < NumEffectiveParticles; ++ParticleIndex)
 			{
+				const int32 TransformIndex = FromParticleToTransformIndex[ParticleIndex];
 				const bool bIsRoot = !GameThreadCollection.GetHasParent(TransformIndex);
 				GameThreadCollection.Active[TransformIndex] = bIsRoot;
 				PhysicsThreadCollection.Active[TransformIndex] = bIsRoot;
 
-				if (FParticle* P = GTParticles[TransformIndex].Get())
+				if (FParticle* P = GTParticles[ParticleIndex].Get())
 				{
 					if (P != nullptr)
 					{
@@ -704,102 +714,98 @@ void FGeometryCollectionPhysicsProxy::CreateGTParticles(const TBitArray<>& Effec
 	const TManagedArray<float>& Mass = Parameters.RestCollection->GetAttribute<float>(MassAttributeName, FTransformCollection::TransformGroup);
 	const TManagedArray<FTransform>& MassToLocal = Parameters.RestCollection->GetAttribute<FTransform>(MassToLocalAttributeName, FTransformCollection::TransformGroup);
 	const TManagedArray<int32>& Level = Parameters.RestCollection->GetAttribute<int32>(LevelAttributeName, FTransformCollection::TransformGroup);
-	const int32 NumTransforms = GameThreadCollection.GetNumTransforms();
 
 	TArray<int32> ChildrenToCheckForParentFix;
 
-	for (int32 Index = 0; Index < NumTransforms; ++Index)
+	for (int32 ParticleIndex = 0; ParticleIndex < NumEffectiveParticles; ++ParticleIndex)
 	{
-		if (EffectiveParticles[Index])
-		{
-			FParticle* P = GTParticles[Index].Get();
-			// Generate all particles unique idx at initialization time
-			if (bInitializationTime)
-			{
-				UniqueIdxs[Index] = Evolution->GenerateUniqueIdx();
-			}
+		const int32 TransformIndex = FromParticleToTransformIndex[ParticleIndex];
 
-			if ((bInitializationTime && Index == Parameters.InitialRootIndex) || // When initializing always create particle for the root
-				((bInitializationTime && (bCreateGTParticleForChildren || bBuildGeometryForChildren)) || // When initializing create all particles if one of the flag is true 
-					(!bInitializationTime && Index != Parameters.InitialRootIndex && !bCreateGTParticleForChildren && !bBuildGeometryForChildren))) // When not initializing create other particles if flag was set to not create
-			{
-				GTParticles[Index] = FParticle::CreateParticle();
-				P = GTParticles[Index].Get();
-				GTParticlesToTransformGroupIndex.Add(P, Index);
-				GTParticles[Index]->SetUniqueIdx(UniqueIdxs[Index]);
+		FParticle* P = GTParticles[ParticleIndex].Get();
+		// Generate all particles unique idx at initialization time
+		if (bInitializationTime)
+		{
+			UniqueIdxs[ParticleIndex] = Evolution->GenerateUniqueIdx();
+		}
 
+		if ((bInitializationTime && TransformIndex == Parameters.InitialRootIndex) || // When initializing always create particle for the root
+			((bInitializationTime && (bCreateGTParticleForChildren || bBuildGeometryForChildren)) || // When initializing create all particles if one of the flag is true 
+				(!bInitializationTime && TransformIndex != Parameters.InitialRootIndex && !bCreateGTParticleForChildren && !bBuildGeometryForChildren))) // When not initializing create other particles if flag was set to not create
+		{
+			GTParticles[ParticleIndex] = FParticle::CreateParticle();
+			P = GTParticles[ParticleIndex].Get();
+			GTParticlesToTransformGroupIndex.Add(P, TransformIndex);
+			GTParticles[ParticleIndex]->SetUniqueIdx(UniqueIdxs[ParticleIndex]);
 #if CHAOS_DEBUG_NAME
-				P->SetDebugName(MakeShared<FString, ESPMode::ThreadSafe>(FString::Printf(TEXT("%s-%d"), *Parameters.Name, Index)));
+				P->SetDebugName(MakeShared<FString, ESPMode::ThreadSafe>(FString::Printf(TEXT("%s-%d"), *Parameters.Name, TransformIndex)));
 #endif
-
-				const float ScaledMass = AdjustMassForScale(Mass[Index]);
-
-				// Note that this transform must match the physics thread transform computation for initialization.
-				// Take for example, a geometry collection in the editor that is linked with a joint constraint. The joint
-				// constraint will query the game thread particle position/rotation for the geometry collection to compute its
-				// reference frame. If that position/rotation does not match up with the physics thread's position/rotation,
-				// the geometry collection particle will have an added velocity computed by the joint constraint solver.
-				const FTransform& T = MassToLocal[Index] * FTransform(GameThreadCollection.GetTransform(Index)) * Parameters.WorldTransform;
-				P->SetX(T.GetTranslation(), false);
-				P->SetR(T.GetRotation(), false);
-				P->SetM(ScaledMass);
-				P->SetUserData(Parameters.UserData);
-				P->SetProxy(this);
-			}
-			if (bInitializationTime && Index == Parameters.InitialRootIndex && bUseStaticMeshCollisionForTraces && CreateTraceCollisionGeometryCallback != nullptr)
-			{
-				const FTransform ToLocal = MassToLocal[Index].Inverse();
-				TArray<Chaos::FImplicitObjectPtr> Geoms;
-				Chaos::FShapesArray Shapes;
-				CreateTraceCollisionGeometryCallback(ToLocal, Geoms, Shapes);
-
-				Chaos::FImplicitObjectPtr ImplicitGeometry = MakeImplicitObjectPtr<Chaos::FImplicitObjectUnion>(MoveTemp(Geoms));
-				P->SetGeometry(ImplicitGeometry);
-			}
-			else if (bInitializationTime == (Index == Parameters.InitialRootIndex) || bBuildGeometryForChildren)
+			const float ScaledMass = AdjustMassForScale(Mass[TransformIndex]);
+
+			// Note that this transform must match the physics thread transform computation for initialization.
+			// Take for example, a geometry collection in the editor that is linked with a joint constraint. The joint
+			// constraint will query the game thread particle position/rotation for the geometry collection to compute its
+			// reference frame. If that position/rotation does not match up with the physics thread's position/rotation,
+			// the geometry collection particle will have an added velocity computed by the joint constraint solver.
+			const FTransform& T = MassToLocal[TransformIndex] * FTransform(GameThreadCollection.GetTransform(TransformIndex)) * Parameters.WorldTransform;
+			P->SetX(T.GetTranslation(), false);
+			P->SetR(T.GetRotation(), false);
+			P->SetM(ScaledMass);
+			P->SetUserData(Parameters.UserData);
+			P->SetProxy(this);
+		}
+		if (bInitializationTime && TransformIndex == Parameters.InitialRootIndex && bUseStaticMeshCollisionForTraces && CreateTraceCollisionGeometryCallback != nullptr)
+		{
+			const FTransform ToLocal = MassToLocal[TransformIndex].Inverse();
+			TArray<Chaos::FImplicitObjectPtr> Geoms;
+			Chaos::FShapesArray Shapes;
+			CreateTraceCollisionGeometryCallback(ToLocal, Geoms, Shapes);
+
+			Chaos::FImplicitObjectPtr ImplicitGeometry = MakeImplicitObjectPtr<Chaos::FImplicitObjectUnion>(MoveTemp(Geoms));
+			P->SetGeometry(ImplicitGeometry);
+		}
+		else if (bInitializationTime == (TransformIndex == Parameters.InitialRootIndex) || bBuildGeometryForChildren)
+		{
+			Chaos::FImplicitObjectPtr ImplicitGeometry = Implicits[TransformIndex];
+			if (ImplicitGeometry && !Scale.Equals(FVector::OneVector))
 			{
-				Chaos::FImplicitObjectPtr ImplicitGeometry = Implicits[Index];
-				if (ImplicitGeometry && !Scale.Equals(FVector::OneVector))
-				{
-					ImplicitGeometry = ImplicitGeometry->CopyGeometryWithScale(Scale);
-				}
-				P->SetGeometry(ImplicitGeometry);
+				ImplicitGeometry = ImplicitGeometry->CopyGeometryWithScale(Scale);
 			}
-			// Reset the root in the ManagedArray if not in mode bInitializeRootOnly
-			if (!bInitializationTime && Index == Parameters.InitialRootIndex)
+			P->SetGeometry(ImplicitGeometry);
+		}
+		// Reset the root in the ManagedArray if not in mode bInitializeRootOnly
+		if (!bInitializationTime && TransformIndex == Parameters.InitialRootIndex)
+		{
+			Implicits[TransformIndex] = P->GetGeometry();
+		}
+
+		if (bInitializationTime == (TransformIndex == Parameters.InitialRootIndex) || bBuildGeometryForChildren)
+		{
+			if (DynamicCollectionAnchoringFacade.IsAnchored(TransformIndex))
 			{
-				Implicits[Index] = P->GetGeometry();
+				P->SetObjectState(Chaos::EObjectStateType::Kinematic, false, false);
 			}
 
-			if (bInitializationTime == (Index == Parameters.InitialRootIndex) || bBuildGeometryForChildren)
-			{
-				if (DynamicCollectionAnchoringFacade.IsAnchored(Index))
-				{
-					P->SetObjectState(Chaos::EObjectStateType::Kinematic, false, false);
-				}
+			// IMPORTANT: we need to set the right spatial index because GT particle is static and PT particle is rigid
+			// this is causing a mismatch when using the separate acceleration structures optimization which can cause crashes when destroying the particle while async tracing 
+			// todo(chaos) we should eventually refactor this code to use rigid particles on the GT side for geometry collection  
+			P->SetSpatialIdx(Chaos::FSpatialAccelerationIdx{ 0,1 });
 
-				// IMPORTANT: we need to set the right spatial index because GT particle is static and PT particle is rigid
-				// this is causing a mismatch when using the separate acceleration structures optimization which can cause crashes when destroying the particle while async tracing 
-				// todo(chaos) we should eventually refactor this code to use rigid particles on the GT side for geometry collection  
+			if (Chaos::AccelerationStructureSplitStaticAndDynamic == 1)
+			{
 				P->SetSpatialIdx(Chaos::FSpatialAccelerationIdx{ 0,1 });
-
-				if (Chaos::AccelerationStructureSplitStaticAndDynamic == 1)
-				{
-					P->SetSpatialIdx(Chaos::FSpatialAccelerationIdx{ 0,1 });
-				}
-				else
-				{
-					P->SetSpatialIdx(Chaos::FSpatialAccelerationIdx{ 0,0 });
-				}
-
-				const bool bIsOneWayInteraction = (Parameters.OneWayInteractionLevel >= 0) && (Level[Index] >= Parameters.OneWayInteractionLevel);
-				P->SetOneWayInteraction(bIsOneWayInteraction);
 			}
+			else
+			{
+				P->SetSpatialIdx(Chaos::FSpatialAccelerationIdx{ 0,0 });
+			}
+
+			const bool bIsOneWayInteraction = (Parameters.OneWayInteractionLevel >= 0) && (Level[TransformIndex] >= Parameters.OneWayInteractionLevel);
+			P->SetOneWayInteraction(bIsOneWayInteraction);
 		}
 		// this step is necessary for Phase 2 where we need to walk back the hierarchy from children to parent 
-		if (bGeometryCollectionAlwaysGenerateGTCollisionForClusters && !GameThreadCollection.HasChildren(Index))
+		if (bGeometryCollectionAlwaysGenerateGTCollisionForClusters && !GameThreadCollection.HasChildren(TransformIndex))
 		{
-			ChildrenToCheckForParentFix.Add(Index);
+			ChildrenToCheckForParentFix.Add(TransformIndex);
 		}
 	}
 
@@ -874,9 +880,9 @@ void FGeometryCollectionPhysicsProxy::CreateGTParticles(const TBitArray<>& Effec
 							Chaos::FImplicitObject* UnionImplicit = new Chaos::FImplicitObjectUnion(MoveTemp(ChildImplicits));
 							Implicits[ParentToFixIndex] = Chaos::FImplicitObjectPtr(UnionImplicit);
 						}
-						if (GTParticles[ParentToFixIndex] != nullptr)
+						if (GTParticles[FromTransformToParticleIndex[ParentToFixIndex]] != nullptr)
 						{
-							GTParticles[ParentToFixIndex]->SetGeometry(Implicits[ParentToFixIndex]);
+							GTParticles[FromTransformToParticleIndex[ParentToFixIndex]]->SetGeometry(Implicits[ParentToFixIndex]);
 						}
 					}
 				}
@@ -889,7 +895,7 @@ void FGeometryCollectionPhysicsProxy::CreateGTParticles(const TBitArray<>& Effec
 	}
 
 	// Phase 3 : finalization of shapes
-	for (int32 Index = 0; Index < NumTransforms; ++Index)
+	for (int32 Index = 0; Index < NumEffectiveParticles; ++Index)
 	{
 		FParticle* P = GTParticles[Index].Get();
 		if (P != nullptr)
@@ -922,7 +928,7 @@ void FGeometryCollectionPhysicsProxy::CreateChildrenGeometry_External()
 
 			TManagedArray<Chaos::FImplicitObjectPtr>& Implicits = GameThreadCollection.ModifyAttribute<Chaos::FImplicitObjectPtr>(FGeometryDynamicCollection::ImplicitsAttribute, FTransformCollection::TransformGroup);
 			TBitArray<> EffectiveParticles;
-			NumEffectiveParticles = CalculateEffectiveParticles(GameThreadCollection, NumParticles, Parameters.MaxSimulatedLevel, Parameters.EnableClustering, GetOwner(), EffectiveParticles);
+			NumEffectiveParticles = CalculateEffectiveParticles(GameThreadCollection, NumTransforms, Parameters.MaxSimulatedLevel, Parameters.EnableClustering, GetOwner(), EffectiveParticles);
 			CreateGTParticles(EffectiveParticles, Implicits, RBDSolver->GetEvolution(), /*bInitializationTime*/false);
 			SyncParticles_External();
 			UniqueIdxs.Empty();
@@ -1033,25 +1039,23 @@ void FGeometryCollectionPhysicsProxy::CreateNonClusteredParticles(Chaos::FPBDRig
 
 	// Count geometry collection leaf node particles to add
 	int NumSimulatedParticles = 0;
-	for (int32 Idx = 0; Idx < SimulatableParticles.Num(); ++Idx)
+	for (int32 ParticleIndex = 0; ParticleIndex < NumEffectiveParticles; ++ParticleIndex)
 	{
-		if (EffectiveParticles[Idx])
+		const int32 TransformIndex = FromParticleToTransformIndex[ParticleIndex];
+		NumSimulatedParticles += SimulatableParticles[TransformIndex];
+		if (SimulatableParticles[TransformIndex] && !RestCollection.IsClustered(TransformIndex) && RestCollection.IsGeometry(TransformIndex))
 		{
-			NumSimulatedParticles += SimulatableParticles[Idx];
-			if (SimulatableParticles[Idx] && !RestCollection.IsClustered(Idx) && RestCollection.IsGeometry(Idx))
+			NumRigids++;
+			Chaos::FUniqueIdx ExistingIndex;
+			if (GTParticles[ParticleIndex] == nullptr)
 			{
-				NumRigids++;
-				Chaos::FUniqueIdx ExistingIndex;
-				if (GTParticles[Idx] == nullptr)
-				{
-					ExistingIndex = RigidsSolver->GetEvolution()->GenerateUniqueIdx();
-				}
-				else
-				{
-					ExistingIndex = GTParticles[Idx]->UniqueIdx();
-				}
-				UniqueIndices.Add(ExistingIndex);
+				ExistingIndex = RigidsSolver->GetEvolution()->GenerateUniqueIdx();
+			}
+			else
+			{
+				ExistingIndex = GTParticles[ParticleIndex]->UniqueIdx();
 			}
+			UniqueIndices.Add(ExistingIndex);
 		}
 	}
 
@@ -1060,10 +1064,11 @@ void FGeometryCollectionPhysicsProxy::CreateNonClusteredParticles(Chaos::FPBDRig
 	TArray<Chaos::TPBDGeometryCollectionParticleHandle<Chaos::FReal, 3>*> Handles = RigidsSolver->GetEvolution()->CreateGeometryCollectionParticles(NumRigids, UniqueIndices.GetData());
 
 	int32 NextIdx = 0;
-	for (int32 Idx = 0; Idx < SimulatableParticles.Num(); ++Idx)
+	for (int32 ParticleIndex = 0; ParticleIndex < NumEffectiveParticles; ++ParticleIndex)
 	{
-		SolverParticleHandles[Idx] = nullptr;
-		if (EffectiveParticles[Idx] && SimulatableParticles[Idx] && !RestCollection.IsClustered(Idx))
+		const int32 TransformIndex = FromParticleToTransformIndex[ParticleIndex];
+		SolverParticleHandles[ParticleIndex] = nullptr;
+		if (EffectiveParticles[TransformIndex] && SimulatableParticles[TransformIndex] && !RestCollection.IsClustered(TransformIndex))
 		{
 			// todo: Unblocked read access of game thread data on the physics thread.
 
@@ -1071,13 +1076,13 @@ void FGeometryCollectionPhysicsProxy::CreateNonClusteredParticles(Chaos::FPBDRig
 
 			Handle->SetPhysicsProxy(this);
 
-			SolverParticleHandles[Idx] = Handle;
-			HandleToTransformGroupIndex.Add(Handle, Idx);
+			SolverParticleHandles[ParticleIndex] = Handle;
+			HandleToTransformGroupIndex.Add(Handle, TransformIndex);
 
 			// We're on the physics thread here but we've already set up the GT particles and we're just linking here
-			Handle->GTGeometryParticle() = GTParticles[Idx].Get();
+			Handle->GTGeometryParticle() = GTParticles[ParticleIndex].Get();
 
-			check(SolverParticleHandles[Idx]->GetParticleType() == Handle->GetParticleType());
+			check(SolverParticleHandles[ParticleIndex]->GetParticleType() == Handle->GetParticleType());
 			RigidsSolver->GetEvolution()->RegisterParticle(Handle);
 		}
 	}
@@ -1136,7 +1141,6 @@ void FGeometryCollectionPhysicsProxy::DirtyAllParticles(const Chaos::FPBDRigidsS
 
 void FGeometryCollectionPhysicsProxy::UpdateDamageThreshold_Internal()
 {
-	const int32 NumTransforms = SolverParticleHandles.Num();
 	ensure(SolverParticleHandles.Num() == SolverClusterHandles.Num());
 
 	if (Chaos::FPhysicsSolver* RBDSolver = GetSolver<Chaos::FPhysicsSolver>())
@@ -1145,12 +1149,12 @@ void FGeometryCollectionPhysicsProxy::UpdateDamageThreshold_Internal()
 
 		const float StrainDefault = Parameters.DamageThreshold.Num() ? Parameters.DamageThreshold[0] : 0;
 
-		for (int32 TransformIndex = 0; TransformIndex < NumTransforms; TransformIndex++)
+		for (int32 ParticleIndex = 0; ParticleIndex < NumEffectiveParticles; ParticleIndex++)
 		{
-			if (Chaos::FPBDRigidClusteredParticleHandle* Handle = SolverParticleHandles[TransformIndex])
+			if (Chaos::FPBDRigidClusteredParticleHandle* Handle = SolverParticleHandles[ParticleIndex])
 			{
+				const int32 TransformIndex = FromParticleToTransformIndex[ParticleIndex];
 				float DamageThreshold = StrainDefault;
-
 				switch (Parameters.DamageModel)
 				{
 				case EDamageModelTypeEnum::Chaos_Damage_Model_UserDefined_Damage_Threshold:
@@ -1175,11 +1179,11 @@ void FGeometryCollectionPhysicsProxy::UpdateDamageThreshold_Internal()
 		if (Parameters.bUsePerClusterOnlyDamageThreshold && 
 			Parameters.DamageModel == EDamageModelTypeEnum::Chaos_Damage_Model_UserDefined_Damage_Threshold)
 		{
-			for (int32 TransformIndex = 0; TransformIndex < NumTransforms; TransformIndex++)
+			for (int32 ParticleIndex = 0; ParticleIndex < NumEffectiveParticles; ParticleIndex++)
 			{
-				if (Chaos::FPBDRigidClusteredParticleHandle* Handle = SolverParticleHandles[TransformIndex])
+				if (Chaos::FPBDRigidClusteredParticleHandle* Handle = SolverParticleHandles[ParticleIndex])
 				{
-					if (Chaos::FPBDRigidClusteredParticleHandle* ParentHandle = SolverClusterHandles[TransformIndex])
+					if (Chaos::FPBDRigidClusteredParticleHandle* ParentHandle = SolverClusterHandles[ParticleIndex])
 					{
 						const float DamageThreshold = ParentHandle->GetInternalStrains();
 						Clustering.SetInternalStrain(Handle, DamageThreshold);
@@ -1207,7 +1211,7 @@ void FGeometryCollectionPhysicsProxy::InitializeBodiesPT(Chaos::FPBDRigidsSolver
 		const TManagedArray<float>& Mass = RestCollection->GetAttribute<float>(MassAttributeName, FTransformCollection::TransformGroup);
 		const TManagedArray<FVector3f>& InertiaTensor = RestCollection->GetAttribute<FVector3f>(InertiaTensorAttributeName, FTransformCollection::TransformGroup);
 
-		const int32 NumTransforms = DynamicCollection.NumElements(FTransformCollection::TransformGroup);
+		check(NumTransforms == DynamicCollection.NumElements(FTransformCollection::TransformGroup));
 		const TManagedArray<uint8>& DynamicState = DynamicCollection.DynamicState;
 		const TManagedArray<bool>& SimulatableParticles = DynamicCollection.SimulatableParticles;
 		const TManagedArray<FTransform>& MassToLocal = RestCollection->GetAttribute<FTransform>(MassToLocalAttributeName, FTransformCollection::TransformGroup);
@@ -1255,9 +1259,10 @@ void FGeometryCollectionPhysicsProxy::InitializeBodiesPT(Chaos::FPBDRigidsSolver
 		// for the case when there's a 1-to-1 mapping between transforms and geometries.
 		// At the point that we start supporting instancing, this assumption will no longer
 		// hold, and those reverse mappints will be INDEX_NONE.
-		ParallelFor(NumTransforms, [&](int32 TransformGroupIndex)
+		ParallelFor(NumEffectiveParticles, [&](int32 ParticleIndex)
 		{
-			if (FClusterHandle* Handle = SolverParticleHandles[TransformGroupIndex])
+			const int32 TransformGroupIndex = FromParticleToTransformIndex[ParticleIndex];
+			if (FClusterHandle* Handle = SolverParticleHandles[ParticleIndex])
 			{
 				const bool bIsAnchored = AnchoringFacade.IsAnchored(TransformGroupIndex);
 
@@ -1353,10 +1358,11 @@ void FGeometryCollectionPhysicsProxy::InitializeBodiesPT(Chaos::FPBDRigidsSolver
 			// rule of thumb for parallelization is that each thread needs at least
 			// 1000 operations in order to overcome the expense of threading.  I don't
 			// think that's generally going to be the case here...
-			for (int32 TransformGroupIndex = 0; TransformGroupIndex < NumTransforms; ++TransformGroupIndex)
+			for (int32 ParticleIndex = 0; ParticleIndex < NumEffectiveParticles; ++ParticleIndex)
 			{
-				if (Chaos::TPBDRigidParticleHandle<Chaos::FReal, 3>* Handle = SolverParticleHandles[TransformGroupIndex])
+				if (Chaos::TPBDRigidParticleHandle<Chaos::FReal, 3>* Handle = SolverParticleHandles[ParticleIndex])
 				{
+					const int32 TransformGroupIndex = FromParticleToTransformIndex[ParticleIndex];
 					if (DynamicState[TransformGroupIndex] == (int32)EObjectStateTypeEnum::Chaos_Object_Dynamic)
 					{
 						Handle->SetV(InitialVelocityFacade.InitialLinearVelocityAttribute[TransformGroupIndex]);
@@ -1380,10 +1386,10 @@ void FGeometryCollectionPhysicsProxy::InitializeBodiesPT(Chaos::FPBDRigidsSolver
 			for (const int32 TransformGroupIndex : RecursiveOrder)
 			{
 				if (SimulatableParticles[TransformGroupIndex] && !RestCollection->IsClustered(TransformGroupIndex))
-
 				{
 					// Rigid node
-					SubTreeContainsSimulatableParticle[TransformGroupIndex] = SolverParticleHandles[TransformGroupIndex] != nullptr || !EffectiveParticles[TransformGroupIndex];
+					const int32 ParticleIndex = FromTransformToParticleIndex[TransformGroupIndex];
+					SubTreeContainsSimulatableParticle[TransformGroupIndex] = (ParticleIndex != INDEX_NONE && SolverParticleHandles[ParticleIndex] != nullptr) || !EffectiveParticles[TransformGroupIndex];
 				}
 				else
 				{
@@ -1441,7 +1447,7 @@ void FGeometryCollectionPhysicsProxy::InitializeBodiesPT(Chaos::FPBDRigidsSolver
 
 						if (EffectiveParticles[ChildIndex])
 						{
-							if (Chaos::TPBDRigidClusteredParticleHandle<Chaos::FReal, 3>*Handle = SolverParticleHandles[ChildIndex])
+							if (Chaos::TPBDRigidClusteredParticleHandle<Chaos::FReal, 3>*Handle = SolverParticleHandles[FromTransformToParticleIndex[ChildIndex]])
 							{
 								RigidChildren.Add(Handle);
 								RigidChildrenTransformGroupIndex.Add(ChildIndex);
@@ -1450,7 +1456,8 @@ void FGeometryCollectionPhysicsProxy::InitializeBodiesPT(Chaos::FPBDRigidsSolver
 						return true;
 					});
 
-				if (SolverParticleHandles[TransformGroupIndex] == nullptr)
+				const int32 ParticleIndex = FromTransformToParticleIndex[TransformGroupIndex];
+				if (ParticleIndex != INDEX_NONE && SolverParticleHandles[ParticleIndex] == nullptr)
 				{
 					if (ReportTooManyChildrenNum >= 0 && RigidChildren.Num() > ReportTooManyChildrenNum)
 					{
@@ -1463,7 +1470,7 @@ void FGeometryCollectionPhysicsProxy::InitializeBodiesPT(Chaos::FPBDRigidsSolver
 					CreationParameters.Scale = Parameters.WorldTransform.GetScale3D();
 
 					// Hook the handle up with the GT particle
-					FParticle* GTParticle = GTParticles[TransformGroupIndex].Get();
+					FParticle* GTParticle = GTParticles[ParticleIndex].Get();
 					Chaos::FUniqueIdx ExistingIndex;
 					if (GTParticle != nullptr)
 					{
@@ -1471,7 +1478,7 @@ void FGeometryCollectionPhysicsProxy::InitializeBodiesPT(Chaos::FPBDRigidsSolver
 					}
 					else
 					{
-						ExistingIndex = UniqueIdxs[TransformGroupIndex];
+						ExistingIndex = UniqueIdxs[ParticleIndex];
 					}
 
 					Chaos::FPBDRigidClusteredParticleHandle* Handle = nullptr;
@@ -1496,12 +1503,15 @@ void FGeometryCollectionPhysicsProxy::InitializeBodiesPT(Chaos::FPBDRigidsSolver
 					int32 RigidChildrenIdx = 0;
 					for(const int32 ChildTransformIndex : RigidChildrenTransformGroupIndex)
 					{
-						SolverClusterID[ChildTransformIndex] = RigidChildren[RigidChildrenIdx++]->CastToClustered()->ClusterIds().Id;;
+						if (FromTransformToParticleIndex[ChildTransformIndex] != INDEX_NONE)
+						{
+							SolverClusterID[FromTransformToParticleIndex[ChildTransformIndex]] = RigidChildren[RigidChildrenIdx++]->CastToClustered()->ClusterIds().Id;
+						}
 					}
-					SolverClusterID[TransformGroupIndex] = Handle->ClusterIds().Id;
+					SolverClusterID[ParticleIndex] = Handle->ClusterIds().Id;
 
-					SolverClusterHandles[TransformGroupIndex] = Handle;
-					SolverParticleHandles[TransformGroupIndex] = Handle;
+					SolverClusterHandles[ParticleIndex] = Handle;
+					SolverParticleHandles[ParticleIndex] = Handle;
 					HandleToTransformGroupIndex.Add(Handle, TransformGroupIndex);
 					Handle->SetPhysicsProxy(this);
 
@@ -1525,11 +1535,12 @@ void FGeometryCollectionPhysicsProxy::InitializeBodiesPT(Chaos::FPBDRigidsSolver
 			{
 				// Set cluster connectivity.  TPBDRigidClustering::CreateClusterParticle() 
 				// will optionally do this, but we switch that functionality off in BuildClusters_Internal().
-				for (int32 TransformGroupIndex = 0; TransformGroupIndex < NumTransforms; ++TransformGroupIndex)
+				for (int32 ParticleIndex = 0; ParticleIndex < NumEffectiveParticles; ++ParticleIndex)
 				{
+					const int32 TransformGroupIndex = FromParticleToTransformIndex[ParticleIndex];
 					if (RestCollection->IsClustered(TransformGroupIndex))
 					{
-						if (FClusterHandle* ClusteredParticle = SolverParticleHandles[TransformGroupIndex])
+						if (FClusterHandle* ClusteredParticle = SolverParticleHandles[ParticleIndex])
 						{
 							Chaos::FClusterCreationParameters ClusterParams;
 							// #todo: should other parameters be set here?  Previously, there was no parameters being sent, and it is unclear
@@ -1545,9 +1556,9 @@ void FGeometryCollectionPhysicsProxy::InitializeBodiesPT(Chaos::FPBDRigidsSolver
 				// since we do not have precise data from the tools, we are using an approximate method using the bounds of the two particles
 				if (Parameters.DamageEvaluationModel == Chaos::EDamageEvaluationModel::StrainFromMaterialStrengthAndConnectivity)
 				{
-					for (int32 TransformGroupIndex = 0; TransformGroupIndex < NumTransforms; ++TransformGroupIndex)
+					for (int32 ParticleIndex = 0; ParticleIndex < NumEffectiveParticles; ++ParticleIndex)
 					{
-						if (FClusterHandle* ClusteredParticle = SolverParticleHandles[TransformGroupIndex])
+						if (FClusterHandle* ClusteredParticle = SolverParticleHandles[ParticleIndex])
 						{
 							for (Chaos::FConnectivityEdge& Edge : ClusteredParticle->ConnectivityEdges())
 							{
@@ -1580,18 +1591,23 @@ void FGeometryCollectionPhysicsProxy::InitializeBodiesPT(Chaos::FPBDRigidsSolver
 						ContactArea = ConnectionFacade.GetConnectionContactArea(ConnectionIdx);
 					}
 
-					if (FClusterHandle* ClusteredParticle = SolverParticleHandles[Connection.Key])
+					const int32 KeyIndex = FromTransformToParticleIndex[Connection.Key];
+					const int32 ValueIndex = FromTransformToParticleIndex[Connection.Value];
+					if (KeyIndex != INDEX_NONE && ValueIndex != INDEX_NONE)
 					{
-						if (FClusterHandle* OtherClusteredParticle = SolverParticleHandles[Connection.Value])
+						if (FClusterHandle* ClusteredParticle = SolverParticleHandles[KeyIndex])
 						{
-							float EdgeStrainOrArea = bHasContactAreas ? ContactArea : (ClusteredParticle->GetInternalStrains() + OtherClusteredParticle->GetInternalStrains()) * 0.5f;
-							if (!bHasContactAreas && Parameters.DamageEvaluationModel == Chaos::EDamageEvaluationModel::StrainFromMaterialStrengthAndConnectivity)
+							if (FClusterHandle* OtherClusteredParticle = SolverParticleHandles[ValueIndex])
 							{
-								EdgeStrainOrArea = AreaFromBoundingBoxOverlap(ClusteredParticle->WorldSpaceInflatedBounds(), OtherClusteredParticle->WorldSpaceInflatedBounds());
+								float EdgeStrainOrArea = bHasContactAreas ? ContactArea : (ClusteredParticle->GetInternalStrains() + OtherClusteredParticle->GetInternalStrains()) * 0.5f;
+								if (!bHasContactAreas && Parameters.DamageEvaluationModel == Chaos::EDamageEvaluationModel::StrainFromMaterialStrengthAndConnectivity)
+								{
+									EdgeStrainOrArea = AreaFromBoundingBoxOverlap(ClusteredParticle->WorldSpaceInflatedBounds(), OtherClusteredParticle->WorldSpaceInflatedBounds());
+								}
+								// Add symmetric Chaos::TConnectivityEdge<Chaos::FReal> edges to each particle
+								ClusteredParticle->ConnectivityEdges().Emplace(OtherClusteredParticle, EdgeStrainOrArea);
+								OtherClusteredParticle->ConnectivityEdges().Emplace(ClusteredParticle, EdgeStrainOrArea);
 							}
-							// Add symmetric Chaos::TConnectivityEdge<Chaos::FReal> edges to each particle
-							ClusteredParticle->ConnectivityEdges().Emplace(OtherClusteredParticle, EdgeStrainOrArea);
-							OtherClusteredParticle->ConnectivityEdges().Emplace(ClusteredParticle, EdgeStrainOrArea);
 						}
 					}
 				}
@@ -1676,7 +1692,7 @@ void FGeometryCollectionPhysicsProxy::CreateChildrenGeometry_Internal()
 				PhysicsThreadCollection.CopyAttribute(*Parameters.RestCollection, FGeometryDynamicCollection::ImplicitsAttribute, FTransformCollection::TransformGroup);
 			}
 
-			const int32 NumTransforms = PhysicsThreadCollection.NumElements(FTransformCollection::TransformGroup);
+			check(NumTransforms == PhysicsThreadCollection.NumElements(FTransformCollection::TransformGroup));
 			const TManagedArray<FTransform>& MassToLocal = RestCollection->GetAttribute<FTransform>(MassToLocalAttributeName, FTransformCollection::TransformGroup);
 			const TManagedArray<Chaos::FImplicitObjectPtr>& Implicits = PhysicsThreadCollection.GetAttribute<Chaos::FImplicitObjectPtr>(FGeometryDynamicCollection::ImplicitsAttribute, FTransformCollection::TransformGroup);
 			const TManagedArray<TUniquePtr<FCollisionStructureManager::FSimplicial>>& Simplicials = PhysicsThreadCollection.Simplicials;
@@ -1692,19 +1708,16 @@ void FGeometryCollectionPhysicsProxy::CreateChildrenGeometry_Internal()
 			TArray<FTransform> Transform;
 			GeometryCollectionAlgo::Private::GlobalMatrices(PhysicsThreadCollection, Transform);
 
-			// Here Clean up Additional particles
-			TBitArray<> EffectiveParticles;
-			NumEffectiveParticles = CalculateEffectiveParticles(PhysicsThreadCollection, PhysicsThreadCollection.NumElements(FGeometryCollection::TransformGroup), Parameters.MaxSimulatedLevel, Parameters.EnableClustering, GetOwner(), EffectiveParticles);
-
-			for (int32 Index = 0; Index < NumTransforms; Index++)
+			for (int32 ParticleIndex = 0; ParticleIndex < NumEffectiveParticles; ParticleIndex++)
 			{
-				if (EffectiveParticles[Index] && Index != Parameters.InitialRootIndex)
+				const int32 TransformIndex = FromParticleToTransformIndex[ParticleIndex];
+				if (TransformIndex != Parameters.InitialRootIndex)
 				{
-					Chaos::TPBDGeometryCollectionParticleHandle<Chaos::FReal, 3>* Handle = static_cast<Chaos::TPBDGeometryCollectionParticleHandle<Chaos::FReal, 3>*>(SolverParticleHandles[Index]);
+					Chaos::TPBDGeometryCollectionParticleHandle<Chaos::FReal, 3>* Handle = static_cast<Chaos::TPBDGeometryCollectionParticleHandle<Chaos::FReal, 3>*>(SolverParticleHandles[ParticleIndex]);
 					if (ensure(Handle) && Handle->GetGeometry() == nullptr)
 					{
-						const FTransform WorldTransform = MassToLocal[Index] * Transform[Index] * Parameters.WorldTransform;
-						SetImplicitToPTParticles(Handle, Parameters.Shared, Simplicials[Index].Get(), Implicits[Index], SimFilter, QueryFilter, WorldTransform, CollisionParticlesPerObjectFraction);
+						const FTransform WorldTransform = MassToLocal[TransformIndex] * Transform[TransformIndex] * Parameters.WorldTransform;
+						SetImplicitToPTParticles(Handle, Parameters.Shared, Simplicials[TransformIndex].Get(), Implicits[TransformIndex], SimFilter, QueryFilter, WorldTransform, CollisionParticlesPerObjectFraction);
 						check(Handle->GetGeometry() != nullptr);
 					}
 				}
@@ -1730,12 +1743,10 @@ void FGeometryCollectionPhysicsProxy::SyncParticles_External()
 		for (int32 Index = 0; Index < ParticlNum; ++Index)
 		{
 			FParticle* GTParticle = GTParticles[Index].Get();
-			if (GTParticle)
-			{
-				FClusterHandle* Handle = SolverParticleHandles[Index];
-				check(Handle != nullptr);
-				Handle->GTGeometryParticle() = GTParticle;
-			}
+			check(GTParticle != nullptr);
+			FClusterHandle* Handle = SolverParticleHandles[Index];
+			check(Handle != nullptr);
+			Handle->GTGeometryParticle() = GTParticle;
 		}
 	});
 }
@@ -1756,9 +1767,10 @@ float FGeometryCollectionPhysicsProxy::ComputeMaterialBasedDamageThreshold_Inter
 	float DamageThreshold = TNumericLimits<float>::Max();
 	if (Parameters.DamageModel == EDamageModelTypeEnum::Chaos_Damage_Model_Material_Strength_And_Connectivity_DamageThreshold)
 	{
-		if (ensure(SolverParticleHandles.IsValidIndex(TransformIndex)))
+		const int32 ParticleIndex = FromTransformToParticleIndex[TransformIndex];
+		if (ensure(SolverParticleHandles.IsValidIndex(ParticleIndex)))
 		{
-			if (Chaos::FPBDRigidClusteredParticleHandle* ClusteredParticle = SolverParticleHandles[TransformIndex])
+			if (Chaos::FPBDRigidClusteredParticleHandle* ClusteredParticle = SolverParticleHandles[ParticleIndex])
 			{
 				DamageThreshold = ComputeMaterialBasedDamageThreshold_Internal(*ClusteredParticle);
 			}
@@ -1824,12 +1836,13 @@ float FGeometryCollectionPhysicsProxy::ComputeUserDefinedDamageThreshold_Interna
 	}
 	else if (Parameters.DamageModel == EDamageModelTypeEnum::Chaos_Damage_Model_UserDefined_Damage_Threshold)
 	{
+		const int32 ParticleIndex = FromTransformToParticleIndex[TransformIndex];
 		// we don't test for the damage model because this is used
-		if (Parameters.bUseSizeSpecificDamageThresholds)
+		if (ParticleIndex != INDEX_NONE && Parameters.bUseSizeSpecificDamageThresholds)
 		{
 			// bounding box volume is used as a fallback to find specific size if the relative size if not available
 			// ( May happen with older GC )
-			const FClusterHandle* Handle = SolverParticleHandles[TransformIndex];
+			const FClusterHandle* Handle = SolverParticleHandles[ParticleIndex];
 			FBox LocalBoundingBox;
 			if (Handle && Handle->HasBounds())
 			{
@@ -1898,14 +1911,15 @@ Chaos::TPBDGeometryCollectionParticleHandle<Chaos::FReal, 3>* FGeometryCollectio
 	Chaos::TPBDGeometryCollectionParticleHandle<Chaos::FReal, 3>* Handle = Handles[0];
 
 	Handle->SetPhysicsProxy(this);
-
-	SolverParticleHandles[CollectionClusterIndex] = Handle;
+	const int32 ParticleIndex = FromTransformToParticleIndex[CollectionClusterIndex];
+	check(ParticleIndex != INDEX_NONE);
+	SolverParticleHandles[ParticleIndex] = Handle;
 	HandleToTransformGroupIndex.Add(Handle, CollectionClusterIndex);
 
 	// We're on the physics thread here but we've already set up the GT particles and we're just linking here
-	Handle->GTGeometryParticle() = GTParticles[CollectionClusterIndex].Get();
+	Handle->GTGeometryParticle() = GTParticles[ParticleIndex].Get();
 
-	check(SolverParticleHandles[CollectionClusterIndex]->GetParticleType() == Handle->GetParticleType());
+	check(SolverParticleHandles[ParticleIndex]->GetParticleType() == Handle->GetParticleType());
 	RigidsSolver->GetEvolution()->RegisterParticle(Handle);
 
 
@@ -2073,12 +2087,14 @@ FGeometryCollectionPhysicsProxy::BuildClusters_Internal(
 		ScaledMass,
 		ScaledInertia,
 		ParticleTM, 
-		(uint8)DynamicState[CollectionClusterIndex], 
+		(uint8)DynamicState[CollectionClusterIndex],
 		0, // static_cast<int16>(CollisionGroup[TransformGroupIndex])
 		CollisionParticlesPerObjectFraction); // CollisionGroup
 
 	// two-way mapping
-	SolverClusterHandles[CollectionClusterIndex] = Parent;
+	const int32 ParticleClusterIndex = FromTransformToParticleIndex[CollectionClusterIndex];
+	check(ParticleClusterIndex != INDEX_NONE);
+	SolverClusterHandles[ParticleClusterIndex] = Parent;
 
 	// #BGTODO This will not automatically update - material properties should only ever exist in the material, not in other arrays
 	const Chaos::FChaosPhysicsMaterial* CurMaterial = static_cast<Chaos::FPBDRigidsSolver*>(Solver)->GetSimMaterials().Get(Parameters.PhysicalMaterialHandle.InnerHandle);
@@ -2103,7 +2119,9 @@ FGeometryCollectionPhysicsProxy::BuildClusters_Internal(
 		Chaos::FPBDRigidParticleHandle* Child = ChildHandles[Idx];
 
 		const int32 ChildTransformGroupIndex = ChildTransformGroupIndices[Idx];
-		SolverClusterHandles[ChildTransformGroupIndex] = Parent;
+		const int32 ChildParticleIndex = FromTransformToParticleIndex[ChildTransformGroupIndex];
+		check(ChildParticleIndex != INDEX_NONE);
+		SolverClusterHandles[ChildParticleIndex] = Parent;
 
 		MinCollisionGroup = FMath::Min(Child->CollisionGroup(), MinCollisionGroup);
 	}
@@ -2133,12 +2151,11 @@ void FGeometryCollectionPhysicsProxy::GetFilteredParticleHandles(
 	if ((ObjectType == EFieldObjectType::Field_Object_All) || (ObjectType == EFieldObjectType::Field_Object_Destruction) || (ObjectType == EFieldObjectType::Field_Object_Max))
 	{
 		// only the local handles
-		TArray<FClusterHandle*>& ParticleHandles = GetSolverParticleHandles();
-		Handles.Reserve(ParticleHandles.Num());
+		Handles.Reserve(SolverParticleHandles.Num());
 
 		if (FilterType == EFieldFilterType::Field_Filter_Dynamic)
 		{
-			for (FClusterHandle* ClusterHandle : ParticleHandles)
+			for (FClusterHandle* ClusterHandle : SolverParticleHandles)
 			{
 				if (ClusterHandle && (ClusterHandle->ObjectState() == Chaos::EObjectStateType::Dynamic))
 				{
@@ -2148,7 +2165,7 @@ void FGeometryCollectionPhysicsProxy::GetFilteredParticleHandles(
 		}
 		else if (FilterType == EFieldFilterType::Field_Filter_Kinematic)
 		{
-			for (FClusterHandle* ClusterHandle : ParticleHandles)
+			for (FClusterHandle* ClusterHandle : SolverParticleHandles)
 			{
 				if (ClusterHandle && (ClusterHandle->ObjectState() == Chaos::EObjectStateType::Kinematic))
 				{
@@ -2158,7 +2175,7 @@ void FGeometryCollectionPhysicsProxy::GetFilteredParticleHandles(
 		}
 		else if (FilterType == EFieldFilterType::Field_Filter_Static)
 		{
-			for (FClusterHandle* ClusterHandle : ParticleHandles)
+			for (FClusterHandle* ClusterHandle : SolverParticleHandles)
 			{
 				if (ClusterHandle && (ClusterHandle->ObjectState() == Chaos::EObjectStateType::Static))
 				{
@@ -2168,7 +2185,7 @@ void FGeometryCollectionPhysicsProxy::GetFilteredParticleHandles(
 		}
 		else if (FilterType == EFieldFilterType::Field_Filter_Sleeping)
 		{
-			for (FClusterHandle* ClusterHandle : ParticleHandles)
+			for (FClusterHandle* ClusterHandle : SolverParticleHandles)
 			{
 				if (ClusterHandle && (ClusterHandle->ObjectState() == Chaos::EObjectStateType::Sleeping))
 				{
@@ -2178,7 +2195,7 @@ void FGeometryCollectionPhysicsProxy::GetFilteredParticleHandles(
 		}
 		else if (FilterType == EFieldFilterType::Field_Filter_Disabled)
 		{
-			for (FClusterHandle* ClusterHandle : ParticleHandles)
+			for (FClusterHandle* ClusterHandle : SolverParticleHandles)
 			{
 				if (ClusterHandle && ClusterHandle->Disabled())
 				{
@@ -2188,7 +2205,7 @@ void FGeometryCollectionPhysicsProxy::GetFilteredParticleHandles(
 		}
 		else if (FilterType == EFieldFilterType::Field_Filter_All)
 		{
-			for (FClusterHandle* ClusterHandle : ParticleHandles)
+			for (FClusterHandle* ClusterHandle : SolverParticleHandles)
 			{
 				if (ClusterHandle && (ClusterHandle->ObjectState() != Chaos::EObjectStateType::Uninitialized))
 				{
@@ -2207,12 +2224,11 @@ void FGeometryCollectionPhysicsProxy::GetRelevantParticleHandles(
 	Handles.SetNum(0, false);
 
 	// only the local handles
-	TArray<FClusterHandle*>& ParticleHandles = GetSolverParticleHandles();
-	Handles.Reserve(ParticleHandles.Num());
+	Handles.Reserve(SolverParticleHandles.Num());
 
 	if (ResolutionType == EFieldResolutionType::Field_Resolution_Maximum)
 	{
-		for (FClusterHandle* ClusterHandle : ParticleHandles)
+		for (FClusterHandle* ClusterHandle : SolverParticleHandles)
 		{
 			if (ClusterHandle )
 			{
@@ -2222,7 +2238,7 @@ void FGeometryCollectionPhysicsProxy::GetRelevantParticleHandles(
 	}
 	else if (ResolutionType == EFieldResolutionType::Field_Resolution_DisabledParents)
 	{
-		for (FClusterHandle* ClusterHandle : ParticleHandles)
+		for (FClusterHandle* ClusterHandle : SolverParticleHandles)
 		{
 			if (ClusterHandle && ClusterHandle->ClusterIds().Id == nullptr)
 			{
@@ -2235,7 +2251,7 @@ void FGeometryCollectionPhysicsProxy::GetRelevantParticleHandles(
 		const auto& Clustering = RigidSolver->GetEvolution()->GetRigidClustering();
 		const auto& ClusterMap = Clustering.GetChildrenMap();
 
-		for (FClusterHandle* ClusterHandle : ParticleHandles)
+		for (FClusterHandle* ClusterHandle : SolverParticleHandles)
 		{
 			if (ClusterHandle && !ClusterHandle->Disabled())
 			{
@@ -2328,8 +2344,10 @@ void FGeometryCollectionPhysicsProxy::DisableParticles_External(TArray<int32>&&
 			{
 				for (int32 TransformIdx : IndicesToDisable)
 				{
-					RBDSolver->GetEvolution()->DisableParticleWithRemovalEvent(SolverParticleHandles[TransformIdx]);
-					RBDSolver->GetParticles().MarkTransientDirtyParticle(SolverParticleHandles[TransformIdx]);
+					const int32 ParticleIndex = FromTransformToParticleIndex[TransformIdx];
+					check(ParticleIndex != INDEX_NONE);
+					RBDSolver->GetEvolution()->DisableParticleWithRemovalEvent(SolverParticleHandles[ParticleIndex]);
+					RBDSolver->GetParticles().MarkTransientDirtyParticle(SolverParticleHandles[ParticleIndex]);
 				}
 			});
 	}
@@ -2486,10 +2504,11 @@ void FGeometryCollectionPhysicsProxy::SetAnchoredByIndex_External(int32 Index, b
 		RBDSolver->EnqueueCommandImmediate([this, Index, bAnchored, RBDSolver]()
 			{
 				Chaos::FPBDRigidsEvolution* Evolution = RBDSolver->GetEvolution();
-				if (SolverParticleHandles.IsValidIndex(Index))
+				const int32 ParticleIndex = FromTransformToParticleIndex[Index];
+				if (SolverParticleHandles.IsValidIndex(ParticleIndex))
 				{
-					SetParticleAnchored_Internal(Evolution, SolverParticleHandles[Index], bAnchored);
-					if (Chaos::FPBDRigidClusteredParticleHandle* TopParentHandle = GetTopActiveClusteredParent_Internal(SolverParticleHandles[Index]))
+					SetParticleAnchored_Internal(Evolution, SolverParticleHandles[ParticleIndex], bAnchored);
+					if (Chaos::FPBDRigidClusteredParticleHandle* TopParentHandle = GetTopActiveClusteredParent_Internal(SolverParticleHandles[ParticleIndex]))
 					{
 						Chaos::UpdateKinematicProperties(TopParentHandle, Evolution->GetRigidClustering().GetChildrenMap(), *Evolution);
 					}
@@ -2514,39 +2533,38 @@ void FGeometryCollectionPhysicsProxy::SetAnchoredByTransformedBox_External(const
 				const int32 MaxLevelToCheck = (InitialLevelAttribute.IsValid()) ? MaxLevel : INDEX_NONE;
 
 				FPBDRigidsEvolution* Evolution = RBDSolver->GetEvolution();
-				for (int32 TransformIndex = 0; TransformIndex < SolverParticleHandles.Num(); TransformIndex++)
+				for (int32 ParticleIndex = 0; ParticleIndex < SolverParticleHandles.Num(); ParticleIndex++)
 				{
-					if (FClusterHandle* ParticleHandle = SolverParticleHandles[TransformIndex])
+					FClusterHandle* ParticleHandle = SolverParticleHandles[ParticleIndex];
+					check(ParticleHandle != nullptr);
+					const FVec3 PositionInBoxSpace = BoxTransform.InverseTransformPosition(ParticleHandle->X());
+					if (BoundsToCheck.Contains(PositionInBoxSpace))
 					{
-						const FVec3 PositionInBoxSpace = BoxTransform.InverseTransformPosition(ParticleHandle->X());
-						if (BoundsToCheck.Contains(PositionInBoxSpace))
+						// if we have a max level , we make sure to anchor a parent of the right level 
+						if (MaxLevelToCheck > INDEX_NONE)
 						{
-							// if we have a max level , we make sure to anchor a parent of the right level 
-							if (MaxLevelToCheck > INDEX_NONE)
+							const int32 ParticleLevel = InitialLevelAttribute.Get()[FromParticleToTransformIndex[ParticleIndex]];
+							if (ParticleLevel > MaxLevelToCheck)
 							{
-								const int32 ParticleLevel = InitialLevelAttribute.Get()[TransformIndex];
-								if (ParticleLevel > MaxLevelToCheck)
+								Chaos::FPBDRigidClusteredParticleHandle* ParentParticle = ParticleHandle->Parent();
+								int32 ParentParticleLevel = (ParticleLevel - 1);
+								// we are above level, so we still contribute to the parent matching the level 
+								while (ParentParticle && ParentParticleLevel > MaxLevelToCheck)
 								{
-									Chaos::FPBDRigidClusteredParticleHandle* ParentParticle = ParticleHandle->Parent();
-									int32 ParentParticleLevel = (ParticleLevel - 1);
-									// we are above level, so we still contribute to the parent matching the level 
-									while (ParentParticle && ParentParticleLevel > MaxLevelToCheck)
-									{
-										ParentParticle = ParentParticle->Parent();
-									}
-									// if we still have a parent , this means the level requirement has been found
-									// let's use this particle to be anchored
-									// if the parent is null, then we don't want to anchor the particle 
-									ParticleHandle = ParentParticle;
+									ParentParticle = ParentParticle->Parent();
 								}
+								// if we still have a parent , this means the level requirement has been found
+								// let's use this particle to be anchored
+								// if the parent is null, then we don't want to anchor the particle 
+								ParticleHandle = ParentParticle;
 							}
-							if (ParticleHandle)
+						}
+						if (ParticleHandle)
+						{
+							SetParticleAnchored_Internal(Evolution, ParticleHandle, bAnchored);
+							if (Chaos::FPBDRigidClusteredParticleHandle* TopParentHandle = GetTopActiveClusteredParent_Internal(ParticleHandle))
 							{
-								SetParticleAnchored_Internal(Evolution, ParticleHandle, bAnchored);
-								if (Chaos::FPBDRigidClusteredParticleHandle* TopParentHandle = GetTopActiveClusteredParent_Internal(ParticleHandle))
-								{
-									TopParentHandles.Add(TopParentHandle);
-								}
+								TopParentHandles.Add(TopParentHandle);
 							}
 						}
 					}
@@ -2567,7 +2585,7 @@ void FGeometryCollectionPhysicsProxy::RemoveAllAnchors_External()
 		RBDSolver->EnqueueCommandImmediate([this, RBDSolver]()
 			{
 				Chaos::FPBDRigidsEvolution* Evolution = RBDSolver->GetEvolution();
-				for (FClusterHandle* ParticleHandle : GetSolverParticleHandles())
+				for (FClusterHandle* ParticleHandle : SolverParticleHandles)
 				{
 					SetParticleAnchored_Internal(Evolution, ParticleHandle, false);
 					// we do not need to update the kinematic state since everything is now dynamic 
@@ -3130,7 +3148,7 @@ void FGeometryCollectionPhysicsProxy::OnRemoveFromScene()
 	}
 
 	const int32 Begin = BaseParticleIndex;
-	const int32 Count = NumParticles;
+	const int32 Count = NumTransforms;
 
 	if (ensure((int32)Particles.Size() > 0 && (Begin + Count) <= (int32)Particles.Size()))
 	{
@@ -3151,7 +3169,7 @@ void FGeometryCollectionPhysicsProxy::SyncBeforeDestroy()
 	{
 		Chaos::FRigidClustering& RigidClustering = RigidSolver->GetEvolution()->GetRigidClustering();
 		Chaos::FClusterUnionManager& ClusterUnionManager = RigidClustering.GetClusterUnionManager();
-		for (Chaos::FPBDRigidClusteredParticleHandle* Handle : GetSolverParticleHandles())
+		for (Chaos::FPBDRigidClusteredParticleHandle* Handle : SolverParticleHandles)
 		{
 			if (Handle)
 			{
@@ -3232,7 +3250,7 @@ void FGeometryCollectionPhysicsProxy::SetWorldTransform_Internal(const FTransfor
 	TSet<FClusterHandle*> ProcessedInternalClusters;
 	const FTransform& ActorToWorld = Parameters.WorldTransform;
 
-	const int32 NumTransforms = PhysicsThreadCollection.NumElements(FGeometryCollection::TransformGroup);
+	check(NumTransforms == PhysicsThreadCollection.NumElements(FGeometryCollection::TransformGroup));
 	if (Chaos::FPhysicsSolver* RigidSolver = GetSolver<Chaos::FPhysicsSolver>())
 	{
 		FClusterUnionManager& ClusterUnionManager = RigidSolver->GetEvolution()->GetRigidClustering().GetClusterUnionManager();
@@ -3244,57 +3262,57 @@ void FGeometryCollectionPhysicsProxy::SetWorldTransform_Internal(const FTransfor
 
 		const TManagedArray<FTransform>& MassToLocal = Parameters.RestCollection->GetAttribute<FTransform>(MassToLocalAttributeName, FTransformCollection::TransformGroup);
 
-		for (int32 TransformGroupIndex = 0; TransformGroupIndex < NumTransforms; ++TransformGroupIndex)
+		for (int32 ParticleIndex = 0; ParticleIndex < NumEffectiveParticles; ++ParticleIndex)
 		{
-			if (Chaos::FPBDRigidClusteredParticleHandle* Handle = SolverParticleHandles[TransformGroupIndex])
-			{
-				FClusterHandle* KinematicRootHandle = nullptr;
-				FClusterHandle* ParentHandle = Handle->Parent();
+			Chaos::FPBDRigidClusteredParticleHandle* Handle = SolverParticleHandles[ParticleIndex];
+			check(Handle != nullptr);
+			FClusterHandle* KinematicRootHandle = nullptr;
+			FClusterHandle* ParentHandle = Handle->Parent();
 
-				if (!Handle->Disabled() && Handle->ObjectState() == Chaos::EObjectStateType::Kinematic)
-				{
-					KinematicRootHandle = Handle;
-				}
-				else
+			if (!Handle->Disabled() && Handle->ObjectState() == Chaos::EObjectStateType::Kinematic)
+			{
+				KinematicRootHandle = Handle;
+			}
+			else
+			{
+				// is there a internal parent as a kinematic root?
+				if (ParentHandle && ParentHandle->InternalCluster() && !ParentHandle->Disabled() && ParentHandle->ObjectState() == Chaos::EObjectStateType::Kinematic && ParentHandle->PhysicsProxy() == this)
 				{
-					// is there a internal parent as a kinematic root?
-					if (ParentHandle && ParentHandle->InternalCluster() && !ParentHandle->Disabled() && ParentHandle->ObjectState() == Chaos::EObjectStateType::Kinematic && ParentHandle->PhysicsProxy() == this)
+					if (!ProcessedInternalClusters.Contains(ParentHandle))
 					{
-						if (!ProcessedInternalClusters.Contains(ParentHandle))
-						{
-							ProcessedInternalClusters.Add(ParentHandle);
-							KinematicRootHandle = ParentHandle;
-						}
+						ProcessedInternalClusters.Add(ParentHandle);
+						KinematicRootHandle = ParentHandle;
 					}
 				}
+			}
 
-				if (KinematicRootHandle)
-				{
-					const FTransform RootWorldTransform(KinematicRootHandle->R(), KinematicRootHandle->X());
-					const FTransform RootRelativeTransform = RootWorldTransform.GetRelativeTransform(Parameters.PrevWorldTransform);
-					const FTransform WorldTransform = RootRelativeTransform * ActorToWorld;
+			if (KinematicRootHandle)
+			{
+				const FTransform RootWorldTransform(KinematicRootHandle->R(), KinematicRootHandle->X());
+				const FTransform RootRelativeTransform = RootWorldTransform.GetRelativeTransform(Parameters.PrevWorldTransform);
+				const FTransform WorldTransform = RootRelativeTransform * ActorToWorld;
 
-					SetClusteredParticleKinematicTarget_Internal(KinematicRootHandle, WorldTransform);
-				}
-				else if (ParentHandle && !ParentHandle->IsDynamic())
+				SetClusteredParticleKinematicTarget_Internal(KinematicRootHandle, WorldTransform);
+			}
+			else if (ParentHandle && !ParentHandle->IsDynamic())
+			{
+				if (ClusterUnionIndex == INDEX_NONE)
 				{
-					if (ClusterUnionIndex == INDEX_NONE)
-					{
-						ClusterUnionIndex = ClusterUnionManager.FindClusterUnionIndexFromParticle(Handle);
-					}
+					ClusterUnionIndex = ClusterUnionManager.FindClusterUnionIndexFromParticle(Handle);
+				}
 
-					if (ClusterUnionIndex != INDEX_NONE)
-					{
-						const FTransform ParentWorldTransform{ ParentHandle->R(), ParentHandle->X() };
-						const FTransform NewWorldTransform = MassToLocal[TransformGroupIndex] * FTransform(PhysicsThreadCollection.GetTransform(TransformGroupIndex)) * Parameters.WorldTransform;
-						const FTransform RelativeTransform = NewWorldTransform.GetRelativeTransform(ParentWorldTransform);
+				if (ClusterUnionIndex != INDEX_NONE)
+				{
+					const int32 TransformGroupIndex = FromParticleToTransformIndex[ParticleIndex];
+					const FTransform ParentWorldTransform{ ParentHandle->R(), ParentHandle->X() };
+					const FTransform NewWorldTransform = MassToLocal[TransformGroupIndex] * FTransform(PhysicsThreadCollection.GetTransform(TransformGroupIndex)) * Parameters.WorldTransform;
+					const FTransform RelativeTransform = NewWorldTransform.GetRelativeTransform(ParentWorldTransform);
 
-						DeferredClusterUnionParticleUpdates.Add(Handle);
-						DeferredClusterUnionChildToParentUpdates.Add(RelativeTransform);
+					DeferredClusterUnionParticleUpdates.Add(Handle);
+					DeferredClusterUnionChildToParentUpdates.Add(RelativeTransform);
 
-						// Make sure the particle is mark dirty to make sure its proxy will properly update the transforms
-						RigidSolver->GetEvolution()->GetParticles().MarkTransientDirtyParticle(Handle);
-					}
+					// Make sure the particle is mark dirty to make sure its proxy will properly update the transforms
+					RigidSolver->GetEvolution()->GetParticles().MarkTransientDirtyParticle(Handle);
 				}
 			}
 		}
@@ -3328,16 +3346,18 @@ void FGeometryCollectionPhysicsProxy::SetUseStaticMeshCollisionForTraces_Externa
 	const bool bUseSMCollisionForTraces = ((bInUseStaticMeshCollisionForTraces && ForceOverrideGCCollisionSetupForTraces == GCCSFT_Property) || (ForceOverrideGCCollisionSetupForTraces == GCCSFT_ForceSM));
 	if (bUseStaticMeshCollisionForTraces != bUseSMCollisionForTraces && CreateTraceCollisionGeometryCallback != nullptr)
 	{
-		const int32 Index = Parameters.InitialRootIndex;
-		if (GTParticles.Num() > Index)
+		const int32 TransformIndex = Parameters.InitialRootIndex;
+		const int32 ParticleIndex = FromTransformToParticleIndex[TransformIndex];
+		if (GTParticles.IsValidIndex(ParticleIndex))
 		{
-			if (FParticle* P = GTParticles[Index].Get())
+			FParticle* P = GTParticles[ParticleIndex].Get();
+			if (P != nullptr)
 			{
 				if (bUseSMCollisionForTraces)
 				{
 					const TManagedArray<FTransform>& MassToLocal = Parameters.RestCollection->GetAttribute<FTransform>("MassToLocal", FGeometryCollection::TransformGroup);
 
-					const FTransform ToLocal = MassToLocal[Index].Inverse();
+					const FTransform ToLocal = MassToLocal[TransformIndex].Inverse();
 					TArray<Chaos::FImplicitObjectPtr> Geoms;
 					Chaos::FShapesArray Shapes;
 					CreateTraceCollisionGeometryCallback(ToLocal, Geoms, Shapes);
@@ -3349,17 +3369,16 @@ void FGeometryCollectionPhysicsProxy::SetUseStaticMeshCollisionForTraces_Externa
 				{
 					const FVector Scale = WorldTransform_External.GetScale3D();
 					TManagedArray<Chaos::FImplicitObjectPtr>& Implicits = GameThreadCollection.ModifyAttribute<Chaos::FImplicitObjectPtr>(FGeometryDynamicCollection::ImplicitsAttribute, FTransformCollection::TransformGroup);
-					Chaos::FImplicitObjectPtr ImplicitGeometry = Implicits[Index];
+					Chaos::FImplicitObjectPtr ImplicitGeometry = Implicits[TransformIndex];
 					if (ImplicitGeometry && !Scale.Equals(FVector::OneVector))
 					{
 						ImplicitGeometry = ImplicitGeometry->CopyGeometryWithScale(Scale);
 					}
 					P->SetGeometry(ImplicitGeometry);
 				}
-
-				// We only want to change the value of our boolean if we actually changed the collision, so the collision state matches
-				bUseStaticMeshCollisionForTraces = bUseSMCollisionForTraces;
 			}
+			// We only want to change the value of our boolean if we actually changed the collision, so the collision state matches
+			bUseStaticMeshCollisionForTraces = bUseSMCollisionForTraces;
 		}
 	}
 }
@@ -3452,20 +3471,19 @@ void FGeometryCollectionPhysicsProxy::SetMaterialOverrideMassScaleMultiplier_Int
 	}
 	Parameters.MaterialOverrideMassScaleMultiplier = NewValue;
 
-	for (int32 TransformIndex = 0; TransformIndex < SolverParticleHandles.Num(); ++TransformIndex)
+	for (int32 ParticleIndex = 0; ParticleIndex < SolverParticleHandles.Num(); ++ParticleIndex)
 	{
-		if (Chaos::FPBDRigidClusteredParticleHandle* Handle = SolverParticleHandles[TransformIndex])
-		{
-			const Chaos::FReal NewM = Handle->M() * MaterialOverrideMassScaleMultiplierChange;
-			const Chaos::FVec3 NewI = Handle->I() * MaterialOverrideMassScaleMultiplierChange;
+		Chaos::FPBDRigidClusteredParticleHandle* Handle = SolverParticleHandles[ParticleIndex];
+		check(Handle != nullptr);
+		const Chaos::FReal NewM = Handle->M() * MaterialOverrideMassScaleMultiplierChange;
+		const Chaos::FVec3 NewI = Handle->I() * MaterialOverrideMassScaleMultiplierChange;
 
-			Handle->SetM(NewM);
-			Handle->SetI(NewI);
-			const Chaos::FReal InvM = (NewM > 0.0f) ? 1.0f / NewM : 0.0f;
-			const Chaos::FVec3 InvI = (NewM > 0.0f) ? Chaos::FVec3(NewI).Reciprocal() : Chaos::FVec3::ZeroVector;
-			Handle->SetInvM(InvM);
-			Handle->SetInvI(InvI);
-		}
+		Handle->SetM(NewM);
+		Handle->SetI(NewI);
+		const Chaos::FReal InvM = (NewM > 0.0f) ? 1.0f / NewM : 0.0f;
+		const Chaos::FVec3 InvI = (NewM > 0.0f) ? Chaos::FVec3(NewI).Reciprocal() : Chaos::FVec3::ZeroVector;
+		Handle->SetInvM(InvM);
+		Handle->SetInvI(InvI);
 	}
 }
 
@@ -3500,7 +3518,7 @@ void FGeometryCollectionPhysicsProxy::SetOneWayInteractionLevel_External(int32 I
 		TUniquePtr<FParticle>& Particle = GTParticles[ParticleIndex];
 		if (Particle.IsValid())
 		{
-			const bool bIsOneWayInteraction = (InOneWayInteractionLevel >= 0) && (Level[ParticleIndex] >= InOneWayInteractionLevel);
+			const bool bIsOneWayInteraction = (InOneWayInteractionLevel >= 0) && (Level[FromParticleToTransformIndex[ParticleIndex]] >= InOneWayInteractionLevel);
 			Particle->SetOneWayInteraction(bIsOneWayInteraction);
 		}
 	}
@@ -3727,7 +3745,7 @@ void FGeometryCollectionPhysicsProxy::BufferPhysicsResults_Internal(Chaos::FPBDR
 
 	UniqueIdxToInternalClusterHandle.Reset();
 	
-	const int32 NumTransforms = PhysicsThreadCollection.GetNumTransforms();
+	check(NumTransforms == PhysicsThreadCollection.GetNumTransforms());
 	if(NumTransforms > 0)
 	{ 
 		SCOPE_CYCLE_COUNTER(STAT_CalcParticleToWorld);
@@ -3741,7 +3759,7 @@ void FGeometryCollectionPhysicsProxy::BufferPhysicsResults_Internal(Chaos::FPBDR
 		// historically assume that the root has no parent which is not true any more).
 		if (Parameters.InitialRootIndex != INDEX_NONE)
 		{
-			const Chaos::FPBDRigidClusteredParticleHandle* RootHandle = SolverParticleHandles[Parameters.InitialRootIndex];
+			const Chaos::FPBDRigidClusteredParticleHandle* RootHandle = SolverParticleHandles[FromTransformToParticleIndex[Parameters.InitialRootIndex]];
 			if (RootHandle != nullptr)
 			{
 				const bool bIsRootBroken = RootHandle->Disabled() && (RootHandle->Parent() == nullptr);
@@ -3749,9 +3767,10 @@ void FGeometryCollectionPhysicsProxy::BufferPhysicsResults_Internal(Chaos::FPBDR
 			}
 		}
 
-		for (int32 TransformGroupIndex = 0; TransformGroupIndex < NumTransforms; ++TransformGroupIndex)
+		for (int32 ParticleIndex = 0; ParticleIndex < NumEffectiveParticles; ++ParticleIndex)
 		{
-			Chaos::FPBDRigidClusteredParticleHandle* Handle = SolverParticleHandles[TransformGroupIndex];
+			const int32 TransformGroupIndex = FromParticleToTransformIndex[ParticleIndex];
+			Chaos::FPBDRigidClusteredParticleHandle* Handle = SolverParticleHandles[ParticleIndex];
 			if (!Handle)
 			{
 				PhysicsThreadCollection.Active[TransformGroupIndex] = false;
@@ -3798,7 +3817,7 @@ void FGeometryCollectionPhysicsProxy::BufferPhysicsResults_Internal(Chaos::FPBDR
 			Chaos::FPBDRigidClusteredParticleHandle* ClusterParent = Handle->Parent();
 
 			// Has parent changed? ( can be a new one or nullptr)
-			const bool bHasNewParent = (SolverClusterID[TransformGroupIndex] != ClusterParent) || (!ClusterParent && PhysicsThreadCollection.GetHasParent(TransformGroupIndex));
+			const bool bHasNewParent = (SolverClusterID[ParticleIndex] != ClusterParent) || (!ClusterParent && PhysicsThreadCollection.GetHasParent(TransformGroupIndex));
 			if (bHasNewParent)
 			{
 				// Force all driven rigid bodies out of the transform hierarchy ( because the new parent can only be an internal cluster )
@@ -3815,7 +3834,7 @@ void FGeometryCollectionPhysicsProxy::BufferPhysicsResults_Internal(Chaos::FPBDR
 					bHasChanged = true;
 				}
 				// set new parent ( likely an internal cluster )
-				SolverClusterID[TransformGroupIndex] = ClusterParent;
+				SolverClusterID[ParticleIndex] = ClusterParent;
 			}
 
 			// do we have an internal cluster parent ( it also implies we are a disabled particle )
@@ -4108,7 +4127,7 @@ bool FGeometryCollectionPhysicsProxy::PullNonInterpolatableDataFromSinglePhysics
 			continue;
 		}
 
-		
+		const int32 ParticleIndex = FromTransformToParticleIndex[TransformGroupIndex];
 		const bool bIsActive = !StateData.State.DisabledState;
 		if (UpdateValue(GameThreadCollection.Active[TransformGroupIndex], bIsActive))
 		{
@@ -4116,19 +4135,20 @@ bool FGeometryCollectionPhysicsProxy::PullNonInterpolatableDataFromSinglePhysics
 			{
 				CreateChildrenGeometry_External();
 			}
-			if (GTParticles[TransformGroupIndex] == nullptr)
+			if (ParticleIndex == INDEX_NONE || GTParticles[ParticleIndex] == nullptr)
 			{
 				continue;
 			}
-			GTParticles[TransformGroupIndex]->SetDisabled(!bIsActive);
+
+			GTParticles[ParticleIndex]->SetDisabled(!bIsActive);
 			bIsCollectionDirty = true;
 		}
 
-		if (GTParticles[TransformGroupIndex] == nullptr)
+		if (ParticleIndex == INDEX_NONE || GTParticles[ParticleIndex] == nullptr)
 		{
 			continue;
 		}
-		FParticle& GTParticle = *GTParticles[TransformGroupIndex];
+		FParticle& GTParticle = *GTParticles[ParticleIndex];
 
 		if (UpdateValue(GameThreadCollection.DynamicState[TransformGroupIndex], static_cast<uint8>(StateData.State.DynamicState)))
 		{
@@ -4183,7 +4203,7 @@ bool FGeometryCollectionPhysicsProxy::PullNonInterpolatableDataFromSinglePhysics
 		// internal cluster index map update
 		if (StateData.InternalClusterUniqueIdx > INDEX_NONE)
 		{
-			GTParticlesToInternalClusterUniqueIdx.Add(GTParticles[TransformGroupIndex].Get(), StateData.InternalClusterUniqueIdx);
+			GTParticlesToInternalClusterUniqueIdx.Add(GTParticles[ParticleIndex].Get(), StateData.InternalClusterUniqueIdx);
 			InternalClusterUniqueIdxToChildrenTransformIndices.FindOrAdd(StateData.InternalClusterUniqueIdx).Add(TransformGroupIndex);
 		}
 	}
@@ -4194,9 +4214,11 @@ bool FGeometryCollectionPhysicsProxy::PullNonInterpolatableDataFromSinglePhysics
 	if ((RootIndex != INDEX_NONE) && CurrentResults.IsRootBroken && GameThreadCollection.Active[RootIndex] && bGeometryCollectionUseRootBrokenFlag)
 	{
 		GameThreadCollection.Active[RootIndex] = false;
-		if (GTParticles[RootIndex] != nullptr)
+		const int32 ParticleIndex = FromTransformToParticleIndex[RootIndex];
+		check(ParticleIndex != INDEX_NONE);
+		if (GTParticles[ParticleIndex] != nullptr)
 		{
-			GTParticles[RootIndex]->SetDisabled(true);
+			GTParticles[ParticleIndex]->SetDisabled(true);
 			bIsCollectionDirty = true;
 		}
 	}
@@ -4235,7 +4257,7 @@ bool FGeometryCollectionPhysicsProxy::PullFromPhysicsState(const Chaos::FDirtyGe
 		InterpolationData.AccumlateErrorXR(Error->ErrorX, Error->ErrorR, SolverSyncTimestamp, RenderInterpErrorCorrectionDurationTicks);
 	}
 
-	const int32 NumTransforms = GameThreadCollection.GetNumTransforms();
+	check(NumTransforms == GameThreadCollection.GetNumTransforms());
 	const bool bNeedInterpolation = (NextPullData != nullptr);
 
 	bool bIsCollectionDirty = false;
@@ -4271,8 +4293,9 @@ bool FGeometryCollectionPhysicsProxy::PullFromPhysicsState(const Chaos::FDirtyGe
 			// for that case we cannot just go through the list of entries since Results and NextResults may have different number of entries that don't always match
 			// so we need to go through the transform indices and find the matching entries on both side 
 			// this is helped by the FResultAccessor structure
-			for (int32 TransformGroupIndex = 0; TransformGroupIndex < NumTransforms; ++TransformGroupIndex)
+			for (int32 ParticleIndex = 0; ParticleIndex < NumEffectiveParticles; ++ParticleIndex)
 			{
+				const int32 TransformGroupIndex = FromParticleToTransformIndex[ParticleIndex];
 				const bool bActive = GameThreadCollection.Active[TransformGroupIndex];
 				bAtLeatOneChildActive |= bActive && TransformGroupIndex != Parameters.InitialRootIndex;
 				if (bAtLeatOneChildActive)
@@ -4280,12 +4303,12 @@ bool FGeometryCollectionPhysicsProxy::PullFromPhysicsState(const Chaos::FDirtyGe
 					CreateChildrenGeometry_External();
 				}
 
-				if (GTParticles[TransformGroupIndex] == nullptr)
+				if (GTParticles[ParticleIndex] == nullptr)
 				{
 					continue;
 				}
 
-				FParticle& GTParticle = *GTParticles[TransformGroupIndex];
+				FParticle& GTParticle = *GTParticles[ParticleIndex];
 				const FResultInterpolator ResultInterpolator(TransformGroupIndex, PrevResults, NextResults, GTParticle, *Alpha);
 				if (ResultInterpolator.HasNoEntry())
 				{
@@ -4330,13 +4353,14 @@ bool FGeometryCollectionPhysicsProxy::PullFromPhysicsState(const Chaos::FDirtyGe
 			const bool bIsComponentTransformScaled = !WorldTransform_External.GetScale3D().Equals(FVector::OneVector);
 			const FTransform ComponentScaleTransform(FQuat::Identity, FVector::ZeroVector, WorldTransform_External.GetScale3D());
 
-			for (int32 TransformIndex = 0; TransformIndex < NumTransforms; TransformIndex++)
+			for (int32 ParticleIndex = 0; ParticleIndex < NumEffectiveParticles; ParticleIndex++)
 			{
+				const int32 TransformIndex = FromParticleToTransformIndex[ParticleIndex];
 				const bool bAnimatingWhileDisabled = AnimationsActive ? (*AnimationsActive)[TransformIndex] : false;
 				const bool bActive = GameThreadCollection.Active[TransformIndex];
 				if (bActive || bAnimatingWhileDisabled)
 				{
-					if (GTParticles[TransformIndex] != nullptr)
+					if (GTParticles[ParticleIndex] != nullptr)
 					{
 						bool bWasModified = PrevResultsModifiedIndices.IsValidIndex(TransformIndex) ? PrevResultsModifiedIndices[TransformIndex] : false;
 						if (NextPullData)
@@ -4346,7 +4370,7 @@ bool FGeometryCollectionPhysicsProxy::PullFromPhysicsState(const Chaos::FDirtyGe
 						}
 						if (bWasModified)
 						{
-							const FParticle& GTParticle = *GTParticles[TransformIndex];
+							const FParticle& GTParticle = *GTParticles[ParticleIndex];
 							const FTransform& ParticleMassToLocal = MassToLocal[TransformIndex];
 							const int32 ParentTransformIndex = GameThreadCollection.GetParent(TransformIndex);
 
@@ -4356,7 +4380,7 @@ bool FGeometryCollectionPhysicsProxy::PullFromPhysicsState(const Chaos::FDirtyGe
 							FTransform ParentWorldTransform = WorldTransform_External;
 							if (ParentTransformIndex != INDEX_NONE)
 							{
-								if (const FParticle* GTParentParticle = (ParentTransformIndex != INDEX_NONE) ? GTParticles[ParentTransformIndex].Get() : nullptr)
+								if (const FParticle* GTParentParticle = (ParentTransformIndex != INDEX_NONE) ? GTParticles[FromTransformToParticleIndex[ParentTransformIndex]].Get() : nullptr)
 								{
 									const FTransform& ParentMassToLocal = MassToLocal[TransformIndex];
 									ParentWorldTransform = ParentMassToLocal.Inverse() * FTransform { GTParentParticle->R(), GTParentParticle->X() };
@@ -4417,8 +4441,7 @@ void FGeometryCollectionPhysicsProxy::UpdateFilterData_External(const FCollision
 	check(IsInGameThread());
 
 	// SimFilter/QueryFilter members are read on both threads, these are const after initialization and are not updated here.
-	const int32 NumTransforms = GameThreadCollection.GetNumTransforms();
-	for (int32 Index = 0; Index < NumTransforms; ++Index)
+	for (int32 Index = 0; Index < NumEffectiveParticles; ++Index)
 	{
 		FParticle* P = GTParticles[Index].Get();
 		if (P != nullptr)
@@ -4452,32 +4475,30 @@ void FGeometryCollectionPhysicsProxy::SetFilterData_Internal(const FCollisionFil
 	{
 		FRigidClustering& RigidClustering = RigidSolver->GetEvolution()->GetRigidClustering();
 
-		const int32 NumTransforms = PhysicsThreadCollection.NumElements(FGeometryCollection::TransformGroup);
-		for (int32 TransformIndex = 0; TransformIndex < NumTransforms; ++TransformIndex)
+		for (int32 ParticleIndex = 0; ParticleIndex < NumEffectiveParticles; ++ParticleIndex)
 		{
-			if (Chaos::FPBDRigidClusteredParticleHandle* Handle = SolverParticleHandles[TransformIndex])
+			Chaos::FPBDRigidClusteredParticleHandle* Handle = SolverParticleHandles[ParticleIndex];
+			check(Handle != nullptr);
+			// Must update our filters before updating an internal cluster parent
+			const Chaos::FShapesArray& ShapesArray = Handle->ShapesArray();
+			for (const TUniquePtr<Chaos::FPerShapeData>& Shape : ShapesArray)
 			{
-				// Must update our filters before updating an internal cluster parent
-				const Chaos::FShapesArray& ShapesArray = Handle->ShapesArray();
-				for (const TUniquePtr<Chaos::FPerShapeData>& Shape : ShapesArray)
-				{
-					Shape->SetQueryData(Parameters.QueryFilterData);
-					Shape->SetSimData(Parameters.SimulationFilterData);
-				}
+				Shape->SetQueryData(Parameters.QueryFilterData);
+				Shape->SetSimData(Parameters.SimulationFilterData);
+			}
 
-				FClusterHandle* ParentHandle = Handle->Parent();
-				if (ParentHandle && ParentHandle->InternalCluster() && ParentHandle->PhysicsProxy() == this)
+			FClusterHandle* ParentHandle = Handle->Parent();
+			if (ParentHandle && ParentHandle->InternalCluster() && ParentHandle->PhysicsProxy() == this)
+			{
+				if (!ProcessedInternalClusters.Contains(ParentHandle))
 				{
-					if (!ProcessedInternalClusters.Contains(ParentHandle))
+					ProcessedInternalClusters.Add(ParentHandle);
+					// Must update our filters before updating an internal cluster parent
+					FRigidClustering::FRigidHandleArray* ChildArray = RigidClustering.GetChildrenMap().Find(ParentHandle);
+					if (ensure(ChildArray))
 					{
-						ProcessedInternalClusters.Add(ParentHandle);
-						// Must update our filters before updating an internal cluster parent
-						FRigidClustering::FRigidHandleArray* ChildArray = RigidClustering.GetChildrenMap().Find(ParentHandle);
-						if (ensure(ChildArray))
-						{
-							// If our filter changed, internal cluster parent's filter may be stale.
-							UpdateClusterFilterDataFromChildren(ParentHandle, *ChildArray);
-						}
+						// If our filter changed, internal cluster parent's filter may be stale.
+						UpdateClusterFilterDataFromChildren(ParentHandle, *ChildArray);
 					}
 				}
 			}
@@ -4507,21 +4528,22 @@ void FGeometryCollectionPhysicsProxy::SetPerParticleFilterData_Internal(const TA
 		{
 			if (Data.bIsValid && Data.ParticleIndex != INDEX_NONE)
 			{
-				if (Chaos::FPhysicsObjectHandle Object = PhysicsObjects[Data.ParticleIndex].Get())
+				const int32 ParticleIndex = FromTransformToParticleIndex[Data.ParticleIndex];
+				if (ParticleIndex != INDEX_NONE)
 				{
+					Chaos::FPhysicsObjectHandle Object = PhysicsObjects[ParticleIndex].Get();
+					check(Object != nullptr);
 					TArrayView<Chaos::FPhysicsObjectHandle> ParticleView{ &Object, 1 };
 					Interface.UpdateShapeCollisionFlags(ParticleView, Data.bSimEnabled, Data.bQueryEnabled);
 					Interface.UpdateShapeFilterData(ParticleView, Data.QueryFilter, Data.SimFilter);
-				}
-
-				// todo(chaos): It's not ideal but the geometry collection needs to request the cluster union update its
-				// cached shape data if the particle is in a cluster union. This is because we don't share shape
-				// data between the GC shapes and the cluster union shapes.
-				Chaos::FRigidClustering& Clustering = RBDSolver->GetEvolution()->GetRigidClustering();
-				Chaos::FClusterUnionManager& ClusterUnionManager = Clustering.GetClusterUnionManager();
-
-				if (Chaos::FPBDRigidClusteredParticleHandle* Particle = SolverParticleHandles[Data.ParticleIndex])
-				{
+					// todo(chaos): It's not ideal but the geometry collection needs to request the cluster union update its
+					// cached shape data if the particle is in a cluster union. This is because we don't share shape
+					// data between the GC shapes and the cluster union shapes.
+					Chaos::FRigidClustering& Clustering = RBDSolver->GetEvolution()->GetRigidClustering();
+					Chaos::FClusterUnionManager& ClusterUnionManager = Clustering.GetClusterUnionManager();
+
+					Chaos::FPBDRigidClusteredParticleHandle* Particle = SolverParticleHandles[ParticleIndex];
+					check(Particle != nullptr);
 					if (Chaos::FClusterUnion* ClusterUnion = ClusterUnionManager.FindClusterUnionFromParticle(Particle))
 					{
 						ClusterUnion->AddPendingGeometryOperation(Chaos::EClusterUnionGeometryOperation::Refresh, Particle);
@@ -5756,9 +5778,9 @@ void FGeometryCollectionPhysicsProxy::FieldForcesUpdateCallback(Chaos::FPBDRigid
 
 FGeometryCollectionPhysicsProxy::FClusterHandle* FGeometryCollectionPhysicsProxy::GetInitialRootHandle_Internal() const
 {
-	if (SolverParticleHandles.IsValidIndex(Parameters.InitialRootIndex))
+	if (SolverParticleHandles.IsValidIndex(FromTransformToParticleIndex[Parameters.InitialRootIndex]))
 	{
-		return SolverParticleHandles[Parameters.InitialRootIndex];
+		return SolverParticleHandles[FromTransformToParticleIndex[Parameters.InitialRootIndex]];
 	}
 	return nullptr;
 }
@@ -5775,31 +5797,40 @@ TArray<Chaos::FPhysicsObjectHandle> FGeometryCollectionPhysicsProxy::GetAllPhysi
 			Handles.Add(Object.Get());
 		}
 	}
-
 	return Handles;
 }
 
 TArray<Chaos::FPhysicsObjectHandle> FGeometryCollectionPhysicsProxy::GetAllPhysicsObjectIncludingNulls() const
 {
 	TArray<Chaos::FPhysicsObjectHandle> Handles;
-	Handles.Reserve(NumEffectiveParticles);
+	Handles.SetNum(NumTransforms);
 
-	for (const Chaos::FPhysicsObjectUniquePtr& Object : PhysicsObjects)
+	for (int32 TransformIndex = 0; TransformIndex < NumTransforms; ++TransformIndex)
 	{
-		Handles.Add(Object.Get());
+		const int32 ParticleIndex = FromTransformToParticleIndex[TransformIndex];
+		if (ParticleIndex == INDEX_NONE)
+		{
+			Handles[TransformIndex] = nullptr;
+		}
+		else
+		{
+			Handles[TransformIndex] = PhysicsObjects[ParticleIndex].Get();
+		}
 	}
-
 	return Handles;
 }
 
 Chaos::FPhysicsObjectHandle FGeometryCollectionPhysicsProxy::GetPhysicsObjectByIndex(int32 Index) const
 {
-	if (!PhysicsObjects.IsValidIndex(Index))
+	if (FromTransformToParticleIndex.IsValidIndex(Index))
 	{
-		return nullptr;
+		const int32 ParticleIndex = FromTransformToParticleIndex[Index];
+		if (ParticleIndex != INDEX_NONE)
+		{
+			return PhysicsObjects[ParticleIndex].Get();
+		}
 	}
-
-	return PhysicsObjects[Index].Get();
+	return nullptr; 
 }
 
 FGeometryCollectionPhysicsProxy::FParticle* FGeometryCollectionPhysicsProxy::GetParticleByIndex_External(int32 Index)
@@ -5809,11 +5840,15 @@ FGeometryCollectionPhysicsProxy::FParticle* FGeometryCollectionPhysicsProxy::Get
 
 const FGeometryCollectionPhysicsProxy::FParticle* FGeometryCollectionPhysicsProxy::GetParticleByIndex_External(int32 Index) const
 {
-	if (Index < 0 || Index >= GTParticles.Num())
+	if (FromTransformToParticleIndex.IsValidIndex(Index))
 	{
-		return nullptr;
+		const int32 ParticleIndex = FromTransformToParticleIndex[Index];
+		if (ParticleIndex != INDEX_NONE)
+		{
+			return GTParticles[ParticleIndex].Get();
+		}
 	}
-	return GTParticles[Index].Get();
+	return nullptr;
 }
 
 FGeometryCollectionPhysicsProxy::FParticleHandle* FGeometryCollectionPhysicsProxy::GetParticleByIndex_Internal(int32 Index)
@@ -5823,11 +5858,15 @@ FGeometryCollectionPhysicsProxy::FParticleHandle* FGeometryCollectionPhysicsProx
 
 const FGeometryCollectionPhysicsProxy::FParticleHandle* FGeometryCollectionPhysicsProxy::GetParticleByIndex_Internal(int32 Index) const
 {
-	if (!SolverParticleHandles.IsValidIndex(Index))
+	if (FromTransformToParticleIndex.IsValidIndex(Index))
 	{
-		return nullptr;
+		const int32 ParticleIndex = FromTransformToParticleIndex[Index];
+		if (ParticleIndex != INDEX_NONE)
+		{
+			return SolverParticleHandles[ParticleIndex];
+		}
 	}
-	return SolverParticleHandles[Index];
+	return nullptr;
 }
 
 void FDamageCollector::Reset(int32 NumTransforms)
diff --git a/Engine/Source/Runtime/Experimental/Chaos/Public/PhysicsProxy/GeometryCollectionPhysicsProxy.h b/Engine/Source/Runtime/Experimental/Chaos/Public/PhysicsProxy/GeometryCollectionPhysicsProxy.h
--- a/Engine/Source/Runtime/Experimental/Chaos/Public/PhysicsProxy/GeometryCollectionPhysicsProxy.h
+++ b/Engine/Source/Runtime/Experimental/Chaos/Public/PhysicsProxy/GeometryCollectionPhysicsProxy.h
@@ -19,6 +19,7 @@
 #include "Chaos/Defines.h"
 #include "Chaos/GeometryParticlesfwd.h"
 
+
 namespace Chaos
 {
 	template <typename T> class TSerializablePtr;
@@ -236,10 +237,22 @@ public:
 	void SetCollisionParticlesPerObjectFraction(float CollisionParticlesPerObjectFractionIn) 
 	{CollisionParticlesPerObjectFraction = CollisionParticlesPerObjectFractionIn;}
 
+	UE_DEPRECATED(5.4, "Use GetSolverClusterHandle_Internal instead")
 	TArray<FClusterHandle*>& GetSolverClusterHandles() { return SolverClusterHandles; }
 
+	UE_DEPRECATED(5.4, "Use GetParticle_Internal instead")
 	TArray<FClusterHandle*>& GetSolverParticleHandles() { return SolverParticleHandles; }
 
+	FClusterHandle* GetSolverClusterHandle_Internal(int32 Index) const
+	{
+		const int32 ParticleIndex = FromTransformToParticleIndex[Index];
+		if (ParticleIndex != INDEX_NONE)
+		{
+			return SolverClusterHandles[ParticleIndex];
+		}
+		return nullptr;
+	}
+
 	const FGeometryCollectionResults* GetConsumerResultsGT() const 
 	{ return PhysToGameInterchange.PeekConsumerBuffer(); }
 
@@ -277,11 +290,27 @@ public:
 	// Set whether the GC should be using collision from the Static Mesh or the GC itself for game thread traces ( this needs to be called on the game thread )
 	CHAOS_API void SetUseStaticMeshCollisionForTraces_External(bool bInUseStaticMeshCollisionForTraces);
 
+	UE_DEPRECATED(5.4, "Use GetParticle_Internal instead")
 	const TArray<FClusterHandle*> GetParticles() const
 	{
 		return SolverParticleHandles;
 	}
 
+	const TArray<FClusterHandle*>GetUnorderedParticles_Internal() const
+	{
+		return SolverParticleHandles;
+	}
+
+	FClusterHandle* GetParticle_Internal(int32 Index) const
+	{
+		const int32 ParticleIndex = FromTransformToParticleIndex[Index];
+		if (ParticleIndex != INDEX_NONE)
+		{
+			return SolverParticleHandles[ParticleIndex];
+		}
+		return nullptr;
+	}
+
 	const FSimulationParameters& GetSimParameters() const
 	{
 		return Parameters;
@@ -302,11 +331,17 @@ public:
 		return GameThreadCollection;
 	}
 
+	UE_DEPRECATED(5.4, "Use GetUnorderedParticles_External instead")
 	TArray<TUniquePtr<FParticle>>& GetExternalParticles()
 	{
 		return GTParticles;
 	}
 
+	TArray<TUniquePtr<FParticle>>& GetUnorderedParticles_External()
+	{
+		return GTParticles;
+	}
+
 	CHAOS_API FParticle* GetParticleByIndex_External(int32 Index);
 	CHAOS_API const FParticle* GetParticleByIndex_External(int32 Index) const;
 
@@ -391,7 +426,7 @@ public:
 		{
 			return FGeometryCollectionItemIndex::CreateInternalClusterItemIndex(*InternalClusterUniqueIdx);
 		}
-		// regular particle that has a matchig transform index 
+		// regular particle that has a matching transform index 
 		if (const int32* TransformGroupIndex = GTParticlesToTransformGroupIndex.Find(GTPParticle))
 		{
 			return FGeometryCollectionItemIndex::CreateTransformItemIndex(*TransformGroupIndex);
@@ -494,7 +529,10 @@ public:
 	CHAOS_API TArray<Chaos::FPhysicsObjectHandle> GetAllPhysicsObjects() const ;
 	CHAOS_API TArray<Chaos::FPhysicsObjectHandle> GetAllPhysicsObjectIncludingNulls() const;
 	CHAOS_API Chaos::FPhysicsObjectHandle GetPhysicsObjectByIndex(int32 Index) const;
-	int32 GetNumParticles() const { return NumParticles; }
+
+	UE_DEPRECATED(5.4, "Use GetNumTransforms instead")
+	int32 GetNumParticles() const { return NumTransforms; }
+	int32 GetNumTransforms() const { return NumTransforms; }
 
 	// todo(chaos): Remove this and move to a cook time approach of the SM data based on the GC property
 	using FCreateTraceCollisionGeometryCallback = TFunction<void(const FTransform& InToLocal, TArray<Chaos::FImplicitObjectPtr>& OutGeoms, Chaos::FShapesArray& OutShapes)>;
@@ -502,6 +540,8 @@ public:
 
 	CHAOS_API void CreateChildrenGeometry_Internal();
 
+	int32 GetFromParticleToTransformIndex(int32 Index) const { check(FromParticleToTransformIndex.IsValidIndex(Index));  return FromParticleToTransformIndex[Index]; }
+
 protected:
 
 	CHAOS_API float ComputeMaterialBasedDamageThreshold_Internal(int32 TransformIndex) const;
@@ -607,23 +647,24 @@ private:
 	//
 	//  Proxy State Information
 	//
-	int32 NumParticles;
+	int32 NumTransforms;
 	int32 NumEffectiveParticles;
 	int32 BaseParticleIndex;
 	TArray<FParticleHandle*> SolverClusterID;
-	TArray<FClusterHandle*> SolverClusterHandles; // make a TArray of the base clase with type
+	TArray<FClusterHandle*> SolverClusterHandles; // make a TArray of the base class with type
 	TArray<FClusterHandle*> SolverParticleHandles;// make a TArray of base class and join with above
-	TSet<FClusterHandle*> SolverAnchors;
 	TMap<FParticleHandle*, int32> HandleToTransformGroupIndex;
 	TMap<int32, FClusterHandle*> UniqueIdxToInternalClusterHandle;
 	TArray<Chaos::FUniqueIdx> UniqueIdxs;
+	TArray<int32> FromParticleToTransformIndex;
+	TArray<int32> FromTransformToParticleIndex;
 
 	//
 	// Buffer Results State Information
 	//
 	bool IsObjectDynamic; // Records current dynamic state
 	bool IsObjectLoading; // Indicate when loaded
-	bool IsObjectDeleting; // Indicatge when pending deletion
+	bool IsObjectDeleting; // Indicate when pending deletion
 
 	EReplicationMode ReplicationMode = EReplicationMode::Unknown;	
 
diff --git a/Engine/Source/Runtime/Experimental/ChaosVehicles/ChaosVehiclesCore/Private/SimModule/DeferredForcesModular.cpp b/Engine/Source/Runtime/Experimental/ChaosVehicles/ChaosVehiclesCore/Private/SimModule/DeferredForcesModular.cpp
--- a/Engine/Source/Runtime/Experimental/ChaosVehicles/ChaosVehiclesCore/Private/SimModule/DeferredForcesModular.cpp
+++ b/Engine/Source/Runtime/Experimental/ChaosVehicles/ChaosVehiclesCore/Private/SimModule/DeferredForcesModular.cpp
@@ -6,6 +6,8 @@
 #include "Chaos/Utilities.h"
 #include "Chaos/ParticleHandle.h"
 
+#include "PhysicsProxy/GeometryCollectionPhysicsProxy.h"
+
 #if CHAOS_DEBUG_DRAW
 #include "Chaos/DebugDrawQueue.h"
 #endif
@@ -36,8 +38,7 @@ Chaos::FPBDRigidParticleHandle* FDeferredForcesModular::GetParticleFromUniqueInd
 }
 
 // #TODO: passing all these parameters in is horrible, tidy this up now we know what needs to be done
-Chaos::FPBDRigidParticleHandle* FDeferredForcesModular::GetParticle(TArray<Chaos::FPBDRigidClusteredParticleHandle*>& Particles
-		, TArray<Chaos::FPBDRigidClusteredParticleHandle*>& ClusterParticles
+Chaos::FPBDRigidParticleHandle* FDeferredForcesModular::GetParticle(FGeometryCollectionPhysicsProxy* Proxy
 		, int TransformIndex
 		, int32 ParticleIdx
 		, const FVector& PositionalOffset
@@ -47,8 +48,8 @@ Chaos::FPBDRigidParticleHandle* FDeferredForcesModular::GetParticle(TArray<Chaos
 		, FTransform& TransformOut)
 {
 
-	Chaos::FPBDRigidParticleHandle* SimParticle = Particles[TransformIndex];
-	Chaos::FPBDRigidClusteredParticleHandle* SimClusterParticle = ClusterParticles[TransformIndex];
+	Chaos::FPBDRigidParticleHandle* SimParticle = Proxy->GetParticle_Internal(TransformIndex);
+	Chaos::FPBDRigidClusteredParticleHandle* SimClusterParticle = Proxy->GetSolverClusterHandle_Internal(TransformIndex);
 
 	if (SimParticle && !SimParticle->Disabled())
 	{
@@ -68,8 +69,7 @@ Chaos::FPBDRigidParticleHandle* FDeferredForcesModular::GetParticle(TArray<Chaos
 }
 
 // #TODO: passing all these parameters in is horrible, tidy this up now we know what needs to be done
-Chaos::FPBDRigidParticleHandle* FDeferredForcesModular::GetParticle(TArray<Chaos::FPBDRigidClusteredParticleHandle*>& Particles
-	, TArray<Chaos::FPBDRigidClusteredParticleHandle*>& ClusterParticles
+Chaos::FPBDRigidParticleHandle* FDeferredForcesModular::GetParticle(FGeometryCollectionPhysicsProxy* Proxy
 	, int TransformIndex
 	, int32 ParticleIdx
 	, const FVector& PositionalOffset
@@ -79,8 +79,8 @@ Chaos::FPBDRigidParticleHandle* FDeferredForcesModular::GetParticle(TArray<Chaos
 	, FTransform& TransformOut)
 {
 
-	Chaos::FPBDRigidParticleHandle* SimParticle = Particles[TransformIndex];
-	Chaos::FPBDRigidClusteredParticleHandle* SimClusterParticle = ClusterParticles[TransformIndex];
+	Chaos::FPBDRigidParticleHandle* SimParticle = Proxy->GetParticle_Internal(TransformIndex);
+	Chaos::FPBDRigidClusteredParticleHandle* SimClusterParticle = Proxy->GetSolverClusterHandle_Internal(TransformIndex);
 
 	if (SimParticle && !SimParticle->Disabled())
 	{
@@ -99,6 +99,48 @@ Chaos::FPBDRigidParticleHandle* FDeferredForcesModular::GetParticle(TArray<Chaos
 	return nullptr;
 }
 
+Chaos::FPBDRigidParticleHandle* FDeferredForcesModular::GetParticle(const FTransform& OffsetTransform
+	, FGeometryCollectionPhysicsProxy* Proxy
+	, int TransformIndex
+	, const FVector& PositionalOffset
+	, FTransform& TransformOut)
+{
+	using namespace Chaos;
+
+	const int SingleChassisIndex = 0;
+	TransformOut = FTransform::Identity;
+
+	Chaos::FPBDRigidClusteredParticleHandle* ClusterParticle = Proxy->GetSolverClusterHandle_Internal(SingleChassisIndex);
+	Chaos::FPBDRigidParticleHandle* Child = Proxy->GetParticle_Internal(TransformIndex);
+	FRigidTransform3 Frame = FRigidTransform3::Identity;
+
+	if (ClusterParticle != nullptr && Child != nullptr)
+	{
+		const FRigidTransform3 ClusterWorldTM(ClusterParticle->X(), ClusterParticle->R());
+		
+		Proxy->GetParticle_Internal(TransformIndex);
+
+		if (FPBDRigidClusteredParticleHandle* ClusterChild = Child->CastToClustered(); ClusterChild && ClusterChild->IsChildToParentLocked())
+		{
+			Frame = ClusterChild->ChildToParent();
+		}
+		else
+		{
+			//FTransform ChildT = Child->GetTransformXRCom();
+			//const FRigidTransform3 ChildWorldTM(ChildT.GetLocation(), ChildT.GetRotation());
+			const FRigidTransform3 ChildWorldTM(Child->X(), Child->R());
+			Frame = ChildWorldTM.GetRelativeTransform(ClusterWorldTM);
+			Frame.SetLocation(Frame.GetLocation() + OffsetTransform.TransformVector(PositionalOffset));
+
+			//UE_LOG(LogChaos, Warning, TEXT("ApplyForce Data TIdx %d %s"), TransformIndex, *Frame.GetTranslation().ToString());
+		}
+	}
+
+	// local transform
+	TransformOut.SetTranslation(Frame.GetTranslation());
+	return ClusterParticle;
+}
+
 Chaos::FPBDRigidParticleHandle* FDeferredForcesModular::GetParticle(const FTransform& OffsetTransform
 	, TArray<Chaos::FPBDRigidParticleHandle*>& Particles
 	, TArray<Chaos::FPBDRigidClusteredParticleHandle*>& ClusterParticles
@@ -137,8 +179,7 @@ Chaos::FPBDRigidParticleHandle* FDeferredForcesModular::GetParticle(const FTrans
 }
 
 template<typename TransformType>
-void FDeferredForcesModular::ApplyTemplate(TArray<Chaos::FPBDRigidClusteredParticleHandle*>& Particles
-		, TArray<Chaos::FPBDRigidClusteredParticleHandle*>& ClusterParticles
+void FDeferredForcesModular::ApplyTemplate(FGeometryCollectionPhysicsProxy* Proxy
 		, const TManagedArray<TransformType>& Transforms
 		, const TManagedArray<FTransform>& CollectionMassToLocal
 		, const TManagedArray<int32>& Parent)
@@ -148,7 +189,7 @@ void FDeferredForcesModular::ApplyTemplate(TArray<Chaos::FPBDRigidClusteredParti
 		FTransform RelativeTransform;
 		for (const FApplyForceData& Data : ApplyForceDatas)
 		{
-			Chaos::FPBDRigidParticleHandle* RigidHandle = GetParticle(Particles, ClusterParticles, Data.TransformIndex, Data.ParticleIdx, FVector::ZeroVector, FTransform(Transforms[Data.TransformIndex]), CollectionMassToLocal, Parent, RelativeTransform);
+			Chaos::FPBDRigidParticleHandle* RigidHandle = GetParticle(Proxy, Data.TransformIndex, Data.ParticleIdx, FVector::ZeroVector, FTransform(Transforms[Data.TransformIndex]), CollectionMassToLocal, Parent, RelativeTransform);
 			if (RigidHandle)
 			{
 				AddForce(RigidHandle, Data, RelativeTransform);
@@ -157,7 +198,7 @@ void FDeferredForcesModular::ApplyTemplate(TArray<Chaos::FPBDRigidClusteredParti
 
 		for (const FApplyForceAtPositionData& Data : ApplyForceAtPositionDatas)
 		{
-			Chaos::FPBDRigidParticleHandle* RigidHandle = GetParticle(Particles, ClusterParticles, Data.TransformIndex, Data.ParticleIdx, Data.Position, FTransform(Transforms[Data.TransformIndex]), CollectionMassToLocal, Parent, RelativeTransform);
+			Chaos::FPBDRigidParticleHandle* RigidHandle = GetParticle(Proxy, Data.TransformIndex, Data.ParticleIdx, Data.Position, FTransform(Transforms[Data.TransformIndex]), CollectionMassToLocal, Parent, RelativeTransform);
 			if (RigidHandle)
 			{
 				AddForceAtPosition(RigidHandle, Data, RelativeTransform);
@@ -166,7 +207,7 @@ void FDeferredForcesModular::ApplyTemplate(TArray<Chaos::FPBDRigidClusteredParti
 
 		for (const FAddTorqueInRadiansData& Data : ApplyTorqueDatas)
 		{
-			Chaos::FPBDRigidParticleHandle* RigidHandle = GetParticle(Particles, ClusterParticles, Data.TransformIndex, Data.ParticleIdx, FVector::ZeroVector, FTransform(Transforms[Data.TransformIndex]), CollectionMassToLocal, Parent, RelativeTransform);
+			Chaos::FPBDRigidParticleHandle* RigidHandle = GetParticle(Proxy, Data.TransformIndex, Data.ParticleIdx, FVector::ZeroVector, FTransform(Transforms[Data.TransformIndex]), CollectionMassToLocal, Parent, RelativeTransform);
 			if (RigidHandle)
 			{
 				AddTorque(RigidHandle, Data, RelativeTransform);
@@ -179,25 +220,61 @@ void FDeferredForcesModular::ApplyTemplate(TArray<Chaos::FPBDRigidClusteredParti
 	ApplyTorqueDatas.Empty();
 }
 
-void FDeferredForcesModular::Apply(TArray<Chaos::FPBDRigidClusteredParticleHandle*>& Particles
-	, TArray<Chaos::FPBDRigidClusteredParticleHandle*>& ClusterParticles
+void FDeferredForcesModular::Apply(FGeometryCollectionPhysicsProxy* Proxy
 	, const TManagedArray<FTransform>& Transforms
 	, const TManagedArray<FTransform>& CollectionMassToLocal
 	, const TManagedArray<int32>& Parent)
 {
-	ApplyTemplate(Particles, ClusterParticles, Transforms, CollectionMassToLocal, Parent);
+	ApplyTemplate(Proxy, Transforms, CollectionMassToLocal, Parent);
 }
 
-void FDeferredForcesModular::Apply(TArray<Chaos::FPBDRigidClusteredParticleHandle*>& Particles
-	, TArray<Chaos::FPBDRigidClusteredParticleHandle*>& ClusterParticles
+void FDeferredForcesModular::Apply(FGeometryCollectionPhysicsProxy* Proxy
 	, const TManagedArray<FTransform3f>& Transforms
 	, const TManagedArray<FTransform>& CollectionMassToLocal
 	, const TManagedArray<int32>& Parent)
 {
-	ApplyTemplate(Particles, ClusterParticles, Transforms, CollectionMassToLocal, Parent);
+	ApplyTemplate(Proxy, Transforms, CollectionMassToLocal, Parent);
 }
 
 
+void FDeferredForcesModular::Apply(FGeometryCollectionPhysicsProxy* Proxy)
+{
+	if (!GCoreModularVehicleDebugParams.DisableForces)
+	{
+		FTransform RelativeTransform;
+		for (const FApplyForceData& Data : ApplyForceDatas)
+		{
+			Chaos::FPBDRigidParticleHandle* RigidHandle = GetParticle(Data.OffsetTransform, Proxy, Data.TransformIndex, FVector::ZeroVector, RelativeTransform);
+			if (RigidHandle)
+			{
+				AddForce(RigidHandle, Data, RelativeTransform);
+			}
+		}
+
+		for (const FApplyForceAtPositionData& Data : ApplyForceAtPositionDatas)
+		{
+			Chaos::FPBDRigidParticleHandle* RigidHandle = GetParticle(Data.OffsetTransform, Proxy, Data.TransformIndex, Data.Position, RelativeTransform);
+			if (RigidHandle)
+			{
+				AddForceAtPosition(RigidHandle, Data, RelativeTransform);
+			}
+		}
+
+		for (const FAddTorqueInRadiansData& Data : ApplyTorqueDatas)
+		{
+			Chaos::FPBDRigidParticleHandle* RigidHandle = GetParticle(Data.OffsetTransform, Proxy, Data.TransformIndex, FVector::ZeroVector, RelativeTransform);
+			if (RigidHandle)
+			{
+				AddTorque(RigidHandle, Data, RelativeTransform);
+			}
+		}
+	}
+
+	ApplyForceDatas.Empty();
+	ApplyForceAtPositionDatas.Empty();
+	ApplyTorqueDatas.Empty();
+}
+
 void FDeferredForcesModular::Apply(TArray<Chaos::FPBDRigidParticleHandle*>& Particles
 	, TArray<Chaos::FPBDRigidClusteredParticleHandle*>& ClusterParticles)
 {
diff --git a/Engine/Source/Runtime/Experimental/ChaosVehicles/ChaosVehiclesCore/Private/SimModule/SimModuleTree.cpp b/Engine/Source/Runtime/Experimental/ChaosVehicles/ChaosVehiclesCore/Private/SimModule/SimModuleTree.cpp
--- a/Engine/Source/Runtime/Experimental/ChaosVehicles/ChaosVehiclesCore/Private/SimModule/SimModuleTree.cpp
+++ b/Engine/Source/Runtime/Experimental/ChaosVehicles/ChaosVehiclesCore/Private/SimModule/SimModuleTree.cpp
@@ -315,15 +315,12 @@ void FSimModuleTree::UpdateModuleVelocites(FGeometryCollectionPhysicsProxy* Phys
 {
 	check(PhysicsProxy);
 
-	const TArray<Chaos::FPBDRigidClusteredParticleHandle*>& Clusters = PhysicsProxy->GetSolverClusterHandles();
-	const TArray<Chaos::FPBDRigidClusteredParticleHandle*>& Particles = PhysicsProxy->GetSolverParticleHandles();
-
 	// capture the velocities at the start of each sim iteration
 	for (int i = 0; i < SimulationModuleTree.Num(); i++)
 	{
 		if (ISimulationModuleBase* Module = SimulationModuleTree[i].SimModule)
 		{
-			if (Chaos::FPBDRigidClusteredParticleHandle* ParentParticle = Clusters[Module->GetTransformIndex()])
+			if (Chaos::FPBDRigidClusteredParticleHandle* ParentParticle = PhysicsProxy->GetSolverClusterHandle_Internal(Module->GetTransformIndex()))
 			{
 				const FTransform BodyTransform(ParentParticle->R(), ParentParticle->X());
 
@@ -332,11 +329,11 @@ void FSimModuleTree::UpdateModuleVelocites(FGeometryCollectionPhysicsProxy* Phys
 					Chaos::FPBDRigidClusteredParticleHandle* Particle = nullptr;
 					if (Module->IsClustered())
 					{
-						Particle = Clusters[Module->GetTransformIndex()];
+						Particle = PhysicsProxy->GetSolverClusterHandle_Internal(Module->GetTransformIndex());
 					}
 					else
 					{
-						Particle = Particles[Module->GetTransformIndex()];				
+						Particle = PhysicsProxy->GetParticle_Internal(Module->GetTransformIndex());
 					}
 
 					if (Particle)
diff --git a/Engine/Source/Runtime/Experimental/ChaosVehicles/ChaosVehiclesCore/Public/SimModule/DeferredForcesModular.h b/Engine/Source/Runtime/Experimental/ChaosVehicles/ChaosVehiclesCore/Public/SimModule/DeferredForcesModular.h
--- a/Engine/Source/Runtime/Experimental/ChaosVehicles/ChaosVehiclesCore/Public/SimModule/DeferredForcesModular.h
+++ b/Engine/Source/Runtime/Experimental/ChaosVehicles/ChaosVehiclesCore/Public/SimModule/DeferredForcesModular.h
@@ -20,6 +20,8 @@ enum class EForceFlags : uint32
 ENUM_CLASS_FLAGS(EForceFlags);	
 
 
+class FGeometryCollectionPhysicsProxy;
+
 class CHAOSVEHICLESCORE_API FDeferredForcesModular
 {
 public:
@@ -159,8 +161,7 @@ public:
 
 	Chaos::FPBDRigidParticleHandle* GetParticleFromUniqueIndex(int32 ParticleUniqueIdx, const TArray<Chaos::FPBDRigidParticleHandle*>& Particles) const;
 
-	Chaos::FPBDRigidParticleHandle* GetParticle(TArray<Chaos::FPBDRigidClusteredParticleHandle*>& Particles
-			, TArray<Chaos::FPBDRigidClusteredParticleHandle*>& ClusterParticles
+	Chaos::FPBDRigidParticleHandle* GetParticle(FGeometryCollectionPhysicsProxy* Proxy
 			, int TransformIndex
 			, int32 ParticleIdx
 			, const FVector& PositionalOffset
@@ -169,8 +170,7 @@ public:
 			, const TManagedArray<int32>& Parent
 			, FTransform& TransformOut);
 
-	Chaos::FPBDRigidParticleHandle* GetParticle(TArray<Chaos::FPBDRigidClusteredParticleHandle*>& Particles
-		, TArray<Chaos::FPBDRigidClusteredParticleHandle*>& ClusterParticles
+	Chaos::FPBDRigidParticleHandle* GetParticle(FGeometryCollectionPhysicsProxy* Proxy
 		, int TransformIndex
 		, int32 ParticleIdx
 		, const FVector& PositionalOffset
@@ -180,20 +180,17 @@ public:
 		, FTransform& TransformOut);
 
 	Chaos::FPBDRigidParticleHandle* GetParticle(const FTransform& OffsetTransform
-		, TArray<Chaos::FPBDRigidParticleHandle*>& Particles
-		, TArray<Chaos::FPBDRigidClusteredParticleHandle*>& ClusterParticles
+		, FGeometryCollectionPhysicsProxy* Proxy
 		, int32 ParticleIdx
 		, const FVector& PositionalOffset
 		, FTransform& TransformOut);
 
-	void Apply(TArray<Chaos::FPBDRigidClusteredParticleHandle*>& Particles
-		, TArray<Chaos::FPBDRigidClusteredParticleHandle*>& ClusterParticles
+	void Apply(FGeometryCollectionPhysicsProxy* Proxy
 		, const TManagedArray<FTransform>& Transforms
 		, const TManagedArray<FTransform>& CollectionMassToLocal
 		, const TManagedArray<int32>& Parent);
 
-	void Apply(TArray<Chaos::FPBDRigidClusteredParticleHandle*>& Particles
-		, TArray<Chaos::FPBDRigidClusteredParticleHandle*>& ClusterParticles
+	void Apply(FGeometryCollectionPhysicsProxy* Proxy
 		, const TManagedArray<FTransform3f>& Transforms
 		, const TManagedArray<FTransform>& CollectionMassToLocal
 		, const TManagedArray<int32>& Parent);
@@ -204,8 +201,7 @@ public:
 private:
 
 	template<typename TransformType>
-	void ApplyTemplate(TArray<Chaos::FPBDRigidClusteredParticleHandle*>& Particles
-		, TArray<Chaos::FPBDRigidClusteredParticleHandle*>& ClusterParticles
+	void ApplyTemplate(FGeometryCollectionPhysicsProxy* Proxy
 		, const TManagedArray<TransformType>& Transforms
 		, const TManagedArray<FTransform>& CollectionMassToLocal
 		, const TManagedArray<int32>& Parent);
diff --git a/Engine/Source/Runtime/Experimental/GeometryCollectionEngine/Private/GeometryCollection/GeometryCollectionComponent.cpp b/Engine/Source/Runtime/Experimental/GeometryCollectionEngine/Private/GeometryCollection/GeometryCollectionComponent.cpp
--- a/Engine/Source/Runtime/Experimental/GeometryCollectionEngine/Private/GeometryCollection/GeometryCollectionComponent.cpp
+++ b/Engine/Source/Runtime/Experimental/GeometryCollectionEngine/Private/GeometryCollection/GeometryCollectionComponent.cpp
@@ -2108,12 +2108,11 @@ void UGeometryCollectionComponent::UpdateRepData()
 		const FRigidClustering& RigidClustering = Solver->GetEvolution()->GetRigidClustering();
 
 		const TManagedArrayAccessor<int32> InitialLevels = PhysicsProxy->GetPhysicsCollection().GetInitialLevels();
-		const TArray<FPBDRigidClusteredParticleHandle*>& ParticleHandles = PhysicsProxy->GetParticles();
+		const TArray<FPBDRigidClusteredParticleHandle*>& UnorderedParticleHandles = PhysicsProxy->GetUnorderedParticles_Internal();
 
 		// Replicate the anchored state of the root particle
 		LocalRepData.bIsRootAnchored = RepData.bIsRootAnchored;
-		const int32 InitialRootIndex = PhysicsProxy->GetSimParameters().InitialRootIndex;
-		if (FPBDRigidClusteredParticleHandle* RootHandle = ParticleHandles[InitialRootIndex])
+		if (FPBDRigidClusteredParticleHandle* RootHandle = PhysicsProxy->GetInitialRootHandle_Internal())
 		{
 			LocalRepData.bIsRootAnchored = RootHandle->IsAnchored();
 		}
@@ -2123,7 +2122,7 @@ void UGeometryCollectionComponent::UpdateRepData()
 
 		// It feels weird to start from the previous one. Need to be checked
 		LocalRepData.OneOffActivated = RepData.OneOffActivated;
-		for (FPBDRigidClusteredParticleHandle* Particle : ParticleHandles)
+		for (FPBDRigidClusteredParticleHandle* Particle : UnorderedParticleHandles)
 		{
 			// Particle can be null if we have embedded geometry 
 			if (Particle)
@@ -2182,7 +2181,7 @@ void UGeometryCollectionComponent::UpdateRepData()
 						//a one off so record it
 						ensureMsgf(TransformGroupIdx >= 0, TEXT("Non-internal cluster should always have a group index"));
 						ensureMsgf(TransformGroupIdx < TNumericLimits<uint16>::Max(), TEXT("Trying to replicate GC with more than 65k pieces. We assumed uint16 would suffice"));
-						ensureMsgf(PhysicsProxy->GetParticles().IsValidIndex(TransformGroupIdx) && PhysicsProxy->GetParticles()[TransformGroupIdx] != nullptr, TEXT("Invalid particle index being replicated - Contact physics team"));
+						ensureMsgf(PhysicsProxy->GetParticle_Internal(TransformGroupIdx) != nullptr, TEXT("Invalid particle index being replicated - Contact physics team"));
 
 						// Because we cull ClustersToRep with abandoned level, we must make sure we don't add duplicates to one off activated.
 						// TODO: avoid search for entry for perf
@@ -2305,8 +2304,8 @@ void UGeometryCollectionComponent::UpdateRepStateAndDynamicData()
 			FPBDRigidsSolver* Solver = PhysicsProxy->GetSolver<Chaos::FPBDRigidsSolver>();
 
 			// let go through all the transform and see which one has changed its state
-			const TArray<FPBDRigidClusteredParticleHandle*>& ParticleHandles = PhysicsProxy->GetParticles();
-			const int32 NumTransforms = ParticleHandles.Num();
+			const TArray<FPBDRigidClusteredParticleHandle*>& UnorderedParticleHandles = PhysicsProxy->GetUnorderedParticles_Internal();
+			const int32 NumParticles = UnorderedParticleHandles.Num();
 
 			const TManagedArrayAccessor<int32> InitialLevels = PhysicsProxy->GetPhysicsCollection().GetInitialLevels();
 			
@@ -2319,7 +2318,7 @@ void UGeometryCollectionComponent::UpdateRepStateAndDynamicData()
 			// root level anchor state
 			LocalRepStateData.bIsRootAnchored = RepStateData.bIsRootAnchored;
 			const int32 InitialRootIndex = PhysicsProxy->GetSimParameters().InitialRootIndex;
-			if (FPBDRigidClusteredParticleHandle* RootHandle = ParticleHandles[InitialRootIndex])
+			if (FPBDRigidClusteredParticleHandle* RootHandle = UnorderedParticleHandles[InitialRootIndex])
 			{
 				LocalRepStateData.bIsRootAnchored = RootHandle->IsAnchored() ? 1 : 0;
 			}
@@ -2340,11 +2339,14 @@ void UGeometryCollectionComponent::UpdateRepStateAndDynamicData()
 			const int32 RootIndex = GetRootIndex();
 
 			// go through particles and send replicated data 
+			const int32 NumTransforms = PhysicsProxy->GetNumTransforms();
 			LocalRepStateData.BrokenState.SetNum(NumTransforms, false);
 			LocalRepStateData.ReleasedData.Reserve(NumTransforms);
-			for (int32 TransformIndex = 0; TransformIndex < NumTransforms; TransformIndex++)
+			for (int32 ParticleIndex = 0; ParticleIndex < NumParticles; ParticleIndex++)
 			{
-				if (FPBDRigidClusteredParticleHandle* ParticleHandle = ParticleHandles[TransformIndex])
+				const int32 TransformIndex = PhysicsProxy->GetFromParticleToTransformIndex(ParticleIndex);
+				FPBDRigidClusteredParticleHandle* ParticleHandle = UnorderedParticleHandles[ParticleIndex];
+				check(ParticleHandle);
 				{
 					const int32 Level = (InitialLevels.IsValid() && InitialLevels.IsValidIndex(TransformIndex)) ? InitialLevels[TransformIndex] : INDEX_NONE;
 
@@ -2518,8 +2520,7 @@ namespace
 		using namespace Chaos;
 
 		const int32 InitialRootIndex = PhysicsProxy.GetSimParameters().InitialRootIndex;
-		const TArray<FPBDRigidClusteredParticleHandle*>& ParticleHandles = PhysicsProxy.GetParticles();
-		if (FPBDRigidClusteredParticleHandle* RootHandle = ParticleHandles[InitialRootIndex])
+		if (FPBDRigidClusteredParticleHandle* RootHandle = PhysicsProxy.GetParticle_Internal(InitialRootIndex))
 		{
 			if (RootHandle->Parent() || !RootHandle->Disabled())
 			{
@@ -2641,7 +2642,7 @@ namespace
 		{
 			for (const FGeometryCollectionClusterRep& RepCluster : RepData.Clusters)
 			{
-				if (Chaos::FPBDRigidParticleHandle* Cluster = PhysicsProxy->GetParticles()[RepCluster.ClusterIdx])
+				if (Chaos::FPBDRigidParticleHandle* Cluster = PhysicsProxy->GetParticle_Internal(RepCluster.ClusterIdx))
 				{
 					if (RepCluster.ClusterState.IsInternalCluster())
 					{
@@ -2681,14 +2682,14 @@ namespace
 			const FGeometryCollectionActivatedCluster& ActivatedCluster = RepData.OneOffActivated[OneOffActivatedProcessed];
 
 #if !UE_BUILD_SHIPPING
-			if(!PhysicsProxy->GetParticles().IsValidIndex(ActivatedCluster.ActivatedIndex))
+			if(PhysicsProxy->GetParticle_Internal(ActivatedCluster.ActivatedIndex) == nullptr)
 			{
 				ensureMsgf(false, TEXT("Invalid activated cluster index processing replication data."));
 				continue;
 			}
 #endif
 
-			FPBDRigidParticleHandle* OneOff = PhysicsProxy->GetParticles()[ActivatedCluster.ActivatedIndex];
+			FPBDRigidParticleHandle* OneOff = PhysicsProxy->GetParticle_Internal(ActivatedCluster.ActivatedIndex);
 
 			if (ensure(OneOff))
 			{
@@ -2739,52 +2740,52 @@ namespace
 			return false;
 		}
 
-		const TArray<FPBDRigidClusteredParticleHandle*>& ParticleHandles = PhysicsProxy->GetParticles();
+		const TArray<FPBDRigidClusteredParticleHandle*>& UnorderedParticleHandle = PhysicsProxy->GetUnorderedParticles_Internal();
 
 		// Update the anchored state of the root particle
 		UpdateRootStateFromReplication(*PhysicsProxy, (bool)RepStateData.bIsRootAnchored);
 		
 		
 		// now sync the broken state of the particles
-		if (RepStateData.BrokenState.Num() == ParticleHandles.Num())
+		if (RepStateData.BrokenState.Num() == PhysicsProxy->GetNumTransforms())
 		{
 			FPBDRigidsSolver* Solver = PhysicsProxy->GetSolver<Chaos::FPBDRigidsSolver>();
 			FRigidClustering& RigidClustering = Solver->GetEvolution()->GetRigidClustering();
 
-			for (int32 TransformIndex = 0; TransformIndex < ParticleHandles.Num(); TransformIndex++)
+			for (int32 ParticleIndex = 0; ParticleIndex < UnorderedParticleHandle.Num(); ParticleIndex++)
 			{
-				if (FPBDRigidClusteredParticleHandle* ParticleHandle = ParticleHandles[TransformIndex])
+				FPBDRigidClusteredParticleHandle* ParticleHandle = UnorderedParticleHandle[ParticleIndex];
+				check(ParticleHandle != nullptr);
+				const int32 TransformIndex = PhysicsProxy->GetFromParticleToTransformIndex(ParticleIndex);
+				// we only check for when thing goes from unbroken to broken 
+				const bool bWasBroken = (ParticleHandle->Parent() == nullptr);
+				const bool bIsBroken = RepStateData.BrokenState[TransformIndex];
+				if (!bWasBroken && bIsBroken)
 				{
-					// we only check for when thing goes from unbroken to broken 
-					const bool bWasBroken = (ParticleHandle->Parent() == nullptr);
-					const bool bIsBroken = RepStateData.BrokenState[TransformIndex];
-					if (!bWasBroken && bIsBroken)
-					{
-						RigidClustering.ForceReleaseChildParticleAndParents(ParticleHandle, /* bTriggerBreakEvents */true);
+					RigidClustering.ForceReleaseChildParticleAndParents(ParticleHandle, /* bTriggerBreakEvents */true);
 
-						if (!ParticleHandle->Disabled())
-						{
-							// check if we have a release velocity to be applied 
-							const FGeometryCollectionRepStateData::FReleasedData* ReleaseData = RepStateData.ReleasedData.FindByPredicate(
-								[TransformIndex](const FGeometryCollectionRepStateData::FReleasedData& Data)
-								{
-									return Data.TransformIndex == TransformIndex;
-								});
-							if (ReleaseData)
+					if (!ParticleHandle->Disabled())
+					{
+						// check if we have a release velocity to be applied 
+						const FGeometryCollectionRepStateData::FReleasedData* ReleaseData = RepStateData.ReleasedData.FindByPredicate(
+							[TransformIndex](const FGeometryCollectionRepStateData::FReleasedData& Data)
 							{
-								ParticleHandle->SetV(ReleaseData->LinearVelocity);
-								ParticleHandle->SetW(FMath::DegreesToRadians(ReleaseData->AngularVelocityInDegreesPerSecond));
-							}
-							else
+								return Data.TransformIndex == TransformIndex;
+							});
+						if (ReleaseData)
+						{
+							ParticleHandle->SetV(ReleaseData->LinearVelocity);
+							ParticleHandle->SetW(FMath::DegreesToRadians(ReleaseData->AngularVelocityInDegreesPerSecond));
+						}
+						else
+						{
+							// if we are a leaf we can disable this particle right away if there's no release data
+							// because this means the server has already disabled the particle from a removal 
+							if (ParticleHandle->ClusterIds().NumChildren == 0)
 							{
-								// if we are a leaf we can disable this particle right away if there's no release data
-								// becuase this means the server has already disabled the particle from a removal 
-								if (ParticleHandle->ClusterIds().NumChildren == 0)
-								{
-									// todo we could probably disable all at once later to have only one call 
-									Solver->GetEvolution()->DisableParticle(ParticleHandle);
-									Solver->GetParticles().MarkTransientDirtyParticle(ParticleHandle);
-								}
+								// todo we could probably disable all at once later to have only one call 
+								Solver->GetEvolution()->DisableParticle(ParticleHandle);
+								Solver->GetParticles().MarkTransientDirtyParticle(ParticleHandle);
 							}
 						}
 					}
@@ -2838,11 +2839,9 @@ namespace
 		if (bReplicateMovement)
 		{
 			FPBDRigidsSolver* Solver = PhysicsProxy->GetSolver<Chaos::FPBDRigidsSolver>();
-			const TArray<FPBDRigidClusteredParticleHandle*>& ParticleHandles = PhysicsProxy->GetParticles();
-
 			for (const FGeometryCollectionRepDynamicData::FClusterData& RepCluster : RepDynamicData.ClusterData)
 			{
-				if (Chaos::FPBDRigidParticleHandle* Cluster = ParticleHandles[RepCluster.TransformIndex])
+				if (Chaos::FPBDRigidParticleHandle* Cluster = PhysicsProxy->GetParticle_Internal(RepCluster.TransformIndex))
 				{
 					if (RepCluster.bIsInternalCluster)
 					{
@@ -2901,7 +2900,7 @@ bool UGeometryCollectionComponent::ProcessRepData(const float DeltaTime, const f
 		{
 			for (const FGeometryCollectionClusterRep& RepCluster : RepData.Clusters)
 			{
-				if (Chaos::FPBDRigidParticleHandle* Cluster = PhysicsProxy->GetParticles()[RepCluster.ClusterIdx])
+				if (Chaos::FPBDRigidParticleHandle* Cluster = PhysicsProxy->GetParticle_Internal(RepCluster.ClusterIdx))
 				{
 					if (RepCluster.ClusterState.IsInternalCluster())
 					{
@@ -3925,15 +3924,13 @@ void UGeometryCollectionComponent::RegisterAndInitializePhysicsProxy()
 						// As we're not in control we make it so our simulated proxy cannot break clusters
 						// We have to set the strain to a high value but be below the max for the data type
 						// so releasing on authority demand works
-						for (Chaos::FPBDRigidClusteredParticleHandle* ParticleHandle : Proxy->GetParticles())
+						for (Chaos::FPBDRigidClusteredParticleHandle* ParticleHandle : Proxy->GetUnorderedParticles_Internal())
 						{
-							if (ParticleHandle)
+							check(ParticleHandle);
+							const int32 Level = EnableAbandonAfterLevel ? ComputeParticleLevel(ParticleHandle) : -1;
+							if (Level <= AbandonAfterLevel)	//we only replicate up until level X, but it means we should replicate the breaking event of level X+1 (but not X+1's positions)
 							{
-								const int32 Level = EnableAbandonAfterLevel ? ComputeParticleLevel(ParticleHandle) : -1;
-								if (Level <= AbandonAfterLevel)	//we only replicate up until level X, but it means we should replicate the breaking event of level X+1 (but not X+1's positions)
-								{
-									ParticleHandle->SetUnbreakable(true);
-								}
+								ParticleHandle->SetUnbreakable(true);
 							}
 						}
 					});
@@ -3983,7 +3980,7 @@ void UGeometryCollectionComponent::UpdateBrokenAndDecayedStates()
 							continue;
 						}
 
-						const Chaos::FPBDRigidParticle* Particle = PhysicsProxy->GetExternalParticles()[TransformIdx].Get();
+						const Chaos::FPBDRigidParticle* Particle = PhysicsProxy->GetParticleByIndex_External(TransformIdx);
 						if (!Particle)
 						{
 							// when no particle we mark it as a decayed one
@@ -4221,7 +4218,7 @@ void UGeometryCollectionComponent::MoveComponentToRootTransform()
 			{
 				const FTransform3f& OriginalComponentSpaceRootTransformOffset = AssetCollection->Transform[RootIndex];
 				DynamicCollection->SetTransform(RootIndex, OriginalComponentSpaceRootTransformOffset);
-				const Chaos::FPBDRigidParticle* RootParticle = PhysicsProxy->GetExternalParticles()[RootIndex].Get();
+				const Chaos::FPBDRigidParticle* RootParticle = PhysicsProxy->GetParticleByIndex_External(RootIndex);
 				const FTransform ParticleWorldPosition(RootParticle->R(), RootParticle->X());
 				const FTransform MassToLocal = MassToLocalAttribute[RootIndex];
 				const FTransform NewRootWorldPosition = MassToLocal.Inverse() * ParticleWorldPosition;
@@ -6993,9 +6990,9 @@ TArray<Chaos::FPhysicsObject*> UGeometryCollectionComponent::GetAllPhysicsObject
 		return {};
 	}
 	TArray<Chaos::FPhysicsObject*> Objects;
-	Objects.Reserve(PhysicsProxy->GetNumParticles());
+	Objects.Reserve(PhysicsProxy->GetNumTransforms());
 	
-	for (int32 Index = 0; Index < PhysicsProxy->GetNumParticles(); ++Index)
+	for (int32 Index = 0; Index < PhysicsProxy->GetNumTransforms(); ++Index)
 	{
 		Objects.Add(GetPhysicsObjectById(Index));
 	}
