Fix build after merge

diff --git a/Engine/Source/Runtime/Experimental/Chaos/Private/PhysicsProxy/GeometryCollectionPhysicsProxy.cpp b/Engine/Source/Runtime/Experimental/Chaos/Private/PhysicsProxy/GeometryCollectionPhysicsProxy.cpp
--- a/Engine/Source/Runtime/Experimental/Chaos/Private/PhysicsProxy/GeometryCollectionPhysicsProxy.cpp
+++ b/Engine/Source/Runtime/Experimental/Chaos/Private/PhysicsProxy/GeometryCollectionPhysicsProxy.cpp
@@ -3227,10 +3227,9 @@ void FGeometryCollectionPhysicsProxy::SetWorldTransform_External(const FTransfor
 
 void FGeometryCollectionPhysicsProxy::ScaleClusterGeometry_Internal(const FVector& WorldScale)
 {
-	const int32 NumTransforms = PhysicsThreadCollection.NumElements(FGeometryCollection::TransformGroup);
-	for (int32 TransformGroupIndex = 0; TransformGroupIndex < NumTransforms; ++TransformGroupIndex)
+	for (int32 ParticleIndex = 0; ParticleIndex < NumEffectiveParticles; ++ParticleIndex)
 	{
-		if (Chaos::FPBDRigidClusteredParticleHandle* ParticleHandle = SolverParticleHandles[TransformGroupIndex])
+		if (Chaos::FPBDRigidClusteredParticleHandle* ParticleHandle = SolverParticleHandles[ParticleIndex])
 		{
 			// Scale the geometry if necessary
 			BuildScaledGeometry(ParticleHandle, ParticleHandle->GetGeometry(), WorldScale);
