[Backout] - CL30211300
Original CL Desc
-----------------------------------------------------------------
Reduce array size to store only required particles.
Several array in the Geometry Collection are used to store particle. According to the max simulation level, not all particle are created.
Now array have their own size, the particle size and the transform size in the proxy is now distinguished.
This change save 55% of memory in initialization.
It also avoid to process all transforms, we now iterate only through particles which should be way faster.

diff --git a/Engine/Plugins/Experimental/ChaosCaching/Source/ChaosCaching/Private/Chaos/Adapters/GeometryCollectionComponentCacheAdapter.cpp b/Engine/Plugins/Experimental/ChaosCaching/Source/ChaosCaching/Private/Chaos/Adapters/GeometryCollectionComponentCacheAdapter.cpp
--- a/Engine/Plugins/Experimental/ChaosCaching/Source/ChaosCaching/Private/Chaos/Adapters/GeometryCollectionComponentCacheAdapter.cpp
+++ b/Engine/Plugins/Experimental/ChaosCaching/Source/ChaosCaching/Private/Chaos/Adapters/GeometryCollectionComponentCacheAdapter.cpp
@@ -135,7 +135,7 @@ namespace Chaos
 
 		for(int32 TransformIndex = 0; TransformIndex < NumTransforms; ++TransformIndex)
 		{
-			FClusterParticle* Handle = Proxy->GetParticle_Internal(TransformIndex);
+			FClusterParticle* Handle = Proxy->GetParticles()[TransformIndex];
 			
 			if(Handle)
 			{
@@ -170,7 +170,7 @@ namespace Chaos
 		{
 			if (TransformIndex > INDEX_NONE)
 			{
-				FClusterParticle* Handle = Proxy->GetParticle_Internal(TransformIndex);
+				FClusterParticle* Handle = Proxy->GetParticles()[TransformIndex];
 				if (Handle)
 				{
 					OutFrame.PendingParticleData.AddDefaulted();
@@ -263,6 +263,7 @@ namespace Chaos
 
 		FGeometryDynamicCollection&      Collection       = Proxy->GetPhysicsCollection();
 		const TManagedArray<FTransform>& MassToLocal      = RestCollection->GetAttribute<FTransform>("MassToLocal", FGeometryCollection::TransformGroup);
+		TArray<FClusterParticle*>        Particles        = Proxy->GetParticles();
 
 		FCacheEvaluationContext Context(TickRecord);
 		Context.bEvaluateTransform = true;
@@ -284,26 +285,31 @@ namespace Chaos
 			{
 				if(FEnableStateEvent* Event = Handle.Get<FEnableStateEvent>())
 				{
-					if(Chaos::FPBDRigidClusteredParticleHandle* ChildParticle = Proxy->GetParticle_Internal(Event->Index))
+					if(Particles.IsValidIndex(Event->Index))
 					{
-						if (ChildParticle->ObjectState() != EObjectStateType::Kinematic)
+						Chaos::FPBDRigidClusteredParticleHandle* ChildParticle = Particles[Event->Index];
+						
+						if (ChildParticle)
 						{
-							// If a field or other external actor set the particle to static or dynamic we no longer apply the cache
-							continue;
-						}
+							if (ChildParticle->ObjectState() != EObjectStateType::Kinematic)
+							{
+								// If a field or other external actor set the particle to static or dynamic we no longer apply the cache
+								continue;
+							}
 
-						if (FRigidParticle* ClusterParent = ChildParticle->ClusterIds().Id)
-						{
-							if (FClusterParticle* Parent = ClusterParent->CastToClustered())
+							if (FRigidParticle* ClusterParent = ChildParticle->ClusterIds().Id)
 							{
-								TArray<FRigidParticle*>& Cluster = NewClusters.FindOrAdd(Parent);
-								Cluster.Add(ChildParticle);
+								if (FClusterParticle* Parent = ClusterParent->CastToClustered())
+								{
+									TArray<FRigidParticle*>& Cluster = NewClusters.FindOrAdd(Parent);
+									Cluster.Add(ChildParticle);
+								}
+							}
+							else
+							{
+								// This is a cluster parent
+								ChildParticle->SetDisabled(!Event->bEnable);
 							}
-						}
-						else
-						{
-							// This is a cluster parent
-							ChildParticle->SetDisabled(!Event->bEnable);
 						}
 					}
 				}
@@ -349,36 +355,41 @@ namespace Chaos
 			{
 				if (FBreakingEvent* Event = Handle.Get<FBreakingEvent>())
 				{
-					if (Chaos::FPBDRigidClusteredParticleHandle* Particle = Proxy->GetParticle_Internal(Event->Index))
+					if (Particles.IsValidIndex(Event->Index))
 					{
-						if (Particle->ObjectState() != EObjectStateType::Kinematic)
-						{
-							// If a field or other external actor set the particle to static or dynamic we no longer apply the cache
-							continue;
-						}
+						Chaos::FPBDRigidClusteredParticleHandle* Particle = Particles[Event->Index];
 
-						FBreakingData CachedBreak;
-						CachedBreak.Proxy = Proxy;
-						CachedBreak.Location = ComponentToWorld.TransformPosition(Event->Location);
-						CachedBreak.Velocity = ComponentToWorld.TransformVector(Event->Velocity);
-						CachedBreak.AngularVelocity = Event->AngularVelocity;
-						CachedBreak.Mass = Event->Mass;
-						CachedBreak.BoundingBox = TAABB<FReal, 3>(Event->BoundingBoxMin, Event->BoundingBoxMax);
-						CachedBreak.BoundingBox = CachedBreak.BoundingBox.TransformedAABB(ComponentToWorld);
-						CachedBreak.SetEmitterFlag(Proxy->GetSimParameters().bGenerateBreakingData, Proxy->GetSimParameters().bGenerateGlobalBreakingData);
-
-						if (!SolverBreakingEventFilter->Enabled() || SolverBreakingEventFilter->Pass(CachedBreak))
+						if (Particle)
 						{
-							float TimeStamp = Solver->GetSolverTime();
-							Solver->GetEventManager()->AddEvent<FBreakingEventData>(EEventType::Breaking, [&CachedBreak, TimeStamp](FBreakingEventData& BreakingEventData)
-								{
-									if (BreakingEventData.BreakingData.TimeCreated != TimeStamp)
+							if (Particle->ObjectState() != EObjectStateType::Kinematic)
+							{
+								// If a field or other external actor set the particle to static or dynamic we no longer apply the cache
+								continue;
+							}
+
+							FBreakingData CachedBreak;
+							CachedBreak.Proxy = Proxy;
+							CachedBreak.Location = ComponentToWorld.TransformPosition(Event->Location);
+							CachedBreak.Velocity = ComponentToWorld.TransformVector(Event->Velocity);
+							CachedBreak.AngularVelocity = Event->AngularVelocity;
+							CachedBreak.Mass = Event->Mass;
+							CachedBreak.BoundingBox = TAABB<FReal, 3>(Event->BoundingBoxMin, Event->BoundingBoxMax);
+							CachedBreak.BoundingBox = CachedBreak.BoundingBox.TransformedAABB(ComponentToWorld);
+							CachedBreak.SetEmitterFlag(Proxy->GetSimParameters().bGenerateBreakingData, Proxy->GetSimParameters().bGenerateGlobalBreakingData);
+
+							if (!SolverBreakingEventFilter->Enabled() || SolverBreakingEventFilter->Pass(CachedBreak))
+							{
+								float TimeStamp = Solver->GetSolverTime();
+								Solver->GetEventManager()->AddEvent<FBreakingEventData>(EEventType::Breaking, [&CachedBreak, TimeStamp](FBreakingEventData& BreakingEventData)
 									{
-										BreakingEventData.BreakingData.AllBreakingsArray.Reset();
-										BreakingEventData.BreakingData.TimeCreated = TimeStamp;
-									}
-									BreakingEventData.BreakingData.AllBreakingsArray.Add(CachedBreak);
-								});
+										if (BreakingEventData.BreakingData.TimeCreated != TimeStamp)
+										{
+											BreakingEventData.BreakingData.AllBreakingsArray.Reset();
+											BreakingEventData.BreakingData.TimeCreated = TimeStamp;
+										}
+										BreakingEventData.BreakingData.AllBreakingsArray.Add(CachedBreak);
+									});
+							}
 						}
 					}
 				}
@@ -393,34 +404,39 @@ namespace Chaos
 			{
 				if (FTrailingEvent* Event = Handle.Get<FTrailingEvent>())
 				{
-					if (Chaos::FPBDRigidClusteredParticleHandle* Particle = Proxy->GetParticle_Internal(Event->Index))
+					if (Particles.IsValidIndex(Event->Index))
 					{
-						if (Particle->ObjectState() != EObjectStateType::Kinematic)
+						Chaos::FPBDRigidClusteredParticleHandle* Particle = Particles[Event->Index];
+
+						if (Particle)
 						{
-							// If a field or other external actor set the particle to static or dynamic we no longer apply the cache
-							continue;
-						}
+							if (Particle->ObjectState() != EObjectStateType::Kinematic)
+							{
+								// If a field or other external actor set the particle to static or dynamic we no longer apply the cache
+								continue;
+							}
 
-						FTrailingData CachedTrail;
-						CachedTrail.Proxy = Proxy;
-						CachedTrail.Location = ComponentToWorld.TransformPosition(Event->Location);
-						CachedTrail.Velocity = ComponentToWorld.TransformVector(Event->Velocity);
-						CachedTrail.AngularVelocity = Event->AngularVelocity;
-						CachedTrail.BoundingBox = TAABB<FReal, 3>(Event->BoundingBoxMin, Event->BoundingBoxMax);
-						CachedTrail.BoundingBox = CachedTrail.BoundingBox.TransformedAABB(ComponentToWorld);
+							FTrailingData CachedTrail;
+							CachedTrail.Proxy = Proxy;
+							CachedTrail.Location = ComponentToWorld.TransformPosition(Event->Location);
+							CachedTrail.Velocity = ComponentToWorld.TransformVector(Event->Velocity);
+							CachedTrail.AngularVelocity = Event->AngularVelocity;
+							CachedTrail.BoundingBox = TAABB<FReal, 3>(Event->BoundingBoxMin, Event->BoundingBoxMax);
+							CachedTrail.BoundingBox = CachedTrail.BoundingBox.TransformedAABB(ComponentToWorld);
 
-						if (!SolverTrailingEventFilter->Enabled() || SolverTrailingEventFilter->Pass(CachedTrail))
-						{
-							float TimeStamp = Solver->GetSolverTime();
-							Solver->GetEventManager()->AddEvent<FTrailingEventData>(EEventType::Trailing, [&CachedTrail , TimeStamp](FTrailingEventData& TrailingEventData)
-								{
-									if (TrailingEventData.TrailingData.TimeCreated != TimeStamp)
+							if (!SolverTrailingEventFilter->Enabled() || SolverTrailingEventFilter->Pass(CachedTrail))
+							{
+								float TimeStamp = Solver->GetSolverTime();
+								Solver->GetEventManager()->AddEvent<FTrailingEventData>(EEventType::Trailing, [&CachedTrail , TimeStamp](FTrailingEventData& TrailingEventData)
 									{
-										TrailingEventData.TrailingData.AllTrailingsArray.Reset();
-										TrailingEventData.TrailingData.TimeCreated = TimeStamp;
-									}
-									TrailingEventData.TrailingData.AllTrailingsArray.Add(CachedTrail);
-								});
+										if (TrailingEventData.TrailingData.TimeCreated != TimeStamp)
+										{
+											TrailingEventData.TrailingData.AllTrailingsArray.Reset();
+											TrailingEventData.TrailingData.TimeCreated = TimeStamp;
+										}
+										TrailingEventData.TrailingData.AllTrailingsArray.Add(CachedTrail);
+									});
+							}
 						}
 					}
 				}
@@ -435,6 +451,7 @@ namespace Chaos
 			{
 				if (FCollisionEvent* Event = Handle.Get<FCollisionEvent>())
 				{
+					
 					FCollidingData CachedCollision;
 					CachedCollision.Location = ComponentToWorld.TransformPosition(Event->Location);
 					CachedCollision.AccumulatedImpulse = ComponentToWorld.TransformVector(Event->AccumulatedImpulse);
@@ -476,9 +493,11 @@ namespace Chaos
 			const int32      ParticleIndex      = EvaluatedResult.ParticleIndices[Index];
 			const FTransform EvaluatedTransform = EvaluatedResult.Transform[Index];
 
-			if (Chaos::FPBDRigidClusteredParticleHandle* Handle = Proxy->GetParticle_Internal(ParticleIndex))
+			if(Particles.IsValidIndex(ParticleIndex))
 			{
-				if(Handle->ObjectState() != EObjectStateType::Kinematic)
+				Chaos::FPBDRigidClusteredParticleHandle* Handle = Particles[ParticleIndex];
+
+				if(!Handle || Handle->ObjectState() != EObjectStateType::Kinematic)
 				{
 					// If a field or other external actor set the particle to static or dynamic we no longer apply the cache
 					continue;
diff --git a/Engine/Plugins/Experimental/ChaosModularVehicle/Source/ChaosModularVehicleEngine/Private/ChaosModularVehicle/ModularVehicleComponent.cpp b/Engine/Plugins/Experimental/ChaosModularVehicle/Source/ChaosModularVehicleEngine/Private/ChaosModularVehicle/ModularVehicleComponent.cpp
--- a/Engine/Plugins/Experimental/ChaosModularVehicle/Source/ChaosModularVehicleEngine/Private/ChaosModularVehicle/ModularVehicleComponent.cpp
+++ b/Engine/Plugins/Experimental/ChaosModularVehicle/Source/ChaosModularVehicleEngine/Private/ChaosModularVehicle/ModularVehicleComponent.cpp
@@ -62,11 +62,11 @@ Chaos::FPBDRigidClusteredParticleHandle* UModularVehicleComponent::GetChassisPar
 
 	if (FGeometryCollectionPhysicsProxy* Proxy = GetPhysicsProxy())
 	{
-		Chaos::FPBDRigidClusteredParticleHandle* Cluster = Proxy->GetSolverClusterHandle_Internal(0);
+		const TArray<Chaos::FPBDRigidClusteredParticleHandle*>& Clusters = Proxy->GetSolverClusterHandles();
 
-		if (Cluster != nullptr && !Cluster->Disabled())
+		if (Clusters.Num() > 0 && !Clusters[0]->Disabled())
 		{
-			Particle = Cluster;
+			Particle = Clusters[0];
 		}
 	}
 	return Particle;
@@ -221,8 +221,8 @@ bool UModularVehicleComponent::ProcessRepData(const float DeltaTime, const float
 								check(RewindData);
 								check(!RewindData->IsResim());
 
-								Chaos::FPBDRigidClusteredParticleHandle* ChassisHandle = Proxy->GetSolverClusterHandle_Internal(ChassisIdx);
-
+								const TArray<Chaos::FPBDRigidClusteredParticleHandle*>& Clusters = Proxy->GetSolverClusterHandles();
+								auto* ChassisHandle = Clusters[ChassisIdx];
 								const FGeometryCollectionClusterRep& RepCluster = RepData.Clusters[ChassisIdx];
 
 								// Need the local frame in order to access correct history data
diff --git a/Engine/Plugins/Experimental/ChaosModularVehicle/Source/ChaosModularVehicleEngine/Private/ChaosModularVehicle/ModularVehicleSimulationCU.cpp b/Engine/Plugins/Experimental/ChaosModularVehicle/Source/ChaosModularVehicleEngine/Private/ChaosModularVehicle/ModularVehicleSimulationCU.cpp
--- a/Engine/Plugins/Experimental/ChaosModularVehicle/Source/ChaosModularVehicleEngine/Private/ChaosModularVehicle/ModularVehicleSimulationCU.cpp
+++ b/Engine/Plugins/Experimental/ChaosModularVehicle/Source/ChaosModularVehicleEngine/Private/ChaosModularVehicle/ModularVehicleSimulationCU.cpp
@@ -363,7 +363,8 @@ void FModularVehicleSimulationCU::ApplyDeferredForces(FGeometryCollectionPhysics
 					const TManagedArray<FTransform>& CollectionMassToLocal = Rest->GetGeometryCollection()->GetAttribute<FTransform>(TEXT("MassToLocal"), FTransformCollection::TransformGroup);
 
 					SimModuleTree->AccessDeferredForces().Apply(
-						Proxy,
+						Proxy->GetSolverParticleHandles(),
+						Proxy->GetSolverClusterHandles(),
 						Rest->GetGeometryCollection()->Transform,
 						CollectionMassToLocal,
 						Rest->GetGeometryCollection()->Parent);
diff --git a/Engine/Plugins/Experimental/ChaosModularVehicle/Source/ChaosModularVehicleEngine/Private/ChaosModularVehicle/ModularVehicleSimulationGC.cpp b/Engine/Plugins/Experimental/ChaosModularVehicle/Source/ChaosModularVehicleEngine/Private/ChaosModularVehicle/ModularVehicleSimulationGC.cpp
--- a/Engine/Plugins/Experimental/ChaosModularVehicle/Source/ChaosModularVehicleEngine/Private/ChaosModularVehicle/ModularVehicleSimulationGC.cpp
+++ b/Engine/Plugins/Experimental/ChaosModularVehicle/Source/ChaosModularVehicleEngine/Private/ChaosModularVehicle/ModularVehicleSimulationGC.cpp
@@ -143,21 +143,25 @@ void FModularVehicleSimulationGC::PerformAdditionalSimWork(UWorld* InWorld, cons
 	Chaos::EnsureIsInPhysicsThreadContext();
 
 	FGeometryDynamicCollection& GeometryCollection = Proxy->GetPhysicsCollection();
+	const TArray<Chaos::FPBDRigidClusteredParticleHandle*>& Clusters = Proxy->GetSolverClusterHandles();
+	const TArray<Chaos::FPBDRigidClusteredParticleHandle*>& Particles = Proxy->GetSolverParticleHandles();
 
-	if (Proxy->GetUnorderedParticles_Internal().Num() > 0)
+	if (Clusters.Num() > 0)
 	{
 		const TArray<Chaos::FSimModuleTree::FSimModuleNode>& ModuleArray = SimModuleTree->GetSimulationModuleTree();
+
 		for (const Chaos::FSimModuleTree::FSimModuleNode& Node : ModuleArray)
 		{
-			const Chaos::FPBDRigidClusteredParticleHandle* Cluster = Proxy->GetParticle_Internal(Node.SimModule->GetTransformIndex());
-			const Chaos::FPBDRigidClusteredParticleHandle* Particle = Proxy->GetSolverClusterHandle_Internal(Node.SimModule->GetTransformIndex());
-			if (Cluster != nullptr && Particle != nullptr)
+			if (!Node.SimModule->IsEnabled() || !Particles[Node.SimModule->GetTransformIndex()]->Disabled())
+			{
+				Node.SimModule->SetStateFlags(Chaos::eSimModuleState::Disabled);
+				continue;
+			}
+
+			check(Node.SimModule->GetTransformIndex() < Clusters.Num());
+			const Chaos::FPBDRigidClusteredParticleHandle* Particle = Clusters[Node.SimModule->GetTransformIndex()];
+			if (Particle)
 			{
-				if (!Node.SimModule->IsEnabled() || !Cluster->Disabled())
-				{
-					Node.SimModule->SetStateFlags(Chaos::eSimModuleState::Disabled);
-					continue;
-				}
 				const FTransform BodyTransform(Particle->R(), Particle->X());
 
 				// #TODO: cheating just now to get it working - do we pass in one or one per particle, as it will be different if it has split into fragments
@@ -296,7 +300,8 @@ void FModularVehicleSimulationGC::ApplyDeferredForces(FGeometryCollectionPhysics
 					const TManagedArray<FTransform>& CollectionMassToLocal = Rest->GetGeometryCollection()->GetAttribute<FTransform>(TEXT("MassToLocal"), FTransformCollection::TransformGroup);
 
 					SimModuleTree->AccessDeferredForces().Apply(
-						Proxy,
+						Proxy->GetSolverParticleHandles(),
+						Proxy->GetSolverClusterHandles(),
 						Rest->GetGeometryCollection()->Transform,
 						CollectionMassToLocal,
 						Rest->GetGeometryCollection()->Parent);
diff --git a/Engine/Source/Runtime/Engine/Private/PhysicsEngine/Experimental/PhysScene_Chaos.cpp b/Engine/Source/Runtime/Engine/Private/PhysicsEngine/Experimental/PhysScene_Chaos.cpp
--- a/Engine/Source/Runtime/Engine/Private/PhysicsEngine/Experimental/PhysScene_Chaos.cpp
+++ b/Engine/Source/Runtime/Engine/Private/PhysicsEngine/Experimental/PhysScene_Chaos.cpp
@@ -784,7 +784,7 @@ void FPhysScene_Chaos::RemoveObject(FGeometryCollectionPhysicsProxy* InObject)
 {
 	Chaos::FPhysicsSolver* Solver = InObject->GetSolver<Chaos::FPhysicsSolver>();
 
-	for (TUniquePtr<Chaos::TPBDRigidParticle<Chaos::FReal, 3>>& GTParticleUnique : InObject->GetUnorderedParticles_External())
+	for (TUniquePtr<Chaos::TPBDRigidParticle<Chaos::FReal, 3>>& GTParticleUnique : InObject->GetExternalParticles())
 	{
 		Chaos::TPBDRigidParticle<Chaos::FReal, 3>* GTParticle = GTParticleUnique.Get();
 		if (GTParticle != nullptr)
@@ -2190,13 +2190,14 @@ static void UpdateAccelerationStructureFromGeometryCollectionProxy(FGeometryColl
 		return Chaos::FAABB3(Particle.X(), Particle.X());
 	};
 
-	const TArray<TUniquePtr<Chaos::FPBDRigidParticle>>& UnorderedGTParticles = Proxy.GetUnorderedParticles_External();
+	const TArray<TUniquePtr<Chaos::FPBDRigidParticle>>& GTParticles = Proxy.GetExternalParticles();
 	if (bIsParentProxyNull)
 	{
-		for (const TUniquePtr<Chaos::FPBDRigidParticle>& GTParticle : UnorderedGTParticles)
+		for (const TUniquePtr<Chaos::FPBDRigidParticle>& GTParticle : GTParticles)
 		{
 			if (GTParticle)
 			{
+
 				const Chaos::FSpatialAccelerationIdx SpatialIndex = GTParticle->SpatialIdx();
 
 				// It's possible to be an enabled particle and not qualify for query collisions if the GC particle has been replication abandoned. 
@@ -2223,7 +2224,7 @@ static void UpdateAccelerationStructureFromGeometryCollectionProxy(FGeometryColl
 	{
 		// if we have a parent proxy ( like attached to a cluster union) then none of the handles should be in the acceleration structure 
 		// todo : right now we are sending all the particles, but we should be able to get away with active + direct children that could reduce the overall cost 
-		for (const TUniquePtr<Chaos::FPBDRigidParticle>& GTParticle : UnorderedGTParticles)
+		for (const TUniquePtr<Chaos::FPBDRigidParticle>& GTParticle : GTParticles)
 		{
 			if (GTParticle)
 			{
diff --git a/Engine/Source/Runtime/Experimental/Chaos/Private/Chaos/StrainModification.cpp b/Engine/Source/Runtime/Experimental/Chaos/Private/Chaos/StrainModification.cpp
--- a/Engine/Source/Runtime/Experimental/Chaos/Private/Chaos/StrainModification.cpp
+++ b/Engine/Source/Runtime/Experimental/Chaos/Private/Chaos/StrainModification.cpp
@@ -71,7 +71,7 @@ static void ForEachRootChildParticle(const Chaos::FStrainedProxyAndRoot& ProxyAn
 		{
 			for (int32 RestChildIdx : *RestChildren)
 			{
-				Chaos::FPBDRigidClusteredParticleHandle* ChildHandle = ProxyGC->GetParticle_Internal(RestChildIdx);
+				Chaos::FPBDRigidClusteredParticleHandle* ChildHandle = ProxyGC->GetSolverParticleHandles()[RestChildIdx];
 				if ((ChildHandle == nullptr) || (ChildHandle->Parent() == nullptr)) { continue; }
 
 				Func(ChildHandle);
@@ -86,7 +86,7 @@ int32 Chaos::FStrainedProxyModifier::GetNumRestBreakables() const
 	{
 		return RestChildren->Num();
 	}
-	return 1; // cluster union case
+	return 1; // cluser union case
 }
 
 int32 Chaos::FStrainedProxyModifier::GetNumBreakingStrains(const bool bDoubleCount, const uint8 StrainTypes) const
@@ -228,10 +228,10 @@ Chaos::FStrainedProxyRange::FStrainedProxyRange(Chaos::FRigidClustering& InRigid
 			{
 				FSimulationParameters& Parameters = ProxyGC->GetSimParameters();
 				const int32 RootIndex = Parameters.InitialRootIndex;
-				Chaos::FPBDRigidClusteredParticleHandle* ParticleHandle = ProxyGC->GetParticle_Internal(RootIndex);
-				if (ParticleHandle != nullptr)
+				TArray<Chaos::FPBDRigidClusteredParticleHandle*>& ParticleHandles = ProxyGC->GetSolverParticleHandles();
+				if (ParticleHandles.IsValidIndex(RootIndex))
 				{
-					if (Cluster != ParticleHandle)
+					if (Cluster != ParticleHandles[RootIndex])
 					{
 						return;
 					}
diff --git a/Engine/Source/Runtime/Experimental/Chaos/Private/PBDRigidsSolver.cpp b/Engine/Source/Runtime/Experimental/Chaos/Private/PBDRigidsSolver.cpp
--- a/Engine/Source/Runtime/Experimental/Chaos/Private/PBDRigidsSolver.cpp
+++ b/Engine/Source/Runtime/Experimental/Chaos/Private/PBDRigidsSolver.cpp
@@ -772,7 +772,7 @@ namespace Chaos
 		InProxy->NewData(); // Buffers data on the proxy.
 
 		// Need to immediately add to the SQ to make sure that the GC is in the external SQ even after the SQ's flip.
-		for (const TUniquePtr<FPBDRigidParticle>& Particle : InProxy->GetUnorderedParticles_External())
+		for (const TUniquePtr<FPBDRigidParticle>& Particle : InProxy->GetExternalParticles())
 		{
 			if (Particle && !Particle->Disabled())
 			{
diff --git a/Engine/Source/Runtime/Experimental/Chaos/Private/PhysicsProxy/GeometryCollectionPhysicsProxy.cpp b/Engine/Source/Runtime/Experimental/Chaos/Private/PhysicsProxy/GeometryCollectionPhysicsProxy.cpp
--- a/Engine/Source/Runtime/Experimental/Chaos/Private/PhysicsProxy/GeometryCollectionPhysicsProxy.cpp
+++ b/Engine/Source/Runtime/Experimental/Chaos/Private/PhysicsProxy/GeometryCollectionPhysicsProxy.cpp
@@ -505,7 +505,7 @@ FGeometryCollectionPhysicsProxy::FGeometryCollectionPhysicsProxy(
 	const Chaos::EMultiBufferMode BufferMode)
 	: Base(InOwner)
 	, Parameters(SimulationParameters)
-	, NumTransforms(INDEX_NONE)
+	, NumParticles(INDEX_NONE)
 	, NumEffectiveParticles(INDEX_NONE)
 	, BaseParticleIndex(INDEX_NONE)
 	, IsObjectDynamic(false)
@@ -563,40 +563,31 @@ void FGeometryCollectionPhysicsProxy::Initialize(Chaos::FPBDRigidsEvolutionBase
 
 	InitializeDynamicCollection(DynamicCollection, *Parameters.RestCollection, Parameters);
 
-	NumTransforms = DynamicCollection.NumElements(FGeometryCollection::TransformGroup);
+	NumParticles = DynamicCollection.NumElements(FGeometryCollection::TransformGroup);
 	BaseParticleIndex = 0; // Are we always zero indexed now?
+	SolverClusterID.Init(nullptr, NumParticles);
+	SolverClusterHandles.Init(nullptr, NumParticles);
+	SolverParticleHandles.Init(nullptr, NumParticles);
+	GTParticles.SetNum(NumParticles);
+	UniqueIdxs.SetNum(NumParticles);
 	
 	TBitArray<> EffectiveParticles;
-	NumEffectiveParticles = CalculateEffectiveParticles(DynamicCollection, NumTransforms, Parameters.MaxSimulatedLevel, Parameters.EnableClustering, GetOwner(), EffectiveParticles);
+	NumEffectiveParticles = CalculateEffectiveParticles(DynamicCollection, NumParticles, Parameters.MaxSimulatedLevel, Parameters.EnableClustering, GetOwner(), EffectiveParticles);
 
-	SolverClusterID.Init(nullptr, NumEffectiveParticles);
-	SolverClusterHandles.Init(nullptr, NumEffectiveParticles);
-	SolverParticleHandles.Init(nullptr, NumEffectiveParticles);
-	GTParticles.SetNum(NumEffectiveParticles);
-	UniqueIdxs.SetNum(NumEffectiveParticles);
-	
-	FromParticleToTransformIndex.SetNum(NumEffectiveParticles);
-	FromTransformToParticleIndex.SetNum(NumTransforms);
 	PhysicsObjects.Empty();
-	PhysicsObjects.Reserve(NumEffectiveParticles);
-	int32 RedirectIndex = 0;
-	for (int32 Index = 0; Index < NumTransforms; ++Index)
+	PhysicsObjects.Reserve(NumParticles);
+	for (int32 Index = 0; Index < NumParticles; ++Index)
 	{
 		if (EffectiveParticles[Index])
 		{
-			FromTransformToParticleIndex[Index] = RedirectIndex;
-			FromParticleToTransformIndex[RedirectIndex] = Index;
 			PhysicsObjects.Emplace(Chaos::FPhysicsObjectFactory::CreatePhysicsObject(this, Index, FName(Parameters.RestCollection->BoneName[Index])));
-			RedirectIndex++;
 		}
 		else
 		{
-			FromTransformToParticleIndex[Index] = INDEX_NONE;
+			PhysicsObjects.Emplace(nullptr);
 		}
 	}
 
-	check(RedirectIndex == NumEffectiveParticles);
-
 	// we need to make sure the world transform is kept up to date on the game thread 
 	WorldTransform_External = Parameters.WorldTransform;
 	PreviousWorldTransform_External = WorldTransform_External;
@@ -621,7 +612,7 @@ void FGeometryCollectionPhysicsProxy::Initialize(Chaos::FPBDRigidsEvolutionBase
 	}
 
 	// Initialise GT/External particles
-	check(NumTransforms == DynamicCollection.GetNumTransforms());
+	const int32 NumTransforms = DynamicCollection.GetNumTransforms();
 	// make sure we copy the anchored information over to the physics thread collection
 	const Chaos::Facades::FCollectionAnchoringFacade DynamicCollectionAnchoringFacade(DynamicCollection);
 	Chaos::Facades::FCollectionAnchoringFacade PhysicsThreadCollectionAnchoringFacade(PhysicsThreadCollection);
@@ -631,7 +622,7 @@ void FGeometryCollectionPhysicsProxy::Initialize(Chaos::FPBDRigidsEvolutionBase
 	const TManagedArray<float>& Mass = Parameters.RestCollection->GetAttribute<float>(MassAttributeName, FTransformCollection::TransformGroup);
 
 	TManagedArray<Chaos::FImplicitObjectPtr>& Implicits = GameThreadCollection.ModifyAttribute<Chaos::FImplicitObjectPtr>(FGeometryDynamicCollection::ImplicitsAttribute, FTransformCollection::TransformGroup);
-	if (ensure(NumTransforms == Implicits.Num() && NumEffectiveParticles == GTParticles.Num())) // Implicits are in the transform group so this invariant should always hold
+	if (ensure(NumTransforms == Implicits.Num() && NumTransforms == GTParticles.Num())) // Implicits are in the transform group so this invariant should always hold
 	{
 		constexpr bool bInitializationTime = true;
 		bHasBuiltGeometryOnGT = bBuildGeometryForChildren;
@@ -677,14 +668,13 @@ void FGeometryCollectionPhysicsProxy::Initialize(Chaos::FPBDRigidsEvolutionBase
 		if (Parameters.EnableClustering)
 		{
 			// make sure we set Activate the right way when clustering is enabled ( only root should be enabled at start ) 
-			for (int32 ParticleIndex = 0; ParticleIndex < NumEffectiveParticles; ++ParticleIndex)
+			for (int32 TransformIndex = 0; TransformIndex < NumTransforms; ++TransformIndex)
 			{
-				const int32 TransformIndex = FromParticleToTransformIndex[ParticleIndex];
 				const bool bIsRoot = !GameThreadCollection.GetHasParent(TransformIndex);
 				GameThreadCollection.Active[TransformIndex] = bIsRoot;
 				PhysicsThreadCollection.Active[TransformIndex] = bIsRoot;
 
-				if (FParticle* P = GTParticles[ParticleIndex].Get())
+				if (FParticle* P = GTParticles[TransformIndex].Get())
 				{
 					if (P != nullptr)
 					{
@@ -714,98 +704,102 @@ void FGeometryCollectionPhysicsProxy::CreateGTParticles(const TBitArray<>& Effec
 	const TManagedArray<float>& Mass = Parameters.RestCollection->GetAttribute<float>(MassAttributeName, FTransformCollection::TransformGroup);
 	const TManagedArray<FTransform>& MassToLocal = Parameters.RestCollection->GetAttribute<FTransform>(MassToLocalAttributeName, FTransformCollection::TransformGroup);
 	const TManagedArray<int32>& Level = Parameters.RestCollection->GetAttribute<int32>(LevelAttributeName, FTransformCollection::TransformGroup);
+	const int32 NumTransforms = GameThreadCollection.GetNumTransforms();
 
 	TArray<int32> ChildrenToCheckForParentFix;
 
-	for (int32 ParticleIndex = 0; ParticleIndex < NumEffectiveParticles; ++ParticleIndex)
+	for (int32 Index = 0; Index < NumTransforms; ++Index)
 	{
-		const int32 TransformIndex = FromParticleToTransformIndex[ParticleIndex];
-
-		FParticle* P = GTParticles[ParticleIndex].Get();
-		// Generate all particles unique idx at initialization time
-		if (bInitializationTime)
-		{
-			UniqueIdxs[ParticleIndex] = Evolution->GenerateUniqueIdx();
-		}
-
-		if ((bInitializationTime && TransformIndex == Parameters.InitialRootIndex) || // When initializing always create particle for the root
-			((bInitializationTime && (bCreateGTParticleForChildren || bBuildGeometryForChildren)) || // When initializing create all particles if one of the flag is true 
-				(!bInitializationTime && TransformIndex != Parameters.InitialRootIndex && !bCreateGTParticleForChildren && !bBuildGeometryForChildren))) // When not initializing create other particles if flag was set to not create
+		if (EffectiveParticles[Index])
 		{
-			GTParticles[ParticleIndex] = FParticle::CreateParticle();
-			P = GTParticles[ParticleIndex].Get();
-			GTParticlesToTransformGroupIndex.Add(P, TransformIndex);
-			GTParticles[ParticleIndex]->SetUniqueIdx(UniqueIdxs[ParticleIndex]);
-#if CHAOS_DEBUG_NAME
-				P->SetDebugName(MakeShared<FString, ESPMode::ThreadSafe>(FString::Printf(TEXT("%s-%d"), *Parameters.Name, TransformIndex)));
-#endif
-			const float ScaledMass = AdjustMassForScale(Mass[TransformIndex]);
-
-			// Note that this transform must match the physics thread transform computation for initialization.
-			// Take for example, a geometry collection in the editor that is linked with a joint constraint. The joint
-			// constraint will query the game thread particle position/rotation for the geometry collection to compute its
-			// reference frame. If that position/rotation does not match up with the physics thread's position/rotation,
-			// the geometry collection particle will have an added velocity computed by the joint constraint solver.
-			const FTransform& T = MassToLocal[TransformIndex] * FTransform(GameThreadCollection.GetTransform(TransformIndex)) * Parameters.WorldTransform;
-			P->SetX(T.GetTranslation(), false);
-			P->SetR(T.GetRotation(), false);
-			P->SetM(ScaledMass);
-			P->SetUserData(Parameters.UserData);
-			P->SetProxy(this);
-		}
-		if (bInitializationTime && TransformIndex == Parameters.InitialRootIndex && bUseStaticMeshCollisionForTraces && CreateTraceCollisionGeometryCallback != nullptr)
-		{
-			const FTransform ToLocal = MassToLocal[TransformIndex].Inverse();
-			TArray<Chaos::FImplicitObjectPtr> Geoms;
-			Chaos::FShapesArray Shapes;
-			CreateTraceCollisionGeometryCallback(ToLocal, Geoms, Shapes);
-
-			Chaos::FImplicitObjectPtr ImplicitGeometry = MakeImplicitObjectPtr<Chaos::FImplicitObjectUnion>(MoveTemp(Geoms));
-			P->SetGeometry(ImplicitGeometry);
-		}
-		else if (bInitializationTime == (TransformIndex == Parameters.InitialRootIndex) || bBuildGeometryForChildren)
-		{
-			Chaos::FImplicitObjectPtr ImplicitGeometry = Implicits[TransformIndex];
-			if (ImplicitGeometry && !Scale.Equals(FVector::OneVector))
+			FParticle* P = GTParticles[Index].Get();
+			// Generate all particles unique idx at initialization time
+			if (bInitializationTime)
 			{
-				ImplicitGeometry = ImplicitGeometry->CopyGeometryWithScale(Scale);
+				UniqueIdxs[Index] = Evolution->GenerateUniqueIdx();
 			}
-			P->SetGeometry(ImplicitGeometry);
-		}
-		// Reset the root in the ManagedArray if not in mode bInitializeRootOnly
-		if (!bInitializationTime && TransformIndex == Parameters.InitialRootIndex)
-		{
-			Implicits[TransformIndex] = P->GetGeometry();
-		}
 
-		if (bInitializationTime == (TransformIndex == Parameters.InitialRootIndex) || bBuildGeometryForChildren)
-		{
-			if (DynamicCollectionAnchoringFacade.IsAnchored(TransformIndex))
+			if ((bInitializationTime && Index == Parameters.InitialRootIndex) || // When initializing always create particle for the root
+				((bInitializationTime && (bCreateGTParticleForChildren || bBuildGeometryForChildren)) || // When initializing create all particles if one of the flag is true 
+					(!bInitializationTime && Index != Parameters.InitialRootIndex && !bCreateGTParticleForChildren && !bBuildGeometryForChildren))) // When not initializing create other particles if flag was set to not create
 			{
-				P->SetObjectState(Chaos::EObjectStateType::Kinematic, false, false);
-			}
+				GTParticles[Index] = FParticle::CreateParticle();
+				P = GTParticles[Index].Get();
+				GTParticlesToTransformGroupIndex.Add(P, Index);
+				GTParticles[Index]->SetUniqueIdx(UniqueIdxs[Index]);
+
+#if CHAOS_DEBUG_NAME
+				P->SetDebugName(MakeShared<FString, ESPMode::ThreadSafe>(FString::Printf(TEXT("%s-%d"), *Parameters.Name, Index)));
+#endif
 
-			// IMPORTANT: we need to set the right spatial index because GT particle is static and PT particle is rigid
-			// this is causing a mismatch when using the separate acceleration structures optimization which can cause crashes when destroying the particle while async tracing 
-			// todo(chaos) we should eventually refactor this code to use rigid particles on the GT side for geometry collection  
-			P->SetSpatialIdx(Chaos::FSpatialAccelerationIdx{ 0,1 });
+				const float ScaledMass = AdjustMassForScale(Mass[Index]);
+
+				// Note that this transform must match the physics thread transform computation for initialization.
+				// Take for example, a geometry collection in the editor that is linked with a joint constraint. The joint
+				// constraint will query the game thread particle position/rotation for the geometry collection to compute its
+				// reference frame. If that position/rotation does not match up with the physics thread's position/rotation,
+				// the geometry collection particle will have an added velocity computed by the joint constraint solver.
+				const FTransform& T = MassToLocal[Index] * FTransform(GameThreadCollection.GetTransform(Index)) * Parameters.WorldTransform;
+				P->SetX(T.GetTranslation(), false);
+				P->SetR(T.GetRotation(), false);
+				P->SetM(ScaledMass);
+				P->SetUserData(Parameters.UserData);
+				P->SetProxy(this);
+			}
+			if (bInitializationTime && Index == Parameters.InitialRootIndex && bUseStaticMeshCollisionForTraces && CreateTraceCollisionGeometryCallback != nullptr)
+			{
+				const FTransform ToLocal = MassToLocal[Index].Inverse();
+				TArray<Chaos::FImplicitObjectPtr> Geoms;
+				Chaos::FShapesArray Shapes;
+				CreateTraceCollisionGeometryCallback(ToLocal, Geoms, Shapes);
 
-			if (Chaos::AccelerationStructureSplitStaticAndDynamic == 1)
+				Chaos::FImplicitObjectPtr ImplicitGeometry = MakeImplicitObjectPtr<Chaos::FImplicitObjectUnion>(MoveTemp(Geoms));
+				P->SetGeometry(ImplicitGeometry);
+			}
+			else if (bInitializationTime == (Index == Parameters.InitialRootIndex) || bBuildGeometryForChildren)
 			{
-				P->SetSpatialIdx(Chaos::FSpatialAccelerationIdx{ 0,1 });
+				Chaos::FImplicitObjectPtr ImplicitGeometry = Implicits[Index];
+				if (ImplicitGeometry && !Scale.Equals(FVector::OneVector))
+				{
+					ImplicitGeometry = ImplicitGeometry->CopyGeometryWithScale(Scale);
+				}
+				P->SetGeometry(ImplicitGeometry);
 			}
-			else
+			// Reset the root in the ManagedArray if not in mode bInitializeRootOnly
+			if (!bInitializationTime && Index == Parameters.InitialRootIndex)
 			{
-				P->SetSpatialIdx(Chaos::FSpatialAccelerationIdx{ 0,0 });
+				Implicits[Index] = P->GetGeometry();
 			}
 
-			const bool bIsOneWayInteraction = (Parameters.OneWayInteractionLevel >= 0) && (Level[TransformIndex] >= Parameters.OneWayInteractionLevel);
-			P->SetOneWayInteraction(bIsOneWayInteraction);
+			if (bInitializationTime == (Index == Parameters.InitialRootIndex) || bBuildGeometryForChildren)
+			{
+				if (DynamicCollectionAnchoringFacade.IsAnchored(Index))
+				{
+					P->SetObjectState(Chaos::EObjectStateType::Kinematic, false, false);
+				}
+
+				// IMPORTANT: we need to set the right spatial index because GT particle is static and PT particle is rigid
+				// this is causing a mismatch when using the separate acceleration structures optimization which can cause crashes when destroying the particle while async tracing 
+				// todo(chaos) we should eventually refactor this code to use rigid particles on the GT side for geometry collection  
+				P->SetSpatialIdx(Chaos::FSpatialAccelerationIdx{ 0,1 });
+
+				if (Chaos::AccelerationStructureSplitStaticAndDynamic == 1)
+				{
+					P->SetSpatialIdx(Chaos::FSpatialAccelerationIdx{ 0,1 });
+				}
+				else
+				{
+					P->SetSpatialIdx(Chaos::FSpatialAccelerationIdx{ 0,0 });
+				}
+
+				const bool bIsOneWayInteraction = (Parameters.OneWayInteractionLevel >= 0) && (Level[Index] >= Parameters.OneWayInteractionLevel);
+				P->SetOneWayInteraction(bIsOneWayInteraction);
+			}
 		}
 		// this step is necessary for Phase 2 where we need to walk back the hierarchy from children to parent 
-		if (bGeometryCollectionAlwaysGenerateGTCollisionForClusters && !GameThreadCollection.HasChildren(TransformIndex))
+		if (bGeometryCollectionAlwaysGenerateGTCollisionForClusters && !GameThreadCollection.HasChildren(Index))
 		{
-			ChildrenToCheckForParentFix.Add(TransformIndex);
+			ChildrenToCheckForParentFix.Add(Index);
 		}
 	}
 
@@ -880,9 +874,9 @@ void FGeometryCollectionPhysicsProxy::CreateGTParticles(const TBitArray<>& Effec
 							Chaos::FImplicitObject* UnionImplicit = new Chaos::FImplicitObjectUnion(MoveTemp(ChildImplicits));
 							Implicits[ParentToFixIndex] = Chaos::FImplicitObjectPtr(UnionImplicit);
 						}
-						if (GTParticles[FromTransformToParticleIndex[ParentToFixIndex]] != nullptr)
+						if (GTParticles[ParentToFixIndex] != nullptr)
 						{
-							GTParticles[FromTransformToParticleIndex[ParentToFixIndex]]->SetGeometry(Implicits[ParentToFixIndex]);
+							GTParticles[ParentToFixIndex]->SetGeometry(Implicits[ParentToFixIndex]);
 						}
 					}
 				}
@@ -895,7 +889,7 @@ void FGeometryCollectionPhysicsProxy::CreateGTParticles(const TBitArray<>& Effec
 	}
 
 	// Phase 3 : finalization of shapes
-	for (int32 Index = 0; Index < NumEffectiveParticles; ++Index)
+	for (int32 Index = 0; Index < NumTransforms; ++Index)
 	{
 		FParticle* P = GTParticles[Index].Get();
 		if (P != nullptr)
@@ -928,7 +922,7 @@ void FGeometryCollectionPhysicsProxy::CreateChildrenGeometry_External()
 
 			TManagedArray<Chaos::FImplicitObjectPtr>& Implicits = GameThreadCollection.ModifyAttribute<Chaos::FImplicitObjectPtr>(FGeometryDynamicCollection::ImplicitsAttribute, FTransformCollection::TransformGroup);
 			TBitArray<> EffectiveParticles;
-			NumEffectiveParticles = CalculateEffectiveParticles(GameThreadCollection, NumTransforms, Parameters.MaxSimulatedLevel, Parameters.EnableClustering, GetOwner(), EffectiveParticles);
+			NumEffectiveParticles = CalculateEffectiveParticles(GameThreadCollection, NumParticles, Parameters.MaxSimulatedLevel, Parameters.EnableClustering, GetOwner(), EffectiveParticles);
 			CreateGTParticles(EffectiveParticles, Implicits, RBDSolver->GetEvolution(), /*bInitializationTime*/false);
 			SyncParticles_External();
 			UniqueIdxs.Empty();
@@ -1039,23 +1033,25 @@ void FGeometryCollectionPhysicsProxy::CreateNonClusteredParticles(Chaos::FPBDRig
 
 	// Count geometry collection leaf node particles to add
 	int NumSimulatedParticles = 0;
-	for (int32 ParticleIndex = 0; ParticleIndex < NumEffectiveParticles; ++ParticleIndex)
+	for (int32 Idx = 0; Idx < SimulatableParticles.Num(); ++Idx)
 	{
-		const int32 TransformIndex = FromParticleToTransformIndex[ParticleIndex];
-		NumSimulatedParticles += SimulatableParticles[TransformIndex];
-		if (SimulatableParticles[TransformIndex] && !RestCollection.IsClustered(TransformIndex) && RestCollection.IsGeometry(TransformIndex))
+		if (EffectiveParticles[Idx])
 		{
-			NumRigids++;
-			Chaos::FUniqueIdx ExistingIndex;
-			if (GTParticles[ParticleIndex] == nullptr)
-			{
-				ExistingIndex = RigidsSolver->GetEvolution()->GenerateUniqueIdx();
-			}
-			else
+			NumSimulatedParticles += SimulatableParticles[Idx];
+			if (SimulatableParticles[Idx] && !RestCollection.IsClustered(Idx) && RestCollection.IsGeometry(Idx))
 			{
-				ExistingIndex = GTParticles[ParticleIndex]->UniqueIdx();
+				NumRigids++;
+				Chaos::FUniqueIdx ExistingIndex;
+				if (GTParticles[Idx] == nullptr)
+				{
+					ExistingIndex = RigidsSolver->GetEvolution()->GenerateUniqueIdx();
+				}
+				else
+				{
+					ExistingIndex = GTParticles[Idx]->UniqueIdx();
+				}
+				UniqueIndices.Add(ExistingIndex);
 			}
-			UniqueIndices.Add(ExistingIndex);
 		}
 	}
 
@@ -1064,11 +1060,10 @@ void FGeometryCollectionPhysicsProxy::CreateNonClusteredParticles(Chaos::FPBDRig
 	TArray<Chaos::TPBDGeometryCollectionParticleHandle<Chaos::FReal, 3>*> Handles = RigidsSolver->GetEvolution()->CreateGeometryCollectionParticles(NumRigids, UniqueIndices.GetData());
 
 	int32 NextIdx = 0;
-	for (int32 ParticleIndex = 0; ParticleIndex < NumEffectiveParticles; ++ParticleIndex)
+	for (int32 Idx = 0; Idx < SimulatableParticles.Num(); ++Idx)
 	{
-		const int32 TransformIndex = FromParticleToTransformIndex[ParticleIndex];
-		SolverParticleHandles[ParticleIndex] = nullptr;
-		if (EffectiveParticles[TransformIndex] && SimulatableParticles[TransformIndex] && !RestCollection.IsClustered(TransformIndex))
+		SolverParticleHandles[Idx] = nullptr;
+		if (EffectiveParticles[Idx] && SimulatableParticles[Idx] && !RestCollection.IsClustered(Idx))
 		{
 			// todo: Unblocked read access of game thread data on the physics thread.
 
@@ -1076,13 +1071,13 @@ void FGeometryCollectionPhysicsProxy::CreateNonClusteredParticles(Chaos::FPBDRig
 
 			Handle->SetPhysicsProxy(this);
 
-			SolverParticleHandles[ParticleIndex] = Handle;
-			HandleToTransformGroupIndex.Add(Handle, TransformIndex);
+			SolverParticleHandles[Idx] = Handle;
+			HandleToTransformGroupIndex.Add(Handle, Idx);
 
 			// We're on the physics thread here but we've already set up the GT particles and we're just linking here
-			Handle->GTGeometryParticle() = GTParticles[ParticleIndex].Get();
+			Handle->GTGeometryParticle() = GTParticles[Idx].Get();
 
-			check(SolverParticleHandles[ParticleIndex]->GetParticleType() == Handle->GetParticleType());
+			check(SolverParticleHandles[Idx]->GetParticleType() == Handle->GetParticleType());
 			RigidsSolver->GetEvolution()->RegisterParticle(Handle);
 		}
 	}
@@ -1141,6 +1136,7 @@ void FGeometryCollectionPhysicsProxy::DirtyAllParticles(const Chaos::FPBDRigidsS
 
 void FGeometryCollectionPhysicsProxy::UpdateDamageThreshold_Internal()
 {
+	const int32 NumTransforms = SolverParticleHandles.Num();
 	ensure(SolverParticleHandles.Num() == SolverClusterHandles.Num());
 
 	if (Chaos::FPhysicsSolver* RBDSolver = GetSolver<Chaos::FPhysicsSolver>())
@@ -1149,12 +1145,12 @@ void FGeometryCollectionPhysicsProxy::UpdateDamageThreshold_Internal()
 
 		const float StrainDefault = Parameters.DamageThreshold.Num() ? Parameters.DamageThreshold[0] : 0;
 
-		for (int32 ParticleIndex = 0; ParticleIndex < NumEffectiveParticles; ParticleIndex++)
+		for (int32 TransformIndex = 0; TransformIndex < NumTransforms; TransformIndex++)
 		{
-			if (Chaos::FPBDRigidClusteredParticleHandle* Handle = SolverParticleHandles[ParticleIndex])
+			if (Chaos::FPBDRigidClusteredParticleHandle* Handle = SolverParticleHandles[TransformIndex])
 			{
-				const int32 TransformIndex = FromParticleToTransformIndex[ParticleIndex];
 				float DamageThreshold = StrainDefault;
+
 				switch (Parameters.DamageModel)
 				{
 				case EDamageModelTypeEnum::Chaos_Damage_Model_UserDefined_Damage_Threshold:
@@ -1179,11 +1175,11 @@ void FGeometryCollectionPhysicsProxy::UpdateDamageThreshold_Internal()
 		if (Parameters.bUsePerClusterOnlyDamageThreshold && 
 			Parameters.DamageModel == EDamageModelTypeEnum::Chaos_Damage_Model_UserDefined_Damage_Threshold)
 		{
-			for (int32 ParticleIndex = 0; ParticleIndex < NumEffectiveParticles; ParticleIndex++)
+			for (int32 TransformIndex = 0; TransformIndex < NumTransforms; TransformIndex++)
 			{
-				if (Chaos::FPBDRigidClusteredParticleHandle* Handle = SolverParticleHandles[ParticleIndex])
+				if (Chaos::FPBDRigidClusteredParticleHandle* Handle = SolverParticleHandles[TransformIndex])
 				{
-					if (Chaos::FPBDRigidClusteredParticleHandle* ParentHandle = SolverClusterHandles[ParticleIndex])
+					if (Chaos::FPBDRigidClusteredParticleHandle* ParentHandle = SolverClusterHandles[TransformIndex])
 					{
 						const float DamageThreshold = ParentHandle->GetInternalStrains();
 						Clustering.SetInternalStrain(Handle, DamageThreshold);
@@ -1211,7 +1207,7 @@ void FGeometryCollectionPhysicsProxy::InitializeBodiesPT(Chaos::FPBDRigidsSolver
 		const TManagedArray<float>& Mass = RestCollection->GetAttribute<float>(MassAttributeName, FTransformCollection::TransformGroup);
 		const TManagedArray<FVector3f>& InertiaTensor = RestCollection->GetAttribute<FVector3f>(InertiaTensorAttributeName, FTransformCollection::TransformGroup);
 
-		check(NumTransforms == DynamicCollection.NumElements(FTransformCollection::TransformGroup));
+		const int32 NumTransforms = DynamicCollection.NumElements(FTransformCollection::TransformGroup);
 		const TManagedArray<uint8>& DynamicState = DynamicCollection.DynamicState;
 		const TManagedArray<bool>& SimulatableParticles = DynamicCollection.SimulatableParticles;
 		const TManagedArray<FTransform>& MassToLocal = RestCollection->GetAttribute<FTransform>(MassToLocalAttributeName, FTransformCollection::TransformGroup);
@@ -1259,10 +1255,9 @@ void FGeometryCollectionPhysicsProxy::InitializeBodiesPT(Chaos::FPBDRigidsSolver
 		// for the case when there's a 1-to-1 mapping between transforms and geometries.
 		// At the point that we start supporting instancing, this assumption will no longer
 		// hold, and those reverse mappints will be INDEX_NONE.
-		ParallelFor(NumEffectiveParticles, [&](int32 ParticleIndex)
+		ParallelFor(NumTransforms, [&](int32 TransformGroupIndex)
 		{
-			const int32 TransformGroupIndex = FromParticleToTransformIndex[ParticleIndex];
-			if (FClusterHandle* Handle = SolverParticleHandles[ParticleIndex])
+			if (FClusterHandle* Handle = SolverParticleHandles[TransformGroupIndex])
 			{
 				const bool bIsAnchored = AnchoringFacade.IsAnchored(TransformGroupIndex);
 
@@ -1358,11 +1353,10 @@ void FGeometryCollectionPhysicsProxy::InitializeBodiesPT(Chaos::FPBDRigidsSolver
 			// rule of thumb for parallelization is that each thread needs at least
 			// 1000 operations in order to overcome the expense of threading.  I don't
 			// think that's generally going to be the case here...
-			for (int32 ParticleIndex = 0; ParticleIndex < NumEffectiveParticles; ++ParticleIndex)
+			for (int32 TransformGroupIndex = 0; TransformGroupIndex < NumTransforms; ++TransformGroupIndex)
 			{
-				if (Chaos::TPBDRigidParticleHandle<Chaos::FReal, 3>* Handle = SolverParticleHandles[ParticleIndex])
+				if (Chaos::TPBDRigidParticleHandle<Chaos::FReal, 3>* Handle = SolverParticleHandles[TransformGroupIndex])
 				{
-					const int32 TransformGroupIndex = FromParticleToTransformIndex[ParticleIndex];
 					if (DynamicState[TransformGroupIndex] == (int32)EObjectStateTypeEnum::Chaos_Object_Dynamic)
 					{
 						Handle->SetV(InitialVelocityFacade.InitialLinearVelocityAttribute[TransformGroupIndex]);
@@ -1386,10 +1380,10 @@ void FGeometryCollectionPhysicsProxy::InitializeBodiesPT(Chaos::FPBDRigidsSolver
 			for (const int32 TransformGroupIndex : RecursiveOrder)
 			{
 				if (SimulatableParticles[TransformGroupIndex] && !RestCollection->IsClustered(TransformGroupIndex))
+
 				{
 					// Rigid node
-					const int32 ParticleIndex = FromTransformToParticleIndex[TransformGroupIndex];
-					SubTreeContainsSimulatableParticle[TransformGroupIndex] = (ParticleIndex != INDEX_NONE && SolverParticleHandles[ParticleIndex] != nullptr) || !EffectiveParticles[TransformGroupIndex];
+					SubTreeContainsSimulatableParticle[TransformGroupIndex] = SolverParticleHandles[TransformGroupIndex] != nullptr || !EffectiveParticles[TransformGroupIndex];
 				}
 				else
 				{
@@ -1447,7 +1441,7 @@ void FGeometryCollectionPhysicsProxy::InitializeBodiesPT(Chaos::FPBDRigidsSolver
 
 						if (EffectiveParticles[ChildIndex])
 						{
-							if (Chaos::TPBDRigidClusteredParticleHandle<Chaos::FReal, 3>*Handle = SolverParticleHandles[FromTransformToParticleIndex[ChildIndex]])
+							if (Chaos::TPBDRigidClusteredParticleHandle<Chaos::FReal, 3>*Handle = SolverParticleHandles[ChildIndex])
 							{
 								RigidChildren.Add(Handle);
 								RigidChildrenTransformGroupIndex.Add(ChildIndex);
@@ -1456,8 +1450,7 @@ void FGeometryCollectionPhysicsProxy::InitializeBodiesPT(Chaos::FPBDRigidsSolver
 						return true;
 					});
 
-				const int32 ParticleIndex = FromTransformToParticleIndex[TransformGroupIndex];
-				if (ParticleIndex != INDEX_NONE && SolverParticleHandles[ParticleIndex] == nullptr)
+				if (SolverParticleHandles[TransformGroupIndex] == nullptr)
 				{
 					if (ReportTooManyChildrenNum >= 0 && RigidChildren.Num() > ReportTooManyChildrenNum)
 					{
@@ -1470,7 +1463,7 @@ void FGeometryCollectionPhysicsProxy::InitializeBodiesPT(Chaos::FPBDRigidsSolver
 					CreationParameters.Scale = Parameters.WorldTransform.GetScale3D();
 
 					// Hook the handle up with the GT particle
-					FParticle* GTParticle = GTParticles[ParticleIndex].Get();
+					FParticle* GTParticle = GTParticles[TransformGroupIndex].Get();
 					Chaos::FUniqueIdx ExistingIndex;
 					if (GTParticle != nullptr)
 					{
@@ -1478,7 +1471,7 @@ void FGeometryCollectionPhysicsProxy::InitializeBodiesPT(Chaos::FPBDRigidsSolver
 					}
 					else
 					{
-						ExistingIndex = UniqueIdxs[ParticleIndex];
+						ExistingIndex = UniqueIdxs[TransformGroupIndex];
 					}
 
 					Chaos::FPBDRigidClusteredParticleHandle* Handle = nullptr;
@@ -1503,15 +1496,12 @@ void FGeometryCollectionPhysicsProxy::InitializeBodiesPT(Chaos::FPBDRigidsSolver
 					int32 RigidChildrenIdx = 0;
 					for(const int32 ChildTransformIndex : RigidChildrenTransformGroupIndex)
 					{
-						if (FromTransformToParticleIndex[ChildTransformIndex] != INDEX_NONE)
-						{
-							SolverClusterID[FromTransformToParticleIndex[ChildTransformIndex]] = RigidChildren[RigidChildrenIdx++]->CastToClustered()->ClusterIds().Id;
-						}
+						SolverClusterID[ChildTransformIndex] = RigidChildren[RigidChildrenIdx++]->CastToClustered()->ClusterIds().Id;;
 					}
-					SolverClusterID[ParticleIndex] = Handle->ClusterIds().Id;
+					SolverClusterID[TransformGroupIndex] = Handle->ClusterIds().Id;
 
-					SolverClusterHandles[ParticleIndex] = Handle;
-					SolverParticleHandles[ParticleIndex] = Handle;
+					SolverClusterHandles[TransformGroupIndex] = Handle;
+					SolverParticleHandles[TransformGroupIndex] = Handle;
 					HandleToTransformGroupIndex.Add(Handle, TransformGroupIndex);
 					Handle->SetPhysicsProxy(this);
 
@@ -1535,12 +1525,11 @@ void FGeometryCollectionPhysicsProxy::InitializeBodiesPT(Chaos::FPBDRigidsSolver
 			{
 				// Set cluster connectivity.  TPBDRigidClustering::CreateClusterParticle() 
 				// will optionally do this, but we switch that functionality off in BuildClusters_Internal().
-				for (int32 ParticleIndex = 0; ParticleIndex < NumEffectiveParticles; ++ParticleIndex)
+				for (int32 TransformGroupIndex = 0; TransformGroupIndex < NumTransforms; ++TransformGroupIndex)
 				{
-					const int32 TransformGroupIndex = FromParticleToTransformIndex[ParticleIndex];
 					if (RestCollection->IsClustered(TransformGroupIndex))
 					{
-						if (FClusterHandle* ClusteredParticle = SolverParticleHandles[ParticleIndex])
+						if (FClusterHandle* ClusteredParticle = SolverParticleHandles[TransformGroupIndex])
 						{
 							Chaos::FClusterCreationParameters ClusterParams;
 							// #todo: should other parameters be set here?  Previously, there was no parameters being sent, and it is unclear
@@ -1556,9 +1545,9 @@ void FGeometryCollectionPhysicsProxy::InitializeBodiesPT(Chaos::FPBDRigidsSolver
 				// since we do not have precise data from the tools, we are using an approximate method using the bounds of the two particles
 				if (Parameters.DamageEvaluationModel == Chaos::EDamageEvaluationModel::StrainFromMaterialStrengthAndConnectivity)
 				{
-					for (int32 ParticleIndex = 0; ParticleIndex < NumEffectiveParticles; ++ParticleIndex)
+					for (int32 TransformGroupIndex = 0; TransformGroupIndex < NumTransforms; ++TransformGroupIndex)
 					{
-						if (FClusterHandle* ClusteredParticle = SolverParticleHandles[ParticleIndex])
+						if (FClusterHandle* ClusteredParticle = SolverParticleHandles[TransformGroupIndex])
 						{
 							for (Chaos::FConnectivityEdge& Edge : ClusteredParticle->ConnectivityEdges())
 							{
@@ -1591,23 +1580,18 @@ void FGeometryCollectionPhysicsProxy::InitializeBodiesPT(Chaos::FPBDRigidsSolver
 						ContactArea = ConnectionFacade.GetConnectionContactArea(ConnectionIdx);
 					}
 
-					const int32 KeyIndex = FromTransformToParticleIndex[Connection.Key];
-					const int32 ValueIndex = FromTransformToParticleIndex[Connection.Value];
-					if (KeyIndex != INDEX_NONE && ValueIndex != INDEX_NONE)
+					if (FClusterHandle* ClusteredParticle = SolverParticleHandles[Connection.Key])
 					{
-						if (FClusterHandle* ClusteredParticle = SolverParticleHandles[KeyIndex])
+						if (FClusterHandle* OtherClusteredParticle = SolverParticleHandles[Connection.Value])
 						{
-							if (FClusterHandle* OtherClusteredParticle = SolverParticleHandles[ValueIndex])
+							float EdgeStrainOrArea = bHasContactAreas ? ContactArea : (ClusteredParticle->GetInternalStrains() + OtherClusteredParticle->GetInternalStrains()) * 0.5f;
+							if (!bHasContactAreas && Parameters.DamageEvaluationModel == Chaos::EDamageEvaluationModel::StrainFromMaterialStrengthAndConnectivity)
 							{
-								float EdgeStrainOrArea = bHasContactAreas ? ContactArea : (ClusteredParticle->GetInternalStrains() + OtherClusteredParticle->GetInternalStrains()) * 0.5f;
-								if (!bHasContactAreas && Parameters.DamageEvaluationModel == Chaos::EDamageEvaluationModel::StrainFromMaterialStrengthAndConnectivity)
-								{
-									EdgeStrainOrArea = AreaFromBoundingBoxOverlap(ClusteredParticle->WorldSpaceInflatedBounds(), OtherClusteredParticle->WorldSpaceInflatedBounds());
-								}
-								// Add symmetric Chaos::TConnectivityEdge<Chaos::FReal> edges to each particle
-								ClusteredParticle->ConnectivityEdges().Emplace(OtherClusteredParticle, EdgeStrainOrArea);
-								OtherClusteredParticle->ConnectivityEdges().Emplace(ClusteredParticle, EdgeStrainOrArea);
+								EdgeStrainOrArea = AreaFromBoundingBoxOverlap(ClusteredParticle->WorldSpaceInflatedBounds(), OtherClusteredParticle->WorldSpaceInflatedBounds());
 							}
+							// Add symmetric Chaos::TConnectivityEdge<Chaos::FReal> edges to each particle
+							ClusteredParticle->ConnectivityEdges().Emplace(OtherClusteredParticle, EdgeStrainOrArea);
+							OtherClusteredParticle->ConnectivityEdges().Emplace(ClusteredParticle, EdgeStrainOrArea);
 						}
 					}
 				}
@@ -1692,7 +1676,7 @@ void FGeometryCollectionPhysicsProxy::CreateChildrenGeometry_Internal()
 				PhysicsThreadCollection.CopyAttribute(*Parameters.RestCollection, FGeometryDynamicCollection::ImplicitsAttribute, FTransformCollection::TransformGroup);
 			}
 
-			check(NumTransforms == PhysicsThreadCollection.NumElements(FTransformCollection::TransformGroup));
+			const int32 NumTransforms = PhysicsThreadCollection.NumElements(FTransformCollection::TransformGroup);
 			const TManagedArray<FTransform>& MassToLocal = RestCollection->GetAttribute<FTransform>(MassToLocalAttributeName, FTransformCollection::TransformGroup);
 			const TManagedArray<Chaos::FImplicitObjectPtr>& Implicits = PhysicsThreadCollection.GetAttribute<Chaos::FImplicitObjectPtr>(FGeometryDynamicCollection::ImplicitsAttribute, FTransformCollection::TransformGroup);
 			const TManagedArray<TUniquePtr<FCollisionStructureManager::FSimplicial>>& Simplicials = PhysicsThreadCollection.Simplicials;
@@ -1708,16 +1692,19 @@ void FGeometryCollectionPhysicsProxy::CreateChildrenGeometry_Internal()
 			TArray<FTransform> Transform;
 			GeometryCollectionAlgo::Private::GlobalMatrices(PhysicsThreadCollection, Transform);
 
-			for (int32 ParticleIndex = 0; ParticleIndex < NumEffectiveParticles; ParticleIndex++)
+			// Here Clean up Additional particles
+			TBitArray<> EffectiveParticles;
+			NumEffectiveParticles = CalculateEffectiveParticles(PhysicsThreadCollection, PhysicsThreadCollection.NumElements(FGeometryCollection::TransformGroup), Parameters.MaxSimulatedLevel, Parameters.EnableClustering, GetOwner(), EffectiveParticles);
+
+			for (int32 Index = 0; Index < NumTransforms; Index++)
 			{
-				const int32 TransformIndex = FromParticleToTransformIndex[ParticleIndex];
-				if (TransformIndex != Parameters.InitialRootIndex)
+				if (EffectiveParticles[Index] && Index != Parameters.InitialRootIndex)
 				{
-					Chaos::TPBDGeometryCollectionParticleHandle<Chaos::FReal, 3>* Handle = static_cast<Chaos::TPBDGeometryCollectionParticleHandle<Chaos::FReal, 3>*>(SolverParticleHandles[ParticleIndex]);
+					Chaos::TPBDGeometryCollectionParticleHandle<Chaos::FReal, 3>* Handle = static_cast<Chaos::TPBDGeometryCollectionParticleHandle<Chaos::FReal, 3>*>(SolverParticleHandles[Index]);
 					if (ensure(Handle) && Handle->GetGeometry() == nullptr)
 					{
-						const FTransform WorldTransform = MassToLocal[TransformIndex] * Transform[TransformIndex] * Parameters.WorldTransform;
-						SetImplicitToPTParticles(Handle, Parameters.Shared, Simplicials[TransformIndex].Get(), Implicits[TransformIndex], SimFilter, QueryFilter, WorldTransform, CollisionParticlesPerObjectFraction);
+						const FTransform WorldTransform = MassToLocal[Index] * Transform[Index] * Parameters.WorldTransform;
+						SetImplicitToPTParticles(Handle, Parameters.Shared, Simplicials[Index].Get(), Implicits[Index], SimFilter, QueryFilter, WorldTransform, CollisionParticlesPerObjectFraction);
 						check(Handle->GetGeometry() != nullptr);
 					}
 				}
@@ -1743,10 +1730,12 @@ void FGeometryCollectionPhysicsProxy::SyncParticles_External()
 		for (int32 Index = 0; Index < ParticlNum; ++Index)
 		{
 			FParticle* GTParticle = GTParticles[Index].Get();
-			check(GTParticle != nullptr);
-			FClusterHandle* Handle = SolverParticleHandles[Index];
-			check(Handle != nullptr);
-			Handle->GTGeometryParticle() = GTParticle;
+			if (GTParticle)
+			{
+				FClusterHandle* Handle = SolverParticleHandles[Index];
+				check(Handle != nullptr);
+				Handle->GTGeometryParticle() = GTParticle;
+			}
 		}
 	});
 }
@@ -1767,10 +1756,9 @@ float FGeometryCollectionPhysicsProxy::ComputeMaterialBasedDamageThreshold_Inter
 	float DamageThreshold = TNumericLimits<float>::Max();
 	if (Parameters.DamageModel == EDamageModelTypeEnum::Chaos_Damage_Model_Material_Strength_And_Connectivity_DamageThreshold)
 	{
-		const int32 ParticleIndex = FromTransformToParticleIndex[TransformIndex];
-		if (ensure(SolverParticleHandles.IsValidIndex(ParticleIndex)))
+		if (ensure(SolverParticleHandles.IsValidIndex(TransformIndex)))
 		{
-			if (Chaos::FPBDRigidClusteredParticleHandle* ClusteredParticle = SolverParticleHandles[ParticleIndex])
+			if (Chaos::FPBDRigidClusteredParticleHandle* ClusteredParticle = SolverParticleHandles[TransformIndex])
 			{
 				DamageThreshold = ComputeMaterialBasedDamageThreshold_Internal(*ClusteredParticle);
 			}
@@ -1836,13 +1824,12 @@ float FGeometryCollectionPhysicsProxy::ComputeUserDefinedDamageThreshold_Interna
 	}
 	else if (Parameters.DamageModel == EDamageModelTypeEnum::Chaos_Damage_Model_UserDefined_Damage_Threshold)
 	{
-		const int32 ParticleIndex = FromTransformToParticleIndex[TransformIndex];
 		// we don't test for the damage model because this is used
-		if (ParticleIndex != INDEX_NONE && Parameters.bUseSizeSpecificDamageThresholds)
+		if (Parameters.bUseSizeSpecificDamageThresholds)
 		{
 			// bounding box volume is used as a fallback to find specific size if the relative size if not available
 			// ( May happen with older GC )
-			const FClusterHandle* Handle = SolverParticleHandles[ParticleIndex];
+			const FClusterHandle* Handle = SolverParticleHandles[TransformIndex];
 			FBox LocalBoundingBox;
 			if (Handle && Handle->HasBounds())
 			{
@@ -1911,15 +1898,14 @@ Chaos::TPBDGeometryCollectionParticleHandle<Chaos::FReal, 3>* FGeometryCollectio
 	Chaos::TPBDGeometryCollectionParticleHandle<Chaos::FReal, 3>* Handle = Handles[0];
 
 	Handle->SetPhysicsProxy(this);
-	const int32 ParticleIndex = FromTransformToParticleIndex[CollectionClusterIndex];
-	check(ParticleIndex != INDEX_NONE);
-	SolverParticleHandles[ParticleIndex] = Handle;
+
+	SolverParticleHandles[CollectionClusterIndex] = Handle;
 	HandleToTransformGroupIndex.Add(Handle, CollectionClusterIndex);
 
 	// We're on the physics thread here but we've already set up the GT particles and we're just linking here
-	Handle->GTGeometryParticle() = GTParticles[ParticleIndex].Get();
+	Handle->GTGeometryParticle() = GTParticles[CollectionClusterIndex].Get();
 
-	check(SolverParticleHandles[ParticleIndex]->GetParticleType() == Handle->GetParticleType());
+	check(SolverParticleHandles[CollectionClusterIndex]->GetParticleType() == Handle->GetParticleType());
 	RigidsSolver->GetEvolution()->RegisterParticle(Handle);
 
 
@@ -2087,14 +2073,12 @@ FGeometryCollectionPhysicsProxy::BuildClusters_Internal(
 		ScaledMass,
 		ScaledInertia,
 		ParticleTM, 
-		(uint8)DynamicState[CollectionClusterIndex],
+		(uint8)DynamicState[CollectionClusterIndex], 
 		0, // static_cast<int16>(CollisionGroup[TransformGroupIndex])
 		CollisionParticlesPerObjectFraction); // CollisionGroup
 
 	// two-way mapping
-	const int32 ParticleClusterIndex = FromTransformToParticleIndex[CollectionClusterIndex];
-	check(ParticleClusterIndex != INDEX_NONE);
-	SolverClusterHandles[ParticleClusterIndex] = Parent;
+	SolverClusterHandles[CollectionClusterIndex] = Parent;
 
 	// #BGTODO This will not automatically update - material properties should only ever exist in the material, not in other arrays
 	const Chaos::FChaosPhysicsMaterial* CurMaterial = static_cast<Chaos::FPBDRigidsSolver*>(Solver)->GetSimMaterials().Get(Parameters.PhysicalMaterialHandle.InnerHandle);
@@ -2119,9 +2103,7 @@ FGeometryCollectionPhysicsProxy::BuildClusters_Internal(
 		Chaos::FPBDRigidParticleHandle* Child = ChildHandles[Idx];
 
 		const int32 ChildTransformGroupIndex = ChildTransformGroupIndices[Idx];
-		const int32 ChildParticleIndex = FromTransformToParticleIndex[ChildTransformGroupIndex];
-		check(ChildParticleIndex != INDEX_NONE);
-		SolverClusterHandles[ChildParticleIndex] = Parent;
+		SolverClusterHandles[ChildTransformGroupIndex] = Parent;
 
 		MinCollisionGroup = FMath::Min(Child->CollisionGroup(), MinCollisionGroup);
 	}
@@ -2151,11 +2133,12 @@ void FGeometryCollectionPhysicsProxy::GetFilteredParticleHandles(
 	if ((ObjectType == EFieldObjectType::Field_Object_All) || (ObjectType == EFieldObjectType::Field_Object_Destruction) || (ObjectType == EFieldObjectType::Field_Object_Max))
 	{
 		// only the local handles
-		Handles.Reserve(SolverParticleHandles.Num());
+		TArray<FClusterHandle*>& ParticleHandles = GetSolverParticleHandles();
+		Handles.Reserve(ParticleHandles.Num());
 
 		if (FilterType == EFieldFilterType::Field_Filter_Dynamic)
 		{
-			for (FClusterHandle* ClusterHandle : SolverParticleHandles)
+			for (FClusterHandle* ClusterHandle : ParticleHandles)
 			{
 				if (ClusterHandle && (ClusterHandle->ObjectState() == Chaos::EObjectStateType::Dynamic))
 				{
@@ -2165,7 +2148,7 @@ void FGeometryCollectionPhysicsProxy::GetFilteredParticleHandles(
 		}
 		else if (FilterType == EFieldFilterType::Field_Filter_Kinematic)
 		{
-			for (FClusterHandle* ClusterHandle : SolverParticleHandles)
+			for (FClusterHandle* ClusterHandle : ParticleHandles)
 			{
 				if (ClusterHandle && (ClusterHandle->ObjectState() == Chaos::EObjectStateType::Kinematic))
 				{
@@ -2175,7 +2158,7 @@ void FGeometryCollectionPhysicsProxy::GetFilteredParticleHandles(
 		}
 		else if (FilterType == EFieldFilterType::Field_Filter_Static)
 		{
-			for (FClusterHandle* ClusterHandle : SolverParticleHandles)
+			for (FClusterHandle* ClusterHandle : ParticleHandles)
 			{
 				if (ClusterHandle && (ClusterHandle->ObjectState() == Chaos::EObjectStateType::Static))
 				{
@@ -2185,7 +2168,7 @@ void FGeometryCollectionPhysicsProxy::GetFilteredParticleHandles(
 		}
 		else if (FilterType == EFieldFilterType::Field_Filter_Sleeping)
 		{
-			for (FClusterHandle* ClusterHandle : SolverParticleHandles)
+			for (FClusterHandle* ClusterHandle : ParticleHandles)
 			{
 				if (ClusterHandle && (ClusterHandle->ObjectState() == Chaos::EObjectStateType::Sleeping))
 				{
@@ -2195,7 +2178,7 @@ void FGeometryCollectionPhysicsProxy::GetFilteredParticleHandles(
 		}
 		else if (FilterType == EFieldFilterType::Field_Filter_Disabled)
 		{
-			for (FClusterHandle* ClusterHandle : SolverParticleHandles)
+			for (FClusterHandle* ClusterHandle : ParticleHandles)
 			{
 				if (ClusterHandle && ClusterHandle->Disabled())
 				{
@@ -2205,7 +2188,7 @@ void FGeometryCollectionPhysicsProxy::GetFilteredParticleHandles(
 		}
 		else if (FilterType == EFieldFilterType::Field_Filter_All)
 		{
-			for (FClusterHandle* ClusterHandle : SolverParticleHandles)
+			for (FClusterHandle* ClusterHandle : ParticleHandles)
 			{
 				if (ClusterHandle && (ClusterHandle->ObjectState() != Chaos::EObjectStateType::Uninitialized))
 				{
@@ -2224,11 +2207,12 @@ void FGeometryCollectionPhysicsProxy::GetRelevantParticleHandles(
 	Handles.SetNum(0, false);
 
 	// only the local handles
-	Handles.Reserve(SolverParticleHandles.Num());
+	TArray<FClusterHandle*>& ParticleHandles = GetSolverParticleHandles();
+	Handles.Reserve(ParticleHandles.Num());
 
 	if (ResolutionType == EFieldResolutionType::Field_Resolution_Maximum)
 	{
-		for (FClusterHandle* ClusterHandle : SolverParticleHandles)
+		for (FClusterHandle* ClusterHandle : ParticleHandles)
 		{
 			if (ClusterHandle )
 			{
@@ -2238,7 +2222,7 @@ void FGeometryCollectionPhysicsProxy::GetRelevantParticleHandles(
 	}
 	else if (ResolutionType == EFieldResolutionType::Field_Resolution_DisabledParents)
 	{
-		for (FClusterHandle* ClusterHandle : SolverParticleHandles)
+		for (FClusterHandle* ClusterHandle : ParticleHandles)
 		{
 			if (ClusterHandle && ClusterHandle->ClusterIds().Id == nullptr)
 			{
@@ -2251,7 +2235,7 @@ void FGeometryCollectionPhysicsProxy::GetRelevantParticleHandles(
 		const auto& Clustering = RigidSolver->GetEvolution()->GetRigidClustering();
 		const auto& ClusterMap = Clustering.GetChildrenMap();
 
-		for (FClusterHandle* ClusterHandle : SolverParticleHandles)
+		for (FClusterHandle* ClusterHandle : ParticleHandles)
 		{
 			if (ClusterHandle && !ClusterHandle->Disabled())
 			{
@@ -2344,10 +2328,8 @@ void FGeometryCollectionPhysicsProxy::DisableParticles_External(TArray<int32>&&
 			{
 				for (int32 TransformIdx : IndicesToDisable)
 				{
-					const int32 ParticleIndex = FromTransformToParticleIndex[TransformIdx];
-					check(ParticleIndex != INDEX_NONE);
-					RBDSolver->GetEvolution()->DisableParticleWithRemovalEvent(SolverParticleHandles[ParticleIndex]);
-					RBDSolver->GetParticles().MarkTransientDirtyParticle(SolverParticleHandles[ParticleIndex]);
+					RBDSolver->GetEvolution()->DisableParticleWithRemovalEvent(SolverParticleHandles[TransformIdx]);
+					RBDSolver->GetParticles().MarkTransientDirtyParticle(SolverParticleHandles[TransformIdx]);
 				}
 			});
 	}
@@ -2504,11 +2486,10 @@ void FGeometryCollectionPhysicsProxy::SetAnchoredByIndex_External(int32 Index, b
 		RBDSolver->EnqueueCommandImmediate([this, Index, bAnchored, RBDSolver]()
 			{
 				Chaos::FPBDRigidsEvolution* Evolution = RBDSolver->GetEvolution();
-				const int32 ParticleIndex = FromTransformToParticleIndex[Index];
-				if (SolverParticleHandles.IsValidIndex(ParticleIndex))
+				if (SolverParticleHandles.IsValidIndex(Index))
 				{
-					SetParticleAnchored_Internal(Evolution, SolverParticleHandles[ParticleIndex], bAnchored);
-					if (Chaos::FPBDRigidClusteredParticleHandle* TopParentHandle = GetTopActiveClusteredParent_Internal(SolverParticleHandles[ParticleIndex]))
+					SetParticleAnchored_Internal(Evolution, SolverParticleHandles[Index], bAnchored);
+					if (Chaos::FPBDRigidClusteredParticleHandle* TopParentHandle = GetTopActiveClusteredParent_Internal(SolverParticleHandles[Index]))
 					{
 						Chaos::UpdateKinematicProperties(TopParentHandle, Evolution->GetRigidClustering().GetChildrenMap(), *Evolution);
 					}
@@ -2533,38 +2514,39 @@ void FGeometryCollectionPhysicsProxy::SetAnchoredByTransformedBox_External(const
 				const int32 MaxLevelToCheck = (InitialLevelAttribute.IsValid()) ? MaxLevel : INDEX_NONE;
 
 				FPBDRigidsEvolution* Evolution = RBDSolver->GetEvolution();
-				for (int32 ParticleIndex = 0; ParticleIndex < SolverParticleHandles.Num(); ParticleIndex++)
+				for (int32 TransformIndex = 0; TransformIndex < SolverParticleHandles.Num(); TransformIndex++)
 				{
-					FClusterHandle* ParticleHandle = SolverParticleHandles[ParticleIndex];
-					check(ParticleHandle != nullptr);
-					const FVec3 PositionInBoxSpace = BoxTransform.InverseTransformPosition(ParticleHandle->X());
-					if (BoundsToCheck.Contains(PositionInBoxSpace))
+					if (FClusterHandle* ParticleHandle = SolverParticleHandles[TransformIndex])
 					{
-						// if we have a max level , we make sure to anchor a parent of the right level 
-						if (MaxLevelToCheck > INDEX_NONE)
+						const FVec3 PositionInBoxSpace = BoxTransform.InverseTransformPosition(ParticleHandle->X());
+						if (BoundsToCheck.Contains(PositionInBoxSpace))
 						{
-							const int32 ParticleLevel = InitialLevelAttribute.Get()[FromParticleToTransformIndex[ParticleIndex]];
-							if (ParticleLevel > MaxLevelToCheck)
+							// if we have a max level , we make sure to anchor a parent of the right level 
+							if (MaxLevelToCheck > INDEX_NONE)
 							{
-								Chaos::FPBDRigidClusteredParticleHandle* ParentParticle = ParticleHandle->Parent();
-								int32 ParentParticleLevel = (ParticleLevel - 1);
-								// we are above level, so we still contribute to the parent matching the level 
-								while (ParentParticle && ParentParticleLevel > MaxLevelToCheck)
+								const int32 ParticleLevel = InitialLevelAttribute.Get()[TransformIndex];
+								if (ParticleLevel > MaxLevelToCheck)
 								{
-									ParentParticle = ParentParticle->Parent();
+									Chaos::FPBDRigidClusteredParticleHandle* ParentParticle = ParticleHandle->Parent();
+									int32 ParentParticleLevel = (ParticleLevel - 1);
+									// we are above level, so we still contribute to the parent matching the level 
+									while (ParentParticle && ParentParticleLevel > MaxLevelToCheck)
+									{
+										ParentParticle = ParentParticle->Parent();
+									}
+									// if we still have a parent , this means the level requirement has been found
+									// let's use this particle to be anchored
+									// if the parent is null, then we don't want to anchor the particle 
+									ParticleHandle = ParentParticle;
 								}
-								// if we still have a parent , this means the level requirement has been found
-								// let's use this particle to be anchored
-								// if the parent is null, then we don't want to anchor the particle 
-								ParticleHandle = ParentParticle;
 							}
-						}
-						if (ParticleHandle)
-						{
-							SetParticleAnchored_Internal(Evolution, ParticleHandle, bAnchored);
-							if (Chaos::FPBDRigidClusteredParticleHandle* TopParentHandle = GetTopActiveClusteredParent_Internal(ParticleHandle))
+							if (ParticleHandle)
 							{
-								TopParentHandles.Add(TopParentHandle);
+								SetParticleAnchored_Internal(Evolution, ParticleHandle, bAnchored);
+								if (Chaos::FPBDRigidClusteredParticleHandle* TopParentHandle = GetTopActiveClusteredParent_Internal(ParticleHandle))
+								{
+									TopParentHandles.Add(TopParentHandle);
+								}
 							}
 						}
 					}
@@ -2585,7 +2567,7 @@ void FGeometryCollectionPhysicsProxy::RemoveAllAnchors_External()
 		RBDSolver->EnqueueCommandImmediate([this, RBDSolver]()
 			{
 				Chaos::FPBDRigidsEvolution* Evolution = RBDSolver->GetEvolution();
-				for (FClusterHandle* ParticleHandle : SolverParticleHandles)
+				for (FClusterHandle* ParticleHandle : GetSolverParticleHandles())
 				{
 					SetParticleAnchored_Internal(Evolution, ParticleHandle, false);
 					// we do not need to update the kinematic state since everything is now dynamic 
@@ -3148,7 +3130,7 @@ void FGeometryCollectionPhysicsProxy::OnRemoveFromScene()
 	}
 
 	const int32 Begin = BaseParticleIndex;
-	const int32 Count = NumTransforms;
+	const int32 Count = NumParticles;
 
 	if (ensure((int32)Particles.Size() > 0 && (Begin + Count) <= (int32)Particles.Size()))
 	{
@@ -3169,7 +3151,7 @@ void FGeometryCollectionPhysicsProxy::SyncBeforeDestroy()
 	{
 		Chaos::FRigidClustering& RigidClustering = RigidSolver->GetEvolution()->GetRigidClustering();
 		Chaos::FClusterUnionManager& ClusterUnionManager = RigidClustering.GetClusterUnionManager();
-		for (Chaos::FPBDRigidClusteredParticleHandle* Handle : SolverParticleHandles)
+		for (Chaos::FPBDRigidClusteredParticleHandle* Handle : GetSolverParticleHandles())
 		{
 			if (Handle)
 			{
@@ -3250,7 +3232,7 @@ void FGeometryCollectionPhysicsProxy::SetWorldTransform_Internal(const FTransfor
 	TSet<FClusterHandle*> ProcessedInternalClusters;
 	const FTransform& ActorToWorld = Parameters.WorldTransform;
 
-	check(NumTransforms == PhysicsThreadCollection.NumElements(FGeometryCollection::TransformGroup));
+	const int32 NumTransforms = PhysicsThreadCollection.NumElements(FGeometryCollection::TransformGroup);
 	if (Chaos::FPhysicsSolver* RigidSolver = GetSolver<Chaos::FPhysicsSolver>())
 	{
 		FClusterUnionManager& ClusterUnionManager = RigidSolver->GetEvolution()->GetRigidClustering().GetClusterUnionManager();
@@ -3262,57 +3244,57 @@ void FGeometryCollectionPhysicsProxy::SetWorldTransform_Internal(const FTransfor
 
 		const TManagedArray<FTransform>& MassToLocal = Parameters.RestCollection->GetAttribute<FTransform>(MassToLocalAttributeName, FTransformCollection::TransformGroup);
 
-		for (int32 ParticleIndex = 0; ParticleIndex < NumEffectiveParticles; ++ParticleIndex)
+		for (int32 TransformGroupIndex = 0; TransformGroupIndex < NumTransforms; ++TransformGroupIndex)
 		{
-			Chaos::FPBDRigidClusteredParticleHandle* Handle = SolverParticleHandles[ParticleIndex];
-			check(Handle != nullptr);
-			FClusterHandle* KinematicRootHandle = nullptr;
-			FClusterHandle* ParentHandle = Handle->Parent();
-
-			if (!Handle->Disabled() && Handle->ObjectState() == Chaos::EObjectStateType::Kinematic)
+			if (Chaos::FPBDRigidClusteredParticleHandle* Handle = SolverParticleHandles[TransformGroupIndex])
 			{
-				KinematicRootHandle = Handle;
-			}
-			else
-			{
-				// is there a internal parent as a kinematic root?
-				if (ParentHandle && ParentHandle->InternalCluster() && !ParentHandle->Disabled() && ParentHandle->ObjectState() == Chaos::EObjectStateType::Kinematic && ParentHandle->PhysicsProxy() == this)
+				FClusterHandle* KinematicRootHandle = nullptr;
+				FClusterHandle* ParentHandle = Handle->Parent();
+
+				if (!Handle->Disabled() && Handle->ObjectState() == Chaos::EObjectStateType::Kinematic)
 				{
-					if (!ProcessedInternalClusters.Contains(ParentHandle))
+					KinematicRootHandle = Handle;
+				}
+				else
+				{
+					// is there a internal parent as a kinematic root?
+					if (ParentHandle && ParentHandle->InternalCluster() && !ParentHandle->Disabled() && ParentHandle->ObjectState() == Chaos::EObjectStateType::Kinematic && ParentHandle->PhysicsProxy() == this)
 					{
-						ProcessedInternalClusters.Add(ParentHandle);
-						KinematicRootHandle = ParentHandle;
+						if (!ProcessedInternalClusters.Contains(ParentHandle))
+						{
+							ProcessedInternalClusters.Add(ParentHandle);
+							KinematicRootHandle = ParentHandle;
+						}
 					}
 				}
-			}
-
-			if (KinematicRootHandle)
-			{
-				const FTransform RootWorldTransform(KinematicRootHandle->R(), KinematicRootHandle->X());
-				const FTransform RootRelativeTransform = RootWorldTransform.GetRelativeTransform(Parameters.PrevWorldTransform);
-				const FTransform WorldTransform = RootRelativeTransform * ActorToWorld;
 
-				SetClusteredParticleKinematicTarget_Internal(KinematicRootHandle, WorldTransform);
-			}
-			else if (ParentHandle && !ParentHandle->IsDynamic())
-			{
-				if (ClusterUnionIndex == INDEX_NONE)
+				if (KinematicRootHandle)
 				{
-					ClusterUnionIndex = ClusterUnionManager.FindClusterUnionIndexFromParticle(Handle);
-				}
+					const FTransform RootWorldTransform(KinematicRootHandle->R(), KinematicRootHandle->X());
+					const FTransform RootRelativeTransform = RootWorldTransform.GetRelativeTransform(Parameters.PrevWorldTransform);
+					const FTransform WorldTransform = RootRelativeTransform * ActorToWorld;
 
-				if (ClusterUnionIndex != INDEX_NONE)
+					SetClusteredParticleKinematicTarget_Internal(KinematicRootHandle, WorldTransform);
+				}
+				else if (ParentHandle && !ParentHandle->IsDynamic())
 				{
-					const int32 TransformGroupIndex = FromParticleToTransformIndex[ParticleIndex];
-					const FTransform ParentWorldTransform{ ParentHandle->R(), ParentHandle->X() };
-					const FTransform NewWorldTransform = MassToLocal[TransformGroupIndex] * FTransform(PhysicsThreadCollection.GetTransform(TransformGroupIndex)) * Parameters.WorldTransform;
-					const FTransform RelativeTransform = NewWorldTransform.GetRelativeTransform(ParentWorldTransform);
+					if (ClusterUnionIndex == INDEX_NONE)
+					{
+						ClusterUnionIndex = ClusterUnionManager.FindClusterUnionIndexFromParticle(Handle);
+					}
+
+					if (ClusterUnionIndex != INDEX_NONE)
+					{
+						const FTransform ParentWorldTransform{ ParentHandle->R(), ParentHandle->X() };
+						const FTransform NewWorldTransform = MassToLocal[TransformGroupIndex] * FTransform(PhysicsThreadCollection.GetTransform(TransformGroupIndex)) * Parameters.WorldTransform;
+						const FTransform RelativeTransform = NewWorldTransform.GetRelativeTransform(ParentWorldTransform);
 
-					DeferredClusterUnionParticleUpdates.Add(Handle);
-					DeferredClusterUnionChildToParentUpdates.Add(RelativeTransform);
+						DeferredClusterUnionParticleUpdates.Add(Handle);
+						DeferredClusterUnionChildToParentUpdates.Add(RelativeTransform);
 
-					// Make sure the particle is mark dirty to make sure its proxy will properly update the transforms
-					RigidSolver->GetEvolution()->GetParticles().MarkTransientDirtyParticle(Handle);
+						// Make sure the particle is mark dirty to make sure its proxy will properly update the transforms
+						RigidSolver->GetEvolution()->GetParticles().MarkTransientDirtyParticle(Handle);
+					}
 				}
 			}
 		}
@@ -3346,18 +3328,16 @@ void FGeometryCollectionPhysicsProxy::SetUseStaticMeshCollisionForTraces_Externa
 	const bool bUseSMCollisionForTraces = ((bInUseStaticMeshCollisionForTraces && ForceOverrideGCCollisionSetupForTraces == GCCSFT_Property) || (ForceOverrideGCCollisionSetupForTraces == GCCSFT_ForceSM));
 	if (bUseStaticMeshCollisionForTraces != bUseSMCollisionForTraces && CreateTraceCollisionGeometryCallback != nullptr)
 	{
-		const int32 TransformIndex = Parameters.InitialRootIndex;
-		const int32 ParticleIndex = FromTransformToParticleIndex[TransformIndex];
-		if (GTParticles.IsValidIndex(ParticleIndex))
+		const int32 Index = Parameters.InitialRootIndex;
+		if (GTParticles.Num() > Index)
 		{
-			FParticle* P = GTParticles[ParticleIndex].Get();
-			if (P != nullptr)
+			if (FParticle* P = GTParticles[Index].Get())
 			{
 				if (bUseSMCollisionForTraces)
 				{
 					const TManagedArray<FTransform>& MassToLocal = Parameters.RestCollection->GetAttribute<FTransform>("MassToLocal", FGeometryCollection::TransformGroup);
 
-					const FTransform ToLocal = MassToLocal[TransformIndex].Inverse();
+					const FTransform ToLocal = MassToLocal[Index].Inverse();
 					TArray<Chaos::FImplicitObjectPtr> Geoms;
 					Chaos::FShapesArray Shapes;
 					CreateTraceCollisionGeometryCallback(ToLocal, Geoms, Shapes);
@@ -3369,16 +3349,17 @@ void FGeometryCollectionPhysicsProxy::SetUseStaticMeshCollisionForTraces_Externa
 				{
 					const FVector Scale = WorldTransform_External.GetScale3D();
 					TManagedArray<Chaos::FImplicitObjectPtr>& Implicits = GameThreadCollection.ModifyAttribute<Chaos::FImplicitObjectPtr>(FGeometryDynamicCollection::ImplicitsAttribute, FTransformCollection::TransformGroup);
-					Chaos::FImplicitObjectPtr ImplicitGeometry = Implicits[TransformIndex];
+					Chaos::FImplicitObjectPtr ImplicitGeometry = Implicits[Index];
 					if (ImplicitGeometry && !Scale.Equals(FVector::OneVector))
 					{
 						ImplicitGeometry = ImplicitGeometry->CopyGeometryWithScale(Scale);
 					}
 					P->SetGeometry(ImplicitGeometry);
 				}
+
+				// We only want to change the value of our boolean if we actually changed the collision, so the collision state matches
+				bUseStaticMeshCollisionForTraces = bUseSMCollisionForTraces;
 			}
-			// We only want to change the value of our boolean if we actually changed the collision, so the collision state matches
-			bUseStaticMeshCollisionForTraces = bUseSMCollisionForTraces;
 		}
 	}
 }
@@ -3471,19 +3452,20 @@ void FGeometryCollectionPhysicsProxy::SetMaterialOverrideMassScaleMultiplier_Int
 	}
 	Parameters.MaterialOverrideMassScaleMultiplier = NewValue;
 
-	for (int32 ParticleIndex = 0; ParticleIndex < SolverParticleHandles.Num(); ++ParticleIndex)
+	for (int32 TransformIndex = 0; TransformIndex < SolverParticleHandles.Num(); ++TransformIndex)
 	{
-		Chaos::FPBDRigidClusteredParticleHandle* Handle = SolverParticleHandles[ParticleIndex];
-		check(Handle != nullptr);
-		const Chaos::FReal NewM = Handle->M() * MaterialOverrideMassScaleMultiplierChange;
-		const Chaos::FVec3 NewI = Handle->I() * MaterialOverrideMassScaleMultiplierChange;
+		if (Chaos::FPBDRigidClusteredParticleHandle* Handle = SolverParticleHandles[TransformIndex])
+		{
+			const Chaos::FReal NewM = Handle->M() * MaterialOverrideMassScaleMultiplierChange;
+			const Chaos::FVec3 NewI = Handle->I() * MaterialOverrideMassScaleMultiplierChange;
 
-		Handle->SetM(NewM);
-		Handle->SetI(NewI);
-		const Chaos::FReal InvM = (NewM > 0.0f) ? 1.0f / NewM : 0.0f;
-		const Chaos::FVec3 InvI = (NewM > 0.0f) ? Chaos::FVec3(NewI).Reciprocal() : Chaos::FVec3::ZeroVector;
-		Handle->SetInvM(InvM);
-		Handle->SetInvI(InvI);
+			Handle->SetM(NewM);
+			Handle->SetI(NewI);
+			const Chaos::FReal InvM = (NewM > 0.0f) ? 1.0f / NewM : 0.0f;
+			const Chaos::FVec3 InvI = (NewM > 0.0f) ? Chaos::FVec3(NewI).Reciprocal() : Chaos::FVec3::ZeroVector;
+			Handle->SetInvM(InvM);
+			Handle->SetInvI(InvI);
+		}
 	}
 }
 
@@ -3518,7 +3500,7 @@ void FGeometryCollectionPhysicsProxy::SetOneWayInteractionLevel_External(int32 I
 		TUniquePtr<FParticle>& Particle = GTParticles[ParticleIndex];
 		if (Particle.IsValid())
 		{
-			const bool bIsOneWayInteraction = (InOneWayInteractionLevel >= 0) && (Level[FromParticleToTransformIndex[ParticleIndex]] >= InOneWayInteractionLevel);
+			const bool bIsOneWayInteraction = (InOneWayInteractionLevel >= 0) && (Level[ParticleIndex] >= InOneWayInteractionLevel);
 			Particle->SetOneWayInteraction(bIsOneWayInteraction);
 		}
 	}
@@ -3745,7 +3727,7 @@ void FGeometryCollectionPhysicsProxy::BufferPhysicsResults_Internal(Chaos::FPBDR
 
 	UniqueIdxToInternalClusterHandle.Reset();
 	
-	check(NumTransforms == PhysicsThreadCollection.GetNumTransforms());
+	const int32 NumTransforms = PhysicsThreadCollection.GetNumTransforms();
 	if(NumTransforms > 0)
 	{ 
 		SCOPE_CYCLE_COUNTER(STAT_CalcParticleToWorld);
@@ -3759,7 +3741,7 @@ void FGeometryCollectionPhysicsProxy::BufferPhysicsResults_Internal(Chaos::FPBDR
 		// historically assume that the root has no parent which is not true any more).
 		if (Parameters.InitialRootIndex != INDEX_NONE)
 		{
-			const Chaos::FPBDRigidClusteredParticleHandle* RootHandle = SolverParticleHandles[FromTransformToParticleIndex[Parameters.InitialRootIndex]];
+			const Chaos::FPBDRigidClusteredParticleHandle* RootHandle = SolverParticleHandles[Parameters.InitialRootIndex];
 			if (RootHandle != nullptr)
 			{
 				const bool bIsRootBroken = RootHandle->Disabled() && (RootHandle->Parent() == nullptr);
@@ -3767,10 +3749,9 @@ void FGeometryCollectionPhysicsProxy::BufferPhysicsResults_Internal(Chaos::FPBDR
 			}
 		}
 
-		for (int32 ParticleIndex = 0; ParticleIndex < NumEffectiveParticles; ++ParticleIndex)
+		for (int32 TransformGroupIndex = 0; TransformGroupIndex < NumTransforms; ++TransformGroupIndex)
 		{
-			const int32 TransformGroupIndex = FromParticleToTransformIndex[ParticleIndex];
-			Chaos::FPBDRigidClusteredParticleHandle* Handle = SolverParticleHandles[ParticleIndex];
+			Chaos::FPBDRigidClusteredParticleHandle* Handle = SolverParticleHandles[TransformGroupIndex];
 			if (!Handle)
 			{
 				PhysicsThreadCollection.Active[TransformGroupIndex] = false;
@@ -3817,7 +3798,7 @@ void FGeometryCollectionPhysicsProxy::BufferPhysicsResults_Internal(Chaos::FPBDR
 			Chaos::FPBDRigidClusteredParticleHandle* ClusterParent = Handle->Parent();
 
 			// Has parent changed? ( can be a new one or nullptr)
-			const bool bHasNewParent = (SolverClusterID[ParticleIndex] != ClusterParent) || (!ClusterParent && PhysicsThreadCollection.GetHasParent(TransformGroupIndex));
+			const bool bHasNewParent = (SolverClusterID[TransformGroupIndex] != ClusterParent) || (!ClusterParent && PhysicsThreadCollection.GetHasParent(TransformGroupIndex));
 			if (bHasNewParent)
 			{
 				// Force all driven rigid bodies out of the transform hierarchy ( because the new parent can only be an internal cluster )
@@ -3834,7 +3815,7 @@ void FGeometryCollectionPhysicsProxy::BufferPhysicsResults_Internal(Chaos::FPBDR
 					bHasChanged = true;
 				}
 				// set new parent ( likely an internal cluster )
-				SolverClusterID[ParticleIndex] = ClusterParent;
+				SolverClusterID[TransformGroupIndex] = ClusterParent;
 			}
 
 			// do we have an internal cluster parent ( it also implies we are a disabled particle )
@@ -4127,7 +4108,7 @@ bool FGeometryCollectionPhysicsProxy::PullNonInterpolatableDataFromSinglePhysics
 			continue;
 		}
 
-		const int32 ParticleIndex = FromTransformToParticleIndex[TransformGroupIndex];
+		
 		const bool bIsActive = !StateData.State.DisabledState;
 		if (UpdateValue(GameThreadCollection.Active[TransformGroupIndex], bIsActive))
 		{
@@ -4135,20 +4116,19 @@ bool FGeometryCollectionPhysicsProxy::PullNonInterpolatableDataFromSinglePhysics
 			{
 				CreateChildrenGeometry_External();
 			}
-			if (ParticleIndex == INDEX_NONE || GTParticles[ParticleIndex] == nullptr)
+			if (GTParticles[TransformGroupIndex] == nullptr)
 			{
 				continue;
 			}
-
-			GTParticles[ParticleIndex]->SetDisabled(!bIsActive);
+			GTParticles[TransformGroupIndex]->SetDisabled(!bIsActive);
 			bIsCollectionDirty = true;
 		}
 
-		if (ParticleIndex == INDEX_NONE || GTParticles[ParticleIndex] == nullptr)
+		if (GTParticles[TransformGroupIndex] == nullptr)
 		{
 			continue;
 		}
-		FParticle& GTParticle = *GTParticles[ParticleIndex];
+		FParticle& GTParticle = *GTParticles[TransformGroupIndex];
 
 		if (UpdateValue(GameThreadCollection.DynamicState[TransformGroupIndex], static_cast<uint8>(StateData.State.DynamicState)))
 		{
@@ -4203,7 +4183,7 @@ bool FGeometryCollectionPhysicsProxy::PullNonInterpolatableDataFromSinglePhysics
 		// internal cluster index map update
 		if (StateData.InternalClusterUniqueIdx > INDEX_NONE)
 		{
-			GTParticlesToInternalClusterUniqueIdx.Add(GTParticles[ParticleIndex].Get(), StateData.InternalClusterUniqueIdx);
+			GTParticlesToInternalClusterUniqueIdx.Add(GTParticles[TransformGroupIndex].Get(), StateData.InternalClusterUniqueIdx);
 			InternalClusterUniqueIdxToChildrenTransformIndices.FindOrAdd(StateData.InternalClusterUniqueIdx).Add(TransformGroupIndex);
 		}
 	}
@@ -4214,11 +4194,9 @@ bool FGeometryCollectionPhysicsProxy::PullNonInterpolatableDataFromSinglePhysics
 	if ((RootIndex != INDEX_NONE) && CurrentResults.IsRootBroken && GameThreadCollection.Active[RootIndex] && bGeometryCollectionUseRootBrokenFlag)
 	{
 		GameThreadCollection.Active[RootIndex] = false;
-		const int32 ParticleIndex = FromTransformToParticleIndex[RootIndex];
-		check(ParticleIndex != INDEX_NONE);
-		if (GTParticles[ParticleIndex] != nullptr)
+		if (GTParticles[RootIndex] != nullptr)
 		{
-			GTParticles[ParticleIndex]->SetDisabled(true);
+			GTParticles[RootIndex]->SetDisabled(true);
 			bIsCollectionDirty = true;
 		}
 	}
@@ -4257,7 +4235,7 @@ bool FGeometryCollectionPhysicsProxy::PullFromPhysicsState(const Chaos::FDirtyGe
 		InterpolationData.AccumlateErrorXR(Error->ErrorX, Error->ErrorR, SolverSyncTimestamp, RenderInterpErrorCorrectionDurationTicks);
 	}
 
-	check(NumTransforms == GameThreadCollection.GetNumTransforms());
+	const int32 NumTransforms = GameThreadCollection.GetNumTransforms();
 	const bool bNeedInterpolation = (NextPullData != nullptr);
 
 	bool bIsCollectionDirty = false;
@@ -4293,9 +4271,8 @@ bool FGeometryCollectionPhysicsProxy::PullFromPhysicsState(const Chaos::FDirtyGe
 			// for that case we cannot just go through the list of entries since Results and NextResults may have different number of entries that don't always match
 			// so we need to go through the transform indices and find the matching entries on both side 
 			// this is helped by the FResultAccessor structure
-			for (int32 ParticleIndex = 0; ParticleIndex < NumEffectiveParticles; ++ParticleIndex)
+			for (int32 TransformGroupIndex = 0; TransformGroupIndex < NumTransforms; ++TransformGroupIndex)
 			{
-				const int32 TransformGroupIndex = FromParticleToTransformIndex[ParticleIndex];
 				const bool bActive = GameThreadCollection.Active[TransformGroupIndex];
 				bAtLeatOneChildActive |= bActive && TransformGroupIndex != Parameters.InitialRootIndex;
 				if (bAtLeatOneChildActive)
@@ -4303,12 +4280,12 @@ bool FGeometryCollectionPhysicsProxy::PullFromPhysicsState(const Chaos::FDirtyGe
 					CreateChildrenGeometry_External();
 				}
 
-				if (GTParticles[ParticleIndex] == nullptr)
+				if (GTParticles[TransformGroupIndex] == nullptr)
 				{
 					continue;
 				}
 
-				FParticle& GTParticle = *GTParticles[ParticleIndex];
+				FParticle& GTParticle = *GTParticles[TransformGroupIndex];
 				const FResultInterpolator ResultInterpolator(TransformGroupIndex, PrevResults, NextResults, GTParticle, *Alpha);
 				if (ResultInterpolator.HasNoEntry())
 				{
@@ -4353,14 +4330,13 @@ bool FGeometryCollectionPhysicsProxy::PullFromPhysicsState(const Chaos::FDirtyGe
 			const bool bIsComponentTransformScaled = !WorldTransform_External.GetScale3D().Equals(FVector::OneVector);
 			const FTransform ComponentScaleTransform(FQuat::Identity, FVector::ZeroVector, WorldTransform_External.GetScale3D());
 
-			for (int32 ParticleIndex = 0; ParticleIndex < NumEffectiveParticles; ParticleIndex++)
+			for (int32 TransformIndex = 0; TransformIndex < NumTransforms; TransformIndex++)
 			{
-				const int32 TransformIndex = FromParticleToTransformIndex[ParticleIndex];
 				const bool bAnimatingWhileDisabled = AnimationsActive ? (*AnimationsActive)[TransformIndex] : false;
 				const bool bActive = GameThreadCollection.Active[TransformIndex];
 				if (bActive || bAnimatingWhileDisabled)
 				{
-					if (GTParticles[ParticleIndex] != nullptr)
+					if (GTParticles[TransformIndex] != nullptr)
 					{
 						bool bWasModified = PrevResultsModifiedIndices.IsValidIndex(TransformIndex) ? PrevResultsModifiedIndices[TransformIndex] : false;
 						if (NextPullData)
@@ -4370,7 +4346,7 @@ bool FGeometryCollectionPhysicsProxy::PullFromPhysicsState(const Chaos::FDirtyGe
 						}
 						if (bWasModified)
 						{
-							const FParticle& GTParticle = *GTParticles[ParticleIndex];
+							const FParticle& GTParticle = *GTParticles[TransformIndex];
 							const FTransform& ParticleMassToLocal = MassToLocal[TransformIndex];
 							const int32 ParentTransformIndex = GameThreadCollection.GetParent(TransformIndex);
 
@@ -4380,7 +4356,7 @@ bool FGeometryCollectionPhysicsProxy::PullFromPhysicsState(const Chaos::FDirtyGe
 							FTransform ParentWorldTransform = WorldTransform_External;
 							if (ParentTransformIndex != INDEX_NONE)
 							{
-								if (const FParticle* GTParentParticle = (ParentTransformIndex != INDEX_NONE) ? GTParticles[FromTransformToParticleIndex[ParentTransformIndex]].Get() : nullptr)
+								if (const FParticle* GTParentParticle = (ParentTransformIndex != INDEX_NONE) ? GTParticles[ParentTransformIndex].Get() : nullptr)
 								{
 									const FTransform& ParentMassToLocal = MassToLocal[TransformIndex];
 									ParentWorldTransform = ParentMassToLocal.Inverse() * FTransform { GTParentParticle->R(), GTParentParticle->X() };
@@ -4441,7 +4417,8 @@ void FGeometryCollectionPhysicsProxy::UpdateFilterData_External(const FCollision
 	check(IsInGameThread());
 
 	// SimFilter/QueryFilter members are read on both threads, these are const after initialization and are not updated here.
-	for (int32 Index = 0; Index < NumEffectiveParticles; ++Index)
+	const int32 NumTransforms = GameThreadCollection.GetNumTransforms();
+	for (int32 Index = 0; Index < NumTransforms; ++Index)
 	{
 		FParticle* P = GTParticles[Index].Get();
 		if (P != nullptr)
@@ -4475,30 +4452,32 @@ void FGeometryCollectionPhysicsProxy::SetFilterData_Internal(const FCollisionFil
 	{
 		FRigidClustering& RigidClustering = RigidSolver->GetEvolution()->GetRigidClustering();
 
-		for (int32 ParticleIndex = 0; ParticleIndex < NumEffectiveParticles; ++ParticleIndex)
+		const int32 NumTransforms = PhysicsThreadCollection.NumElements(FGeometryCollection::TransformGroup);
+		for (int32 TransformIndex = 0; TransformIndex < NumTransforms; ++TransformIndex)
 		{
-			Chaos::FPBDRigidClusteredParticleHandle* Handle = SolverParticleHandles[ParticleIndex];
-			check(Handle != nullptr);
-			// Must update our filters before updating an internal cluster parent
-			const Chaos::FShapesArray& ShapesArray = Handle->ShapesArray();
-			for (const TUniquePtr<Chaos::FPerShapeData>& Shape : ShapesArray)
+			if (Chaos::FPBDRigidClusteredParticleHandle* Handle = SolverParticleHandles[TransformIndex])
 			{
-				Shape->SetQueryData(Parameters.QueryFilterData);
-				Shape->SetSimData(Parameters.SimulationFilterData);
-			}
+				// Must update our filters before updating an internal cluster parent
+				const Chaos::FShapesArray& ShapesArray = Handle->ShapesArray();
+				for (const TUniquePtr<Chaos::FPerShapeData>& Shape : ShapesArray)
+				{
+					Shape->SetQueryData(Parameters.QueryFilterData);
+					Shape->SetSimData(Parameters.SimulationFilterData);
+				}
 
-			FClusterHandle* ParentHandle = Handle->Parent();
-			if (ParentHandle && ParentHandle->InternalCluster() && ParentHandle->PhysicsProxy() == this)
-			{
-				if (!ProcessedInternalClusters.Contains(ParentHandle))
+				FClusterHandle* ParentHandle = Handle->Parent();
+				if (ParentHandle && ParentHandle->InternalCluster() && ParentHandle->PhysicsProxy() == this)
 				{
-					ProcessedInternalClusters.Add(ParentHandle);
-					// Must update our filters before updating an internal cluster parent
-					FRigidClustering::FRigidHandleArray* ChildArray = RigidClustering.GetChildrenMap().Find(ParentHandle);
-					if (ensure(ChildArray))
+					if (!ProcessedInternalClusters.Contains(ParentHandle))
 					{
-						// If our filter changed, internal cluster parent's filter may be stale.
-						UpdateClusterFilterDataFromChildren(ParentHandle, *ChildArray);
+						ProcessedInternalClusters.Add(ParentHandle);
+						// Must update our filters before updating an internal cluster parent
+						FRigidClustering::FRigidHandleArray* ChildArray = RigidClustering.GetChildrenMap().Find(ParentHandle);
+						if (ensure(ChildArray))
+						{
+							// If our filter changed, internal cluster parent's filter may be stale.
+							UpdateClusterFilterDataFromChildren(ParentHandle, *ChildArray);
+						}
 					}
 				}
 			}
@@ -4528,22 +4507,21 @@ void FGeometryCollectionPhysicsProxy::SetPerParticleFilterData_Internal(const TA
 		{
 			if (Data.bIsValid && Data.ParticleIndex != INDEX_NONE)
 			{
-				const int32 ParticleIndex = FromTransformToParticleIndex[Data.ParticleIndex];
-				if (ParticleIndex != INDEX_NONE)
+				if (Chaos::FPhysicsObjectHandle Object = PhysicsObjects[Data.ParticleIndex].Get())
 				{
-					Chaos::FPhysicsObjectHandle Object = PhysicsObjects[ParticleIndex].Get();
-					check(Object != nullptr);
 					TArrayView<Chaos::FPhysicsObjectHandle> ParticleView{ &Object, 1 };
 					Interface.UpdateShapeCollisionFlags(ParticleView, Data.bSimEnabled, Data.bQueryEnabled);
 					Interface.UpdateShapeFilterData(ParticleView, Data.QueryFilter, Data.SimFilter);
-					// todo(chaos): It's not ideal but the geometry collection needs to request the cluster union update its
-					// cached shape data if the particle is in a cluster union. This is because we don't share shape
-					// data between the GC shapes and the cluster union shapes.
-					Chaos::FRigidClustering& Clustering = RBDSolver->GetEvolution()->GetRigidClustering();
-					Chaos::FClusterUnionManager& ClusterUnionManager = Clustering.GetClusterUnionManager();
-
-					Chaos::FPBDRigidClusteredParticleHandle* Particle = SolverParticleHandles[ParticleIndex];
-					check(Particle != nullptr);
+				}
+
+				// todo(chaos): It's not ideal but the geometry collection needs to request the cluster union update its
+				// cached shape data if the particle is in a cluster union. This is because we don't share shape
+				// data between the GC shapes and the cluster union shapes.
+				Chaos::FRigidClustering& Clustering = RBDSolver->GetEvolution()->GetRigidClustering();
+				Chaos::FClusterUnionManager& ClusterUnionManager = Clustering.GetClusterUnionManager();
+
+				if (Chaos::FPBDRigidClusteredParticleHandle* Particle = SolverParticleHandles[Data.ParticleIndex])
+				{
 					if (Chaos::FClusterUnion* ClusterUnion = ClusterUnionManager.FindClusterUnionFromParticle(Particle))
 					{
 						ClusterUnion->AddPendingGeometryOperation(Chaos::EClusterUnionGeometryOperation::Refresh, Particle);
@@ -5778,9 +5756,9 @@ void FGeometryCollectionPhysicsProxy::FieldForcesUpdateCallback(Chaos::FPBDRigid
 
 FGeometryCollectionPhysicsProxy::FClusterHandle* FGeometryCollectionPhysicsProxy::GetInitialRootHandle_Internal() const
 {
-	if (SolverParticleHandles.IsValidIndex(FromTransformToParticleIndex[Parameters.InitialRootIndex]))
+	if (SolverParticleHandles.IsValidIndex(Parameters.InitialRootIndex))
 	{
-		return SolverParticleHandles[FromTransformToParticleIndex[Parameters.InitialRootIndex]];
+		return SolverParticleHandles[Parameters.InitialRootIndex];
 	}
 	return nullptr;
 }
@@ -5797,40 +5775,31 @@ TArray<Chaos::FPhysicsObjectHandle> FGeometryCollectionPhysicsProxy::GetAllPhysi
 			Handles.Add(Object.Get());
 		}
 	}
+
 	return Handles;
 }
 
 TArray<Chaos::FPhysicsObjectHandle> FGeometryCollectionPhysicsProxy::GetAllPhysicsObjectIncludingNulls() const
 {
 	TArray<Chaos::FPhysicsObjectHandle> Handles;
-	Handles.SetNum(NumTransforms);
+	Handles.Reserve(NumEffectiveParticles);
 
-	for (int32 TransformIndex = 0; TransformIndex < NumTransforms; ++TransformIndex)
+	for (const Chaos::FPhysicsObjectUniquePtr& Object : PhysicsObjects)
 	{
-		const int32 ParticleIndex = FromTransformToParticleIndex[TransformIndex];
-		if (ParticleIndex == INDEX_NONE)
-		{
-			Handles[TransformIndex] = nullptr;
-		}
-		else
-		{
-			Handles[TransformIndex] = PhysicsObjects[ParticleIndex].Get();
-		}
+		Handles.Add(Object.Get());
 	}
+
 	return Handles;
 }
 
 Chaos::FPhysicsObjectHandle FGeometryCollectionPhysicsProxy::GetPhysicsObjectByIndex(int32 Index) const
 {
-	if (FromTransformToParticleIndex.IsValidIndex(Index))
+	if (!PhysicsObjects.IsValidIndex(Index))
 	{
-		const int32 ParticleIndex = FromTransformToParticleIndex[Index];
-		if (ParticleIndex != INDEX_NONE)
-		{
-			return PhysicsObjects[ParticleIndex].Get();
-		}
+		return nullptr;
 	}
-	return nullptr; 
+
+	return PhysicsObjects[Index].Get();
 }
 
 FGeometryCollectionPhysicsProxy::FParticle* FGeometryCollectionPhysicsProxy::GetParticleByIndex_External(int32 Index)
@@ -5840,15 +5809,11 @@ FGeometryCollectionPhysicsProxy::FParticle* FGeometryCollectionPhysicsProxy::Get
 
 const FGeometryCollectionPhysicsProxy::FParticle* FGeometryCollectionPhysicsProxy::GetParticleByIndex_External(int32 Index) const
 {
-	if (FromTransformToParticleIndex.IsValidIndex(Index))
+	if (Index < 0 || Index >= GTParticles.Num())
 	{
-		const int32 ParticleIndex = FromTransformToParticleIndex[Index];
-		if (ParticleIndex != INDEX_NONE)
-		{
-			return GTParticles[ParticleIndex].Get();
-		}
+		return nullptr;
 	}
-	return nullptr;
+	return GTParticles[Index].Get();
 }
 
 FGeometryCollectionPhysicsProxy::FParticleHandle* FGeometryCollectionPhysicsProxy::GetParticleByIndex_Internal(int32 Index)
@@ -5858,15 +5823,11 @@ FGeometryCollectionPhysicsProxy::FParticleHandle* FGeometryCollectionPhysicsProx
 
 const FGeometryCollectionPhysicsProxy::FParticleHandle* FGeometryCollectionPhysicsProxy::GetParticleByIndex_Internal(int32 Index) const
 {
-	if (FromTransformToParticleIndex.IsValidIndex(Index))
+	if (!SolverParticleHandles.IsValidIndex(Index))
 	{
-		const int32 ParticleIndex = FromTransformToParticleIndex[Index];
-		if (ParticleIndex != INDEX_NONE)
-		{
-			return SolverParticleHandles[ParticleIndex];
-		}
+		return nullptr;
 	}
-	return nullptr;
+	return SolverParticleHandles[Index];
 }
 
 void FDamageCollector::Reset(int32 NumTransforms)
diff --git a/Engine/Source/Runtime/Experimental/Chaos/Public/PhysicsProxy/GeometryCollectionPhysicsProxy.h b/Engine/Source/Runtime/Experimental/Chaos/Public/PhysicsProxy/GeometryCollectionPhysicsProxy.h
--- a/Engine/Source/Runtime/Experimental/Chaos/Public/PhysicsProxy/GeometryCollectionPhysicsProxy.h
+++ b/Engine/Source/Runtime/Experimental/Chaos/Public/PhysicsProxy/GeometryCollectionPhysicsProxy.h
@@ -19,7 +19,6 @@
 #include "Chaos/Defines.h"
 #include "Chaos/GeometryParticlesfwd.h"
 
-
 namespace Chaos
 {
 	template <typename T> class TSerializablePtr;
@@ -237,22 +236,10 @@ public:
 	void SetCollisionParticlesPerObjectFraction(float CollisionParticlesPerObjectFractionIn) 
 	{CollisionParticlesPerObjectFraction = CollisionParticlesPerObjectFractionIn;}
 
-	UE_DEPRECATED(5.4, "Use GetSolverClusterHandle_Internal instead")
 	TArray<FClusterHandle*>& GetSolverClusterHandles() { return SolverClusterHandles; }
 
-	UE_DEPRECATED(5.4, "Use GetParticle_Internal instead")
 	TArray<FClusterHandle*>& GetSolverParticleHandles() { return SolverParticleHandles; }
 
-	FClusterHandle* GetSolverClusterHandle_Internal(int32 Index) const
-	{
-		const int32 ParticleIndex = FromTransformToParticleIndex[Index];
-		if (ParticleIndex != INDEX_NONE)
-		{
-			return SolverClusterHandles[ParticleIndex];
-		}
-		return nullptr;
-	}
-
 	const FGeometryCollectionResults* GetConsumerResultsGT() const 
 	{ return PhysToGameInterchange.PeekConsumerBuffer(); }
 
@@ -290,27 +277,11 @@ public:
 	// Set whether the GC should be using collision from the Static Mesh or the GC itself for game thread traces ( this needs to be called on the game thread )
 	CHAOS_API void SetUseStaticMeshCollisionForTraces_External(bool bInUseStaticMeshCollisionForTraces);
 
-	UE_DEPRECATED(5.4, "Use GetParticle_Internal instead")
 	const TArray<FClusterHandle*> GetParticles() const
 	{
 		return SolverParticleHandles;
 	}
 
-	const TArray<FClusterHandle*>GetUnorderedParticles_Internal() const
-	{
-		return SolverParticleHandles;
-	}
-
-	FClusterHandle* GetParticle_Internal(int32 Index) const
-	{
-		const int32 ParticleIndex = FromTransformToParticleIndex[Index];
-		if (ParticleIndex != INDEX_NONE)
-		{
-			return SolverParticleHandles[ParticleIndex];
-		}
-		return nullptr;
-	}
-
 	const FSimulationParameters& GetSimParameters() const
 	{
 		return Parameters;
@@ -331,17 +302,11 @@ public:
 		return GameThreadCollection;
 	}
 
-	UE_DEPRECATED(5.4, "Use GetUnorderedParticles_External instead")
 	TArray<TUniquePtr<FParticle>>& GetExternalParticles()
 	{
 		return GTParticles;
 	}
 
-	TArray<TUniquePtr<FParticle>>& GetUnorderedParticles_External()
-	{
-		return GTParticles;
-	}
-
 	CHAOS_API FParticle* GetParticleByIndex_External(int32 Index);
 	CHAOS_API const FParticle* GetParticleByIndex_External(int32 Index) const;
 
@@ -426,7 +391,7 @@ public:
 		{
 			return FGeometryCollectionItemIndex::CreateInternalClusterItemIndex(*InternalClusterUniqueIdx);
 		}
-		// regular particle that has a matching transform index 
+		// regular particle that has a matchig transform index 
 		if (const int32* TransformGroupIndex = GTParticlesToTransformGroupIndex.Find(GTPParticle))
 		{
 			return FGeometryCollectionItemIndex::CreateTransformItemIndex(*TransformGroupIndex);
@@ -529,10 +494,7 @@ public:
 	CHAOS_API TArray<Chaos::FPhysicsObjectHandle> GetAllPhysicsObjects() const ;
 	CHAOS_API TArray<Chaos::FPhysicsObjectHandle> GetAllPhysicsObjectIncludingNulls() const;
 	CHAOS_API Chaos::FPhysicsObjectHandle GetPhysicsObjectByIndex(int32 Index) const;
-
-	UE_DEPRECATED(5.4, "Use GetNumTransforms instead")
-	int32 GetNumParticles() const { return NumTransforms; }
-	int32 GetNumTransforms() const { return NumTransforms; }
+	int32 GetNumParticles() const { return NumParticles; }
 
 	// todo(chaos): Remove this and move to a cook time approach of the SM data based on the GC property
 	using FCreateTraceCollisionGeometryCallback = TFunction<void(const FTransform& InToLocal, TArray<Chaos::FImplicitObjectPtr>& OutGeoms, Chaos::FShapesArray& OutShapes)>;
@@ -540,8 +502,6 @@ public:
 
 	CHAOS_API void CreateChildrenGeometry_Internal();
 
-	int32 GetFromParticleToTransformIndex(int32 Index) const { check(FromParticleToTransformIndex.IsValidIndex(Index));  return FromParticleToTransformIndex[Index]; }
-
 protected:
 
 	CHAOS_API float ComputeMaterialBasedDamageThreshold_Internal(int32 TransformIndex) const;
@@ -647,24 +607,23 @@ private:
 	//
 	//  Proxy State Information
 	//
-	int32 NumTransforms;
+	int32 NumParticles;
 	int32 NumEffectiveParticles;
 	int32 BaseParticleIndex;
 	TArray<FParticleHandle*> SolverClusterID;
-	TArray<FClusterHandle*> SolverClusterHandles; // make a TArray of the base class with type
+	TArray<FClusterHandle*> SolverClusterHandles; // make a TArray of the base clase with type
 	TArray<FClusterHandle*> SolverParticleHandles;// make a TArray of base class and join with above
+	TSet<FClusterHandle*> SolverAnchors;
 	TMap<FParticleHandle*, int32> HandleToTransformGroupIndex;
 	TMap<int32, FClusterHandle*> UniqueIdxToInternalClusterHandle;
 	TArray<Chaos::FUniqueIdx> UniqueIdxs;
-	TArray<int32> FromParticleToTransformIndex;
-	TArray<int32> FromTransformToParticleIndex;
 
 	//
 	// Buffer Results State Information
 	//
 	bool IsObjectDynamic; // Records current dynamic state
 	bool IsObjectLoading; // Indicate when loaded
-	bool IsObjectDeleting; // Indicate when pending deletion
+	bool IsObjectDeleting; // Indicatge when pending deletion
 
 	EReplicationMode ReplicationMode = EReplicationMode::Unknown;	
 
diff --git a/Engine/Source/Runtime/Experimental/ChaosVehicles/ChaosVehiclesCore/Private/SimModule/DeferredForcesModular.cpp b/Engine/Source/Runtime/Experimental/ChaosVehicles/ChaosVehiclesCore/Private/SimModule/DeferredForcesModular.cpp
--- a/Engine/Source/Runtime/Experimental/ChaosVehicles/ChaosVehiclesCore/Private/SimModule/DeferredForcesModular.cpp
+++ b/Engine/Source/Runtime/Experimental/ChaosVehicles/ChaosVehiclesCore/Private/SimModule/DeferredForcesModular.cpp
@@ -6,8 +6,6 @@
 #include "Chaos/Utilities.h"
 #include "Chaos/ParticleHandle.h"
 
-#include "PhysicsProxy/GeometryCollectionPhysicsProxy.h"
-
 #if CHAOS_DEBUG_DRAW
 #include "Chaos/DebugDrawQueue.h"
 #endif
@@ -38,7 +36,8 @@ Chaos::FPBDRigidParticleHandle* FDeferredForcesModular::GetParticleFromUniqueInd
 }
 
 // #TODO: passing all these parameters in is horrible, tidy this up now we know what needs to be done
-Chaos::FPBDRigidParticleHandle* FDeferredForcesModular::GetParticle(FGeometryCollectionPhysicsProxy* Proxy
+Chaos::FPBDRigidParticleHandle* FDeferredForcesModular::GetParticle(TArray<Chaos::FPBDRigidClusteredParticleHandle*>& Particles
+		, TArray<Chaos::FPBDRigidClusteredParticleHandle*>& ClusterParticles
 		, int TransformIndex
 		, int32 ParticleIdx
 		, const FVector& PositionalOffset
@@ -48,8 +47,8 @@ Chaos::FPBDRigidParticleHandle* FDeferredForcesModular::GetParticle(FGeometryCol
 		, FTransform& TransformOut)
 {
 
-	Chaos::FPBDRigidParticleHandle* SimParticle = Proxy->GetParticle_Internal(TransformIndex);
-	Chaos::FPBDRigidClusteredParticleHandle* SimClusterParticle = Proxy->GetSolverClusterHandle_Internal(TransformIndex);
+	Chaos::FPBDRigidParticleHandle* SimParticle = Particles[TransformIndex];
+	Chaos::FPBDRigidClusteredParticleHandle* SimClusterParticle = ClusterParticles[TransformIndex];
 
 	if (SimParticle && !SimParticle->Disabled())
 	{
@@ -69,7 +68,8 @@ Chaos::FPBDRigidParticleHandle* FDeferredForcesModular::GetParticle(FGeometryCol
 }
 
 // #TODO: passing all these parameters in is horrible, tidy this up now we know what needs to be done
-Chaos::FPBDRigidParticleHandle* FDeferredForcesModular::GetParticle(FGeometryCollectionPhysicsProxy* Proxy
+Chaos::FPBDRigidParticleHandle* FDeferredForcesModular::GetParticle(TArray<Chaos::FPBDRigidClusteredParticleHandle*>& Particles
+	, TArray<Chaos::FPBDRigidClusteredParticleHandle*>& ClusterParticles
 	, int TransformIndex
 	, int32 ParticleIdx
 	, const FVector& PositionalOffset
@@ -79,8 +79,8 @@ Chaos::FPBDRigidParticleHandle* FDeferredForcesModular::GetParticle(FGeometryCol
 	, FTransform& TransformOut)
 {
 
-	Chaos::FPBDRigidParticleHandle* SimParticle = Proxy->GetParticle_Internal(TransformIndex);
-	Chaos::FPBDRigidClusteredParticleHandle* SimClusterParticle = Proxy->GetSolverClusterHandle_Internal(TransformIndex);
+	Chaos::FPBDRigidParticleHandle* SimParticle = Particles[TransformIndex];
+	Chaos::FPBDRigidClusteredParticleHandle* SimClusterParticle = ClusterParticles[TransformIndex];
 
 	if (SimParticle && !SimParticle->Disabled())
 	{
@@ -99,48 +99,6 @@ Chaos::FPBDRigidParticleHandle* FDeferredForcesModular::GetParticle(FGeometryCol
 	return nullptr;
 }
 
-Chaos::FPBDRigidParticleHandle* FDeferredForcesModular::GetParticle(const FTransform& OffsetTransform
-	, FGeometryCollectionPhysicsProxy* Proxy
-	, int TransformIndex
-	, const FVector& PositionalOffset
-	, FTransform& TransformOut)
-{
-	using namespace Chaos;
-
-	const int SingleChassisIndex = 0;
-	TransformOut = FTransform::Identity;
-
-	Chaos::FPBDRigidClusteredParticleHandle* ClusterParticle = Proxy->GetSolverClusterHandle_Internal(SingleChassisIndex);
-	Chaos::FPBDRigidParticleHandle* Child = Proxy->GetParticle_Internal(TransformIndex);
-	FRigidTransform3 Frame = FRigidTransform3::Identity;
-
-	if (ClusterParticle != nullptr && Child != nullptr)
-	{
-		const FRigidTransform3 ClusterWorldTM(ClusterParticle->X(), ClusterParticle->R());
-		
-		Proxy->GetParticle_Internal(TransformIndex);
-
-		if (FPBDRigidClusteredParticleHandle* ClusterChild = Child->CastToClustered(); ClusterChild && ClusterChild->IsChildToParentLocked())
-		{
-			Frame = ClusterChild->ChildToParent();
-		}
-		else
-		{
-			//FTransform ChildT = Child->GetTransformXRCom();
-			//const FRigidTransform3 ChildWorldTM(ChildT.GetLocation(), ChildT.GetRotation());
-			const FRigidTransform3 ChildWorldTM(Child->X(), Child->R());
-			Frame = ChildWorldTM.GetRelativeTransform(ClusterWorldTM);
-			Frame.SetLocation(Frame.GetLocation() + OffsetTransform.TransformVector(PositionalOffset));
-
-			//UE_LOG(LogChaos, Warning, TEXT("ApplyForce Data TIdx %d %s"), TransformIndex, *Frame.GetTranslation().ToString());
-		}
-	}
-
-	// local transform
-	TransformOut.SetTranslation(Frame.GetTranslation());
-	return ClusterParticle;
-}
-
 Chaos::FPBDRigidParticleHandle* FDeferredForcesModular::GetParticle(const FTransform& OffsetTransform
 	, TArray<Chaos::FPBDRigidParticleHandle*>& Particles
 	, TArray<Chaos::FPBDRigidClusteredParticleHandle*>& ClusterParticles
@@ -179,7 +137,8 @@ Chaos::FPBDRigidParticleHandle* FDeferredForcesModular::GetParticle(const FTrans
 }
 
 template<typename TransformType>
-void FDeferredForcesModular::ApplyTemplate(FGeometryCollectionPhysicsProxy* Proxy
+void FDeferredForcesModular::ApplyTemplate(TArray<Chaos::FPBDRigidClusteredParticleHandle*>& Particles
+		, TArray<Chaos::FPBDRigidClusteredParticleHandle*>& ClusterParticles
 		, const TManagedArray<TransformType>& Transforms
 		, const TManagedArray<FTransform>& CollectionMassToLocal
 		, const TManagedArray<int32>& Parent)
@@ -189,7 +148,7 @@ void FDeferredForcesModular::ApplyTemplate(FGeometryCollectionPhysicsProxy* Prox
 		FTransform RelativeTransform;
 		for (const FApplyForceData& Data : ApplyForceDatas)
 		{
-			Chaos::FPBDRigidParticleHandle* RigidHandle = GetParticle(Proxy, Data.TransformIndex, Data.ParticleIdx, FVector::ZeroVector, FTransform(Transforms[Data.TransformIndex]), CollectionMassToLocal, Parent, RelativeTransform);
+			Chaos::FPBDRigidParticleHandle* RigidHandle = GetParticle(Particles, ClusterParticles, Data.TransformIndex, Data.ParticleIdx, FVector::ZeroVector, FTransform(Transforms[Data.TransformIndex]), CollectionMassToLocal, Parent, RelativeTransform);
 			if (RigidHandle)
 			{
 				AddForce(RigidHandle, Data, RelativeTransform);
@@ -198,7 +157,7 @@ void FDeferredForcesModular::ApplyTemplate(FGeometryCollectionPhysicsProxy* Prox
 
 		for (const FApplyForceAtPositionData& Data : ApplyForceAtPositionDatas)
 		{
-			Chaos::FPBDRigidParticleHandle* RigidHandle = GetParticle(Proxy, Data.TransformIndex, Data.ParticleIdx, Data.Position, FTransform(Transforms[Data.TransformIndex]), CollectionMassToLocal, Parent, RelativeTransform);
+			Chaos::FPBDRigidParticleHandle* RigidHandle = GetParticle(Particles, ClusterParticles, Data.TransformIndex, Data.ParticleIdx, Data.Position, FTransform(Transforms[Data.TransformIndex]), CollectionMassToLocal, Parent, RelativeTransform);
 			if (RigidHandle)
 			{
 				AddForceAtPosition(RigidHandle, Data, RelativeTransform);
@@ -207,7 +166,7 @@ void FDeferredForcesModular::ApplyTemplate(FGeometryCollectionPhysicsProxy* Prox
 
 		for (const FAddTorqueInRadiansData& Data : ApplyTorqueDatas)
 		{
-			Chaos::FPBDRigidParticleHandle* RigidHandle = GetParticle(Proxy, Data.TransformIndex, Data.ParticleIdx, FVector::ZeroVector, FTransform(Transforms[Data.TransformIndex]), CollectionMassToLocal, Parent, RelativeTransform);
+			Chaos::FPBDRigidParticleHandle* RigidHandle = GetParticle(Particles, ClusterParticles, Data.TransformIndex, Data.ParticleIdx, FVector::ZeroVector, FTransform(Transforms[Data.TransformIndex]), CollectionMassToLocal, Parent, RelativeTransform);
 			if (RigidHandle)
 			{
 				AddTorque(RigidHandle, Data, RelativeTransform);
@@ -220,61 +179,25 @@ void FDeferredForcesModular::ApplyTemplate(FGeometryCollectionPhysicsProxy* Prox
 	ApplyTorqueDatas.Empty();
 }
 
-void FDeferredForcesModular::Apply(FGeometryCollectionPhysicsProxy* Proxy
+void FDeferredForcesModular::Apply(TArray<Chaos::FPBDRigidClusteredParticleHandle*>& Particles
+	, TArray<Chaos::FPBDRigidClusteredParticleHandle*>& ClusterParticles
 	, const TManagedArray<FTransform>& Transforms
 	, const TManagedArray<FTransform>& CollectionMassToLocal
 	, const TManagedArray<int32>& Parent)
 {
-	ApplyTemplate(Proxy, Transforms, CollectionMassToLocal, Parent);
+	ApplyTemplate(Particles, ClusterParticles, Transforms, CollectionMassToLocal, Parent);
 }
 
-void FDeferredForcesModular::Apply(FGeometryCollectionPhysicsProxy* Proxy
+void FDeferredForcesModular::Apply(TArray<Chaos::FPBDRigidClusteredParticleHandle*>& Particles
+	, TArray<Chaos::FPBDRigidClusteredParticleHandle*>& ClusterParticles
 	, const TManagedArray<FTransform3f>& Transforms
 	, const TManagedArray<FTransform>& CollectionMassToLocal
 	, const TManagedArray<int32>& Parent)
 {
-	ApplyTemplate(Proxy, Transforms, CollectionMassToLocal, Parent);
+	ApplyTemplate(Particles, ClusterParticles, Transforms, CollectionMassToLocal, Parent);
 }
 
 
-void FDeferredForcesModular::Apply(FGeometryCollectionPhysicsProxy* Proxy)
-{
-	if (!GCoreModularVehicleDebugParams.DisableForces)
-	{
-		FTransform RelativeTransform;
-		for (const FApplyForceData& Data : ApplyForceDatas)
-		{
-			Chaos::FPBDRigidParticleHandle* RigidHandle = GetParticle(Data.OffsetTransform, Proxy, Data.TransformIndex, FVector::ZeroVector, RelativeTransform);
-			if (RigidHandle)
-			{
-				AddForce(RigidHandle, Data, RelativeTransform);
-			}
-		}
-
-		for (const FApplyForceAtPositionData& Data : ApplyForceAtPositionDatas)
-		{
-			Chaos::FPBDRigidParticleHandle* RigidHandle = GetParticle(Data.OffsetTransform, Proxy, Data.TransformIndex, Data.Position, RelativeTransform);
-			if (RigidHandle)
-			{
-				AddForceAtPosition(RigidHandle, Data, RelativeTransform);
-			}
-		}
-
-		for (const FAddTorqueInRadiansData& Data : ApplyTorqueDatas)
-		{
-			Chaos::FPBDRigidParticleHandle* RigidHandle = GetParticle(Data.OffsetTransform, Proxy, Data.TransformIndex, FVector::ZeroVector, RelativeTransform);
-			if (RigidHandle)
-			{
-				AddTorque(RigidHandle, Data, RelativeTransform);
-			}
-		}
-	}
-
-	ApplyForceDatas.Empty();
-	ApplyForceAtPositionDatas.Empty();
-	ApplyTorqueDatas.Empty();
-}
-
 void FDeferredForcesModular::Apply(TArray<Chaos::FPBDRigidParticleHandle*>& Particles
 	, TArray<Chaos::FPBDRigidClusteredParticleHandle*>& ClusterParticles)
 {
diff --git a/Engine/Source/Runtime/Experimental/ChaosVehicles/ChaosVehiclesCore/Private/SimModule/SimModuleTree.cpp b/Engine/Source/Runtime/Experimental/ChaosVehicles/ChaosVehiclesCore/Private/SimModule/SimModuleTree.cpp
--- a/Engine/Source/Runtime/Experimental/ChaosVehicles/ChaosVehiclesCore/Private/SimModule/SimModuleTree.cpp
+++ b/Engine/Source/Runtime/Experimental/ChaosVehicles/ChaosVehiclesCore/Private/SimModule/SimModuleTree.cpp
@@ -315,12 +315,15 @@ void FSimModuleTree::UpdateModuleVelocites(FGeometryCollectionPhysicsProxy* Phys
 {
 	check(PhysicsProxy);
 
+	const TArray<Chaos::FPBDRigidClusteredParticleHandle*>& Clusters = PhysicsProxy->GetSolverClusterHandles();
+	const TArray<Chaos::FPBDRigidClusteredParticleHandle*>& Particles = PhysicsProxy->GetSolverParticleHandles();
+
 	// capture the velocities at the start of each sim iteration
 	for (int i = 0; i < SimulationModuleTree.Num(); i++)
 	{
 		if (ISimulationModuleBase* Module = SimulationModuleTree[i].SimModule)
 		{
-			if (Chaos::FPBDRigidClusteredParticleHandle* ParentParticle = PhysicsProxy->GetSolverClusterHandle_Internal(Module->GetTransformIndex()))
+			if (Chaos::FPBDRigidClusteredParticleHandle* ParentParticle = Clusters[Module->GetTransformIndex()])
 			{
 				const FTransform BodyTransform(ParentParticle->R(), ParentParticle->X());
 
@@ -329,11 +332,11 @@ void FSimModuleTree::UpdateModuleVelocites(FGeometryCollectionPhysicsProxy* Phys
 					Chaos::FPBDRigidClusteredParticleHandle* Particle = nullptr;
 					if (Module->IsClustered())
 					{
-						Particle = PhysicsProxy->GetSolverClusterHandle_Internal(Module->GetTransformIndex());
+						Particle = Clusters[Module->GetTransformIndex()];
 					}
 					else
 					{
-						Particle = PhysicsProxy->GetParticle_Internal(Module->GetTransformIndex());
+						Particle = Particles[Module->GetTransformIndex()];				
 					}
 
 					if (Particle)
diff --git a/Engine/Source/Runtime/Experimental/ChaosVehicles/ChaosVehiclesCore/Public/SimModule/DeferredForcesModular.h b/Engine/Source/Runtime/Experimental/ChaosVehicles/ChaosVehiclesCore/Public/SimModule/DeferredForcesModular.h
--- a/Engine/Source/Runtime/Experimental/ChaosVehicles/ChaosVehiclesCore/Public/SimModule/DeferredForcesModular.h
+++ b/Engine/Source/Runtime/Experimental/ChaosVehicles/ChaosVehiclesCore/Public/SimModule/DeferredForcesModular.h
@@ -20,8 +20,6 @@ enum class EForceFlags : uint32
 ENUM_CLASS_FLAGS(EForceFlags);	
 
 
-class FGeometryCollectionPhysicsProxy;
-
 class CHAOSVEHICLESCORE_API FDeferredForcesModular
 {
 public:
@@ -161,7 +159,8 @@ public:
 
 	Chaos::FPBDRigidParticleHandle* GetParticleFromUniqueIndex(int32 ParticleUniqueIdx, const TArray<Chaos::FPBDRigidParticleHandle*>& Particles) const;
 
-	Chaos::FPBDRigidParticleHandle* GetParticle(FGeometryCollectionPhysicsProxy* Proxy
+	Chaos::FPBDRigidParticleHandle* GetParticle(TArray<Chaos::FPBDRigidClusteredParticleHandle*>& Particles
+			, TArray<Chaos::FPBDRigidClusteredParticleHandle*>& ClusterParticles
 			, int TransformIndex
 			, int32 ParticleIdx
 			, const FVector& PositionalOffset
@@ -170,7 +169,8 @@ public:
 			, const TManagedArray<int32>& Parent
 			, FTransform& TransformOut);
 
-	Chaos::FPBDRigidParticleHandle* GetParticle(FGeometryCollectionPhysicsProxy* Proxy
+	Chaos::FPBDRigidParticleHandle* GetParticle(TArray<Chaos::FPBDRigidClusteredParticleHandle*>& Particles
+		, TArray<Chaos::FPBDRigidClusteredParticleHandle*>& ClusterParticles
 		, int TransformIndex
 		, int32 ParticleIdx
 		, const FVector& PositionalOffset
@@ -180,17 +180,20 @@ public:
 		, FTransform& TransformOut);
 
 	Chaos::FPBDRigidParticleHandle* GetParticle(const FTransform& OffsetTransform
-		, FGeometryCollectionPhysicsProxy* Proxy
+		, TArray<Chaos::FPBDRigidParticleHandle*>& Particles
+		, TArray<Chaos::FPBDRigidClusteredParticleHandle*>& ClusterParticles
 		, int32 ParticleIdx
 		, const FVector& PositionalOffset
 		, FTransform& TransformOut);
 
-	void Apply(FGeometryCollectionPhysicsProxy* Proxy
+	void Apply(TArray<Chaos::FPBDRigidClusteredParticleHandle*>& Particles
+		, TArray<Chaos::FPBDRigidClusteredParticleHandle*>& ClusterParticles
 		, const TManagedArray<FTransform>& Transforms
 		, const TManagedArray<FTransform>& CollectionMassToLocal
 		, const TManagedArray<int32>& Parent);
 
-	void Apply(FGeometryCollectionPhysicsProxy* Proxy
+	void Apply(TArray<Chaos::FPBDRigidClusteredParticleHandle*>& Particles
+		, TArray<Chaos::FPBDRigidClusteredParticleHandle*>& ClusterParticles
 		, const TManagedArray<FTransform3f>& Transforms
 		, const TManagedArray<FTransform>& CollectionMassToLocal
 		, const TManagedArray<int32>& Parent);
@@ -201,7 +204,8 @@ public:
 private:
 
 	template<typename TransformType>
-	void ApplyTemplate(FGeometryCollectionPhysicsProxy* Proxy
+	void ApplyTemplate(TArray<Chaos::FPBDRigidClusteredParticleHandle*>& Particles
+		, TArray<Chaos::FPBDRigidClusteredParticleHandle*>& ClusterParticles
 		, const TManagedArray<TransformType>& Transforms
 		, const TManagedArray<FTransform>& CollectionMassToLocal
 		, const TManagedArray<int32>& Parent);
diff --git a/Engine/Source/Runtime/Experimental/GeometryCollectionEngine/Private/GeometryCollection/GeometryCollectionComponent.cpp b/Engine/Source/Runtime/Experimental/GeometryCollectionEngine/Private/GeometryCollection/GeometryCollectionComponent.cpp
--- a/Engine/Source/Runtime/Experimental/GeometryCollectionEngine/Private/GeometryCollection/GeometryCollectionComponent.cpp
+++ b/Engine/Source/Runtime/Experimental/GeometryCollectionEngine/Private/GeometryCollection/GeometryCollectionComponent.cpp
@@ -2108,11 +2108,12 @@ void UGeometryCollectionComponent::UpdateRepData()
 		const FRigidClustering& RigidClustering = Solver->GetEvolution()->GetRigidClustering();
 
 		const TManagedArrayAccessor<int32> InitialLevels = PhysicsProxy->GetPhysicsCollection().GetInitialLevels();
-		const TArray<FPBDRigidClusteredParticleHandle*>& UnorderedParticleHandles = PhysicsProxy->GetUnorderedParticles_Internal();
+		const TArray<FPBDRigidClusteredParticleHandle*>& ParticleHandles = PhysicsProxy->GetParticles();
 
 		// Replicate the anchored state of the root particle
 		LocalRepData.bIsRootAnchored = RepData.bIsRootAnchored;
-		if (FPBDRigidClusteredParticleHandle* RootHandle = PhysicsProxy->GetInitialRootHandle_Internal())
+		const int32 InitialRootIndex = PhysicsProxy->GetSimParameters().InitialRootIndex;
+		if (FPBDRigidClusteredParticleHandle* RootHandle = ParticleHandles[InitialRootIndex])
 		{
 			LocalRepData.bIsRootAnchored = RootHandle->IsAnchored();
 		}
@@ -2122,7 +2123,7 @@ void UGeometryCollectionComponent::UpdateRepData()
 
 		// It feels weird to start from the previous one. Need to be checked
 		LocalRepData.OneOffActivated = RepData.OneOffActivated;
-		for (FPBDRigidClusteredParticleHandle* Particle : UnorderedParticleHandles)
+		for (FPBDRigidClusteredParticleHandle* Particle : ParticleHandles)
 		{
 			// Particle can be null if we have embedded geometry 
 			if (Particle)
@@ -2181,7 +2182,7 @@ void UGeometryCollectionComponent::UpdateRepData()
 						//a one off so record it
 						ensureMsgf(TransformGroupIdx >= 0, TEXT("Non-internal cluster should always have a group index"));
 						ensureMsgf(TransformGroupIdx < TNumericLimits<uint16>::Max(), TEXT("Trying to replicate GC with more than 65k pieces. We assumed uint16 would suffice"));
-						ensureMsgf(PhysicsProxy->GetParticle_Internal(TransformGroupIdx) != nullptr, TEXT("Invalid particle index being replicated - Contact physics team"));
+						ensureMsgf(PhysicsProxy->GetParticles().IsValidIndex(TransformGroupIdx) && PhysicsProxy->GetParticles()[TransformGroupIdx] != nullptr, TEXT("Invalid particle index being replicated - Contact physics team"));
 
 						// Because we cull ClustersToRep with abandoned level, we must make sure we don't add duplicates to one off activated.
 						// TODO: avoid search for entry for perf
@@ -2304,8 +2305,8 @@ void UGeometryCollectionComponent::UpdateRepStateAndDynamicData()
 			FPBDRigidsSolver* Solver = PhysicsProxy->GetSolver<Chaos::FPBDRigidsSolver>();
 
 			// let go through all the transform and see which one has changed its state
-			const TArray<FPBDRigidClusteredParticleHandle*>& UnorderedParticleHandles = PhysicsProxy->GetUnorderedParticles_Internal();
-			const int32 NumParticles = UnorderedParticleHandles.Num();
+			const TArray<FPBDRigidClusteredParticleHandle*>& ParticleHandles = PhysicsProxy->GetParticles();
+			const int32 NumTransforms = ParticleHandles.Num();
 
 			const TManagedArrayAccessor<int32> InitialLevels = PhysicsProxy->GetPhysicsCollection().GetInitialLevels();
 			
@@ -2318,7 +2319,7 @@ void UGeometryCollectionComponent::UpdateRepStateAndDynamicData()
 			// root level anchor state
 			LocalRepStateData.bIsRootAnchored = RepStateData.bIsRootAnchored;
 			const int32 InitialRootIndex = PhysicsProxy->GetSimParameters().InitialRootIndex;
-			if (FPBDRigidClusteredParticleHandle* RootHandle = UnorderedParticleHandles[InitialRootIndex])
+			if (FPBDRigidClusteredParticleHandle* RootHandle = ParticleHandles[InitialRootIndex])
 			{
 				LocalRepStateData.bIsRootAnchored = RootHandle->IsAnchored() ? 1 : 0;
 			}
@@ -2339,14 +2340,11 @@ void UGeometryCollectionComponent::UpdateRepStateAndDynamicData()
 			const int32 RootIndex = GetRootIndex();
 
 			// go through particles and send replicated data 
-			const int32 NumTransforms = PhysicsProxy->GetNumTransforms();
 			LocalRepStateData.BrokenState.SetNum(NumTransforms, false);
 			LocalRepStateData.ReleasedData.Reserve(NumTransforms);
-			for (int32 ParticleIndex = 0; ParticleIndex < NumParticles; ParticleIndex++)
+			for (int32 TransformIndex = 0; TransformIndex < NumTransforms; TransformIndex++)
 			{
-				const int32 TransformIndex = PhysicsProxy->GetFromParticleToTransformIndex(ParticleIndex);
-				FPBDRigidClusteredParticleHandle* ParticleHandle = UnorderedParticleHandles[ParticleIndex];
-				check(ParticleHandle);
+				if (FPBDRigidClusteredParticleHandle* ParticleHandle = ParticleHandles[TransformIndex])
 				{
 					const int32 Level = (InitialLevels.IsValid() && InitialLevels.IsValidIndex(TransformIndex)) ? InitialLevels[TransformIndex] : INDEX_NONE;
 
@@ -2520,7 +2518,8 @@ namespace
 		using namespace Chaos;
 
 		const int32 InitialRootIndex = PhysicsProxy.GetSimParameters().InitialRootIndex;
-		if (FPBDRigidClusteredParticleHandle* RootHandle = PhysicsProxy.GetParticle_Internal(InitialRootIndex))
+		const TArray<FPBDRigidClusteredParticleHandle*>& ParticleHandles = PhysicsProxy.GetParticles();
+		if (FPBDRigidClusteredParticleHandle* RootHandle = ParticleHandles[InitialRootIndex])
 		{
 			if (RootHandle->Parent() || !RootHandle->Disabled())
 			{
@@ -2642,7 +2641,7 @@ namespace
 		{
 			for (const FGeometryCollectionClusterRep& RepCluster : RepData.Clusters)
 			{
-				if (Chaos::FPBDRigidParticleHandle* Cluster = PhysicsProxy->GetParticle_Internal(RepCluster.ClusterIdx))
+				if (Chaos::FPBDRigidParticleHandle* Cluster = PhysicsProxy->GetParticles()[RepCluster.ClusterIdx])
 				{
 					if (RepCluster.ClusterState.IsInternalCluster())
 					{
@@ -2682,14 +2681,14 @@ namespace
 			const FGeometryCollectionActivatedCluster& ActivatedCluster = RepData.OneOffActivated[OneOffActivatedProcessed];
 
 #if !UE_BUILD_SHIPPING
-			if(PhysicsProxy->GetParticle_Internal(ActivatedCluster.ActivatedIndex) == nullptr)
+			if(!PhysicsProxy->GetParticles().IsValidIndex(ActivatedCluster.ActivatedIndex))
 			{
 				ensureMsgf(false, TEXT("Invalid activated cluster index processing replication data."));
 				continue;
 			}
 #endif
 
-			FPBDRigidParticleHandle* OneOff = PhysicsProxy->GetParticle_Internal(ActivatedCluster.ActivatedIndex);
+			FPBDRigidParticleHandle* OneOff = PhysicsProxy->GetParticles()[ActivatedCluster.ActivatedIndex];
 
 			if (ensure(OneOff))
 			{
@@ -2740,52 +2739,52 @@ namespace
 			return false;
 		}
 
-		const TArray<FPBDRigidClusteredParticleHandle*>& UnorderedParticleHandle = PhysicsProxy->GetUnorderedParticles_Internal();
+		const TArray<FPBDRigidClusteredParticleHandle*>& ParticleHandles = PhysicsProxy->GetParticles();
 
 		// Update the anchored state of the root particle
 		UpdateRootStateFromReplication(*PhysicsProxy, (bool)RepStateData.bIsRootAnchored);
 		
 		
 		// now sync the broken state of the particles
-		if (RepStateData.BrokenState.Num() == PhysicsProxy->GetNumTransforms())
+		if (RepStateData.BrokenState.Num() == ParticleHandles.Num())
 		{
 			FPBDRigidsSolver* Solver = PhysicsProxy->GetSolver<Chaos::FPBDRigidsSolver>();
 			FRigidClustering& RigidClustering = Solver->GetEvolution()->GetRigidClustering();
 
-			for (int32 ParticleIndex = 0; ParticleIndex < UnorderedParticleHandle.Num(); ParticleIndex++)
+			for (int32 TransformIndex = 0; TransformIndex < ParticleHandles.Num(); TransformIndex++)
 			{
-				FPBDRigidClusteredParticleHandle* ParticleHandle = UnorderedParticleHandle[ParticleIndex];
-				check(ParticleHandle != nullptr);
-				const int32 TransformIndex = PhysicsProxy->GetFromParticleToTransformIndex(ParticleIndex);
-				// we only check for when thing goes from unbroken to broken 
-				const bool bWasBroken = (ParticleHandle->Parent() == nullptr);
-				const bool bIsBroken = RepStateData.BrokenState[TransformIndex];
-				if (!bWasBroken && bIsBroken)
+				if (FPBDRigidClusteredParticleHandle* ParticleHandle = ParticleHandles[TransformIndex])
 				{
-					RigidClustering.ForceReleaseChildParticleAndParents(ParticleHandle, /* bTriggerBreakEvents */true);
-
-					if (!ParticleHandle->Disabled())
+					// we only check for when thing goes from unbroken to broken 
+					const bool bWasBroken = (ParticleHandle->Parent() == nullptr);
+					const bool bIsBroken = RepStateData.BrokenState[TransformIndex];
+					if (!bWasBroken && bIsBroken)
 					{
-						// check if we have a release velocity to be applied 
-						const FGeometryCollectionRepStateData::FReleasedData* ReleaseData = RepStateData.ReleasedData.FindByPredicate(
-							[TransformIndex](const FGeometryCollectionRepStateData::FReleasedData& Data)
-							{
-								return Data.TransformIndex == TransformIndex;
-							});
-						if (ReleaseData)
-						{
-							ParticleHandle->SetV(ReleaseData->LinearVelocity);
-							ParticleHandle->SetW(FMath::DegreesToRadians(ReleaseData->AngularVelocityInDegreesPerSecond));
-						}
-						else
+						RigidClustering.ForceReleaseChildParticleAndParents(ParticleHandle, /* bTriggerBreakEvents */true);
+
+						if (!ParticleHandle->Disabled())
 						{
-							// if we are a leaf we can disable this particle right away if there's no release data
-							// because this means the server has already disabled the particle from a removal 
-							if (ParticleHandle->ClusterIds().NumChildren == 0)
+							// check if we have a release velocity to be applied 
+							const FGeometryCollectionRepStateData::FReleasedData* ReleaseData = RepStateData.ReleasedData.FindByPredicate(
+								[TransformIndex](const FGeometryCollectionRepStateData::FReleasedData& Data)
+								{
+									return Data.TransformIndex == TransformIndex;
+								});
+							if (ReleaseData)
 							{
-								// todo we could probably disable all at once later to have only one call 
-								Solver->GetEvolution()->DisableParticle(ParticleHandle);
-								Solver->GetParticles().MarkTransientDirtyParticle(ParticleHandle);
+								ParticleHandle->SetV(ReleaseData->LinearVelocity);
+								ParticleHandle->SetW(FMath::DegreesToRadians(ReleaseData->AngularVelocityInDegreesPerSecond));
+							}
+							else
+							{
+								// if we are a leaf we can disable this particle right away if there's no release data
+								// becuase this means the server has already disabled the particle from a removal 
+								if (ParticleHandle->ClusterIds().NumChildren == 0)
+								{
+									// todo we could probably disable all at once later to have only one call 
+									Solver->GetEvolution()->DisableParticle(ParticleHandle);
+									Solver->GetParticles().MarkTransientDirtyParticle(ParticleHandle);
+								}
 							}
 						}
 					}
@@ -2839,9 +2838,11 @@ namespace
 		if (bReplicateMovement)
 		{
 			FPBDRigidsSolver* Solver = PhysicsProxy->GetSolver<Chaos::FPBDRigidsSolver>();
+			const TArray<FPBDRigidClusteredParticleHandle*>& ParticleHandles = PhysicsProxy->GetParticles();
+
 			for (const FGeometryCollectionRepDynamicData::FClusterData& RepCluster : RepDynamicData.ClusterData)
 			{
-				if (Chaos::FPBDRigidParticleHandle* Cluster = PhysicsProxy->GetParticle_Internal(RepCluster.TransformIndex))
+				if (Chaos::FPBDRigidParticleHandle* Cluster = ParticleHandles[RepCluster.TransformIndex])
 				{
 					if (RepCluster.bIsInternalCluster)
 					{
@@ -2900,7 +2901,7 @@ bool UGeometryCollectionComponent::ProcessRepData(const float DeltaTime, const f
 		{
 			for (const FGeometryCollectionClusterRep& RepCluster : RepData.Clusters)
 			{
-				if (Chaos::FPBDRigidParticleHandle* Cluster = PhysicsProxy->GetParticle_Internal(RepCluster.ClusterIdx))
+				if (Chaos::FPBDRigidParticleHandle* Cluster = PhysicsProxy->GetParticles()[RepCluster.ClusterIdx])
 				{
 					if (RepCluster.ClusterState.IsInternalCluster())
 					{
@@ -3924,13 +3925,15 @@ void UGeometryCollectionComponent::RegisterAndInitializePhysicsProxy()
 						// As we're not in control we make it so our simulated proxy cannot break clusters
 						// We have to set the strain to a high value but be below the max for the data type
 						// so releasing on authority demand works
-						for (Chaos::FPBDRigidClusteredParticleHandle* ParticleHandle : Proxy->GetUnorderedParticles_Internal())
+						for (Chaos::FPBDRigidClusteredParticleHandle* ParticleHandle : Proxy->GetParticles())
 						{
-							check(ParticleHandle);
-							const int32 Level = EnableAbandonAfterLevel ? ComputeParticleLevel(ParticleHandle) : -1;
-							if (Level <= AbandonAfterLevel)	//we only replicate up until level X, but it means we should replicate the breaking event of level X+1 (but not X+1's positions)
+							if (ParticleHandle)
 							{
-								ParticleHandle->SetUnbreakable(true);
+								const int32 Level = EnableAbandonAfterLevel ? ComputeParticleLevel(ParticleHandle) : -1;
+								if (Level <= AbandonAfterLevel)	//we only replicate up until level X, but it means we should replicate the breaking event of level X+1 (but not X+1's positions)
+								{
+									ParticleHandle->SetUnbreakable(true);
+								}
 							}
 						}
 					});
@@ -3980,7 +3983,7 @@ void UGeometryCollectionComponent::UpdateBrokenAndDecayedStates()
 							continue;
 						}
 
-						const Chaos::FPBDRigidParticle* Particle = PhysicsProxy->GetParticleByIndex_External(TransformIdx);
+						const Chaos::FPBDRigidParticle* Particle = PhysicsProxy->GetExternalParticles()[TransformIdx].Get();
 						if (!Particle)
 						{
 							// when no particle we mark it as a decayed one
@@ -4218,7 +4221,7 @@ void UGeometryCollectionComponent::MoveComponentToRootTransform()
 			{
 				const FTransform3f& OriginalComponentSpaceRootTransformOffset = AssetCollection->Transform[RootIndex];
 				DynamicCollection->SetTransform(RootIndex, OriginalComponentSpaceRootTransformOffset);
-				const Chaos::FPBDRigidParticle* RootParticle = PhysicsProxy->GetParticleByIndex_External(RootIndex);
+				const Chaos::FPBDRigidParticle* RootParticle = PhysicsProxy->GetExternalParticles()[RootIndex].Get();
 				const FTransform ParticleWorldPosition(RootParticle->R(), RootParticle->X());
 				const FTransform MassToLocal = MassToLocalAttribute[RootIndex];
 				const FTransform NewRootWorldPosition = MassToLocal.Inverse() * ParticleWorldPosition;
@@ -6990,9 +6993,9 @@ TArray<Chaos::FPhysicsObject*> UGeometryCollectionComponent::GetAllPhysicsObject
 		return {};
 	}
 	TArray<Chaos::FPhysicsObject*> Objects;
-	Objects.Reserve(PhysicsProxy->GetNumTransforms());
+	Objects.Reserve(PhysicsProxy->GetNumParticles());
 	
-	for (int32 Index = 0; Index < PhysicsProxy->GetNumTransforms(); ++Index)
+	for (int32 Index = 0; Index < PhysicsProxy->GetNumParticles(); ++Index)
 	{
 		Objects.Add(GetPhysicsObjectById(Index));
 	}
