Fix some occasionally incorrect logging when toggling IAS flags on/off

#rnx

- Add missing conversion to text for UE::IO::IAS::FBackendStatus::AbandonCache otherwise it will not be correctly logged when the flag is toggled.

diff --git a/Engine/Source/Runtime/Experimental/IoStoreOnDemand/Private/OnDemandIoDispatcherBackend.cpp b/Engine/Source/Runtime/Experimental/IoStoreOnDemand/Private/OnDemandIoDispatcherBackend.cpp
--- a/Engine/Source/Runtime/Experimental/IoStoreOnDemand/Private/OnDemandIoDispatcherBackend.cpp
+++ b/Engine/Source/Runtime/Experimental/IoStoreOnDemand/Private/OnDemandIoDispatcherBackend.cpp
@@ -1199,6 +1199,8 @@ struct FBackendStatus
 		HttpError					= (1 << 2),
 		HttpBulkOptionalDisabled	= (1 << 3),
 		AbandonCache				= (1 << 4),
+
+		// When adding new values here, remember to update operator<<(FStringBuilderBase& Sb, EFlags StatusFlags) below!
 	};
 
 	bool IsHttpEnabled() const
@@ -1340,6 +1342,7 @@ private:
 		AppendIf(EFlags::HttpEnabled, TEXT("HttpEnabled"));
 		AppendIf(EFlags::HttpError, TEXT("HttpError"));
 		AppendIf(EFlags::HttpBulkOptionalDisabled, TEXT("HttpBulkOptionalDisabled"));
+		AppendIf(EFlags::AbandonCache, TEXT("AbandonCache"));
 
 		return Sb;
 	}
