Chaos Cloth: improved XPBD Damping.

XPBD Damping causes the cloth to become too soft. This seemed to be because elasticity and damping shared the same Lagrange multipliers. This change splits the XPBD constraint into two constraints: one for just damping and one for elasticity, each with their own Lagrange multipliers.

Internal testing showed that doing two passes, damping first, then elasticity was the most stable. There's a CVAR that allows testing a few different options. Non-ISPC springs has several options. Bending elements and ISPC springs just have the option to do the new split method and the old single Lagrange multiplier version.

I switched the sign for a lot of the DeltaX computations to match the original XPBD paper. However, I also did the thing I had already done for springs which was make things in terms of stiffness rather than 1/stiffness.

diff --git a/Engine/Source/Runtime/Experimental/Chaos/Private/Chaos/SoftsSpring.isph b/Engine/Source/Runtime/Experimental/Chaos/Private/Chaos/SoftsSpring.isph
--- a/Engine/Source/Runtime/Experimental/Chaos/Private/Chaos/SoftsSpring.isph
+++ b/Engine/Source/Runtime/Experimental/Chaos/Private/Chaos/SoftsSpring.isph
@@ -16,6 +16,7 @@ inline float SafeNormalize(FVector3f &Direction)
 	return (SizeSquared < FLOAT_KINDA_SMALL_NUMBER) ? FLOAT_ZERO : Size;
 }
 
+// Spring without damping
 template <typename StiffnessFloatType>
 inline FVector3f GetXPBDSpringDelta(const FVector3f& P1, const FVector3f& P2, const varying float InvM1, const varying float InvM2,
 	 const uniform float Dt, const varying float Dist, const varying float Lambda, const StiffnessFloatType Stiffness, const uniform float MinStiffness, float& DLambda)
@@ -28,7 +29,7 @@ inline FVector3f GetXPBDSpringDelta(const FVector3f& P1, const FVector3f& P2, co
 	
 	const float CombinedInvMass = InvM2 + InvM1;
 
-	const float ComputedDLambda = (AlphaInv * Offset - Lambda) / (AlphaInv * CombinedInvMass + (uniform float)1.f);
+	const float ComputedDLambda = (-AlphaInv * Offset - Lambda) / (AlphaInv * CombinedInvMass + (uniform float)1.f);
 	DLambda = select(Stiffness < MinStiffness, 0.f, ComputedDLambda);
 
 	const FVector3f ComputedDelta = DLambda * Direction;
@@ -36,29 +37,35 @@ inline FVector3f GetXPBDSpringDelta(const FVector3f& P1, const FVector3f& P2, co
 	return VectorSelect((InvM2 == FLOAT_ZERO && InvM1 == FLOAT_ZERO), FloatZeroVector, ComputedDelta);
 }
 
+// This is a following the original XPBD paper using a single lambda for stretch and damping.
 template <typename StiffnessFloatType>
-inline FVector3f GetXPBDAxialSpringDelta(const FVector3f& P1, const FVector3f& P2, const FVector3f& P3, const varying float InvM1, const varying float InvM2, const varying float InvM3,
-	 const uniform float Dt, const varying float Bary, const varying float Dist, const varying float Lambda, const StiffnessFloatType Stiffness, const uniform float MinStiffness, float& DLambda)
+inline FVector3f GetXPBDSpringDeltaWithDamping(const FVector3f& P1, const FVector3f& P2, const FVector3f& X1, const FVector3f& X2,
+	const varying float InvM1, const varying float InvM2,
+	const uniform float Dt, const varying float Dist, const varying float Lambda, const StiffnessFloatType Stiffness, const StiffnessFloatType DampingRatio, const uniform float MinStiffness, float& DLambda)
 {
-	FVector3f P = (P2 - P3) * Bary + P3;
-	FVector3f Direction = P1 - P;
+	FVector3f Direction = P1 - P2;
 	const float Distance = SafeNormalize(Direction);
 	const varying float Offset = Distance - Dist;
+	const varying float OneOverDist = select(Dist< FLOAT_SMALL_NUMBER, 1.f, 1.f / Dist);
 	
 	const varying float AlphaInv = Stiffness * Dt * Dt;
 	
-	const float CombinedInvMass = (1.f - Bary) * InvM3 + Bary * InvM2 + InvM1;
+	const float CombinedInvMass = InvM2 + InvM1;
+	const varying float Damping = (uniform float)2.f * DampingRatio * sqrt(Stiffness / CombinedInvMass) * OneOverDist;
+	const varying float BetaDt = Damping * Dt;
+	const varying FVector3f RelativeVelocityTimesDt = P1 - X1 - P2 + X2;
 
-	const float ComputedDLambda = (AlphaInv * Offset - Lambda) / (AlphaInv * CombinedInvMass + (uniform float)1.f);
+	const float ComputedDLambda = (-AlphaInv * Offset - Lambda - BetaDt * VectorDot(Direction, RelativeVelocityTimesDt)) / ((AlphaInv + BetaDt)*CombinedInvMass + (uniform float)1.f);
 	DLambda = select(Stiffness < MinStiffness, 0.f, ComputedDLambda);
 
 	const FVector3f ComputedDelta = DLambda * Direction;
 
-	return VectorSelect((CombinedInvMass == FLOAT_ZERO), FloatZeroVector, ComputedDelta);
+	return VectorSelect(CombinedInvMass < FLOAT_SMALL_NUMBER, FloatZeroVector, ComputedDelta);
 }
 
+// Spring damping constraint (separate from spring stretching)
 template <typename StiffnessFloatType>
-inline FVector3f GetXPBDSpringDeltaWithDamping(const FVector3f& P1, const FVector3f& P2, const FVector3f& X1, const FVector3f& X2,
+inline FVector3f GetXPBDSpringDampingDelta(const FVector3f& P1, const FVector3f& P2, const FVector3f& X1, const FVector3f& X2,
 	const varying float InvM1, const varying float InvM2,
 	const uniform float Dt, const varying float Dist, const varying float Lambda, const StiffnessFloatType Stiffness, const StiffnessFloatType DampingRatio, const uniform float MinStiffness, float& DLambda)
 {
@@ -67,21 +74,42 @@ inline FVector3f GetXPBDSpringDeltaWithDamping(const FVector3f& P1, const FVecto
 	const varying float Offset = Distance - Dist;
 	const varying float OneOverDist = select(Dist< FLOAT_SMALL_NUMBER, 1.f, 1.f / Dist);
 	
-	const varying float AlphaInv = Stiffness * Dt * Dt;
-	
 	const float CombinedInvMass = InvM2 + InvM1;
 	const varying float Damping = (uniform float)2.f * DampingRatio * sqrt(Stiffness / CombinedInvMass) * OneOverDist;
 	const varying float BetaDt = Damping * Dt;
 	const varying FVector3f RelativeVelocityTimesDt = P1 - X1 - P2 + X2;
 
-	const float ComputedDLambda = (AlphaInv * Offset - Lambda + BetaDt * VectorDot(Direction, RelativeVelocityTimesDt)) / ((AlphaInv + BetaDt)*CombinedInvMass + (uniform float)1.f);
+	const float ComputedDLambda = ( -BetaDt * VectorDot(Direction, RelativeVelocityTimesDt) - Lambda) / (BetaDt * CombinedInvMass + (uniform float)1.f);
 	DLambda = select(Stiffness < MinStiffness, 0.f, ComputedDLambda);
 
 	const FVector3f ComputedDelta = DLambda * Direction;
 
-	return VectorSelect((InvM2 == FLOAT_ZERO && InvM1 == FLOAT_ZERO), FloatZeroVector, ComputedDelta);
+	return VectorSelect(CombinedInvMass < FLOAT_SMALL_NUMBER, FloatZeroVector, ComputedDelta);
 }
 
+// Spring without damping
+template <typename StiffnessFloatType>
+inline FVector3f GetXPBDAxialSpringDelta(const FVector3f& P1, const FVector3f& P2, const FVector3f& P3, const varying float InvM1, const varying float InvM2, const varying float InvM3,
+	 const uniform float Dt, const varying float Bary, const varying float Dist, const varying float Lambda, const StiffnessFloatType Stiffness, const uniform float MinStiffness, float& DLambda)
+{
+	FVector3f P = (P2 - P3) * Bary + P3;
+	FVector3f Direction = P1 - P;
+	const float Distance = SafeNormalize(Direction);
+	const varying float Offset = Distance - Dist;
+	
+	const varying float AlphaInv = Stiffness * Dt * Dt;
+	
+	const float CombinedInvMass = (1.f - Bary) * InvM3 + Bary * InvM2 + InvM1;
+
+	const float ComputedDLambda = (-AlphaInv * Offset - Lambda) / (AlphaInv * CombinedInvMass + (uniform float)1.f);
+	DLambda = select(Stiffness < MinStiffness, 0.f, ComputedDLambda);
+
+	const FVector3f ComputedDelta = DLambda * Direction;
+
+	return VectorSelect((CombinedInvMass == FLOAT_ZERO), FloatZeroVector, ComputedDelta);
+}
+
+// This is a following the original XPBD paper using a single lambda for stretch and damping.
 template <typename StiffnessFloatType>
 inline FVector3f GetXPBDAxialSpringDeltaWithDamping(const FVector3f& P1, const FVector3f& P2, const FVector3f& P3, 
 	const FVector3f& X1, const FVector3f& X2, const FVector3f& X3,
@@ -93,7 +121,7 @@ inline FVector3f GetXPBDAxialSpringDeltaWithDamping(const FVector3f& P1, const F
 	FVector3f Direction = P1 - P;
 	const float Distance = SafeNormalize(Direction);
 	const varying float Offset = Distance - Dist;
-	const varying float OneOverDist = select(Dist< FLOAT_SMALL_NUMBER, 1.f, 1.f / Dist);
+	const varying float OneOverDist = select(Dist < FLOAT_SMALL_NUMBER, 1.f, 1.f / Dist);
 	
 	const varying float AlphaInv = Stiffness * Dt * Dt;
 	
@@ -102,11 +130,38 @@ inline FVector3f GetXPBDAxialSpringDeltaWithDamping(const FVector3f& P1, const F
 	const varying float BetaDt = Damping * Dt;
 	const varying FVector3f RelativeVelocityTimesDt = P1 - X1 - P2 + X2;
 
-	const float ComputedDLambda = (AlphaInv * Offset - Lambda + BetaDt * VectorDot(Direction, RelativeVelocityTimesDt)) / ((AlphaInv + BetaDt)*CombinedInvMass + (uniform float)1.f);
+	const float ComputedDLambda = (-AlphaInv * Offset - Lambda - BetaDt * VectorDot(Direction, RelativeVelocityTimesDt)) / ((AlphaInv + BetaDt)*CombinedInvMass + (uniform float)1.f);
 	DLambda = select(Stiffness < MinStiffness, 0.f, ComputedDLambda);
 
 	const FVector3f ComputedDelta = DLambda * Direction;
 
-	return VectorSelect((CombinedInvMass == FLOAT_ZERO), FloatZeroVector, ComputedDelta);
+	return VectorSelect((CombinedInvMass < FLOAT_SMALL_NUMBER), FloatZeroVector, ComputedDelta);
+}
+
+// Spring damping constraint (separate from spring stretching)
+template <typename StiffnessFloatType>
+inline FVector3f GetXPBDAxialSpringDampingDelta(const FVector3f& P1, const FVector3f& P2, const FVector3f& P3, 
+	const FVector3f& X1, const FVector3f& X2, const FVector3f& X3,
+	const varying float InvM1, const varying float InvM2, const varying float InvM3,
+	const uniform float Dt, const varying float Bary, const varying float Dist, const varying float Lambda, 
+	const StiffnessFloatType Stiffness, const StiffnessFloatType DampingRatio, const uniform float MinStiffness, float& DLambda)
+{
+	FVector3f P = (P2 - P3) * Bary + P3;
+	FVector3f Direction = P1 - P;
+	const float Distance = SafeNormalize(Direction);
+	const varying float Offset = Distance - Dist;
+	const varying float OneOverDist = select(Dist < FLOAT_SMALL_NUMBER, 1.f, 1.f / Dist);
+	
+	const float CombinedInvMass = (1.f - Bary) * InvM3 + Bary * InvM2 + InvM1;
+	const varying float Damping = (uniform float)2.f * DampingRatio * sqrt(Stiffness / CombinedInvMass) * OneOverDist;
+	const varying float BetaDt = Damping * Dt;
+	const varying FVector3f RelativeVelocityTimesDt = P1 - X1 - P2 + X2;
+
+	const float ComputedDLambda = (- BetaDt * VectorDot(Direction, RelativeVelocityTimesDt) - Lambda ) / (BetaDt*CombinedInvMass + (uniform float)1.f);
+	DLambda = select(Stiffness < MinStiffness, 0.f, ComputedDLambda);
+
+	const FVector3f ComputedDelta = DLambda * Direction;
+
+	return VectorSelect((CombinedInvMass < FLOAT_SMALL_NUMBER), FloatZeroVector, ComputedDelta);
 }
 #endif
diff --git a/Engine/Source/Runtime/Experimental/Chaos/Private/Chaos/XPBDAnisotropicBendingConstraints.cpp b/Engine/Source/Runtime/Experimental/Chaos/Private/Chaos/XPBDAnisotropicBendingConstraints.cpp
--- a/Engine/Source/Runtime/Experimental/Chaos/Private/Chaos/XPBDAnisotropicBendingConstraints.cpp
+++ b/Engine/Source/Runtime/Experimental/Chaos/Private/Chaos/XPBDAnisotropicBendingConstraints.cpp
@@ -179,6 +179,7 @@ FXPBDAnisotropicBendingConstraints::FXPBDAnisotropicBendingConstraints(const FSo
 	, XPBDAnisoRestAngleTypeIndex(PropertyCollection)
 {
 	Lambdas.Init((FSolverReal)0., Constraints.Num());
+	LambdasDamping.Init((FSolverReal)0., Constraints.Num());
 	InitColor(InParticles);
 }
 
@@ -271,6 +272,7 @@ FXPBDAnisotropicBendingConstraints::FXPBDAnisotropicBendingConstraints(const FSo
 	, XPBDAnisoRestAngleTypeIndex(ForceInit)
 {
 	Lambdas.Init((FSolverReal)0., Constraints.Num());
+	LambdasDamping.Init((FSolverReal)0., Constraints.Num());
 	InitColor(InParticles);
 }
 
@@ -539,49 +541,50 @@ void FXPBDAnisotropicBendingConstraints::SetProperties(
 	}
 }
 
-template<typename SolverParticlesOrRange>
+template<bool bDampingOnly, bool bElasticOnly, typename SolverParticlesOrRange>
 void FXPBDAnisotropicBendingConstraints::ApplyHelper(SolverParticlesOrRange& Particles, const FSolverReal Dt, const int32 ConstraintIndex, const FSolverVec3& ExpStiffnessValues,
 	const FSolverVec3& ExpBucklingStiffnessValues, const FSolverReal DampingRatioValue) const
 {
 	const TVec4<int32>& Constraint = Constraints[ConstraintIndex];
-	const int32 i1 = Constraint[0];
-	const int32 i2 = Constraint[1];
-	const int32 i3 = Constraint[2];
-	const int32 i4 = Constraint[3];
+	const int32 Index1 = Constraint[0];
+	const int32 Index2 = Constraint[1];
+	const int32 Index3 = Constraint[2];
+	const int32 Index4 = Constraint[3];
 	const FSolverVec3& WarpWeftBiasBaseMultiplier = WarpWeftBiasBaseMultipliers[ConstraintIndex];
 
 	const FSolverReal BiphasicStiffnessValue = IsBuckled[ConstraintIndex] ?
 		WarpWeftBiasBaseMultiplier.Dot(ExpBucklingStiffnessValues): WarpWeftBiasBaseMultiplier.Dot(ExpStiffnessValues);
 
-	if (BiphasicStiffnessValue < MinStiffness || (Particles.InvM(i1) == (FSolverReal)0. && Particles.InvM(i2) == (FSolverReal)0. && Particles.InvM(i3) == (FSolverReal)0. && Particles.InvM(i4) == (FSolverReal)0.))
+	const FSolverReal CombinedInvMass = Particles.InvM(Index1) + Particles.InvM(Index2) + Particles.InvM(Index3) + Particles.InvM(Index4);
+	if (BiphasicStiffnessValue < MinStiffness || CombinedInvMass < UE_SMALL_NUMBER)
 	{
 		return;
 	}
 
-	const FSolverReal CombinedInvMass = Particles.InvM(i1) + Particles.InvM(i2) + Particles.InvM(i3) + Particles.InvM(i4);
 	const FSolverReal Damping = (FSolverReal)2.f * DampingRatioValue * FMath::Sqrt(BiphasicStiffnessValue / CombinedInvMass);
 
-	const FSolverReal Angle = CalcAngle(Particles.P(i1), Particles.P(i2), Particles.P(i3), Particles.P(i4));
+	const FSolverReal Angle = CalcAngle(Particles.P(Index1), Particles.P(Index2), Particles.P(Index3), Particles.P(Index4));
 	const TStaticArray<FSolverVec3, 4> Grads = Base::GetGradients(Particles, ConstraintIndex);
 
-	const FSolverVec3 V1TimesDt = Particles.P(i1) - Particles.X(i1);
-	const FSolverVec3 V2TimesDt = Particles.P(i2) - Particles.X(i2);
-	const FSolverVec3 V3TimesDt = Particles.P(i3) - Particles.X(i3);
-	const FSolverVec3 V4TimesDt = Particles.P(i4) - Particles.X(i4);
+	const FSolverVec3 V1TimesDt = Particles.P(Index1) - Particles.X(Index1);
+	const FSolverVec3 V2TimesDt = Particles.P(Index2) - Particles.X(Index2);
+	const FSolverVec3 V3TimesDt = Particles.P(Index3) - Particles.X(Index3);
+	const FSolverVec3 V4TimesDt = Particles.P(Index4) - Particles.X(Index4);
 
-	FSolverReal& Lambda = Lambdas[ConstraintIndex];
-	const FSolverReal Alpha = (FSolverReal)1.f / (BiphasicStiffnessValue * Dt * Dt);
-	const FSolverReal Gamma = Alpha * Damping * Dt;
+	FSolverReal& Lambda = bDampingOnly ? LambdasDamping[ConstraintIndex] : Lambdas[ConstraintIndex];
+	const FSolverReal AlphaInv = bDampingOnly ? (FSolverReal)0.f : BiphasicStiffnessValue * Dt * Dt;
+	const FSolverReal BetaDt = bElasticOnly ? (FSolverReal)0.f : Damping * Dt;
 
-	const FSolverReal DampingTerm = Gamma * (FSolverVec3::DotProduct(V1TimesDt, Grads[0]) + FSolverVec3::DotProduct(V2TimesDt, Grads[1]) + FSolverVec3::DotProduct(V3TimesDt, Grads[2]) + FSolverVec3::DotProduct(V4TimesDt, Grads[3]));
+	const FSolverReal DampingTerm = bElasticOnly ? (FSolverReal)0.f : BetaDt * (FSolverVec3::DotProduct(V1TimesDt, Grads[0]) + FSolverVec3::DotProduct(V2TimesDt, Grads[1]) + FSolverVec3::DotProduct(V3TimesDt, Grads[2]) + FSolverVec3::DotProduct(V4TimesDt, Grads[3]));
+	const FSolverReal ElasticTerm = bDampingOnly ? (FSolverReal)0.f : AlphaInv * (Angle - RestAngles[ConstraintIndex]);
 
-	const FSolverReal Denom = ((FSolverReal)1.f + Gamma) * (Particles.InvM(i1) * Grads[0].SizeSquared() + Particles.InvM(i2) * Grads[1].SizeSquared() + Particles.InvM(i3) * Grads[2].SizeSquared() + Particles.InvM(i4) * Grads[3].SizeSquared()) + Alpha;
-	const FSolverReal DLambda = (Angle - RestAngles[ConstraintIndex] - Alpha * Lambda + DampingTerm) / Denom;
+	const FSolverReal Denom = (AlphaInv + BetaDt) * (Particles.InvM(Index1) * Grads[0].SizeSquared() + Particles.InvM(Index2) * Grads[1].SizeSquared() + Particles.InvM(Index3) * Grads[2].SizeSquared() + Particles.InvM(Index4) * Grads[3].SizeSquared()) + (FSolverReal)1.;
+	const FSolverReal DLambda = (-ElasticTerm - DampingTerm - Lambda) / Denom;
 
-	Particles.P(i1) -= DLambda * Particles.InvM(i1) * Grads[0];
-	Particles.P(i2) -= DLambda * Particles.InvM(i2) * Grads[1];
-	Particles.P(i3) -= DLambda * Particles.InvM(i3) * Grads[2];
-	Particles.P(i4) -= DLambda * Particles.InvM(i4) * Grads[3];
+	Particles.P(Index1) += DLambda * Particles.InvM(Index1) * Grads[0];
+	Particles.P(Index2) += DLambda * Particles.InvM(Index2) * Grads[1];
+	Particles.P(Index3) += DLambda * Particles.InvM(Index3) * Grads[2];
+	Particles.P(Index4) += DLambda * Particles.InvM(Index4) * Grads[3];
 	Lambda += DLambda;
 }
 
@@ -602,123 +605,245 @@ void FXPBDAnisotropicBendingConstraints::Apply(SolverParticlesOrRange& Particles
 	if (ConstraintsPerColorStartIndex.Num() > 0 && Constraints.Num() > Chaos_XPBDBending_ParallelConstraintCount)
 	{
 		const int32 ConstraintColorNum = ConstraintsPerColorStartIndex.Num() - 1;
-		if (!StiffnessHasWeightMap && !StiffnessWeftHasWeightMap && !StiffnessBiasHasWeightMap &&
-			!BucklingStiffnessHasWeightMap && !BucklingStiffnessWeftHasWeightMap && !BucklingStiffnessBiasHasWeightMap &&
-			!DampingHasWeightMap)
+#if INTEL_ISPC
+		if (bRealTypeCompatibleWithISPC && bChaos_XPBDBending_ISPC_Enabled)
 		{
-			const FSolverVec3 ExpStiffnessValue((FSolverReal)StiffnessWeft, (FSolverReal)Stiffness, (FSolverReal)StiffnessBias);
-			const FSolverVec3 ExpBucklingValue((FSolverReal)BucklingStiffnessWeft, (FSolverReal)BucklingStiffness, (FSolverReal)BucklingStiffnessBias);
-			const FSolverReal DampingRatioValue = (FSolverReal)DampingRatio;
-
-			if (ExpStiffnessValue.Max() < MinStiffness && ExpBucklingValue.Max() < MinStiffness)
+			if (!StiffnessHasWeightMap && !StiffnessWeftHasWeightMap && !StiffnessBiasHasWeightMap &&
+				!BucklingStiffnessHasWeightMap && !BucklingStiffnessWeftHasWeightMap && !BucklingStiffnessBiasHasWeightMap &&
+				!DampingHasWeightMap)
 			{
-				return;
-			}
+				const FSolverVec3 ExpStiffnessValue((FSolverReal)StiffnessWeft, (FSolverReal)Stiffness, (FSolverReal)StiffnessBias);
+				const FSolverVec3 ExpBucklingValue((FSolverReal)BucklingStiffnessWeft, (FSolverReal)BucklingStiffness, (FSolverReal)BucklingStiffnessBias);
+				const FSolverReal DampingRatioValue = (FSolverReal)DampingRatio;
+
+				if (ExpStiffnessValue.Max() < MinStiffness && ExpBucklingValue.Max() < MinStiffness)
+				{
+					return;
+				}
 
-#if INTEL_ISPC
-			if (bRealTypeCompatibleWithISPC && bChaos_XPBDBending_ISPC_Enabled)
-			{
 				if (DampingRatioValue > 0)
 				{
-					for (int32 ConstraintColorIndex = 0; ConstraintColorIndex < ConstraintColorNum; ++ConstraintColorIndex)
+					if (bChaos_XPBDBending_SplitLambdaDamping)
 					{
-						const int32 ColorStart = ConstraintsPerColorStartIndex[ConstraintColorIndex];
-						const int32 ColorSize = ConstraintsPerColorStartIndex[ConstraintColorIndex + 1] - ColorStart;
-						ispc::ApplyXPBDAnisotropicBendingConstraintsWithDamping(
-							(ispc::FVector4f*)Particles.GetPAndInvM().GetData(),
-							(const ispc::FVector3f*)Particles.XArray().GetData(),
-							(ispc::FIntVector4*)&Constraints.GetData()[ColorStart],
-							&RestAngles.GetData()[ColorStart],
-							&IsBuckled.GetData()[ColorStart],
-							(const ispc::FVector3f*)&WarpWeftBiasBaseMultipliers.GetData()[ColorStart],
-							&Lambdas.GetData()[ColorStart],
-							Dt,
-							MinStiffness,
-							reinterpret_cast<const ispc::FVector3f&>(ExpStiffnessValue),
-							reinterpret_cast<const ispc::FVector3f&>(ExpBucklingValue),
-							DampingRatioValue,
-							ColorSize);
+						for (int32 ConstraintColorIndex = 0; ConstraintColorIndex < ConstraintColorNum; ++ConstraintColorIndex)
+						{
+							const int32 ColorStart = ConstraintsPerColorStartIndex[ConstraintColorIndex];
+							const int32 ColorSize = ConstraintsPerColorStartIndex[ConstraintColorIndex + 1] - ColorStart;
+							ispc::ApplyXPBDAnisotropicBendingDampingConstraints(
+								(ispc::FVector4f*)Particles.GetPAndInvM().GetData(),
+								(const ispc::FVector3f*)Particles.XArray().GetData(),
+								(ispc::FIntVector4*)&Constraints.GetData()[ColorStart],
+								&RestAngles.GetData()[ColorStart],
+								&IsBuckled.GetData()[ColorStart],
+								(const ispc::FVector3f*)&WarpWeftBiasBaseMultipliers.GetData()[ColorStart],
+								&LambdasDamping.GetData()[ColorStart],
+								Dt,
+								MinStiffness,
+								reinterpret_cast<const ispc::FVector3f&>(ExpStiffnessValue),
+								reinterpret_cast<const ispc::FVector3f&>(ExpBucklingValue),
+								DampingRatioValue,
+								ColorSize);
+						}
 					}
-				}
-				else
-				{
-					for (int32 ConstraintColorIndex = 0; ConstraintColorIndex < ConstraintColorNum; ++ConstraintColorIndex)
+					else
 					{
-						const int32 ColorStart = ConstraintsPerColorStartIndex[ConstraintColorIndex];
-						const int32 ColorSize = ConstraintsPerColorStartIndex[ConstraintColorIndex + 1] - ColorStart;
-						ispc::ApplyXPBDAnisotropicBendingConstraints(
-							(ispc::FVector4f*)Particles.GetPAndInvM().GetData(),
-							(ispc::FIntVector4*)&Constraints.GetData()[ColorStart],
-							&RestAngles.GetData()[ColorStart],
-							&IsBuckled.GetData()[ColorStart],
-							(const ispc::FVector3f*)&WarpWeftBiasBaseMultipliers.GetData()[ColorStart],
-							&Lambdas.GetData()[ColorStart],
-							Dt,
-							MinStiffness,
-							reinterpret_cast<const ispc::FVector3f&>(ExpStiffnessValue),
-							reinterpret_cast<const ispc::FVector3f&>(ExpBucklingValue),
-							ColorSize);
+						for (int32 ConstraintColorIndex = 0; ConstraintColorIndex < ConstraintColorNum; ++ConstraintColorIndex)
+						{
+							const int32 ColorStart = ConstraintsPerColorStartIndex[ConstraintColorIndex];
+							const int32 ColorSize = ConstraintsPerColorStartIndex[ConstraintColorIndex + 1] - ColorStart;
+							ispc::ApplyXPBDAnisotropicBendingConstraintsWithDamping(
+								(ispc::FVector4f*)Particles.GetPAndInvM().GetData(),
+								(const ispc::FVector3f*)Particles.XArray().GetData(),
+								(ispc::FIntVector4*)&Constraints.GetData()[ColorStart],
+								&RestAngles.GetData()[ColorStart],
+								&IsBuckled.GetData()[ColorStart],
+								(const ispc::FVector3f*)&WarpWeftBiasBaseMultipliers.GetData()[ColorStart],
+								&Lambdas.GetData()[ColorStart],
+								Dt,
+								MinStiffness,
+								reinterpret_cast<const ispc::FVector3f&>(ExpStiffnessValue),
+								reinterpret_cast<const ispc::FVector3f&>(ExpBucklingValue),
+								DampingRatioValue,
+								ColorSize);
+						}
+						return;
 					}
 				}
+				for (int32 ConstraintColorIndex = 0; ConstraintColorIndex < ConstraintColorNum; ++ConstraintColorIndex)
+				{
+					const int32 ColorStart = ConstraintsPerColorStartIndex[ConstraintColorIndex];
+					const int32 ColorSize = ConstraintsPerColorStartIndex[ConstraintColorIndex + 1] - ColorStart;
+					ispc::ApplyXPBDAnisotropicBendingConstraints(
+						(ispc::FVector4f*)Particles.GetPAndInvM().GetData(),
+						(ispc::FIntVector4*)&Constraints.GetData()[ColorStart],
+						&RestAngles.GetData()[ColorStart],
+						&IsBuckled.GetData()[ColorStart],
+						(const ispc::FVector3f*)&WarpWeftBiasBaseMultipliers.GetData()[ColorStart],
+						&Lambdas.GetData()[ColorStart],
+						Dt,
+						MinStiffness,
+						reinterpret_cast<const ispc::FVector3f&>(ExpStiffnessValue),
+						reinterpret_cast<const ispc::FVector3f&>(ExpBucklingValue),
+						ColorSize);
+				}
 			}
 			else
-#endif
 			{
+				// ISPC with maps
+				if (DampingHasWeightMap || (FSolverReal)DampingRatio > 0)
+				{
+					if (bChaos_XPBDBending_SplitLambdaDamping)
+					{
+						for (int32 ConstraintColorIndex = 0; ConstraintColorIndex < ConstraintColorNum; ++ConstraintColorIndex)
+						{
+							const int32 ColorStart = ConstraintsPerColorStartIndex[ConstraintColorIndex];
+							const int32 ColorSize = ConstraintsPerColorStartIndex[ConstraintColorIndex + 1] - ColorStart;
+							ispc::ApplyXPBDAnisotropicBendingDampingConstraintsWithMaps(
+								(ispc::FVector4f*)Particles.GetPAndInvM().GetData(),
+								(const ispc::FVector3f*)Particles.XArray().GetData(),
+								(ispc::FIntVector4*)&Constraints.GetData()[ColorStart],
+								&RestAngles.GetData()[ColorStart],
+								&IsBuckled.GetData()[ColorStart],
+								(const ispc::FVector3f*)&WarpWeftBiasBaseMultipliers.GetData()[ColorStart],
+								&LambdasDamping.GetData()[ColorStart],
+								Dt,
+								MinStiffness,
+								StiffnessHasWeightMap,
+								StiffnessHasWeightMap ? &Stiffness.GetIndices().GetData()[ColorStart] : nullptr,
+								&Stiffness.GetTable().GetData()[0],
+								StiffnessWeftHasWeightMap,
+								StiffnessWeftHasWeightMap ? &StiffnessWeft.GetIndices().GetData()[ColorStart] : nullptr,
+								&StiffnessWeft.GetTable().GetData()[0],
+								StiffnessBiasHasWeightMap,
+								StiffnessBiasHasWeightMap ? &StiffnessBias.GetIndices().GetData()[ColorStart] : nullptr,
+								&StiffnessBias.GetTable().GetData()[0],
+								BucklingStiffnessHasWeightMap,
+								BucklingStiffnessHasWeightMap ? &BucklingStiffness.GetIndices().GetData()[ColorStart] : nullptr,
+								&BucklingStiffness.GetTable().GetData()[0],
+								BucklingStiffnessWeftHasWeightMap,
+								BucklingStiffnessWeftHasWeightMap ? &BucklingStiffnessWeft.GetIndices().GetData()[ColorStart] : nullptr,
+								&BucklingStiffnessWeft.GetTable().GetData()[0],
+								BucklingStiffnessBiasHasWeightMap,
+								BucklingStiffnessBiasHasWeightMap ? &BucklingStiffnessBias.GetIndices().GetData()[ColorStart] : nullptr,
+								&BucklingStiffnessBias.GetTable().GetData()[0],
+								DampingHasWeightMap,
+								DampingHasWeightMap ? &DampingRatio.GetIndices().GetData()[ColorStart] : nullptr,
+								&DampingRatio.GetTable().GetData()[0],
+								ColorSize);
+						}
+					}
+					else
+					{
+						for (int32 ConstraintColorIndex = 0; ConstraintColorIndex < ConstraintColorNum; ++ConstraintColorIndex)
+						{
+							const int32 ColorStart = ConstraintsPerColorStartIndex[ConstraintColorIndex];
+							const int32 ColorSize = ConstraintsPerColorStartIndex[ConstraintColorIndex + 1] - ColorStart;
+							ispc::ApplyXPBDAnisotropicBendingConstraintsWithDampingAndMaps(
+								(ispc::FVector4f*)Particles.GetPAndInvM().GetData(),
+								(const ispc::FVector3f*)Particles.XArray().GetData(),
+								(ispc::FIntVector4*)&Constraints.GetData()[ColorStart],
+								&RestAngles.GetData()[ColorStart],
+								&IsBuckled.GetData()[ColorStart],
+								(const ispc::FVector3f*)&WarpWeftBiasBaseMultipliers.GetData()[ColorStart],
+								&Lambdas.GetData()[ColorStart],
+								Dt,
+								MinStiffness,
+								StiffnessHasWeightMap,
+								StiffnessHasWeightMap ? &Stiffness.GetIndices().GetData()[ColorStart] : nullptr,
+								&Stiffness.GetTable().GetData()[0],
+								StiffnessWeftHasWeightMap,
+								StiffnessWeftHasWeightMap ? &StiffnessWeft.GetIndices().GetData()[ColorStart] : nullptr,
+								&StiffnessWeft.GetTable().GetData()[0],
+								StiffnessBiasHasWeightMap,
+								StiffnessBiasHasWeightMap ? &StiffnessBias.GetIndices().GetData()[ColorStart] : nullptr,
+								&StiffnessBias.GetTable().GetData()[0],
+								BucklingStiffnessHasWeightMap,
+								BucklingStiffnessHasWeightMap ? &BucklingStiffness.GetIndices().GetData()[ColorStart] : nullptr,
+								&BucklingStiffness.GetTable().GetData()[0],
+								BucklingStiffnessWeftHasWeightMap,
+								BucklingStiffnessWeftHasWeightMap ? &BucklingStiffnessWeft.GetIndices().GetData()[ColorStart] : nullptr,
+								&BucklingStiffnessWeft.GetTable().GetData()[0],
+								BucklingStiffnessBiasHasWeightMap,
+								BucklingStiffnessBiasHasWeightMap ? &BucklingStiffnessBias.GetIndices().GetData()[ColorStart] : nullptr,
+								&BucklingStiffnessBias.GetTable().GetData()[0],
+								DampingHasWeightMap,
+								DampingHasWeightMap ? &DampingRatio.GetIndices().GetData()[ColorStart] : nullptr,
+								&DampingRatio.GetTable().GetData()[0],
+								ColorSize);
+						}
+						return;
+					}
+				}
 				for (int32 ConstraintColorIndex = 0; ConstraintColorIndex < ConstraintColorNum; ++ConstraintColorIndex)
 				{
 					const int32 ColorStart = ConstraintsPerColorStartIndex[ConstraintColorIndex];
 					const int32 ColorSize = ConstraintsPerColorStartIndex[ConstraintColorIndex + 1] - ColorStart;
-					PhysicsParallelFor(ColorSize, [&](const int32 Index)
-					{
-						const int32 ConstraintIndex = ColorStart + Index;
-						ApplyHelper(Particles, Dt, ConstraintIndex, ExpStiffnessValue, ExpBucklingValue, DampingRatioValue);
-					});
+					ispc::ApplyXPBDAnisotropicBendingConstraintsWithMaps(
+						(ispc::FVector4f*)Particles.GetPAndInvM().GetData(),
+						(ispc::FIntVector4*)&Constraints.GetData()[ColorStart],
+						&RestAngles.GetData()[ColorStart],
+						&IsBuckled.GetData()[ColorStart],
+						(const ispc::FVector3f*)&WarpWeftBiasBaseMultipliers.GetData()[ColorStart],
+						&Lambdas.GetData()[ColorStart],
+						Dt,
+						MinStiffness,
+						StiffnessHasWeightMap,
+						StiffnessHasWeightMap ? &Stiffness.GetIndices().GetData()[ColorStart] : nullptr,
+						&Stiffness.GetTable().GetData()[0],
+						StiffnessWeftHasWeightMap,
+						StiffnessWeftHasWeightMap ? &StiffnessWeft.GetIndices().GetData()[ColorStart] : nullptr,
+						&StiffnessWeft.GetTable().GetData()[0],
+						StiffnessBiasHasWeightMap,
+						StiffnessBiasHasWeightMap ? &StiffnessBias.GetIndices().GetData()[ColorStart] : nullptr,
+						&StiffnessBias.GetTable().GetData()[0],
+						BucklingStiffnessHasWeightMap,
+						BucklingStiffnessHasWeightMap ? &BucklingStiffness.GetIndices().GetData()[ColorStart] : nullptr,
+						&BucklingStiffness.GetTable().GetData()[0],
+						BucklingStiffnessWeftHasWeightMap,
+						BucklingStiffnessWeftHasWeightMap ? &BucklingStiffnessWeft.GetIndices().GetData()[ColorStart] : nullptr,
+						&BucklingStiffnessWeft.GetTable().GetData()[0],
+						BucklingStiffnessBiasHasWeightMap,
+						BucklingStiffnessBiasHasWeightMap ? &BucklingStiffnessBias.GetIndices().GetData()[ColorStart] : nullptr,
+						&BucklingStiffnessBias.GetTable().GetData()[0],
+						ColorSize);
 				}
 			}
 		}
-		else  // Has weight maps
+		else
+#endif
 		{
-#if INTEL_ISPC
-			if (bRealTypeCompatibleWithISPC && bChaos_XPBDBending_ISPC_Enabled)
+			// Parallel non-ispc
+			const FSolverReal StiffnessNoMap = (FSolverReal)Stiffness;
+			const FSolverReal StiffnessWeftNoMap = (FSolverReal)StiffnessWeft;
+			const FSolverReal StiffnessBiasNoMap = (FSolverReal)StiffnessBias;
+			const FSolverReal BucklingStiffnessNoMap = (FSolverReal)BucklingStiffness;
+			const FSolverReal BucklingStiffnessWeftNoMap = (FSolverReal)BucklingStiffnessWeft;
+			const FSolverReal BucklingStiffnessBiasNoMap = (FSolverReal)BucklingStiffnessBias;
+			const FSolverReal DampingNoMap = (FSolverReal)DampingRatio;
+
+			if (DampingHasWeightMap || (FSolverReal)DampingRatio > 0)
 			{
-				if (DampingHasWeightMap || (FSolverReal)DampingRatio > 0)
+				if (bChaos_XPBDBending_SplitLambdaDamping)
 				{
 					for (int32 ConstraintColorIndex = 0; ConstraintColorIndex < ConstraintColorNum; ++ConstraintColorIndex)
 					{
 						const int32 ColorStart = ConstraintsPerColorStartIndex[ConstraintColorIndex];
 						const int32 ColorSize = ConstraintsPerColorStartIndex[ConstraintColorIndex + 1] - ColorStart;
-						ispc::ApplyXPBDAnisotropicBendingConstraintsWithDampingAndMaps(
-							(ispc::FVector4f*)Particles.GetPAndInvM().GetData(),
-							(const ispc::FVector3f*)Particles.XArray().GetData(),
-							(ispc::FIntVector4*)&Constraints.GetData()[ColorStart],
-							&RestAngles.GetData()[ColorStart],
-							&IsBuckled.GetData()[ColorStart],
-							(const ispc::FVector3f*)&WarpWeftBiasBaseMultipliers.GetData()[ColorStart],
-							&Lambdas.GetData()[ColorStart],
-							Dt,
-							MinStiffness,
-							StiffnessHasWeightMap,
-							StiffnessHasWeightMap ? &Stiffness.GetIndices().GetData()[ColorStart] : nullptr,
-							&Stiffness.GetTable().GetData()[0],
-							StiffnessWeftHasWeightMap,
-							StiffnessWeftHasWeightMap ? &StiffnessWeft.GetIndices().GetData()[ColorStart] : nullptr,
-							&StiffnessWeft.GetTable().GetData()[0],
-							StiffnessBiasHasWeightMap,
-							StiffnessBiasHasWeightMap ? &StiffnessBias.GetIndices().GetData()[ColorStart] : nullptr,
-							&StiffnessBias.GetTable().GetData()[0],
-							BucklingStiffnessHasWeightMap,
-							BucklingStiffnessHasWeightMap ? &BucklingStiffness.GetIndices().GetData()[ColorStart] : nullptr,
-							&BucklingStiffness.GetTable().GetData()[0],
-							BucklingStiffnessWeftHasWeightMap,
-							BucklingStiffnessWeftHasWeightMap ? &BucklingStiffnessWeft.GetIndices().GetData()[ColorStart] : nullptr,
-							&BucklingStiffnessWeft.GetTable().GetData()[0],
-							BucklingStiffnessBiasHasWeightMap,
-							BucklingStiffnessBiasHasWeightMap ? &BucklingStiffnessBias.GetIndices().GetData()[ColorStart] : nullptr,
-							&BucklingStiffnessBias.GetTable().GetData()[0],
-							DampingHasWeightMap,
-							DampingHasWeightMap ? &DampingRatio.GetIndices().GetData()[ColorStart] : nullptr,
-							&DampingRatio.GetTable().GetData()[0],
-							ColorSize);
+						PhysicsParallelFor(ColorSize, [&](const int32 Index)
+						{
+							const int32 ConstraintIndex = ColorStart + Index;
+							const FSolverReal ExpStiffnessValue = StiffnessHasWeightMap ? Stiffness[ConstraintIndex] : StiffnessNoMap;
+							const FSolverReal ExpStiffnessWeftValue = StiffnessWeftHasWeightMap ? StiffnessWeft[ConstraintIndex] : StiffnessWeftNoMap;
+							const FSolverReal ExpStiffnessBiasValue = StiffnessBiasHasWeightMap ? StiffnessBias[ConstraintIndex] : StiffnessBiasNoMap;
+							const FSolverReal ExpBucklingValue = BucklingStiffnessHasWeightMap ? BucklingStiffness[ConstraintIndex] : BucklingStiffnessNoMap;
+							const FSolverReal ExpBucklingWeftValue = BucklingStiffnessWeftHasWeightMap ? BucklingStiffnessWeft[ConstraintIndex] : BucklingStiffnessWeftNoMap;
+							const FSolverReal ExpBucklingBiasValue = BucklingStiffnessBiasHasWeightMap ? BucklingStiffnessBias[ConstraintIndex] : BucklingStiffnessBiasNoMap;
+							const FSolverReal DampingRatioValue = DampingHasWeightMap ? DampingRatio[ConstraintIndex] : DampingNoMap;
+							constexpr bool bDampingOnly = true;
+							constexpr bool bElasticOnly = false;
+							ApplyHelper<bDampingOnly, bElasticOnly>(Particles, Dt, ConstraintIndex,
+								FSolverVec3(ExpStiffnessWeftValue, ExpStiffnessValue, ExpStiffnessBiasValue), FSolverVec3(ExpBucklingWeftValue, ExpBucklingValue, ExpBucklingBiasValue), DampingRatioValue);
+						});
 					}
 				}
 				else
@@ -727,47 +852,25 @@ void FXPBDAnisotropicBendingConstraints::Apply(SolverParticlesOrRange& Particles
 					{
 						const int32 ColorStart = ConstraintsPerColorStartIndex[ConstraintColorIndex];
 						const int32 ColorSize = ConstraintsPerColorStartIndex[ConstraintColorIndex + 1] - ColorStart;
-						ispc::ApplyXPBDAnisotropicBendingConstraintsWithMaps(
-							(ispc::FVector4f*)Particles.GetPAndInvM().GetData(),
-							(ispc::FIntVector4*)&Constraints.GetData()[ColorStart],
-							&RestAngles.GetData()[ColorStart],
-							&IsBuckled.GetData()[ColorStart],
-							(const ispc::FVector3f*)&WarpWeftBiasBaseMultipliers.GetData()[ColorStart],
-							&Lambdas.GetData()[ColorStart],
-							Dt,
-							MinStiffness,
-							StiffnessHasWeightMap,
-							StiffnessHasWeightMap ? &Stiffness.GetIndices().GetData()[ColorStart] : nullptr,
-							&Stiffness.GetTable().GetData()[0],
-							StiffnessWeftHasWeightMap,
-							StiffnessWeftHasWeightMap ? &StiffnessWeft.GetIndices().GetData()[ColorStart] : nullptr,
-							&StiffnessWeft.GetTable().GetData()[0],
-							StiffnessBiasHasWeightMap,
-							StiffnessBiasHasWeightMap ? &StiffnessBias.GetIndices().GetData()[ColorStart] : nullptr,
-							&StiffnessBias.GetTable().GetData()[0],
-							BucklingStiffnessHasWeightMap,
-							BucklingStiffnessHasWeightMap ? &BucklingStiffness.GetIndices().GetData()[ColorStart] : nullptr,
-							&BucklingStiffness.GetTable().GetData()[0],
-							BucklingStiffnessWeftHasWeightMap,
-							BucklingStiffnessWeftHasWeightMap ? &BucklingStiffnessWeft.GetIndices().GetData()[ColorStart] : nullptr,
-							&BucklingStiffnessWeft.GetTable().GetData()[0],
-							BucklingStiffnessBiasHasWeightMap,
-							BucklingStiffnessBiasHasWeightMap ? &BucklingStiffnessBias.GetIndices().GetData()[ColorStart] : nullptr,
-							&BucklingStiffnessBias.GetTable().GetData()[0],
-							ColorSize);
+						PhysicsParallelFor(ColorSize, [&](const int32 Index)
+						{
+							const int32 ConstraintIndex = ColorStart + Index;
+							const FSolverReal ExpStiffnessValue = StiffnessHasWeightMap ? Stiffness[ConstraintIndex] : StiffnessNoMap;
+							const FSolverReal ExpStiffnessWeftValue = StiffnessWeftHasWeightMap ? StiffnessWeft[ConstraintIndex] : StiffnessWeftNoMap;
+							const FSolverReal ExpStiffnessBiasValue = StiffnessBiasHasWeightMap ? StiffnessBias[ConstraintIndex] : StiffnessBiasNoMap;
+							const FSolverReal ExpBucklingValue = BucklingStiffnessHasWeightMap ? BucklingStiffness[ConstraintIndex] : BucklingStiffnessNoMap;
+							const FSolverReal ExpBucklingWeftValue = BucklingStiffnessWeftHasWeightMap ? BucklingStiffnessWeft[ConstraintIndex] : BucklingStiffnessWeftNoMap;
+							const FSolverReal ExpBucklingBiasValue = BucklingStiffnessBiasHasWeightMap ? BucklingStiffnessBias[ConstraintIndex] : BucklingStiffnessBiasNoMap;
+							const FSolverReal DampingRatioValue = DampingHasWeightMap ? DampingRatio[ConstraintIndex] : DampingNoMap;
+							constexpr bool bDampingOnly = false;
+							constexpr bool bElasticOnly = false;
+							ApplyHelper<bDampingOnly, bElasticOnly>(Particles, Dt, ConstraintIndex,
+								FSolverVec3(ExpStiffnessWeftValue, ExpStiffnessValue, ExpStiffnessBiasValue), FSolverVec3(ExpBucklingWeftValue, ExpBucklingValue, ExpBucklingBiasValue), DampingRatioValue);
+						});
 					}
+					return;
 				}
-			}
-			else
-#endif
-			{
-				const FSolverReal StiffnessNoMap = (FSolverReal)Stiffness;
-				const FSolverReal StiffnessWeftNoMap = (FSolverReal)StiffnessWeft;
-				const FSolverReal StiffnessBiasNoMap = (FSolverReal)StiffnessBias;
-				const FSolverReal BucklingStiffnessNoMap = (FSolverReal)BucklingStiffness;
-				const FSolverReal BucklingStiffnessWeftNoMap = (FSolverReal)BucklingStiffnessWeft;
-				const FSolverReal BucklingStiffnessBiasNoMap = (FSolverReal)BucklingStiffnessBias;
-				const FSolverReal DampingNoMap = (FSolverReal)DampingRatio;
+
 				for (int32 ConstraintColorIndex = 0; ConstraintColorIndex < ConstraintColorNum; ++ConstraintColorIndex)
 				{
 					const int32 ColorStart = ConstraintsPerColorStartIndex[ConstraintColorIndex];
@@ -782,8 +885,10 @@ void FXPBDAnisotropicBendingConstraints::Apply(SolverParticlesOrRange& Particles
 						const FSolverReal ExpBucklingWeftValue = BucklingStiffnessWeftHasWeightMap ? BucklingStiffnessWeft[ConstraintIndex] : BucklingStiffnessWeftNoMap;
 						const FSolverReal ExpBucklingBiasValue = BucklingStiffnessBiasHasWeightMap ? BucklingStiffnessBias[ConstraintIndex] : BucklingStiffnessBiasNoMap;
 						const FSolverReal DampingRatioValue = DampingHasWeightMap ? DampingRatio[ConstraintIndex] : DampingNoMap;
-						ApplyHelper(Particles, Dt, ConstraintIndex, FSolverVec3(ExpStiffnessWeftValue, ExpStiffnessValue, ExpStiffnessBiasValue),
-							FSolverVec3(ExpBucklingWeftValue, ExpBucklingValue, ExpBucklingBiasValue), DampingRatioValue);
+						constexpr bool bDampingOnly = false;
+						constexpr bool bElasticOnly = true;
+						ApplyHelper<bDampingOnly, bElasticOnly>(Particles, Dt, ConstraintIndex,
+							FSolverVec3(ExpStiffnessWeftValue, ExpStiffnessValue, ExpStiffnessBiasValue), FSolverVec3(ExpBucklingWeftValue, ExpBucklingValue, ExpBucklingBiasValue), DampingRatioValue);
 					});
 				}
 			}
@@ -791,45 +896,66 @@ void FXPBDAnisotropicBendingConstraints::Apply(SolverParticlesOrRange& Particles
 	}
 	else
 	{
-		if (!StiffnessHasWeightMap && !StiffnessWeftHasWeightMap && !StiffnessBiasHasWeightMap &&
-			!BucklingStiffnessHasWeightMap && !BucklingStiffnessWeftHasWeightMap && !BucklingStiffnessBiasHasWeightMap &&
-			!DampingHasWeightMap)
+		// Single-threaded
+		const FSolverReal StiffnessNoMap = (FSolverReal)Stiffness;
+		const FSolverReal StiffnessWeftNoMap = (FSolverReal)StiffnessWeft;
+		const FSolverReal StiffnessBiasNoMap = (FSolverReal)StiffnessBias;
+		const FSolverReal BucklingStiffnessNoMap = (FSolverReal)BucklingStiffness;
+		const FSolverReal BucklingStiffnessWeftNoMap = (FSolverReal)BucklingStiffnessWeft;
+		const FSolverReal BucklingStiffnessBiasNoMap = (FSolverReal)BucklingStiffnessBias;
+		const FSolverReal DampingNoMap = (FSolverReal)DampingRatio;
+		if (DampingHasWeightMap || (FSolverReal)DampingRatio > 0)
 		{
-			const FSolverVec3 ExpStiffnessValue((FSolverReal)StiffnessWeft, (FSolverReal)Stiffness, (FSolverReal)StiffnessBias);
-			const FSolverVec3 ExpBucklingValue((FSolverReal)BucklingStiffnessWeft, (FSolverReal)BucklingStiffness, (FSolverReal)BucklingStiffnessBias);
-			const FSolverReal DampingRatioValue = (FSolverReal)DampingRatio;
-
-			if (ExpStiffnessValue.Max() < MinStiffness && ExpBucklingValue.Max() < MinStiffness)
+			if (bChaos_XPBDBending_SplitLambdaDamping)
 			{
-				return;
+				for (int32 ConstraintIndex = 0; ConstraintIndex < Constraints.Num(); ++ConstraintIndex)
+				{
+					const FSolverReal ExpStiffnessValue = StiffnessHasWeightMap ? Stiffness[ConstraintIndex] : StiffnessNoMap;
+					const FSolverReal ExpStiffnessWeftValue = StiffnessWeftHasWeightMap ? StiffnessWeft[ConstraintIndex] : StiffnessWeftNoMap;
+					const FSolverReal ExpStiffnessBiasValue = StiffnessBiasHasWeightMap ? StiffnessBias[ConstraintIndex] : StiffnessBiasNoMap;
+					const FSolverReal ExpBucklingValue = BucklingStiffnessHasWeightMap ? BucklingStiffness[ConstraintIndex] : BucklingStiffnessNoMap;
+					const FSolverReal ExpBucklingWeftValue = BucklingStiffnessWeftHasWeightMap ? BucklingStiffnessWeft[ConstraintIndex] : BucklingStiffnessWeftNoMap;
+					const FSolverReal ExpBucklingBiasValue = BucklingStiffnessBiasHasWeightMap ? BucklingStiffnessBias[ConstraintIndex] : BucklingStiffnessBiasNoMap;
+					const FSolverReal DampingRatioValue = DampingHasWeightMap ? DampingRatio[ConstraintIndex] : DampingNoMap;
+					constexpr bool bDampingOnly = true;
+					constexpr bool bElasticOnly = false;
+					ApplyHelper<bDampingOnly, bElasticOnly>(Particles, Dt, ConstraintIndex,
+						FSolverVec3(ExpStiffnessWeftValue, ExpStiffnessValue, ExpStiffnessBiasValue), FSolverVec3(ExpBucklingWeftValue, ExpBucklingValue, ExpBucklingBiasValue), DampingRatioValue);
+				}
 			}
-
-			for (int32 ConstraintIndex = 0; ConstraintIndex < Constraints.Num(); ++ConstraintIndex)
+			else
 			{
-				ApplyHelper(Particles, Dt, ConstraintIndex, ExpStiffnessValue, ExpBucklingValue, DampingRatioValue);
+				for (int32 ConstraintIndex = 0; ConstraintIndex < Constraints.Num(); ++ConstraintIndex)
+				{
+					const FSolverReal ExpStiffnessValue = StiffnessHasWeightMap ? Stiffness[ConstraintIndex] : StiffnessNoMap;
+					const FSolverReal ExpStiffnessWeftValue = StiffnessWeftHasWeightMap ? StiffnessWeft[ConstraintIndex] : StiffnessWeftNoMap;
+					const FSolverReal ExpStiffnessBiasValue = StiffnessBiasHasWeightMap ? StiffnessBias[ConstraintIndex] : StiffnessBiasNoMap;
+					const FSolverReal ExpBucklingValue = BucklingStiffnessHasWeightMap ? BucklingStiffness[ConstraintIndex] : BucklingStiffnessNoMap;
+					const FSolverReal ExpBucklingWeftValue = BucklingStiffnessWeftHasWeightMap ? BucklingStiffnessWeft[ConstraintIndex] : BucklingStiffnessWeftNoMap;
+					const FSolverReal ExpBucklingBiasValue = BucklingStiffnessBiasHasWeightMap ? BucklingStiffnessBias[ConstraintIndex] : BucklingStiffnessBiasNoMap;
+					const FSolverReal DampingRatioValue = DampingHasWeightMap ? DampingRatio[ConstraintIndex] : DampingNoMap;
+					constexpr bool bDampingOnly = false;
+					constexpr bool bElasticOnly = false;
+					ApplyHelper<bDampingOnly, bElasticOnly>(Particles, Dt, ConstraintIndex,
+						FSolverVec3(ExpStiffnessWeftValue, ExpStiffnessValue, ExpStiffnessBiasValue), FSolverVec3(ExpBucklingWeftValue, ExpBucklingValue, ExpBucklingBiasValue), DampingRatioValue);
+				}
+				return;
 			}
 		}
-		else
+
+		for (int32 ConstraintIndex = 0; ConstraintIndex < Constraints.Num(); ++ConstraintIndex)
 		{
-			const FSolverReal StiffnessNoMap = (FSolverReal)Stiffness;
-			const FSolverReal StiffnessWeftNoMap = (FSolverReal)StiffnessWeft;
-			const FSolverReal StiffnessBiasNoMap = (FSolverReal)StiffnessBias;
-			const FSolverReal BucklingStiffnessNoMap = (FSolverReal)BucklingStiffness;
-			const FSolverReal BucklingStiffnessWeftNoMap = (FSolverReal)BucklingStiffnessWeft;
-			const FSolverReal BucklingStiffnessBiasNoMap = (FSolverReal)BucklingStiffnessBias;
-			const FSolverReal DampingNoMap = (FSolverReal)DampingRatio;
-			for (int32 ConstraintIndex = 0; ConstraintIndex < Constraints.Num(); ++ConstraintIndex)
-			{
-				const FSolverReal ExpStiffnessValue = StiffnessHasWeightMap ? Stiffness[ConstraintIndex] : StiffnessNoMap;
-				const FSolverReal ExpStiffnessWeftValue = StiffnessWeftHasWeightMap ? StiffnessWeft[ConstraintIndex] : StiffnessWeftNoMap;
-				const FSolverReal ExpStiffnessBiasValue = StiffnessBiasHasWeightMap ? StiffnessBias[ConstraintIndex] : StiffnessBiasNoMap;
-				const FSolverReal ExpBucklingValue = BucklingStiffnessHasWeightMap ? BucklingStiffness[ConstraintIndex] : BucklingStiffnessNoMap;
-				const FSolverReal ExpBucklingWeftValue = BucklingStiffnessWeftHasWeightMap ? BucklingStiffnessWeft[ConstraintIndex] : BucklingStiffnessWeftNoMap;
-				const FSolverReal ExpBucklingBiasValue = BucklingStiffnessBiasHasWeightMap ? BucklingStiffnessBias[ConstraintIndex] : BucklingStiffnessBiasNoMap;
-				const FSolverReal DampingRatioValue = DampingHasWeightMap ? DampingRatio[ConstraintIndex] : DampingNoMap;
-				ApplyHelper(Particles, Dt, ConstraintIndex, FSolverVec3(ExpStiffnessWeftValue, ExpStiffnessValue, ExpStiffnessBiasValue),
-					FSolverVec3(ExpBucklingWeftValue, ExpBucklingValue, ExpBucklingBiasValue), DampingRatioValue);
-			}
+			const FSolverReal ExpStiffnessValue = StiffnessHasWeightMap ? Stiffness[ConstraintIndex] : StiffnessNoMap;
+			const FSolverReal ExpStiffnessWeftValue = StiffnessWeftHasWeightMap ? StiffnessWeft[ConstraintIndex] : StiffnessWeftNoMap;
+			const FSolverReal ExpStiffnessBiasValue = StiffnessBiasHasWeightMap ? StiffnessBias[ConstraintIndex] : StiffnessBiasNoMap;
+			const FSolverReal ExpBucklingValue = BucklingStiffnessHasWeightMap ? BucklingStiffness[ConstraintIndex] : BucklingStiffnessNoMap;
+			const FSolverReal ExpBucklingWeftValue = BucklingStiffnessWeftHasWeightMap ? BucklingStiffnessWeft[ConstraintIndex] : BucklingStiffnessWeftNoMap;
+			const FSolverReal ExpBucklingBiasValue = BucklingStiffnessBiasHasWeightMap ? BucklingStiffnessBias[ConstraintIndex] : BucklingStiffnessBiasNoMap;
+			const FSolverReal DampingRatioValue = DampingHasWeightMap ? DampingRatio[ConstraintIndex] : DampingNoMap;
+			constexpr bool bDampingOnly = false;
+			constexpr bool bElasticOnly = true;
+			ApplyHelper<bDampingOnly, bElasticOnly>(Particles, Dt, ConstraintIndex,
+				FSolverVec3(ExpStiffnessWeftValue, ExpStiffnessValue, ExpStiffnessBiasValue), FSolverVec3(ExpBucklingWeftValue, ExpBucklingValue, ExpBucklingBiasValue), DampingRatioValue);
 		}
 	}
 }
diff --git a/Engine/Source/Runtime/Experimental/Chaos/Private/Chaos/XPBDAnisotropicBendingConstraints.ispc b/Engine/Source/Runtime/Experimental/Chaos/Private/Chaos/XPBDAnisotropicBendingConstraints.ispc
--- a/Engine/Source/Runtime/Experimental/Chaos/Private/Chaos/XPBDAnisotropicBendingConstraints.ispc
+++ b/Engine/Source/Runtime/Experimental/Chaos/Private/Chaos/XPBDAnisotropicBendingConstraints.ispc
@@ -4,51 +4,74 @@
 
 #include "Math/Vector.isph"
 #include "Chaos/PBDSofts.isph"
+#include "Chaos/XPBDBendingElement.isph"
 
-static inline FVector3f SafeDivide(const FVector3f& Numerator, const float Denominator)
+export void ApplyXPBDAnisotropicBendingConstraints(uniform FVector4f PandInvM[],
+									const uniform FIntVector4 Constraints[],
+									const uniform float RestAngles[],
+									const uniform bool IsBuckled[],
+									const uniform FVector3f WarpWeftBiasBaseMultiplers[],
+									uniform float Lambdas[],
+									const uniform float Dt,
+									const uniform float MinStiffness,
+									const uniform FVector3f& Stiffness,
+									const uniform FVector3f& BucklingStiffness,
+									const uniform int32 NumConstraints)
 {
-	return VectorSelect(Denominator <= FLOAT_SMALL_NUMBER, FloatZeroVector, Numerator / Denominator);
-}
 
-static inline varying float SafeDivide(const varying float Numerator, const varying float Denominator)
-{
-	return select(Denominator <= FLOAT_SMALL_NUMBER, (uniform float)0, Numerator / Denominator);
-}
+	foreach(i = 0 ... NumConstraints)
+	{
+		const varying FIntVector4 Constraint = VectorLoad(&Constraints[extract(i,0)]);
 
-// Returns Angle
-static inline float CalcGradientsAndAngle(const FVector3f& P1, const FVector3f& P2, const FVector3f& P3, const FVector3f& P4, FVector3f& Grad1, FVector3f& Grad2, FVector3f& Grad3, FVector3f& Grad4)
-{
-	const FVector3f SharedEdgeNormalized = VectorGetSafeNormal(P2 - P1);
-	const FVector3f P13CrossP23 = VectorCross(P1 - P3, P2 - P3);
-	const float Normal1Len = VectorSize(P13CrossP23);
-	const FVector3f Normal1 = SafeDivide(P13CrossP23, Normal1Len);
-	const FVector3f P24CrossP14 = VectorCross(P2 - P4, P1 - P4);
-	const float Normal2Len = VectorSize(P24CrossP14);
-	const FVector3f Normal2 = SafeDivide(P24CrossP14, Normal2Len);
+		const varying int32 i1 = Constraint.V[0];
+		const varying int32 i2 = Constraint.V[1];
+		const varying int32 i3 = Constraint.V[2];
+		const varying int32 i4 = Constraint.V[3];
+		
+		const varying FVector4f PandInvM1 = VectorGather(&PandInvM[i1]);
+		const varying FVector4f PandInvM2 = VectorGather(&PandInvM[i2]);
+		const varying FVector4f PandInvM3 = VectorGather(&PandInvM[i3]);
+		const varying FVector4f PandInvM4 = VectorGather(&PandInvM[i4]);
 
-	const FVector3f N2CrossN1 = VectorCross(Normal2, Normal1);
+		varying FVector3f P1, P2, P3, P4;
+		varying float M1, M2, M3, M4;
 
-	const float CosPhi = clamp(VectorDot(Normal1, Normal2), (uniform float)(-1), (uniform float)(1));
-	const float SinPhi = clamp(VectorDot(N2CrossN1, SharedEdgeNormalized), (uniform float)(-1), (uniform float)(1));
+		UnzipPandInvM(PandInvM1, P1, M1);
+		UnzipPandInvM(PandInvM2, P2, M2);
+		UnzipPandInvM(PandInvM3, P3, M3);
+		UnzipPandInvM(PandInvM4, P4, M4);
 
-	const FVector3f DPhiDN1_OverNormal1Len = SafeDivide(CosPhi * VectorCross(SharedEdgeNormalized, Normal2) - SinPhi * Normal2, Normal1Len);
-	const FVector3f DPhiDN2_OverNormal2Len = SafeDivide(CosPhi * VectorCross(Normal1, SharedEdgeNormalized) - SinPhi * Normal1, Normal2Len);
+		const varying FVector3f WarpWeftBiasBaseMultiplier = VectorGather(&WarpWeftBiasBaseMultiplers[i]);
+		const varying float StiffnessValue = select(IsBuckled[i], VectorDot(WarpWeftBiasBaseMultiplier, BucklingStiffness), VectorDot(WarpWeftBiasBaseMultiplier, Stiffness));
+		
+		const varying float Lambda = Lambdas[i];
 
-	const FVector3f DPhiDP13 = VectorCross(P2 - P3, DPhiDN1_OverNormal1Len);
-	const FVector3f DPhiDP23 = VectorCross(DPhiDN1_OverNormal1Len, P1 - P3);
-	const FVector3f DPhiDP24 = VectorCross(P1 - P4, DPhiDN2_OverNormal2Len);
-	const FVector3f DPhiDP14 = VectorCross(DPhiDN2_OverNormal2Len, P2 - P4);
+		varying FVector3f Delta1, Delta2, Delta3, Delta4;
+		const varying float DLambda = GetXPBDBendingElementDeltas(P1, P2, P3, P4, M1, M2, M3, M4, Dt, RestAngles[i], Lambda, StiffnessValue, MinStiffness, Delta1, Delta2, Delta3, Delta4);
 
-	Grad1 = DPhiDP13 + DPhiDP14;
-	Grad2 = DPhiDP23 + DPhiDP24;
-	Grad3 = -1 * DPhiDP13 - DPhiDP23;
-	Grad4 = -1 * DPhiDP14 - DPhiDP24;
+		Lambdas[i] = Lambda + DLambda;
 
-	return atan2(SinPhi, CosPhi);
+		if (M1 > 0)
+		{
+			VectorScatter(&PandInvM[i1], SetVector4(P1 + Delta1, M1));
+		}
+		if (M2 > 0)
+		{
+			VectorScatter(&PandInvM[i2], SetVector4(P2 + Delta2, M2));
+		}
+		if (M3 > 0)
+		{
+			VectorScatter(&PandInvM[i3], SetVector4(P3 + Delta3, M3));
+		}
+		if (M4 > 0)
+		{
+			VectorScatter(&PandInvM[i4], SetVector4(P4 + Delta4, M4));
+		}
+	}
 }
 
-
-export void ApplyXPBDAnisotropicBendingConstraints(uniform FVector4f PandInvM[],
+export void ApplyXPBDAnisotropicBendingConstraintsWithDamping(uniform FVector4f PandInvM[],
+									const uniform FVector3f X[],
 									const uniform FIntVector4 Constraints[],
 									const uniform float RestAngles[],
 									const uniform bool IsBuckled[],
@@ -58,6 +81,7 @@ export void ApplyXPBDAnisotropicBendingConstraints(uniform FVector4f PandInvM[],
 									const uniform float MinStiffness,
 									const uniform FVector3f& Stiffness,
 									const uniform FVector3f& BucklingStiffness,
+									const uniform float DampingRatio,
 									const uniform int32 NumConstraints)
 {
 
@@ -83,41 +107,40 @@ export void ApplyXPBDAnisotropicBendingConstraints(uniform FVector4f PandInvM[],
 		UnzipPandInvM(PandInvM3, P3, M3);
 		UnzipPandInvM(PandInvM4, P4, M4);
 
-		varying FVector3f Grad1, Grad2, Grad3, Grad4;
-		const varying float Angle = CalcGradientsAndAngle(P1, P2, P3, P4, Grad1, Grad2, Grad3, Grad4);
-		
-		const varying float RestAngle = RestAngles[i];
+		const varying FVector3f X1 = VectorGather(&X[i1]);
+		const varying FVector3f X2 = VectorGather(&X[i2]);
+		const varying FVector3f X3 = VectorGather(&X[i3]);
+		const varying FVector3f X4 = VectorGather(&X[i4]);
+
 		const varying FVector3f WarpWeftBiasBaseMultiplier = VectorGather(&WarpWeftBiasBaseMultiplers[i]);
 		const varying float StiffnessValue = select(IsBuckled[i], VectorDot(WarpWeftBiasBaseMultiplier, BucklingStiffness), VectorDot(WarpWeftBiasBaseMultiplier, Stiffness));
-		
 		const varying float Lambda = Lambdas[i];
-		const varying float Alpha  = 1.f / (StiffnessValue * Dt * Dt);
-		const varying float Denom = M1 * VectorSizeSquared(Grad1) + M2 * VectorSizeSquared(Grad2) + M3 * VectorSizeSquared(Grad3) + M4 * VectorSizeSquared(Grad4) + Alpha;
 
-		const varying float ComputedDLambda = (Angle - RestAngle - Alpha * Lambda) / Denom;
-		const varying float DLambda = select(StiffnessValue < MinStiffness, (uniform float)0, ComputedDLambda);
+		varying FVector3f Delta1, Delta2, Delta3, Delta4;
+		const varying float DLambda = GetXPBDBendingElementDeltasWithDamping(P1, P2, P3, P4, X1, X2, X3, X4, M1, M2, M3, M4, Dt, RestAngles[i], Lambda, StiffnessValue, DampingRatio, MinStiffness, Delta1, Delta2, Delta3, Delta4);
+
 		Lambdas[i] = Lambda + DLambda;
 
 		if (M1 > 0)
 		{
-			VectorScatter(&PandInvM[i1], SetVector4(P1 - DLambda * M1 * Grad1, M1));
+			VectorScatter(&PandInvM[i1], SetVector4(P1 + Delta1, M1));
 		}
 		if (M2 > 0)
 		{
-			VectorScatter(&PandInvM[i2], SetVector4(P2 - DLambda * M2 * Grad2, M2));
+			VectorScatter(&PandInvM[i2], SetVector4(P2 + Delta2, M2));
 		}
 		if (M3 > 0)
 		{
-			VectorScatter(&PandInvM[i3], SetVector4(P3 - DLambda * M3 * Grad3, M3));
+			VectorScatter(&PandInvM[i3], SetVector4(P3 + Delta3, M3));
 		}
 		if (M4 > 0)
 		{
-			VectorScatter(&PandInvM[i4], SetVector4(P4 - DLambda * M4 * Grad4, M4));
+			VectorScatter(&PandInvM[i4], SetVector4(P4 + Delta4, M4));
 		}
 	}
 }
 
-export void ApplyXPBDAnisotropicBendingConstraintsWithDamping(uniform FVector4f PandInvM[],
+export void ApplyXPBDAnisotropicBendingDampingConstraints(uniform FVector4f PandInvM[],
 									const uniform FVector3f X[],
 									const uniform FIntVector4 Constraints[],
 									const uniform float RestAngles[],
@@ -154,52 +177,35 @@ export void ApplyXPBDAnisotropicBendingConstraintsWithDamping(uniform FVector4f
 		UnzipPandInvM(PandInvM3, P3, M3);
 		UnzipPandInvM(PandInvM4, P4, M4);
 
-		varying FVector3f Grad1, Grad2, Grad3, Grad4;
-		const varying float Angle = CalcGradientsAndAngle(P1, P2, P3, P4, Grad1, Grad2, Grad3, Grad4);
-		
-		const varying float RestAngle = RestAngles[i];
-		const varying FVector3f WarpWeftBiasBaseMultiplier = VectorGather(&WarpWeftBiasBaseMultiplers[i]);
-		const varying float StiffnessValue = select(IsBuckled[i], VectorDot(WarpWeftBiasBaseMultiplier, BucklingStiffness), VectorDot(WarpWeftBiasBaseMultiplier, Stiffness));
-
 		const varying FVector3f X1 = VectorGather(&X[i1]);
 		const varying FVector3f X2 = VectorGather(&X[i2]);
 		const varying FVector3f X3 = VectorGather(&X[i3]);
 		const varying FVector3f X4 = VectorGather(&X[i4]);
 
-		const varying float CombinedInvMass = M1 + M2 + M3 + M4;
-		const varying float Damping = 2.f * DampingRatio * sqrt(StiffnessValue / CombinedInvMass);
-
+		const varying FVector3f WarpWeftBiasBaseMultiplier = VectorGather(&WarpWeftBiasBaseMultiplers[i]);
+		const varying float StiffnessValue = select(IsBuckled[i], VectorDot(WarpWeftBiasBaseMultiplier, BucklingStiffness), VectorDot(WarpWeftBiasBaseMultiplier, Stiffness));
 		const varying float Lambda = Lambdas[i];
-		const varying float Alpha  = 1.f / (StiffnessValue * Dt * Dt);
-		const varying float Gamma = Alpha * Damping * Dt;
-
-		const varying FVector3f V1TimesDt = P1 - X1;
-		const varying FVector3f V2TimesDt = P2 - X2;
-		const varying FVector3f V3TimesDt = P3 - X3;
-		const varying FVector3f V4TimesDt = P4 - X4;
-		const varying float DampingTerm = Gamma * (VectorDot(V1TimesDt, Grad1) + VectorDot(V2TimesDt, Grad2) + VectorDot(V3TimesDt, Grad3) + VectorDot(V4TimesDt, Grad4));
 
-		const varying float Denom = (1.f + Gamma)*(M1 * VectorSizeSquared(Grad1) + M2 * VectorSizeSquared(Grad2) + M3 * VectorSizeSquared(Grad3) + M4 * VectorSizeSquared(Grad4)) + Alpha;
+		varying FVector3f Delta1, Delta2, Delta3, Delta4;
+		const varying float DLambda = GetXPBDBendingElementDampingDeltas(P1, P2, P3, P4, X1, X2, X3, X4, M1, M2, M3, M4, Dt, RestAngles[i], Lambda, StiffnessValue, DampingRatio, MinStiffness, Delta1, Delta2, Delta3, Delta4);
 
-		const varying float ComputedDLambda = (Angle - RestAngle - Alpha * Lambda + DampingTerm) / Denom;
-		const varying float DLambda = select(StiffnessValue < MinStiffness, (uniform float)0, ComputedDLambda);
 		Lambdas[i] = Lambda + DLambda;
 
 		if (M1 > 0)
 		{
-			VectorScatter(&PandInvM[i1], SetVector4(P1 - DLambda * M1 * Grad1, M1));
+			VectorScatter(&PandInvM[i1], SetVector4(P1 + Delta1, M1));
 		}
 		if (M2 > 0)
 		{
-			VectorScatter(&PandInvM[i2], SetVector4(P2 - DLambda * M2 * Grad2, M2));
+			VectorScatter(&PandInvM[i2], SetVector4(P2 + Delta2, M2));
 		}
 		if (M3 > 0)
 		{
-			VectorScatter(&PandInvM[i3], SetVector4(P3 - DLambda * M3 * Grad3, M3));
+			VectorScatter(&PandInvM[i3], SetVector4(P3 + Delta3, M3));
 		}
 		if (M4 > 0)
 		{
-			VectorScatter(&PandInvM[i4], SetVector4(P4 - DLambda * M4 * Grad4, M4));
+			VectorScatter(&PandInvM[i4], SetVector4(P4 + Delta4, M4));
 		}
 	}
 }
@@ -260,11 +266,7 @@ export void ApplyXPBDAnisotropicBendingConstraintsWithMaps(uniform FVector4f Pan
 		UnzipPandInvM(PandInvM2, P2, M2);
 		UnzipPandInvM(PandInvM3, P3, M3);
 		UnzipPandInvM(PandInvM4, P4, M4);
-
-
-		varying FVector3f Grad1, Grad2, Grad3, Grad4;
-		const varying float Angle = CalcGradientsAndAngle(P1, P2, P3, P4, Grad1, Grad2, Grad3, Grad4);
-		
+				
 		#pragma ignore warning(perf)
 		const varying float Stiffness = StiffnessTable[StiffnessIndex];
 		#pragma ignore warning(perf)
@@ -283,33 +285,31 @@ export void ApplyXPBDAnisotropicBendingConstraintsWithMaps(uniform FVector4f Pan
 		
 		const varying FVector3f BucklingStiffnessVector = {{BucklingStiffnessWeft, BucklingStiffness, BucklingStiffnessBias}};
 		
-		const varying float RestAngle = RestAngles[i];
 		const varying FVector3f WarpWeftBiasBaseMultiplier = VectorGather(&WarpWeftBiasBaseMultiplers[i]);
 		const varying float StiffnessValue = select(IsBuckled[i], VectorDot(WarpWeftBiasBaseMultiplier, BucklingStiffnessVector), VectorDot(WarpWeftBiasBaseMultiplier, StiffnessVector));
 		
 		const varying float Lambda = Lambdas[i];
-		const varying float Alpha  = 1.f / (StiffnessValue * Dt * Dt);
-		const varying float Denom = M1 * VectorSizeSquared(Grad1) + M2 * VectorSizeSquared(Grad2) + M3 * VectorSizeSquared(Grad3) + M4 * VectorSizeSquared(Grad4) + Alpha;
 
-		const varying float ComputedDLambda = (Angle - RestAngle - Alpha * Lambda) / Denom;
-		const varying float DLambda = select( StiffnessValue < MinStiffness, (uniform float)0, ComputedDLambda);
+		varying FVector3f Delta1, Delta2, Delta3, Delta4;
+		const varying float DLambda = GetXPBDBendingElementDeltas(P1, P2, P3, P4, M1, M2, M3, M4, Dt, RestAngles[i], Lambda, StiffnessValue, MinStiffness, Delta1, Delta2, Delta3, Delta4);
+
 		Lambdas[i] = Lambda + DLambda;
 
 		if (M1 > 0)
 		{
-			VectorScatter(&PandInvM[i1], SetVector4(P1 - DLambda * M1 * Grad1, M1));
+			VectorScatter(&PandInvM[i1], SetVector4(P1 + Delta1, M1));
 		}
 		if (M2 > 0)
 		{
-			VectorScatter(&PandInvM[i2], SetVector4(P2 - DLambda * M2 * Grad2, M2));
+			VectorScatter(&PandInvM[i2], SetVector4(P2 + Delta2, M2));
 		}
 		if (M3 > 0)
 		{
-			VectorScatter(&PandInvM[i3], SetVector4(P3 - DLambda * M3 * Grad3, M3));
+			VectorScatter(&PandInvM[i3], SetVector4(P3 + Delta3, M3));
 		}
 		if (M4 > 0)
 		{
-			VectorScatter(&PandInvM[i4], SetVector4(P4 - DLambda * M4 * Grad4, M4));
+			VectorScatter(&PandInvM[i4], SetVector4(P4 + Delta4, M4));
 		}
 	}
 }
@@ -376,11 +376,124 @@ export void ApplyXPBDAnisotropicBendingConstraintsWithDampingAndMaps(uniform FVe
 		UnzipPandInvM(PandInvM3, P3, M3);
 		UnzipPandInvM(PandInvM4, P4, M4);
 
+		#pragma ignore warning(perf)
+		const varying float Stiffness = StiffnessTable[StiffnessIndex];
+		#pragma ignore warning(perf)
+		const varying float StiffnessWeft = StiffnessWeftTable[StiffnessWeftIndex];
+		#pragma ignore warning(perf)
+		const varying float StiffnessBias = StiffnessBiasTable[StiffnessBiasIndex];
+
+		const varying FVector3f StiffnessVector = {{StiffnessWeft, Stiffness, StiffnessBias}};
 
-		varying FVector3f Grad1, Grad2, Grad3, Grad4;
-		const varying float Angle = CalcGradientsAndAngle(P1, P2, P3, P4, Grad1, Grad2, Grad3, Grad4);
+		#pragma ignore warning(perf)
+		const varying float BucklingStiffness = BucklingStiffnessTable[BucklingStiffnessIndex];
+		#pragma ignore warning(perf)
+		const varying float BucklingStiffnessWeft = BucklingStiffnessWeftTable[BucklingStiffnessWeftIndex];
+		#pragma ignore warning(perf)
+		const varying float BucklingStiffnessBias = BucklingStiffnessBiasTable[BucklingStiffnessBiasIndex];
 		
+		const varying FVector3f BucklingStiffnessVector = {{BucklingStiffnessWeft, BucklingStiffness, BucklingStiffnessBias}};
 		
+		#pragma ignore warning(perf)
+		const varying float DampingRatio = DampingTable[DampingIndex];
+		
+		const varying float RestAngle = RestAngles[i];
+		const varying FVector3f WarpWeftBiasBaseMultiplier = VectorGather(&WarpWeftBiasBaseMultiplers[i]);
+		const varying float StiffnessValue = select(IsBuckled[i], VectorDot(WarpWeftBiasBaseMultiplier, BucklingStiffnessVector), VectorDot(WarpWeftBiasBaseMultiplier, StiffnessVector));
+
+		const varying FVector3f X1 = VectorGather(&X[i1]);
+		const varying FVector3f X2 = VectorGather(&X[i2]);
+		const varying FVector3f X3 = VectorGather(&X[i3]);
+		const varying FVector3f X4 = VectorGather(&X[i4]);
+		
+		const varying float Lambda = Lambdas[i];
+
+		varying FVector3f Delta1, Delta2, Delta3, Delta4;
+		const varying float DLambda = GetXPBDBendingElementDeltasWithDamping(P1, P2, P3, P4, X1, X2, X3, X4, M1, M2, M3, M4, Dt, RestAngles[i], Lambda, StiffnessValue, DampingRatio, MinStiffness, Delta1, Delta2, Delta3, Delta4);
+
+		Lambdas[i] = Lambda + DLambda;
+
+		if (M1 > 0)
+		{
+			VectorScatter(&PandInvM[i1], SetVector4(P1 + Delta1, M1));
+		}
+		if (M2 > 0)
+		{
+			VectorScatter(&PandInvM[i2], SetVector4(P2 + Delta2, M2));
+		}
+		if (M3 > 0)
+		{
+			VectorScatter(&PandInvM[i3], SetVector4(P3 + Delta3, M3));
+		}
+		if (M4 > 0)
+		{
+			VectorScatter(&PandInvM[i4], SetVector4(P4 + Delta4, M4));
+		}
+	}
+}
+
+export void ApplyXPBDAnisotropicBendingDampingConstraintsWithMaps(uniform FVector4f PandInvM[],
+											const uniform FVector3f X[],
+											const uniform FIntVector4 Constraints[],
+											const uniform float RestAngles[],
+											const uniform bool IsBuckled[],
+											const uniform FVector3f WarpWeftBiasBaseMultiplers[],
+											uniform float Lambdas[],
+											const uniform float Dt,
+											const uniform float MinStiffness,
+											const uniform bool StiffnessHasMap,
+											const uniform uint8 StiffnessIndices[],
+											const uniform float StiffnessTable[],
+											const uniform bool StiffnessWeftHasMap,
+											const uniform uint8 StiffnessWeftIndices[],
+											const uniform float StiffnessWeftTable[],
+											const uniform bool StiffnessBiasHasMap,
+											const uniform uint8 StiffnessBiasIndices[],
+											const uniform float StiffnessBiasTable[],
+											const uniform bool BucklingStiffnessHasMap,
+											const uniform uint8 BucklingStiffnessIndices[],
+											const uniform float BucklingStiffnessTable[],
+											const uniform bool BucklingStiffnessWeftHasMap,
+											const uniform uint8 BucklingStiffnessWeftIndices[],
+											const uniform float BucklingStiffnessWeftTable[],
+											const uniform bool BucklingStiffnessBiasHasMap,
+											const uniform uint8 BucklingStiffnessBiasIndices[],
+											const uniform float BucklingStiffnessBiasTable[],
+											const uniform bool DampingHasMap,
+											const uniform uint8 DampingIndices[],
+											const uniform float DampingTable[],
+											const uniform int32 NumConstraints)
+{
+
+	foreach(i = 0 ... NumConstraints)
+	{
+		const varying FIntVector4 Constraint = VectorLoad(&Constraints[extract(i,0)]);
+
+		const varying int32 i1 = Constraint.V[0];
+		const varying int32 i2 = Constraint.V[1];
+		const varying int32 i3 = Constraint.V[2];
+		const varying int32 i4 = Constraint.V[3];
+		const varying int8 StiffnessIndex = StiffnessHasMap ? StiffnessIndices[i] : 0;
+		const varying int8 StiffnessWeftIndex = StiffnessWeftHasMap ? StiffnessWeftIndices[i] : 0;
+		const varying int8 StiffnessBiasIndex = StiffnessBiasHasMap ? StiffnessBiasIndices[i] : 0;
+		const varying int8 BucklingStiffnessIndex = BucklingStiffnessHasMap ? BucklingStiffnessIndices[i] : 0;
+		const varying int8 BucklingStiffnessWeftIndex = BucklingStiffnessWeftHasMap ? BucklingStiffnessWeftIndices[i] : 0;
+		const varying int8 BucklingStiffnessBiasIndex = BucklingStiffnessBiasHasMap ? BucklingStiffnessBiasIndices[i] : 0;
+		const varying int8 DampingIndex = DampingHasMap ? DampingIndices[i] : 0;
+
+		const varying FVector4f PandInvM2 = VectorGather(&PandInvM[i2]);
+		const varying FVector4f PandInvM3 = VectorGather(&PandInvM[i3]);
+		const varying FVector4f PandInvM1 = VectorGather(&PandInvM[i1]);
+		const varying FVector4f PandInvM4 = VectorGather(&PandInvM[i4]);
+
+		varying FVector3f P1, P2, P3, P4;
+		varying float M1, M2, M3, M4;
+
+		UnzipPandInvM(PandInvM1, P1, M1);
+		UnzipPandInvM(PandInvM2, P2, M2);
+		UnzipPandInvM(PandInvM3, P3, M3);
+		UnzipPandInvM(PandInvM4, P4, M4);
+
 		#pragma ignore warning(perf)
 		const varying float Stiffness = StiffnessTable[StiffnessIndex];
 		#pragma ignore warning(perf)
@@ -410,41 +523,29 @@ export void ApplyXPBDAnisotropicBendingConstraintsWithDampingAndMaps(uniform FVe
 		const varying FVector3f X2 = VectorGather(&X[i2]);
 		const varying FVector3f X3 = VectorGather(&X[i3]);
 		const varying FVector3f X4 = VectorGather(&X[i4]);
-
-		const varying float CombinedInvMass = M1 + M2 + M3 + M4;
-		const varying float Damping = 2.f * DampingRatio * sqrt(StiffnessValue / CombinedInvMass);
 		
 		const varying float Lambda = Lambdas[i];
-		const varying float Alpha  = 1.f / (StiffnessValue * Dt * Dt);
-		const varying float Gamma = Alpha * Damping * Dt;
 
-		const varying FVector3f V1TimesDt = P1 - X1;
-		const varying FVector3f V2TimesDt = P2 - X2;
-		const varying FVector3f V3TimesDt = P3 - X3;
-		const varying FVector3f V4TimesDt = P4 - X4;
-		const varying float DampingTerm = Gamma * (VectorDot(V1TimesDt, Grad1) + VectorDot(V2TimesDt, Grad2) + VectorDot(V3TimesDt, Grad3) + VectorDot(V4TimesDt, Grad4));
+		varying FVector3f Delta1, Delta2, Delta3, Delta4;
+		const varying float DLambda = GetXPBDBendingElementDampingDeltas(P1, P2, P3, P4, X1, X2, X3, X4, M1, M2, M3, M4, Dt, RestAngles[i], Lambda, StiffnessValue, DampingRatio, MinStiffness, Delta1, Delta2, Delta3, Delta4);
 
-		const varying float Denom = (1.f + Gamma)*(M1 * VectorSizeSquared(Grad1) + M2 * VectorSizeSquared(Grad2) + M3 * VectorSizeSquared(Grad3) + M4 * VectorSizeSquared(Grad4)) + Alpha;
-
-		const varying float ComputedDLambda = (Angle - RestAngle - Alpha * Lambda + DampingTerm) / Denom;
-		const varying float DLambda = select( StiffnessValue < MinStiffness, (uniform float)0, ComputedDLambda);
 		Lambdas[i] = Lambda + DLambda;
 
 		if (M1 > 0)
 		{
-			VectorScatter(&PandInvM[i1], SetVector4(P1 - DLambda * M1 * Grad1, M1));
+			VectorScatter(&PandInvM[i1], SetVector4(P1 + Delta1, M1));
 		}
 		if (M2 > 0)
 		{
-			VectorScatter(&PandInvM[i2], SetVector4(P2 - DLambda * M2 * Grad2, M2));
+			VectorScatter(&PandInvM[i2], SetVector4(P2 + Delta2, M2));
 		}
 		if (M3 > 0)
 		{
-			VectorScatter(&PandInvM[i3], SetVector4(P3 - DLambda * M3 * Grad3, M3));
+			VectorScatter(&PandInvM[i3], SetVector4(P3 + Delta3, M3));
 		}
 		if (M4 > 0)
 		{
-			VectorScatter(&PandInvM[i4], SetVector4(P4 - DLambda * M4 * Grad4, M4));
+			VectorScatter(&PandInvM[i4], SetVector4(P4 + Delta4, M4));
 		}
 	}
-}
\ No newline at end of file
+}
diff --git a/Engine/Source/Runtime/Experimental/Chaos/Private/Chaos/XPBDAnisotropicSpringConstraints.cpp b/Engine/Source/Runtime/Experimental/Chaos/Private/Chaos/XPBDAnisotropicSpringConstraints.cpp
--- a/Engine/Source/Runtime/Experimental/Chaos/Private/Chaos/XPBDAnisotropicSpringConstraints.cpp
+++ b/Engine/Source/Runtime/Experimental/Chaos/Private/Chaos/XPBDAnisotropicSpringConstraints.cpp
@@ -6,6 +6,7 @@
 #include "Chaos/SoftsSpring.h"
 #include "Chaos/TriangleMesh.h"
 #include "ChaosStats.h"
+#include "XPBDInternal.h"
 
 #if INTEL_ISPC
 #include "XPBDAnisotropicSpringConstraints.ispc.generated.h"
@@ -320,7 +321,7 @@ void FXPBDAnisotropicEdgeSpringConstraints::UpdateDists()
 	Private::UpdateDists(BaseDists, WarpWeftScaleBaseMultipliers, WarpScale, WeftScale, Dists);
 }
 
-template<typename SolverParticlesOrRange>
+template<bool bDampingBefore, bool bSingleLambda, bool bSeparateStretch, bool bDampingAfter, typename SolverParticlesOrRange>
 void FXPBDAnisotropicEdgeSpringConstraints::ApplyHelper(SolverParticlesOrRange& Particles, const FSolverReal Dt, const int32 ConstraintIndex, const FSolverVec3& ExpStiffnessValues, const FSolverReal DampingRatioValue) const
 {
 	const TVec2<int32>& Constraint = Constraints[ConstraintIndex];
@@ -328,18 +329,40 @@ void FXPBDAnisotropicEdgeSpringConstraints::ApplyHelper(SolverParticlesOrRange&
 	const int32 Index2 = Constraint[1];
 
 	const FSolverVec3& WarpWeftBiasBaseMultiplier = WarpWeftBiasBaseMultipliers[ConstraintIndex];
-	const FSolverReal FinalStiffnessValue = WarpWeftBiasBaseMultiplier.Dot(ExpStiffnessValues);
+	const FSolverReal ExpStiffnessValue = WarpWeftBiasBaseMultiplier.Dot(ExpStiffnessValues);
 
-	const FSolverVec3 Delta = Spring::GetXPBDSpringDelta(Particles, Dt, Constraint, Dists[ConstraintIndex], Lambdas[ConstraintIndex],
-		FinalStiffnessValue, MinStiffness, DampingRatioValue);
+	FSolverVec3 Delta(0.f);
+	if constexpr (bDampingBefore)
+	{
+		Delta += Spring::GetXPBDSpringDampingDelta(Particles, Dt, Constraint, Dists[ConstraintIndex], LambdasDamping[ConstraintIndex],
+			ExpStiffnessValue, MinStiffness, DampingRatioValue);
+	}
+
+	if constexpr (bSingleLambda)
+	{
+		Delta += Spring::GetXPBDSpringDeltaWithDamping(Particles, Dt, Constraint, Dists[ConstraintIndex], Lambdas[ConstraintIndex],
+			ExpStiffnessValue, MinStiffness, DampingRatioValue);
+	}
+
+	if constexpr (bSeparateStretch)
+	{
+		Delta += Spring::GetXPBDSpringDelta(Particles, Dt, Constraint, Dists[ConstraintIndex], Lambdas[ConstraintIndex],
+			ExpStiffnessValue, MinStiffness);
+	}
+
+	if constexpr (bDampingAfter)
+	{
+		Delta += Spring::GetXPBDSpringDampingDelta(Particles, Dt, Constraint, Dists[ConstraintIndex], LambdasDamping[ConstraintIndex],
+			ExpStiffnessValue, MinStiffness, DampingRatioValue);
+	}
 
 	if (Particles.InvM(Index1) > (FSolverReal)0.)
 	{
-		Particles.P(Index1) -= Particles.InvM(Index1) * Delta;
+		Particles.P(Index1) += Particles.InvM(Index1) * Delta;
 	}
 	if (Particles.InvM(Index2) > (FSolverReal)0.)
 	{
-		Particles.P(Index2) += Particles.InvM(Index2) * Delta;
+		Particles.P(Index2) -= Particles.InvM(Index2) * Delta;
 	}
 }
 
@@ -357,152 +380,338 @@ void FXPBDAnisotropicEdgeSpringConstraints::Apply(SolverParticlesOrRange& Partic
 	if ((ConstraintsPerColorStartIndex.Num() > 1) && (Constraints.Num() > Chaos_XPBDSpring_ParallelConstraintCount))
 	{
 		const int32 ConstraintColorNum = ConstraintsPerColorStartIndex.Num() - 1;
-		if (!StiffnessHasWeightMap && !StiffnessWeftHasWeightMap && !StiffnessBiasHasWeightMap && !DampingHasWeightMap)
+#if INTEL_ISPC
+		if (bRealTypeCompatibleWithISPC && bChaos_XPBDSpring_ISPC_Enabled)
 		{
-			const FSolverVec3 ExpStiffnessValue((FSolverReal)StiffnessWeft, (FSolverReal)Stiffness, (FSolverReal)StiffnessBias);
-			const FSolverReal DampingRatioValue = (FSolverReal)DampingRatio;
-			if (ExpStiffnessValue.Max() < MinStiffness)
-			{
-				return;
-			}
+			const bool bSingleLambda = Chaos_XPBDSpring_SplitDampingMode == (int32)EXPBDSplitDampingMode::SingleLambda;
 
-#if INTEL_ISPC
-			if (bRealTypeCompatibleWithISPC && bChaos_XPBDSpring_ISPC_Enabled)
+			if (!StiffnessHasWeightMap && !StiffnessWeftHasWeightMap && !StiffnessBiasHasWeightMap && !DampingHasWeightMap)
 			{
+				const FSolverVec3 ExpStiffnessValue((FSolverReal)StiffnessWeft, (FSolverReal)Stiffness, (FSolverReal)StiffnessBias);
+				const FSolverReal DampingRatioValue = (FSolverReal)DampingRatio;
+				if (ExpStiffnessValue.Max() < MinStiffness)
+				{
+					return;
+				}
+
 				if (DampingRatioValue > 0)
 				{
-					for (int32 ConstraintColorIndex = 0; ConstraintColorIndex < ConstraintColorNum; ++ConstraintColorIndex)
+					if (bSingleLambda)
 					{
-						const int32 ColorStart = ConstraintsPerColorStartIndex[ConstraintColorIndex];
-						const int32 ColorSize = ConstraintsPerColorStartIndex[ConstraintColorIndex + 1] - ColorStart;
-						ispc::ApplyXPBDAnisoSpringConstraintsWithDamping(
-							(ispc::FVector4f*)Particles.GetPAndInvM().GetData(),
-							(const ispc::FVector3f*)Particles.XArray().GetData(),
-							(ispc::FIntVector2*)&Constraints.GetData()[ColorStart],
-							&Dists.GetData()[ColorStart],
-							(const ispc::FVector3f*)&WarpWeftBiasBaseMultipliers.GetData()[ColorStart],
-							&Lambdas.GetData()[ColorStart],
-							Dt,
-							reinterpret_cast<const ispc::FVector3f&>(ExpStiffnessValue),
-							DampingRatioValue,
-							ColorSize);
+						for (int32 ConstraintColorIndex = 0; ConstraintColorIndex < ConstraintColorNum; ++ConstraintColorIndex)
+						{
+							const int32 ColorStart = ConstraintsPerColorStartIndex[ConstraintColorIndex];
+							const int32 ColorSize = ConstraintsPerColorStartIndex[ConstraintColorIndex + 1] - ColorStart;
+							ispc::ApplyXPBDAnisoSpringConstraintsWithDamping(
+								(ispc::FVector4f*)Particles.GetPAndInvM().GetData(),
+								(const ispc::FVector3f*)Particles.XArray().GetData(),
+								(ispc::FIntVector2*)&Constraints.GetData()[ColorStart],
+								&Dists.GetData()[ColorStart],
+								(const ispc::FVector3f*)&WarpWeftBiasBaseMultipliers.GetData()[ColorStart],
+								&Lambdas.GetData()[ColorStart],
+								Dt,
+								reinterpret_cast<const ispc::FVector3f&>(ExpStiffnessValue),
+								DampingRatioValue,
+								ColorSize);
+						}
+						return;
 					}
-				}
-				else
-				{
-					for (int32 ConstraintColorIndex = 0; ConstraintColorIndex < ConstraintColorNum; ++ConstraintColorIndex)
+					else
 					{
-						const int32 ColorStart = ConstraintsPerColorStartIndex[ConstraintColorIndex];
-						const int32 ColorSize = ConstraintsPerColorStartIndex[ConstraintColorIndex + 1] - ColorStart;
-						ispc::ApplyXPBDAnisoSpringConstraints(
-							(ispc::FVector4f*)Particles.GetPAndInvM().GetData(),
-							(ispc::FIntVector2*)&Constraints.GetData()[ColorStart],
-							&Dists.GetData()[ColorStart],
-							(const ispc::FVector3f*)&WarpWeftBiasBaseMultipliers.GetData()[ColorStart],
-							&Lambdas.GetData()[ColorStart],
-							Dt,
-							reinterpret_cast<const ispc::FVector3f&>(ExpStiffnessValue),
-							ColorSize);
+						for (int32 ConstraintColorIndex = 0; ConstraintColorIndex < ConstraintColorNum; ++ConstraintColorIndex)
+						{
+							const int32 ColorStart = ConstraintsPerColorStartIndex[ConstraintColorIndex];
+							const int32 ColorSize = ConstraintsPerColorStartIndex[ConstraintColorIndex + 1] - ColorStart;
+							ispc::ApplyXPBDAnisoSpringDampingConstraints(
+								(ispc::FVector4f*)Particles.GetPAndInvM().GetData(),
+								(const ispc::FVector3f*)Particles.XArray().GetData(),
+								(ispc::FIntVector2*)&Constraints.GetData()[ColorStart],
+								&Dists.GetData()[ColorStart],
+								(const ispc::FVector3f*)&WarpWeftBiasBaseMultipliers.GetData()[ColorStart],
+								&LambdasDamping.GetData()[ColorStart],
+								Dt,
+								reinterpret_cast<const ispc::FVector3f&>(ExpStiffnessValue),
+								DampingRatioValue,
+								ColorSize);
+						}
 					}
 				}
+				for (int32 ConstraintColorIndex = 0; ConstraintColorIndex < ConstraintColorNum; ++ConstraintColorIndex)
+				{
+					const int32 ColorStart = ConstraintsPerColorStartIndex[ConstraintColorIndex];
+					const int32 ColorSize = ConstraintsPerColorStartIndex[ConstraintColorIndex + 1] - ColorStart;
+					ispc::ApplyXPBDAnisoSpringConstraints(
+						(ispc::FVector4f*)Particles.GetPAndInvM().GetData(),
+						(ispc::FIntVector2*)&Constraints.GetData()[ColorStart],
+						&Dists.GetData()[ColorStart],
+						(const ispc::FVector3f*)&WarpWeftBiasBaseMultipliers.GetData()[ColorStart],
+						&Lambdas.GetData()[ColorStart],
+						Dt,
+						reinterpret_cast<const ispc::FVector3f&>(ExpStiffnessValue),
+						ColorSize);
+				}
 			}
-			else
-#endif
+			else // ISPC with maps
 			{
+				if (DampingHasWeightMap || (FSolverReal)DampingRatio > 0)
+				{
+					if (bSingleLambda)
+					{
+						for (int32 ConstraintColorIndex = 0; ConstraintColorIndex < ConstraintColorNum; ++ConstraintColorIndex)
+						{
+							const int32 ColorStart = ConstraintsPerColorStartIndex[ConstraintColorIndex];
+							const int32 ColorSize = ConstraintsPerColorStartIndex[ConstraintColorIndex + 1] - ColorStart;
+							ispc::ApplyXPBDAnisoSpringConstraintsWithDampingAndWeightMaps(
+								(ispc::FVector4f*)Particles.GetPAndInvM().GetData(),
+								(const ispc::FVector3f*)Particles.XArray().GetData(),
+								(ispc::FIntVector2*)&Constraints.GetData()[ColorStart],
+								&Dists.GetData()[ColorStart],
+								(const ispc::FVector3f*)&WarpWeftBiasBaseMultipliers.GetData()[ColorStart],
+								&Lambdas.GetData()[ColorStart],
+								Dt,
+								MinStiffness,
+								StiffnessHasWeightMap,
+								StiffnessHasWeightMap ? &Stiffness.GetIndices().GetData()[ColorStart] : nullptr,
+								&Stiffness.GetTable().GetData()[0],
+								StiffnessWeftHasWeightMap,
+								StiffnessWeftHasWeightMap ? &StiffnessWeft.GetIndices().GetData()[ColorStart] : nullptr,
+								&StiffnessWeft.GetTable().GetData()[0],
+								StiffnessBiasHasWeightMap,
+								StiffnessBiasHasWeightMap ? &StiffnessBias.GetIndices().GetData()[ColorStart] : nullptr,
+								&StiffnessBias.GetTable().GetData()[0],
+								DampingHasWeightMap,
+								DampingHasWeightMap ? &DampingRatio.GetIndices().GetData()[ColorStart] : nullptr,
+								&DampingRatio.GetTable().GetData()[0],
+								ColorSize);
+						}
+						return;
+					}
+					else
+					{
+						for (int32 ConstraintColorIndex = 0; ConstraintColorIndex < ConstraintColorNum; ++ConstraintColorIndex)
+						{
+							const int32 ColorStart = ConstraintsPerColorStartIndex[ConstraintColorIndex];
+							const int32 ColorSize = ConstraintsPerColorStartIndex[ConstraintColorIndex + 1] - ColorStart;
+							ispc::ApplyXPBDAnisoSpringDampingConstraintsWithWeightMaps(
+								(ispc::FVector4f*)Particles.GetPAndInvM().GetData(),
+								(const ispc::FVector3f*)Particles.XArray().GetData(),
+								(ispc::FIntVector2*)&Constraints.GetData()[ColorStart],
+								&Dists.GetData()[ColorStart],
+								(const ispc::FVector3f*)&WarpWeftBiasBaseMultipliers.GetData()[ColorStart],
+								&LambdasDamping.GetData()[ColorStart],
+								Dt,
+								MinStiffness,
+								StiffnessHasWeightMap,
+								StiffnessHasWeightMap ? &Stiffness.GetIndices().GetData()[ColorStart] : nullptr,
+								&Stiffness.GetTable().GetData()[0],
+								StiffnessWeftHasWeightMap,
+								StiffnessWeftHasWeightMap ? &StiffnessWeft.GetIndices().GetData()[ColorStart] : nullptr,
+								&StiffnessWeft.GetTable().GetData()[0],
+								StiffnessBiasHasWeightMap,
+								StiffnessBiasHasWeightMap ? &StiffnessBias.GetIndices().GetData()[ColorStart] : nullptr,
+								&StiffnessBias.GetTable().GetData()[0],
+								DampingHasWeightMap,
+								DampingHasWeightMap ? &DampingRatio.GetIndices().GetData()[ColorStart] : nullptr,
+								&DampingRatio.GetTable().GetData()[0],
+								ColorSize);
+						}
+					}
+				}
 				for (int32 ConstraintColorIndex = 0; ConstraintColorIndex < ConstraintColorNum; ++ConstraintColorIndex)
 				{
 					const int32 ColorStart = ConstraintsPerColorStartIndex[ConstraintColorIndex];
 					const int32 ColorSize = ConstraintsPerColorStartIndex[ConstraintColorIndex + 1] - ColorStart;
-					PhysicsParallelFor(ColorSize, [this, &Particles, Dt, ColorStart, &ExpStiffnessValue, DampingRatioValue](const int32 Index)
-					{
-						const int32 ConstraintIndex = ColorStart + Index;
-						ApplyHelper(Particles, Dt, ConstraintIndex, ExpStiffnessValue, DampingRatioValue);
-					});
+					ispc::ApplyXPBDAnisoSpringConstraintsWithWeightMaps(
+						(ispc::FVector4f*)Particles.GetPAndInvM().GetData(),
+						(ispc::FIntVector2*)&Constraints.GetData()[ColorStart],
+						&Dists.GetData()[ColorStart],
+						(const ispc::FVector3f*)&WarpWeftBiasBaseMultipliers.GetData()[ColorStart],
+						&Lambdas.GetData()[ColorStart],
+						Dt,
+						MinStiffness,
+						StiffnessHasWeightMap,
+						&Stiffness.GetIndices().GetData()[ColorStart],
+						&Stiffness.GetTable().GetData()[0],
+						StiffnessWeftHasWeightMap,
+						StiffnessWeftHasWeightMap ? &StiffnessWeft.GetIndices().GetData()[ColorStart] : nullptr,
+						&StiffnessWeft.GetTable().GetData()[0],
+						StiffnessBiasHasWeightMap,
+						StiffnessBiasHasWeightMap ? &StiffnessBias.GetIndices().GetData()[ColorStart] : nullptr,
+						&StiffnessBias.GetTable().GetData()[0],
+						ColorSize);
 				}
 			}
 		}
-		else  // Has weight maps
+		else
+#endif
 		{
-#if INTEL_ISPC
-			if (bRealTypeCompatibleWithISPC && bChaos_XPBDSpring_ISPC_Enabled)
+			// Parallel non-ispc
+			const FSolverReal StiffnessNoMap = (FSolverReal)Stiffness;
+			const FSolverReal StiffnessWeftNoMap = (FSolverReal)StiffnessWeft;
+			const FSolverReal StiffnessBiasNoMap = (FSolverReal)StiffnessBias;
+			const FSolverReal DampingNoMap = (FSolverReal)DampingRatio;
+			if (DampingNoMap > 0 || DampingHasWeightMap)
 			{
-				if (DampingHasWeightMap || (FSolverReal)DampingRatio > 0)
+				if (Chaos_XPBDSpring_SplitDampingMode == (int32)EXPBDSplitDampingMode::TwoPassBefore)
 				{
 					for (int32 ConstraintColorIndex = 0; ConstraintColorIndex < ConstraintColorNum; ++ConstraintColorIndex)
 					{
 						const int32 ColorStart = ConstraintsPerColorStartIndex[ConstraintColorIndex];
 						const int32 ColorSize = ConstraintsPerColorStartIndex[ConstraintColorIndex + 1] - ColorStart;
-						ispc::ApplyXPBDAnisoSpringConstraintsWithDampingAndWeightMaps(
-							(ispc::FVector4f*)Particles.GetPAndInvM().GetData(),
-							(const ispc::FVector3f*)Particles.XArray().GetData(),
-							(ispc::FIntVector2*)&Constraints.GetData()[ColorStart],
-							&Dists.GetData()[ColorStart],
-							(const ispc::FVector3f*)&WarpWeftBiasBaseMultipliers.GetData()[ColorStart],
-							&Lambdas.GetData()[ColorStart],
-							Dt,
-							MinStiffness,
-							StiffnessHasWeightMap,
-							StiffnessHasWeightMap ? &Stiffness.GetIndices().GetData()[ColorStart] : nullptr,
-							&Stiffness.GetTable().GetData()[0],
-							StiffnessWeftHasWeightMap,
-							StiffnessWeftHasWeightMap ? &StiffnessWeft.GetIndices().GetData()[ColorStart] : nullptr,
-							&StiffnessWeft.GetTable().GetData()[0],
-							StiffnessBiasHasWeightMap,
-							StiffnessBiasHasWeightMap ? &StiffnessBias.GetIndices().GetData()[ColorStart] : nullptr,
-							&StiffnessBias.GetTable().GetData()[0],
-							DampingHasWeightMap,
-							DampingHasWeightMap ? &DampingRatio.GetIndices().GetData()[ColorStart] : nullptr,
-							&DampingRatio.GetTable().GetData()[0],
-							ColorSize);
+						PhysicsParallelFor(ColorSize, [this, &Particles, Dt, ColorStart, StiffnessHasWeightMap, StiffnessNoMap, StiffnessWeftHasWeightMap, StiffnessWeftNoMap, StiffnessBiasHasWeightMap, StiffnessBiasNoMap, DampingHasWeightMap, DampingNoMap](const int32 Index)
+						{
+							const int32 ConstraintIndex = ColorStart + Index;
+							constexpr bool bDampingBefore = true;
+							constexpr bool bSingleLambda = false;
+							constexpr bool bSeparateStretch = false;
+							constexpr bool bDampingAfter = false;
+							const FSolverReal ExpStiffnessValue = StiffnessHasWeightMap ? Stiffness[ConstraintIndex] : StiffnessNoMap;
+							const FSolverReal ExpStiffnessWeftValue = StiffnessWeftHasWeightMap ? StiffnessWeft[ConstraintIndex] : StiffnessWeftNoMap;
+							const FSolverReal ExpStiffnessBiasValue = StiffnessBiasHasWeightMap ? StiffnessBias[ConstraintIndex] : StiffnessBiasNoMap;
+							const FSolverReal DampingRatioValue = DampingHasWeightMap ? DampingRatio[ConstraintIndex] : DampingNoMap;
+							ApplyHelper<bDampingBefore, bSingleLambda, bSeparateStretch, bDampingAfter>(Particles, Dt, ConstraintIndex, FSolverVec3(ExpStiffnessWeftValue, ExpStiffnessValue, ExpStiffnessBiasValue), DampingRatioValue);
+						});
 					}
 				}
-				else
+				
+				// Stretch part (possibly with damping depending on split mode)
+				switch (Chaos_XPBDSpring_SplitDampingMode)
+				{
+				case (int32)EXPBDSplitDampingMode::TwoPassBefore: // fallthrough
+				case (int32)EXPBDSplitDampingMode::TwoPassAfter: // fallthrough
+				default:
 				{
+					// Do a pass with stretch only
 					for (int32 ConstraintColorIndex = 0; ConstraintColorIndex < ConstraintColorNum; ++ConstraintColorIndex)
 					{
 						const int32 ColorStart = ConstraintsPerColorStartIndex[ConstraintColorIndex];
 						const int32 ColorSize = ConstraintsPerColorStartIndex[ConstraintColorIndex + 1] - ColorStart;
-						ispc::ApplyXPBDAnisoSpringConstraintsWithWeightMaps(
-							(ispc::FVector4f*)Particles.GetPAndInvM().GetData(),
-							(ispc::FIntVector2*)&Constraints.GetData()[ColorStart],
-							&Dists.GetData()[ColorStart],
-							(const ispc::FVector3f*)&WarpWeftBiasBaseMultipliers.GetData()[ColorStart],
-							&Lambdas.GetData()[ColorStart],
-							Dt,
-							MinStiffness,
-							StiffnessHasWeightMap,
-							&Stiffness.GetIndices().GetData()[ColorStart],
-							&Stiffness.GetTable().GetData()[0],
-							StiffnessWeftHasWeightMap,
-							StiffnessWeftHasWeightMap ? &StiffnessWeft.GetIndices().GetData()[ColorStart] : nullptr,
-							&StiffnessWeft.GetTable().GetData()[0],
-							StiffnessBiasHasWeightMap,
-							StiffnessBiasHasWeightMap ? &StiffnessBias.GetIndices().GetData()[ColorStart] : nullptr,
-							&StiffnessBias.GetTable().GetData()[0],
-							ColorSize);
+						PhysicsParallelFor(ColorSize, [this, &Particles, Dt, ColorStart, StiffnessHasWeightMap, StiffnessNoMap, StiffnessWeftHasWeightMap, StiffnessWeftNoMap, StiffnessBiasHasWeightMap, StiffnessBiasNoMap, DampingHasWeightMap, DampingNoMap](const int32 Index)
+						{
+							const int32 ConstraintIndex = ColorStart + Index;
+							constexpr bool bDampingBefore = false;
+							constexpr bool bSingleLambda = false;
+							constexpr bool bSeparateStretch = true;
+							constexpr bool bDampingAfter = false;
+							const FSolverReal ExpStiffnessValue = StiffnessHasWeightMap ? Stiffness[ConstraintIndex] : StiffnessNoMap;
+							const FSolverReal ExpStiffnessWeftValue = StiffnessWeftHasWeightMap ? StiffnessWeft[ConstraintIndex] : StiffnessWeftNoMap;
+							const FSolverReal ExpStiffnessBiasValue = StiffnessBiasHasWeightMap ? StiffnessBias[ConstraintIndex] : StiffnessBiasNoMap;
+							const FSolverReal DampingRatioValue = DampingHasWeightMap ? DampingRatio[ConstraintIndex] : DampingNoMap;
+							ApplyHelper<bDampingBefore, bSingleLambda, bSeparateStretch, bDampingAfter>(Particles, Dt, ConstraintIndex, FSolverVec3(ExpStiffnessWeftValue, ExpStiffnessValue, ExpStiffnessBiasValue), DampingRatioValue);
+						});
+					}
+				}
+				break;
+				case (int32)EXPBDSplitDampingMode::SingleLambda:
+				{
+					// Do a pass with combined stretch and damping with a single lambda
+					for (int32 ConstraintColorIndex = 0; ConstraintColorIndex < ConstraintColorNum; ++ConstraintColorIndex)
+					{
+						const int32 ColorStart = ConstraintsPerColorStartIndex[ConstraintColorIndex];
+						const int32 ColorSize = ConstraintsPerColorStartIndex[ConstraintColorIndex + 1] - ColorStart;
+						PhysicsParallelFor(ColorSize, [this, &Particles, Dt, ColorStart, StiffnessHasWeightMap, StiffnessNoMap, StiffnessWeftHasWeightMap, StiffnessWeftNoMap, StiffnessBiasHasWeightMap, StiffnessBiasNoMap, DampingHasWeightMap, DampingNoMap](const int32 Index)
+						{
+							const int32 ConstraintIndex = ColorStart + Index;
+							constexpr bool bDampingBefore = false;
+							constexpr bool bSingleLambda = true;
+							constexpr bool bSeparateStretch = false;
+							constexpr bool bDampingAfter = false;
+							const FSolverReal ExpStiffnessValue = StiffnessHasWeightMap ? Stiffness[ConstraintIndex] : StiffnessNoMap;
+							const FSolverReal ExpStiffnessWeftValue = StiffnessWeftHasWeightMap ? StiffnessWeft[ConstraintIndex] : StiffnessWeftNoMap;
+							const FSolverReal ExpStiffnessBiasValue = StiffnessBiasHasWeightMap ? StiffnessBias[ConstraintIndex] : StiffnessBiasNoMap;
+							const FSolverReal DampingRatioValue = DampingHasWeightMap ? DampingRatio[ConstraintIndex] : DampingNoMap;
+							ApplyHelper<bDampingBefore, bSingleLambda, bSeparateStretch, bDampingAfter>(Particles, Dt, ConstraintIndex, FSolverVec3(ExpStiffnessWeftValue, ExpStiffnessValue, ExpStiffnessBiasValue), DampingRatioValue);
+						});
+					}
+				}
+				break;
+				case (int32)EXPBDSplitDampingMode::InterleavedBefore:
+				{
+					// Do a pass with damping before stretch
+					for (int32 ConstraintColorIndex = 0; ConstraintColorIndex < ConstraintColorNum; ++ConstraintColorIndex)
+					{
+						const int32 ColorStart = ConstraintsPerColorStartIndex[ConstraintColorIndex];
+						const int32 ColorSize = ConstraintsPerColorStartIndex[ConstraintColorIndex + 1] - ColorStart;
+						PhysicsParallelFor(ColorSize, [this, &Particles, Dt, ColorStart, StiffnessHasWeightMap, StiffnessNoMap, StiffnessWeftHasWeightMap, StiffnessWeftNoMap, StiffnessBiasHasWeightMap, StiffnessBiasNoMap, DampingHasWeightMap, DampingNoMap](const int32 Index)
+						{
+							const int32 ConstraintIndex = ColorStart + Index;
+							constexpr bool bDampingBefore = true;
+							constexpr bool bSingleLambda = false;
+							constexpr bool bSeparateStretch = true;
+							constexpr bool bDampingAfter = false;
+							const FSolverReal ExpStiffnessValue = StiffnessHasWeightMap ? Stiffness[ConstraintIndex] : StiffnessNoMap;
+							const FSolverReal ExpStiffnessWeftValue = StiffnessWeftHasWeightMap ? StiffnessWeft[ConstraintIndex] : StiffnessWeftNoMap;
+							const FSolverReal ExpStiffnessBiasValue = StiffnessBiasHasWeightMap ? StiffnessBias[ConstraintIndex] : StiffnessBiasNoMap;
+							const FSolverReal DampingRatioValue = DampingHasWeightMap ? DampingRatio[ConstraintIndex] : DampingNoMap;
+							ApplyHelper<bDampingBefore, bSingleLambda, bSeparateStretch, bDampingAfter>(Particles, Dt, ConstraintIndex, FSolverVec3(ExpStiffnessWeftValue, ExpStiffnessValue, ExpStiffnessBiasValue), DampingRatioValue);
+						});
+					}
+				}
+				break;
+				case (int32)EXPBDSplitDampingMode::InterleavedAfter:
+				{
+					// Do a pass with damping after stretch
+					for (int32 ConstraintColorIndex = 0; ConstraintColorIndex < ConstraintColorNum; ++ConstraintColorIndex)
+					{
+						const int32 ColorStart = ConstraintsPerColorStartIndex[ConstraintColorIndex];
+						const int32 ColorSize = ConstraintsPerColorStartIndex[ConstraintColorIndex + 1] - ColorStart;
+						PhysicsParallelFor(ColorSize, [this, &Particles, Dt, ColorStart, StiffnessHasWeightMap, StiffnessNoMap, StiffnessWeftHasWeightMap, StiffnessWeftNoMap, StiffnessBiasHasWeightMap, StiffnessBiasNoMap, DampingHasWeightMap, DampingNoMap](const int32 Index)
+						{
+							const int32 ConstraintIndex = ColorStart + Index;
+							constexpr bool bDampingBefore = false;
+							constexpr bool bSingleLambda = false;
+							constexpr bool bSeparateStretch = true;
+							constexpr bool bDampingAfter = true;
+							const FSolverReal ExpStiffnessValue = StiffnessHasWeightMap ? Stiffness[ConstraintIndex] : StiffnessNoMap;
+							const FSolverReal ExpStiffnessWeftValue = StiffnessWeftHasWeightMap ? StiffnessWeft[ConstraintIndex] : StiffnessWeftNoMap;
+							const FSolverReal ExpStiffnessBiasValue = StiffnessBiasHasWeightMap ? StiffnessBias[ConstraintIndex] : StiffnessBiasNoMap;
+							const FSolverReal DampingRatioValue = DampingHasWeightMap ? DampingRatio[ConstraintIndex] : DampingNoMap;
+							ApplyHelper<bDampingBefore, bSingleLambda, bSeparateStretch, bDampingAfter>(Particles, Dt, ConstraintIndex, FSolverVec3(ExpStiffnessWeftValue, ExpStiffnessValue, ExpStiffnessBiasValue), DampingRatioValue);
+						});
+					}
+				}
+				break;
+				}
+				if (Chaos_XPBDSpring_SplitDampingMode == (int32)EXPBDSplitDampingMode::TwoPassAfter)
+				{
+					for (int32 ConstraintColorIndex = 0; ConstraintColorIndex < ConstraintColorNum; ++ConstraintColorIndex)
+					{
+						const int32 ColorStart = ConstraintsPerColorStartIndex[ConstraintColorIndex];
+						const int32 ColorSize = ConstraintsPerColorStartIndex[ConstraintColorIndex + 1] - ColorStart;
+						PhysicsParallelFor(ColorSize, [this, &Particles, Dt, ColorStart, StiffnessHasWeightMap, StiffnessNoMap, StiffnessWeftHasWeightMap, StiffnessWeftNoMap, StiffnessBiasHasWeightMap, StiffnessBiasNoMap, DampingHasWeightMap, DampingNoMap](const int32 Index)
+						{
+							const int32 ConstraintIndex = ColorStart + Index;
+							constexpr bool bDampingBefore = false;
+							constexpr bool bSingleLambda = false;
+							constexpr bool bSeparateStretch = false;
+							constexpr bool bDampingAfter = true;
+							const FSolverReal ExpStiffnessValue = StiffnessHasWeightMap ? Stiffness[ConstraintIndex] : StiffnessNoMap;
+							const FSolverReal ExpStiffnessWeftValue = StiffnessWeftHasWeightMap ? StiffnessWeft[ConstraintIndex] : StiffnessWeftNoMap;
+							const FSolverReal ExpStiffnessBiasValue = StiffnessBiasHasWeightMap ? StiffnessBias[ConstraintIndex] : StiffnessBiasNoMap;
+							const FSolverReal DampingRatioValue = DampingHasWeightMap ? DampingRatio[ConstraintIndex] : DampingNoMap;
+							ApplyHelper<bDampingBefore, bSingleLambda, bSeparateStretch, bDampingAfter>(Particles, Dt, ConstraintIndex, FSolverVec3(ExpStiffnessWeftValue, ExpStiffnessValue, ExpStiffnessBiasValue), DampingRatioValue);
+						});
 					}
 				}
 			}
 			else
-#endif
 			{
-				const FSolverReal StiffnessNoMap = (FSolverReal)Stiffness;
-				const FSolverReal StiffnessWeftNoMap = (FSolverReal)StiffnessWeft;
-				const FSolverReal StiffnessBiasNoMap = (FSolverReal)StiffnessBias;
-				const FSolverReal DampingNoMap = (FSolverReal)DampingRatio;
+				// No damping. 
 				for (int32 ConstraintColorIndex = 0; ConstraintColorIndex < ConstraintColorNum; ++ConstraintColorIndex)
 				{
 					const int32 ColorStart = ConstraintsPerColorStartIndex[ConstraintColorIndex];
 					const int32 ColorSize = ConstraintsPerColorStartIndex[ConstraintColorIndex + 1] - ColorStart;
-					PhysicsParallelFor(ColorSize, [this, &Particles, Dt, ColorStart, StiffnessHasWeightMap, StiffnessNoMap, StiffnessWeftHasWeightMap, StiffnessWeftNoMap, StiffnessBiasHasWeightMap, 
-						StiffnessBiasNoMap, DampingHasWeightMap, DampingNoMap](const int32 Index)
+					PhysicsParallelFor(ColorSize, [this, &Particles, Dt, ColorStart, StiffnessHasWeightMap, StiffnessNoMap, StiffnessWeftHasWeightMap, StiffnessWeftNoMap, StiffnessBiasHasWeightMap, StiffnessBiasNoMap](const int32 Index)
 					{
 						const int32 ConstraintIndex = ColorStart + Index;
+						constexpr bool bDampingBefore = false;
+						constexpr bool bSingleLambda = false;
+						constexpr bool bSeparateStretch = true;
+						constexpr bool bDampingAfter = false;
 						const FSolverReal ExpStiffnessValue = StiffnessHasWeightMap ? Stiffness[ConstraintIndex] : StiffnessNoMap;
 						const FSolverReal ExpStiffnessWeftValue = StiffnessWeftHasWeightMap ? StiffnessWeft[ConstraintIndex] : StiffnessWeftNoMap;
 						const FSolverReal ExpStiffnessBiasValue = StiffnessBiasHasWeightMap ? StiffnessBias[ConstraintIndex] : StiffnessBiasNoMap;
-						const FSolverReal DampingRatioValue = DampingHasWeightMap ? DampingRatio[ConstraintIndex] : DampingNoMap;
-						ApplyHelper(Particles, Dt, ConstraintIndex, FSolverVec3(ExpStiffnessWeftValue, ExpStiffnessValue, ExpStiffnessBiasValue), DampingRatioValue);
+						ApplyHelper<bDampingBefore, bSingleLambda, bSeparateStretch, bDampingAfter>(Particles, Dt, ConstraintIndex, FSolverVec3(ExpStiffnessWeftValue, ExpStiffnessValue, ExpStiffnessBiasValue), (FSolverReal)0.);
 					});
 				}
 			}
@@ -510,32 +719,133 @@ void FXPBDAnisotropicEdgeSpringConstraints::Apply(SolverParticlesOrRange& Partic
 	}
 	else
 	{
-		if (!StiffnessHasWeightMap && !StiffnessWeftHasWeightMap && !StiffnessBiasHasWeightMap && !DampingHasWeightMap)
+		// Single threaded
+		const FSolverReal StiffnessNoMap = (FSolverReal)Stiffness;
+		const FSolverReal StiffnessWeftNoMap = (FSolverReal)StiffnessWeft;
+		const FSolverReal StiffnessBiasNoMap = (FSolverReal)StiffnessBias;
+		const FSolverReal DampingNoMap = (FSolverReal)DampingRatio;
+		if (DampingNoMap > 0 || DampingHasWeightMap)
 		{
-			const FSolverVec3 ExpStiffnessValue((FSolverReal)StiffnessWeft, (FSolverReal)Stiffness, (FSolverReal)StiffnessBias);
-			const FSolverReal DampingRatioValue = (FSolverReal)DampingRatio;
-			if (ExpStiffnessValue.Max() < MinStiffness)
+			if (Chaos_XPBDSpring_SplitDampingMode == (int32)EXPBDSplitDampingMode::TwoPassBefore)
 			{
-				return;
+				for (int32 ConstraintIndex = 0; ConstraintIndex < Constraints.Num(); ++ConstraintIndex)
+				{
+					constexpr bool bDampingBefore = true;
+					constexpr bool bSingleLambda = false;
+					constexpr bool bSeparateStretch = false;
+					constexpr bool bDampingAfter = false;
+					const FSolverReal ExpStiffnessValue = StiffnessHasWeightMap ? Stiffness[ConstraintIndex] : StiffnessNoMap;
+					const FSolverReal ExpStiffnessWeftValue = StiffnessWeftHasWeightMap ? StiffnessWeft[ConstraintIndex] : StiffnessWeftNoMap;
+					const FSolverReal ExpStiffnessBiasValue = StiffnessBiasHasWeightMap ? StiffnessBias[ConstraintIndex] : StiffnessBiasNoMap;
+					const FSolverReal DampingRatioValue = DampingHasWeightMap ? DampingRatio[ConstraintIndex] : DampingNoMap;
+					ApplyHelper<bDampingBefore, bSingleLambda, bSeparateStretch, bDampingAfter>(Particles, Dt, ConstraintIndex, FSolverVec3(ExpStiffnessWeftValue, ExpStiffnessValue, ExpStiffnessBiasValue), DampingRatioValue);
+				}
 			}
-			for (int32 ConstraintIndex = 0; ConstraintIndex < Constraints.Num(); ++ConstraintIndex)
+
+			// Stretch part (possibly with damping depending on split mode)
+			switch (Chaos_XPBDSpring_SplitDampingMode)
+			{
+			case (int32)EXPBDSplitDampingMode::TwoPassBefore: // fallthrough
+			case (int32)EXPBDSplitDampingMode::TwoPassAfter: // fallthrough
+			default:
 			{
-				ApplyHelper(Particles, Dt, ConstraintIndex, ExpStiffnessValue, DampingRatioValue);
+				// Do a pass with stretch only
+				for (int32 ConstraintIndex = 0; ConstraintIndex < Constraints.Num(); ++ConstraintIndex)
+				{
+					constexpr bool bDampingBefore = false;
+					constexpr bool bSingleLambda = false;
+					constexpr bool bSeparateStretch = true;
+					constexpr bool bDampingAfter = false;
+					const FSolverReal ExpStiffnessValue = StiffnessHasWeightMap ? Stiffness[ConstraintIndex] : StiffnessNoMap;
+					const FSolverReal ExpStiffnessWeftValue = StiffnessWeftHasWeightMap ? StiffnessWeft[ConstraintIndex] : StiffnessWeftNoMap;
+					const FSolverReal ExpStiffnessBiasValue = StiffnessBiasHasWeightMap ? StiffnessBias[ConstraintIndex] : StiffnessBiasNoMap;
+					const FSolverReal DampingRatioValue = DampingHasWeightMap ? DampingRatio[ConstraintIndex] : DampingNoMap;
+					ApplyHelper<bDampingBefore, bSingleLambda, bSeparateStretch, bDampingAfter>(Particles, Dt, ConstraintIndex, FSolverVec3(ExpStiffnessWeftValue, ExpStiffnessValue, ExpStiffnessBiasValue), DampingRatioValue);
+				}
+			}
+			break;
+			case (int32)EXPBDSplitDampingMode::SingleLambda:
+			{
+				// Do a pass with combined stretch and damping with a single lambda
+				for (int32 ConstraintIndex = 0; ConstraintIndex < Constraints.Num(); ++ConstraintIndex)
+				{
+					constexpr bool bDampingBefore = false;
+					constexpr bool bSingleLambda = true;
+					constexpr bool bSeparateStretch = false;
+					constexpr bool bDampingAfter = false;
+					const FSolverReal ExpStiffnessValue = StiffnessHasWeightMap ? Stiffness[ConstraintIndex] : StiffnessNoMap;
+					const FSolverReal ExpStiffnessWeftValue = StiffnessWeftHasWeightMap ? StiffnessWeft[ConstraintIndex] : StiffnessWeftNoMap;
+					const FSolverReal ExpStiffnessBiasValue = StiffnessBiasHasWeightMap ? StiffnessBias[ConstraintIndex] : StiffnessBiasNoMap;
+					const FSolverReal DampingRatioValue = DampingHasWeightMap ? DampingRatio[ConstraintIndex] : DampingNoMap;
+					ApplyHelper<bDampingBefore, bSingleLambda, bSeparateStretch, bDampingAfter>(Particles, Dt, ConstraintIndex, FSolverVec3(ExpStiffnessWeftValue, ExpStiffnessValue, ExpStiffnessBiasValue), DampingRatioValue);
+				}
+			}
+			break;
+			case (int32)EXPBDSplitDampingMode::InterleavedBefore:
+			{
+				// Do a pass with damping before stretch
+				for (int32 ConstraintIndex = 0; ConstraintIndex < Constraints.Num(); ++ConstraintIndex)
+				{
+					constexpr bool bDampingBefore = true;
+					constexpr bool bSingleLambda = false;
+					constexpr bool bSeparateStretch = true;
+					constexpr bool bDampingAfter = false;
+					const FSolverReal ExpStiffnessValue = StiffnessHasWeightMap ? Stiffness[ConstraintIndex] : StiffnessNoMap;
+					const FSolverReal ExpStiffnessWeftValue = StiffnessWeftHasWeightMap ? StiffnessWeft[ConstraintIndex] : StiffnessWeftNoMap;
+					const FSolverReal ExpStiffnessBiasValue = StiffnessBiasHasWeightMap ? StiffnessBias[ConstraintIndex] : StiffnessBiasNoMap;
+					const FSolverReal DampingRatioValue = DampingHasWeightMap ? DampingRatio[ConstraintIndex] : DampingNoMap;
+					ApplyHelper<bDampingBefore, bSingleLambda, bSeparateStretch, bDampingAfter>(Particles, Dt, ConstraintIndex, FSolverVec3(ExpStiffnessWeftValue, ExpStiffnessValue, ExpStiffnessBiasValue), DampingRatioValue);
+				}
+			}
+			break;
+			case (int32)EXPBDSplitDampingMode::InterleavedAfter:
+			{
+				// Do a pass with damping after stretch
+				for (int32 ConstraintIndex = 0; ConstraintIndex < Constraints.Num(); ++ConstraintIndex)
+				{
+					constexpr bool bDampingBefore = false;
+					constexpr bool bSingleLambda = false;
+					constexpr bool bSeparateStretch = true;
+					constexpr bool bDampingAfter = true;
+					const FSolverReal ExpStiffnessValue = StiffnessHasWeightMap ? Stiffness[ConstraintIndex] : StiffnessNoMap;
+					const FSolverReal ExpStiffnessWeftValue = StiffnessWeftHasWeightMap ? StiffnessWeft[ConstraintIndex] : StiffnessWeftNoMap;
+					const FSolverReal ExpStiffnessBiasValue = StiffnessBiasHasWeightMap ? StiffnessBias[ConstraintIndex] : StiffnessBiasNoMap;
+					const FSolverReal DampingRatioValue = DampingHasWeightMap ? DampingRatio[ConstraintIndex] : DampingNoMap;
+					ApplyHelper<bDampingBefore, bSingleLambda, bSeparateStretch, bDampingAfter>(Particles, Dt, ConstraintIndex, FSolverVec3(ExpStiffnessWeftValue, ExpStiffnessValue, ExpStiffnessBiasValue), DampingRatioValue);
+				}
+			}
+			break;
+			}
+			if (Chaos_XPBDSpring_SplitDampingMode == (int32)EXPBDSplitDampingMode::TwoPassAfter)
+			{
+				for (int32 ConstraintIndex = 0; ConstraintIndex < Constraints.Num(); ++ConstraintIndex)
+				{
+					constexpr bool bDampingBefore = false;
+					constexpr bool bSingleLambda = false;
+					constexpr bool bSeparateStretch = false;
+					constexpr bool bDampingAfter = true;
+					const FSolverReal ExpStiffnessValue = StiffnessHasWeightMap ? Stiffness[ConstraintIndex] : StiffnessNoMap;
+					const FSolverReal ExpStiffnessWeftValue = StiffnessWeftHasWeightMap ? StiffnessWeft[ConstraintIndex] : StiffnessWeftNoMap;
+					const FSolverReal ExpStiffnessBiasValue = StiffnessBiasHasWeightMap ? StiffnessBias[ConstraintIndex] : StiffnessBiasNoMap;
+					const FSolverReal DampingRatioValue = DampingHasWeightMap ? DampingRatio[ConstraintIndex] : DampingNoMap;
+					ApplyHelper<bDampingBefore, bSingleLambda, bSeparateStretch, bDampingAfter>(Particles, Dt, ConstraintIndex, FSolverVec3(ExpStiffnessWeftValue, ExpStiffnessValue, ExpStiffnessBiasValue), DampingRatioValue);
+				}
 			}
 		}
 		else
 		{
-			const FSolverReal StiffnessNoMap = (FSolverReal)Stiffness;
-			const FSolverReal StiffnessWeftNoMap = (FSolverReal)StiffnessWeft;
-			const FSolverReal StiffnessBiasNoMap = (FSolverReal)StiffnessBias;
-			const FSolverReal DampingNoMap = (FSolverReal)DampingRatio;
+			// No damping. 
 			for (int32 ConstraintIndex = 0; ConstraintIndex < Constraints.Num(); ++ConstraintIndex)
 			{
+				constexpr bool bDampingBefore = false;
+				constexpr bool bSingleLambda = false;
+				constexpr bool bSeparateStretch = true;
+				constexpr bool bDampingAfter = false;
 				const FSolverReal ExpStiffnessValue = StiffnessHasWeightMap ? Stiffness[ConstraintIndex] : StiffnessNoMap;
 				const FSolverReal ExpStiffnessWeftValue = StiffnessWeftHasWeightMap ? StiffnessWeft[ConstraintIndex] : StiffnessWeftNoMap;
 				const FSolverReal ExpStiffnessBiasValue = StiffnessBiasHasWeightMap ? StiffnessBias[ConstraintIndex] : StiffnessBiasNoMap;
-				const FSolverReal DampingRatioValue = DampingHasWeightMap ? DampingRatio[ConstraintIndex] : DampingNoMap;
-				ApplyHelper(Particles, Dt, ConstraintIndex, FSolverVec3(ExpStiffnessWeftValue, ExpStiffnessValue, ExpStiffnessBiasValue), DampingRatioValue);
+				ApplyHelper<bDampingBefore, bSingleLambda, bSeparateStretch, bDampingAfter>(Particles, Dt, ConstraintIndex, FSolverVec3(ExpStiffnessWeftValue, ExpStiffnessValue, ExpStiffnessBiasValue), (FSolverReal)0.);
+
 			}
 		}
 	}
@@ -916,7 +1226,7 @@ void FXPBDAnisotropicAxialSpringConstraints::UpdateDists()
 	Private::UpdateDists(BaseDists, WarpWeftScaleBaseMultipliers, WarpScale, WeftScale, Dists);
 }
 
-template<typename SolverParticlesOrRange>
+template<bool bDampingBefore, bool bSingleLambda, bool bSeparateStretch, bool bDampingAfter, typename SolverParticlesOrRange>
 void FXPBDAnisotropicAxialSpringConstraints::ApplyHelper(SolverParticlesOrRange& Particles, const FSolverReal Dt, const int32 ConstraintIndex, const FSolverVec3& ExpStiffnessValues, const FSolverReal DampingRatioValue) const
 {
 	const TVec3<int32>& Constraint = Constraints[ConstraintIndex];
@@ -925,22 +1235,44 @@ void FXPBDAnisotropicAxialSpringConstraints::ApplyHelper(SolverParticlesOrRange&
 	const int32 Index3 = Constraint[2];
 
 	const FSolverVec3& WarpWeftBiasBaseMultiplier = WarpWeftBiasBaseMultipliers[ConstraintIndex];
-	const FSolverReal FinalStiffnessValue = WarpWeftBiasBaseMultiplier.Dot(ExpStiffnessValues);
+	const FSolverReal ExpStiffnessValue = WarpWeftBiasBaseMultiplier.Dot(ExpStiffnessValues);
 
-	const FSolverVec3 Delta = Spring::GetXPBDAxialSpringDelta(Particles, Dt, Constraint, Barys[ConstraintIndex],  Dists[ConstraintIndex], Lambdas[ConstraintIndex],
-		FinalStiffnessValue, MinStiffness, DampingRatioValue);
+	FSolverVec3 Delta(0.f);
+	if constexpr (bDampingBefore)
+	{
+		Delta += Spring::GetXPBDAxialSpringDampingDelta(Particles, Dt, Constraint, Barys[ConstraintIndex], Dists[ConstraintIndex], LambdasDamping[ConstraintIndex],
+			ExpStiffnessValue, MinStiffness, DampingRatioValue);
+	}
+
+	if constexpr (bSingleLambda)
+	{
+		Delta += Spring::GetXPBDAxialSpringDeltaWithDamping(Particles, Dt, Constraint, Barys[ConstraintIndex], Dists[ConstraintIndex], Lambdas[ConstraintIndex],
+			ExpStiffnessValue, MinStiffness, DampingRatioValue);
+	}
+
+	if constexpr (bSeparateStretch)
+	{
+		Delta += Spring::GetXPBDAxialSpringDelta(Particles, Dt, Constraint, Barys[ConstraintIndex], Dists[ConstraintIndex], Lambdas[ConstraintIndex],
+			ExpStiffnessValue, MinStiffness);
+	}
+
+	if constexpr (bDampingAfter)
+	{
+		Delta += Spring::GetXPBDAxialSpringDampingDelta(Particles, Dt, Constraint, Barys[ConstraintIndex], Dists[ConstraintIndex], LambdasDamping[ConstraintIndex],
+			ExpStiffnessValue, MinStiffness, DampingRatioValue);
+	}
 
 	if (Particles.InvM(Index1) > (FSolverReal)0.)
 	{
-		Particles.P(Index1) -= Particles.InvM(Index1) * Delta;
+		Particles.P(Index1) += Particles.InvM(Index1) * Delta;
 	}
 	if (Particles.InvM(Index2) > (FSolverReal)0.)
 	{
-		Particles.P(Index2) += Particles.InvM(Index2) * Barys[ConstraintIndex] * Delta;
+		Particles.P(Index2) -= Particles.InvM(Index2) * Barys[ConstraintIndex] * Delta;
 	}
 	if (Particles.InvM(Index3) > (FSolverReal)0.)
 	{
-		Particles.P(Index3) += Particles.InvM(Index3) * ((FSolverReal)1. - Barys[ConstraintIndex]) * Delta;
+		Particles.P(Index3) -= Particles.InvM(Index3) * ((FSolverReal)1. - Barys[ConstraintIndex]) * Delta;
 	}
 }
 
@@ -958,91 +1290,72 @@ void FXPBDAnisotropicAxialSpringConstraints::Apply(SolverParticlesOrRange& Parti
 	if ((ConstraintsPerColorStartIndex.Num() > 1) && (Constraints.Num() > Chaos_XPBDSpring_ParallelConstraintCount))
 	{
 		const int32 ConstraintColorNum = ConstraintsPerColorStartIndex.Num() - 1;
-		if (!StiffnessHasWeightMap && !StiffnessWeftHasWeightMap && !StiffnessBiasHasWeightMap && !DampingHasWeightMap)
+#if INTEL_ISPC
+		if (bRealTypeCompatibleWithISPC && bChaos_XPBDSpring_ISPC_Enabled)
 		{
-			const FSolverVec3 ExpStiffnessValue((FSolverReal)StiffnessWeft, (FSolverReal)Stiffness, (FSolverReal)StiffnessBias);
-			const FSolverReal DampingRatioValue = (FSolverReal)DampingRatio;
-			if (ExpStiffnessValue.Max() < MinStiffness)
+			if (!StiffnessHasWeightMap && !StiffnessWeftHasWeightMap && !StiffnessBiasHasWeightMap && !DampingHasWeightMap)
 			{
-				return;
-			}
+				const FSolverVec3 ExpStiffnessValue((FSolverReal)StiffnessWeft, (FSolverReal)Stiffness, (FSolverReal)StiffnessBias);
+				const FSolverReal DampingRatioValue = (FSolverReal)DampingRatio;
+				if (ExpStiffnessValue.Max() < MinStiffness)
+				{
+					return;
+				}
 
-#if INTEL_ISPC
-			if (bRealTypeCompatibleWithISPC && bChaos_XPBDSpring_ISPC_Enabled)
-			{
+				// ISPC always does two pass damping before stretch
 				if (DampingRatioValue > 0)
 				{
 					for (int32 ConstraintColorIndex = 0; ConstraintColorIndex < ConstraintColorNum; ++ConstraintColorIndex)
 					{
 						const int32 ColorStart = ConstraintsPerColorStartIndex[ConstraintColorIndex];
 						const int32 ColorSize = ConstraintsPerColorStartIndex[ConstraintColorIndex + 1] - ColorStart;
-						ispc::ApplyXPBDAnisoAxialSpringConstraintsWithDamping(
+						ispc::ApplyXPBDAnisoAxialSpringDampingConstraints(
 							(ispc::FVector4f*)Particles.GetPAndInvM().GetData(),
 							(const ispc::FVector3f*)Particles.XArray().GetData(),
 							(ispc::FIntVector*)&Constraints.GetData()[ColorStart],
 							&Barys.GetData()[ColorStart],
 							&Dists.GetData()[ColorStart],
 							(const ispc::FVector3f*)&WarpWeftBiasBaseMultipliers.GetData()[ColorStart],
-							&Lambdas.GetData()[ColorStart],
+							&LambdasDamping.GetData()[ColorStart],
 							Dt,
 							reinterpret_cast<const ispc::FVector3f&>(ExpStiffnessValue),
 							DampingRatioValue,
 							ColorSize);
 					}
 				}
-				else
-				{
-					for (int32 ConstraintColorIndex = 0; ConstraintColorIndex < ConstraintColorNum; ++ConstraintColorIndex)
-					{
-						const int32 ColorStart = ConstraintsPerColorStartIndex[ConstraintColorIndex];
-						const int32 ColorSize = ConstraintsPerColorStartIndex[ConstraintColorIndex + 1] - ColorStart;
-						ispc::ApplyXPBDAnisoAxialSpringConstraints(
-							(ispc::FVector4f*)Particles.GetPAndInvM().GetData(),
-							(ispc::FIntVector*)&Constraints.GetData()[ColorStart],
-							&Barys.GetData()[ColorStart],
-							&Dists.GetData()[ColorStart],
-							(const ispc::FVector3f*)&WarpWeftBiasBaseMultipliers.GetData()[ColorStart],
-							&Lambdas.GetData()[ColorStart],
-							Dt,
-							reinterpret_cast<const ispc::FVector3f&>(ExpStiffnessValue),
-							ColorSize);
-					}
-				}
-			}
-			else
-#endif
-			{
 				for (int32 ConstraintColorIndex = 0; ConstraintColorIndex < ConstraintColorNum; ++ConstraintColorIndex)
 				{
 					const int32 ColorStart = ConstraintsPerColorStartIndex[ConstraintColorIndex];
 					const int32 ColorSize = ConstraintsPerColorStartIndex[ConstraintColorIndex + 1] - ColorStart;
-					PhysicsParallelFor(ColorSize, [this, &Particles, Dt, ColorStart, &ExpStiffnessValue, DampingRatioValue](const int32 Index)
-					{
-						const int32 ConstraintIndex = ColorStart + Index;
-					ApplyHelper(Particles, Dt, ConstraintIndex, ExpStiffnessValue, DampingRatioValue);
-					});
+					ispc::ApplyXPBDAnisoAxialSpringConstraints(
+						(ispc::FVector4f*)Particles.GetPAndInvM().GetData(),
+						(ispc::FIntVector*)&Constraints.GetData()[ColorStart],
+						&Barys.GetData()[ColorStart],
+						&Dists.GetData()[ColorStart],
+						(const ispc::FVector3f*)&WarpWeftBiasBaseMultipliers.GetData()[ColorStart],
+						&Lambdas.GetData()[ColorStart],
+						Dt,
+						reinterpret_cast<const ispc::FVector3f&>(ExpStiffnessValue),
+						ColorSize);
 				}
 			}
-		}
-		else  // Has weight maps
-		{
-#if INTEL_ISPC
-			if (bRealTypeCompatibleWithISPC && bChaos_XPBDSpring_ISPC_Enabled)
+			else // ISPC with maps
 			{
+				// ISPC always does two pass damping before stretch
 				if (DampingHasWeightMap || (FSolverReal)DampingRatio > 0)
 				{
 					for (int32 ConstraintColorIndex = 0; ConstraintColorIndex < ConstraintColorNum; ++ConstraintColorIndex)
 					{
 						const int32 ColorStart = ConstraintsPerColorStartIndex[ConstraintColorIndex];
 						const int32 ColorSize = ConstraintsPerColorStartIndex[ConstraintColorIndex + 1] - ColorStart;
-						ispc::ApplyXPBDAnisoAxialSpringConstraintsWithDampingAndWeightMaps(
+						ispc::ApplyXPBDAnisoAxialSpringDampingConstraintsWithWeightMaps(
 							(ispc::FVector4f*)Particles.GetPAndInvM().GetData(),
 							(const ispc::FVector3f*)Particles.XArray().GetData(),
 							(ispc::FIntVector*)&Constraints.GetData()[ColorStart],
 							&Barys.GetData()[ColorStart],
 							&Dists.GetData()[ColorStart],
 							(const ispc::FVector3f*)&WarpWeftBiasBaseMultipliers.GetData()[ColorStart],
-							&Lambdas.GetData()[ColorStart],
+							&LambdasDamping.GetData()[ColorStart],
 							Dt,
 							MinStiffness,
 							StiffnessHasWeightMap,
@@ -1060,54 +1373,202 @@ void FXPBDAnisotropicAxialSpringConstraints::Apply(SolverParticlesOrRange& Parti
 							ColorSize);
 					}
 				}
-				else
+				for (int32 ConstraintColorIndex = 0; ConstraintColorIndex < ConstraintColorNum; ++ConstraintColorIndex)
+				{
+					const int32 ColorStart = ConstraintsPerColorStartIndex[ConstraintColorIndex];
+					const int32 ColorSize = ConstraintsPerColorStartIndex[ConstraintColorIndex + 1] - ColorStart;
+					ispc::ApplyXPBDAnisoAxialSpringConstraintsWithWeightMaps(
+						(ispc::FVector4f*)Particles.GetPAndInvM().GetData(),
+						(ispc::FIntVector*)&Constraints.GetData()[ColorStart],
+						&Barys.GetData()[ColorStart],
+						&Dists.GetData()[ColorStart],
+						(const ispc::FVector3f*)&WarpWeftBiasBaseMultipliers.GetData()[ColorStart],
+						&Lambdas.GetData()[ColorStart],
+						Dt,
+						MinStiffness,
+						StiffnessHasWeightMap,
+						&Stiffness.GetIndices().GetData()[ColorStart],
+						&Stiffness.GetTable().GetData()[0],
+						StiffnessWeftHasWeightMap,
+						StiffnessWeftHasWeightMap ? &StiffnessWeft.GetIndices().GetData()[ColorStart] : nullptr,
+						&StiffnessWeft.GetTable().GetData()[0],
+						StiffnessBiasHasWeightMap,
+						StiffnessBiasHasWeightMap ? &StiffnessBias.GetIndices().GetData()[ColorStart] : nullptr,
+						&StiffnessBias.GetTable().GetData()[0],
+						ColorSize);
+				}
+			}
+		}
+		else
+#endif
+		{
+			// Parallel non-ispc
+			const FSolverReal StiffnessNoMap = (FSolverReal)Stiffness;
+			const FSolverReal StiffnessWeftNoMap = (FSolverReal)StiffnessWeft;
+			const FSolverReal StiffnessBiasNoMap = (FSolverReal)StiffnessBias;
+			const FSolverReal DampingNoMap = (FSolverReal)DampingRatio;
+			if (DampingNoMap > 0 || DampingHasWeightMap)
+			{
+				if (Chaos_XPBDSpring_SplitDampingMode == (int32)EXPBDSplitDampingMode::TwoPassBefore)
 				{
 					for (int32 ConstraintColorIndex = 0; ConstraintColorIndex < ConstraintColorNum; ++ConstraintColorIndex)
 					{
 						const int32 ColorStart = ConstraintsPerColorStartIndex[ConstraintColorIndex];
 						const int32 ColorSize = ConstraintsPerColorStartIndex[ConstraintColorIndex + 1] - ColorStart;
-						ispc::ApplyXPBDAnisoAxialSpringConstraintsWithWeightMaps(
-							(ispc::FVector4f*)Particles.GetPAndInvM().GetData(),
-							(ispc::FIntVector*)&Constraints.GetData()[ColorStart],
-							&Barys.GetData()[ColorStart],
-							&Dists.GetData()[ColorStart],
-							(const ispc::FVector3f*)&WarpWeftBiasBaseMultipliers.GetData()[ColorStart],
-							&Lambdas.GetData()[ColorStart],
-							Dt,
-							MinStiffness,
-							StiffnessHasWeightMap,
-							&Stiffness.GetIndices().GetData()[ColorStart],
-							&Stiffness.GetTable().GetData()[0],
-							StiffnessWeftHasWeightMap,
-							StiffnessWeftHasWeightMap ? &StiffnessWeft.GetIndices().GetData()[ColorStart] : nullptr,
-							&StiffnessWeft.GetTable().GetData()[0],
-							StiffnessBiasHasWeightMap,
-							StiffnessBiasHasWeightMap ? &StiffnessBias.GetIndices().GetData()[ColorStart] : nullptr,
-							&StiffnessBias.GetTable().GetData()[0],
-							ColorSize);
+						PhysicsParallelFor(ColorSize, [this, &Particles, Dt, ColorStart, StiffnessHasWeightMap, StiffnessNoMap, StiffnessWeftHasWeightMap, StiffnessWeftNoMap, StiffnessBiasHasWeightMap, StiffnessBiasNoMap, DampingHasWeightMap, DampingNoMap](const int32 Index)
+						{
+							const int32 ConstraintIndex = ColorStart + Index;
+							constexpr bool bDampingBefore = true;
+							constexpr bool bSingleLambda = false;
+							constexpr bool bSeparateStretch = false;
+							constexpr bool bDampingAfter = false;
+							const FSolverReal ExpStiffnessValue = StiffnessHasWeightMap ? Stiffness[ConstraintIndex] : StiffnessNoMap;
+							const FSolverReal ExpStiffnessWeftValue = StiffnessWeftHasWeightMap ? StiffnessWeft[ConstraintIndex] : StiffnessWeftNoMap;
+							const FSolverReal ExpStiffnessBiasValue = StiffnessBiasHasWeightMap ? StiffnessBias[ConstraintIndex] : StiffnessBiasNoMap;
+							const FSolverReal DampingRatioValue = DampingHasWeightMap ? DampingRatio[ConstraintIndex] : DampingNoMap;
+							ApplyHelper<bDampingBefore, bSingleLambda, bSeparateStretch, bDampingAfter>(Particles, Dt, ConstraintIndex, FSolverVec3(ExpStiffnessWeftValue, ExpStiffnessValue, ExpStiffnessBiasValue), DampingRatioValue);
+						});
+					}
+				}
+				
+				// Stretch part (possibly with damping depending on split mode)
+				switch (Chaos_XPBDSpring_SplitDampingMode)
+				{
+				case (int32)EXPBDSplitDampingMode::TwoPassBefore: // fallthrough
+				case (int32)EXPBDSplitDampingMode::TwoPassAfter: // fallthrough
+				default:
+				{
+					// Do a pass with stretch only
+					for (int32 ConstraintColorIndex = 0; ConstraintColorIndex < ConstraintColorNum; ++ConstraintColorIndex)
+					{
+						const int32 ColorStart = ConstraintsPerColorStartIndex[ConstraintColorIndex];
+						const int32 ColorSize = ConstraintsPerColorStartIndex[ConstraintColorIndex + 1] - ColorStart;
+						PhysicsParallelFor(ColorSize, [this, &Particles, Dt, ColorStart, StiffnessHasWeightMap, StiffnessNoMap, StiffnessWeftHasWeightMap, StiffnessWeftNoMap, StiffnessBiasHasWeightMap, StiffnessBiasNoMap, DampingHasWeightMap, DampingNoMap](const int32 Index)
+						{
+							const int32 ConstraintIndex = ColorStart + Index;
+							constexpr bool bDampingBefore = false;
+							constexpr bool bSingleLambda = false;
+							constexpr bool bSeparateStretch = true;
+							constexpr bool bDampingAfter = false;
+							const FSolverReal ExpStiffnessValue = StiffnessHasWeightMap ? Stiffness[ConstraintIndex] : StiffnessNoMap;
+							const FSolverReal ExpStiffnessWeftValue = StiffnessWeftHasWeightMap ? StiffnessWeft[ConstraintIndex] : StiffnessWeftNoMap;
+							const FSolverReal ExpStiffnessBiasValue = StiffnessBiasHasWeightMap ? StiffnessBias[ConstraintIndex] : StiffnessBiasNoMap;
+							const FSolverReal DampingRatioValue = DampingHasWeightMap ? DampingRatio[ConstraintIndex] : DampingNoMap;
+							ApplyHelper<bDampingBefore, bSingleLambda, bSeparateStretch, bDampingAfter>(Particles, Dt, ConstraintIndex, FSolverVec3(ExpStiffnessWeftValue, ExpStiffnessValue, ExpStiffnessBiasValue), DampingRatioValue);
+						});
+					}
+				}
+				break;
+				case (int32)EXPBDSplitDampingMode::SingleLambda:
+				{
+					// Do a pass with combined stretch and damping with a single lambda
+					for (int32 ConstraintColorIndex = 0; ConstraintColorIndex < ConstraintColorNum; ++ConstraintColorIndex)
+					{
+						const int32 ColorStart = ConstraintsPerColorStartIndex[ConstraintColorIndex];
+						const int32 ColorSize = ConstraintsPerColorStartIndex[ConstraintColorIndex + 1] - ColorStart;
+						PhysicsParallelFor(ColorSize, [this, &Particles, Dt, ColorStart, StiffnessHasWeightMap, StiffnessNoMap, StiffnessWeftHasWeightMap, StiffnessWeftNoMap, StiffnessBiasHasWeightMap, StiffnessBiasNoMap, DampingHasWeightMap, DampingNoMap](const int32 Index)
+						{
+							const int32 ConstraintIndex = ColorStart + Index;
+							constexpr bool bDampingBefore = false;
+							constexpr bool bSingleLambda = true;
+							constexpr bool bSeparateStretch = false;
+							constexpr bool bDampingAfter = false;
+							const FSolverReal ExpStiffnessValue = StiffnessHasWeightMap ? Stiffness[ConstraintIndex] : StiffnessNoMap;
+							const FSolverReal ExpStiffnessWeftValue = StiffnessWeftHasWeightMap ? StiffnessWeft[ConstraintIndex] : StiffnessWeftNoMap;
+							const FSolverReal ExpStiffnessBiasValue = StiffnessBiasHasWeightMap ? StiffnessBias[ConstraintIndex] : StiffnessBiasNoMap;
+							const FSolverReal DampingRatioValue = DampingHasWeightMap ? DampingRatio[ConstraintIndex] : DampingNoMap;
+							ApplyHelper<bDampingBefore, bSingleLambda, bSeparateStretch, bDampingAfter>(Particles, Dt, ConstraintIndex, FSolverVec3(ExpStiffnessWeftValue, ExpStiffnessValue, ExpStiffnessBiasValue), DampingRatioValue);
+						});
+					}
+				}
+				break;
+				case (int32)EXPBDSplitDampingMode::InterleavedBefore:
+				{
+					// Do a pass with damping before stretch
+					for (int32 ConstraintColorIndex = 0; ConstraintColorIndex < ConstraintColorNum; ++ConstraintColorIndex)
+					{
+						const int32 ColorStart = ConstraintsPerColorStartIndex[ConstraintColorIndex];
+						const int32 ColorSize = ConstraintsPerColorStartIndex[ConstraintColorIndex + 1] - ColorStart;
+						PhysicsParallelFor(ColorSize, [this, &Particles, Dt, ColorStart, StiffnessHasWeightMap, StiffnessNoMap, StiffnessWeftHasWeightMap, StiffnessWeftNoMap, StiffnessBiasHasWeightMap, StiffnessBiasNoMap, DampingHasWeightMap, DampingNoMap](const int32 Index)
+						{
+							const int32 ConstraintIndex = ColorStart + Index;
+							constexpr bool bDampingBefore = true;
+							constexpr bool bSingleLambda = false;
+							constexpr bool bSeparateStretch = true;
+							constexpr bool bDampingAfter = false;
+							const FSolverReal ExpStiffnessValue = StiffnessHasWeightMap ? Stiffness[ConstraintIndex] : StiffnessNoMap;
+							const FSolverReal ExpStiffnessWeftValue = StiffnessWeftHasWeightMap ? StiffnessWeft[ConstraintIndex] : StiffnessWeftNoMap;
+							const FSolverReal ExpStiffnessBiasValue = StiffnessBiasHasWeightMap ? StiffnessBias[ConstraintIndex] : StiffnessBiasNoMap;
+							const FSolverReal DampingRatioValue = DampingHasWeightMap ? DampingRatio[ConstraintIndex] : DampingNoMap;
+							ApplyHelper<bDampingBefore, bSingleLambda, bSeparateStretch, bDampingAfter>(Particles, Dt, ConstraintIndex, FSolverVec3(ExpStiffnessWeftValue, ExpStiffnessValue, ExpStiffnessBiasValue), DampingRatioValue);
+						});
+					}
+				}
+				break;
+				case (int32)EXPBDSplitDampingMode::InterleavedAfter:
+				{
+					// Do a pass with damping after stretch
+					for (int32 ConstraintColorIndex = 0; ConstraintColorIndex < ConstraintColorNum; ++ConstraintColorIndex)
+					{
+						const int32 ColorStart = ConstraintsPerColorStartIndex[ConstraintColorIndex];
+						const int32 ColorSize = ConstraintsPerColorStartIndex[ConstraintColorIndex + 1] - ColorStart;
+						PhysicsParallelFor(ColorSize, [this, &Particles, Dt, ColorStart, StiffnessHasWeightMap, StiffnessNoMap, StiffnessWeftHasWeightMap, StiffnessWeftNoMap, StiffnessBiasHasWeightMap, StiffnessBiasNoMap, DampingHasWeightMap, DampingNoMap](const int32 Index)
+						{
+							const int32 ConstraintIndex = ColorStart + Index;
+							constexpr bool bDampingBefore = false;
+							constexpr bool bSingleLambda = false;
+							constexpr bool bSeparateStretch = true;
+							constexpr bool bDampingAfter = true;
+							const FSolverReal ExpStiffnessValue = StiffnessHasWeightMap ? Stiffness[ConstraintIndex] : StiffnessNoMap;
+							const FSolverReal ExpStiffnessWeftValue = StiffnessWeftHasWeightMap ? StiffnessWeft[ConstraintIndex] : StiffnessWeftNoMap;
+							const FSolverReal ExpStiffnessBiasValue = StiffnessBiasHasWeightMap ? StiffnessBias[ConstraintIndex] : StiffnessBiasNoMap;
+							const FSolverReal DampingRatioValue = DampingHasWeightMap ? DampingRatio[ConstraintIndex] : DampingNoMap;
+							ApplyHelper<bDampingBefore, bSingleLambda, bSeparateStretch, bDampingAfter>(Particles, Dt, ConstraintIndex, FSolverVec3(ExpStiffnessWeftValue, ExpStiffnessValue, ExpStiffnessBiasValue), DampingRatioValue);
+						});
+					}
+				}
+				break;
+				}
+				if (Chaos_XPBDSpring_SplitDampingMode == (int32)EXPBDSplitDampingMode::TwoPassAfter)
+				{
+					for (int32 ConstraintColorIndex = 0; ConstraintColorIndex < ConstraintColorNum; ++ConstraintColorIndex)
+					{
+						const int32 ColorStart = ConstraintsPerColorStartIndex[ConstraintColorIndex];
+						const int32 ColorSize = ConstraintsPerColorStartIndex[ConstraintColorIndex + 1] - ColorStart;
+						PhysicsParallelFor(ColorSize, [this, &Particles, Dt, ColorStart, StiffnessHasWeightMap, StiffnessNoMap, StiffnessWeftHasWeightMap, StiffnessWeftNoMap, StiffnessBiasHasWeightMap, StiffnessBiasNoMap, DampingHasWeightMap, DampingNoMap](const int32 Index)
+						{
+							const int32 ConstraintIndex = ColorStart + Index;
+							constexpr bool bDampingBefore = false;
+							constexpr bool bSingleLambda = false;
+							constexpr bool bSeparateStretch = false;
+							constexpr bool bDampingAfter = true;
+							const FSolverReal ExpStiffnessValue = StiffnessHasWeightMap ? Stiffness[ConstraintIndex] : StiffnessNoMap;
+							const FSolverReal ExpStiffnessWeftValue = StiffnessWeftHasWeightMap ? StiffnessWeft[ConstraintIndex] : StiffnessWeftNoMap;
+							const FSolverReal ExpStiffnessBiasValue = StiffnessBiasHasWeightMap ? StiffnessBias[ConstraintIndex] : StiffnessBiasNoMap;
+							const FSolverReal DampingRatioValue = DampingHasWeightMap ? DampingRatio[ConstraintIndex] : DampingNoMap;
+							ApplyHelper<bDampingBefore, bSingleLambda, bSeparateStretch, bDampingAfter>(Particles, Dt, ConstraintIndex, FSolverVec3(ExpStiffnessWeftValue, ExpStiffnessValue, ExpStiffnessBiasValue), DampingRatioValue);
+						});
 					}
 				}
 			}
 			else
-#endif
 			{
-				const FSolverReal StiffnessNoMap = (FSolverReal)Stiffness;
-				const FSolverReal StiffnessWeftNoMap = (FSolverReal)StiffnessWeft;
-				const FSolverReal StiffnessBiasNoMap = (FSolverReal)StiffnessBias;
-				const FSolverReal DampingNoMap = (FSolverReal)DampingRatio;
+				// No damping. 
 				for (int32 ConstraintColorIndex = 0; ConstraintColorIndex < ConstraintColorNum; ++ConstraintColorIndex)
 				{
 					const int32 ColorStart = ConstraintsPerColorStartIndex[ConstraintColorIndex];
 					const int32 ColorSize = ConstraintsPerColorStartIndex[ConstraintColorIndex + 1] - ColorStart;
-					PhysicsParallelFor(ColorSize, [this, &Particles, Dt, ColorStart, StiffnessHasWeightMap, StiffnessNoMap, StiffnessWeftHasWeightMap, StiffnessWeftNoMap, StiffnessBiasHasWeightMap,
-						StiffnessBiasNoMap, DampingHasWeightMap, DampingNoMap](const int32 Index)
+					PhysicsParallelFor(ColorSize, [this, &Particles, Dt, ColorStart, StiffnessHasWeightMap, StiffnessNoMap, StiffnessWeftHasWeightMap, StiffnessWeftNoMap, StiffnessBiasHasWeightMap, StiffnessBiasNoMap](const int32 Index)
 					{
 						const int32 ConstraintIndex = ColorStart + Index;
+						constexpr bool bDampingBefore = false;
+						constexpr bool bSingleLambda = false;
+						constexpr bool bSeparateStretch = true;
+						constexpr bool bDampingAfter = false;
 						const FSolverReal ExpStiffnessValue = StiffnessHasWeightMap ? Stiffness[ConstraintIndex] : StiffnessNoMap;
 						const FSolverReal ExpStiffnessWeftValue = StiffnessWeftHasWeightMap ? StiffnessWeft[ConstraintIndex] : StiffnessWeftNoMap;
 						const FSolverReal ExpStiffnessBiasValue = StiffnessBiasHasWeightMap ? StiffnessBias[ConstraintIndex] : StiffnessBiasNoMap;
-						const FSolverReal DampingRatioValue = DampingHasWeightMap ? DampingRatio[ConstraintIndex] : DampingNoMap;
-						ApplyHelper(Particles, Dt, ConstraintIndex, FSolverVec3(ExpStiffnessWeftValue, ExpStiffnessValue, ExpStiffnessBiasValue), DampingRatioValue);
+						ApplyHelper<bDampingBefore, bSingleLambda, bSeparateStretch, bDampingAfter>(Particles, Dt, ConstraintIndex, FSolverVec3(ExpStiffnessWeftValue, ExpStiffnessValue, ExpStiffnessBiasValue), (FSolverReal)0.);
 					});
 				}
 			}
@@ -1115,32 +1576,133 @@ void FXPBDAnisotropicAxialSpringConstraints::Apply(SolverParticlesOrRange& Parti
 	}
 	else
 	{
-		if (!StiffnessHasWeightMap && !StiffnessWeftHasWeightMap && !StiffnessBiasHasWeightMap && !DampingHasWeightMap)
+		// Single threaded
+		const FSolverReal StiffnessNoMap = (FSolverReal)Stiffness;
+		const FSolverReal StiffnessWeftNoMap = (FSolverReal)StiffnessWeft;
+		const FSolverReal StiffnessBiasNoMap = (FSolverReal)StiffnessBias;
+		const FSolverReal DampingNoMap = (FSolverReal)DampingRatio;
+		if (DampingNoMap > 0 || DampingHasWeightMap)
 		{
-			const FSolverVec3 ExpStiffnessValue((FSolverReal)StiffnessWeft, (FSolverReal)Stiffness, (FSolverReal)StiffnessBias);
-			const FSolverReal DampingRatioValue = (FSolverReal)DampingRatio;
-			if (ExpStiffnessValue.Max() < MinStiffness)
+			if (Chaos_XPBDSpring_SplitDampingMode == (int32)EXPBDSplitDampingMode::TwoPassBefore)
 			{
-				return;
+				for (int32 ConstraintIndex = 0; ConstraintIndex < Constraints.Num(); ++ConstraintIndex)
+				{
+					constexpr bool bDampingBefore = true;
+					constexpr bool bSingleLambda = false;
+					constexpr bool bSeparateStretch = false;
+					constexpr bool bDampingAfter = false;
+					const FSolverReal ExpStiffnessValue = StiffnessHasWeightMap ? Stiffness[ConstraintIndex] : StiffnessNoMap;
+					const FSolverReal ExpStiffnessWeftValue = StiffnessWeftHasWeightMap ? StiffnessWeft[ConstraintIndex] : StiffnessWeftNoMap;
+					const FSolverReal ExpStiffnessBiasValue = StiffnessBiasHasWeightMap ? StiffnessBias[ConstraintIndex] : StiffnessBiasNoMap;
+					const FSolverReal DampingRatioValue = DampingHasWeightMap ? DampingRatio[ConstraintIndex] : DampingNoMap;
+					ApplyHelper<bDampingBefore, bSingleLambda, bSeparateStretch, bDampingAfter>(Particles, Dt, ConstraintIndex, FSolverVec3(ExpStiffnessWeftValue, ExpStiffnessValue, ExpStiffnessBiasValue), DampingRatioValue);
+				}
 			}
-			for (int32 ConstraintIndex = 0; ConstraintIndex < Constraints.Num(); ++ConstraintIndex)
+
+			// Stretch part (possibly with damping depending on split mode)
+			switch (Chaos_XPBDSpring_SplitDampingMode)
+			{
+			case (int32)EXPBDSplitDampingMode::TwoPassBefore: // fallthrough
+			case (int32)EXPBDSplitDampingMode::TwoPassAfter: // fallthrough
+			default:
+			{
+				// Do a pass with stretch only
+				for (int32 ConstraintIndex = 0; ConstraintIndex < Constraints.Num(); ++ConstraintIndex)
+				{
+					constexpr bool bDampingBefore = false;
+					constexpr bool bSingleLambda = false;
+					constexpr bool bSeparateStretch = true;
+					constexpr bool bDampingAfter = false;
+					const FSolverReal ExpStiffnessValue = StiffnessHasWeightMap ? Stiffness[ConstraintIndex] : StiffnessNoMap;
+					const FSolverReal ExpStiffnessWeftValue = StiffnessWeftHasWeightMap ? StiffnessWeft[ConstraintIndex] : StiffnessWeftNoMap;
+					const FSolverReal ExpStiffnessBiasValue = StiffnessBiasHasWeightMap ? StiffnessBias[ConstraintIndex] : StiffnessBiasNoMap;
+					const FSolverReal DampingRatioValue = DampingHasWeightMap ? DampingRatio[ConstraintIndex] : DampingNoMap;
+					ApplyHelper<bDampingBefore, bSingleLambda, bSeparateStretch, bDampingAfter>(Particles, Dt, ConstraintIndex, FSolverVec3(ExpStiffnessWeftValue, ExpStiffnessValue, ExpStiffnessBiasValue), DampingRatioValue);
+				}
+			}
+			break;
+			case (int32)EXPBDSplitDampingMode::SingleLambda:
+			{
+				// Do a pass with combined stretch and damping with a single lambda
+				for (int32 ConstraintIndex = 0; ConstraintIndex < Constraints.Num(); ++ConstraintIndex)
+				{
+					constexpr bool bDampingBefore = false;
+					constexpr bool bSingleLambda = true;
+					constexpr bool bSeparateStretch = false;
+					constexpr bool bDampingAfter = false;
+					const FSolverReal ExpStiffnessValue = StiffnessHasWeightMap ? Stiffness[ConstraintIndex] : StiffnessNoMap;
+					const FSolverReal ExpStiffnessWeftValue = StiffnessWeftHasWeightMap ? StiffnessWeft[ConstraintIndex] : StiffnessWeftNoMap;
+					const FSolverReal ExpStiffnessBiasValue = StiffnessBiasHasWeightMap ? StiffnessBias[ConstraintIndex] : StiffnessBiasNoMap;
+					const FSolverReal DampingRatioValue = DampingHasWeightMap ? DampingRatio[ConstraintIndex] : DampingNoMap;
+					ApplyHelper<bDampingBefore, bSingleLambda, bSeparateStretch, bDampingAfter>(Particles, Dt, ConstraintIndex, FSolverVec3(ExpStiffnessWeftValue, ExpStiffnessValue, ExpStiffnessBiasValue), DampingRatioValue);
+				}
+			}
+			break;
+			case (int32)EXPBDSplitDampingMode::InterleavedBefore:
 			{
-				ApplyHelper(Particles, Dt, ConstraintIndex, ExpStiffnessValue, DampingRatioValue);
+				// Do a pass with damping before stretch
+				for (int32 ConstraintIndex = 0; ConstraintIndex < Constraints.Num(); ++ConstraintIndex)
+				{
+					constexpr bool bDampingBefore = true;
+					constexpr bool bSingleLambda = false;
+					constexpr bool bSeparateStretch = true;
+					constexpr bool bDampingAfter = false;
+					const FSolverReal ExpStiffnessValue = StiffnessHasWeightMap ? Stiffness[ConstraintIndex] : StiffnessNoMap;
+					const FSolverReal ExpStiffnessWeftValue = StiffnessWeftHasWeightMap ? StiffnessWeft[ConstraintIndex] : StiffnessWeftNoMap;
+					const FSolverReal ExpStiffnessBiasValue = StiffnessBiasHasWeightMap ? StiffnessBias[ConstraintIndex] : StiffnessBiasNoMap;
+					const FSolverReal DampingRatioValue = DampingHasWeightMap ? DampingRatio[ConstraintIndex] : DampingNoMap;
+					ApplyHelper<bDampingBefore, bSingleLambda, bSeparateStretch, bDampingAfter>(Particles, Dt, ConstraintIndex, FSolverVec3(ExpStiffnessWeftValue, ExpStiffnessValue, ExpStiffnessBiasValue), DampingRatioValue);
+				}
+			}
+			break;
+			case (int32)EXPBDSplitDampingMode::InterleavedAfter:
+			{
+				// Do a pass with damping after stretch
+				for (int32 ConstraintIndex = 0; ConstraintIndex < Constraints.Num(); ++ConstraintIndex)
+				{
+					constexpr bool bDampingBefore = false;
+					constexpr bool bSingleLambda = false;
+					constexpr bool bSeparateStretch = true;
+					constexpr bool bDampingAfter = true;
+					const FSolverReal ExpStiffnessValue = StiffnessHasWeightMap ? Stiffness[ConstraintIndex] : StiffnessNoMap;
+					const FSolverReal ExpStiffnessWeftValue = StiffnessWeftHasWeightMap ? StiffnessWeft[ConstraintIndex] : StiffnessWeftNoMap;
+					const FSolverReal ExpStiffnessBiasValue = StiffnessBiasHasWeightMap ? StiffnessBias[ConstraintIndex] : StiffnessBiasNoMap;
+					const FSolverReal DampingRatioValue = DampingHasWeightMap ? DampingRatio[ConstraintIndex] : DampingNoMap;
+					ApplyHelper<bDampingBefore, bSingleLambda, bSeparateStretch, bDampingAfter>(Particles, Dt, ConstraintIndex, FSolverVec3(ExpStiffnessWeftValue, ExpStiffnessValue, ExpStiffnessBiasValue), DampingRatioValue);
+				}
+			}
+			break;
+			}
+			if (Chaos_XPBDSpring_SplitDampingMode == (int32)EXPBDSplitDampingMode::TwoPassAfter)
+			{
+				for (int32 ConstraintIndex = 0; ConstraintIndex < Constraints.Num(); ++ConstraintIndex)
+				{
+					constexpr bool bDampingBefore = false;
+					constexpr bool bSingleLambda = false;
+					constexpr bool bSeparateStretch = false;
+					constexpr bool bDampingAfter = true;
+					const FSolverReal ExpStiffnessValue = StiffnessHasWeightMap ? Stiffness[ConstraintIndex] : StiffnessNoMap;
+					const FSolverReal ExpStiffnessWeftValue = StiffnessWeftHasWeightMap ? StiffnessWeft[ConstraintIndex] : StiffnessWeftNoMap;
+					const FSolverReal ExpStiffnessBiasValue = StiffnessBiasHasWeightMap ? StiffnessBias[ConstraintIndex] : StiffnessBiasNoMap;
+					const FSolverReal DampingRatioValue = DampingHasWeightMap ? DampingRatio[ConstraintIndex] : DampingNoMap;
+					ApplyHelper<bDampingBefore, bSingleLambda, bSeparateStretch, bDampingAfter>(Particles, Dt, ConstraintIndex, FSolverVec3(ExpStiffnessWeftValue, ExpStiffnessValue, ExpStiffnessBiasValue), DampingRatioValue);
+				}
 			}
 		}
 		else
 		{
-			const FSolverReal StiffnessNoMap = (FSolverReal)Stiffness;
-			const FSolverReal StiffnessWeftNoMap = (FSolverReal)StiffnessWeft;
-			const FSolverReal StiffnessBiasNoMap = (FSolverReal)StiffnessBias;
-			const FSolverReal DampingNoMap = (FSolverReal)DampingRatio;
+			// No damping. 
 			for (int32 ConstraintIndex = 0; ConstraintIndex < Constraints.Num(); ++ConstraintIndex)
 			{
+				constexpr bool bDampingBefore = false;
+				constexpr bool bSingleLambda = false;
+				constexpr bool bSeparateStretch = true;
+				constexpr bool bDampingAfter = false;
 				const FSolverReal ExpStiffnessValue = StiffnessHasWeightMap ? Stiffness[ConstraintIndex] : StiffnessNoMap;
 				const FSolverReal ExpStiffnessWeftValue = StiffnessWeftHasWeightMap ? StiffnessWeft[ConstraintIndex] : StiffnessWeftNoMap;
 				const FSolverReal ExpStiffnessBiasValue = StiffnessBiasHasWeightMap ? StiffnessBias[ConstraintIndex] : StiffnessBiasNoMap;
-				const FSolverReal DampingRatioValue = DampingHasWeightMap ? DampingRatio[ConstraintIndex] : DampingNoMap;
-				ApplyHelper(Particles, Dt, ConstraintIndex, FSolverVec3(ExpStiffnessWeftValue, ExpStiffnessValue, ExpStiffnessBiasValue), DampingRatioValue);
+				ApplyHelper<bDampingBefore, bSingleLambda, bSeparateStretch, bDampingAfter>(Particles, Dt, ConstraintIndex, FSolverVec3(ExpStiffnessWeftValue, ExpStiffnessValue, ExpStiffnessBiasValue), (FSolverReal)0.);
+
 			}
 		}
 	}
@@ -1195,13 +1757,13 @@ void FXPBDAnisotropicAxialSpringConstraints::UpdateLinearSystem(const FSolverPar
 					StiffnessBiasNoMap, DampingHasWeightMap, DampingNoMap](const int32 Index)
 				{
 					const int32 ConstraintIndex = ColorStart + Index;
-				const FSolverReal ExpStiffnessValue = StiffnessHasWeightMap ? Stiffness[ConstraintIndex] : StiffnessNoMap;
-				const FSolverReal ExpStiffnessWeftValue = StiffnessWeftHasWeightMap ? StiffnessWeft[ConstraintIndex] : StiffnessWeftNoMap;
-				const FSolverReal ExpStiffnessBiasValue = StiffnessBiasHasWeightMap ? StiffnessBias[ConstraintIndex] : StiffnessBiasNoMap;
-				const FSolverReal DampingRatioValue = DampingHasWeightMap ? DampingRatio[ConstraintIndex] : DampingNoMap;
-				const FSolverVec3& WarpWeftBiasBaseMultiplier = WarpWeftBiasBaseMultipliers[ConstraintIndex];
-				const FSolverReal FinalStiffnessValue = WarpWeftBiasBaseMultiplier.Dot(FSolverVec3(ExpStiffnessWeftValue, ExpStiffnessValue, ExpStiffnessBiasValue));
-				Spring::UpdateAxialSpringLinearSystem(Particles, Dt, Constraints[ConstraintIndex], Barys[ConstraintIndex], Dists[ConstraintIndex], FinalStiffnessValue, MinStiffness, DampingRatioValue, LinearSystem);
+					const FSolverReal ExpStiffnessValue = StiffnessHasWeightMap ? Stiffness[ConstraintIndex] : StiffnessNoMap;
+					const FSolverReal ExpStiffnessWeftValue = StiffnessWeftHasWeightMap ? StiffnessWeft[ConstraintIndex] : StiffnessWeftNoMap;
+					const FSolverReal ExpStiffnessBiasValue = StiffnessBiasHasWeightMap ? StiffnessBias[ConstraintIndex] : StiffnessBiasNoMap;
+					const FSolverReal DampingRatioValue = DampingHasWeightMap ? DampingRatio[ConstraintIndex] : DampingNoMap;
+					const FSolverVec3& WarpWeftBiasBaseMultiplier = WarpWeftBiasBaseMultipliers[ConstraintIndex];
+					const FSolverReal FinalStiffnessValue = WarpWeftBiasBaseMultiplier.Dot(FSolverVec3(ExpStiffnessWeftValue, ExpStiffnessValue, ExpStiffnessBiasValue));
+					Spring::UpdateAxialSpringLinearSystem(Particles, Dt, Constraints[ConstraintIndex], Barys[ConstraintIndex], Dists[ConstraintIndex], FinalStiffnessValue, MinStiffness, DampingRatioValue, LinearSystem);
 				});
 			}
 		}
diff --git a/Engine/Source/Runtime/Experimental/Chaos/Private/Chaos/XPBDAnisotropicSpringConstraints.ispc b/Engine/Source/Runtime/Experimental/Chaos/Private/Chaos/XPBDAnisotropicSpringConstraints.ispc
--- a/Engine/Source/Runtime/Experimental/Chaos/Private/Chaos/XPBDAnisotropicSpringConstraints.ispc
+++ b/Engine/Source/Runtime/Experimental/Chaos/Private/Chaos/XPBDAnisotropicSpringConstraints.ispc
@@ -41,11 +41,11 @@ export void ApplyXPBDAnisoSpringConstraints(uniform FVector4f PandInvM[],
 
 		if (InvM1 > FLOAT_ZERO)
 		{
-			VectorScatter(&PandInvM[Index1], SetVector4( P1 - (InvM1 * Delta), InvM1 ));
+			VectorScatter(&PandInvM[Index1], SetVector4( P1 + (InvM1 * Delta), InvM1 ));
 		}
 		if (InvM2 > FLOAT_ZERO)
 		{
-			VectorScatter(&PandInvM[Index2], SetVector4( P2 + (InvM2 * Delta), InvM2 ));
+			VectorScatter(&PandInvM[Index2], SetVector4( P2 - (InvM2 * Delta), InvM2 ));
 		}
 	}
 }
@@ -90,11 +90,60 @@ export void ApplyXPBDAnisoSpringConstraintsWithDamping(uniform FVector4f PandInv
 
 		if (InvM1 > FLOAT_ZERO)
 		{
-			VectorScatter(&PandInvM[Index1], SetVector4( P1 - (InvM1 * Delta), InvM1 ));
+			VectorScatter(&PandInvM[Index1], SetVector4( P1 + (InvM1 * Delta), InvM1 ));
 		}
 		if (InvM2 > FLOAT_ZERO)
 		{
-			VectorScatter(&PandInvM[Index2], SetVector4( P2 + (InvM2 * Delta), InvM2 ));
+			VectorScatter(&PandInvM[Index2], SetVector4( P2 - (InvM2 * Delta), InvM2 ));
+		}
+	}
+}
+
+export void ApplyXPBDAnisoSpringDampingConstraints(uniform FVector4f PandInvM[],
+												const uniform FVector3f X[],
+												const uniform FIntVector2 Constraints[],
+												const uniform float Dists[],
+												const uniform FVector3f WarpWeftBiasBaseMultiplers[],
+												uniform float Lambdas[],
+												const uniform float Dt,
+												const uniform FVector3f& Stiffness,
+												const uniform float DampingRatio,
+												const uniform int32 NumConstraints)
+{
+	foreach(Index = 0 ... NumConstraints)
+	{
+		const varying FIntVector2 Constraint = VectorLoad(&Constraints[extract(Index,0)]);
+		const varying int32 Index1 = Constraint.V[0];
+		const varying int32 Index2 = Constraint.V[1];
+
+		const varying FVector4f PandInvM1 = VectorGather(&PandInvM[Index1]);
+		const varying FVector4f PandInvM2 = VectorGather(&PandInvM[Index2]);
+
+		varying FVector3f P1, P2;
+		varying float InvM1, InvM2;
+		UnzipPandInvM(PandInvM1, P1, InvM1);
+		UnzipPandInvM(PandInvM2, P2, InvM2);
+
+		const varying FVector3f X1 = VectorGather(&X[Index1]);
+		const varying FVector3f X2 = VectorGather(&X[Index2]);
+
+		const varying float Dist = Dists[Index];
+		const varying float Lambda = Lambdas[Index];
+		const varying FVector3f WarpWeftBiasBaseMultiplier = VectorGather(&WarpWeftBiasBaseMultiplers[Index]);
+		const varying float StiffnessValue = VectorDot(WarpWeftBiasBaseMultiplier, Stiffness);
+
+		varying float DLambda;
+		const FVector3f Delta = GetXPBDSpringDampingDelta(P1, P2, X1, X2, InvM1, InvM2, Dt, Dist, Lambda, StiffnessValue, DampingRatio, FLOAT_SMALL_NUMBER, DLambda);
+
+		Lambdas[Index] = Lambda + DLambda;
+
+		if (InvM1 > FLOAT_ZERO)
+		{
+			VectorScatter(&PandInvM[Index1], SetVector4( P1 + (InvM1 * Delta), InvM1 ));
+		}
+		if (InvM2 > FLOAT_ZERO)
+		{
+			VectorScatter(&PandInvM[Index2], SetVector4( P2 - (InvM2 * Delta), InvM2 ));
 		}
 	}
 }
@@ -157,11 +206,11 @@ export void ApplyXPBDAnisoSpringConstraintsWithWeightMaps(uniform FVector4f Pand
 
 		if (InvM1 > FLOAT_ZERO)
 		{
-			VectorScatter(&PandInvM[Index1], SetVector4( P1 - (InvM1 * Delta), InvM1 ));
+			VectorScatter(&PandInvM[Index1], SetVector4( P1 + (InvM1 * Delta), InvM1 ));
 		}
 		if (InvM2 > FLOAT_ZERO)
 		{
-			VectorScatter(&PandInvM[Index2], SetVector4( P2 + (InvM2 * Delta), InvM2 ));
+			VectorScatter(&PandInvM[Index2], SetVector4( P2 - (InvM2 * Delta), InvM2 ));
 		}
 	}
 }
@@ -233,11 +282,87 @@ export void ApplyXPBDAnisoSpringConstraintsWithDampingAndWeightMaps(uniform FVec
 
 		if (InvM1 > FLOAT_ZERO)
 		{
-			VectorScatter(&PandInvM[Index1], SetVector4( P1 - (InvM1 * Delta), InvM1 ));
+			VectorScatter(&PandInvM[Index1], SetVector4( P1 + (InvM1 * Delta), InvM1 ));
 		}
 		if (InvM2 > FLOAT_ZERO)
 		{
-			VectorScatter(&PandInvM[Index2], SetVector4( P2 + (InvM2 * Delta), InvM2 ));
+			VectorScatter(&PandInvM[Index2], SetVector4( P2 - (InvM2 * Delta), InvM2 ));
+		}
+	}
+}
+
+export void ApplyXPBDAnisoSpringDampingConstraintsWithWeightMaps(uniform FVector4f PandInvM[],
+															const uniform FVector3f X[],
+															const uniform FIntVector2 Constraints[],
+															const uniform float Dists[],
+															const uniform FVector3f WarpWeftBiasBaseMultiplers[],
+															uniform float Lambdas[],
+															const uniform float Dt,
+															const uniform float MinStiffness,
+															const uniform bool StiffnessHasMap,
+															const uniform uint8 StiffnessIndices[],
+															const uniform float StiffnessTable[],
+															const uniform bool StiffnessWeftHasMap,
+															const uniform uint8 StiffnessWeftIndices[],
+															const uniform float StiffnessWeftTable[],
+															const uniform bool StiffnessBiasHasMap,
+															const uniform uint8 StiffnessBiasIndices[],
+															const uniform float StiffnessBiasTable[],
+															const uniform bool DampingHasMap,
+															const uniform uint8 DampingIndices[],
+															const uniform float DampingTable[],
+															const uniform int32 NumConstraints)
+{
+
+	foreach(Index = 0 ... NumConstraints)
+	{
+		const varying FIntVector2 Constraint = VectorLoad(&Constraints[extract(Index,0)]);
+		const varying int32 Index1 = Constraint.V[0];
+		const varying int32 Index2 = Constraint.V[1];
+		const varying int8 StiffnessIndex = StiffnessHasMap ? StiffnessIndices[Index] : 0;
+		const varying int8 StiffnessWeftIndex = StiffnessWeftHasMap ? StiffnessWeftIndices[Index] : 0;
+		const varying int8 StiffnessBiasIndex = StiffnessBiasHasMap ? StiffnessBiasIndices[Index] : 0;
+		const varying int8 DampingIndex = DampingHasMap ? DampingIndices[Index] : 0;
+
+		const varying FVector4f PandInvM1 = VectorGather(&PandInvM[Index1]);
+		const varying FVector4f PandInvM2 = VectorGather(&PandInvM[Index2]);
+
+		varying FVector3f P1, P2;
+		varying float InvM1, InvM2;
+		UnzipPandInvM(PandInvM1, P1, InvM1);
+		UnzipPandInvM(PandInvM2, P2, InvM2);
+
+		const varying FVector3f X1 = VectorGather(&X[Index1]);
+		const varying FVector3f X2 = VectorGather(&X[Index2]);
+
+		const varying float Dist = Dists[Index];
+		const varying float Lambda = Lambdas[Index];
+
+		#pragma ignore warning(perf)
+		const varying float Stiffness = StiffnessTable[StiffnessIndex];
+		#pragma ignore warning(perf)
+		const varying float StiffnessWeft = StiffnessWeftTable[StiffnessWeftIndex];
+		#pragma ignore warning(perf)
+		const varying float StiffnessBias = StiffnessBiasTable[StiffnessBiasIndex];
+		#pragma ignore warning(perf)
+		const varying float DampingRatio = DampingTable[DampingIndex];
+
+		const varying FVector3f StiffnessVector = {{StiffnessWeft, Stiffness, StiffnessBias}};
+		const varying FVector3f WarpWeftBiasBaseMultiplier = VectorGather(&WarpWeftBiasBaseMultiplers[Index]);
+		const varying float StiffnessValue = VectorDot(WarpWeftBiasBaseMultiplier, StiffnessVector);
+
+		varying float DLambda;
+		const FVector3f Delta = GetXPBDSpringDampingDelta(P1, P2, X1, X2, InvM1, InvM2, Dt, Dist, Lambda, StiffnessValue, DampingRatio, MinStiffness, DLambda);
+
+		Lambdas[Index] = Lambda + DLambda;
+
+		if (InvM1 > FLOAT_ZERO)
+		{
+			VectorScatter(&PandInvM[Index1], SetVector4( P1 + (InvM1 * Delta), InvM1 ));
+		}
+		if (InvM2 > FLOAT_ZERO)
+		{
+			VectorScatter(&PandInvM[Index2], SetVector4( P2 - (InvM2 * Delta), InvM2 ));
 		}
 	}
 }
@@ -282,15 +407,15 @@ export void ApplyXPBDAnisoAxialSpringConstraints(uniform FVector4f PandInvM[],
 
 		if (InvM1 > FLOAT_ZERO)
 		{
-			VectorScatter(&PandInvM[Index1], SetVector4( P1 - (InvM1 * Delta), InvM1 ));
+			VectorScatter(&PandInvM[Index1], SetVector4( P1 + (InvM1 * Delta), InvM1 ));
 		}
 		if (InvM2 > FLOAT_ZERO)
 		{
-			VectorScatter(&PandInvM[Index2], SetVector4( P2 + (InvM2 * Bary * Delta), InvM2 ));
+			VectorScatter(&PandInvM[Index2], SetVector4( P2 - (InvM2 * Bary * Delta), InvM2 ));
 		}
 		if (InvM3 > FLOAT_ZERO)
 		{
-			VectorScatter(&PandInvM[Index3], SetVector4( P3 + (InvM3 * (1.f - Bary) * Delta), InvM3 ));
+			VectorScatter(&PandInvM[Index3], SetVector4( P3 - (InvM3 * (1.f - Bary) * Delta), InvM3 ));
 		}
 	}
 }
@@ -341,15 +466,74 @@ export void ApplyXPBDAnisoAxialSpringConstraintsWithDamping(uniform FVector4f Pa
 
 		if (InvM1 > FLOAT_ZERO)
 		{
-			VectorScatter(&PandInvM[Index1], SetVector4( P1 - (InvM1 * Delta), InvM1 ));
+			VectorScatter(&PandInvM[Index1], SetVector4( P1 + (InvM1 * Delta), InvM1 ));
 		}
 		if (InvM2 > FLOAT_ZERO)
 		{
-			VectorScatter(&PandInvM[Index2], SetVector4( P2 + (InvM2 * Bary * Delta), InvM2 ));
+			VectorScatter(&PandInvM[Index2], SetVector4( P2 - (InvM2 * Bary * Delta), InvM2 ));
 		}
 		if (InvM3 > FLOAT_ZERO)
 		{
-			VectorScatter(&PandInvM[Index3], SetVector4( P3 + (InvM3 * (1.f - Bary) * Delta), InvM3 ));
+			VectorScatter(&PandInvM[Index3], SetVector4( P3 - (InvM3 * (1.f - Bary) * Delta), InvM3 ));
+		}
+	}
+}
+
+export void ApplyXPBDAnisoAxialSpringDampingConstraints(uniform FVector4f PandInvM[],
+												const uniform FVector3f X[],
+												const uniform FIntVector Constraints[],
+												const uniform float Barys[],
+												const uniform float Dists[],
+												const uniform FVector3f WarpWeftBiasBaseMultiplers[],
+												uniform float Lambdas[],
+												const uniform float Dt,
+												const uniform FVector3f& Stiffness,
+												const uniform float DampingRatio,
+												const uniform int32 NumConstraints)
+{
+	foreach(Index = 0 ... NumConstraints)
+	{
+		const varying FIntVector Constraint = VectorLoad(&Constraints[extract(Index,0)]);
+		const varying int32 Index1 = Constraint.V[0];
+		const varying int32 Index2 = Constraint.V[1];
+		const varying int32 Index3 = Constraint.V[2];
+		
+		const varying FVector4f PandInvM1 = VectorGather(&PandInvM[Index1]);
+		const varying FVector4f PandInvM2 = VectorGather(&PandInvM[Index2]);
+		const varying FVector4f PandInvM3 = VectorGather(&PandInvM[Index3]);
+
+		varying FVector3f P1, P2, P3;
+		varying float InvM1, InvM2, InvM3;
+		UnzipPandInvM(PandInvM1, P1, InvM1);
+		UnzipPandInvM(PandInvM2, P2, InvM2);
+		UnzipPandInvM(PandInvM3, P3, InvM3);
+
+		const varying FVector3f X1 = VectorGather(&X[Index1]);
+		const varying FVector3f X2 = VectorGather(&X[Index2]);
+		const varying FVector3f X3 = VectorGather(&X[Index3]);
+		
+		const varying float Bary = Barys[Index];
+		const varying float Dist = Dists[Index];
+		const varying float Lambda = Lambdas[Index];
+		const varying FVector3f WarpWeftBiasBaseMultiplier = VectorGather(&WarpWeftBiasBaseMultiplers[Index]);
+		const varying float StiffnessValue = VectorDot(WarpWeftBiasBaseMultiplier, Stiffness);
+
+		varying float DLambda;
+		const FVector3f Delta = GetXPBDAxialSpringDampingDelta(P1, P2, P3, X1, X2, X3, InvM1, InvM2, InvM3, Dt, Bary, Dist, Lambda, StiffnessValue, DampingRatio, FLOAT_SMALL_NUMBER, DLambda);
+
+		Lambdas[Index] = Lambda + DLambda;
+
+		if (InvM1 > FLOAT_ZERO)
+		{
+			VectorScatter(&PandInvM[Index1], SetVector4( P1 + (InvM1 * Delta), InvM1 ));
+		}
+		if (InvM2 > FLOAT_ZERO)
+		{
+			VectorScatter(&PandInvM[Index2], SetVector4( P2 - (InvM2 * Bary * Delta), InvM2 ));
+		}
+		if (InvM3 > FLOAT_ZERO)
+		{
+			VectorScatter(&PandInvM[Index3], SetVector4( P3 - (InvM3 * (1.f - Bary) * Delta), InvM3 ));
 		}
 	}
 }
@@ -417,15 +601,15 @@ export void ApplyXPBDAnisoAxialSpringConstraintsWithWeightMaps(uniform FVector4f
 
 		if (InvM1 > FLOAT_ZERO)
 		{
-			VectorScatter(&PandInvM[Index1], SetVector4( P1 - (InvM1 * Delta), InvM1 ));
+			VectorScatter(&PandInvM[Index1], SetVector4( P1 + (InvM1 * Delta), InvM1 ));
 		}
 		if (InvM2 > FLOAT_ZERO)
 		{
-			VectorScatter(&PandInvM[Index2], SetVector4( P2 + (InvM2 * Bary * Delta), InvM2 ));
+			VectorScatter(&PandInvM[Index2], SetVector4( P2 - (InvM2 * Bary * Delta), InvM2 ));
 		}
 		if (InvM3 > FLOAT_ZERO)
 		{
-			VectorScatter(&PandInvM[Index3], SetVector4( P3 + (InvM3 * (1.f - Bary) * Delta), InvM3 ));
+			VectorScatter(&PandInvM[Index3], SetVector4( P3 - (InvM3 * (1.f - Bary) * Delta), InvM3 ));
 		}
 	}
 }
@@ -504,15 +688,102 @@ export void ApplyXPBDAnisoAxialSpringConstraintsWithDampingAndWeightMaps(uniform
 
 		if (InvM1 > FLOAT_ZERO)
 		{
-			VectorScatter(&PandInvM[Index1], SetVector4( P1 - (InvM1 * Delta), InvM1 ));
+			VectorScatter(&PandInvM[Index1], SetVector4( P1 + (InvM1 * Delta), InvM1 ));
+		}
+		if (InvM2 > FLOAT_ZERO)
+		{
+			VectorScatter(&PandInvM[Index2], SetVector4( P2 - (InvM2 * Bary * Delta), InvM2 ));
+		}
+		if (InvM3 > FLOAT_ZERO)
+		{
+			VectorScatter(&PandInvM[Index3], SetVector4( P3 - (InvM3 * (1.f - Bary) * Delta), InvM3 ));
+		}
+	}
+}
+
+export void ApplyXPBDAnisoAxialSpringDampingConstraintsWithWeightMaps(uniform FVector4f PandInvM[],
+															const uniform FVector3f X[],
+															const uniform FIntVector Constraints[],
+															const uniform float Barys[],
+															const uniform float Dists[],
+															const uniform FVector3f WarpWeftBiasBaseMultiplers[],
+															uniform float Lambdas[],
+															const uniform float Dt,
+															const uniform float MinStiffness,
+															const uniform bool StiffnessHasMap,
+															const uniform uint8 StiffnessIndices[],
+															const uniform float StiffnessTable[],
+															const uniform bool StiffnessWeftHasMap,
+															const uniform uint8 StiffnessWeftIndices[],
+															const uniform float StiffnessWeftTable[],
+															const uniform bool StiffnessBiasHasMap,
+															const uniform uint8 StiffnessBiasIndices[],
+															const uniform float StiffnessBiasTable[],
+															const uniform bool DampingHasMap,
+															const uniform uint8 DampingIndices[],
+															const uniform float DampingTable[],
+															const uniform int32 NumConstraints)
+{
+
+	foreach(Index = 0 ... NumConstraints)
+	{
+		const varying FIntVector Constraint = VectorLoad(&Constraints[extract(Index,0)]);
+		const varying int32 Index1 = Constraint.V[0];
+		const varying int32 Index2 = Constraint.V[1];
+		const varying int32 Index3 = Constraint.V[2];
+
+		const varying int8 StiffnessIndex = StiffnessHasMap ? StiffnessIndices[Index] : 0;
+		const varying int8 StiffnessWeftIndex = StiffnessWeftHasMap ? StiffnessWeftIndices[Index] : 0;
+		const varying int8 StiffnessBiasIndex = StiffnessBiasHasMap ? StiffnessBiasIndices[Index] : 0;
+		const varying int8 DampingIndex = DampingHasMap ? DampingIndices[Index] : 0;
+		
+		const varying FVector4f PandInvM1 = VectorGather(&PandInvM[Index1]);
+		const varying FVector4f PandInvM2 = VectorGather(&PandInvM[Index2]);
+		const varying FVector4f PandInvM3 = VectorGather(&PandInvM[Index3]);
+
+		varying FVector3f P1, P2, P3;
+		varying float InvM1, InvM2, InvM3;
+		UnzipPandInvM(PandInvM1, P1, InvM1);
+		UnzipPandInvM(PandInvM2, P2, InvM2);
+		UnzipPandInvM(PandInvM3, P3, InvM3);
+		
+		const varying FVector3f X1 = VectorGather(&X[Index1]);
+		const varying FVector3f X2 = VectorGather(&X[Index2]);
+		const varying FVector3f X3 = VectorGather(&X[Index3]);
+		
+		const varying float Bary = Barys[Index];
+		const varying float Dist = Dists[Index];
+		const varying float Lambda = Lambdas[Index];
+
+		#pragma ignore warning(perf)
+		const varying float Stiffness = StiffnessTable[StiffnessIndex];
+		#pragma ignore warning(perf)
+		const varying float StiffnessWeft = StiffnessWeftTable[StiffnessWeftIndex];
+		#pragma ignore warning(perf)
+		const varying float StiffnessBias = StiffnessBiasTable[StiffnessBiasIndex];
+		#pragma ignore warning(perf)
+		const varying float DampingRatio = DampingTable[DampingIndex];
+
+		const varying FVector3f StiffnessVector = {{StiffnessWeft, Stiffness, StiffnessBias}};
+		const varying FVector3f WarpWeftBiasBaseMultiplier = VectorGather(&WarpWeftBiasBaseMultiplers[Index]);
+		const varying float StiffnessValue = VectorDot(WarpWeftBiasBaseMultiplier, StiffnessVector);
+
+		varying float DLambda;
+		const FVector3f Delta = GetXPBDAxialSpringDampingDelta(P1, P2, P3, X1, X2, X3, InvM1, InvM2, InvM3, Dt, Bary, Dist, Lambda, StiffnessValue, DampingRatio, FLOAT_SMALL_NUMBER, DLambda);
+
+		Lambdas[Index] = Lambda + DLambda;
+
+		if (InvM1 > FLOAT_ZERO)
+		{
+			VectorScatter(&PandInvM[Index1], SetVector4( P1 + (InvM1 * Delta), InvM1 ));
 		}
 		if (InvM2 > FLOAT_ZERO)
 		{
-			VectorScatter(&PandInvM[Index2], SetVector4( P2 + (InvM2 * Bary * Delta), InvM2 ));
+			VectorScatter(&PandInvM[Index2], SetVector4( P2 - (InvM2 * Bary * Delta), InvM2 ));
 		}
 		if (InvM3 > FLOAT_ZERO)
 		{
-			VectorScatter(&PandInvM[Index3], SetVector4( P3 + (InvM3 * (1.f - Bary) * Delta), InvM3 ));
+			VectorScatter(&PandInvM[Index3], SetVector4( P3 - (InvM3 * (1.f - Bary) * Delta), InvM3 ));
 		}
 	}
 }
diff --git a/Engine/Source/Runtime/Experimental/Chaos/Private/Chaos/XPBDBendingConstraints.cpp b/Engine/Source/Runtime/Experimental/Chaos/Private/Chaos/XPBDBendingConstraints.cpp
--- a/Engine/Source/Runtime/Experimental/Chaos/Private/Chaos/XPBDBendingConstraints.cpp
+++ b/Engine/Source/Runtime/Experimental/Chaos/Private/Chaos/XPBDBendingConstraints.cpp
@@ -23,6 +23,11 @@ namespace Chaos::Softs {
 int32 Chaos_XPBDBending_ParallelConstraintCount = 100;
 FAutoConsoleVariableRef CVarChaosXPBDBendingParallelConstraintCount(TEXT("p.Chaos.XPBDBending.ParallelConstraintCount"), Chaos_XPBDBending_ParallelConstraintCount, TEXT("If we have more constraints than this, use parallel-for in Apply."));
 
+#if !UE_BUILD_SHIPPING
+bool bChaos_XPBDBending_SplitLambdaDamping = true;
+FAutoConsoleVariableRef CVarChaosXPBDBendingSplitLambdaDamping(TEXT("p.Chaos.XPBDBending.SplitLambdaDamping"), bChaos_XPBDBending_SplitLambdaDamping, TEXT("Use the split two-pass damping model (slower but doesn't make cloth too soft at high damping levels)."));
+#endif
+
 template<typename SolverParticlesOrRange>
 void FXPBDBendingConstraints::InitColor(const SolverParticlesOrRange& InParticles)
 {
@@ -141,45 +146,48 @@ void FXPBDBendingConstraints::SetProperties(
 	}
 }
 
-template<typename SolverParticlesOrRange>
+template<bool bDampingOnly, bool bElasticOnly, typename SolverParticlesOrRange>
 void FXPBDBendingConstraints::ApplyHelper(SolverParticlesOrRange& Particles, const FSolverReal Dt, const int32 ConstraintIndex, const FSolverReal StiffnessValue, const FSolverReal BucklingValue, const FSolverReal DampingRatioValue) const
 {
+	check(!(bDampingOnly && bElasticOnly));
+
 	const TVec4<int32>& Constraint = Constraints[ConstraintIndex];
-	const int32 i1 = Constraint[0];
-	const int32 i2 = Constraint[1];
-	const int32 i3 = Constraint[2];
-	const int32 i4 = Constraint[3];
+	const int32 Index1 = Constraint[0];
+	const int32 Index2 = Constraint[1];
+	const int32 Index3 = Constraint[2];
+	const int32 Index4 = Constraint[3];
 
 	const FSolverReal BiphasicStiffnessValue = IsBuckled[ConstraintIndex] ? BucklingValue : StiffnessValue;
-	if (BiphasicStiffnessValue < MinStiffness || (Particles.InvM(i1) == (FSolverReal)0. && Particles.InvM(i2) == (FSolverReal)0. && Particles.InvM(i3) == (FSolverReal)0. && Particles.InvM(i4) == (FSolverReal)0.))
+	const FSolverReal CombinedInvMass = Particles.InvM(Index1) + Particles.InvM(Index2) + Particles.InvM(Index3) + Particles.InvM(Index4);
+	if (BiphasicStiffnessValue < MinStiffness || CombinedInvMass < UE_SMALL_NUMBER)
 	{
 		return;
 	}
 
-	const FSolverReal CombinedInvMass = Particles.InvM(i1) + Particles.InvM(i2) + Particles.InvM(i3) + Particles.InvM(i4);
 	const FSolverReal Damping = (FSolverReal)2.f * DampingRatioValue * FMath::Sqrt(BiphasicStiffnessValue / CombinedInvMass);
 
-	const FSolverReal Angle = CalcAngle(Particles.P(i1), Particles.P(i2), Particles.P(i3), Particles.P(i4));
+	const FSolverReal Angle = CalcAngle(Particles.P(Index1), Particles.P(Index2), Particles.P(Index3), Particles.P(Index4));
 	const TStaticArray<FSolverVec3, 4> Grads = Base::GetGradients(Particles, ConstraintIndex);
 
-	const FSolverVec3 V1TimesDt = Particles.P(i1) - Particles.X(i1);
-	const FSolverVec3 V2TimesDt = Particles.P(i2) - Particles.X(i2);
-	const FSolverVec3 V3TimesDt = Particles.P(i3) - Particles.X(i3);
-	const FSolverVec3 V4TimesDt = Particles.P(i4) - Particles.X(i4);
+	const FSolverVec3 V1TimesDt = Particles.P(Index1) - Particles.X(Index1);
+	const FSolverVec3 V2TimesDt = Particles.P(Index2) - Particles.X(Index2);
+	const FSolverVec3 V3TimesDt = Particles.P(Index3) - Particles.X(Index3);
+	const FSolverVec3 V4TimesDt = Particles.P(Index4) - Particles.X(Index4);
 
-	FSolverReal& Lambda = Lambdas[ConstraintIndex];
-	const FSolverReal Alpha = (FSolverReal)1.f / (BiphasicStiffnessValue * Dt * Dt);
-	const FSolverReal Gamma = Alpha * Damping * Dt;
+	FSolverReal& Lambda = bDampingOnly ? LambdasDamping[ConstraintIndex] : Lambdas[ConstraintIndex];
+	const FSolverReal AlphaInv = bDampingOnly ? (FSolverReal)0.f : BiphasicStiffnessValue * Dt * Dt;
+	const FSolverReal BetaDt = bElasticOnly ? (FSolverReal)0.f : Damping * Dt;
 
-	const FSolverReal DampingTerm = Gamma * (FSolverVec3::DotProduct(V1TimesDt, Grads[0]) + FSolverVec3::DotProduct(V2TimesDt, Grads[1]) + FSolverVec3::DotProduct(V3TimesDt, Grads[2]) + FSolverVec3::DotProduct(V4TimesDt, Grads[3]));
+	const FSolverReal DampingTerm = bElasticOnly ? (FSolverReal)0.f :  BetaDt * (FSolverVec3::DotProduct(V1TimesDt, Grads[0]) + FSolverVec3::DotProduct(V2TimesDt, Grads[1]) + FSolverVec3::DotProduct(V3TimesDt, Grads[2]) + FSolverVec3::DotProduct(V4TimesDt, Grads[3]));
+	const FSolverReal ElasticTerm = bDampingOnly ? (FSolverReal)0.f : AlphaInv * (Angle - RestAngles[ConstraintIndex]);
 
-	const FSolverReal Denom = ((FSolverReal)1.f + Gamma) * (Particles.InvM(i1) * Grads[0].SizeSquared() + Particles.InvM(i2) * Grads[1].SizeSquared() + Particles.InvM(i3) * Grads[2].SizeSquared() + Particles.InvM(i4) * Grads[3].SizeSquared()) + Alpha;
-	const FSolverReal DLambda = (Angle - RestAngles[ConstraintIndex] - Alpha * Lambda + DampingTerm) / Denom;
+	const FSolverReal Denom = (AlphaInv + BetaDt) * (Particles.InvM(Index1) * Grads[0].SizeSquared() + Particles.InvM(Index2) * Grads[1].SizeSquared() + Particles.InvM(Index3) * Grads[2].SizeSquared() + Particles.InvM(Index4) * Grads[3].SizeSquared()) + (FSolverReal)1.;
+	const FSolverReal DLambda = (-ElasticTerm - DampingTerm - Lambda) / Denom;
 
-	Particles.P(i1) -= DLambda * Particles.InvM(i1) * Grads[0];
-	Particles.P(i2) -= DLambda * Particles.InvM(i2) * Grads[1];
-	Particles.P(i3) -= DLambda * Particles.InvM(i3) * Grads[2];
-	Particles.P(i4) -= DLambda * Particles.InvM(i4) * Grads[3];
+	Particles.P(Index1) += DLambda * Particles.InvM(Index1) * Grads[0];
+	Particles.P(Index2) += DLambda * Particles.InvM(Index2) * Grads[1];
+	Particles.P(Index3) += DLambda * Particles.InvM(Index3) * Grads[2];
+	Particles.P(Index4) += DLambda * Particles.InvM(Index4) * Grads[3];
 	Lambda += DLambda;
 }
 
@@ -196,106 +204,192 @@ void FXPBDBendingConstraints::Apply(SolverParticlesOrRange& Particles, const FSo
 	if (ConstraintsPerColorStartIndex.Num() > 0 && Constraints.Num() > Chaos_XPBDBending_ParallelConstraintCount)
 	{
 		const int32 ConstraintColorNum = ConstraintsPerColorStartIndex.Num() - 1;
-		if (!StiffnessHasWeightMap && !BucklingStiffnessHasWeightMap && !DampingHasWeightMap)
-		{
-			const FSolverReal ExpStiffnessValue = (FSolverReal)Stiffness;
-			const FSolverReal ExpBucklingValue = (FSolverReal)BucklingStiffness;
-			const FSolverReal DampingRatioValue = (FSolverReal)DampingRatio;
-
-			if (ExpStiffnessValue < MinStiffness && ExpBucklingValue < MinStiffness)
-			{
-				return;
-			}
 
 #if INTEL_ISPC
-			if (bRealTypeCompatibleWithISPC && bChaos_XPBDBending_ISPC_Enabled)
+		if (bRealTypeCompatibleWithISPC && bChaos_XPBDBending_ISPC_Enabled)
+		{
+			if (!StiffnessHasWeightMap && !BucklingStiffnessHasWeightMap && !DampingHasWeightMap)
 			{
+				const FSolverReal ExpStiffnessValue = (FSolverReal)Stiffness;
+				const FSolverReal ExpBucklingValue = (FSolverReal)BucklingStiffness;
+				const FSolverReal DampingRatioValue = (FSolverReal)DampingRatio;
+
+				if (ExpStiffnessValue < MinStiffness && ExpBucklingValue < MinStiffness)
+				{
+					return;
+				}
 				if (DampingRatioValue > 0)
 				{
-					for (int32 ConstraintColorIndex = 0; ConstraintColorIndex < ConstraintColorNum; ++ConstraintColorIndex)
+					if (bChaos_XPBDBending_SplitLambdaDamping)
 					{
-						const int32 ColorStart = ConstraintsPerColorStartIndex[ConstraintColorIndex];
-						const int32 ColorSize = ConstraintsPerColorStartIndex[ConstraintColorIndex + 1] - ColorStart;
-						ispc::ApplyXPBDBendingConstraintsWithDamping(
-							(ispc::FVector4f*)Particles.GetPAndInvM().GetData(),
-							(const ispc::FVector3f*)Particles.XArray().GetData(),
-							(ispc::FIntVector4*)&Constraints.GetData()[ColorStart],
-							&RestAngles.GetData()[ColorStart],
-							&IsBuckled.GetData()[ColorStart],
-							&Lambdas.GetData()[ColorStart],
-							Dt,
-							MinStiffness,
-							ExpStiffnessValue,
-							ExpBucklingValue,
-							DampingRatioValue,
-							ColorSize);
+						for (int32 ConstraintColorIndex = 0; ConstraintColorIndex < ConstraintColorNum; ++ConstraintColorIndex)
+						{
+							const int32 ColorStart = ConstraintsPerColorStartIndex[ConstraintColorIndex];
+							const int32 ColorSize = ConstraintsPerColorStartIndex[ConstraintColorIndex + 1] - ColorStart;
+							ispc::ApplyXPBDBendingDampingConstraints(
+								(ispc::FVector4f*)Particles.GetPAndInvM().GetData(),
+								(const ispc::FVector3f*)Particles.XArray().GetData(),
+								(ispc::FIntVector4*)&Constraints.GetData()[ColorStart],
+								&RestAngles.GetData()[ColorStart],
+								&IsBuckled.GetData()[ColorStart],
+								&LambdasDamping.GetData()[ColorStart],
+								Dt,
+								MinStiffness,
+								ExpStiffnessValue,
+								ExpBucklingValue,
+								DampingRatioValue,
+								ColorSize);
+						}
 					}
-				}
-				else
-				{
-					for (int32 ConstraintColorIndex = 0; ConstraintColorIndex < ConstraintColorNum; ++ConstraintColorIndex)
+					else
 					{
-						const int32 ColorStart = ConstraintsPerColorStartIndex[ConstraintColorIndex];
-						const int32 ColorSize = ConstraintsPerColorStartIndex[ConstraintColorIndex + 1] - ColorStart;
-						ispc::ApplyXPBDBendingConstraints(
-							(ispc::FVector4f*)Particles.GetPAndInvM().GetData(),
-							(ispc::FIntVector4*)&Constraints.GetData()[ColorStart],
-							&RestAngles.GetData()[ColorStart],
-							&IsBuckled.GetData()[ColorStart],
-							&Lambdas.GetData()[ColorStart],
-							Dt,
-							MinStiffness,
-							ExpStiffnessValue,
-							ExpBucklingValue,
-							ColorSize);
+						for (int32 ConstraintColorIndex = 0; ConstraintColorIndex < ConstraintColorNum; ++ConstraintColorIndex)
+						{
+							const int32 ColorStart = ConstraintsPerColorStartIndex[ConstraintColorIndex];
+							const int32 ColorSize = ConstraintsPerColorStartIndex[ConstraintColorIndex + 1] - ColorStart;
+							ispc::ApplyXPBDBendingConstraintsWithDamping(
+								(ispc::FVector4f*)Particles.GetPAndInvM().GetData(),
+								(const ispc::FVector3f*)Particles.XArray().GetData(),
+								(ispc::FIntVector4*)&Constraints.GetData()[ColorStart],
+								&RestAngles.GetData()[ColorStart],
+								&IsBuckled.GetData()[ColorStart],
+								&Lambdas.GetData()[ColorStart],
+								Dt,
+								MinStiffness,
+								ExpStiffnessValue,
+								ExpBucklingValue,
+								DampingRatioValue,
+								ColorSize);
+						}
+						return;
 					}
 				}
+				for (int32 ConstraintColorIndex = 0; ConstraintColorIndex < ConstraintColorNum; ++ConstraintColorIndex)
+				{
+					const int32 ColorStart = ConstraintsPerColorStartIndex[ConstraintColorIndex];
+					const int32 ColorSize = ConstraintsPerColorStartIndex[ConstraintColorIndex + 1] - ColorStart;
+					ispc::ApplyXPBDBendingConstraints(
+						(ispc::FVector4f*)Particles.GetPAndInvM().GetData(),
+						(ispc::FIntVector4*)&Constraints.GetData()[ColorStart],
+						&RestAngles.GetData()[ColorStart],
+						&IsBuckled.GetData()[ColorStart],
+						&Lambdas.GetData()[ColorStart],
+						Dt,
+						MinStiffness,
+						ExpStiffnessValue,
+						ExpBucklingValue,
+						ColorSize);
+				}
 			}
 			else
-#endif
 			{
+				// ISPC with maps
+				if (DampingHasWeightMap || (FSolverReal)DampingRatio > 0)
+				{
+					if (bChaos_XPBDBending_SplitLambdaDamping)
+					{
+						for (int32 ConstraintColorIndex = 0; ConstraintColorIndex < ConstraintColorNum; ++ConstraintColorIndex)
+						{
+							const int32 ColorStart = ConstraintsPerColorStartIndex[ConstraintColorIndex];
+							const int32 ColorSize = ConstraintsPerColorStartIndex[ConstraintColorIndex + 1] - ColorStart;
+							ispc::ApplyXPBDBendingDampingConstraintsWithMaps(
+								(ispc::FVector4f*)Particles.GetPAndInvM().GetData(),
+								(const ispc::FVector3f*)Particles.XArray().GetData(),
+								(ispc::FIntVector4*)&Constraints.GetData()[ColorStart],
+								&RestAngles.GetData()[ColorStart],
+								&IsBuckled.GetData()[ColorStart],
+								&LambdasDamping.GetData()[ColorStart],
+								Dt,
+								MinStiffness,
+								StiffnessHasWeightMap,
+								StiffnessHasWeightMap ? &Stiffness.GetIndices().GetData()[ColorStart] : nullptr,
+								&Stiffness.GetTable().GetData()[0],
+								BucklingStiffnessHasWeightMap,
+								BucklingStiffnessHasWeightMap ? &BucklingStiffness.GetIndices().GetData()[ColorStart] : nullptr,
+								&BucklingStiffness.GetTable().GetData()[0],
+								DampingHasWeightMap,
+								DampingHasWeightMap ? &DampingRatio.GetIndices().GetData()[ColorStart] : nullptr,
+								&DampingRatio.GetTable().GetData()[0],
+								ColorSize);
+						}
+					}
+					else
+					{
+						for (int32 ConstraintColorIndex = 0; ConstraintColorIndex < ConstraintColorNum; ++ConstraintColorIndex)
+						{
+							const int32 ColorStart = ConstraintsPerColorStartIndex[ConstraintColorIndex];
+							const int32 ColorSize = ConstraintsPerColorStartIndex[ConstraintColorIndex + 1] - ColorStart;
+							ispc::ApplyXPBDBendingConstraintsWithDampingAndMaps(
+								(ispc::FVector4f*)Particles.GetPAndInvM().GetData(),
+								(const ispc::FVector3f*)Particles.XArray().GetData(),
+								(ispc::FIntVector4*)&Constraints.GetData()[ColorStart],
+								&RestAngles.GetData()[ColorStart],
+								&IsBuckled.GetData()[ColorStart],
+								&Lambdas.GetData()[ColorStart],
+								Dt,
+								MinStiffness,
+								StiffnessHasWeightMap,
+								StiffnessHasWeightMap ? &Stiffness.GetIndices().GetData()[ColorStart] : nullptr,
+								&Stiffness.GetTable().GetData()[0],
+								BucklingStiffnessHasWeightMap,
+								BucklingStiffnessHasWeightMap ? &BucklingStiffness.GetIndices().GetData()[ColorStart] : nullptr,
+								&BucklingStiffness.GetTable().GetData()[0],
+								DampingHasWeightMap,
+								DampingHasWeightMap ? &DampingRatio.GetIndices().GetData()[ColorStart] : nullptr,
+								&DampingRatio.GetTable().GetData()[0],
+								ColorSize);
+						}
+						return;
+					}
+				}
 				for (int32 ConstraintColorIndex = 0; ConstraintColorIndex < ConstraintColorNum; ++ConstraintColorIndex)
 				{
 					const int32 ColorStart = ConstraintsPerColorStartIndex[ConstraintColorIndex];
 					const int32 ColorSize = ConstraintsPerColorStartIndex[ConstraintColorIndex + 1] - ColorStart;
-					PhysicsParallelFor(ColorSize, [&](const int32 Index)
-					{
-						const int32 ConstraintIndex = ColorStart + Index;
-						ApplyHelper(Particles, Dt, ConstraintIndex, ExpStiffnessValue, ExpBucklingValue, DampingRatioValue);
-					});
+					ispc::ApplyXPBDBendingConstraintsWithMaps(
+						(ispc::FVector4f*)Particles.GetPAndInvM().GetData(),
+						(ispc::FIntVector4*)&Constraints.GetData()[ColorStart],
+						&RestAngles.GetData()[ColorStart],
+						&IsBuckled.GetData()[ColorStart],
+						&Lambdas.GetData()[ColorStart],
+						Dt,
+						MinStiffness,
+						StiffnessHasWeightMap,
+						StiffnessHasWeightMap ? &Stiffness.GetIndices().GetData()[ColorStart] : nullptr,
+						&Stiffness.GetTable().GetData()[0],
+						BucklingStiffnessHasWeightMap,
+						BucklingStiffnessHasWeightMap ? &BucklingStiffness.GetIndices().GetData()[ColorStart] : nullptr,
+						&BucklingStiffness.GetTable().GetData()[0],
+						ColorSize);
 				}
 			}
 		}
-		else  // Has weight maps
+		else
+#endif
 		{
-#if INTEL_ISPC
-			if (bRealTypeCompatibleWithISPC && bChaos_XPBDBending_ISPC_Enabled)
+			// Parallel non-ispc
+			const FSolverReal StiffnessNoMap = (FSolverReal)Stiffness;
+			const FSolverReal BucklingStiffnessNoMap = (FSolverReal)BucklingStiffness;
+			const FSolverReal DampingNoMap = (FSolverReal)DampingRatio;
+
+			if (DampingHasWeightMap || (FSolverReal)DampingRatio > 0)
 			{
-				if (DampingHasWeightMap || (FSolverReal)DampingRatio > 0)
+				if (bChaos_XPBDBending_SplitLambdaDamping)
 				{
 					for (int32 ConstraintColorIndex = 0; ConstraintColorIndex < ConstraintColorNum; ++ConstraintColorIndex)
 					{
 						const int32 ColorStart = ConstraintsPerColorStartIndex[ConstraintColorIndex];
 						const int32 ColorSize = ConstraintsPerColorStartIndex[ConstraintColorIndex + 1] - ColorStart;
-						ispc::ApplyXPBDBendingConstraintsWithDampingAndMaps(
-							(ispc::FVector4f*)Particles.GetPAndInvM().GetData(),
-							(const ispc::FVector3f*)Particles.XArray().GetData(),
-							(ispc::FIntVector4*)&Constraints.GetData()[ColorStart],
-							&RestAngles.GetData()[ColorStart],
-							&IsBuckled.GetData()[ColorStart],
-							&Lambdas.GetData()[ColorStart],
-							Dt,
-							MinStiffness,
-							StiffnessHasWeightMap,
-							StiffnessHasWeightMap ? &Stiffness.GetIndices().GetData()[ColorStart] : nullptr,
-							&Stiffness.GetTable().GetData()[0],
-							BucklingStiffnessHasWeightMap,
-							BucklingStiffnessHasWeightMap ? &BucklingStiffness.GetIndices().GetData()[ColorStart] : nullptr,
-							&BucklingStiffness.GetTable().GetData()[0],
-							DampingHasWeightMap,
-							DampingHasWeightMap ? &DampingRatio.GetIndices().GetData()[ColorStart] : nullptr,
-							&DampingRatio.GetTable().GetData()[0],
-							ColorSize);
+						PhysicsParallelFor(ColorSize, [&](const int32 Index)
+						{
+							const int32 ConstraintIndex = ColorStart + Index;
+							const FSolverReal ExpStiffnessValue = StiffnessHasWeightMap ? Stiffness[ConstraintIndex] : StiffnessNoMap;
+							const FSolverReal ExpBucklingValue = BucklingStiffnessHasWeightMap ? BucklingStiffness[ConstraintIndex] : BucklingStiffnessNoMap;
+							const FSolverReal DampingRatioValue = DampingHasWeightMap ? DampingRatio[ConstraintIndex] : DampingNoMap;
+							constexpr bool bDampingOnly = true;
+							constexpr bool bElasticOnly = false;
+							ApplyHelper<bDampingOnly, bElasticOnly>(Particles, Dt, ConstraintIndex, ExpStiffnessValue, ExpBucklingValue, DampingRatioValue);
+						});
 					}
 				}
 				else
@@ -304,74 +398,80 @@ void FXPBDBendingConstraints::Apply(SolverParticlesOrRange& Particles, const FSo
 					{
 						const int32 ColorStart = ConstraintsPerColorStartIndex[ConstraintColorIndex];
 						const int32 ColorSize = ConstraintsPerColorStartIndex[ConstraintColorIndex + 1] - ColorStart;
-						ispc::ApplyXPBDBendingConstraintsWithMaps(
-							(ispc::FVector4f*)Particles.GetPAndInvM().GetData(),
-							(ispc::FIntVector4*)&Constraints.GetData()[ColorStart],
-							&RestAngles.GetData()[ColorStart],
-							&IsBuckled.GetData()[ColorStart],
-							&Lambdas.GetData()[ColorStart],
-							Dt,
-							MinStiffness,
-							StiffnessHasWeightMap,
-							StiffnessHasWeightMap ? &Stiffness.GetIndices().GetData()[ColorStart] : nullptr,
-							&Stiffness.GetTable().GetData()[0],
-							BucklingStiffnessHasWeightMap,
-							BucklingStiffnessHasWeightMap ? &BucklingStiffness.GetIndices().GetData()[ColorStart] : nullptr,
-							&BucklingStiffness.GetTable().GetData()[0],
-							ColorSize);
+						PhysicsParallelFor(ColorSize, [&](const int32 Index)
+						{
+							const int32 ConstraintIndex = ColorStart + Index;
+							const FSolverReal ExpStiffnessValue = StiffnessHasWeightMap ? Stiffness[ConstraintIndex] : StiffnessNoMap;
+							const FSolverReal ExpBucklingValue = BucklingStiffnessHasWeightMap ? BucklingStiffness[ConstraintIndex] : BucklingStiffnessNoMap;
+							const FSolverReal DampingRatioValue = DampingHasWeightMap ? DampingRatio[ConstraintIndex] : DampingNoMap;
+							constexpr bool bDampingOnly = false;
+							constexpr bool bElasticOnly = false;
+							ApplyHelper<bDampingOnly, bElasticOnly>(Particles, Dt, ConstraintIndex, ExpStiffnessValue, ExpBucklingValue, DampingRatioValue);
+						});
 					}
+					return;
 				}
 			}
-			else
-#endif
+			for (int32 ConstraintColorIndex = 0; ConstraintColorIndex < ConstraintColorNum; ++ConstraintColorIndex)
 			{
-				const FSolverReal StiffnessNoMap = (FSolverReal)Stiffness;
-				const FSolverReal BucklingStiffnessNoMap = (FSolverReal)BucklingStiffness;
-				const FSolverReal DampingNoMap = (FSolverReal)DampingRatio;
-				for (int32 ConstraintColorIndex = 0; ConstraintColorIndex < ConstraintColorNum; ++ConstraintColorIndex)
+				const int32 ColorStart = ConstraintsPerColorStartIndex[ConstraintColorIndex];
+				const int32 ColorSize = ConstraintsPerColorStartIndex[ConstraintColorIndex + 1] - ColorStart;
+				PhysicsParallelFor(ColorSize, [&](const int32 Index)
 				{
-					const int32 ColorStart = ConstraintsPerColorStartIndex[ConstraintColorIndex];
-					const int32 ColorSize = ConstraintsPerColorStartIndex[ConstraintColorIndex + 1] - ColorStart;
-					PhysicsParallelFor(ColorSize, [&](const int32 Index)
-					{
-						const int32 ConstraintIndex = ColorStart + Index;
-						const FSolverReal ExpStiffnessValue = StiffnessHasWeightMap ? Stiffness[ConstraintIndex] : StiffnessNoMap;
-						const FSolverReal ExpBucklingValue = BucklingStiffnessHasWeightMap ? BucklingStiffness[ConstraintIndex] : BucklingStiffnessNoMap;
-						const FSolverReal DampingRatioValue = DampingHasWeightMap ? DampingRatio[ConstraintIndex] : DampingNoMap;
-						ApplyHelper(Particles, Dt, ConstraintIndex, ExpStiffnessValue, ExpBucklingValue, DampingRatioValue);
-					});
-				}
+					const int32 ConstraintIndex = ColorStart + Index;
+					const FSolverReal ExpStiffnessValue = StiffnessHasWeightMap ? Stiffness[ConstraintIndex] : StiffnessNoMap;
+					const FSolverReal ExpBucklingValue = BucklingStiffnessHasWeightMap ? BucklingStiffness[ConstraintIndex] : BucklingStiffnessNoMap;
+					const FSolverReal DampingRatioValue = DampingHasWeightMap ? DampingRatio[ConstraintIndex] : DampingNoMap;
+					constexpr bool bDampingOnly = false;
+					constexpr bool bElasticOnly = true;
+					ApplyHelper<bDampingOnly, bElasticOnly>(Particles, Dt, ConstraintIndex, ExpStiffnessValue, ExpBucklingValue, DampingRatioValue);
+				});
 			}
 		}
 	}
 	else
 	{
-		if (!StiffnessHasWeightMap && !BucklingStiffnessHasWeightMap)
+		// Single-threaded
+		const FSolverReal StiffnessNoMap = (FSolverReal)Stiffness;
+		const FSolverReal BucklingStiffnessNoMap = (FSolverReal)BucklingStiffness;
+		const FSolverReal DampingNoMap = (FSolverReal)DampingRatio;
+
+		if (DampingHasWeightMap || (FSolverReal)DampingRatio > 0)
 		{
-			const FSolverReal ExpStiffnessValue = (FSolverReal)Stiffness;
-			const FSolverReal ExpBucklingValue = (FSolverReal)BucklingStiffness;
-			const FSolverReal DampingRatioValue = (FSolverReal)DampingRatio;
-			if (ExpStiffnessValue < MinStiffness && ExpBucklingValue < MinStiffness)
+			if (bChaos_XPBDBending_SplitLambdaDamping)
 			{
-				return;
+				for (int32 ConstraintIndex = 0; ConstraintIndex < Constraints.Num(); ++ConstraintIndex)
+				{
+					const FSolverReal ExpStiffnessValue = StiffnessHasWeightMap ? Stiffness[ConstraintIndex] : StiffnessNoMap;
+					const FSolverReal ExpBucklingValue = BucklingStiffnessHasWeightMap ? BucklingStiffness[ConstraintIndex] : BucklingStiffnessNoMap;
+					const FSolverReal DampingRatioValue = DampingHasWeightMap ? DampingRatio[ConstraintIndex] : DampingNoMap;
+					constexpr bool bDampingOnly = true;
+					constexpr bool bElasticOnly = false;
+					ApplyHelper<bDampingOnly, bElasticOnly>(Particles, Dt, ConstraintIndex, ExpStiffnessValue, ExpBucklingValue, DampingRatioValue);
+				}
 			}
-			for (int32 ConstraintIndex = 0; ConstraintIndex < Constraints.Num(); ++ConstraintIndex)
+			else
 			{
-				ApplyHelper(Particles, Dt, ConstraintIndex, ExpStiffnessValue, ExpBucklingValue, DampingRatioValue);
+				for (int32 ConstraintIndex = 0; ConstraintIndex < Constraints.Num(); ++ConstraintIndex)
+				{
+					const FSolverReal ExpStiffnessValue = StiffnessHasWeightMap ? Stiffness[ConstraintIndex] : StiffnessNoMap;
+					const FSolverReal ExpBucklingValue = BucklingStiffnessHasWeightMap ? BucklingStiffness[ConstraintIndex] : BucklingStiffnessNoMap;
+					const FSolverReal DampingRatioValue = DampingHasWeightMap ? DampingRatio[ConstraintIndex] : DampingNoMap;
+					constexpr bool bDampingOnly = false;
+					constexpr bool bElasticOnly = false;
+					ApplyHelper<bDampingOnly, bElasticOnly>(Particles, Dt, ConstraintIndex, ExpStiffnessValue, ExpBucklingValue, DampingRatioValue);
+				}
+				return;
 			}
 		}
-		else
+		for (int32 ConstraintIndex = 0; ConstraintIndex < Constraints.Num(); ++ConstraintIndex)
 		{
-			const FSolverReal StiffnessNoMap = (FSolverReal)Stiffness;
-			const FSolverReal BucklingStiffnessNoMap = (FSolverReal)BucklingStiffness;
-			const FSolverReal DampingNoMap = (FSolverReal)DampingRatio;
-			for (int32 ConstraintIndex = 0; ConstraintIndex < Constraints.Num(); ++ConstraintIndex)
-			{
-				const FSolverReal ExpStiffnessValue = StiffnessHasWeightMap ? Stiffness[ConstraintIndex] : StiffnessNoMap;
-				const FSolverReal ExpBucklingValue = BucklingStiffnessHasWeightMap ? BucklingStiffness[ConstraintIndex] : BucklingStiffnessNoMap;
-				const FSolverReal DampingRatioValue = DampingHasWeightMap ? DampingRatio[ConstraintIndex] : DampingNoMap;
-				ApplyHelper(Particles, Dt, ConstraintIndex, ExpStiffnessValue, ExpBucklingValue, DampingRatioValue);
-			}
+			const FSolverReal ExpStiffnessValue = StiffnessHasWeightMap ? Stiffness[ConstraintIndex] : StiffnessNoMap;
+			const FSolverReal ExpBucklingValue = BucklingStiffnessHasWeightMap ? BucklingStiffness[ConstraintIndex] : BucklingStiffnessNoMap;
+			const FSolverReal DampingRatioValue = DampingHasWeightMap ? DampingRatio[ConstraintIndex] : DampingNoMap;
+			constexpr bool bDampingOnly = false;
+			constexpr bool bElasticOnly = true;
+			ApplyHelper<bDampingOnly, bElasticOnly>(Particles, Dt, ConstraintIndex, ExpStiffnessValue, ExpBucklingValue, DampingRatioValue);
 		}
 	}
 }
diff --git a/Engine/Source/Runtime/Experimental/Chaos/Private/Chaos/XPBDBendingConstraints.ispc b/Engine/Source/Runtime/Experimental/Chaos/Private/Chaos/XPBDBendingConstraints.ispc
--- a/Engine/Source/Runtime/Experimental/Chaos/Private/Chaos/XPBDBendingConstraints.ispc
+++ b/Engine/Source/Runtime/Experimental/Chaos/Private/Chaos/XPBDBendingConstraints.ispc
@@ -4,51 +4,72 @@
 
 #include "Math/Vector.isph"
 #include "Chaos/PBDSofts.isph"
+#include "Chaos/XPBDBendingElement.isph"
 
-static inline FVector3f SafeDivide(const FVector3f& Numerator, const float Denominator)
+export void ApplyXPBDBendingConstraints(uniform FVector4f PandInvM[],
+									const uniform FIntVector4 Constraints[],
+									const uniform float RestAngles[],
+									const uniform bool IsBuckled[],
+									uniform float Lambdas[],
+									const uniform float Dt,
+									const uniform float MinStiffness,
+									const uniform float Stiffness,
+									const uniform float BucklingStiffness,
+									const uniform int32 NumConstraints)
 {
-	return VectorSelect(Denominator <= FLOAT_SMALL_NUMBER, FloatZeroVector, Numerator / Denominator);
-}
 
-static inline varying float SafeDivide(const varying float Numerator, const varying float Denominator)
-{
-	return select(Denominator <= FLOAT_SMALL_NUMBER, (uniform float)0, Numerator / Denominator);
-}
+	foreach(i = 0 ... NumConstraints)
+	{
+		const varying FIntVector4 Constraint = VectorLoad(&Constraints[extract(i,0)]);
 
-// Returns Angle
-static inline float CalcGradientsAndAngle(const FVector3f& P1, const FVector3f& P2, const FVector3f& P3, const FVector3f& P4, FVector3f& Grad1, FVector3f& Grad2, FVector3f& Grad3, FVector3f& Grad4)
-{
-	const FVector3f SharedEdgeNormalized = VectorGetSafeNormal(P2 - P1);
-	const FVector3f P13CrossP23 = VectorCross(P1 - P3, P2 - P3);
-	const float Normal1Len = VectorSize(P13CrossP23);
-	const FVector3f Normal1 = SafeDivide(P13CrossP23, Normal1Len);
-	const FVector3f P24CrossP14 = VectorCross(P2 - P4, P1 - P4);
-	const float Normal2Len = VectorSize(P24CrossP14);
-	const FVector3f Normal2 = SafeDivide(P24CrossP14, Normal2Len);
+		const varying int32 i1 = Constraint.V[0];
+		const varying int32 i2 = Constraint.V[1];
+		const varying int32 i3 = Constraint.V[2];
+		const varying int32 i4 = Constraint.V[3];
+		
+		const varying FVector4f PandInvM1 = VectorGather(&PandInvM[i1]);
+		const varying FVector4f PandInvM2 = VectorGather(&PandInvM[i2]);
+		const varying FVector4f PandInvM3 = VectorGather(&PandInvM[i3]);
+		const varying FVector4f PandInvM4 = VectorGather(&PandInvM[i4]);
 
-	const FVector3f N2CrossN1 = VectorCross(Normal2, Normal1);
+		varying FVector3f P1, P2, P3, P4;
+		varying float M1, M2, M3, M4;
 
-	const float CosPhi = clamp(VectorDot(Normal1, Normal2), (uniform float)(-1), (uniform float)(1));
-	const float SinPhi = clamp(VectorDot(N2CrossN1, SharedEdgeNormalized), (uniform float)(-1), (uniform float)(1));
+		UnzipPandInvM(PandInvM1, P1, M1);
+		UnzipPandInvM(PandInvM2, P2, M2);
+		UnzipPandInvM(PandInvM3, P3, M3);
+		UnzipPandInvM(PandInvM4, P4, M4);
 
-	const FVector3f DPhiDN1_OverNormal1Len = SafeDivide(CosPhi * VectorCross(SharedEdgeNormalized, Normal2) - SinPhi * Normal2, Normal1Len);
-	const FVector3f DPhiDN2_OverNormal2Len = SafeDivide(CosPhi * VectorCross(Normal1, SharedEdgeNormalized) - SinPhi * Normal1, Normal2Len);
+		const varying float StiffnessValue = select(IsBuckled[i], BucklingStiffness, Stiffness);
+		
+		const varying float Lambda = Lambdas[i];
 
-	const FVector3f DPhiDP13 = VectorCross(P2 - P3, DPhiDN1_OverNormal1Len);
-	const FVector3f DPhiDP23 = VectorCross(DPhiDN1_OverNormal1Len, P1 - P3);
-	const FVector3f DPhiDP24 = VectorCross(P1 - P4, DPhiDN2_OverNormal2Len);
-	const FVector3f DPhiDP14 = VectorCross(DPhiDN2_OverNormal2Len, P2 - P4);
+		varying FVector3f Delta1, Delta2, Delta3, Delta4;
+		const varying float DLambda = GetXPBDBendingElementDeltas(P1, P2, P3, P4, M1, M2, M3, M4, Dt, RestAngles[i], Lambda, StiffnessValue, MinStiffness, Delta1, Delta2, Delta3, Delta4);
 
-	Grad1 = DPhiDP13 + DPhiDP14;
-	Grad2 = DPhiDP23 + DPhiDP24;
-	Grad3 = -1 * DPhiDP13 - DPhiDP23;
-	Grad4 = -1 * DPhiDP14 - DPhiDP24;
+		Lambdas[i] = Lambda + DLambda;
 
-	return atan2(SinPhi, CosPhi);
+		if (M1 > 0)
+		{
+			VectorScatter(&PandInvM[i1], SetVector4(P1 + Delta1, M1));
+		}
+		if (M2 > 0)
+		{
+			VectorScatter(&PandInvM[i2], SetVector4(P2 + Delta2, M2));
+		}
+		if (M3 > 0)
+		{
+			VectorScatter(&PandInvM[i3], SetVector4(P3 + Delta3, M3));
+		}
+		if (M4 > 0)
+		{
+			VectorScatter(&PandInvM[i4], SetVector4(P4 + Delta4, M4));
+		}
+	}
 }
 
-
-export void ApplyXPBDBendingConstraints(uniform FVector4f PandInvM[],
+export void ApplyXPBDBendingConstraintsWithDamping(uniform FVector4f PandInvM[],
+									const uniform FVector3f X[],
 									const uniform FIntVector4 Constraints[],
 									const uniform float RestAngles[],
 									const uniform bool IsBuckled[],
@@ -57,6 +78,7 @@ export void ApplyXPBDBendingConstraints(uniform FVector4f PandInvM[],
 									const uniform float MinStiffness,
 									const uniform float Stiffness,
 									const uniform float BucklingStiffness,
+									const uniform float DampingRatio,
 									const uniform int32 NumConstraints)
 {
 
@@ -82,40 +104,39 @@ export void ApplyXPBDBendingConstraints(uniform FVector4f PandInvM[],
 		UnzipPandInvM(PandInvM3, P3, M3);
 		UnzipPandInvM(PandInvM4, P4, M4);
 
-		varying FVector3f Grad1, Grad2, Grad3, Grad4;
-		const varying float Angle = CalcGradientsAndAngle(P1, P2, P3, P4, Grad1, Grad2, Grad3, Grad4);
+		const varying FVector3f X1 = VectorGather(&X[i1]);
+		const varying FVector3f X2 = VectorGather(&X[i2]);
+		const varying FVector3f X3 = VectorGather(&X[i3]);
+		const varying FVector3f X4 = VectorGather(&X[i4]);
 		
-		const varying float RestAngle = RestAngles[i];
 		const varying float StiffnessValue = select(IsBuckled[i], BucklingStiffness, Stiffness);
-		
 		const varying float Lambda = Lambdas[i];
-		const varying float Alpha  = 1.f / (StiffnessValue * Dt * Dt);
-		const varying float Denom = M1 * VectorSizeSquared(Grad1) + M2 * VectorSizeSquared(Grad2) + M3 * VectorSizeSquared(Grad3) + M4 * VectorSizeSquared(Grad4) + Alpha;
 
-		const varying float ComputedDLambda = (Angle - RestAngle - Alpha * Lambda) / Denom;
-		const varying float DLambda = select(StiffnessValue < MinStiffness, (uniform float)0, ComputedDLambda);
+		varying FVector3f Delta1, Delta2, Delta3, Delta4;
+		const varying float DLambda = GetXPBDBendingElementDeltasWithDamping(P1, P2, P3, P4, X1, X2, X3, X4, M1, M2, M3, M4, Dt, RestAngles[i], Lambda, StiffnessValue, DampingRatio, MinStiffness, Delta1, Delta2, Delta3, Delta4);
+
 		Lambdas[i] = Lambda + DLambda;
 
 		if (M1 > 0)
 		{
-			VectorScatter(&PandInvM[i1], SetVector4(P1 - DLambda * M1 * Grad1, M1));
+			VectorScatter(&PandInvM[i1], SetVector4(P1 + Delta1, M1));
 		}
 		if (M2 > 0)
 		{
-			VectorScatter(&PandInvM[i2], SetVector4(P2 - DLambda * M2 * Grad2, M2));
+			VectorScatter(&PandInvM[i2], SetVector4(P2 + Delta2, M2));
 		}
 		if (M3 > 0)
 		{
-			VectorScatter(&PandInvM[i3], SetVector4(P3 - DLambda * M3 * Grad3, M3));
+			VectorScatter(&PandInvM[i3], SetVector4(P3 + Delta3, M3));
 		}
 		if (M4 > 0)
 		{
-			VectorScatter(&PandInvM[i4], SetVector4(P4 - DLambda * M4 * Grad4, M4));
+			VectorScatter(&PandInvM[i4], SetVector4(P4 + Delta4, M4));
 		}
 	}
 }
 
-export void ApplyXPBDBendingConstraintsWithDamping(uniform FVector4f PandInvM[],
+export void ApplyXPBDBendingDampingConstraints(uniform FVector4f PandInvM[],
 									const uniform FVector3f X[],
 									const uniform FIntVector4 Constraints[],
 									const uniform float RestAngles[],
@@ -151,51 +172,34 @@ export void ApplyXPBDBendingConstraintsWithDamping(uniform FVector4f PandInvM[],
 		UnzipPandInvM(PandInvM3, P3, M3);
 		UnzipPandInvM(PandInvM4, P4, M4);
 
-		varying FVector3f Grad1, Grad2, Grad3, Grad4;
-		const varying float Angle = CalcGradientsAndAngle(P1, P2, P3, P4, Grad1, Grad2, Grad3, Grad4);
-		
-		const varying float RestAngle = RestAngles[i];
-		const varying float StiffnessValue = select(IsBuckled[i], BucklingStiffness, Stiffness);
-
 		const varying FVector3f X1 = VectorGather(&X[i1]);
 		const varying FVector3f X2 = VectorGather(&X[i2]);
 		const varying FVector3f X3 = VectorGather(&X[i3]);
 		const varying FVector3f X4 = VectorGather(&X[i4]);
-
-		const varying float CombinedInvMass = M1 + M2 + M3 + M4;
-		const varying float Damping = 2.f * DampingRatio * sqrt(StiffnessValue / CombinedInvMass);
-
+		
+		const varying float StiffnessValue = select(IsBuckled[i], BucklingStiffness, Stiffness);
 		const varying float Lambda = Lambdas[i];
-		const varying float Alpha  = 1.f / (StiffnessValue * Dt * Dt);
-		const varying float Gamma = Alpha * Damping * Dt;
 
-		const varying FVector3f V1TimesDt = P1 - X1;
-		const varying FVector3f V2TimesDt = P2 - X2;
-		const varying FVector3f V3TimesDt = P3 - X3;
-		const varying FVector3f V4TimesDt = P4 - X4;
-		const varying float DampingTerm = Gamma * (VectorDot(V1TimesDt, Grad1) + VectorDot(V2TimesDt, Grad2) + VectorDot(V3TimesDt, Grad3) + VectorDot(V4TimesDt, Grad4));
+		varying FVector3f Delta1, Delta2, Delta3, Delta4;
+		const varying float DLambda = GetXPBDBendingElementDampingDeltas(P1, P2, P3, P4, X1, X2, X3, X4, M1, M2, M3, M4, Dt, RestAngles[i], Lambda, StiffnessValue, DampingRatio, MinStiffness, Delta1, Delta2, Delta3, Delta4);
 
-		const varying float Denom = (1.f + Gamma)*(M1 * VectorSizeSquared(Grad1) + M2 * VectorSizeSquared(Grad2) + M3 * VectorSizeSquared(Grad3) + M4 * VectorSizeSquared(Grad4)) + Alpha;
-
-		const varying float ComputedDLambda = (Angle - RestAngle - Alpha * Lambda + DampingTerm) / Denom;
-		const varying float DLambda = select(StiffnessValue < MinStiffness, (uniform float)0, ComputedDLambda);
 		Lambdas[i] = Lambda + DLambda;
 
 		if (M1 > 0)
 		{
-			VectorScatter(&PandInvM[i1], SetVector4(P1 - DLambda * M1 * Grad1, M1));
+			VectorScatter(&PandInvM[i1], SetVector4(P1 + Delta1, M1));
 		}
 		if (M2 > 0)
 		{
-			VectorScatter(&PandInvM[i2], SetVector4(P2 - DLambda * M2 * Grad2, M2));
+			VectorScatter(&PandInvM[i2], SetVector4(P2 + Delta2, M2));
 		}
 		if (M3 > 0)
 		{
-			VectorScatter(&PandInvM[i3], SetVector4(P3 - DLambda * M3 * Grad3, M3));
+			VectorScatter(&PandInvM[i3], SetVector4(P3 + Delta3, M3));
 		}
 		if (M4 > 0)
 		{
-			VectorScatter(&PandInvM[i4], SetVector4(P4 - DLambda * M4 * Grad4, M4));
+			VectorScatter(&PandInvM[i4], SetVector4(P4 + Delta4, M4));
 		}
 	}
 }
@@ -239,42 +243,36 @@ export void ApplyXPBDBendingConstraintsWithMaps(uniform FVector4f PandInvM[],
 		UnzipPandInvM(PandInvM2, P2, M2);
 		UnzipPandInvM(PandInvM3, P3, M3);
 		UnzipPandInvM(PandInvM4, P4, M4);
-
-
-		varying FVector3f Grad1, Grad2, Grad3, Grad4;
-		const varying float Angle = CalcGradientsAndAngle(P1, P2, P3, P4, Grad1, Grad2, Grad3, Grad4);
 		
 		#pragma ignore warning(perf)
 		const varying float Stiffness = StiffnessTable[StiffnessIndex];
 		#pragma ignore warning(perf)
 		const varying float BucklingStiffness = BucklingStiffnessTable[BucklingStiffnessIndex];
 		
-		const varying float RestAngle = RestAngles[i];
 		const varying float StiffnessValue = select(IsBuckled[i], BucklingStiffness, Stiffness);
 		
 		const varying float Lambda = Lambdas[i];
-		const varying float Alpha  = 1.f / (StiffnessValue * Dt * Dt);
-		const varying float Denom = M1 * VectorSizeSquared(Grad1) + M2 * VectorSizeSquared(Grad2) + M3 * VectorSizeSquared(Grad3) + M4 * VectorSizeSquared(Grad4) + Alpha;
 
-		const varying float ComputedDLambda = (Angle - RestAngle - Alpha * Lambda) / Denom;
-		const varying float DLambda = select( StiffnessValue < MinStiffness, (uniform float)0, ComputedDLambda);
+		varying FVector3f Delta1, Delta2, Delta3, Delta4;
+		const varying float DLambda = GetXPBDBendingElementDeltas(P1, P2, P3, P4, M1, M2, M3, M4, Dt, RestAngles[i], Lambda, StiffnessValue, MinStiffness, Delta1, Delta2, Delta3, Delta4);
+
 		Lambdas[i] = Lambda + DLambda;
 
 		if (M1 > 0)
 		{
-			VectorScatter(&PandInvM[i1], SetVector4(P1 - DLambda * M1 * Grad1, M1));
+			VectorScatter(&PandInvM[i1], SetVector4(P1 + Delta1, M1));
 		}
 		if (M2 > 0)
 		{
-			VectorScatter(&PandInvM[i2], SetVector4(P2 - DLambda * M2 * Grad2, M2));
+			VectorScatter(&PandInvM[i2], SetVector4(P2 + Delta2, M2));
 		}
 		if (M3 > 0)
 		{
-			VectorScatter(&PandInvM[i3], SetVector4(P3 - DLambda * M3 * Grad3, M3));
+			VectorScatter(&PandInvM[i3], SetVector4(P3 + Delta3, M3));
 		}
 		if (M4 > 0)
 		{
-			VectorScatter(&PandInvM[i4], SetVector4(P4 - DLambda * M4 * Grad4, M4));
+			VectorScatter(&PandInvM[i4], SetVector4(P4 + Delta4, M4));
 		}
 	}
 }
@@ -325,9 +323,6 @@ export void ApplyXPBDBendingConstraintsWithDampingAndMaps(uniform FVector4f Pand
 		UnzipPandInvM(PandInvM4, P4, M4);
 
 
-		varying FVector3f Grad1, Grad2, Grad3, Grad4;
-		const varying float Angle = CalcGradientsAndAngle(P1, P2, P3, P4, Grad1, Grad2, Grad3, Grad4);
-		
 		#pragma ignore warning(perf)
 		const varying float Stiffness = StiffnessTable[StiffnessIndex];
 		#pragma ignore warning(perf)
@@ -335,48 +330,121 @@ export void ApplyXPBDBendingConstraintsWithDampingAndMaps(uniform FVector4f Pand
 		#pragma ignore warning(perf)
 		const varying float DampingRatio = DampingTable[DampingIndex];
 		
-		const varying float RestAngle = RestAngles[i];
 		const varying float StiffnessValue = select(IsBuckled[i], BucklingStiffness, Stiffness);
 
 		const varying FVector3f X1 = VectorGather(&X[i1]);
 		const varying FVector3f X2 = VectorGather(&X[i2]);
 		const varying FVector3f X3 = VectorGather(&X[i3]);
 		const varying FVector3f X4 = VectorGather(&X[i4]);
-
-		const varying float CombinedInvMass = M1 + M2 + M3 + M4;
-		const varying float Damping = 2.f * DampingRatio * sqrt(StiffnessValue / CombinedInvMass);
 		
 		const varying float Lambda = Lambdas[i];
-		const varying float Alpha  = 1.f / (StiffnessValue * Dt * Dt);
-		const varying float Gamma = Alpha * Damping * Dt;
 
-		const varying FVector3f V1TimesDt = P1 - X1;
-		const varying FVector3f V2TimesDt = P2 - X2;
-		const varying FVector3f V3TimesDt = P3 - X3;
-		const varying FVector3f V4TimesDt = P4 - X4;
-		const varying float DampingTerm = Gamma * (VectorDot(V1TimesDt, Grad1) + VectorDot(V2TimesDt, Grad2) + VectorDot(V3TimesDt, Grad3) + VectorDot(V4TimesDt, Grad4));
+		varying FVector3f Delta1, Delta2, Delta3, Delta4;
+		const varying float DLambda = GetXPBDBendingElementDeltasWithDamping(P1, P2, P3, P4, X1, X2, X3, X4, M1, M2, M3, M4, Dt, RestAngles[i], Lambda, StiffnessValue, DampingRatio, MinStiffness, Delta1, Delta2, Delta3, Delta4);
 
-		const varying float Denom = (1.f + Gamma)*(M1 * VectorSizeSquared(Grad1) + M2 * VectorSizeSquared(Grad2) + M3 * VectorSizeSquared(Grad3) + M4 * VectorSizeSquared(Grad4)) + Alpha;
+		Lambdas[i] = Lambda + DLambda;
+
+		if (M1 > 0)
+		{
+			VectorScatter(&PandInvM[i1], SetVector4(P1 + Delta1, M1));
+		}
+		if (M2 > 0)
+		{
+			VectorScatter(&PandInvM[i2], SetVector4(P2 + Delta2, M2));
+		}
+		if (M3 > 0)
+		{
+			VectorScatter(&PandInvM[i3], SetVector4(P3 + Delta3, M3));
+		}
+		if (M4 > 0)
+		{
+			VectorScatter(&PandInvM[i4], SetVector4(P4 + Delta4, M4));
+		}
+	}
+}
+
+export void ApplyXPBDBendingDampingConstraintsWithMaps(uniform FVector4f PandInvM[],
+											const uniform FVector3f X[],
+											const uniform FIntVector4 Constraints[],
+											const uniform float RestAngles[],
+											const uniform bool IsBuckled[],
+											uniform float Lambdas[],
+											const uniform float Dt,
+											const uniform float MinStiffness,
+											const uniform bool StiffnessHasMap,
+											const uniform uint8 StiffnessIndices[],
+											const uniform float StiffnessTable[],
+											const uniform bool BucklingStiffnessHasMap,
+											const uniform uint8 BucklingStiffnessIndices[],
+											const uniform float BucklingStiffnessTable[],
+											const uniform bool DampingHasMap,
+											const uniform uint8 DampingIndices[],
+											const uniform float DampingTable[],
+											const uniform int32 NumConstraints)
+{
+
+	foreach(i = 0 ... NumConstraints)
+	{
+		const varying FIntVector4 Constraint = VectorLoad(&Constraints[extract(i,0)]);
+
+		const varying int32 i1 = Constraint.V[0];
+		const varying int32 i2 = Constraint.V[1];
+		const varying int32 i3 = Constraint.V[2];
+		const varying int32 i4 = Constraint.V[3];
+		const varying int8 StiffnessIndex = StiffnessHasMap ? StiffnessIndices[i] : 0;
+		const varying int8 BucklingStiffnessIndex = BucklingStiffnessHasMap ? BucklingStiffnessIndices[i] : 0;
+		const varying int8 DampingIndex = DampingHasMap ? DampingIndices[i] : 0;
+
+		const varying FVector4f PandInvM2 = VectorGather(&PandInvM[i2]);
+		const varying FVector4f PandInvM3 = VectorGather(&PandInvM[i3]);
+		const varying FVector4f PandInvM1 = VectorGather(&PandInvM[i1]);
+		const varying FVector4f PandInvM4 = VectorGather(&PandInvM[i4]);
+
+		varying FVector3f P1, P2, P3, P4;
+		varying float M1, M2, M3, M4;
+
+		UnzipPandInvM(PandInvM1, P1, M1);
+		UnzipPandInvM(PandInvM2, P2, M2);
+		UnzipPandInvM(PandInvM3, P3, M3);
+		UnzipPandInvM(PandInvM4, P4, M4);
+
+
+		#pragma ignore warning(perf)
+		const varying float Stiffness = StiffnessTable[StiffnessIndex];
+		#pragma ignore warning(perf)
+		const varying float BucklingStiffness = BucklingStiffnessTable[BucklingStiffnessIndex];		
+		#pragma ignore warning(perf)
+		const varying float DampingRatio = DampingTable[DampingIndex];
+		
+		const varying float StiffnessValue = select(IsBuckled[i], BucklingStiffness, Stiffness);
+
+		const varying FVector3f X1 = VectorGather(&X[i1]);
+		const varying FVector3f X2 = VectorGather(&X[i2]);
+		const varying FVector3f X3 = VectorGather(&X[i3]);
+		const varying FVector3f X4 = VectorGather(&X[i4]);
+		
+		const varying float Lambda = Lambdas[i];
+
+		varying FVector3f Delta1, Delta2, Delta3, Delta4;
+		const varying float DLambda = GetXPBDBendingElementDampingDeltas(P1, P2, P3, P4, X1, X2, X3, X4, M1, M2, M3, M4, Dt, RestAngles[i], Lambda, StiffnessValue, DampingRatio, MinStiffness, Delta1, Delta2, Delta3, Delta4);
 
-		const varying float ComputedDLambda = (Angle - RestAngle - Alpha * Lambda + DampingTerm) / Denom;
-		const varying float DLambda = select( StiffnessValue < MinStiffness, (uniform float)0, ComputedDLambda);
 		Lambdas[i] = Lambda + DLambda;
 
 		if (M1 > 0)
 		{
-			VectorScatter(&PandInvM[i1], SetVector4(P1 - DLambda * M1 * Grad1, M1));
+			VectorScatter(&PandInvM[i1], SetVector4(P1 + Delta1, M1));
 		}
 		if (M2 > 0)
 		{
-			VectorScatter(&PandInvM[i2], SetVector4(P2 - DLambda * M2 * Grad2, M2));
+			VectorScatter(&PandInvM[i2], SetVector4(P2 + Delta2, M2));
 		}
 		if (M3 > 0)
 		{
-			VectorScatter(&PandInvM[i3], SetVector4(P3 - DLambda * M3 * Grad3, M3));
+			VectorScatter(&PandInvM[i3], SetVector4(P3 + Delta3, M3));
 		}
 		if (M4 > 0)
 		{
-			VectorScatter(&PandInvM[i4], SetVector4(P4 - DLambda * M4 * Grad4, M4));
+			VectorScatter(&PandInvM[i4], SetVector4(P4 + Delta4, M4));
 		}
 	}
-}
\ No newline at end of file
+}
diff --git a/Engine/Source/Runtime/Experimental/Chaos/Private/Chaos/XPBDBendingElement.isph b/Engine/Source/Runtime/Experimental/Chaos/Private/Chaos/XPBDBendingElement.isph
new file mode 100644
--- /dev/null
+++ b/Engine/Source/Runtime/Experimental/Chaos/Private/Chaos/XPBDBendingElement.isph
@@ -0,0 +1,176 @@
+// Copyright Epic Games, Inc. All Rights Reserved.
+
+#ifndef __XPBDBENDING_ELEMENT_ISPH__
+#define __XPBDBENDING_ELEMENT_ISPH__
+#define EXPLICIT_VECTOR4 1
+
+#include "Math/Vector.isph"
+#include "Chaos/PBDSofts.isph"
+
+inline FVector3f SafeDivide(const FVector3f& Numerator, const float Denominator)
+{
+	return VectorSelect(Denominator <= FLOAT_SMALL_NUMBER, FloatZeroVector, Numerator / Denominator);
+}
+
+inline varying float SafeDivide(const varying float Numerator, const varying float Denominator)
+{
+	return select(Denominator <= FLOAT_SMALL_NUMBER, (uniform float)0, Numerator / Denominator);
+}
+
+inline void CalcGradients(const FVector3f& P1, const FVector3f& P2, const FVector3f& P3, const FVector3f& P4, FVector3f& Grad1, FVector3f& Grad2, FVector3f& Grad3, FVector3f& Grad4)
+{
+	const FVector3f SharedEdgeNormalized = VectorGetSafeNormal(P2 - P1);
+	const FVector3f P13CrossP23 = VectorCross(P1 - P3, P2 - P3);
+	const float Normal1Len = VectorSize(P13CrossP23);
+	const FVector3f Normal1 = SafeDivide(P13CrossP23, Normal1Len);
+	const FVector3f P24CrossP14 = VectorCross(P2 - P4, P1 - P4);
+	const float Normal2Len = VectorSize(P24CrossP14);
+	const FVector3f Normal2 = SafeDivide(P24CrossP14, Normal2Len);
+
+	const FVector3f N2CrossN1 = VectorCross(Normal2, Normal1);
+
+	const float CosPhi = clamp(VectorDot(Normal1, Normal2), (uniform float)(-1), (uniform float)(1));
+	const float SinPhi = clamp(VectorDot(N2CrossN1, SharedEdgeNormalized), (uniform float)(-1), (uniform float)(1));
+
+	const FVector3f DPhiDN1_OverNormal1Len = SafeDivide(CosPhi * VectorCross(SharedEdgeNormalized, Normal2) - SinPhi * Normal2, Normal1Len);
+	const FVector3f DPhiDN2_OverNormal2Len = SafeDivide(CosPhi * VectorCross(Normal1, SharedEdgeNormalized) - SinPhi * Normal1, Normal2Len);
+
+	const FVector3f DPhiDP13 = VectorCross(P2 - P3, DPhiDN1_OverNormal1Len);
+	const FVector3f DPhiDP23 = VectorCross(DPhiDN1_OverNormal1Len, P1 - P3);
+	const FVector3f DPhiDP24 = VectorCross(P1 - P4, DPhiDN2_OverNormal2Len);
+	const FVector3f DPhiDP14 = VectorCross(DPhiDN2_OverNormal2Len, P2 - P4);
+
+	Grad1 = DPhiDP13 + DPhiDP14;
+	Grad2 = DPhiDP23 + DPhiDP24;
+	Grad3 = -1 * DPhiDP13 - DPhiDP23;
+	Grad4 = -1 * DPhiDP14 - DPhiDP24;
+}
+
+// Returns Angle
+inline float CalcGradientsAndAngle(const FVector3f& P1, const FVector3f& P2, const FVector3f& P3, const FVector3f& P4, FVector3f& Grad1, FVector3f& Grad2, FVector3f& Grad3, FVector3f& Grad4)
+{
+	const FVector3f SharedEdgeNormalized = VectorGetSafeNormal(P2 - P1);
+	const FVector3f P13CrossP23 = VectorCross(P1 - P3, P2 - P3);
+	const float Normal1Len = VectorSize(P13CrossP23);
+	const FVector3f Normal1 = SafeDivide(P13CrossP23, Normal1Len);
+	const FVector3f P24CrossP14 = VectorCross(P2 - P4, P1 - P4);
+	const float Normal2Len = VectorSize(P24CrossP14);
+	const FVector3f Normal2 = SafeDivide(P24CrossP14, Normal2Len);
+
+	const FVector3f N2CrossN1 = VectorCross(Normal2, Normal1);
+
+	const float CosPhi = clamp(VectorDot(Normal1, Normal2), (uniform float)(-1), (uniform float)(1));
+	const float SinPhi = clamp(VectorDot(N2CrossN1, SharedEdgeNormalized), (uniform float)(-1), (uniform float)(1));
+
+	const FVector3f DPhiDN1_OverNormal1Len = SafeDivide(CosPhi * VectorCross(SharedEdgeNormalized, Normal2) - SinPhi * Normal2, Normal1Len);
+	const FVector3f DPhiDN2_OverNormal2Len = SafeDivide(CosPhi * VectorCross(Normal1, SharedEdgeNormalized) - SinPhi * Normal1, Normal2Len);
+
+	const FVector3f DPhiDP13 = VectorCross(P2 - P3, DPhiDN1_OverNormal1Len);
+	const FVector3f DPhiDP23 = VectorCross(DPhiDN1_OverNormal1Len, P1 - P3);
+	const FVector3f DPhiDP24 = VectorCross(P1 - P4, DPhiDN2_OverNormal2Len);
+	const FVector3f DPhiDP14 = VectorCross(DPhiDN2_OverNormal2Len, P2 - P4);
+
+	Grad1 = DPhiDP13 + DPhiDP14;
+	Grad2 = DPhiDP23 + DPhiDP24;
+	Grad3 = -1 * DPhiDP13 - DPhiDP23;
+	Grad4 = -1 * DPhiDP14 - DPhiDP24;
+
+	return atan2(SinPhi, CosPhi);
+}
+
+// Elastic part only
+inline float GetXPBDBendingElementDeltas(const FVector3f& P1, const FVector3f& P2, const FVector3f& P3, const FVector3f& P4,
+	const varying float InvM1, const varying float InvM2, const varying float InvM3, const varying float InvM4,
+	const uniform float Dt, const varying float RestAngle, const varying float Lambda, const varying float StiffnessValue, const uniform float MinStiffness, 
+	FVector3f& Delta1, FVector3f& Delta2, FVector3f& Delta3, FVector3f& Delta4)
+{
+	varying FVector3f Grad1, Grad2, Grad3, Grad4;
+	const varying float Angle = CalcGradientsAndAngle(P1, P2, P3, P4, Grad1, Grad2, Grad3, Grad4);
+	
+	const varying float AlphaInv  = StiffnessValue * Dt * Dt;	
+	const varying float ElasticTerm = AlphaInv * (Angle - RestAngle);
+	const varying float Denom = AlphaInv*(InvM1 * VectorSizeSquared(Grad1) + InvM2 * VectorSizeSquared(Grad2) + InvM3 * VectorSizeSquared(Grad3) + InvM4 * VectorSizeSquared(Grad4)) + 1.f;
+
+	const varying float ComputedDLambda = (-ElasticTerm - Lambda) / Denom;
+	const varying float DLambda = select(StiffnessValue < MinStiffness, (uniform float)0, ComputedDLambda);
+
+	Delta1 = DLambda * InvM1 * Grad1;
+	Delta2 = DLambda * InvM2 * Grad2;
+	Delta3 = DLambda * InvM3 * Grad3;
+	Delta4 = DLambda * InvM4 * Grad4;
+	return DLambda;
+}
+
+// This is a following the original XPBD paper using a single lambda for stretch and damping.
+template <typename DampingFloatType>
+inline float GetXPBDBendingElementDeltasWithDamping(const FVector3f& P1, const FVector3f& P2, const FVector3f& P3, const FVector3f& P4,
+	const FVector3f& X1, const FVector3f& X2, const FVector3f& X3, const FVector3f& X4,
+	const varying float InvM1, const varying float InvM2, const varying float InvM3, const varying float InvM4,
+	const uniform float Dt, const varying float RestAngle, const varying float Lambda, 
+	const varying float StiffnessValue, const DampingFloatType DampingRatio, const uniform float MinStiffness, 
+	FVector3f& Delta1, FVector3f& Delta2, FVector3f& Delta3, FVector3f& Delta4)
+{
+	varying FVector3f Grad1, Grad2, Grad3, Grad4;
+	const varying float Angle = CalcGradientsAndAngle(P1, P2, P3, P4, Grad1, Grad2, Grad3, Grad4);
+	
+	const varying float AlphaInv = StiffnessValue * Dt * Dt;
+	const varying float CombinedInvMass = InvM1 + InvM2 + InvM3 + InvM4;
+	const varying float Damping = 2.f * DampingRatio * sqrt(StiffnessValue / CombinedInvMass);
+	const varying float BetaDt = Damping * Dt;
+	
+	const varying FVector3f V1TimesDt = P1 - X1;
+	const varying FVector3f V2TimesDt = P2 - X2;
+	const varying FVector3f V3TimesDt = P3 - X3;
+	const varying FVector3f V4TimesDt = P4 - X4;
+	const varying float DampingTerm = BetaDt * (VectorDot(V1TimesDt, Grad1) + VectorDot(V2TimesDt, Grad2) + VectorDot(V3TimesDt, Grad3) + VectorDot(V4TimesDt, Grad4));
+	const varying float ElasticTerm = AlphaInv * (Angle - RestAngle);
+
+	const varying float Denom = (AlphaInv + BetaDt)*(InvM1 * VectorSizeSquared(Grad1) + InvM2 * VectorSizeSquared(Grad2) + InvM3 * VectorSizeSquared(Grad3) + InvM4 * VectorSizeSquared(Grad4)) + 1.f;
+
+	const varying float ComputedDLambda = (-ElasticTerm - DampingTerm - Lambda) / Denom;
+
+	const varying float DLambda = select(StiffnessValue < MinStiffness, (uniform float)0, ComputedDLambda);
+
+	Delta1 = DLambda * InvM1 * Grad1;
+	Delta2 = DLambda * InvM2 * Grad2;
+	Delta3 = DLambda * InvM3 * Grad3;
+	Delta4 = DLambda * InvM4 * Grad4;
+	return DLambda;
+}
+
+// Damping only (separate from elastic part)
+template <typename DampingFloatType>
+inline float GetXPBDBendingElementDampingDeltas(const FVector3f& P1, const FVector3f& P2, const FVector3f& P3, const FVector3f& P4,
+	const FVector3f& X1, const FVector3f& X2, const FVector3f& X3, const FVector3f& X4,
+	const varying float InvM1, const varying float InvM2, const varying float InvM3, const varying float InvM4,
+	const uniform float Dt, const varying float RestAngle, const varying float Lambda, 
+	const varying float StiffnessValue, const DampingFloatType DampingRatio, const uniform float MinStiffness, 
+	FVector3f& Delta1, FVector3f& Delta2, FVector3f& Delta3, FVector3f& Delta4)
+{
+	varying FVector3f Grad1, Grad2, Grad3, Grad4;
+	CalcGradients(P1, P2, P3, P4, Grad1, Grad2, Grad3, Grad4);
+	
+	const varying float CombinedInvMass = InvM1 + InvM2 + InvM3 + InvM4;
+	const varying float Damping = 2.f * DampingRatio * sqrt(StiffnessValue / CombinedInvMass);
+	const varying float BetaDt = Damping * Dt;
+	
+	const varying FVector3f V1TimesDt = P1 - X1;
+	const varying FVector3f V2TimesDt = P2 - X2;
+	const varying FVector3f V3TimesDt = P3 - X3;
+	const varying FVector3f V4TimesDt = P4 - X4;
+	const varying float DampingTerm = BetaDt * (VectorDot(V1TimesDt, Grad1) + VectorDot(V2TimesDt, Grad2) + VectorDot(V3TimesDt, Grad3) + VectorDot(V4TimesDt, Grad4));
+
+	const varying float Denom = BetaDt*(InvM1 * VectorSizeSquared(Grad1) + InvM2 * VectorSizeSquared(Grad2) + InvM3 * VectorSizeSquared(Grad3) + InvM4 * VectorSizeSquared(Grad4)) + 1.f;
+
+	const varying float ComputedDLambda = ( -DampingTerm - Lambda) / Denom;
+
+	const varying float DLambda = select(StiffnessValue < MinStiffness, (uniform float)0, ComputedDLambda);
+
+	Delta1 = DLambda * InvM1 * Grad1;
+	Delta2 = DLambda * InvM2 * Grad2;
+	Delta3 = DLambda * InvM3 * Grad3;
+	Delta4 = DLambda * InvM4 * Grad4;
+	return DLambda;
+}
+
+#endif
\ No newline at end of file
diff --git a/Engine/Source/Runtime/Experimental/Chaos/Private/Chaos/XPBDInternal.h b/Engine/Source/Runtime/Experimental/Chaos/Private/Chaos/XPBDInternal.h
new file mode 100644
--- /dev/null
+++ b/Engine/Source/Runtime/Experimental/Chaos/Private/Chaos/XPBDInternal.h
@@ -0,0 +1,21 @@
+// Copyright Epic Games, Inc. All Rights Reserved.
+#pragma once
+
+namespace Chaos::Softs
+{
+// Used for controlling damping model for xpbd springs (CVar-based testing only)
+enum struct EXPBDSplitDampingMode
+{
+	SingleLambda = 0,
+	InterleavedAfter = 1,
+	InterleavedBefore = 2,
+	TwoPassAfter = 3,
+	TwoPassBefore = 4
+};
+#if UE_BUILD_SHIPPING
+const int32 Chaos_XPBDSpring_SplitDampingMode = (int32)EXPBDSplitDampingMode::TwoPassBefore;
+#else
+extern int32 Chaos_XPBDSpring_SplitDampingMode;
+#endif
+
+}  // End namespace Chaos::Softs
diff --git a/Engine/Source/Runtime/Experimental/Chaos/Private/Chaos/XPBDSpringConstraints.cpp b/Engine/Source/Runtime/Experimental/Chaos/Private/Chaos/XPBDSpringConstraints.cpp
--- a/Engine/Source/Runtime/Experimental/Chaos/Private/Chaos/XPBDSpringConstraints.cpp
+++ b/Engine/Source/Runtime/Experimental/Chaos/Private/Chaos/XPBDSpringConstraints.cpp
@@ -8,6 +8,8 @@
 #include "ChaosLog.h"
 #include "HAL/IConsoleManager.h"
 
+#include "XPBDInternal.h"
+
 #if INTEL_ISPC
 #include "XPBDSpringConstraints.ispc.generated.h"
 #endif
@@ -29,6 +31,9 @@ namespace Chaos::Softs {
 int32 Chaos_XPBDSpring_ParallelConstraintCount = 100;
 #if !UE_BUILD_SHIPPING
 FAutoConsoleVariableRef CVarChaosXPBDSpringParallelConstraintCount(TEXT("p.Chaos.XPBDSpring.ParallelConstraintCount"), Chaos_XPBDSpring_ParallelConstraintCount, TEXT("If we have more constraints than this, use parallel-for in Apply."));
+
+int32 Chaos_XPBDSpring_SplitDampingMode = (int32)EXPBDSplitDampingMode::TwoPassBefore;
+FAutoConsoleVariableRef CVarChaosXPBDSpringSplitDamping(TEXT("p.Chaos.XPBDSpring.SplitDamping"), Chaos_XPBDSpring_SplitDampingMode, TEXT("Test xpbd spring split damping mode. 0 = single lambda, 1 = interleaved with damping after (non-ispc only), 2 interleaved with damping before (non-ispc only), 3 =  two passes damping after (non-ispc only), 4 = two passes damping before (default)."));
 #endif
 
 template<typename SolverParticlesOrRange>
@@ -76,20 +81,45 @@ void FXPBDSpringConstraints::InitColor(const SolverParticlesOrRange& Particles)
 template CHAOS_API void FXPBDSpringConstraints::InitColor(const FSolverParticles& Particles);
 template CHAOS_API void FXPBDSpringConstraints::InitColor(const FSolverParticlesRange& Particles);
 
-template<typename SolverParticlesOrRange>
+template<bool bDampingBefore, bool bSingleLambda, bool bSeparateStretch, bool bDampingAfter, typename SolverParticlesOrRange>
 void FXPBDSpringConstraints::ApplyHelper(SolverParticlesOrRange& Particles, const FSolverReal Dt, const int32 ConstraintIndex, const FSolverReal ExpStiffnessValue, const FSolverReal DampingRatioValue) const
 {
 	const TVec2<int32>& Constraint = Constraints[ConstraintIndex];
 	const int32 Index1 = Constraint[0];
 	const int32 Index2 = Constraint[1];
-	const FSolverVec3 Delta =  GetDelta(Particles, Dt, ConstraintIndex, ExpStiffnessValue, DampingRatioValue);
+
+	FSolverVec3 Delta(0.f);
+	if constexpr (bDampingBefore)
+	{
+		Delta += Spring::GetXPBDSpringDampingDelta(Particles, Dt, Constraint, Dists[ConstraintIndex], LambdasDamping[ConstraintIndex],
+			ExpStiffnessValue, MinStiffness, DampingRatioValue);
+	}
+
+	if constexpr (bSingleLambda)
+	{
+		Delta += Spring::GetXPBDSpringDeltaWithDamping(Particles, Dt, Constraint, Dists[ConstraintIndex], Lambdas[ConstraintIndex],
+			ExpStiffnessValue, MinStiffness, DampingRatioValue);
+	}
+
+	if constexpr (bSeparateStretch)
+	{
+		Delta += Spring::GetXPBDSpringDelta(Particles, Dt, Constraint, Dists[ConstraintIndex], Lambdas[ConstraintIndex],
+			ExpStiffnessValue, MinStiffness);
+	}
+
+	if constexpr (bDampingAfter)
+	{
+		Delta += Spring::GetXPBDSpringDampingDelta(Particles, Dt, Constraint, Dists[ConstraintIndex], LambdasDamping[ConstraintIndex],
+			ExpStiffnessValue, MinStiffness, DampingRatioValue);
+	}
+
 	if (Particles.InvM(Index1) > (FSolverReal)0.)
 	{
-		Particles.P(Index1) -= Particles.InvM(Index1) * Delta;
+		Particles.P(Index1) += Particles.InvM(Index1) * Delta;
 	}
 	if (Particles.InvM(Index2) > (FSolverReal)0.)
 	{
-		Particles.P(Index2) += Particles.InvM(Index2) * Delta;
+		Particles.P(Index2) -= Particles.InvM(Index2) * Delta;
 	}
 }
 
@@ -105,130 +135,297 @@ void FXPBDSpringConstraints::Apply(SolverParticlesOrRange& Particles, const FSol
 	if ((ConstraintsPerColorStartIndex.Num() > 1) && (Constraints.Num() > Chaos_XPBDSpring_ParallelConstraintCount))
 	{
 		const int32 ConstraintColorNum = ConstraintsPerColorStartIndex.Num() - 1;
-		if (!StiffnessHasWeightMap && !DampingHasWeightMap)
+#if INTEL_ISPC
+		if (bRealTypeCompatibleWithISPC && bChaos_XPBDSpring_ISPC_Enabled)
 		{
-			const FSolverReal ExpStiffnessValue = (FSolverReal)Stiffness;
-			const FSolverReal DampingRatioValue = (FSolverReal)DampingRatio;
-			if (ExpStiffnessValue < MinStiffness)
+			const bool bSingleLambda = Chaos_XPBDSpring_SplitDampingMode == (int32)EXPBDSplitDampingMode::SingleLambda;
+			if (!StiffnessHasWeightMap && !DampingHasWeightMap)
 			{
-				return;
-			}
+				const FSolverReal ExpStiffnessValue = (FSolverReal)Stiffness;
+				const FSolverReal DampingRatioValue = (FSolverReal)DampingRatio;
+				if (ExpStiffnessValue < MinStiffness)
+				{
+					return;
+				}
+				
+				if (DampingRatioValue > 0)
+				{
+					if (bSingleLambda)
+					{
+						for (int32 ConstraintColorIndex = 0; ConstraintColorIndex < ConstraintColorNum; ++ConstraintColorIndex)
+						{
+							const int32 ColorStart = ConstraintsPerColorStartIndex[ConstraintColorIndex];
+							const int32 ColorSize = ConstraintsPerColorStartIndex[ConstraintColorIndex + 1] - ColorStart;
+							ispc::ApplyXPBDSpringConstraintsWithDamping(
+								(ispc::FVector4f*)Particles.GetPAndInvM().GetData(),
+								(const ispc::FVector3f*)Particles.XArray().GetData(),
+								(ispc::FIntVector2*)&Constraints.GetData()[ColorStart],
+								&Dists.GetData()[ColorStart],
+								&Lambdas.GetData()[ColorStart],
+								Dt,
+								ExpStiffnessValue,
+								DampingRatioValue,
+								ColorSize);
+						}
+						return;
+					}
+					else
+					{
+						for (int32 ConstraintColorIndex = 0; ConstraintColorIndex < ConstraintColorNum; ++ConstraintColorIndex)
+						{
+							const int32 ColorStart = ConstraintsPerColorStartIndex[ConstraintColorIndex];
+							const int32 ColorSize = ConstraintsPerColorStartIndex[ConstraintColorIndex + 1] - ColorStart;
+							ispc::ApplyXPBDSpringDampingConstraints(
+								(ispc::FVector4f*)Particles.GetPAndInvM().GetData(),
+								(const ispc::FVector3f*)Particles.XArray().GetData(),
+								(ispc::FIntVector2*)&Constraints.GetData()[ColorStart],
+								&Dists.GetData()[ColorStart],
+								&LambdasDamping.GetData()[ColorStart],
+								Dt,
+								ExpStiffnessValue,
+								DampingRatioValue,
+								ColorSize);
+						}
+					}
+				}
 
-#if INTEL_ISPC
-			if (bRealTypeCompatibleWithISPC && bChaos_XPBDSpring_ISPC_Enabled)
+				for (int32 ConstraintColorIndex = 0; ConstraintColorIndex < ConstraintColorNum; ++ConstraintColorIndex)
+				{
+					const int32 ColorStart = ConstraintsPerColorStartIndex[ConstraintColorIndex];
+					const int32 ColorSize = ConstraintsPerColorStartIndex[ConstraintColorIndex + 1] - ColorStart;
+					ispc::ApplyXPBDSpringConstraints(
+						(ispc::FVector4f*)Particles.GetPAndInvM().GetData(),
+						(ispc::FIntVector2*)&Constraints.GetData()[ColorStart],
+						&Dists.GetData()[ColorStart],
+						&Lambdas.GetData()[ColorStart],
+						Dt,
+						ExpStiffnessValue,
+						ColorSize);
+				}
+			}
+			else // ISPC with maps
 			{
-				if (DampingRatioValue > 0)
+				if (DampingHasWeightMap || (FSolverReal)DampingRatio > 0)
+				{
+					if (bSingleLambda)
+					{
+						for (int32 ConstraintColorIndex = 0; ConstraintColorIndex < ConstraintColorNum; ++ConstraintColorIndex)
+						{
+							const int32 ColorStart = ConstraintsPerColorStartIndex[ConstraintColorIndex];
+							const int32 ColorSize = ConstraintsPerColorStartIndex[ConstraintColorIndex + 1] - ColorStart;
+							ispc::ApplyXPBDSpringConstraintsWithDampingAndWeightMaps(
+								(ispc::FVector4f*)Particles.GetPAndInvM().GetData(),
+								(const ispc::FVector3f*)Particles.XArray().GetData(),
+								(ispc::FIntVector2*)&Constraints.GetData()[ColorStart],
+								&Dists.GetData()[ColorStart],
+								&Lambdas.GetData()[ColorStart],
+								Dt,
+								MinStiffness,
+								StiffnessHasWeightMap,
+								StiffnessHasWeightMap ? &Stiffness.GetIndices().GetData()[ColorStart] : nullptr,
+								&Stiffness.GetTable().GetData()[0],
+								DampingHasWeightMap,
+								DampingHasWeightMap ? &DampingRatio.GetIndices().GetData()[ColorStart] : nullptr,
+								&DampingRatio.GetTable().GetData()[0],
+								ColorSize);
+						}
+						return;
+					}
+					else
+					{
+						for (int32 ConstraintColorIndex = 0; ConstraintColorIndex < ConstraintColorNum; ++ConstraintColorIndex)
+						{
+							const int32 ColorStart = ConstraintsPerColorStartIndex[ConstraintColorIndex];
+							const int32 ColorSize = ConstraintsPerColorStartIndex[ConstraintColorIndex + 1] - ColorStart;
+							ispc::ApplyXPBDSpringDampingConstraintsWithWeightMaps(
+								(ispc::FVector4f*)Particles.GetPAndInvM().GetData(),
+								(const ispc::FVector3f*)Particles.XArray().GetData(),
+								(ispc::FIntVector2*)&Constraints.GetData()[ColorStart],
+								&Dists.GetData()[ColorStart],
+								&LambdasDamping.GetData()[ColorStart],
+								Dt,
+								MinStiffness,
+								StiffnessHasWeightMap,
+								StiffnessHasWeightMap ? &Stiffness.GetIndices().GetData()[ColorStart] : nullptr,
+								&Stiffness.GetTable().GetData()[0],
+								DampingHasWeightMap,
+								DampingHasWeightMap ? &DampingRatio.GetIndices().GetData()[ColorStart] : nullptr,
+								&DampingRatio.GetTable().GetData()[0],
+								ColorSize);
+						}
+					}
+				}
+				for (int32 ConstraintColorIndex = 0; ConstraintColorIndex < ConstraintColorNum; ++ConstraintColorIndex)
+				{
+					const int32 ColorStart = ConstraintsPerColorStartIndex[ConstraintColorIndex];
+					const int32 ColorSize = ConstraintsPerColorStartIndex[ConstraintColorIndex + 1] - ColorStart;
+					ispc::ApplyXPBDSpringConstraintsWithWeightMaps(
+						(ispc::FVector4f*)Particles.GetPAndInvM().GetData(),
+						(ispc::FIntVector2*)&Constraints.GetData()[ColorStart],
+						&Dists.GetData()[ColorStart],
+						&Lambdas.GetData()[ColorStart],
+						Dt,
+						MinStiffness,
+						&Stiffness.GetIndices().GetData()[ColorStart],
+						&Stiffness.GetTable().GetData()[0],
+						ColorSize);
+				}
+			}
+		}
+		else
+#endif
+		{
+			// Parallel non-ispc
+			const FSolverReal StiffnessNoMap = (FSolverReal)Stiffness;
+			const FSolverReal DampingNoMap = (FSolverReal)DampingRatio;
+			if (DampingNoMap > 0 || DampingHasWeightMap)
+			{
+				if (Chaos_XPBDSpring_SplitDampingMode == (int32)EXPBDSplitDampingMode::TwoPassBefore)
 				{
 					for (int32 ConstraintColorIndex = 0; ConstraintColorIndex < ConstraintColorNum; ++ConstraintColorIndex)
 					{
 						const int32 ColorStart = ConstraintsPerColorStartIndex[ConstraintColorIndex];
 						const int32 ColorSize = ConstraintsPerColorStartIndex[ConstraintColorIndex + 1] - ColorStart;
-						ispc::ApplyXPBDSpringConstraintsWithDamping(
-							(ispc::FVector4f*)Particles.GetPAndInvM().GetData(),
-							(const ispc::FVector3f*)Particles.XArray().GetData(),
-							(ispc::FIntVector2*)&Constraints.GetData()[ColorStart],
-							&Dists.GetData()[ColorStart],
-							&Lambdas.GetData()[ColorStart],
-							Dt,
-							ExpStiffnessValue,
-							DampingRatioValue,
-							ColorSize);
+						PhysicsParallelFor(ColorSize, [this, &Particles, Dt, ColorStart, StiffnessHasWeightMap, StiffnessNoMap, DampingHasWeightMap, DampingNoMap](const int32 Index)
+						{
+							const int32 ConstraintIndex = ColorStart + Index;
+							constexpr bool bDampingBefore = true;
+							constexpr bool bSingleLambda = false;
+							constexpr bool bSeparateStretch = false;
+							constexpr bool bDampingAfter = false;
+							const FSolverReal ExpStiffnessValue = StiffnessHasWeightMap ? Stiffness[ConstraintIndex] : StiffnessNoMap;
+							const FSolverReal DampingRatioValue = DampingHasWeightMap ? DampingRatio[ConstraintIndex] : DampingNoMap;
+							ApplyHelper<bDampingBefore, bSingleLambda, bSeparateStretch, bDampingAfter>(Particles, Dt, ConstraintIndex, ExpStiffnessValue, DampingRatioValue);
+						});
 					}
 				}
-				else
+
+				// Stretch part (possibly with damping depending on split mode)
+				switch (Chaos_XPBDSpring_SplitDampingMode)
 				{
+				case (int32)EXPBDSplitDampingMode::TwoPassBefore: // fallthrough
+				case (int32)EXPBDSplitDampingMode::TwoPassAfter: // fallthrough
+				default:
+				{
+					// Do a pass with stretch only
 					for (int32 ConstraintColorIndex = 0; ConstraintColorIndex < ConstraintColorNum; ++ConstraintColorIndex)
 					{
 						const int32 ColorStart = ConstraintsPerColorStartIndex[ConstraintColorIndex];
 						const int32 ColorSize = ConstraintsPerColorStartIndex[ConstraintColorIndex + 1] - ColorStart;
-						ispc::ApplyXPBDSpringConstraints(
-							(ispc::FVector4f*)Particles.GetPAndInvM().GetData(),
-							(ispc::FIntVector2*)&Constraints.GetData()[ColorStart],
-							&Dists.GetData()[ColorStart],
-							&Lambdas.GetData()[ColorStart],
-							Dt,
-							ExpStiffnessValue,
-							ColorSize);
+						PhysicsParallelFor(ColorSize, [this, &Particles, Dt, ColorStart, StiffnessHasWeightMap, StiffnessNoMap, DampingHasWeightMap, DampingNoMap](const int32 Index)
+						{
+							const int32 ConstraintIndex = ColorStart + Index;
+							constexpr bool bDampingBefore = false;
+							constexpr bool bSingleLambda = false;
+							constexpr bool bSeparateStretch = true;
+							constexpr bool bDampingAfter = false;
+							const FSolverReal ExpStiffnessValue = StiffnessHasWeightMap ? Stiffness[ConstraintIndex] : StiffnessNoMap;
+							const FSolverReal DampingRatioValue = DampingHasWeightMap ? DampingRatio[ConstraintIndex] : DampingNoMap;
+							ApplyHelper<bDampingBefore, bSingleLambda, bSeparateStretch, bDampingAfter>(Particles, Dt, ConstraintIndex, ExpStiffnessValue, DampingRatioValue);
+						});
 					}
 				}
-			}
-			else
-#endif
-			{
-				for (int32 ConstraintColorIndex = 0; ConstraintColorIndex < ConstraintColorNum; ++ConstraintColorIndex)
+				break;
+				case (int32)EXPBDSplitDampingMode::SingleLambda:
 				{
-					const int32 ColorStart = ConstraintsPerColorStartIndex[ConstraintColorIndex];
-					const int32 ColorSize = ConstraintsPerColorStartIndex[ConstraintColorIndex + 1] - ColorStart;
-					PhysicsParallelFor(ColorSize, [this, &Particles, Dt, ColorStart, ExpStiffnessValue, DampingRatioValue](const int32 Index)
+					// Do a pass with combined stretch and damping with a single lambda
+					for (int32 ConstraintColorIndex = 0; ConstraintColorIndex < ConstraintColorNum; ++ConstraintColorIndex)
 					{
-						const int32 ConstraintIndex = ColorStart + Index;
-						ApplyHelper(Particles, Dt, ConstraintIndex, ExpStiffnessValue, DampingRatioValue);
-					});
+						const int32 ColorStart = ConstraintsPerColorStartIndex[ConstraintColorIndex];
+						const int32 ColorSize = ConstraintsPerColorStartIndex[ConstraintColorIndex + 1] - ColorStart;
+						PhysicsParallelFor(ColorSize, [this, &Particles, Dt, ColorStart, StiffnessHasWeightMap, StiffnessNoMap, DampingHasWeightMap, DampingNoMap](const int32 Index)
+						{
+							const int32 ConstraintIndex = ColorStart + Index;
+							constexpr bool bDampingBefore = false;
+							constexpr bool bSingleLambda = true;
+							constexpr bool bSeparateStretch = false;
+							constexpr bool bDampingAfter = false;
+							const FSolverReal ExpStiffnessValue = StiffnessHasWeightMap ? Stiffness[ConstraintIndex] : StiffnessNoMap;
+							const FSolverReal DampingRatioValue = DampingHasWeightMap ? DampingRatio[ConstraintIndex] : DampingNoMap;
+							ApplyHelper<bDampingBefore, bSingleLambda, bSeparateStretch, bDampingAfter>(Particles, Dt, ConstraintIndex, ExpStiffnessValue, DampingRatioValue);
+						});
+					}
 				}
-			}
-		}
-		else  // Has weight maps
-		{
-#if INTEL_ISPC
-			if (bRealTypeCompatibleWithISPC && bChaos_XPBDSpring_ISPC_Enabled)
-			{
-				if (DampingHasWeightMap || (FSolverReal)DampingRatio > 0)
+				break;
+				case (int32)EXPBDSplitDampingMode::InterleavedBefore:
+				{
+					// Do a pass with damping before stretch
+					for (int32 ConstraintColorIndex = 0; ConstraintColorIndex < ConstraintColorNum; ++ConstraintColorIndex)
+					{
+						const int32 ColorStart = ConstraintsPerColorStartIndex[ConstraintColorIndex];
+						const int32 ColorSize = ConstraintsPerColorStartIndex[ConstraintColorIndex + 1] - ColorStart;
+						PhysicsParallelFor(ColorSize, [this, &Particles, Dt, ColorStart, StiffnessHasWeightMap, StiffnessNoMap, DampingHasWeightMap, DampingNoMap](const int32 Index)
+						{
+							const int32 ConstraintIndex = ColorStart + Index;
+							constexpr bool bDampingBefore = true;
+							constexpr bool bSingleLambda = false;
+							constexpr bool bSeparateStretch = true;
+							constexpr bool bDampingAfter = false;
+							const FSolverReal ExpStiffnessValue = StiffnessHasWeightMap ? Stiffness[ConstraintIndex] : StiffnessNoMap;
+							const FSolverReal DampingRatioValue = DampingHasWeightMap ? DampingRatio[ConstraintIndex] : DampingNoMap;
+							ApplyHelper<bDampingBefore, bSingleLambda, bSeparateStretch, bDampingAfter>(Particles, Dt, ConstraintIndex, ExpStiffnessValue, DampingRatioValue);
+						});
+					}
+				}
+				break;
+				case (int32)EXPBDSplitDampingMode::InterleavedAfter:
 				{
+					// Do a pass with damping after stretch
 					for (int32 ConstraintColorIndex = 0; ConstraintColorIndex < ConstraintColorNum; ++ConstraintColorIndex)
 					{
 						const int32 ColorStart = ConstraintsPerColorStartIndex[ConstraintColorIndex];
 						const int32 ColorSize = ConstraintsPerColorStartIndex[ConstraintColorIndex + 1] - ColorStart;
-						ispc::ApplyXPBDSpringConstraintsWithDampingAndWeightMaps(
-							(ispc::FVector4f*)Particles.GetPAndInvM().GetData(),
-							(const ispc::FVector3f*)Particles.XArray().GetData(),
-							(ispc::FIntVector2*)&Constraints.GetData()[ColorStart],
-							&Dists.GetData()[ColorStart],
-							&Lambdas.GetData()[ColorStart],
-							Dt,
-							MinStiffness,
-							StiffnessHasWeightMap,
-							StiffnessHasWeightMap ? &Stiffness.GetIndices().GetData()[ColorStart] : nullptr,
-							&Stiffness.GetTable().GetData()[0],
-							DampingHasWeightMap,
-							DampingHasWeightMap ? &DampingRatio.GetIndices().GetData()[ColorStart] : nullptr,
-							&DampingRatio.GetTable().GetData()[0],
-							ColorSize);
+						PhysicsParallelFor(ColorSize, [this, &Particles, Dt, ColorStart, StiffnessHasWeightMap, StiffnessNoMap, DampingHasWeightMap, DampingNoMap](const int32 Index)
+						{
+							const int32 ConstraintIndex = ColorStart + Index;
+							constexpr bool bDampingBefore = false;
+							constexpr bool bSingleLambda = false;
+							constexpr bool bSeparateStretch = true;
+							constexpr bool bDampingAfter = true;
+							const FSolverReal ExpStiffnessValue = StiffnessHasWeightMap ? Stiffness[ConstraintIndex] : StiffnessNoMap;
+							const FSolverReal DampingRatioValue = DampingHasWeightMap ? DampingRatio[ConstraintIndex] : DampingNoMap;
+							ApplyHelper<bDampingBefore, bSingleLambda, bSeparateStretch, bDampingAfter>(Particles, Dt, ConstraintIndex, ExpStiffnessValue, DampingRatioValue);
+						});
 					}
 				}
-				else
+				break;
+				}
+				if (Chaos_XPBDSpring_SplitDampingMode == (int32)EXPBDSplitDampingMode::TwoPassAfter)
 				{
 					for (int32 ConstraintColorIndex = 0; ConstraintColorIndex < ConstraintColorNum; ++ConstraintColorIndex)
 					{
 						const int32 ColorStart = ConstraintsPerColorStartIndex[ConstraintColorIndex];
 						const int32 ColorSize = ConstraintsPerColorStartIndex[ConstraintColorIndex + 1] - ColorStart;
-						ispc::ApplyXPBDSpringConstraintsWithWeightMaps(
-							(ispc::FVector4f*)Particles.GetPAndInvM().GetData(),
-							(ispc::FIntVector2*)&Constraints.GetData()[ColorStart],
-							&Dists.GetData()[ColorStart],
-							&Lambdas.GetData()[ColorStart],
-							Dt,
-							MinStiffness,
-							&Stiffness.GetIndices().GetData()[ColorStart],
-							&Stiffness.GetTable().GetData()[0],
-							ColorSize);
+						PhysicsParallelFor(ColorSize, [this, &Particles, Dt, ColorStart, StiffnessHasWeightMap, StiffnessNoMap, DampingHasWeightMap, DampingNoMap](const int32 Index)
+						{
+							const int32 ConstraintIndex = ColorStart + Index;
+							constexpr bool bDampingBefore = false;
+							constexpr bool bSingleLambda = false;
+							constexpr bool bSeparateStretch = false;
+							constexpr bool bDampingAfter = true;
+							const FSolverReal ExpStiffnessValue = StiffnessHasWeightMap ? Stiffness[ConstraintIndex] : StiffnessNoMap;
+							const FSolverReal DampingRatioValue = DampingHasWeightMap ? DampingRatio[ConstraintIndex] : DampingNoMap;
+							ApplyHelper<bDampingBefore, bSingleLambda, bSeparateStretch, bDampingAfter>(Particles, Dt, ConstraintIndex, ExpStiffnessValue, DampingRatioValue);
+						});
 					}
 				}
 			}
 			else
-#endif
 			{
-				const FSolverReal StiffnessNoMap = (FSolverReal)Stiffness;
-				const FSolverReal DampingNoMap = (FSolverReal)DampingRatio;
+				// No damping. 
 				for (int32 ConstraintColorIndex = 0; ConstraintColorIndex < ConstraintColorNum; ++ConstraintColorIndex)
 				{
 					const int32 ColorStart = ConstraintsPerColorStartIndex[ConstraintColorIndex];
 					const int32 ColorSize = ConstraintsPerColorStartIndex[ConstraintColorIndex + 1] - ColorStart;
-					PhysicsParallelFor(ColorSize, [this, &Particles, Dt, ColorStart, StiffnessHasWeightMap, StiffnessNoMap, DampingHasWeightMap, DampingNoMap](const int32 Index)
+					PhysicsParallelFor(ColorSize, [this, &Particles, Dt, ColorStart, StiffnessHasWeightMap, StiffnessNoMap](const int32 Index)
 					{
 						const int32 ConstraintIndex = ColorStart + Index;
+						constexpr bool bDampingBefore = false;
+						constexpr bool bSingleLambda = false;
+						constexpr bool bSeparateStretch = true;
+						constexpr bool bDampingAfter = false;
 						const FSolverReal ExpStiffnessValue = StiffnessHasWeightMap ? Stiffness[ConstraintIndex] : StiffnessNoMap;
-						const FSolverReal DampingRatioValue = DampingHasWeightMap ? DampingRatio[ConstraintIndex] : DampingNoMap;
-						ApplyHelper(Particles, Dt, ConstraintIndex, ExpStiffnessValue, DampingRatioValue);
+						ApplyHelper<bDampingBefore, bSingleLambda, bSeparateStretch, bDampingAfter>(Particles, Dt, ConstraintIndex, ExpStiffnessValue, (FSolverReal)0.);
 					});
 				}
 			}
@@ -236,28 +433,118 @@ void FXPBDSpringConstraints::Apply(SolverParticlesOrRange& Particles, const FSol
 	}
 	else
 	{
-		if (!StiffnessHasWeightMap && !DampingHasWeightMap)
+		// Single threaded
+		const FSolverReal StiffnessNoMap = (FSolverReal)Stiffness;
+		const FSolverReal DampingNoMap = (FSolverReal)DampingRatio;
+		if (DampingNoMap > 0 || DampingHasWeightMap)
 		{
-			const FSolverReal ExpStiffnessValue = (FSolverReal)Stiffness;
-			const FSolverReal DampingRatioValue = (FSolverReal)DampingRatio;
-			if (ExpStiffnessValue < MinStiffness)
+			if (Chaos_XPBDSpring_SplitDampingMode == (int32)EXPBDSplitDampingMode::TwoPassBefore)
 			{
-				return;
+				for (int32 ConstraintIndex = 0; ConstraintIndex < Constraints.Num(); ++ConstraintIndex)
+				{
+					constexpr bool bDampingBefore = true;
+					constexpr bool bSingleLambda = false;
+					constexpr bool bSeparateStretch = false;
+					constexpr bool bDampingAfter = false;
+					const FSolverReal ExpStiffnessValue = StiffnessHasWeightMap ? Stiffness[ConstraintIndex] : StiffnessNoMap;
+					const FSolverReal DampingRatioValue = DampingHasWeightMap ? DampingRatio[ConstraintIndex] : DampingNoMap;
+					ApplyHelper<bDampingBefore, bSingleLambda, bSeparateStretch, bDampingAfter>(Particles, Dt, ConstraintIndex, ExpStiffnessValue, DampingRatioValue);
+				}
 			}
-			for (int32 ConstraintIndex = 0; ConstraintIndex < Constraints.Num(); ++ConstraintIndex)
+
+			// Stretch part (possibly with damping depending on split mode)
+			switch (Chaos_XPBDSpring_SplitDampingMode)
 			{
-				ApplyHelper(Particles, Dt, ConstraintIndex, ExpStiffnessValue, DampingRatioValue);
+			case (int32)EXPBDSplitDampingMode::TwoPassBefore: // fallthrough
+			case (int32)EXPBDSplitDampingMode::TwoPassAfter: // fallthrough
+			default:
+			{
+				// Do a pass with stretch only
+				for (int32 ConstraintIndex = 0; ConstraintIndex < Constraints.Num(); ++ConstraintIndex)
+				{
+					constexpr bool bDampingBefore = false;
+					constexpr bool bSingleLambda = false;
+					constexpr bool bSeparateStretch = true;
+					constexpr bool bDampingAfter = false;
+					const FSolverReal ExpStiffnessValue = StiffnessHasWeightMap ? Stiffness[ConstraintIndex] : StiffnessNoMap;
+					const FSolverReal DampingRatioValue = DampingHasWeightMap ? DampingRatio[ConstraintIndex] : DampingNoMap;
+					ApplyHelper<bDampingBefore, bSingleLambda, bSeparateStretch, bDampingAfter>(Particles, Dt, ConstraintIndex, ExpStiffnessValue, DampingRatioValue);
+				}
+			}
+			break;
+			case (int32)EXPBDSplitDampingMode::SingleLambda:
+			{
+				// Do a pass with combined stretch and damping with a single lambda
+				for (int32 ConstraintIndex = 0; ConstraintIndex < Constraints.Num(); ++ConstraintIndex)
+				{
+					constexpr bool bDampingBefore = false;
+					constexpr bool bSingleLambda = true;
+					constexpr bool bSeparateStretch = false;
+					constexpr bool bDampingAfter = false;
+					const FSolverReal ExpStiffnessValue = StiffnessHasWeightMap ? Stiffness[ConstraintIndex] : StiffnessNoMap;
+					const FSolverReal DampingRatioValue = DampingHasWeightMap ? DampingRatio[ConstraintIndex] : DampingNoMap;
+					ApplyHelper<bDampingBefore, bSingleLambda, bSeparateStretch, bDampingAfter>(Particles, Dt, ConstraintIndex, ExpStiffnessValue, DampingRatioValue);
+				}
+			}
+			break;
+			case (int32)EXPBDSplitDampingMode::InterleavedBefore:
+			{
+				// Do a pass with damping before stretch
+				for (int32 ConstraintIndex = 0; ConstraintIndex < Constraints.Num(); ++ConstraintIndex)
+				{
+					constexpr bool bDampingBefore = true;
+					constexpr bool bSingleLambda = false;
+					constexpr bool bSeparateStretch = true;
+					constexpr bool bDampingAfter = false;
+					const FSolverReal ExpStiffnessValue = StiffnessHasWeightMap ? Stiffness[ConstraintIndex] : StiffnessNoMap;
+					const FSolverReal DampingRatioValue = DampingHasWeightMap ? DampingRatio[ConstraintIndex] : DampingNoMap;
+					ApplyHelper<bDampingBefore, bSingleLambda, bSeparateStretch, bDampingAfter>(Particles, Dt, ConstraintIndex, ExpStiffnessValue, DampingRatioValue);
+				}
+			}
+			break;
+			case (int32)EXPBDSplitDampingMode::InterleavedAfter:
+			{
+				// Do a pass with damping after stretch
+				for (int32 ConstraintIndex = 0; ConstraintIndex < Constraints.Num(); ++ConstraintIndex)
+				{
+					constexpr bool bDampingBefore = false;
+					constexpr bool bSingleLambda = false;
+					constexpr bool bSeparateStretch = true;
+					constexpr bool bDampingAfter = true;
+					const FSolverReal ExpStiffnessValue = StiffnessHasWeightMap ? Stiffness[ConstraintIndex] : StiffnessNoMap;
+					const FSolverReal DampingRatioValue = DampingHasWeightMap ? DampingRatio[ConstraintIndex] : DampingNoMap;
+					ApplyHelper<bDampingBefore, bSingleLambda, bSeparateStretch, bDampingAfter>(Particles, Dt, ConstraintIndex, ExpStiffnessValue, DampingRatioValue);
+				}
+			}
+			break;
+			}
+			if (Chaos_XPBDSpring_SplitDampingMode == (int32)EXPBDSplitDampingMode::TwoPassAfter)
+			{
+				for (int32 ConstraintIndex = 0; ConstraintIndex < Constraints.Num(); ++ConstraintIndex)
+				{
+					constexpr bool bDampingBefore = false;
+					constexpr bool bSingleLambda = false;
+					constexpr bool bSeparateStretch = false;
+					constexpr bool bDampingAfter = true;
+					const FSolverReal ExpStiffnessValue = StiffnessHasWeightMap ? Stiffness[ConstraintIndex] : StiffnessNoMap;
+					const FSolverReal DampingRatioValue = DampingHasWeightMap ? DampingRatio[ConstraintIndex] : DampingNoMap;
+					ApplyHelper<bDampingBefore, bSingleLambda, bSeparateStretch, bDampingAfter>(Particles, Dt, ConstraintIndex, ExpStiffnessValue, DampingRatioValue);
+				}
 			}
 		}
 		else
 		{
-			const FSolverReal StiffnessNoMap = (FSolverReal)Stiffness;
-			const FSolverReal DampingNoMap = (FSolverReal)DampingRatio;
+			// No damping. 
 			for (int32 ConstraintIndex = 0; ConstraintIndex < Constraints.Num(); ++ConstraintIndex)
 			{
+				constexpr bool bDampingBefore = false;
+				constexpr bool bSingleLambda = false;
+				constexpr bool bSeparateStretch = true;
+				constexpr bool bDampingAfter = false;
 				const FSolverReal ExpStiffnessValue = StiffnessHasWeightMap ? Stiffness[ConstraintIndex] : StiffnessNoMap;
 				const FSolverReal DampingRatioValue = DampingHasWeightMap ? DampingRatio[ConstraintIndex] : DampingNoMap;
-				ApplyHelper(Particles, Dt, ConstraintIndex, ExpStiffnessValue, DampingRatioValue);
+				ApplyHelper<bDampingBefore, bSingleLambda, bSeparateStretch, bDampingAfter>(Particles, Dt, ConstraintIndex, ExpStiffnessValue, (FSolverReal)0.);
+
 			}
 		}
 	}
diff --git a/Engine/Source/Runtime/Experimental/Chaos/Private/Chaos/XPBDSpringConstraints.ispc b/Engine/Source/Runtime/Experimental/Chaos/Private/Chaos/XPBDSpringConstraints.ispc
--- a/Engine/Source/Runtime/Experimental/Chaos/Private/Chaos/XPBDSpringConstraints.ispc
+++ b/Engine/Source/Runtime/Experimental/Chaos/Private/Chaos/XPBDSpringConstraints.ispc
@@ -38,11 +38,11 @@ export void ApplyXPBDSpringConstraints(uniform FVector4f PandInvM[],
 
 		if (InvM1 > FLOAT_ZERO)
 		{
-			VectorScatter(&PandInvM[i1], SetVector4( P1 - (InvM1 * Delta), InvM1 ));
+			VectorScatter(&PandInvM[i1], SetVector4( P1 + (InvM1 * Delta), InvM1 ));
 		}
 		if (InvM2 > FLOAT_ZERO)
 		{
-			VectorScatter(&PandInvM[i2], SetVector4( P2 + (InvM2 * Delta), InvM2 ));
+			VectorScatter(&PandInvM[i2], SetVector4( P2 - (InvM2 * Delta), InvM2 ));
 		}
 	}
 }
@@ -84,11 +84,57 @@ export void ApplyXPBDSpringConstraintsWithDamping(uniform FVector4f PandInvM[],
 
 		if (InvM1 > FLOAT_ZERO)
 		{
-			VectorScatter(&PandInvM[i1], SetVector4( P1 - (InvM1 * Delta), InvM1 ));
+			VectorScatter(&PandInvM[i1], SetVector4( P1 + (InvM1 * Delta), InvM1 ));
 		}
 		if (InvM2 > FLOAT_ZERO)
 		{
-			VectorScatter(&PandInvM[i2], SetVector4( P2 + (InvM2 * Delta), InvM2 ));
+			VectorScatter(&PandInvM[i2], SetVector4( P2 - (InvM2 * Delta), InvM2 ));
+		}
+	}
+}
+
+export void ApplyXPBDSpringDampingConstraints(uniform FVector4f PandInvM[],
+												const uniform FVector3f X[],
+												const uniform FIntVector2 Constraints[],
+												const uniform float Dists[],
+												uniform float Lambdas[],
+												const uniform float Dt,
+												const uniform float Stiffness,
+												const uniform float DampingRatio,
+												const uniform int32 NumConstraints)
+{
+	foreach(i = 0 ... NumConstraints)
+	{
+		const varying FIntVector2 Constraint = VectorLoad(&Constraints[extract(i,0)]);
+		const varying int32 i1 = Constraint.V[0];
+		const varying int32 i2 = Constraint.V[1];
+
+		const varying FVector4f PandInvM1 = VectorGather(&PandInvM[i1]);
+		const varying FVector4f PandInvM2 = VectorGather(&PandInvM[i2]);
+
+		varying FVector3f P1, P2;
+		varying float InvM1, InvM2;
+		UnzipPandInvM(PandInvM1, P1, InvM1);
+		UnzipPandInvM(PandInvM2, P2, InvM2);
+
+		const varying FVector3f X1 = VectorGather(&X[i1]);
+		const varying FVector3f X2 = VectorGather(&X[i2]);
+
+		const varying float Dist = Dists[i];
+		const varying float Lambda = Lambdas[i];
+
+		varying float DLambda;
+		const FVector3f Delta = GetXPBDSpringDampingDelta(P1, P2, X1, X2, InvM1, InvM2, Dt, Dist, Lambda, Stiffness, DampingRatio, FLOAT_SMALL_NUMBER, DLambda);
+
+		Lambdas[i] = Lambda + DLambda;
+
+		if (InvM1 > FLOAT_ZERO)
+		{
+			VectorScatter(&PandInvM[i1], SetVector4( P1 + (InvM1 * Delta), InvM1 ));
+		}
+		if (InvM2 > FLOAT_ZERO)
+		{
+			VectorScatter(&PandInvM[i2], SetVector4( P2 - (InvM2 * Delta), InvM2 ));
 		}
 	}
 }
@@ -133,7 +179,7 @@ export void ApplyXPBDSpringConstraintsWithWeightMaps(uniform FVector4f PandInvM[
 
 		if (InvM1 > FLOAT_ZERO)
 		{
-			VectorScatter(&PandInvM[i1], SetVector4( P1 - (InvM1 * Delta), InvM1 ));
+			VectorScatter(&PandInvM[i1], SetVector4( P1 + (InvM1 * Delta), InvM1 ));
 		}
 		if (InvM2 > FLOAT_ZERO)
 		{
@@ -192,11 +238,70 @@ export void ApplyXPBDSpringConstraintsWithDampingAndWeightMaps(uniform FVector4f
 
 		if (InvM1 > FLOAT_ZERO)
 		{
-			VectorScatter(&PandInvM[i1], SetVector4( P1 - (InvM1 * Delta), InvM1 ));
+			VectorScatter(&PandInvM[i1], SetVector4( P1 + (InvM1 * Delta), InvM1 ));
 		}
 		if (InvM2 > FLOAT_ZERO)
 		{
-			VectorScatter(&PandInvM[i2], SetVector4( P2 + (InvM2 * Delta), InvM2 ));
+			VectorScatter(&PandInvM[i2], SetVector4( P2 - (InvM2 * Delta), InvM2 ));
 		}
 	}
 }
+
+export void ApplyXPBDSpringDampingConstraintsWithWeightMaps(uniform FVector4f PandInvM[],
+															const uniform FVector3f X[],
+															const uniform FIntVector2 Constraints[],
+															const uniform float Dists[],
+															uniform float Lambdas[],
+															const uniform float Dt,
+															const uniform float MinStiffness,
+															const uniform bool StiffnessHasMap,
+															const uniform uint8 StiffnessIndices[],
+															const uniform float StiffnessTable[],
+															const uniform bool DampingHasMap,
+															const uniform uint8 DampingIndices[],
+															const uniform float DampingTable[],
+															const uniform int32 NumConstraints)
+{
+
+	foreach(i = 0 ... NumConstraints)
+	{
+		const varying FIntVector2 Constraint = VectorLoad(&Constraints[extract(i,0)]);
+		const varying int32 i1 = Constraint.V[0];
+		const varying int32 i2 = Constraint.V[1];
+		const varying int8 StiffnessIndex = StiffnessHasMap ? StiffnessIndices[i] : 0;
+		const varying int8 DampingIndex = DampingHasMap ? DampingIndices[i] : 0;
+
+		const varying FVector4f PandInvM1 = VectorGather(&PandInvM[i1]);
+		const varying FVector4f PandInvM2 = VectorGather(&PandInvM[i2]);
+
+		varying FVector3f P1, P2;
+		varying float InvM1, InvM2;
+		UnzipPandInvM(PandInvM1, P1, InvM1);
+		UnzipPandInvM(PandInvM2, P2, InvM2);
+
+		const varying FVector3f X1 = VectorGather(&X[i1]);
+		const varying FVector3f X2 = VectorGather(&X[i2]);
+
+		const varying float Dist = Dists[i];
+		const varying float Lambda = Lambdas[i];
+
+		#pragma ignore warning(perf)
+		const varying float Stiffness = StiffnessTable[StiffnessIndex];
+		#pragma ignore warning(perf)
+		const varying float DampingRatio = DampingTable[DampingIndex];
+
+		varying float DLambda;
+		const FVector3f Delta = GetXPBDSpringDampingDelta(P1, P2, X1, X2, InvM1, InvM2, Dt, Dist, Lambda, Stiffness, DampingRatio, MinStiffness, DLambda);
+
+		Lambdas[i] = Lambda + DLambda;
+
+		if (InvM1 > FLOAT_ZERO)
+		{
+			VectorScatter(&PandInvM[i1], SetVector4( P1 + (InvM1 * Delta), InvM1 ));
+		}
+		if (InvM2 > FLOAT_ZERO)
+		{
+			VectorScatter(&PandInvM[i2], SetVector4( P2 - (InvM2 * Delta), InvM2 ));
+		}
+	}
+}
\ No newline at end of file
diff --git a/Engine/Source/Runtime/Experimental/Chaos/Public/Chaos/SoftsSpring.h b/Engine/Source/Runtime/Experimental/Chaos/Public/Chaos/SoftsSpring.h
--- a/Engine/Source/Runtime/Experimental/Chaos/Public/Chaos/SoftsSpring.h
+++ b/Engine/Source/Runtime/Experimental/Chaos/Public/Chaos/SoftsSpring.h
@@ -1,8 +1,6 @@
 // Copyright Epic Games, Inc. All Rights Reserved.
 #pragma once
-
 #include "Chaos/SoftsEvolutionLinearSystem.h"
-#include "Chaos/SoftsSolverParticlesRange.h"
 
 namespace Chaos::Softs
 {
@@ -10,10 +8,11 @@ namespace Chaos::Softs
 namespace Spring
 {
 
+// Spring without damping
 template<typename SolverParticlesOrRange>
 FSolverVec3 GetXPBDSpringDelta(const SolverParticlesOrRange& Particles, const FSolverReal Dt,
 	const TVec2<int32>& Constraint, const FSolverReal RestLength, FSolverReal& Lambda,
-	const FSolverReal StiffnessValue, const FSolverReal MinStiffness, const FSolverReal DampingRatioValue)
+	const FSolverReal StiffnessValue, const FSolverReal MinStiffness)
 {
 	const int32 Index1 = Constraint[0];
 	const int32 Index2 = Constraint[1];
@@ -24,6 +23,37 @@ FSolverVec3 GetXPBDSpringDelta(const SolverParticlesOrRange& Particles, const FS
 	}
 
 	const FSolverReal CombinedInvMass = Particles.InvM(Index2) + Particles.InvM(Index1);
+
+	const FSolverVec3& P1 = Particles.P(Index1);
+	const FSolverVec3& P2 = Particles.P(Index2);
+	FSolverVec3 Direction = P1 - P2;
+	const FSolverReal Distance = Direction.SafeNormalize();
+	const FSolverReal Offset = Distance - RestLength;
+
+	const FSolverReal AlphaInv = StiffnessValue * Dt * Dt;
+
+	const FSolverReal DLambda = (-AlphaInv * Offset - Lambda) / (AlphaInv * CombinedInvMass + (FSolverReal)1.);
+	const FSolverVec3 Delta = DLambda * Direction;
+	Lambda += DLambda;
+
+	return Delta;
+}
+
+// This is a following the original XPBD paper using a single lambda for stretch and damping.
+template<typename SolverParticlesOrRange>
+FSolverVec3 GetXPBDSpringDeltaWithDamping(const SolverParticlesOrRange& Particles, const FSolverReal Dt,
+	const TVec2<int32>& Constraint, const FSolverReal RestLength, FSolverReal& Lambda,
+	const FSolverReal StiffnessValue, const FSolverReal MinStiffness, const FSolverReal DampingRatioValue)
+{
+	const int32 Index1 = Constraint[0];
+	const int32 Index2 = Constraint[1];
+
+	const FSolverReal CombinedInvMass = Particles.InvM(Index2) + Particles.InvM(Index1);
+	if (StiffnessValue < MinStiffness || CombinedInvMass < UE_SMALL_NUMBER)
+	{
+		return FSolverVec3((FSolverReal)0.);
+	}
+
 	const FSolverReal Damping = DampingRatioValue * 2.f * FMath::Sqrt(StiffnessValue / CombinedInvMass) * (RestLength > UE_SMALL_NUMBER ? (FSolverReal)1. / RestLength : (FSolverReal)1.);
 
 	const FSolverVec3& P1 = Particles.P(Index1);
@@ -40,17 +70,49 @@ FSolverVec3 GetXPBDSpringDelta(const SolverParticlesOrRange& Particles, const FS
 	const FSolverReal AlphaInv = StiffnessValue * Dt * Dt;
 	const FSolverReal BetaDt = Damping * Dt;
 
-	const FSolverReal DLambda = (AlphaInv * Offset - Lambda + BetaDt * FSolverVec3::DotProduct(Direction, RelativeVelocityTimesDt)) / ((AlphaInv + BetaDt) * CombinedInvMass + (FSolverReal)1.);
+	const FSolverReal DLambda = (-AlphaInv * Offset - Lambda - BetaDt * FSolverVec3::DotProduct(Direction, RelativeVelocityTimesDt)) / ((AlphaInv + BetaDt) * CombinedInvMass + (FSolverReal)1.);
 	const FSolverVec3 Delta = DLambda * Direction;
 	Lambda += DLambda;
 
 	return Delta;
 }
 
+// Spring damping constraint (separate from spring stretching)
+template<typename SolverParticlesOrRange>
+FSolverVec3 GetXPBDSpringDampingDelta(const SolverParticlesOrRange& Particles, const FSolverReal Dt,
+	const TVec2<int32>& Constraint, const FSolverReal RestLength, FSolverReal& Lambda,
+	const FSolverReal StiffnessValue, const FSolverReal MinStiffness, const FSolverReal DampingRatioValue)
+{
+	const int32 Index1 = Constraint[0];
+	const int32 Index2 = Constraint[1];
+	const FSolverReal CombinedInvMass = Particles.InvM(Index2) + Particles.InvM(Index1);
+	if (StiffnessValue < MinStiffness || CombinedInvMass < UE_SMALL_NUMBER)
+	{
+		return FSolverVec3((FSolverReal)0.);
+	}
+
+	const FSolverReal Damping = DampingRatioValue * 2.f * FMath::Sqrt(StiffnessValue / CombinedInvMass) * (RestLength > UE_SMALL_NUMBER ? (FSolverReal)1. / RestLength : (FSolverReal)1.);
+
+	const FSolverVec3& P1 = Particles.P(Index1);
+	const FSolverVec3& P2 = Particles.P(Index2);
+	FSolverVec3 Direction = (P1 - P2);
+	Direction.SafeNormalize();
+
+	const FSolverVec3& X1 = Particles.X(Index1);
+	const FSolverVec3& X2 = Particles.X(Index2);
+	const FSolverVec3 RelativeVelocityTimesDt = P1 - X1 - P2 + X2;
+	const FSolverReal BetaDt = Damping * Dt;
+	const FSolverReal DLambda = (-BetaDt * FSolverVec3::DotProduct(Direction, RelativeVelocityTimesDt) - Lambda) / (BetaDt * CombinedInvMass + (FSolverReal)1.);
+	const FSolverVec3 Delta = DLambda * Direction;
+	Lambda += DLambda;
+	return Delta;
+}
+
+// Spring without damping
 template<typename SolverParticlesOrRange>
 FSolverVec3 GetXPBDAxialSpringDelta(const SolverParticlesOrRange& Particles, const FSolverReal Dt,
 	const TVec3<int32>& Constraint, const FSolverReal Bary, const FSolverReal RestLength, FSolverReal& Lambda,
-	const FSolverReal StiffnessValue, const FSolverReal MinStiffness, const FSolverReal DampingRatioValue)
+	const FSolverReal StiffnessValue, const FSolverReal MinStiffness)
 {
 	const int32 Index1 = Constraint[0];
 	const int32 Index2 = Constraint[1];
@@ -63,6 +125,41 @@ FSolverVec3 GetXPBDAxialSpringDelta(const SolverParticlesOrRange& Particles, con
 	}
 
 	const FSolverReal CombinedInvMass = PInvMass + Particles.InvM(Index1);
+
+	const FSolverVec3& P1 = Particles.P(Index1);
+	const FSolverVec3& P2 = Particles.P(Index2);
+	const FSolverVec3& P3 = Particles.P(Index3);
+	const FSolverVec3 P = (P2 - P3) * Bary + P3;
+	FSolverVec3 Direction = P1 - P;
+	const FSolverReal Distance = Direction.SafeNormalize();
+	const FSolverReal Offset = Distance - RestLength;
+
+	const FSolverReal AlphaInv = StiffnessValue * Dt * Dt;
+
+	const FSolverReal DLambda = (-AlphaInv * Offset - Lambda) / (AlphaInv * CombinedInvMass + (FSolverReal)1.);
+	const FSolverVec3 Delta = DLambda * Direction;
+	Lambda += DLambda;
+
+	return Delta;
+}
+
+// This is a following the original XPBD paper using a single lambda for stretch and damping.
+template<typename SolverParticlesOrRange>
+FSolverVec3 GetXPBDAxialSpringDeltaWithDamping(const SolverParticlesOrRange& Particles, const FSolverReal Dt,
+	const TVec3<int32>& Constraint, const FSolverReal Bary, const FSolverReal RestLength, FSolverReal& Lambda,
+	const FSolverReal StiffnessValue, const FSolverReal MinStiffness, const FSolverReal DampingRatioValue)
+{
+	const int32 Index1 = Constraint[0];
+	const int32 Index2 = Constraint[1];
+	const int32 Index3 = Constraint[2];
+
+	const FSolverReal PInvMass = Particles.InvM(Index3) * ((FSolverReal)1. - Bary) + Particles.InvM(Index2) * Bary;
+	const FSolverReal CombinedInvMass = PInvMass + Particles.InvM(Index1);
+	if (StiffnessValue < MinStiffness || CombinedInvMass < UE_SMALL_NUMBER)
+	{
+		return FSolverVec3((FSolverReal)0.);
+	}
+
 	const FSolverReal Damping = DampingRatioValue * 2.f * FMath::Sqrt(StiffnessValue / CombinedInvMass) * (RestLength > UE_SMALL_NUMBER ? (FSolverReal)1. / RestLength : (FSolverReal)1.);
 
 	const FSolverVec3& P1 = Particles.P(Index1);
@@ -83,7 +180,47 @@ FSolverVec3 GetXPBDAxialSpringDelta(const SolverParticlesOrRange& Particles, con
 	const FSolverReal AlphaInv = StiffnessValue * Dt * Dt;
 	const FSolverReal BetaDt = Damping * Dt;
 
-	const FSolverReal DLambda = (AlphaInv * Offset - Lambda + BetaDt * FSolverVec3::DotProduct(Direction, RelativeVelocityTimesDt)) / ((AlphaInv + BetaDt) * CombinedInvMass + (FSolverReal)1.);
+	const FSolverReal DLambda = (-AlphaInv * Offset - Lambda - BetaDt * FSolverVec3::DotProduct(Direction, RelativeVelocityTimesDt)) / ((AlphaInv + BetaDt) * CombinedInvMass + (FSolverReal)1.);
+	const FSolverVec3 Delta = DLambda * Direction;
+	Lambda += DLambda;
+
+	return Delta;
+}
+
+// Spring damping constraint (separate from spring stretching)
+template<typename SolverParticlesOrRange>
+FSolverVec3 GetXPBDAxialSpringDampingDelta(const SolverParticlesOrRange& Particles, const FSolverReal Dt,
+	const TVec3<int32>& Constraint, const FSolverReal Bary, const FSolverReal RestLength, FSolverReal& Lambda,
+	const FSolverReal StiffnessValue, const FSolverReal MinStiffness, const FSolverReal DampingRatioValue)
+{
+	const int32 Index1 = Constraint[0];
+	const int32 Index2 = Constraint[1];
+	const int32 Index3 = Constraint[2];
+
+	const FSolverReal PInvMass = Particles.InvM(Index3) * ((FSolverReal)1. - Bary) + Particles.InvM(Index2) * Bary;
+	const FSolverReal CombinedInvMass = PInvMass + Particles.InvM(Index1);
+	if (StiffnessValue < MinStiffness || CombinedInvMass < UE_SMALL_NUMBER)
+	{
+		return FSolverVec3((FSolverReal)0.);
+	}
+
+	const FSolverReal Damping = DampingRatioValue * 2.f * FMath::Sqrt(StiffnessValue / CombinedInvMass) * (RestLength > UE_SMALL_NUMBER ? (FSolverReal)1. / RestLength : (FSolverReal)1.);
+
+	const FSolverVec3& P1 = Particles.P(Index1);
+	const FSolverVec3& P2 = Particles.P(Index2);
+	const FSolverVec3& P3 = Particles.P(Index3);
+	const FSolverVec3 P = (P2 - P3) * Bary + P3;
+	FSolverVec3 Direction = P1 - P;
+	const FSolverReal Distance = Direction.SafeNormalize();
+	const FSolverReal Offset = Distance - RestLength;
+
+	const FSolverVec3& X1 = Particles.X(Index1);
+	const FSolverVec3& X2 = Particles.X(Index2);
+	const FSolverVec3& X3 = Particles.X(Index3);
+	const FSolverVec3 X = (X2 - X3) * Bary + X3;
+	const FSolverVec3 RelativeVelocityTimesDt = P1 - X1 - P + X;
+	const FSolverReal BetaDt = Damping * Dt;
+	const FSolverReal DLambda = (-BetaDt * FSolverVec3::DotProduct(Direction, RelativeVelocityTimesDt) - Lambda) / (BetaDt * CombinedInvMass + (FSolverReal)1.);
 	const FSolverVec3 Delta = DLambda * Direction;
 	Lambda += DLambda;
 
diff --git a/Engine/Source/Runtime/Experimental/Chaos/Public/Chaos/XPBDAnisotropicBendingConstraints.h b/Engine/Source/Runtime/Experimental/Chaos/Public/Chaos/XPBDAnisotropicBendingConstraints.h
--- a/Engine/Source/Runtime/Experimental/Chaos/Public/Chaos/XPBDAnisotropicBendingConstraints.h
+++ b/Engine/Source/Runtime/Experimental/Chaos/Public/Chaos/XPBDAnisotropicBendingConstraints.h
@@ -74,6 +74,8 @@ public:
 	{ 
 		Lambdas.Reset();
 		Lambdas.AddZeroed(Constraints.Num());
+		LambdasDamping.Reset();
+		LambdasDamping.AddZeroed(Constraints.Num());
 		FPBDBendingConstraintsBase::Init(InParticles);
 	}
 
@@ -128,7 +130,7 @@ public:
 private:
 	template<typename SolverParticlesOrRange>
 	void InitColor(const SolverParticlesOrRange& InParticles);
-	template<typename SolverParticlesOrRange>
+	template<bool bDampingOnly, bool bElasticOnly, typename SolverParticlesOrRange>
 	void ApplyHelper(SolverParticlesOrRange& Particles, const FSolverReal Dt, const int32 ConstraintIndex, const FSolverVec3& ExpStiffnessValues,
 		const FSolverVec3& ExpBucklingStiffnessValues, const FSolverReal DampingRatioValue) const;
 
@@ -183,6 +185,7 @@ private:
 
 	FPBDWeightMap DampingRatio;
 	mutable TArray<FSolverReal> Lambdas;
+	mutable TArray<FSolverReal> LambdasDamping;
 	TArray<int32> ConstraintsPerColorStartIndex; // Constraints are ordered so each batch is contiguous. This is ColorNum + 1 length so it can be used as start and end.
 
 	TArray<FSolverVec3> WarpWeftBiasBaseMultipliers;
diff --git a/Engine/Source/Runtime/Experimental/Chaos/Public/Chaos/XPBDAnisotropicSpringConstraints.h b/Engine/Source/Runtime/Experimental/Chaos/Public/Chaos/XPBDAnisotropicSpringConstraints.h
--- a/Engine/Source/Runtime/Experimental/Chaos/Public/Chaos/XPBDAnisotropicSpringConstraints.h
+++ b/Engine/Source/Runtime/Experimental/Chaos/Public/Chaos/XPBDAnisotropicSpringConstraints.h
@@ -75,6 +75,8 @@ public:
 	{
 		Lambdas.Reset();
 		Lambdas.AddZeroed(Constraints.Num());
+		LambdasDamping.Reset();
+		LambdasDamping.AddZeroed(Constraints.Num());
 	}
 
 	// Update stiffness table, as well as the simulation stiffness exponent
@@ -111,7 +113,7 @@ private:
 
 	CHAOS_API void UpdateDists();
 	
-	template<typename SolverParticlesOrRange>
+	template<bool bDampingBefore, bool bSingleLambda, bool bSeparateStretch, bool bDampingAfter, typename SolverParticlesOrRange>
 	void ApplyHelper(SolverParticlesOrRange& Particles, const FSolverReal Dt, const int32 ConstraintIndex, const FSolverVec3& StiffnessValue, const FSolverReal DampingRatioValue) const;
 
 	using Base::Constraints;
@@ -127,6 +129,7 @@ private:
 	FPBDWeightMap WeftScale;
 
 	mutable TArray<FSolverReal> Lambdas;
+	mutable TArray<FSolverReal> LambdasDamping;
 	TArray<FSolverReal> BaseDists; // Without Warp/Weft Scale applied
 	TArray<FSolverVec3> WarpWeftBiasBaseMultipliers;
 	TArray<FSolverVec2> WarpWeftScaleBaseMultipliers;
@@ -194,6 +197,8 @@ public:
 		{
 			Lambdas.Reset();
 			Lambdas.AddZeroed(Constraints.Num());
+			LambdasDamping.Reset();
+			LambdasDamping.AddZeroed(Constraints.Num());
 		}
 
 		// Update stiffness table, as well as the simulation stiffness exponent
@@ -230,7 +235,7 @@ private:
 
 	CHAOS_API void UpdateDists();
 
-	template<typename SolverParticlesOrRange>
+	template<bool bDampingBefore, bool bSingleLambda, bool bSeparateStretch, bool bDampingAfter, typename SolverParticlesOrRange>
 	void ApplyHelper(SolverParticlesOrRange& Particles, const FSolverReal Dt, const int32 ConstraintIndex, const FSolverVec3& StiffnessValue, const FSolverReal DampingRatioValue) const;
 
 	using Base::Constraints;
@@ -247,6 +252,7 @@ private:
 	FPBDWeightMap WeftScale;
 
 	mutable TArray<FSolverReal> Lambdas;
+	mutable TArray<FSolverReal> LambdasDamping;
 	TArray<FSolverReal> BaseDists; // Without Warp/Weft Scale applied
 	TArray<FSolverVec3> WarpWeftBiasBaseMultipliers;
 	TArray<FSolverVec2> WarpWeftScaleBaseMultipliers;
diff --git a/Engine/Source/Runtime/Experimental/Chaos/Public/Chaos/XPBDBendingConstraints.h b/Engine/Source/Runtime/Experimental/Chaos/Public/Chaos/XPBDBendingConstraints.h
--- a/Engine/Source/Runtime/Experimental/Chaos/Public/Chaos/XPBDBendingConstraints.h
+++ b/Engine/Source/Runtime/Experimental/Chaos/Public/Chaos/XPBDBendingConstraints.h
@@ -63,6 +63,7 @@ public:
 		, XPBDRestAngleTypeIndex(PropertyCollection)
 	{
 		Lambdas.Init((FSolverReal)0., Constraints.Num());
+		LambdasDamping.Init((FSolverReal)0., Constraints.Num());
 		InitColor(InParticles);
 	}
 
@@ -103,6 +104,7 @@ public:
 		, XPBDRestAngleTypeIndex(PropertyCollection)
 	{
 		Lambdas.Init((FSolverReal)0., Constraints.Num());
+		LambdasDamping.Init((FSolverReal)0., Constraints.Num());
 		InitColor(InParticles);
 	}
 
@@ -143,6 +145,7 @@ public:
 		, XPBDRestAngleTypeIndex(PropertyCollection)
 	{
 		Lambdas.Init((FSolverReal)0., Constraints.Num());
+		LambdasDamping.Init((FSolverReal)0., Constraints.Num());
 		InitColor(InParticles);
 	}
 
@@ -179,6 +182,7 @@ public:
 		, XPBDRestAngleTypeIndex(ForceInit)
 	{
 		Lambdas.Init((FSolverReal)0., Constraints.Num());
+		LambdasDamping.Init((FSolverReal)0., Constraints.Num());
 		InitColor(InParticles);
 	}
 
@@ -221,6 +225,7 @@ public:
 		, XPBDRestAngleTypeIndex(ForceInit)
 	{
 		Lambdas.Init((FSolverReal)0., Constraints.Num());
+		LambdasDamping.Init((FSolverReal)0., Constraints.Num());
 		InitColor(InParticles);
 	}
 
@@ -231,6 +236,8 @@ public:
 	{ 
 		Lambdas.Reset();
 		Lambdas.AddZeroed(Constraints.Num());
+		LambdasDamping.Reset();
+		LambdasDamping.AddZeroed(Constraints.Num());
 		FPBDBendingConstraintsBase::Init(InParticles);
 	}
 
@@ -288,7 +295,7 @@ public:
 private:
 	template<typename SolverParticlesOrRange>
 	CHAOS_API void InitColor(const SolverParticlesOrRange& InParticles);
-	template<typename SolverParticlesOrRange>
+	template<bool bDampingOnly, bool bElasticOnly, typename SolverParticlesOrRange>
 	void ApplyHelper(SolverParticlesOrRange& Particles, const FSolverReal Dt, const int32 ConstraintIndex, const FSolverReal ExpStiffnessValue, const FSolverReal ExpBucklingValue, const FSolverReal DampingRatioValue) const;
 
 	TConstArrayView<FRealSingle> GetRestAngleMapFromCollection(
@@ -334,6 +341,7 @@ private:
 
 	FPBDWeightMap DampingRatio;
 	mutable TArray<FSolverReal> Lambdas;
+	mutable TArray<FSolverReal> LambdasDamping;
 	TArray<int32> ConstraintsPerColorStartIndex; // Constraints are ordered so each batch is contiguous. This is ColorNum + 1 length so it can be used as start and end.
 
 	UE_CHAOS_DECLARE_PROPERTYCOLLECTION_NAME(XPBDBendingElementStiffness, float);
@@ -345,6 +353,13 @@ private:
 	UE_CHAOS_DECLARE_PROPERTYCOLLECTION_NAME(XPBDRestAngleType, int32);
 };
 
+// Support split vs shared damping models in non-shipping builds
+#if UE_BUILD_SHIPPING
+const bool bChaos_XPBDBending_SplitLambdaDamping = true;
+#else
+extern bool bChaos_XPBDBending_SplitLambdaDamping;
+#endif
+
 }  // End namespace Chaos::Softs
 
 // Support ISPC enable/disable in non-shipping builds
diff --git a/Engine/Source/Runtime/Experimental/Chaos/Public/Chaos/XPBDSpringConstraints.h b/Engine/Source/Runtime/Experimental/Chaos/Public/Chaos/XPBDSpringConstraints.h
--- a/Engine/Source/Runtime/Experimental/Chaos/Public/Chaos/XPBDSpringConstraints.h
+++ b/Engine/Source/Runtime/Experimental/Chaos/Public/Chaos/XPBDSpringConstraints.h
@@ -107,6 +107,8 @@ public:
 	{
 		Lambdas.Reset();
 		Lambdas.SetNumZeroed(Constraints.Num());
+		LambdasDamping.Reset();
+		LambdasDamping.SetNumZeroed(Constraints.Num());
 		InitColor(Particles);
 	}
 
@@ -116,6 +118,8 @@ public:
 	{
 		Lambdas.Reset();
 		Lambdas.SetNumZeroed(Constraints.Num());
+		LambdasDamping.Reset();
+		LambdasDamping.SetNumZeroed(Constraints.Num());
 	}
 
 	// Update stiffness values
@@ -142,14 +146,9 @@ public:
 private:
 	template<typename SolverParticlesOrRange>
 	CHAOS_API void InitColor(const SolverParticlesOrRange& InParticles);
-	template<typename SolverParticlesOrRange>
-	void ApplyHelper(SolverParticlesOrRange& Particles, const FSolverReal Dt, const int32 ConstraintIndex, const FSolverReal StiffnessValue, const FSolverReal DampingRatioValue) const; 
 
-	template<typename SolverParticlesOrRange>
-	FSolverVec3 GetDelta(const SolverParticlesOrRange& Particles, const FSolverReal Dt, const int32 ConstraintIndex, const FSolverReal StiffnessValue, const FSolverReal DampingRatioValue) const
-	{
-		return Spring::GetXPBDSpringDelta(Particles, Dt, Constraints[ConstraintIndex], Dists[ConstraintIndex], Lambdas[ConstraintIndex], StiffnessValue, MinStiffness, DampingRatioValue);
-	}
+	template<bool bDampingBefore, bool bSingleLambda, bool bSeparateStretch, bool bDampingAfter, typename SolverParticlesOrRange >
+	void ApplyHelper(SolverParticlesOrRange& Particles, const FSolverReal Dt, const int32 ConstraintIndex, const FSolverReal StiffnessValue, const FSolverReal DampingRatioValue) const; 
 
 protected:
 	using Base::Constraints;
@@ -161,6 +160,7 @@ protected:
 private:
 	using Base::Dists;
 	mutable TArray<FSolverReal> Lambdas;
+	mutable TArray<FSolverReal> LambdasDamping;
 	TArray<int32> ConstraintsPerColorStartIndex; // Constraints are ordered so each batch is contiguous. This is ColorNum + 1 length so it can be used as start and end.
 };
 
