Fixing assignment causes data loss warnings.

diff --git a/Engine/Source/Runtime/Experimental/Chaos/Public/Chaos/SoftsExternalForces.h b/Engine/Source/Runtime/Experimental/Chaos/Public/Chaos/SoftsExternalForces.h
--- a/Engine/Source/Runtime/Experimental/Chaos/Public/Chaos/SoftsExternalForces.h
+++ b/Engine/Source/Runtime/Experimental/Chaos/Public/Chaos/SoftsExternalForces.h
@@ -92,8 +92,8 @@ public:
 					if (Direction.Normalize())
 					{
 						// Scale by angle
-						const FReal DirectionDot = FVec3::DotProduct(Direction, N[Index]);
-						const FReal ScaleFactor = FMath::Min(1.f, FMath::Abs(DirectionDot) * LegacyWindAdaptation);
+						const FSolverReal DirectionDot = FSolverVec3::DotProduct(Direction, N[Index]);
+						const FSolverReal ScaleFactor = FMath::Min(1.f, FMath::Abs(DirectionDot) * LegacyWindAdaptation);
 						Force += VelocityDelta * ScaleFactor * M[Index];
 					}
 				}
diff --git a/Engine/Source/Runtime/Networking/Public/Common/TcpListener.h b/Engine/Source/Runtime/Networking/Public/Common/TcpListener.h
--- a/Engine/Source/Runtime/Networking/Public/Common/TcpListener.h
+++ b/Engine/Source/Runtime/Networking/Public/Common/TcpListener.h
@@ -189,7 +189,7 @@ public:
 			}
 			else
 			{
-				FPlatformProcess::Sleep(SleepTime.GetSeconds());
+				FPlatformProcess::Sleep(static_cast<float>(SleepTime.GetSeconds()));
 			}
 		}
 
diff --git a/Engine/Source/Runtime/Networking/Public/Common/TcpSocketBuilder.h b/Engine/Source/Runtime/Networking/Public/Common/TcpSocketBuilder.h
--- a/Engine/Source/Runtime/Networking/Public/Common/TcpSocketBuilder.h
+++ b/Engine/Source/Runtime/Networking/Public/Common/TcpSocketBuilder.h
@@ -135,7 +135,7 @@ public:
 	 * @return This instance (for method chaining).
 	 * @see BoundToAddress, BoundToEndpoint
 	 */
-	FTcpSocketBuilder BoundToPort(int32 Port)
+	FTcpSocketBuilder BoundToPort(uint16 Port)
 	{
 		BoundEndpoint = FIPv4Endpoint(BoundEndpoint.Address, Port);
 		Bound = true;
diff --git a/Engine/Source/Runtime/Networking/Public/Common/UdpSocketBuilder.h b/Engine/Source/Runtime/Networking/Public/Common/UdpSocketBuilder.h
--- a/Engine/Source/Runtime/Networking/Public/Common/UdpSocketBuilder.h
+++ b/Engine/Source/Runtime/Networking/Public/Common/UdpSocketBuilder.h
@@ -121,7 +121,7 @@ public:
 	 * @return This instance (for method chaining).
 	 * @see BoundToAddress
 	 */
-	FUdpSocketBuilder& BoundToPort(int32 Port)
+	FUdpSocketBuilder& BoundToPort(uint16 Port)
 	{
 		BoundEndpoint = FIPv4Endpoint(BoundEndpoint.Address, Port);
 		Bound = true;
