Fix particle index access

diff --git a/Engine/Source/Runtime/Experimental/Chaos/Private/PhysicsProxy/GeometryCollectionPhysicsProxy.cpp b/Engine/Source/Runtime/Experimental/Chaos/Private/PhysicsProxy/GeometryCollectionPhysicsProxy.cpp
--- a/Engine/Source/Runtime/Experimental/Chaos/Private/PhysicsProxy/GeometryCollectionPhysicsProxy.cpp
+++ b/Engine/Source/Runtime/Experimental/Chaos/Private/PhysicsProxy/GeometryCollectionPhysicsProxy.cpp
@@ -1105,9 +1105,10 @@ Chaos::FPBDRigidClusteredParticleHandle* FGeometryCollectionPhysicsProxy::FindCl
 	else
 	{
 		const int32 TransformIndex = ItemIndex.GetTransformIndex();
-		if (SolverParticleHandles.IsValidIndex(TransformIndex))
+		int32 ParticleIndex = FromTransformToParticleIndex[TransformIndex];
+		if (SolverParticleHandles.IsValidIndex(ParticleIndex))
 		{
-			ResultHandle = SolverParticleHandles[TransformIndex];
+			ResultHandle = SolverParticleHandles[ParticleIndex];
 		}
 	}
 	return ResultHandle;
