Chaos - Fix one-way interaction following geometry collection particle array packing

diff --git a/Engine/Source/Runtime/Experimental/Chaos/Private/Chaos/PBDRigidClustering.cpp b/Engine/Source/Runtime/Experimental/Chaos/Private/Chaos/PBDRigidClustering.cpp
--- a/Engine/Source/Runtime/Experimental/Chaos/Private/Chaos/PBDRigidClustering.cpp
+++ b/Engine/Source/Runtime/Experimental/Chaos/Private/Chaos/PBDRigidClustering.cpp
@@ -280,6 +280,11 @@ namespace Chaos
 		{
 			NewParticle->SetInternalStrains(NewParticle->GetInternalStrains() / static_cast<FRealSingle>(ChildrenSet.Num()));
 			UpdateTopLevelParticle(NewParticle);
+
+			// NOTE: These property values are only known when we have children. They should be overwritten when 
+			// children are added to an empty cluster, but we also shouldn't set non-default values before then
+			NewParticle->SetSleeping(bClusterIsAsleep);
+			NewParticle->SetOneWayInteraction(bClusterIsOneWayInteraction);
 		}
 
 		if (ForceMassOrientation)
@@ -302,10 +307,6 @@ namespace Chaos
 		// Build the convex optimizer if required
 		//FRigidClustering::BuildConvexOptimizer(NewParticle);
 
-		NewParticle->SetSleeping(bClusterIsAsleep);
-
-		NewParticle->SetOneWayInteraction(bClusterIsOneWayInteraction);
-
 		auto AddToClusterUnion = [this](int32 ClusterID, FPBDRigidClusteredParticleHandle* Handle)
 		{
 			if (ClusterID <= 0)
diff --git a/Engine/Source/Runtime/Experimental/Chaos/Private/PhysicsProxy/GeometryCollectionPhysicsProxy.cpp b/Engine/Source/Runtime/Experimental/Chaos/Private/PhysicsProxy/GeometryCollectionPhysicsProxy.cpp
--- a/Engine/Source/Runtime/Experimental/Chaos/Private/PhysicsProxy/GeometryCollectionPhysicsProxy.cpp
+++ b/Engine/Source/Runtime/Experimental/Chaos/Private/PhysicsProxy/GeometryCollectionPhysicsProxy.cpp
@@ -3542,7 +3542,7 @@ void FGeometryCollectionPhysicsProxy::SetOneWayInteractionLevel_Internal(int32 I
 		Chaos::FPBDRigidClusteredParticleHandle* Handle = SolverParticleHandles[ParticleIndex];
 		if (Handle)
 		{
-			const bool bIsOneWayInteraction = (Parameters.OneWayInteractionLevel >= 0) && (Level[ParticleIndex] >= Parameters.OneWayInteractionLevel);
+			const bool bIsOneWayInteraction = (Parameters.OneWayInteractionLevel >= 0) && (Level[FromParticleToTransformIndex[ParticleIndex]] >= Parameters.OneWayInteractionLevel);
 			Handle->SetOneWayInteraction(bIsOneWayInteraction);
 		}
 	}
