Geometry collection: fix ivisible geometry collection when using ISM Pool but no root proxy meshes

diff --git a/Engine/Source/Runtime/Experimental/GeometryCollectionEngine/Private/GeometryCollection/GeometryCollectionComponent.cpp b/Engine/Source/Runtime/Experimental/GeometryCollectionEngine/Private/GeometryCollection/GeometryCollectionComponent.cpp
--- a/Engine/Source/Runtime/Experimental/GeometryCollectionEngine/Private/GeometryCollection/GeometryCollectionComponent.cpp
+++ b/Engine/Source/Runtime/Experimental/GeometryCollectionEngine/Private/GeometryCollection/GeometryCollectionComponent.cpp
@@ -6107,7 +6107,7 @@ void UGeometryCollectionComponent::RefreshCustomRenderer()
 					}
 				#endif						
 
-					const bool bRenderRootProxy = bEnableRootProxyForCustomRenderer && !bIsBroken;
+					const bool bRenderRootProxy = bEnableRootProxyForCustomRenderer && !bIsBroken && (RestCollection->RootProxyData.ProxyMeshes.Num() > 0);
 
 					uint32 StateFlags = 0;
 					StateFlags |= bHiddenInGame ? 0 : IGeometryCollectionExternalRenderInterface::EState_Visible;
