[Iris]
* Fixed non-unity compile error

diff --git a/Engine/Source/Runtime/Experimental/Iris/Core/Private/Iris/ReplicationSystem/NetRefHandleManager.cpp b/Engine/Source/Runtime/Experimental/Iris/Core/Private/Iris/ReplicationSystem/NetRefHandleManager.cpp
--- a/Engine/Source/Runtime/Experimental/Iris/Core/Private/Iris/ReplicationSystem/NetRefHandleManager.cpp
+++ b/Engine/Source/Runtime/Experimental/Iris/Core/Private/Iris/ReplicationSystem/NetRefHandleManager.cpp
@@ -8,6 +8,7 @@
 #include "Iris/Core/IrisProfiler.h"
 #include "Iris/Serialization/NetSerializationContext.h"
 #include "Iris/Serialization/InternalNetSerializationContext.h"
+#include "Net/Core/Trace/NetDebugName.h"
 #include "Net/Core/Trace/NetTrace.h"
 #include "UObject/CoreNetTypes.h"
 #include "UObject/Object.h"
