Iris
- Changed assert to ensure and added some additional logging

diff --git a/Engine/Source/Runtime/Experimental/Iris/Core/Private/Iris/ReplicationSystem/ReplicationWriter.cpp b/Engine/Source/Runtime/Experimental/Iris/Core/Private/Iris/ReplicationSystem/ReplicationWriter.cpp
--- a/Engine/Source/Runtime/Experimental/Iris/Core/Private/Iris/ReplicationSystem/ReplicationWriter.cpp
+++ b/Engine/Source/Runtime/Experimental/Iris/Core/Private/Iris/ReplicationSystem/ReplicationWriter.cpp
@@ -935,8 +935,13 @@ void FReplicationWriter::HandleDeliveredRecord(const FReplicationRecord::FRecord
 	EReplicatedObjectState CurrentState = Info.GetState();
 	const uint32 InternalIndex = RecordInfo.Index;
 	
-	checkf(CurrentState != EReplicatedObjectState::Invalid, TEXT("Object ( InternalIndex: %u ) has an invalid state."), InternalIndex);
-
+	if (CurrentState == EReplicatedObjectState::Invalid)
+	{
+		ensure(CurrentState != EReplicatedObjectState::Invalid);
+		UE_LOG_REPLICATIONWRITER_WARNING(TEXT("FReplicationWriter::HandleDeliveredRecord - Warning Object ( InternalIndex: %u ) is invalid. DeliveredState %s WasDestroySubObject: %u"), InternalIndex, LexToString(DeliveredState), RecordInfo.WroteDestroySubObject)
+		return;
+	}
+	
 	// We confirmed a new baseline
 	if (RecordInfo.NewBaselineIndex != FDeltaCompressionBaselineManager::InvalidBaselineIndex)
 	{
