GeometryCollection Indeterminism: Prevent ValidateSizeSpecificDataDefaults from being called on the default GeometryCollection object when it is serialized. Validate only needs to be called when loading and calling it other times (e.g. ObjectCollector archives) can incorrectly trigger it on the default object. If the default object has ValidateSizeSpecificDataDefaults, then it will afterwards change the save serialization of instances that have SizeSpecificData equal to the default values constructed by ValidateSizeSpecificDataDefaults.
#rnx

diff --git a/Engine/Source/Runtime/Experimental/GeometryCollectionEngine/Private/GeometryCollection/GeometryCollectionObject.cpp b/Engine/Source/Runtime/Experimental/GeometryCollectionEngine/Private/GeometryCollection/GeometryCollectionObject.cpp
--- a/Engine/Source/Runtime/Experimental/GeometryCollectionEngine/Private/GeometryCollection/GeometryCollectionObject.cpp
+++ b/Engine/Source/Runtime/Experimental/GeometryCollectionEngine/Private/GeometryCollection/GeometryCollectionObject.cpp
@@ -967,7 +967,7 @@ void UGeometryCollection::Serialize(FArchive& Ar)
 		DataflowAsset = StrippedDataflowAsset;
 	}
 
-	if (!SizeSpecificData.Num())
+	if (Ar.IsLoading() && !SizeSpecificData.Num())
 	{
 		ValidateSizeSpecificDataDefaults();
 	}
