Addition of a warning in UGeometryCollectionComponent to catch GC that might be using bStripOnCook without bUseRootProxyForNavigation (resulting in exporting nothing).

diff --git a/Engine/Source/Runtime/Experimental/GeometryCollectionEngine/Private/GeometryCollection/GeometryCollectionComponent.cpp b/Engine/Source/Runtime/Experimental/GeometryCollectionEngine/Private/GeometryCollection/GeometryCollectionComponent.cpp
--- a/Engine/Source/Runtime/Experimental/GeometryCollectionEngine/Private/GeometryCollection/GeometryCollectionComponent.cpp
+++ b/Engine/Source/Runtime/Experimental/GeometryCollectionEngine/Private/GeometryCollection/GeometryCollectionComponent.cpp
@@ -136,6 +136,7 @@ int32 GeometryCollectionNetAwakeningMode = 1;
 FAutoConsoleVariableRef CVarGeometryCollectionNetAwakeningMode(TEXT("p.Chaos.GC.NetAwakeningMode"), GeometryCollectionNetAwakeningMode, TEXT("Changes how GC components ensure that their owner is awake for replication. 0 = ForceDormancyAwake, 1 = Use Flush Net Dormancy"));
 
 DEFINE_LOG_CATEGORY_STATIC(UGCC_LOG, Error, All);
+DEFINE_LOG_CATEGORY_STATIC(LogGeometryCollectionComponent, Warning, All);
 
 extern FGeometryCollectionDynamicDataPool GDynamicDataPool;
 
@@ -1417,6 +1418,12 @@ bool UGeometryCollectionComponent::DoCustomNavigableGeometryExport(FNavigableGeo
 		return true;
 	}
 
+	if (RestCollection->bStripOnCook)
+	{
+		UE_LOG(LogGeometryCollectionComponent, Warning,
+			TEXT("Collection data is set to be stripped on cook so there is nothing left to export. Make sure to use bUseRootProxyForNavigation with bStripOnCook for %s."), *GetFullNameSafe(GetOwner()));
+	}
+	
 	TArray<FVector> OutVertexBuffer;
 	TArray<int32> OutIndexBuffer;
 
