Iris
Increase max registered groups from 4 to 8

#rnx

diff --git a/Engine/Source/Runtime/Experimental/Iris/Core/Private/Iris/ReplicationSystem/Filtering/NetObjectGroups.cpp b/Engine/Source/Runtime/Experimental/Iris/Core/Private/Iris/ReplicationSystem/Filtering/NetObjectGroups.cpp
--- a/Engine/Source/Runtime/Experimental/Iris/Core/Private/Iris/ReplicationSystem/Filtering/NetObjectGroups.cpp
+++ b/Engine/Source/Runtime/Experimental/Iris/Core/Private/Iris/ReplicationSystem/Filtering/NetObjectGroups.cpp
@@ -2,6 +2,7 @@
 
 #include "NetObjectGroups.h"
 #include "Math/UnrealMathUtility.h"
+#include "Iris/Core/IrisLog.h"
 #include "Iris/ReplicationSystem/Filtering/NetObjectFilter.h"
 #include "Iris/ReplicationSystem/NetRefHandleManager.h" // for InvalidInternalIndex
 #include "Containers/ArrayView.h"
@@ -119,7 +120,7 @@ void FNetObjectGroups::SetGroupName(FNetObjectGroupHandle GroupHandle, FName Gro
 {
 	if (FNetObjectGroup* Group = GetGroup(GroupHandle))
 	{
-		if (ensureAlwaysMsgf(EnumHasAnyFlags(Group->Traits, ENetObjectGroupTraits::IsFindableByName), TEXT("FNetObjectGroups::SetGroupName Cannot SetGroupName for grouphandle %u as it is FindableByName"), GroupHandle.GetGroupIndex()))
+		if (ensureMsgf(EnumHasAnyFlags(Group->Traits, ENetObjectGroupTraits::IsFindableByName), TEXT("FNetObjectGroups::SetGroupName Cannot SetGroupName for grouphandle %u as it is FindableByName"), GroupHandle.GetGroupIndex()))
 		{
 			Group->GroupName = GroupName;
 		}
@@ -140,7 +141,7 @@ FNetObjectGroupHandle FNetObjectGroups::CreateNamedGroup(FName GroupName)
 {
 	if (NamedGroups.Contains(GroupName))
 	{
-		ensureAlwaysMsgf(false, TEXT("FNetObjectGroups, trying to create named group %s that already exists"), *GroupName.ToString());
+		ensureMsgf(false, TEXT("FNetObjectGroups, trying to create named group %s that already exists"), *GroupName.ToString());
 		return FNetObjectGroupHandle();
 	}
 
@@ -236,7 +237,8 @@ void FNetObjectGroups::AddToGroup(FNetObjectGroupHandle GroupHandle, FInternalNe
 		}
 		else
 		{
-			checkf(false, TEXT("FNetObjectGroups::AddToGroup, Failed to add ( InternalIndex: %u ) to (GroupIndex: %u) A NetObject can only be a member of %u groups."), InternalIndex, GroupHandle.GetRawValue(), FNetObjectGroupMembership::MaxAssignedGroupCount);
+			UE_LOG(LogIris, Error, TEXT("FNetObjectGroups::AddToGroup, Failed to add ( InternalIndex: %u ) to Group %s (GroupIndex: %u) A NetObject can only be a member of %u groups."), InternalIndex, *GetGroupName(GroupHandle).ToString(), GroupHandle.GetRawValue(), FNetObjectGroupMembership::MaxAssignedGroupCount);
+			ensure(false);
 		}
 	}
 }
diff --git a/Engine/Source/Runtime/Experimental/Iris/Core/Private/Iris/ReplicationSystem/Filtering/NetObjectGroups.h b/Engine/Source/Runtime/Experimental/Iris/Core/Private/Iris/ReplicationSystem/Filtering/NetObjectGroups.h
--- a/Engine/Source/Runtime/Experimental/Iris/Core/Private/Iris/ReplicationSystem/Filtering/NetObjectGroups.h
+++ b/Engine/Source/Runtime/Experimental/Iris/Core/Private/Iris/ReplicationSystem/Filtering/NetObjectGroups.h
@@ -118,7 +118,7 @@ public:
 private:
 	struct FNetObjectGroupMembership
 	{
-		enum { MaxAssignedGroupCount = 4 };
+		enum { MaxAssignedGroupCount = 8 };
 		FNetObjectGroupHandle Groups[MaxAssignedGroupCount];		
 	};
 
