Iris
- Fixes static analysis warning

diff --git a/Engine/Source/Runtime/Experimental/Iris/Core/Private/Iris/ReplicationSystem/ReplicationWriter.cpp b/Engine/Source/Runtime/Experimental/Iris/Core/Private/Iris/ReplicationSystem/ReplicationWriter.cpp
--- a/Engine/Source/Runtime/Experimental/Iris/Core/Private/Iris/ReplicationSystem/ReplicationWriter.cpp
+++ b/Engine/Source/Runtime/Experimental/Iris/Core/Private/Iris/ReplicationSystem/ReplicationWriter.cpp
@@ -937,8 +937,8 @@ void FReplicationWriter::HandleDeliveredRecord(const FReplicationRecord::FRecord
 	
 	if (CurrentState == EReplicatedObjectState::Invalid)
 	{
-		ensure(CurrentState != EReplicatedObjectState::Invalid);
 		UE_LOG_REPLICATIONWRITER_WARNING(TEXT("FReplicationWriter::HandleDeliveredRecord - Warning Object ( InternalIndex: %u ) is invalid. DeliveredState %s WasDestroySubObject: %u"), InternalIndex, LexToString(DeliveredState), RecordInfo.WroteDestroySubObject)
+		ensure(false);
 		return;
 	}
 	
