Handle geometry collection particles being decayed immediately on the physics thread and fix immediately decayed internal geometry collection clusters reporting their children as active.

diff --git a/Engine/Source/Runtime/Experimental/Chaos/Private/PhysicsProxy/GeometryCollectionPhysicsProxy.cpp b/Engine/Source/Runtime/Experimental/Chaos/Private/PhysicsProxy/GeometryCollectionPhysicsProxy.cpp
--- a/Engine/Source/Runtime/Experimental/Chaos/Private/PhysicsProxy/GeometryCollectionPhysicsProxy.cpp
+++ b/Engine/Source/Runtime/Experimental/Chaos/Private/PhysicsProxy/GeometryCollectionPhysicsProxy.cpp
@@ -144,6 +144,12 @@ FAutoConsoleVariableRef CVarGeometryCollectionUseRootBrokenFlag(
 	bGeometryCollectionUseRootBrokenFlag,
 	TEXT("If enabled, check if the root transform is broken in the proxy and disable the GT particle if so. Should be enabled - cvar is a failsafe to revert behaviour"));
 
+bool bPropagateInternalClusterDisableFlagToChildren = true;
+FAutoConsoleVariableRef CVarPropagateInternalClusterDisableFlagToChildren(
+	TEXT("p.GeometryCollection.PropagateInternalClusterDisableFlagToChildren"),
+	bPropagateInternalClusterDisableFlagToChildren,
+	TEXT("If enabled, disabled internal clusters will propagate their disabled flag to their children when buffering instead of implicitly activating the children."));
+	
 bool bGeometryCollectionScaleClusterGeometry = true;
 FAutoConsoleVariableRef CVarGeometryCollectionScaleClusterGeometry(
 	TEXT("p.GeometryCollection.ScaleClusterGeometry"),
@@ -3795,7 +3801,8 @@ void FGeometryCollectionPhysicsProxy::BufferPhysicsResults_Internal(Chaos::FPBDR
 			StateData.TransformIndex = TransformGroupIndex;
 			StateData.HasParent = PhysicsThreadCollection.GetHasParent(TransformGroupIndex);
 			StateData.InternalClusterUniqueIdx = INDEX_NONE;
-			StateData.State.DisabledState = Handle->Disabled();
+			StateData.State.DisabledState = Handle->Disabled(); // this can change if the particle has a cluster union parent 
+			StateData.State.HasDecayed = Handle->Disabled() && (Handle->Parent() == nullptr);
 			StateData.State.HasInternalClusterParent = false;
 			StateData.State.DynamicInternalClusterParent = false;
 			StateData.State.HasClusterUnionParent = false;
@@ -3862,7 +3869,7 @@ void FGeometryCollectionPhysicsProxy::BufferPhysicsResults_Internal(Chaos::FPBDR
 				PositionData.ParticleR = ParticleToWorld.GetRotation();
 				
 				// Indicate that this object needs to be updated and the proxy is active.
-				StateData.State.DisabledState = false;
+				StateData.State.DisabledState = bPropagateInternalClusterDisableFlagToChildren ? ClusterParent->Disabled() : false;
 				StateData.State.HasInternalClusterParent = true;
 				StateData.State.DynamicInternalClusterParent = (ClusterParent->IsDynamic());
 				StateData.State.HasClusterUnionParent = ClusterParent->PhysicsProxy()->GetType() == EPhysicsProxyType::ClusterUnionProxy;
@@ -4131,26 +4138,40 @@ bool FGeometryCollectionPhysicsProxy::PullNonInterpolatableDataFromSinglePhysics
 		}
 
 		const int32 ParticleIndex = FromTransformToParticleIndex[TransformGroupIndex];
+		const bool bIsRootIndex = (Parameters.InitialRootIndex == TransformGroupIndex);
+
+		if (!bIsRootIndex)
+		{
+			CreateChildrenGeometry_External();
+		}
+
 		const bool bIsActive = !StateData.State.DisabledState;
 		if (UpdateValue(GameThreadCollection.Active[TransformGroupIndex], bIsActive))
 		{
-			if (bIsActive && Parameters.InitialRootIndex != TransformGroupIndex)
+			if (GTParticles[ParticleIndex].IsValid())
 			{
-				CreateChildrenGeometry_External();
+				GTParticles[ParticleIndex]->SetDisabled(!bIsActive);
+				bIsCollectionDirty = true;
 			}
-			if (ParticleIndex == INDEX_NONE || GTParticles[ParticleIndex] == nullptr)
+		}
+
+		if (StateData.State.HasDecayed)
+		{
+			if (!bIsRootIndex)
 			{
-				continue;
+				if (ParticleIndex != INDEX_NONE && GTParticles[ParticleIndex].IsValid())
+				{
+					GTParticles[ParticleIndex]->SetDisabled(true);
+					bIsCollectionDirty = true;
+				}
 			}
-
-			GTParticles[ParticleIndex]->SetDisabled(!bIsActive);
-			bIsCollectionDirty = true;
 		}
 
 		if (ParticleIndex == INDEX_NONE || GTParticles[ParticleIndex] == nullptr)
 		{
 			continue;
 		}
+
 		FParticle& GTParticle = *GTParticles[ParticleIndex];
 
 		if (UpdateValue(GameThreadCollection.DynamicState[TransformGroupIndex], static_cast<uint8>(StateData.State.DynamicState)))
diff --git a/Engine/Source/Runtime/Experimental/Chaos/Public/GeometryCollectionProxyData.h b/Engine/Source/Runtime/Experimental/Chaos/Public/GeometryCollectionProxyData.h
--- a/Engine/Source/Runtime/Experimental/Chaos/Public/GeometryCollectionProxyData.h
+++ b/Engine/Source/Runtime/Experimental/Chaos/Public/GeometryCollectionProxyData.h
@@ -247,10 +247,11 @@ public:
 	{
 		uint16 DynamicState: 8; // need to fit EObjectStateTypeEnum
 		uint16 DisabledState: 1;
+		uint16 HasDecayed : 1; // particle has been released but disabled right away
 		uint16 HasInternalClusterParent: 1;
 		uint16 DynamicInternalClusterParent: 1;
 		uint16 HasClusterUnionParent: 1;
-		// 4 bits left
+		// 3 bits left
 	};
 
 	struct FStateData
diff --git a/Engine/Source/Runtime/Experimental/GeometryCollectionEngine/Private/GeometryCollection/GeometryCollectionComponent.cpp b/Engine/Source/Runtime/Experimental/GeometryCollectionEngine/Private/GeometryCollection/GeometryCollectionComponent.cpp
--- a/Engine/Source/Runtime/Experimental/GeometryCollectionEngine/Private/GeometryCollection/GeometryCollectionComponent.cpp
+++ b/Engine/Source/Runtime/Experimental/GeometryCollectionEngine/Private/GeometryCollection/GeometryCollectionComponent.cpp
@@ -3970,7 +3970,7 @@ void UGeometryCollectionComponent::UpdateBrokenAndDecayedStates()
 		if (RootIndex != INDEX_NONE)
 		{
 			const int32 NumTransforms = DynamicCollection->GetNumTransforms();
-
+			
 			const bool bIsRootBroken = !DynamicCollection->Active[RootIndex];
 			BrokenAndDecayedStates.SetRootIsBroken(bIsRootBroken);
 			// we mark it both broken and decayed
@@ -3998,11 +3998,38 @@ void UGeometryCollectionComponent::UpdateBrokenAndDecayedStates()
 							continue;
 						}
 
-						const bool bInternalClusterParentBrokenOff = DynamicStateFacade.HasDynamicInternalClusterParent(TransformIdx) && !DynamicStateFacade.HasClusterUnionParent(TransformIdx);
-						const bool bSelfBrokenOff = DynamicStateFacade.HasBrokenOff(TransformIdx) && !DynamicStateFacade.HasInternalClusterParent(TransformIdx);
-						if (bSelfBrokenOff || bInternalClusterParentBrokenOff)
+						if (!BrokenAndDecayedStates.GetIsBroken(TransformIdx))
+						{
+							const bool bInternalClusterParentBrokenOff = DynamicStateFacade.HasDynamicInternalClusterParent(TransformIdx) && !DynamicStateFacade.HasClusterUnionParent(TransformIdx);
+							const bool bSelfBrokenOff = DynamicStateFacade.HasBrokenOff(TransformIdx) && !DynamicStateFacade.HasInternalClusterParent(TransformIdx);
+							if (bSelfBrokenOff || bInternalClusterParentBrokenOff)
+							{
+								BrokenAndDecayedStates.SetIsBroken(TransformIdx);								
+							}
+						}
+
+						if (BrokenAndDecayedStates.GetIsBroken(TransformIdx))
 						{
-							BrokenAndDecayedStates.SetIsBroken(TransformIdx);
+							if (Particle->Disabled())
+							{
+								// Todo(chaos) in the future we shoud probably store a bitArray to know if the bone is a simulated leaf
+								bool bHasAnyChildWithValidParticle = false;
+								DynamicCollection->IterateThroughChildren(TransformIdx, 
+									[&bHasAnyChildWithValidParticle, this] (int32 ChildTransformIdx) -> bool
+									{
+										if (PhysicsProxy->GetParticleByIndex_External(ChildTransformIdx))
+										{
+											bHasAnyChildWithValidParticle = true;
+											return false; // do no continue iterating
+										}
+										return true; // continue iterating
+									});
+
+								if (!bHasAnyChildWithValidParticle)
+								{
+									BrokenAndDecayedStates.SetHasDecayed(TransformIdx);
+								}
+							}
 						}
 					}
 				}
@@ -4179,6 +4206,10 @@ void UGeometryCollectionComponent::OnPostPhysicsSync()
 				PrimaryComponentTick.SetTickFunctionEnable(true);
 			}
 		}
+		else
+		{
+			InitializeRemovalDynamicAttributesIfNeeded();
+		}
 
 		UpdateRemovalIfNeeded();
 
@@ -4312,7 +4343,7 @@ void UGeometryCollectionComponent::UpdateRemovalIfNeeded()
 
 	// if removal is enabled, update the dynamic collection transform based on the decay 
 	// todo: we could optimize this using a list of transform to update from when we update the decay values
-	if (DynamicCollection && bAllowRemovalOnBreak && bAllowRemovalOnSleep)
+	if (DynamicCollection && bAllowRemovalOnBreak && bAllowRemovalOnSleep && PhysicsProxy)
 	{
 		FGeometryCollectionDecayDynamicFacade DecayFacade(*DynamicCollection);
 		if (DecayFacade.IsValid())
@@ -4329,11 +4360,27 @@ void UGeometryCollectionComponent::UpdateRemovalIfNeeded()
 				CompSpaceTransform = &ComponentSpaceTransforms.RequestAllTransforms();
 			}
 
+			const int32 RootIndex = RestCollection->GetRootIndex();
+
 			const int32 NumTransforms = DecayFacade.GetDecayAttributeSize();
 			for (int32 TransformIndex = 0; TransformIndex < NumTransforms; ++TransformIndex)
 			{
 				// only update values if the decay has changed 
-				const float Decay = DecayFacade.GetDecay(TransformIndex);
+				float Decay = DecayFacade.GetDecay(TransformIndex);
+
+				// reconcile changes to the decay state with the decay value
+				if (Decay < 1.f && TransformIndex != RootIndex)
+				{
+					if (PhysicsProxy->GetParticleByIndex_External(TransformIndex))
+					{
+						if (BrokenAndDecayedStates.GetHasDecayed(TransformIndex))
+						{
+							Decay = 1.0f;
+							DecayFacade.SetDecay(TransformIndex, Decay);
+						}
+					}
+				}
+
 				if (Decay > 0.f && Decay <= 1.f)
 				{
 					const float Scale = 1.0 - Decay;
