GeometryCollection Indeterminism: Fix crash due to not calling  ValidateSizeSpecificDataDefaults when it is saved; the validation done during serialize is needed in some cases.
Crash repro instructions: create a new worldpartition level, add cube, create a level instance from that cube.
#rnx

diff --git a/Engine/Source/Runtime/Experimental/GeometryCollectionEngine/Private/GeometryCollection/GeometryCollectionObject.cpp b/Engine/Source/Runtime/Experimental/GeometryCollectionEngine/Private/GeometryCollection/GeometryCollectionObject.cpp
--- a/Engine/Source/Runtime/Experimental/GeometryCollectionEngine/Private/GeometryCollection/GeometryCollectionObject.cpp
+++ b/Engine/Source/Runtime/Experimental/GeometryCollectionEngine/Private/GeometryCollection/GeometryCollectionObject.cpp
@@ -967,8 +967,10 @@ void UGeometryCollection::Serialize(FArchive& Ar)
 		DataflowAsset = StrippedDataflowAsset;
 	}
 
-	if (Ar.IsLoading() && !SizeSpecificData.Num())
+	if ((Ar.IsLoading() || Ar.IsSaving()) && !SizeSpecificData.Num())
 	{
+		// Validation is necessary when loading old version and when saving newly created version
+		// that might not have created the defaults yet; the defaults are used during EnsureDataIsCooked.
 		ValidateSizeSpecificDataDefaults();
 	}
 
