Chaos - add explicit particle wake method which resets sleep counters etc

diff --git a/Engine/Source/Runtime/Experimental/Chaos/Private/Chaos/Island/IslandManager.cpp b/Engine/Source/Runtime/Experimental/Chaos/Private/Chaos/Island/IslandManager.cpp
--- a/Engine/Source/Runtime/Experimental/Chaos/Private/Chaos/Island/IslandManager.cpp
+++ b/Engine/Source/Runtime/Experimental/Chaos/Private/Chaos/Island/IslandManager.cpp
@@ -667,6 +667,19 @@ namespace Chaos::Private
 		return IslandConstraints;
 	}
 
+	void FPBDIslandManager::WakeParticleIsland(const FGeometryParticleHandle* Particle)
+	{
+		if (FPBDIslandParticle* Node = GetGraphNode(Particle))
+		{
+			if (Node->Island != nullptr)
+			{
+				// When we explicitly wake, we reset sleep counters etc
+				const bool bIsSleepAllowed = false;
+				EnqueueIslandCheckSleep(Node->Island, bIsSleepAllowed);
+			}
+		}
+	}
+
 	void FPBDIslandManager::AddParticle(FGeometryParticleHandle* Particle)
 	{
 	}
diff --git a/Engine/Source/Runtime/Experimental/Chaos/Private/Chaos/PBDRigidsEvolution.cpp b/Engine/Source/Runtime/Experimental/Chaos/Private/Chaos/PBDRigidsEvolution.cpp
--- a/Engine/Source/Runtime/Experimental/Chaos/Private/Chaos/PBDRigidsEvolution.cpp
+++ b/Engine/Source/Runtime/Experimental/Chaos/Private/Chaos/PBDRigidsEvolution.cpp
@@ -1099,6 +1099,18 @@ namespace Chaos
 		}
 	}
 
+	void FPBDRigidsEvolutionBase::WakeParticle(FPBDRigidParticleHandle* Particle)
+	{
+		if (Particle->IsDynamic())
+		{
+			SetParticleObjectState(Particle, EObjectStateType::Dynamic);
+
+			// Reset the sleep counters etc in the sleep management system
+			// @todo(chaos): ideally we would just set reset the sleepiness of the particle
+			IslandManager.WakeParticleIsland(Particle);
+		}
+	}
+
 	void FPBDRigidsEvolutionBase::SetParticleSleepType(FPBDRigidParticleHandle* Particle, ESleepType InSleepType)
 	{
 		//ESleepType InitialSleepType = Particle->SleepType();
diff --git a/Engine/Source/Runtime/Experimental/Chaos/Public/Chaos/Island/IslandManager.h b/Engine/Source/Runtime/Experimental/Chaos/Public/Chaos/Island/IslandManager.h
--- a/Engine/Source/Runtime/Experimental/Chaos/Public/Chaos/Island/IslandManager.h
+++ b/Engine/Source/Runtime/Experimental/Chaos/Public/Chaos/Island/IslandManager.h
@@ -389,15 +389,16 @@ namespace Chaos::Private
 		// The number of registered constraint containers
 		CHAOS_API int32 GetNumConstraintContainers() const;
 
-		// Add/Remove particles
+		// Particles management
 		int32 GetNumParticles() const { return Nodes.Num(); }
 		CHAOS_API void AddParticle(FGeometryParticleHandle* Particle);
 		CHAOS_API void RemoveParticle(FGeometryParticleHandle* Particle);
 		CHAOS_API int32 ReserveParticles(const int32 InNumParticles);
 		CHAOS_API void UpdateParticleMaterial(FGeometryParticleHandle* Particle);
 		CHAOS_API int32 GetParticleLevel(FGeometryParticleHandle* Particle) const;
+		CHAOS_API void WakeParticleIsland(const FGeometryParticleHandle* Particle);
 
-		// Add/Remove constraints
+		// Constraint management
 		int32 GetNumConstraints() const { return Edges.Num(); }
 		CHAOS_API void AddConstraint(const int32 ContainerId, FConstraintHandle* Constraint, const TVec2<FGeometryParticleHandle*>& ConstrainedParticles);
 		template<typename ConstraintContainerType> void AddContainerConstraints(ConstraintContainerType& ConstraintContainer);
diff --git a/Engine/Source/Runtime/Experimental/Chaos/Public/Chaos/PBDRigidsEvolution.h b/Engine/Source/Runtime/Experimental/Chaos/Public/Chaos/PBDRigidsEvolution.h
--- a/Engine/Source/Runtime/Experimental/Chaos/Public/Chaos/PBDRigidsEvolution.h
+++ b/Engine/Source/Runtime/Experimental/Chaos/Public/Chaos/PBDRigidsEvolution.h
@@ -590,6 +590,9 @@ public:
 
 	CHAOS_API void SetParticleObjectState(FPBDRigidParticleHandle* Particle, EObjectStateType ObjectState);
 
+	// Wake a dynamic particle and reset sleep counters for its island
+	CHAOS_API void WakeParticle(FPBDRigidParticleHandle* Particle);
+
 	CHAOS_API void SetParticleSleepType(FPBDRigidParticleHandle* Particle, ESleepType InSleepType);
 
 	CHAOS_API void DisableParticles(const TSet<FGeometryParticleHandle*>& InParticles);
