Fix potential crash introduced when doing memory optimization

diff --git a/Engine/Source/Runtime/Experimental/Chaos/Private/PhysicsProxy/GeometryCollectionPhysicsProxy.cpp b/Engine/Source/Runtime/Experimental/Chaos/Private/PhysicsProxy/GeometryCollectionPhysicsProxy.cpp
--- a/Engine/Source/Runtime/Experimental/Chaos/Private/PhysicsProxy/GeometryCollectionPhysicsProxy.cpp
+++ b/Engine/Source/Runtime/Experimental/Chaos/Private/PhysicsProxy/GeometryCollectionPhysicsProxy.cpp
@@ -1650,8 +1650,7 @@ void FGeometryCollectionPhysicsProxy::InitializeBodiesPT(Chaos::FPBDRigidsSolver
 		const TManagedArray<int32>& Level = PhysicsThreadCollection.GetInitialLevels().Get();
 		for (int32 ParticleIndex = 0; ParticleIndex < SolverParticleHandles.Num(); ++ParticleIndex)
 		{
-			Chaos::FPBDRigidParticleHandle* Handle = SolverParticleHandles[ParticleIndex];
-			if (Handle)
+			if (Chaos::FPBDRigidParticleHandle* Handle = SolverParticleHandles[ParticleIndex])
 			{
 				const bool bIsOneWayInteraction = (Parameters.OneWayInteractionLevel >= 0) && (Level[ParticleIndex] >= Parameters.OneWayInteractionLevel);
 
@@ -3271,55 +3270,56 @@ void FGeometryCollectionPhysicsProxy::SetWorldTransform_Internal(const FTransfor
 
 		for (int32 ParticleIndex = 0; ParticleIndex < NumEffectiveParticles; ++ParticleIndex)
 		{
-			Chaos::FPBDRigidClusteredParticleHandle* Handle = SolverParticleHandles[ParticleIndex];
-			check(Handle != nullptr);
-			FClusterHandle* KinematicRootHandle = nullptr;
-			FClusterHandle* ParentHandle = Handle->Parent();
-
-			if (!Handle->Disabled() && Handle->ObjectState() == Chaos::EObjectStateType::Kinematic)
-			{
-				KinematicRootHandle = Handle;
-			}
-			else
+			if (Chaos::FPBDRigidClusteredParticleHandle* Handle = SolverParticleHandles[ParticleIndex])
 			{
-				// is there a internal parent as a kinematic root?
-				if (ParentHandle && ParentHandle->InternalCluster() && !ParentHandle->Disabled() && ParentHandle->ObjectState() == Chaos::EObjectStateType::Kinematic && ParentHandle->PhysicsProxy() == this)
+				FClusterHandle* KinematicRootHandle = nullptr;
+				FClusterHandle* ParentHandle = Handle->Parent();
+
+				if (!Handle->Disabled() && Handle->ObjectState() == Chaos::EObjectStateType::Kinematic)
 				{
-					if (!ProcessedInternalClusters.Contains(ParentHandle))
+					KinematicRootHandle = Handle;
+				}
+				else
+				{
+					// is there a internal parent as a kinematic root?
+					if (ParentHandle && ParentHandle->InternalCluster() && !ParentHandle->Disabled() && ParentHandle->ObjectState() == Chaos::EObjectStateType::Kinematic && ParentHandle->PhysicsProxy() == this)
 					{
-						ProcessedInternalClusters.Add(ParentHandle);
-						KinematicRootHandle = ParentHandle;
+						if (!ProcessedInternalClusters.Contains(ParentHandle))
+						{
+							ProcessedInternalClusters.Add(ParentHandle);
+							KinematicRootHandle = ParentHandle;
+						}
 					}
 				}
-			}
-
-			if (KinematicRootHandle)
-			{
-				const FTransform RootWorldTransform(KinematicRootHandle->R(), KinematicRootHandle->X());
-				const FTransform RootRelativeTransform = RootWorldTransform.GetRelativeTransform(Parameters.PrevWorldTransform);
-				const FTransform WorldTransform = RootRelativeTransform * ActorToWorld;
 
-				SetClusteredParticleKinematicTarget_Internal(KinematicRootHandle, WorldTransform);
-			}
-			else if (ParentHandle && !ParentHandle->IsDynamic())
-			{
-				if (ClusterUnionIndex == INDEX_NONE)
+				if (KinematicRootHandle)
 				{
-					ClusterUnionIndex = ClusterUnionManager.FindClusterUnionIndexFromParticle(Handle);
-				}
+					const FTransform RootWorldTransform(KinematicRootHandle->R(), KinematicRootHandle->X());
+					const FTransform RootRelativeTransform = RootWorldTransform.GetRelativeTransform(Parameters.PrevWorldTransform);
+					const FTransform WorldTransform = RootRelativeTransform * ActorToWorld;
 
-				if (ClusterUnionIndex != INDEX_NONE)
+					SetClusteredParticleKinematicTarget_Internal(KinematicRootHandle, WorldTransform);
+				}
+				else if (ParentHandle && !ParentHandle->IsDynamic())
 				{
-					const int32 TransformGroupIndex = FromParticleToTransformIndex[ParticleIndex];
-					const FTransform ParentWorldTransform{ ParentHandle->R(), ParentHandle->X() };
-					const FTransform NewWorldTransform = MassToLocal[TransformGroupIndex] * FTransform(PhysicsThreadCollection.GetTransform(TransformGroupIndex)) * Parameters.WorldTransform;
-					const FTransform RelativeTransform = NewWorldTransform.GetRelativeTransform(ParentWorldTransform);
+					if (ClusterUnionIndex == INDEX_NONE)
+					{
+						ClusterUnionIndex = ClusterUnionManager.FindClusterUnionIndexFromParticle(Handle);
+					}
 
-					DeferredClusterUnionParticleUpdates.Add(Handle);
-					DeferredClusterUnionChildToParentUpdates.Add(RelativeTransform);
+					if (ClusterUnionIndex != INDEX_NONE)
+					{
+						const int32 TransformGroupIndex = FromParticleToTransformIndex[ParticleIndex];
+						const FTransform ParentWorldTransform{ ParentHandle->R(), ParentHandle->X() };
+						const FTransform NewWorldTransform = MassToLocal[TransformGroupIndex] * FTransform(PhysicsThreadCollection.GetTransform(TransformGroupIndex)) * Parameters.WorldTransform;
+						const FTransform RelativeTransform = NewWorldTransform.GetRelativeTransform(ParentWorldTransform);
 
-					// Make sure the particle is mark dirty to make sure its proxy will properly update the transforms
-					RigidSolver->GetEvolution()->GetParticles().MarkTransientDirtyParticle(Handle);
+						DeferredClusterUnionParticleUpdates.Add(Handle);
+						DeferredClusterUnionChildToParentUpdates.Add(RelativeTransform);
+
+						// Make sure the particle is mark dirty to make sure its proxy will properly update the transforms
+						RigidSolver->GetEvolution()->GetParticles().MarkTransientDirtyParticle(Handle);
+					}
 				}
 			}
 		}
@@ -3480,8 +3480,7 @@ void FGeometryCollectionPhysicsProxy::SetMaterialOverrideMassScaleMultiplier_Int
 
 	for (int32 ParticleIndex = 0; ParticleIndex < SolverParticleHandles.Num(); ++ParticleIndex)
 	{
-		Chaos::FPBDRigidClusteredParticleHandle* Handle = SolverParticleHandles[ParticleIndex];
-		if (Handle != nullptr)
+		if (Chaos::FPBDRigidClusteredParticleHandle* Handle = SolverParticleHandles[ParticleIndex])
 		{
 			const Chaos::FReal NewM = Handle->M() * MaterialOverrideMassScaleMultiplierChange;
 			const Chaos::FVec3 NewI = Handle->I() * MaterialOverrideMassScaleMultiplierChange;
@@ -3545,8 +3544,7 @@ void FGeometryCollectionPhysicsProxy::SetOneWayInteractionLevel_Internal(int32 I
 	const TManagedArray<int32>& Level = PhysicsThreadCollection.GetInitialLevels().Get();
 	for (int32 ParticleIndex = 0; ParticleIndex < SolverParticleHandles.Num(); ++ParticleIndex)
 	{
-		Chaos::FPBDRigidClusteredParticleHandle* Handle = SolverParticleHandles[ParticleIndex];
-		if (Handle)
+		if (Chaos::FPBDRigidClusteredParticleHandle* Handle = SolverParticleHandles[ParticleIndex])
 		{
 			const bool bIsOneWayInteraction = (Parameters.OneWayInteractionLevel >= 0) && (Level[FromParticleToTransformIndex[ParticleIndex]] >= Parameters.OneWayInteractionLevel);
 			Handle->SetOneWayInteraction(bIsOneWayInteraction);
