Chaos - fix cluster union wake so that it ensures the correct number of waking frames based on sleep settings

diff --git a/Engine/Source/Runtime/Engine/Private/PhysicsEngine/ClusterUnionComponent.cpp b/Engine/Source/Runtime/Engine/Private/PhysicsEngine/ClusterUnionComponent.cpp
--- a/Engine/Source/Runtime/Engine/Private/PhysicsEngine/ClusterUnionComponent.cpp
+++ b/Engine/Source/Runtime/Engine/Private/PhysicsEngine/ClusterUnionComponent.cpp
@@ -1374,15 +1374,9 @@ void UClusterUnionComponent::SetSimulatePhysics(bool bSimulate)
 
 void UClusterUnionComponent::WakeAllRigidBodies()
 {
-	if (!PhysicsProxy)
-	{
-		return;
-	}
-
-	Chaos::EObjectStateType State = PhysicsProxy->GetObjectState_External();
-	if (State == Chaos::EObjectStateType::Sleeping)
+	if (PhysicsProxy)
 	{
-		SetRigidState(Chaos::EObjectStateType::Dynamic);
+		PhysicsProxy->Wake_External();
 	}
 }
 
diff --git a/Engine/Source/Runtime/Experimental/Chaos/Private/Chaos/Island/IslandManager.cpp b/Engine/Source/Runtime/Experimental/Chaos/Private/Chaos/Island/IslandManager.cpp
--- a/Engine/Source/Runtime/Experimental/Chaos/Private/Chaos/Island/IslandManager.cpp
+++ b/Engine/Source/Runtime/Experimental/Chaos/Private/Chaos/Island/IslandManager.cpp
@@ -59,6 +59,13 @@ namespace Chaos::CVars
 	/** Cvar to override the sleep angular threshold if necessary */
 	FRealSingle ChaosSolverCollisionDefaultAngularSleepThreshold = 0.0087f;  //~1/2 unit mass degree
 	FAutoConsoleVariableRef CVarChaosSolverCollisionDefaultAngularSleepThreshold(TEXT("p.ChaosSolverCollisionDefaultAngularSleepThreshold"), ChaosSolverCollisionDefaultAngularSleepThreshold, TEXT("Default angular threshold for sleeping.[def:0.0087]"));
+
+	/** The size of object for which the angular sleep threshold is defined. Large objects reduce the threshold propertionally. 0 means do not apply size scale. */
+	// E.g., if ChaosSolverCollisionAngularSleepThresholdSize=100, an objects with a bounds of 500 will have 1/5x the sleep threshold.
+	// We are effectively converting the angular threshold into a linear threshold calculated at the object extents.
+	// @todo(chaos): male this a project setting or something
+	FRealSingle ChaosSolverCollisionAngularSleepThresholdSize = 0;
+	FAutoConsoleVariableRef CVarChaosSolverCollisionAngularSleepThresholdSize(TEXT("p.ChaosSolverCollisionAngularSleepThresholdSize"), ChaosSolverCollisionAngularSleepThresholdSize, TEXT("Scales the angular threshold based on size (0 to disable size based scaling)"));
 }
 
 
@@ -155,6 +162,19 @@ namespace Chaos::Private
 				OutSleepCounterThreshold = CVars::ChaosSolverCollisionDefaultSleepCounterThreshold;
 			}
 
+			// Adjust angular threshold for size. It is equivalent to converting the angular threshold into a linear
+			// movement threshold at the extreme points on the particle.
+			const FRealSingle AngularSleepThresholdSize = CVars::ChaosSolverCollisionAngularSleepThresholdSize;
+			if ((AngularSleepThresholdSize > 0) && Rigid->HasBounds())
+			{
+				const FRealSingle RigidSize = FRealSingle(Rigid->LocalBounds().Extents().GetMax());
+				if (RigidSize > AngularSleepThresholdSize)
+				{
+					const FRealSingle ThresholdScale = AngularSleepThresholdSize / RigidSize;
+					OutSleepAngularThreshold *= ThresholdScale;
+				}
+			}
+
 			return true;
 		}
 
@@ -667,21 +687,31 @@ namespace Chaos::Private
 		return IslandConstraints;
 	}
 
-	void FPBDIslandManager::WakeParticleIsland(const FGeometryParticleHandle* Particle)
+	void FPBDIslandManager::WakeParticleIsland(FGeometryParticleHandle* Particle)
 	{
+		// When we explicitly wake we reset sleep counters etc
 		if (FPBDIslandParticle* Node = GetGraphNode(Particle))
 		{
+			// If we are in an island flag for sleep reset
 			if (Node->Island != nullptr)
 			{
-				// When we explicitly wake, we reset sleep counters etc
 				const bool bIsSleepAllowed = false;
 				EnqueueIslandCheckSleep(Node->Island, bIsSleepAllowed);
 			}
 		}
+
+		// If we are an isolated particle we keep our own sleep counter
+		// NOTE: We won't even be in the graph if we are dynamic but have no constraints
+		if (FPBDRigidParticleHandle* Rigid = Particle->CastToRigidParticle())
+		{
+			Rigid->SetSleepCounter(0);
+		}
 	}
 
 	void FPBDIslandManager::AddParticle(FGeometryParticleHandle* Particle)
 	{
+		// Particles get auto-added via AddConstraint (see CreateGraphEdge, GetOrCreateGraphNode)
+		// But must be explicitly removed when destroyed.
 	}
 
 	void FPBDIslandManager::RemoveParticle(FGeometryParticleHandle* Particle)
diff --git a/Engine/Source/Runtime/Experimental/Chaos/Private/PhysicsProxy/ClusterUnionPhysicsProxy.cpp b/Engine/Source/Runtime/Experimental/Chaos/Private/PhysicsProxy/ClusterUnionPhysicsProxy.cpp
--- a/Engine/Source/Runtime/Experimental/Chaos/Private/PhysicsProxy/ClusterUnionPhysicsProxy.cpp
+++ b/Engine/Source/Runtime/Experimental/Chaos/Private/PhysicsProxy/ClusterUnionPhysicsProxy.cpp
@@ -272,6 +272,27 @@ namespace Chaos
 		);
 	}
 
+	void FClusterUnionPhysicsProxy::Wake_External()
+	{
+		if (!Solver || !ensure(Particle_External))
+		{
+			return;
+		}
+
+		Solver->EnqueueCommandImmediate(
+			[this]() mutable
+			{
+				if (Particle_Internal)
+				{
+					if (FPBDRigidsEvolutionGBF* Evolution = GetEvolution(this))
+					{
+						Evolution->WakeParticle(Particle_Internal);
+					}
+				}
+			}
+		);
+	}
+
 	void FClusterUnionPhysicsProxy::RemoveShapes_External(const TArray<FPBDRigidParticle*>& ShapeParticles)
     {
 		RemoveParticlesFromClusterUnionGeometry(Particle_External.Get(), ShapeParticles, GeometryChildParticles_External);
diff --git a/Engine/Source/Runtime/Experimental/Chaos/Public/Chaos/Island/IslandManager.h b/Engine/Source/Runtime/Experimental/Chaos/Public/Chaos/Island/IslandManager.h
--- a/Engine/Source/Runtime/Experimental/Chaos/Public/Chaos/Island/IslandManager.h
+++ b/Engine/Source/Runtime/Experimental/Chaos/Public/Chaos/Island/IslandManager.h
@@ -396,7 +396,7 @@ namespace Chaos::Private
 		CHAOS_API int32 ReserveParticles(const int32 InNumParticles);
 		CHAOS_API void UpdateParticleMaterial(FGeometryParticleHandle* Particle);
 		CHAOS_API int32 GetParticleLevel(FGeometryParticleHandle* Particle) const;
-		CHAOS_API void WakeParticleIsland(const FGeometryParticleHandle* Particle);
+		CHAOS_API void WakeParticleIsland(FGeometryParticleHandle* Particle);
 
 		// Constraint management
 		int32 GetNumConstraints() const { return Edges.Num(); }
diff --git a/Engine/Source/Runtime/Experimental/Chaos/Public/PhysicsProxy/ClusterUnionPhysicsProxy.h b/Engine/Source/Runtime/Experimental/Chaos/Public/PhysicsProxy/ClusterUnionPhysicsProxy.h
--- a/Engine/Source/Runtime/Experimental/Chaos/Public/PhysicsProxy/ClusterUnionPhysicsProxy.h
+++ b/Engine/Source/Runtime/Experimental/Chaos/Public/PhysicsProxy/ClusterUnionPhysicsProxy.h
@@ -82,6 +82,9 @@ namespace Chaos
 		CHAOS_API EObjectStateType GetObjectState_External() const;
 		CHAOS_API void SetObjectState_External(EObjectStateType State);
 
+		// Explicitly wake the physics object. The wakes the physics object if it is sleeping and resets any sleep state (whether awake or sleeping)
+		CHAOS_API void Wake_External();
+
 		// Set GT geometry - this is only for smoothing over any changes until the PT syncs back to the GT.
 		CHAOS_API void SetGeometry_External(const Chaos::FImplicitObjectPtr& Geometry, const TArray<FPBDRigidParticle*>& ShapeParticles);
 
