Ias: Fix some InFlight book-keeping issues
1) checking !IsInFlight(Request) in CompleteRequest/Release is not valid since a new request with the same key may have already been added to the Inflight map.
2) FChunkRequests::Cancel sets bCancelled=true and removes the entry from the Inflight map, then CompleteRequest/RemoveAndRelease must not remove from the Inflight map again, since a new request with the same key may have already been added to the Inflight map.
3) In CancelAllPendingRequests/DrainHttpRequests don't set bCancelled since that would not remove the requests from the Inflight map, it is better to take the fail path. On the other hand, on shutdown there should only ever be cancelled requests in the http queue since IoDispatcher waits for all pending io before shutting down the backends.

#rnx

diff --git a/Engine/Source/Runtime/Experimental/IoStoreOnDemand/Private/OnDemandIoDispatcherBackend.cpp b/Engine/Source/Runtime/Experimental/IoStoreOnDemand/Private/OnDemandIoDispatcherBackend.cpp
--- a/Engine/Source/Runtime/Experimental/IoStoreOnDemand/Private/OnDemandIoDispatcherBackend.cpp
+++ b/Engine/Source/Runtime/Experimental/IoStoreOnDemand/Private/OnDemandIoDispatcherBackend.cpp
@@ -1498,17 +1498,9 @@ class FOnDemandIoBackend final
 		void Release(FChunkRequest* Request)
 		{
 			FScopeLock _(&Mutex);
-			check(!IsInFlight(Request));
 			Destroy(Request);
 		}
 		
-		void RemoveAndRelease(FChunkRequest* Request)
-		{
-			FScopeLock _(&Mutex);
-			Inflight.Remove(Request->Params.ChunkKey);
-			Destroy(Request);
-		}
-
 		int32 Num()
 		{
 			FScopeLock _(&Mutex);
@@ -1523,20 +1515,6 @@ class FOnDemandIoBackend final
 			check(ChunkRequestCount >= 0);
 		}
 
-		/** Helper intended to be called by methods that have already locked ::Mutex */
-		inline bool IsInFlight(const FChunkRequest* Request) const
-		{
-			const FChunkRequest* const* InFlightRequest = Inflight.Find(Request->Params.ChunkKey);
-			if (InFlightRequest == nullptr)
-			{
-				return false;
-			}
-			else
-			{
-				return *InFlightRequest == Request;
-			}
-		}
-
 		TSingleThreadedSlabAllocator<FChunkRequest, 128> Allocator;
 		TMap<FIoHash, FChunkRequest*> Inflight;
 		FCriticalSection Mutex;
@@ -1575,8 +1553,6 @@ public:
 
 private:
 
-	void CancelAllPendingRequests();
-
 	FString GetEndpointTestPath() const;
 	void ConditionallyStartBackendThread();
 	void CompleteRequest(FChunkRequest* ChunkRequest);
@@ -1604,7 +1580,8 @@ private:
 
 	bool SetupHttpThread();
 	void ProcessHttpRequests(FHttpClient& HttpClient, FBitWindow& HttpErrors, int32 MaxConcurrentRequests);
-	int32 WaitForPendingRequests(float WaitTimeSeconds, float PollTimeSeconds);
+	int32 WaitForCompleteRequestTasks(float WaitTimeSeconds, float PollTimeSeconds);
+	void DrainHttpRequests();
 
 	struct FTocParams
 	{
@@ -1698,39 +1675,20 @@ void FOnDemandIoBackend::Shutdown()
 
 	UE_LOG(LogIas, Log, TEXT("Shutting down on demand I/O dispatcher backend"));
 
-	Stop();
-
-	{
-		TRACE_CPUPROFILER_EVENT_SCOPE(IasBackend::Shutdown::StoppingIasHttpThread);
-		BackendThread.Reset();
-	}
+	// Stop and wait for our backend thread to finish.
+	// Note, the IoDispatcher typically waits for all its pending io requests before shutting down its backends.
+	BackendThread.Reset();
 
-	CancelAllPendingRequests();	
+	// Drain any reamaining (cancelled) http requests that already been completed from the IoDispatcher point of view.
+	DrainHttpRequests();
 
-	const int32 NumPending = WaitForPendingRequests(5.0f, 0.1f);
+	// The CompleteRequest tasks may still be executing a while after the IoDispatcher has been notified about the completed io requests.
+	const int32 NumPending = WaitForCompleteRequestTasks(5.0f, 0.1f);
 	UE_CLOG(NumPending > 0, LogIas, Warning, TEXT("%d request(s) still pending after shutdown"), NumPending);
 
 	BackendContext.Reset();
 }
 
-void FOnDemandIoBackend::CancelAllPendingRequests()
-{
-	FChunkRequest* Iterator = HttpRequests.Dequeue();
-	while (Iterator != nullptr)
-	{
-		FChunkRequest* Request = Iterator;
-		Iterator = Iterator->NextRequest;
-
-		// We need to call this to increment the number of requests in flight or the subsequent call to Stats.OnHttpCancel will start to mismatch that counter.
-		Stats.OnHttpDequeue();
-
-		Request->bCancelled = true;
-		CompleteRequest(Request);
-		
-		Stats.OnHttpCancel();
-	}
-}
-
 FString FOnDemandIoBackend::GetEndpointTestPath() const
 {
 	FString TestPath = IoStore.IsValid() ? IoStore->GetFirstTocPath() : FString();
@@ -1764,7 +1722,7 @@ void FOnDemandIoBackend::CompleteRequest(FChunkRequest* ChunkRequest)
 	{
 		check(ChunkRequest->RequestHead == nullptr);
 		check(ChunkRequest->RequestTail == nullptr);
-		return ChunkRequests.RemoveAndRelease(ChunkRequest);
+		return ChunkRequests.Release(ChunkRequest);
 	}
 
 	ChunkRequests.Remove(ChunkRequest);
@@ -2385,11 +2343,7 @@ void FOnDemandIoBackend::ProcessHttpRequests(FHttpClient& HttpClient, FBitWindow
 				}
 				else if (BackendStatus.IsHttpEnabled() == false)
 				{
-					UE::Tasks::Launch(UE_SOURCE_LOCATION, [this, ChunkRequest]()
-					{
-						CompleteRequest(ChunkRequest);
-					});
-
+					CompleteRequest(ChunkRequest);
 					// Technically this request is being skipped because of a pre-existing error. It is not
 					// an error itself and it is not being canceled by higher level code. However we do not
 					// currently have a statistic for that and we have to call one of the existing types in
@@ -2491,6 +2445,20 @@ void FOnDemandIoBackend::ProcessHttpRequests(FHttpClient& HttpClient, FBitWindow
 	} 
 }
 
+void FOnDemandIoBackend::DrainHttpRequests()
+{
+	FChunkRequest* Iterator = HttpRequests.Dequeue();
+	while (Iterator != nullptr)
+	{
+		FChunkRequest* Request = Iterator;
+		Iterator = Iterator->NextRequest;
+
+		Stats.OnHttpDequeue();
+		CompleteRequest(Request);
+		Stats.OnHttpCancel();
+	}
+}
+
 bool FOnDemandIoBackend::SetupHttpThread()
 {
 	// A disk only flush should be fairly quick at this point, especially if 's.IasEnableThreadedTocGeneration' is true
@@ -2616,7 +2584,7 @@ void FOnDemandIoBackend::Stop()
 	DistributedEndpointEvent->Trigger();
 }
 
-int32 FOnDemandIoBackend::WaitForPendingRequests(float WaitTimeSeconds, float PollTimeSeconds)
+int32 FOnDemandIoBackend::WaitForCompleteRequestTasks(float WaitTimeSeconds, float PollTimeSeconds)
 {
 	const double StartTime = FPlatformTime::Seconds();
 	while (ChunkRequests.Num() > 0 && float(FPlatformTime::Seconds() - StartTime) < WaitTimeSeconds)
