Iris
* Removed ensureAlways from StateDescriptorBuilder

#rnx

diff --git a/Engine/Source/Runtime/Experimental/Iris/Core/Private/Iris/ReplicationState/ReplicationStateDescriptorBuilder.cpp b/Engine/Source/Runtime/Experimental/Iris/Core/Private/Iris/ReplicationState/ReplicationStateDescriptorBuilder.cpp
--- a/Engine/Source/Runtime/Experimental/Iris/Core/Private/Iris/ReplicationState/ReplicationStateDescriptorBuilder.cpp
+++ b/Engine/Source/Runtime/Experimental/Iris/Core/Private/Iris/ReplicationState/ReplicationStateDescriptorBuilder.cpp
@@ -1108,7 +1108,7 @@ void FPropertyReplicationStateDescriptorBuilder::BuildMemberTagDescriptors(FRepl
 	}
 
 	// Fill in descriptor
-	ensureAlwaysMsgf(Context.MemberTagCache.TagCount < MAX_uint16 - 1U, TEXT("Tag count cannot exceed %u. Building descriptor %s."), MAX_uint16 - 1U, ToCStr(Descriptor->DebugName));
+	ensureMsgf(Context.MemberTagCache.TagCount < MAX_uint16 - 1U, TEXT("Tag count cannot exceed %u. Building descriptor %s."), MAX_uint16 - 1U, ToCStr(Descriptor->DebugName));
 	Descriptor->MemberTagDescriptors = MemberTagDescriptors;
 	Descriptor->TagCount = static_cast<uint16>(Context.MemberTagCache.TagCount);
 }
@@ -1174,7 +1174,7 @@ void FPropertyReplicationStateDescriptorBuilder::BuildMemberReferenceDescriptors
 	}
 
 	// Fill in descriptor
-	ensureAlwaysMsgf(Context.ReferenceCache.ReferenceCount < std::numeric_limits<uint16>::max(), TEXT("Object reference count cannot exceed %u. Building descriptor %s."), std::numeric_limits<uint16>::max(), ToCStr(Descriptor->DebugName));
+	ensureMsgf(Context.ReferenceCache.ReferenceCount < std::numeric_limits<uint16>::max(), TEXT("Object reference count cannot exceed %u. Building descriptor %s."), std::numeric_limits<uint16>::max(), ToCStr(Descriptor->DebugName));
 	Descriptor->MemberReferenceDescriptors = MemberReferenceDescriptors;
 	Descriptor->ObjectReferenceCount = static_cast<uint16>(Context.ReferenceCache.ReferenceCount);
 }
@@ -1193,7 +1193,7 @@ void FPropertyReplicationStateDescriptorBuilder::BuildMemberFunctionDescriptors(
 		UE_LOG_DESCRIPTORBUILDER(Verbose, TEXT("FPropertyReplicationStateDescriptorBuilder::BuildMemberFunctionDescriptors AddingFunction %s"), ToCStr(MemberFunction.Function->GetName()));		
 	}
 
-	ensureAlwaysMsgf(Context.MemberFunctionCache.Num() <= std::numeric_limits<uint16>::max(), TEXT("Function count cannot exceed %u. Building descriptor %s."), std::numeric_limits<uint16>::max(), ToCStr(Descriptor->DebugName));
+	ensureMsgf(Context.MemberFunctionCache.Num() <= std::numeric_limits<uint16>::max(), TEXT("Function count cannot exceed %u. Building descriptor %s."), std::numeric_limits<uint16>::max(), ToCStr(Descriptor->DebugName));
 	Descriptor->FunctionCount = static_cast<uint16>(Context.MemberFunctionCache.Num());
 	Descriptor->MemberFunctionDescriptors = MemberFunctionDescriptors;
 }
@@ -2146,7 +2146,7 @@ void FPropertyReplicationStateDescriptorBuilder::GetPropertyPathName(const FProp
 	uint32 PropertyChainIndex = MaxHierarchyDepth;
 	for (FFieldVariant Object = Property; Object.IsValid(); Object = Object.GetOwnerVariant())
 	{
-		if (!ensureAlwaysMsgf(PropertyChainIndex > 0, TEXT("Property hieararchy depth exceeds %u"), MaxHierarchyDepth))
+		if (!ensureMsgf(PropertyChainIndex > 0, TEXT("Property hieararchy depth exceeds %u"), MaxHierarchyDepth))
 		{
 			break;
 		}
@@ -2269,7 +2269,7 @@ EMemberPropertyTraits FPropertyReplicationStateDescriptorBuilder::GetFastArrayPr
 			else if (Struct->IsChildOf(FFastArraySerializerItem::StaticStruct()))
 			{
 				// Invalidate FastArrayItem if it has a custom NetSerializer which we currently do not support, the workaround is to wrap the struct with the custom netserializer in a struct
-				if (ensureAlwaysMsgf(!IsStructWithCustomSerializer(Struct), 
+				if (ensureMsgf(!IsStructWithCustomSerializer(Struct), 
 					TEXT("FPropertyReplicationStateDescriptorBuilder Iris does not support custom NetSerializers for FastArrayItems %s, if required use a property in the struct wrapping the custom NetSerializer"), *Struct->GetName())
 				)
 				{
@@ -2669,7 +2669,7 @@ SIZE_T FReplicationStateDescriptorBuilder::CreateDescriptorsForClass(FResult& Cr
 				else if (EnumHasAnyFlags(MemberProperty.Traits, EMemberPropertyTraits::IsFastArray))
 				{
 					// FastArrayProperties should use a custom replication fragment
-					ensureAlwaysMsgf(false, TEXT("FReplicationStateDescriptorBuilder::CreateDescriptorsForClass FFastArray property %s not registered! Usually this means a call to SetupIrisSupport is needed in the module's Build.cs file."), *Property->GetFullName());
+					ensureMsgf(false, TEXT("FReplicationStateDescriptorBuilder::CreateDescriptorsForClass FFastArray property %s not registered! Usually this means a call to SetupIrisSupport is needed in the module's Build.cs file."), *Property->GetFullName());
 					continue;
 				}
 
@@ -2743,7 +2743,7 @@ SIZE_T FReplicationStateDescriptorBuilder::CreateDescriptorsForClass(FResult& Cr
 		// NetFields only contain the class specific networked fields, excluding super class fields.
 		for (const UClass* CurrentClass = InObjectClass; CurrentClass != nullptr; CurrentClass = CurrentClass->GetSuperClass())
 		{
-			if (!ensureAlwaysMsgf(ClassIndex > 0, TEXT("Class hieararchy depth exceeds %u"), MaxClassHierarchyDepth))
+			if (!ensureMsgf(ClassIndex > 0, TEXT("Class hieararchy depth exceeds %u"), MaxClassHierarchyDepth))
 			{
 				break;
 			}
