Remove UE::IO::IAS::EOnDemandEndpointType as we no longer have any planned use for it.

#rnx

diff --git a/Engine/Source/Runtime/Experimental/IoStoreOnDemand/Private/IoStoreOnDemand.cpp b/Engine/Source/Runtime/Experimental/IoStoreOnDemand/Private/IoStoreOnDemand.cpp
--- a/Engine/Source/Runtime/Experimental/IoStoreOnDemand/Private/IoStoreOnDemand.cpp
+++ b/Engine/Source/Runtime/Experimental/IoStoreOnDemand/Private/IoStoreOnDemand.cpp
@@ -2211,8 +2211,6 @@ void FIoStoreOnDemandModule::InitializeInternal()
 
 	FLatencyInjector::Initialize(CommandLine);
 
-	Endpoint.EndpointType = EOnDemandEndpointType::CDN;
-
 	TUniquePtr<IIasCache> Cache;
 	FIasCacheConfig CacheConfig = GetIasCacheConfig(CommandLine);
 	CacheConfig.DropCache = DeferredAbandonCache.Get(CacheConfig.DropCache);
diff --git a/Engine/Source/Runtime/Experimental/IoStoreOnDemand/Private/OnDemandIoDispatcherBackend.cpp b/Engine/Source/Runtime/Experimental/IoStoreOnDemand/Private/OnDemandIoDispatcherBackend.cpp
--- a/Engine/Source/Runtime/Experimental/IoStoreOnDemand/Private/OnDemandIoDispatcherBackend.cpp
+++ b/Engine/Source/Runtime/Experimental/IoStoreOnDemand/Private/OnDemandIoDispatcherBackend.cpp
@@ -2221,57 +2221,50 @@ void FOnDemandIoBackend::Mount(const FOnDemandEndpoint& Endpoint)
 			});
 	}
 
-	if (EnumHasAnyFlags(Endpoint.EndpointType, EOnDemandEndpointType::CDN))
 	{
+		FWriteScopeLock _(Lock);
+		if (Endpoint.ServiceUrls.IsEmpty())
 		{
-			FWriteScopeLock _(Lock);
-			if (Endpoint.ServiceUrls.IsEmpty())
+			if (AvailableEps.HasCurrent() == false)
 			{
-				if (AvailableEps.HasCurrent() == false)
+				if (!DistributionUrl.IsValid())
 				{
-					if (!DistributionUrl.IsValid())
-					{
-						DistributionUrl = { Endpoint.DistributionUrl, Endpoint.FallbackUrl };
-					}
-					DeferredTocs.Add(FTocParams{Endpoint.TocPath, Endpoint.bForceTocDownload});
-					return;
+					DistributionUrl = { Endpoint.DistributionUrl, Endpoint.FallbackUrl };
 				}
-			}
-			else if (AvailableEps.Urls.IsEmpty())
-			{
-				for (const FString& Url : Endpoint.ServiceUrls)
-				{
-					AvailableEps.Urls.Add(Url.Replace(TEXT("https"), TEXT("http")));
-				}
-				AvailableEps.Current = 0;
+				DeferredTocs.Add(FTocParams{ Endpoint.TocPath, Endpoint.bForceTocDownload });
+				return;
 			}
 		}
-
-		check(AvailableEps.HasCurrent());
-
-		if (GIasGenerateOnDemandToc && !Endpoint.bForceTocDownload)
+		else if (AvailableEps.Urls.IsEmpty())
 		{
-			FIoStatus GeneratedResult = ApplyGeneratedOnDemandToc(AvailableEps.GetCurrent(), Endpoint.TocPath);
-			if (!GeneratedResult.IsOk())
+			for (const FString& Url : Endpoint.ServiceUrls)
 			{
-				UE_LOG(LogIas, Error, TEXT("Failed to add generated toc', reason '%s'"), *GeneratedResult.ToString());
+				AvailableEps.Urls.Add(Url.Replace(TEXT("https"), TEXT("http")));
 			}
+			AvailableEps.Current = 0;
 		}
-		else
+	}
+
+	check(AvailableEps.HasCurrent());
+
+	if (GIasGenerateOnDemandToc && !Endpoint.bForceTocDownload)
+	{
+		FIoStatus GeneratedResult = ApplyGeneratedOnDemandToc(AvailableEps.GetCurrent(), Endpoint.TocPath);
+		if (!GeneratedResult.IsOk())
 		{
-			FIoStatus Result = DownloadoadOnDemandToc(AvailableEps.GetCurrent(), Endpoint.TocPath);
-			if (!Result.IsOk())
-			{
-				UE_LOG(LogIas, Error, TEXT("Deferring TOC '%s/%s' due to '%s'"), *AvailableEps.GetCurrent(), *Endpoint.TocPath, *Result.ToString());
-					BackendStatus.SetHttpError(true);
-					FWriteScopeLock _(Lock);
-					DeferredTocs.Add(FTocParams{Endpoint.TocPath, Endpoint.bForceTocDownload});
-			}
+			UE_LOG(LogIas, Error, TEXT("Failed to add generated toc', reason '%s'"), *GeneratedResult.ToString());
 		}
 	}
 	else
 	{
-		UE_LOG(LogIas, Log, TEXT("Mounting ZEN endpoint, Url='%s'"), *Endpoint.ServiceUrls[0]);
+		FIoStatus Result = DownloadoadOnDemandToc(AvailableEps.GetCurrent(), Endpoint.TocPath);
+		if (!Result.IsOk())
+		{
+			UE_LOG(LogIas, Error, TEXT("Deferring TOC '%s/%s' due to '%s'"), *AvailableEps.GetCurrent(), *Endpoint.TocPath, *Result.ToString());
+			BackendStatus.SetHttpError(true);
+			FWriteScopeLock _(Lock);
+			DeferredTocs.Add(FTocParams{ Endpoint.TocPath, Endpoint.bForceTocDownload });
+		}
 	}
 }
 
diff --git a/Engine/Source/Runtime/Experimental/IoStoreOnDemand/Private/OnDemandIoDispatcherBackend.h b/Engine/Source/Runtime/Experimental/IoStoreOnDemand/Private/OnDemandIoDispatcherBackend.h
--- a/Engine/Source/Runtime/Experimental/IoStoreOnDemand/Private/OnDemandIoDispatcherBackend.h
+++ b/Engine/Source/Runtime/Experimental/IoStoreOnDemand/Private/OnDemandIoDispatcherBackend.h
@@ -11,12 +11,6 @@ struct FAnalyticsEventAttribute;
 namespace UE::IO::IAS
 {
 
-enum class EOnDemandEndpointType
-{
-	CDN = 1,
-	ZEN
-};
-
 struct FDistributedEndpointUrl
 {
 	FString EndpointUrl;
@@ -41,7 +35,6 @@ struct FDistributedEndpointUrl
 
 struct FOnDemandEndpoint
 {
-	EOnDemandEndpointType EndpointType;
 	FString DistributionUrl;
 	FString FallbackUrl;
 
