Chaos Cloth:
- Added a new Intrinsic flag to the property facade to mark properties that when changed requires the simulated object to be re-created.
- Added the GetTypeHash implementation to the property facade.

diff --git a/Engine/Source/Runtime/Experimental/Chaos/Private/Chaos/CollectionPropertyFacade.cpp b/Engine/Source/Runtime/Experimental/Chaos/Private/Chaos/CollectionPropertyFacade.cpp
--- a/Engine/Source/Runtime/Experimental/Chaos/Private/Chaos/CollectionPropertyFacade.cpp
+++ b/Engine/Source/Runtime/Experimental/Chaos/Private/Chaos/CollectionPropertyFacade.cpp
@@ -15,7 +15,7 @@ namespace Chaos::Softs
 		static const FName StringValueName("StringValue");  // String value, or weight map name, ...etc.
 		static const FName FlagsName("Flags");  // Whether this property is enabled, animatable, ...etc.
 	}
-	
+
 	FCollectionPropertyConstFacade::FCollectionPropertyConstFacade(const TSharedPtr<const FManagedArrayCollection>& InManagedArrayCollection)
 		: ManagedArrayCollection(InManagedArrayCollection)
 	{
@@ -111,8 +111,8 @@ namespace Chaos::Softs
 	{
 		// Cannot set string dirty without also dirtying the property
 		Flags |= EnumHasAnyFlags(Flags, ECollectionPropertyFlags::StringDirty) ? ECollectionPropertyFlags::Dirty : ECollectionPropertyFlags::None;
-		// Cannot remove the dirty flags
-		Flags |= GetFlagsArray()[KeyIndex] & (ECollectionPropertyFlags::Dirty | ECollectionPropertyFlags::StringDirty);
+		// Cannot remove the Dirty, StringDirty, and Intrinsic flags
+		Flags |= GetFlagsArray()[KeyIndex] & (ECollectionPropertyFlags::Dirty | ECollectionPropertyFlags::StringDirty | ECollectionPropertyFlags::Intrinsic);
 
 		SetValue(KeyIndex, GetFlagsArray(), Flags);
 	}
@@ -191,11 +191,12 @@ namespace Chaos::Softs
 		RebuildKeyIndices();
 	}
 
-	int32 FCollectionPropertyMutableFacade::AddProperty(const FString& Key, bool bEnabled, bool bAnimatable)
+	int32 FCollectionPropertyMutableFacade::AddProperty(const FString& Key, bool bEnabled, bool bAnimatable, bool bIntrinsic)
 	{
 		const ECollectionPropertyFlags Flags =
 			(bEnabled ? ECollectionPropertyFlags::Enabled : ECollectionPropertyFlags::None) |
-			(bAnimatable ? ECollectionPropertyFlags::Animatable : ECollectionPropertyFlags::None);
+			(bAnimatable ? ECollectionPropertyFlags::Animatable : ECollectionPropertyFlags::None) |
+			(bIntrinsic ? ECollectionPropertyFlags::Intrinsic : ECollectionPropertyFlags::None);
 		return AddProperty(Key, Flags);
 	}
 
@@ -217,11 +218,12 @@ namespace Chaos::Softs
 		return Index;
 	}
 
-	int32 FCollectionPropertyMutableFacade::AddProperties(const TArray<FString>& Keys, bool bEnabled, bool bAnimatable)
+	int32 FCollectionPropertyMutableFacade::AddProperties(const TArray<FString>& Keys, bool bEnabled, bool bAnimatable, bool bIntrinsic)
 	{
 		const ECollectionPropertyFlags Flags =
 			(bEnabled ? ECollectionPropertyFlags::Enabled : ECollectionPropertyFlags::None) |
-			(bAnimatable ? ECollectionPropertyFlags::Animatable : ECollectionPropertyFlags::None);
+			(bAnimatable ? ECollectionPropertyFlags::Animatable : ECollectionPropertyFlags::None) |
+			(bIntrinsic ? ECollectionPropertyFlags::Intrinsic : ECollectionPropertyFlags::None);
 
 		return AddProperties(Keys, Flags);
 	}
@@ -365,9 +367,9 @@ PRAGMA_ENABLE_DEPRECATION_WARNINGS
 		}
 	}
 
-	int32 FCollectionPropertyMutableFacade::AddWeightedFloatValue(const FString& Key, const FVector2f& Value, bool bEnabled, bool bAnimatable)
+	int32 FCollectionPropertyMutableFacade::AddWeightedFloatValue(const FString& Key, const FVector2f& Value, bool bEnabled, bool bAnimatable, bool bIntrinsic)
 	{
-		const int32 KeyIndex = AddProperty(Key, bEnabled, bAnimatable);
+		const int32 KeyIndex = AddProperty(Key, bEnabled, bAnimatable, bIntrinsic);
 		SetWeightedFloatValue(KeyIndex, Value);
 		return KeyIndex;
 	}
@@ -379,9 +381,9 @@ PRAGMA_ENABLE_DEPRECATION_WARNINGS
 		return KeyIndex;
 	}
 
-	int32 FCollectionPropertyMutableFacade::AddStringValue(const FString& Key, const FString& Value, bool bEnabled, bool bAnimatable)
+	int32 FCollectionPropertyMutableFacade::AddStringValue(const FString& Key, const FString& Value, bool bEnabled, bool bAnimatable, bool bIntrinsic)
 	{
-		const int32 KeyIndex = AddProperty(Key, bEnabled, bAnimatable);
+		const int32 KeyIndex = AddProperty(Key, bEnabled, bAnimatable, bIntrinsic);
 		SetStringValue(Key, Value);
 		return KeyIndex;
 	}
@@ -392,4 +394,4 @@ PRAGMA_ENABLE_DEPRECATION_WARNINGS
 		SetStringValue(Key, Value);
 		return KeyIndex;
 	}
-}  // End namespace Chaos
+}  // End namespace Chaos::Softs
diff --git a/Engine/Source/Runtime/Experimental/Chaos/Public/Chaos/CollectionPropertyFacade.h b/Engine/Source/Runtime/Experimental/Chaos/Public/Chaos/CollectionPropertyFacade.h
--- a/Engine/Source/Runtime/Experimental/Chaos/Public/Chaos/CollectionPropertyFacade.h
+++ b/Engine/Source/Runtime/Experimental/Chaos/Public/Chaos/CollectionPropertyFacade.h
@@ -26,8 +26,9 @@ namespace Chaos::Softs
 	{
 		None = 0,
 		Enabled = 1 << 0,  /** Whether this property is enabled(so that it doesn't have to be removed from the collection when not needed). */
-		Animatable = 1 << 1,  /** Whether this property needs to be set at every frame. */
+		Animatable = 1 << 1,  /** Whether this property needs to be set at every frame. This flag is ignored when the Intrinsic flag is also set. */
 		Legacy = 1 << 2,  /** Whether this property has been set by a legacy system predating the property collection. Can be useful for overriding/upgrading some properties post conversion. */
+		Intrinsic = 1 << 4,  /** Whether this property is intrinsically built into the simulated object model, rather than affecting the simulation itself (see Animatable in this case). Changing this property requires a re-construction of the simulated object model to be effective. Implies non Animatable. */
 		//~ Add new flags above this line
 		StringDirty = 1 << 6,  /** Whether this property's string has changed and needs to be updated at the next frame. */
 		Dirty = 1 << 7  /** Whether this property's value has changed and needs to be updated at the next frame. */
@@ -94,8 +95,9 @@ namespace Chaos::Softs
 		uint8 GetFlags(int32 KeyIndex) const { return FlagsArray[KeyIndex]; }
 
 		bool IsEnabled(int32 KeyIndex) const { return HasAnyFlags(KeyIndex, ECollectionPropertyFlags::Enabled); }
-		bool IsAnimatable(int32 KeyIndex) const { return HasAnyFlags(KeyIndex, ECollectionPropertyFlags::Animatable); }
+		bool IsAnimatable(int32 KeyIndex) const { return HasAnyFlags(KeyIndex, ECollectionPropertyFlags::Animatable) && !HasAnyFlags(KeyIndex, ECollectionPropertyFlags::Intrinsic); }
 		bool IsLegacy(int32 KeyIndex) const { return HasAnyFlags(KeyIndex, ECollectionPropertyFlags::Legacy); }
+		bool IsIntrinsic(int32 KeyIndex) const { return HasAnyFlags(KeyIndex, ECollectionPropertyFlags::Intrinsic); }
 		bool IsStringDirty(int32 KeyIndex) const { return HasAnyFlags(KeyIndex, ECollectionPropertyFlags::StringDirty); }
 		bool IsDirty(int32 KeyIndex) const { return HasAnyFlags(KeyIndex, ECollectionPropertyFlags::Dirty); }
 
@@ -157,6 +159,11 @@ PRAGMA_ENABLE_DEPRECATION_WARNINGS
 			return SafeGet(Key, [this](int32 KeyIndex)->bool { return IsLegacy(KeyIndex); }, bDefault, OutKeyIndex);
 		}
 
+		bool IsIntrinsic(const FString& Key, bool bDefault = false, int32* OutKeyIndex = nullptr) const
+		{
+			return SafeGet(Key, [this](int32 KeyIndex)->bool { return IsIntrinsic(KeyIndex); }, bDefault, OutKeyIndex);
+		}
+
 		bool IsStringDirty(const FString& Key, bool bDefault = false, int32* OutKeyIndex = nullptr) const
 		{
 			return SafeGet(Key, [this](int32 KeyIndex)->bool { return IsStringDirty(KeyIndex); }, bDefault, OutKeyIndex);
@@ -167,6 +174,20 @@ PRAGMA_ENABLE_DEPRECATION_WARNINGS
 			return SafeGet(Key, [this](int32 KeyIndex)->bool { return IsDirty(KeyIndex); }, bDefault, OutKeyIndex);
 		}
 
+		friend ::uint32 GetTypeHash(const Chaos::Softs::FCollectionPropertyConstFacade& PropertyFacade)
+		{
+			uint32 Hash = 0;
+			if (PropertyFacade.IsValid())
+			{
+				Hash = GetArrayHash(PropertyFacade.KeyArray.GetData(), PropertyFacade.KeyArray.Num(), Hash);
+				Hash = GetArrayHash(PropertyFacade.LowValueArray.GetData(), PropertyFacade.LowValueArray.Num(), Hash);
+				Hash = GetArrayHash(PropertyFacade.HighValueArray.GetData(), PropertyFacade.HighValueArray.Num(), Hash);
+				Hash = GetArrayHash(PropertyFacade.StringValueArray.GetData(), PropertyFacade.StringValueArray.Num(), Hash);
+				Hash = GetArrayHash(PropertyFacade.FlagsArray.GetData(), PropertyFacade.FlagsArray.Num(), Hash);
+			}
+			return Hash;
+		}
+
 	protected:
 		// No init constructor for FCollectionPropertyFacade
 		CHAOS_API FCollectionPropertyConstFacade(const TSharedPtr<const FManagedArrayCollection>& InManagedArrayCollection, ENoInit);
@@ -255,7 +276,7 @@ PRAGMA_ENABLE_DEPRECATION_WARNINGS
 
 		void SetStringValue(int32 KeyIndex, const FString& Value) { if (GetStringValueArray()[KeyIndex] != Value) { GetStringValueArray()[KeyIndex] = Value; SetStringDirty(KeyIndex); } }
 
-		/** SetFlags cannot be used to remove Dirty or StringDirty flags. Use ClearDirtyFlags to remove dirty flags. */
+		/** SetFlags cannot be used to remove Dirty, StringDirty, or Intrinsic flags. Use ClearDirtyFlags to remove dirty flags. */
 		void SetFlags(int32 KeyIndex, ECollectionPropertyFlags Flags);
 		UE_DEPRECATED(5.3, "Use SetFlags(int32, ECollectionPropertyFlags) instead.")
 		void SetFlags(int32 KeyIndex, uint8 Flags) { return SetFlags(KeyIndex, (ECollectionPropertyFlags)Flags); }
@@ -263,6 +284,9 @@ PRAGMA_ENABLE_DEPRECATION_WARNINGS
 		void SetEnabled(int32 KeyIndex, bool bEnabled) { EnableFlags(KeyIndex, ECollectionPropertyFlags::Enabled, bEnabled); }
 		void SetAnimatable(int32 KeyIndex, bool bAnimatable) { EnableFlags(KeyIndex, ECollectionPropertyFlags::Animatable, bAnimatable); }
 		void SetLegacy(int32 KeyIndex, bool bLegacy) { EnableFlags(KeyIndex, ECollectionPropertyFlags::Legacy, bLegacy); }
+
+		/** Set the intrinsic flag for this property. This flag cannot be removed and implies non Animatable. */
+		void SetIntrinsic(int32 KeyIndex) { EnableFlags(KeyIndex, ECollectionPropertyFlags::Intrinsic, true); }
 		void SetDirty(int32 KeyIndex) { EnableFlags(KeyIndex, ECollectionPropertyFlags::Dirty, true); }
 		UE_DEPRECATED(5.3, "SetDirty can only be set, to unset use ClearDirtyFlags instead.")
 		void SetDirty(int32 KeyIndex, bool bDirty) { EnableFlags(KeyIndex, ECollectionPropertyFlags::Dirty, bDirty); }
@@ -325,6 +349,12 @@ PRAGMA_ENABLE_DEPRECATION_WARNINGS
 			return SafeSet(Key, [this, bLegacy](int32 KeyIndex) { SetLegacy(KeyIndex, bLegacy); });
 		}
 
+		/** Set the intrinsic flag for this property. This flag cannot be removed and implies non Animatable. */
+		int32 SetIntrinsic(const FString& Key)
+		{
+			return SafeSet(Key, [this](int32 KeyIndex) { SetIntrinsic(KeyIndex); });
+		}
+
 		UE_DEPRECATED(5.3, "SetDirty can only be set, to unset use ClearDirtyFlags instead.")
 		int32 SetDirty(const FString& Key, bool bDirty)
 		{
@@ -414,13 +444,13 @@ PRAGMA_ENABLE_DEPRECATION_WARNINGS
 		CHAOS_API int32 AddProperty(const FString& Key, ECollectionPropertyFlags Flags = ECollectionPropertyFlags::Enabled);
 
 		/** Add a single property, and return its index. */
-		CHAOS_API int32 AddProperty(const FString& Key, bool bEnabled, bool bAnimatable = false);
+		CHAOS_API int32 AddProperty(const FString& Key, bool bEnabled, bool bAnimatable = false, bool bIntrinsic = false);
 
 		/** Add new properties, and return the index of the first added property. */
 		CHAOS_API int32 AddProperties(const TArray<FString>& Keys, ECollectionPropertyFlags Flags = ECollectionPropertyFlags::Enabled);
 
 		/** Add new properties, and return the index of the first added property. */
-		CHAOS_API int32 AddProperties(const TArray<FString>& Keys, bool bEnabled, bool bAnimatable = false);
+		CHAOS_API int32 AddProperties(const TArray<FString>& Keys, bool bEnabled, bool bAnimatable = false, bool bIntrinsic = false);
 
 		/**
 		 * Append all properties and values from an existing collection to this property collection.
@@ -453,10 +483,10 @@ PRAGMA_ENABLE_DEPRECATION_WARNINGS
 		template<typename T, TEMPLATE_REQUIRES(TIsWeightedType<T>::Value)>
 		inline int32 AddWeightedValue(const FString& Key, const T& LowValue, const T& HighValue, ECollectionPropertyFlags Flags = ECollectionPropertyFlags::Enabled);
 		template<typename T, TEMPLATE_REQUIRES(TIsWeightedType<T>::Value)>
-		inline int32 AddWeightedValue(const FString& Key, const T& LowValue, const T& HighValue, bool bEnabled, bool bAnimatable = false);
+		inline int32 AddWeightedValue(const FString& Key, const T& LowValue, const T& HighValue, bool bEnabled, bool bAnimatable = false, bool bIntrinsic = false);
 
 		CHAOS_API int32 AddWeightedFloatValue(const FString& Key, const FVector2f& Value, ECollectionPropertyFlags Flags = ECollectionPropertyFlags::Enabled);
-		CHAOS_API int32 AddWeightedFloatValue(const FString& Key, const FVector2f& Value, bool bEnabled, bool bAnimatable);
+		CHAOS_API int32 AddWeightedFloatValue(const FString& Key, const FVector2f& Value, bool bEnabled, bool bAnimatable, bool bIntrinsic = false);
 
 		template<typename T, TEMPLATE_REQUIRES(TIsWeightedType<T>::Value)>
 		int32 AddValue(const FString& Key, const T& Value, ECollectionPropertyFlags Flags = ECollectionPropertyFlags::Enabled) { return AddWeightedValue(Key, Value, Value, Flags); }
@@ -464,7 +494,7 @@ PRAGMA_ENABLE_DEPRECATION_WARNINGS
 		int32 AddValue(const FString& Key, const T& Value, bool bEnabled, bool bAnimatable = false) { return AddWeightedValue(Key, Value, Value, bEnabled, bAnimatable); }
 
 		CHAOS_API int32 AddStringValue(const FString& Key, const FString& Value, ECollectionPropertyFlags Flags = ECollectionPropertyFlags::Enabled);
-		CHAOS_API int32 AddStringValue(const FString& Key, const FString& Value, bool bEnabled, bool bAnimatable = false);
+		CHAOS_API int32 AddStringValue(const FString& Key, const FString& Value, bool bEnabled, bool bAnimatable = false, bool bIntrinsic = false);
 	};
 
 	template<typename T>
@@ -486,9 +516,9 @@ PRAGMA_ENABLE_DEPRECATION_WARNINGS
 	}
 
 	template<typename T, typename TEnableIf<TIsWeightedType<T>::Value, int>::type>
-	inline int32 FCollectionPropertyMutableFacade::AddWeightedValue(const FString& Key, const T& LowValue, const T& HighValue, bool bEnabled, bool bAnimatable)
+	inline int32 FCollectionPropertyMutableFacade::AddWeightedValue(const FString& Key, const T& LowValue, const T& HighValue, bool bEnabled, bool bAnimatable, bool bIntrinsic)
 	{
-		const int32 KeyIndex = AddProperty(Key, bEnabled, bAnimatable);
+		const int32 KeyIndex = AddProperty(Key, bEnabled, bAnimatable, bIntrinsic);
 		SetWeightedValue(KeyIndex, LowValue, HighValue);
 		return KeyIndex;
 	}
