DispatchConsumerData has a stale pointer and uses heap memory after it has been freed

#localization none
#tests linux server

diff --git a/Engine/Source/Runtime/Experimental/Chaos/Private/EventDefaults.cpp b/Engine/Source/Runtime/Experimental/Chaos/Private/EventDefaults.cpp
--- a/Engine/Source/Runtime/Experimental/Chaos/Private/EventDefaults.cpp
+++ b/Engine/Source/Runtime/Experimental/Chaos/Private/EventDefaults.cpp
@@ -287,12 +287,13 @@ namespace Chaos
 						TMap<IPhysicsProxyBase*, TArray<int32>>& AllCollisionsIndicesByPhysicsProxy = CollisionEventData.PhysicsProxyToCollisionIndices.PhysicsProxyToIndicesMap;
 						for (int32 IdxCollision = 0; IdxCollision < NumValidCollisions; ++IdxCollision)
 						{
-							if (DupAllCollisionsDataArray[IdxCollision].Proxy1 != nullptr)
+							if (DupAllCollisionsDataArray[IdxCollision].Proxy1 != nullptr && !DupAllCollisionsDataArray[IdxCollision].Proxy1->GetMarkedDeleted())
 							{
 								int32 NewIdx = AllCollisionsDataArray.Add(DupAllCollisionsDataArray[IdxCollision]);
 								AllCollisionsIndicesByPhysicsProxy.FindOrAdd(AllCollisionsDataArray[NewIdx].Proxy1).Add(FEventManager::EncodeCollisionIndex(NewIdx, false));
 
-								if (AllCollisionsDataArray[NewIdx].Proxy2 && AllCollisionsDataArray[NewIdx].Proxy2 != AllCollisionsDataArray[NewIdx].Proxy1)
+								if (AllCollisionsDataArray[NewIdx].Proxy2 && AllCollisionsDataArray[NewIdx].Proxy2 != AllCollisionsDataArray[NewIdx].Proxy1 
+									&& !AllCollisionsDataArray[NewIdx].Proxy2->GetMarkedDeleted())
 								{
 									AllCollisionsIndicesByPhysicsProxy.FindOrAdd(AllCollisionsDataArray[NewIdx].Proxy2).Add(FEventManager::EncodeCollisionIndex(NewIdx, true));
 								}
