Dataflow : GeometryCollection
- Crash when TransformStarts or InstanceCounts are empty.

diff --git a/Engine/Source/Runtime/Experimental/GeometryCollectionEngine/Private/GeometryCollection/GeometryCollectionObject.cpp b/Engine/Source/Runtime/Experimental/GeometryCollectionEngine/Private/GeometryCollection/GeometryCollectionObject.cpp
--- a/Engine/Source/Runtime/Experimental/GeometryCollectionEngine/Private/GeometryCollection/GeometryCollectionObject.cpp
+++ b/Engine/Source/Runtime/Experimental/GeometryCollectionEngine/Private/GeometryCollection/GeometryCollectionObject.cpp
@@ -343,37 +343,43 @@ void UGeometryCollection::CacheAutoInstanceTransformRemapIndices()
 	}
 
 	const int32 NumMeshes = AutoInstanceMeshes.Num();
+	if(NumMeshes!=0)
+	{
+		TArray<int32> TransformGroups;
+		TransformGroups.AddZeroed(NumMeshes);
+		TArray<int32> TransformStarts;
+		TransformStarts.AddUninitialized(NumMeshes);
+		TArray<int32> InstanceCounts;
+		InstanceCounts.AddUninitialized(NumMeshes);
+		TArray<int32> WrittenTransformCounts;
+		WrittenTransformCounts.AddZeroed(NumMeshes);
 
-	TArray<int32> TransformGroups;
-	TransformGroups.AddZeroed(NumMeshes);
-	TArray<int32> TransformStarts;
-	TransformStarts.AddUninitialized(NumMeshes);
-	TArray<int32> InstanceCounts;
-	InstanceCounts.AddUninitialized(NumMeshes);
-	TArray<int32> WrittenTransformCounts;
-	WrittenTransformCounts.AddZeroed(NumMeshes);
 
-	for (int32 MeshIndex = 0; MeshIndex < NumMeshes; MeshIndex++)
-	{
-		const int32 NumInstances = AutoInstanceMeshes[MeshIndex].NumInstances;
-		TransformStarts[MeshIndex] = MeshIndex == 0 ? 0 : TransformStarts[MeshIndex - 1] + InstanceCounts[MeshIndex - 1];
-		InstanceCounts[MeshIndex] = NumInstances;
-	}
-	
-	AutoInstanceTransformRemapIndices.AddUninitialized(TransformStarts.Last() + InstanceCounts.Last());
+		for (int32 MeshIndex = 0; MeshIndex < NumMeshes; MeshIndex++)
+		{
+			const int32 NumInstances = AutoInstanceMeshes[MeshIndex].NumInstances;
+			TransformStarts[MeshIndex] = MeshIndex == 0 ? 0 : TransformStarts[MeshIndex - 1] + InstanceCounts[MeshIndex - 1];
+			InstanceCounts[MeshIndex] = NumInstances;
+		}
 
-	const int32 NumTransforms = InstancedMeshFacade.GetNumIndices();
-	for (int32 TransformIndex = 0; TransformIndex < NumTransforms; TransformIndex++)
-	{
-		if (GeometryCollection->Children[TransformIndex].Num() == 0)
+		AutoInstanceTransformRemapIndices.AddUninitialized(TransformStarts.Last() + InstanceCounts.Last());
+
+		const int32 NumTransforms = InstancedMeshFacade.GetNumIndices();
+		for (int32 TransformIndex = 0; TransformIndex < NumTransforms; TransformIndex++)
 		{
-			const int32 AutoInstanceMeshIndex = InstancedMeshFacade.GetIndex(TransformIndex);
-			const int32 WriteIndex = WrittenTransformCounts[AutoInstanceMeshIndex];
-			if (WriteIndex < InstanceCounts[AutoInstanceMeshIndex])
+			if (GeometryCollection->Children[TransformIndex].Num() == 0)
 			{
-				const int32 TransformArrayIndex = TransformStarts[AutoInstanceMeshIndex] + WriteIndex;
-			AutoInstanceTransformRemapIndices[TransformArrayIndex] = TransformIndex;
-				WrittenTransformCounts[AutoInstanceMeshIndex]++;
+				const int32 AutoInstanceMeshIndex = InstancedMeshFacade.GetIndex(TransformIndex);
+				const int32 WriteIndex = WrittenTransformCounts[AutoInstanceMeshIndex];
+				if (WriteIndex < InstanceCounts[AutoInstanceMeshIndex])
+				{
+					const int32 TransformArrayIndex = TransformStarts[AutoInstanceMeshIndex] + WriteIndex;
+					if(AutoInstanceTransformRemapIndices.IsValidIndex(TransformArrayIndex))
+					{
+						AutoInstanceTransformRemapIndices[TransformArrayIndex] = TransformIndex;
+						WrittenTransformCounts[AutoInstanceMeshIndex]++;
+					}
+				}
 			}
 		}
 	}
