Add MinValue / MaxValue parameters to ApplyradialImpulse function ( default 0, 1) to allow fine tuning

diff --git a/Engine/Source/Runtime/Engine/Classes/PhysicsEngine/PhysicsObjectBlueprintLibrary.h b/Engine/Source/Runtime/Engine/Classes/PhysicsEngine/PhysicsObjectBlueprintLibrary.h
--- a/Engine/Source/Runtime/Engine/Classes/PhysicsEngine/PhysicsObjectBlueprintLibrary.h
+++ b/Engine/Source/Runtime/Engine/Classes/PhysicsEngine/PhysicsObjectBlueprintLibrary.h
@@ -22,6 +22,20 @@ public:
 	UFUNCTION(BlueprintPure, Category = "Physics Object")
 	static FTransform GetPhysicsObjectWorldTransform(UPrimitiveComponent* Component, FName BoneName);
 
+	/**
+	* Apply a physics radial impulse with an optional strain on a specific component
+	* Effect is applied within a sphere. When using linear falloff the effect will be minimum at the outer edge of the sphere and maximum at its center
+	* @param Component		Primitive component to apply the impulse / strain on
+	* @param Origin			Positition of the origin of the radial effect in world space
+	* @param Radius			Radius of the radial effect ( beyond the radius, impulse will not be applied )
+	* @param Strength		Strength of the impulse to apply ( Unit : (Kg * m / s) or ( m /s ) if bVelChange is true
+	* @param FallOff		Type of falloff to use ( constant, linear )
+	* @param bApplyStrain	Whether or not to apply strain on top of the impulse ( for destructible objects )
+	* @param Strain			If bApplyStrain is true, Strain to apply to the physics particles ( for destructible objects )
+	* @param bVelChange		If true, impulse Strength parameter is interpretation as a change of velocity
+	* @param MinValue		When using linear falloff, this define the falloff value at the outer edge of the sphere
+	* @param MaxValue		When using linear falloff, this define the falloff value at the center of the sphere
+	*/
 	UFUNCTION(BlueprintCallable, Category = "Physics Object")
-	static ENGINE_API void ApplyRadialImpulse(UPrimitiveComponent* Component, FVector Origin, float Radius, float Strength, enum ERadialImpulseFalloff Falloff, bool bApplyStrain, float Strain, bool bVelChange = false);
+	static ENGINE_API void ApplyRadialImpulse(UPrimitiveComponent* Component, FVector Origin, float Radius, float Strength, enum ERadialImpulseFalloff Falloff, bool bApplyStrain, float Strain, bool bVelChange = false, float MinValue = 0.f, float MaxValue = 1.f);
 };
\ No newline at end of file
diff --git a/Engine/Source/Runtime/Engine/Private/PhysicsEngine/PhysicsObjectBlueprintLibrary.cpp b/Engine/Source/Runtime/Engine/Private/PhysicsEngine/PhysicsObjectBlueprintLibrary.cpp
--- a/Engine/Source/Runtime/Engine/Private/PhysicsEngine/PhysicsObjectBlueprintLibrary.cpp
+++ b/Engine/Source/Runtime/Engine/Private/PhysicsEngine/PhysicsObjectBlueprintLibrary.cpp
@@ -54,7 +54,7 @@ FTransform UPhysicsObjectBlueprintLibrary::GetPhysicsObjectWorldTransform(UPrimi
 	return FTransform::Identity;
 }
 
-void UPhysicsObjectBlueprintLibrary::ApplyRadialImpulse(UPrimitiveComponent* Component, FVector Origin, float Radius, float Strength, enum ERadialImpulseFalloff Falloff, bool bApplyStrain, float Strain, bool bVelChange)
+void UPhysicsObjectBlueprintLibrary::ApplyRadialImpulse(UPrimitiveComponent* Component, FVector Origin, float Radius, float Strength, enum ERadialImpulseFalloff Falloff, bool bApplyStrain, float Strain, bool bVelChange, float MinValue, float MaxValue)
 {
 	if (Component == nullptr)
 	{
@@ -66,6 +66,6 @@ void UPhysicsObjectBlueprintLibrary::ApplyRadialImpulse(UPrimitiveComponent* Com
 		TArray<Chaos::FPhysicsObjectHandle>  PhysicsObjects;
 		PhysicsObjects.Add(PhysicsObject);
 		FLockedWritePhysicsObjectExternalInterface Interface = FPhysicsObjectExternalInterface::LockWrite(PhysicsObjects);
-		return Interface->AddRadialImpulse(PhysicsObjects, Origin, Radius, Strength, Falloff, bApplyStrain, Strain, /* bInvalidate */ true, bVelChange);
+		return Interface->AddRadialImpulse(PhysicsObjects, Origin, Radius, Strength, Falloff, bApplyStrain, Strain, /* bInvalidate */ true, bVelChange, MinValue, MaxValue);
 	}
 }
diff --git a/Engine/Source/Runtime/Experimental/Chaos/Private/Chaos/PhysicsObjectInterface.cpp b/Engine/Source/Runtime/Experimental/Chaos/Private/Chaos/PhysicsObjectInterface.cpp
--- a/Engine/Source/Runtime/Experimental/Chaos/Private/Chaos/PhysicsObjectInterface.cpp
+++ b/Engine/Source/Runtime/Experimental/Chaos/Private/Chaos/PhysicsObjectInterface.cpp
@@ -90,7 +90,7 @@ namespace
 	// Apply an impulse based on an origin,radius and falloff parameters
 	// return the fall off value ( between 0 and 1 )
 	template<typename T>
-	float AddRadialImpulseHelper(T ParticleHandle, const FVector& Origin, const float Radius, const float Strength, const enum ERadialImpulseFalloff Falloff, const float VelocityRatio = 1.0f, bool bInvalidate = true, bool bVelChange = false)
+	float AddRadialImpulseHelper(T ParticleHandle, const FVector& Origin, const float Radius, const float Strength, const enum ERadialImpulseFalloff Falloff, const float VelocityRatio = 1.0f, bool bInvalidate = true, bool bVelChange = false, float MinValue = 0.f, float MaxValue = 1.f)
 	{
 		using namespace Chaos;
 
@@ -125,6 +125,19 @@ namespace
 				{
 					const float Distance = FMath::Sqrt(DistanceSquared);
 					FalloffAlpha = static_cast<float>(1.0f - Distance / Radius);
+
+					if (FalloffAlpha >= 0)
+					{
+						const float PrevFalloffAlpha = FalloffAlpha;
+						// remap the value within MinValue/MaxValue within the radius
+						FalloffAlpha = FMath::Lerp(MinValue, MaxValue, FalloffAlpha);
+						UE_LOG(LogChaos, Warning, TEXT("FalloffAlpha = %f = Lerp(%f, %f, %f)"), FalloffAlpha, MinValue, MaxValue, PrevFalloffAlpha);
+					}
+					else
+					{
+						// outside of the sphere, clamp to 0
+						FalloffAlpha = 0.f;
+					}
 				}
 
 				// if the strength was still strong enough to consider
@@ -935,7 +948,7 @@ namespace Chaos
 	}
 
 	template<EThreadContext Id>
-	void FWritePhysicsObjectInterface<Id>::AddRadialImpulse(TArrayView<const FPhysicsObjectHandle> InObjects, FVector Origin, float Radius, float Strength, enum ERadialImpulseFalloff Falloff, bool bApplyStrain, float Strain, bool bInvalidate, bool bVelChange)
+	void FWritePhysicsObjectInterface<Id>::AddRadialImpulse(TArrayView<const FPhysicsObjectHandle> InObjects, FVector Origin, float Radius, float Strength, enum ERadialImpulseFalloff Falloff, bool bApplyStrain, float Strain, bool bInvalidate, bool bVelChange, float MinValue, float MaxValue)
 	{
 		//TODO: create a PT version of this, plus the damping functions
 		if (Chaos::FPBDRigidsSolver* RigidSolver = Chaos::FPhysicsObjectInterface::GetSolver(InObjects))
@@ -950,7 +963,9 @@ namespace Chaos
 				bApplyStrain,
 				Strain,
 				bInvalidate,
-				bVelChange
+				bVelChange,
+				MinValue,
+				MaxValue
 				]()
 				{
 					using namespace Chaos;
@@ -984,7 +999,7 @@ namespace Chaos
 
 								for (FPBDRigidParticleHandle* ChildHandle : *ChildrenHandles)
 								{
-									const float FalloffAlpha = AddRadialImpulseHelper(ChildHandle, Origin, Radius, Strength, Falloff, VelocityRatio, bInvalidate, bVelChange);
+									const float FalloffAlpha = AddRadialImpulseHelper(ChildHandle, Origin, Radius, Strength, Falloff, VelocityRatio, bInvalidate, bVelChange, MinValue, MaxValue);
 
 									// todo(chaos) : Remove StrainModifier cvar when material system is in place and densities are updated
 									const float StrainToApply =
@@ -1006,8 +1021,8 @@ namespace Chaos
 									RigidSolver->GetEvolution()->SetParticleObjectState(ParticleHandle, Chaos::EObjectStateType::Dynamic);
 									RigidSolver->GetEvolution()->GetParticles().MarkTransientDirtyParticle(ParticleHandle);
 								}
-
-								AddRadialImpulseHelper(ParticleHandle, Origin, Radius, Strength, Falloff, bInvalidate, bVelChange);
+								constexpr float VelocityRatio = 1.f;
+								AddRadialImpulseHelper(ParticleHandle, Origin, Radius, Strength, Falloff, VelocityRatio, bInvalidate, bVelChange, MinValue, MaxValue);
 							}
 						}
 					}
diff --git a/Engine/Source/Runtime/Experimental/Chaos/Public/Chaos/PhysicsObjectInterface.h b/Engine/Source/Runtime/Experimental/Chaos/Public/Chaos/PhysicsObjectInterface.h
--- a/Engine/Source/Runtime/Experimental/Chaos/Public/Chaos/PhysicsObjectInterface.h
+++ b/Engine/Source/Runtime/Experimental/Chaos/Public/Chaos/PhysicsObjectInterface.h
@@ -132,7 +132,7 @@ namespace Chaos
 		UE_DEPRECATED(5.4, "This version AddRadialImpulse has been deprecated. Please use the version where the strain value is passed explicitly")
 		CHAOS_API void AddRadialImpulse(TArrayView<const FPhysicsObjectHandle> InObjects, FVector Origin, float Radius, float Strength, enum ERadialImpulseFalloff Falloff, bool bApplyStrain, bool bInvalidate, bool bVelChange = false);
 
-		CHAOS_API void AddRadialImpulse(TArrayView<const FPhysicsObjectHandle> InObjects, FVector Origin, float Radius, float Strength, enum ERadialImpulseFalloff Falloff, bool bApplyStrain, float Strain, bool bInvalidate, bool bVelChange = false);
+		CHAOS_API void AddRadialImpulse(TArrayView<const FPhysicsObjectHandle> InObjects, FVector Origin, float Radius, float Strength, enum ERadialImpulseFalloff Falloff, bool bApplyStrain, float Strain, bool bInvalidate, bool bVelChange = false, float MinValue = 0.f, float MaxValue = 1.f);
 		CHAOS_API void SetLinearEtherDrag(TArrayView<const FPhysicsObjectHandle> InObjects, float InLinearDrag);
 		CHAOS_API void SetAngularEtherDrag(TArrayView<const FPhysicsObjectHandle> InObjects, float InAngularDrag);
 
