Iris
* Added PrintObjectFromIndex to make Iris logs more useful
* Added logs to NetObjectGroups

#rnx

diff --git a/Engine/Source/Runtime/Experimental/Iris/Core/Private/Iris/ReplicationState/ReplicationStateDescriptorBuilder.cpp b/Engine/Source/Runtime/Experimental/Iris/Core/Private/Iris/ReplicationState/ReplicationStateDescriptorBuilder.cpp
--- a/Engine/Source/Runtime/Experimental/Iris/Core/Private/Iris/ReplicationState/ReplicationStateDescriptorBuilder.cpp
+++ b/Engine/Source/Runtime/Experimental/Iris/Core/Private/Iris/ReplicationState/ReplicationStateDescriptorBuilder.cpp
@@ -2269,12 +2269,15 @@ EMemberPropertyTraits FPropertyReplicationStateDescriptorBuilder::GetFastArrayPr
 			else if (Struct->IsChildOf(FFastArraySerializerItem::StaticStruct()))
 			{
 				// Invalidate FastArrayItem if it has a custom NetSerializer which we currently do not support, the workaround is to wrap the struct with the custom netserializer in a struct
-				if (ensureMsgf(!IsStructWithCustomSerializer(Struct), 
-					TEXT("FPropertyReplicationStateDescriptorBuilder Iris does not support custom NetSerializers for FastArrayItems %s, if required use a property in the struct wrapping the custom NetSerializer"), *Struct->GetName())
-				)
+				if (IsStructWithCustomSerializer(Struct))
 				{
-					return EMemberPropertyTraits::IsFastArrayItem;
+					UE_LOG(LogIris, Error, TEXT("FPropertyReplicationStateDescriptorBuilder found unsupported custom NetSerializers for FastArrayItems %s"), *Struct->GetName());
+					ensureMsgf(false, TEXT("FPropertyReplicationStateDescriptorBuilder Iris does not support custom NetSerializers for FastArrayItems %s, if required use a property in the struct wrapping the custom NetSerializer"), *Struct->GetName());
+
+					return EMemberPropertyTraits::None;
 				}
+
+				return EMemberPropertyTraits::IsFastArrayItem;
 			}
 		}
 	}
@@ -2669,6 +2672,7 @@ SIZE_T FReplicationStateDescriptorBuilder::CreateDescriptorsForClass(FResult& Cr
 				else if (EnumHasAnyFlags(MemberProperty.Traits, EMemberPropertyTraits::IsFastArray))
 				{
 					// FastArrayProperties should use a custom replication fragment
+					UE_LOG(LogIris, Error, TEXT("FReplicationStateDescriptorBuilder::CreateDescriptorsForClass FFastArray property %s not registered and won't be replicated."), *Property->GetFullName());
 					ensureMsgf(false, TEXT("FReplicationStateDescriptorBuilder::CreateDescriptorsForClass FFastArray property %s not registered! Usually this means a call to SetupIrisSupport is needed in the module's Build.cs file."), *Property->GetFullName());
 					continue;
 				}
diff --git a/Engine/Source/Runtime/Experimental/Iris/Core/Private/Iris/ReplicationSystem/Filtering/NetObjectGroups.cpp b/Engine/Source/Runtime/Experimental/Iris/Core/Private/Iris/ReplicationSystem/Filtering/NetObjectGroups.cpp
--- a/Engine/Source/Runtime/Experimental/Iris/Core/Private/Iris/ReplicationSystem/Filtering/NetObjectGroups.cpp
+++ b/Engine/Source/Runtime/Experimental/Iris/Core/Private/Iris/ReplicationSystem/Filtering/NetObjectGroups.cpp
@@ -4,15 +4,16 @@
 #include "Math/UnrealMathUtility.h"
 #include "Iris/Core/IrisLog.h"
 #include "Iris/ReplicationSystem/Filtering/NetObjectFilter.h"
-#include "Iris/ReplicationSystem/NetRefHandleManager.h" // for InvalidInternalIndex
+#include "Iris/ReplicationSystem/NetRefHandleManager.h"
 #include "Containers/ArrayView.h"
 
+DEFINE_LOG_CATEGORY_STATIC(LogIrisGroup, Log, All)
+
 namespace UE::Net::Private
 {
 
 FNetObjectGroups::FNetObjectGroups()
-: MaxGroupCount(0U)
-, CurrentEpoch(++NextEpoch)
+: CurrentEpoch(++NextEpoch)
 {
 }
 
@@ -73,6 +74,8 @@ void FNetObjectGroups::ResetGroupMembership(FNetObjectGroupMembership& Target)
 
 void FNetObjectGroups::Init(const FNetObjectGroupInitParams& Params)
 {
+	NetRefHandleManager = Params.NetRefHandleManager;
+
 	ensureMsgf(Params.MaxGroupCount < std::numeric_limits<FNetObjectGroupHandle::FGroupIndexType>::max(), TEXT("MaxGroupCount cannot exceed %u"), std::numeric_limits<FNetObjectGroupHandle::FGroupIndexType>::max());
 	MaxGroupCount = FMath::Clamp<uint32>(Params.MaxGroupCount, 0U, std::numeric_limits<FNetObjectGroupHandle::FGroupIndexType>::max());
 
@@ -226,6 +229,9 @@ void FNetObjectGroups::AddToGroup(FNetObjectGroupHandle GroupHandle, FInternalNe
 	FNetObjectGroup* Group = GetGroup(GroupHandle);
 	if (InternalIndex != FNetRefHandleManager::InvalidInternalIndex && Group)
 	{
+		UE_LOG(LogIrisGroup, Verbose, TEXT("FNetObjectGroups::AddToGroup Adding %s to Group %s (GroupIndex: %u)"),
+			*NetRefHandleManager->PrintObjectFromIndex(InternalIndex), *GetGroupName(GroupHandle).ToString(), GroupHandle.GetGroupIndex());
+
 		if (AddGroupMembership(GroupMemberships[InternalIndex], GroupHandle))
 		{
 			Group->Members.AddUnique(InternalIndex);
@@ -237,7 +243,8 @@ void FNetObjectGroups::AddToGroup(FNetObjectGroupHandle GroupHandle, FInternalNe
 		}
 		else
 		{
-			UE_LOG(LogIris, Error, TEXT("FNetObjectGroups::AddToGroup, Failed to add ( InternalIndex: %u ) to Group %s (GroupIndex: %u) A NetObject can only be a member of %u groups."), InternalIndex, *GetGroupName(GroupHandle).ToString(), GroupHandle.GetRawValue(), FNetObjectGroupMembership::MaxAssignedGroupCount);
+			UE_LOG(LogIrisGroup, Error, TEXT("FNetObjectGroups::AddToGroup Failed to add %s to Group %s (GroupIndex: %u) A NetObject can only be a member of %u groups."),
+				*NetRefHandleManager->PrintObjectFromIndex(InternalIndex), *GetGroupName(GroupHandle).ToString(), GroupHandle.GetGroupIndex(), FNetObjectGroupMembership::MaxAssignedGroupCount);
 			ensure(false);
 		}
 	}
@@ -248,6 +255,9 @@ void FNetObjectGroups::RemoveFromGroup(FNetObjectGroupHandle GroupHandle, FInter
 	FNetObjectGroup* Group = GetGroup(GroupHandle);
 	if (InternalIndex != FNetRefHandleManager::InvalidInternalIndex && Group)
 	{
+		UE_LOG(LogIrisGroup, Verbose, TEXT("FNetObjectGroups::RemoveFromGroup Removing %s from Group %s (GroupIndex: %u)"),
+			*NetRefHandleManager->PrintObjectFromIndex(InternalIndex), *GetGroupName(GroupHandle).ToString(), GroupHandle.GetGroupIndex());
+
 		FNetObjectGroupMembership& GroupMembership = GroupMemberships[InternalIndex];
 		checkSlow(IsMemberOf(GroupMembership, GroupHandle));
 
@@ -268,6 +278,8 @@ void FNetObjectGroups::AddExclusionFilterTrait(FNetObjectGroupHandle GroupHandle
 	{
 		if (!IsFilterGroup(*Group))
 		{
+			UE_LOG(LogIrisGroup, Verbose, TEXT("FNetObjectGroups::AddExclusionFilterTrait to Group %s (GroupIndex: %u)"), *GetGroupName(GroupHandle).ToString(), GroupHandle.GetGroupIndex());
+
 			Group->Traits |= ENetObjectGroupTraits::IsExclusionFiltering;
 
 			// Flag all current members of this group that they are now filterable
@@ -292,6 +304,8 @@ void FNetObjectGroups::RemoveExclusionFilterTrait(FNetObjectGroupHandle GroupHan
 		return;
 	}
 
+	UE_LOG(LogIrisGroup, Verbose, TEXT("FNetObjectGroups::RemoveExclusionFilterTrait to Group %s (GroupIndex: %u)"), *GetGroupName(GroupHandle).ToString(), GroupHandle.GetGroupIndex());
+
 	Group->Traits &= ~(ENetObjectGroupTraits::IsExclusionFiltering);
 
 	for (FInternalNetRefIndex MemberIndex : Group->Members)
diff --git a/Engine/Source/Runtime/Experimental/Iris/Core/Private/Iris/ReplicationSystem/Filtering/NetObjectGroups.h b/Engine/Source/Runtime/Experimental/Iris/Core/Private/Iris/ReplicationSystem/Filtering/NetObjectGroups.h
--- a/Engine/Source/Runtime/Experimental/Iris/Core/Private/Iris/ReplicationSystem/Filtering/NetObjectGroups.h
+++ b/Engine/Source/Runtime/Experimental/Iris/Core/Private/Iris/ReplicationSystem/Filtering/NetObjectGroups.h
@@ -15,6 +15,8 @@ namespace UE::Net
 	namespace Private
 	{
 		typedef uint32 FInternalNetRefIndex;
+
+		class FNetRefHandleManager;
 	}
 }
 
@@ -40,8 +42,9 @@ struct FNetObjectGroup
 
 struct FNetObjectGroupInitParams
 {
-	uint32 MaxObjectCount;
-	uint32 MaxGroupCount;
+	FNetRefHandleManager* NetRefHandleManager = nullptr;
+	uint32 MaxObjectCount = 0;
+	uint32 MaxGroupCount = 0;
 };
 
 class FNetObjectGroups
@@ -54,6 +57,7 @@ public:
 
 	void Init(const FNetObjectGroupInitParams& Params);
 
+	//$IRIS TODO: Groups should have a mandatory unique name. Otherwise it's impossible to debug issues with groups.
 	FNetObjectGroupHandle CreateGroup();
 	void DestroyGroup(FNetObjectGroupHandle GroupHandle);
 	void ClearGroup(FNetObjectGroupHandle GroupHandle);
@@ -133,12 +137,16 @@ private:
 
 	bool IsInAnyFilterGroup(const FNetObjectGroupMembership& GroupMembership) const;
 
+private:
+
+	FNetRefHandleManager* NetRefHandleManager = nullptr;
+
 	// Group usage pattern should not be high frequency so memory layout should not be a major concern
 	TSparseArray<FNetObjectGroup> Groups;
 
 	// Track what groups each internal handle is a member of, we can tighten this up a bit if needed
 	TArray<FNetObjectGroupMembership> GroupMemberships;
-	uint32 MaxGroupCount;
+	uint32 MaxGroupCount = 0U;
 
 	// List of objects that are members of a group with a filter trait
 	FNetBitArray GroupFilteredOutObjects;
diff --git a/Engine/Source/Runtime/Experimental/Iris/Core/Private/Iris/ReplicationSystem/NetRefHandleManager.cpp b/Engine/Source/Runtime/Experimental/Iris/Core/Private/Iris/ReplicationSystem/NetRefHandleManager.cpp
--- a/Engine/Source/Runtime/Experimental/Iris/Core/Private/Iris/ReplicationSystem/NetRefHandleManager.cpp
+++ b/Engine/Source/Runtime/Experimental/Iris/Core/Private/Iris/ReplicationSystem/NetRefHandleManager.cpp
@@ -157,17 +157,8 @@ const FReplicationInstanceProtocol* FNetRefHandleManager::DetachInstanceProtocol
 		const bool bDetachedInstancedIsPendingEndReplication = ReplicatedObjectData[InternalIndex].bPendingEndReplication == 1U;
 		if (!bDetachedInstancedIsPendingEndReplication)
 		{
-			ensureMsgf(false, TEXT("DetachInstanceProtocol for %s (%s) from client"), *Data.RefHandle.ToString(), Data.Protocol->DebugName->Name);
-
-			if (Data.SubObjectRootIndex != InvalidInternalIndex)
-			{
-				FReplicatedObjectData& RootObjectData = ReplicatedObjectData[Data.SubObjectRootIndex];
-				UE_LOG(LogIris, Warning, TEXT("FNetRefHandleManager::DetachInstanceProtocol - DetachInstanceProtocol %s (%s Root:%s) from client"), *Data.RefHandle.ToString(), Data.Protocol->DebugName->Name, RootObjectData.Protocol->DebugName->Name);
-			}
-			else
-			{
-				UE_LOG(LogIris, Warning, TEXT("FNetRefHandleManager::DetachInstanceProtocol - DetachInstanceProtocol %s (%s) from client"), *Data.RefHandle.ToString(), Data.Protocol->DebugName->Name);
-			}
+			UE_LOG(LogIris, Warning, TEXT("FNetRefHandleManager::DetachInstanceProtocol - DetachInstanceProtocol %s %s (%s) from client"), *PrintObjectFromIndex(InternalIndex),*Data.RefHandle.ToString(), Data.Protocol->DebugName->Name);
+			ensureMsgf(false, TEXT("DetachInstanceProtocol for %s %s (%s) from client"), *PrintObjectFromIndex(InternalIndex),*Data.RefHandle.ToString(), Data.Protocol->DebugName->Name);
 		}
 		
 		Data.InstanceProtocol = nullptr;
@@ -834,4 +825,29 @@ void FNetRefHandleManager::OnPostSendUpdate()
 	ScopeFrameData.bIsValid = false;
 }
 
-}
\ No newline at end of file
+FString FNetRefHandleManager::PrintObjectFromIndex(FInternalNetRefIndex ObjectIndex) const
+{
+	if (ObjectIndex != InvalidInternalIndex)
+	{
+		const FNetRefHandle NetRefHandle = GetNetRefHandleFromInternalIndex(ObjectIndex);
+		const FReplicatedObjectData& ObjectData = GetReplicatedObjectDataNoCheck(ObjectIndex);
+
+		if (ObjectData.SubObjectRootIndex == InvalidInternalIndex)
+		{
+			return FString::Printf(TEXT("RootObject %s (InternalIndex: %u) (%s)"), *GetNameSafe(ReplicatedInstances[ObjectIndex]), ObjectIndex, *NetRefHandle.ToString());
+		}
+		else
+		{
+			const FNetRefHandle RootNetRefHandle = GetNetRefHandleFromInternalIndex(ObjectData.SubObjectRootIndex);
+			return FString::Printf(TEXT("SubObject %s (InternalIndex: %u) (%s) tied to RootObject %s (InternalIndex: %u) (%s)"), 
+								*GetNameSafe(ReplicatedInstances[ObjectIndex]), ObjectIndex, *NetRefHandle.ToString(),
+								*GetNameSafe(ReplicatedInstances[ObjectData.SubObjectRootIndex]), ObjectData.SubObjectRootIndex, *RootNetRefHandle.ToString());
+		}
+	}
+	else
+	{
+		return FString(TEXT("InvalidObject (InternalIndex: Invalid)"));
+	}
+}
+
+} // end namespace UE::Net::Private
diff --git a/Engine/Source/Runtime/Experimental/Iris/Core/Private/Iris/ReplicationSystem/NetRefHandleManager.h b/Engine/Source/Runtime/Experimental/Iris/Core/Private/Iris/ReplicationSystem/NetRefHandleManager.h
--- a/Engine/Source/Runtime/Experimental/Iris/Core/Private/Iris/ReplicationSystem/NetRefHandleManager.h
+++ b/Engine/Source/Runtime/Experimental/Iris/Core/Private/Iris/ReplicationSystem/NetRefHandleManager.h
@@ -265,11 +265,13 @@ public:
 	// Get child subobjects and condtionals array if one exists, if there are no child subobjects the method returns false
 	inline bool GetChildSubObjects(FInternalNetRefIndex OwnerIndex, FChildSubObjectsInfo& OutInfo) const;
 
+	/** Get the map that translates RefHandles into InternalIndexes */
 	const FRefHandleMap& GetReplicatedHandles() const { return RefHandleToInternalIndex; }
 
 	// Get the replicated object represented by a given internal index.
 	UObject* GetReplicatedObjectInstance(FInternalNetRefIndex ObjectIndex) const { return ReplicatedInstances[ObjectIndex]; }
 
+	/** Get the array of all held object pointers */
 	const TArray<UObject*>& GetReplicatedInstances() const { return ReplicatedInstances; }
 
 	void AddReferencedObjects(FReferenceCollector& Collector);
@@ -279,9 +281,15 @@ public:
 
 	const FNetBitArray& GetDestroyedStartupObjectInternalIndices() const { return DestroyedStartupObjectInternalIndices; }
 
+	/** List of replicated objects that want to be dormant */
 	const FNetBitArray& GetWantToBeDormantInternalIndices() const { return WantToBeDormantInternalIndices; }
 	FNetBitArray& GetWantToBeDormantInternalIndices() { return WantToBeDormantInternalIndices; }
 
+	/** Return a string to identify the object linked to an index in logs */
+	FString PrintObjectFromIndex(FInternalNetRefIndex ObjectIndex) const;
+
+public:
+
 	// Iterate over all dependent objects and their dependent objects
 	template <typename T>
 	void ForAllDependentObjectsRecursive(FInternalNetRefIndex ObjectIndex, T&& Functor) const
diff --git a/Engine/Source/Runtime/Experimental/Iris/Core/Private/Iris/ReplicationSystem/ObjectReplicationBridge.cpp b/Engine/Source/Runtime/Experimental/Iris/Core/Private/Iris/ReplicationSystem/ObjectReplicationBridge.cpp
--- a/Engine/Source/Runtime/Experimental/Iris/Core/Private/Iris/ReplicationSystem/ObjectReplicationBridge.cpp
+++ b/Engine/Source/Runtime/Experimental/Iris/Core/Private/Iris/ReplicationSystem/ObjectReplicationBridge.cpp
@@ -1973,13 +1973,13 @@ void UObjectReplicationBridge::OnErrorWithNetRefHandleReported(uint32 ErrorType,
 	// Ensure at the end so the log contains all the relevant information
 	ON_SCOPE_EXIT
 	{
-		ensureMsgf(false, TEXT("NetRefHandle error(%u) reported. Look at the log for important information on the object tied to the handle."), ErrorType);
+		ensureMsgf(false, TEXT("NetRefHandle error: %u reported. Look at the log for important information on the object tied to the handle."), ErrorType);
 	};
 
 	const FInternalNetRefIndex ObjectInternalIndex = NetRefHandleManager->GetInternalIndex(RefHandle);
 	if (ObjectInternalIndex == FNetRefHandleManager::InvalidInternalIndex)
 	{
-		UE_LOG(LogIris, Warning, TEXT("OnErrorWithNetRefHandleReported(%u) from Connection:%u for %s but object has no InternalIndex."), ErrorType, ConnectionId, *RefHandle.ToString());
+		UE_LOG(LogIris, Warning, TEXT("OnErrorWithNetRefHandleReported: %u from Connection:%u for %s but object has no InternalIndex."), ErrorType, ConnectionId, *RefHandle.ToString());
 		return;
 	}
 
@@ -1991,17 +1991,17 @@ void UObjectReplicationBridge::OnErrorWithNetRefHandleReported(uint32 ErrorType,
 		UObject* RootObjInstance = NetRefHandleManager->GetReplicatedObjectInstance(ObjData.SubObjectRootIndex);
 		const FNetRefHandle RootObjNetHandle = NetRefHandleManager->GetNetRefHandleFromInternalIndex(ObjData.SubObjectRootIndex);
 		
-		UE_LOG(LogIris, Error, TEXT("OnErrorWithNetRefHandleReported(%u) from client:%u. %s maps to SubObject: %s owned by RootObject: %s using %s"), 
+		UE_LOG(LogIris, Error, TEXT("OnErrorWithNetRefHandleReported: %u from client:%u. %s maps to SubObject: %s owned by RootObject: %s using %s"), 
 			ErrorType, ConnectionId, 
-			*RefHandle.ToString(), *GetNameSafe(ObjInstance),
+			*RefHandle.ToString(), *GetPathNameSafe(ObjInstance),
 			*GetNameSafe(RootObjInstance), *RootObjNetHandle.ToString()
 		);
 	}
 	else
 	{
-		UE_LOG(LogIris, Error, TEXT("OnErrorWithNetRefHandleReported(%u) from client:%u. %s maps to RootObject: %s"), 
+		UE_LOG(LogIris, Error, TEXT("OnErrorWithNetRefHandleReported: %u from client:%u. %s maps to RootObject: %s"), 
 			ErrorType, ConnectionId, 
-			*RefHandle.ToString(), *GetNameSafe(ObjInstance)
+			*RefHandle.ToString(), *GetPathNameSafe(ObjInstance)
 		);
 	}
 }
diff --git a/Engine/Source/Runtime/Experimental/Iris/Core/Private/Iris/ReplicationSystem/ReplicationSystem.cpp b/Engine/Source/Runtime/Experimental/Iris/Core/Private/Iris/ReplicationSystem/ReplicationSystem.cpp
--- a/Engine/Source/Runtime/Experimental/Iris/Core/Private/Iris/ReplicationSystem/ReplicationSystem.cpp
+++ b/Engine/Source/Runtime/Experimental/Iris/Core/Private/Iris/ReplicationSystem/ReplicationSystem.cpp
@@ -122,6 +122,7 @@ public:
 		FNetObjectGroups& Groups = ReplicationSystemInternal.GetGroups();
 		{
 			FNetObjectGroupInitParams InitParams = {};
+			InitParams.NetRefHandleManager = &NetRefHandleManager;
 			InitParams.MaxObjectCount = MaxObjectCount;
 			InitParams.MaxGroupCount = Params.MaxNetObjectGroupCount;
 
