geometry collection : fix pre-placed geometry collection not properly setting collision when calling SetUseStaticMeshCollisionForTraces from construction script
- When placing in editor, construction script runs before physics proxy is set , which was not setting the collision properly
- Move the proxy side flag to be part of the simulation parameters, so that when the proxy get created it is properly set from the component bUseStaticMeshCollisionForTraces property

diff --git a/Engine/Source/Runtime/Experimental/Chaos/Private/PhysicsProxy/GeometryCollectionPhysicsProxy.cpp b/Engine/Source/Runtime/Experimental/Chaos/Private/PhysicsProxy/GeometryCollectionPhysicsProxy.cpp
--- a/Engine/Source/Runtime/Experimental/Chaos/Private/PhysicsProxy/GeometryCollectionPhysicsProxy.cpp
+++ b/Engine/Source/Runtime/Experimental/Chaos/Private/PhysicsProxy/GeometryCollectionPhysicsProxy.cpp
@@ -759,7 +759,7 @@ void FGeometryCollectionPhysicsProxy::CreateGTParticles(const TBitArray<>& Effec
 			P->SetUserData(Parameters.UserData);
 			P->SetProxy(this);
 		}
-		if (bInitializationTime && TransformIndex == Parameters.InitialRootIndex && bUseStaticMeshCollisionForTraces && CreateTraceCollisionGeometryCallback != nullptr)
+		if (bInitializationTime && TransformIndex == Parameters.InitialRootIndex && Parameters.bUseStaticMeshCollisionForTraces && CreateTraceCollisionGeometryCallback != nullptr)
 		{
 			const FTransform ToLocal = MassToLocal[TransformIndex].Inverse();
 			TArray<Chaos::FImplicitObjectPtr> Geoms;
@@ -3351,7 +3351,7 @@ void FGeometryCollectionPhysicsProxy::SetUseStaticMeshCollisionForTraces_Externa
 
 	// Expensive, so only do it when we need to
 	const bool bUseSMCollisionForTraces = ((bInUseStaticMeshCollisionForTraces && ForceOverrideGCCollisionSetupForTraces == GCCSFT_Property) || (ForceOverrideGCCollisionSetupForTraces == GCCSFT_ForceSM));
-	if (bUseStaticMeshCollisionForTraces != bUseSMCollisionForTraces && CreateTraceCollisionGeometryCallback != nullptr)
+	if (Parameters.bUseStaticMeshCollisionForTraces != bUseSMCollisionForTraces && CreateTraceCollisionGeometryCallback != nullptr)
 	{
 		const int32 TransformIndex = Parameters.InitialRootIndex;
 		const int32 ParticleIndex = FromTransformToParticleIndex[TransformIndex];
@@ -3385,7 +3385,8 @@ void FGeometryCollectionPhysicsProxy::SetUseStaticMeshCollisionForTraces_Externa
 				}
 			}
 			// We only want to change the value of our boolean if we actually changed the collision, so the collision state matches
-			bUseStaticMeshCollisionForTraces = bUseSMCollisionForTraces;
+			// note we can set thsi member of Parameters on the GT because it is only used and set on the GT 
+			Parameters.bUseStaticMeshCollisionForTraces = bUseSMCollisionForTraces;
 		}
 	}
 }
diff --git a/Engine/Source/Runtime/Experimental/Chaos/Public/GeometryCollection/GeometryCollectionSimulationCoreTypes.h b/Engine/Source/Runtime/Experimental/Chaos/Public/GeometryCollection/GeometryCollectionSimulationCoreTypes.h
--- a/Engine/Source/Runtime/Experimental/Chaos/Public/GeometryCollection/GeometryCollectionSimulationCoreTypes.h
+++ b/Engine/Source/Runtime/Experimental/Chaos/Public/GeometryCollection/GeometryCollectionSimulationCoreTypes.h
@@ -224,6 +224,7 @@ struct FSimulationParameters
 		, QueryFilterData()
 		, UserData(nullptr)
 		, bEnableStrainOnCollision(true)
+		, bUseStaticMeshCollisionForTraces(false)
 	{}
 
 	FSimulationParameters(const FSimulationParameters& Other)
@@ -285,6 +286,7 @@ struct FSimulationParameters
 		, QueryFilterData(Other.QueryFilterData)
 		, UserData(Other.UserData)
 		, bEnableStrainOnCollision(Other.bEnableStrainOnCollision)
+		, bUseStaticMeshCollisionForTraces(Other.bUseStaticMeshCollisionForTraces)
 	{
 	}
 
@@ -378,4 +380,6 @@ struct FSimulationParameters
 	FCollisionFilterData QueryFilterData;
 	void* UserData;
 	bool bEnableStrainOnCollision;
+
+	bool bUseStaticMeshCollisionForTraces;
 };
diff --git a/Engine/Source/Runtime/Experimental/Chaos/Public/PhysicsProxy/GeometryCollectionPhysicsProxy.h b/Engine/Source/Runtime/Experimental/Chaos/Public/PhysicsProxy/GeometryCollectionPhysicsProxy.h
--- a/Engine/Source/Runtime/Experimental/Chaos/Public/PhysicsProxy/GeometryCollectionPhysicsProxy.h
+++ b/Engine/Source/Runtime/Experimental/Chaos/Public/PhysicsProxy/GeometryCollectionPhysicsProxy.h
@@ -686,9 +686,6 @@ private:
 	// todo(chaos): Remove this and move to a cook time approach of the SM data based on the GC property
 	FCreateTraceCollisionGeometryCallback CreateTraceCollisionGeometryCallback;
 	
-	// Use SetUseStaticMeshCollisionForTraces_External instead of setting this directly, so that the collision can be recreated if needed
-	bool bUseStaticMeshCollisionForTraces = false;
-
 #ifdef TODO_REIMPLEMENT_RIGID_CACHING
 	TFunction<void(void)> ResetAnimationCacheCallback;
 	TFunction<void(const TArrayView<FTransform> &)> UpdateTransformsCallback;
diff --git a/Engine/Source/Runtime/Experimental/GeometryCollectionEngine/Private/GeometryCollection/GeometryCollectionComponent.cpp b/Engine/Source/Runtime/Experimental/GeometryCollectionEngine/Private/GeometryCollection/GeometryCollectionComponent.cpp
--- a/Engine/Source/Runtime/Experimental/GeometryCollectionEngine/Private/GeometryCollection/GeometryCollectionComponent.cpp
+++ b/Engine/Source/Runtime/Experimental/GeometryCollectionEngine/Private/GeometryCollection/GeometryCollectionComponent.cpp
@@ -3831,6 +3831,7 @@ void UGeometryCollectionComponent::RegisterAndInitializePhysicsProxy()
 		SimulationParameters.WorldTransform = GetComponentToWorld();
 		SimulationParameters.UserData = static_cast<void*>(&PhysicsUserData);
 		SimulationParameters.bEnableStrainOnCollision = bEnableDamageFromCollision;
+		SimulationParameters.bUseStaticMeshCollisionForTraces = bUseStaticMeshCollisionForTraces;
 
 		UPhysicalMaterial* EnginePhysicalMaterial = GetPhysicalMaterial();
 		if (ensure(EnginePhysicalMaterial))
