Compile-type error checking to ensure the argument passed into UE_NET_IMPLEMENT_NAMED_STRUCT_LASTRESORT_NETSERIALIZER_AND_REGISTRY_DELEGATES corresponds to a declared UE type.

diff --git a/Engine/Source/Runtime/Experimental/Iris/Core/Public/Iris/ReplicationState/PropertyNetSerializerInfoRegistry.h b/Engine/Source/Runtime/Experimental/Iris/Core/Public/Iris/ReplicationState/PropertyNetSerializerInfoRegistry.h
--- a/Engine/Source/Runtime/Experimental/Iris/Core/Public/Iris/ReplicationState/PropertyNetSerializerInfoRegistry.h
+++ b/Engine/Source/Runtime/Experimental/Iris/Core/Public/Iris/ReplicationState/PropertyNetSerializerInfoRegistry.h
@@ -24,6 +24,25 @@ typedef FReplicationFragment* (*CreateAndRegisterReplicationFragmentFunc)(UObjec
 namespace UE::Net
 {
 
+/**
+ * This function should be called in the following way to take advantage of SFINAE and perform compile-time
+ * error checking for declared types:
+ *     static_assert(IsDeclaredType((struct MyStruct*)Ptr), "Error message..");
+ *
+ * If MyStruct is not declared, the cast will pass a void* pointer to IsDeclaredType(). Otherwise, if MyStruct
+ * is declared then the cast will pass a MyStruct* pointer to IsDeclaredType().
+ */
+template<typename StructName>
+constexpr auto IsDeclaredType(StructName*) -> decltype(sizeof(StructName))
+{
+	return true;
+}
+
+constexpr auto IsDeclaredType(void*)
+{
+	return false;
+}
+
 /**
  * Currently we require each supported type to register FPropertyNetSerializerInfo 
  * It provides information on what NetSerializer to use for which property and how to build the required NetSerializer config which is used when we build the dynamic descriptor
@@ -161,6 +180,13 @@ void IRISCORE_API ValidateForwardingNetSerializerTraits(const FNetSerializer* Se
 
 }
 
+// Produce a compiler error if the name does not correspond to a declared UE type.
+#define UE_NET_IS_DECLARED_TYPE(Name) \
+	static_assert( \
+	UE::Net::IsDeclaredType((struct F##Name*)nullptr) || \
+	UE::Net::IsDeclaredType((struct U##Name*)nullptr) \
+	, "The UE type name '" #Name "' cannot be found. Make sure you have removed the 'F' or 'U' prefix from the type name.");
+
 // Only needed if we want to export PropertyNetSerializerInfo, this goes in the header if we need to export it
 #define UE_NET_DECLARE_NETSERIALIZER_INFO(NetSerializerInfo) \
 const UE::Net::FPropertyNetSerializerInfo& GetPropertyNetSerializerInfo_##NetSerializerInfo();
@@ -224,6 +250,7 @@ static F##Name##NetSerializerRegistryDelegates Name##NetSerializerRegistryDelega
 
 // Utility that can be used to forward serialization of a Struct to a last resort net serializer
 #define UE_NET_IMPLEMENT_NAMED_STRUCT_LASTRESORT_NETSERIALIZER_AND_REGISTRY_DELEGATES(Name) \
+	UE_NET_IS_DECLARED_TYPE(Name); \
 	static const FName PropertyNetSerializerRegistry_NAME_##Name( PREPROCESSOR_TO_STRING(Name) ); \
 	UE_NET_IMPLEMENT_NAMED_STRUCT_LASTRESORT_NETSERIALIZER_INFO(PropertyNetSerializerRegistry_NAME_##Name); \
 	UE_NET_IMPLEMENT_NETSERIALIZER_REGISTRY_DELEGATES(Name)
