No creating particle and geometry on Game Thread when creating a Geometry Collection.

Still creating geometry on the Physics Thread, otherwise a partial destruction issue occurs.
Setup query flag after particle and geometry creation

diff --git a/Engine/Source/Runtime/Experimental/Chaos/Private/PhysicsProxy/GeometryCollectionPhysicsProxy.cpp b/Engine/Source/Runtime/Experimental/Chaos/Private/PhysicsProxy/GeometryCollectionPhysicsProxy.cpp
--- a/Engine/Source/Runtime/Experimental/Chaos/Private/PhysicsProxy/GeometryCollectionPhysicsProxy.cpp
+++ b/Engine/Source/Runtime/Experimental/Chaos/Private/PhysicsProxy/GeometryCollectionPhysicsProxy.cpp
@@ -543,8 +543,11 @@ FGeometryCollectionPhysicsProxy::~FGeometryCollectionPhysicsProxy()
 float ReportHighParticleFraction = -1.f;
 FAutoConsoleVariableRef CVarReportHighParticleFraction(TEXT("p.gc.ReportHighParticleFraction"), ReportHighParticleFraction, TEXT("Report any objects with particle fraction above this threshold"));
 
-bool bBuildGeometryForChildren = true;
-FAutoConsoleVariableRef CVarbBuildGeometryForChildren(TEXT("p.gc.BuildGeometryForChildren"), bBuildGeometryForChildren, TEXT("If true build all children geometry at initilaization time, otherwise wait until destruction occurs."));
+bool bBuildGeometryForChildrenOnPT = true;
+FAutoConsoleVariableRef CVarbBuildGeometryForChildrenOnPT(TEXT("p.gc.BuildGeometryForChildrenOnPT"), bBuildGeometryForChildrenOnPT, TEXT("If true build all children geometry on Physics Thread at initilaization time, otherwise wait until destruction occurs."));
+
+bool bBuildGeometryForChildrenOnGT = true;
+FAutoConsoleVariableRef CVarbBuildGeometryForChildrenOnGT(TEXT("p.gc.BuildGeometryForChildrenOnGT"), bBuildGeometryForChildrenOnGT, TEXT("If true build all children geometry  on Game Thread at initilaization time, otherwise wait until destruction occurs."));
 
 bool bCreateGTParticleForChildren = true;
 FAutoConsoleVariableRef CVarCreateGTParticleForChildren(TEXT("p.gc.CreateGTParticlesForChildren"), bCreateGTParticleForChildren, TEXT("If true create all children particles at initilaization time, otherwise wait until destruction occurs."));
@@ -640,7 +643,7 @@ void FGeometryCollectionPhysicsProxy::Initialize(Chaos::FPBDRigidsEvolutionBase
 	if (ensure(NumTransforms == Implicits.Num() && NumEffectiveParticles == GTParticles.Num())) // Implicits are in the transform group so this invariant should always hold
 	{
 		constexpr bool bInitializationTime = true;
-		bHasBuiltGeometryOnGT = bBuildGeometryForChildren;
+		bHasBuiltGeometryOnGT = bBuildGeometryForChildrenOnGT;
 		CreateGTParticles(EffectiveParticles, Implicits, Evolution, bInitializationTime);
 
 		// Skip simplicials, as they're owned by unique pointers.
@@ -707,7 +710,7 @@ void FGeometryCollectionPhysicsProxy::Initialize(Chaos::FPBDRigidsEvolutionBase
 		GameThreadCollection.RemoveAttribute(FGeometryDynamicCollection::ImplicitsAttribute, FTransformCollection::TransformGroup);
 	}
 	// If GT particle have been created, we can use directly the UniqueIdx from the GTParticle, then free up some space.
-	if (bBuildGeometryForChildren || bCreateGTParticleForChildren)
+	if (bBuildGeometryForChildrenOnGT || bCreateGTParticleForChildren)
 	{
 		UniqueIdxs.Empty();
 	}
@@ -735,8 +738,8 @@ void FGeometryCollectionPhysicsProxy::CreateGTParticles(const TBitArray<>& Effec
 		}
 
 		if ((bInitializationTime && TransformIndex == Parameters.InitialRootIndex) || // When initializing always create particle for the root
-			((bInitializationTime && (bCreateGTParticleForChildren || bBuildGeometryForChildren)) || // When initializing create all particles if one of the flag is true 
-				(!bInitializationTime && TransformIndex != Parameters.InitialRootIndex && !bCreateGTParticleForChildren && !bBuildGeometryForChildren))) // When not initializing create other particles if flag was set to not create
+			((bInitializationTime && (bCreateGTParticleForChildren || bBuildGeometryForChildrenOnGT)) || // When initializing create all particles if one of the flag is true 
+				(!bInitializationTime && TransformIndex != Parameters.InitialRootIndex && !bCreateGTParticleForChildren && !bBuildGeometryForChildrenOnGT))) // When not initializing create other particles if flag was set to not create
 		{
 			GTParticles[ParticleIndex] = FParticle::CreateParticle();
 			P = GTParticles[ParticleIndex].Get();
@@ -769,7 +772,7 @@ void FGeometryCollectionPhysicsProxy::CreateGTParticles(const TBitArray<>& Effec
 			Chaos::FImplicitObjectPtr ImplicitGeometry = MakeImplicitObjectPtr<Chaos::FImplicitObjectUnion>(MoveTemp(Geoms));
 			P->SetGeometry(ImplicitGeometry);
 		}
-		else if (bInitializationTime == (TransformIndex == Parameters.InitialRootIndex) || bBuildGeometryForChildren)
+		else if (bInitializationTime == (TransformIndex == Parameters.InitialRootIndex) || bBuildGeometryForChildrenOnGT)
 		{
 			Chaos::FImplicitObjectPtr ImplicitGeometry = Implicits[TransformIndex];
 			if (ImplicitGeometry && !Scale.Equals(FVector::OneVector))
@@ -784,7 +787,7 @@ void FGeometryCollectionPhysicsProxy::CreateGTParticles(const TBitArray<>& Effec
 			Implicits[TransformIndex] = P->GetGeometry();
 		}
 
-		if (bInitializationTime == (TransformIndex == Parameters.InitialRootIndex) || bBuildGeometryForChildren)
+		if (bInitializationTime == (TransformIndex == Parameters.InitialRootIndex) || bBuildGeometryForChildrenOnGT)
 		{
 			if (DynamicCollectionAnchoringFacade.IsAnchored(TransformIndex))
 			{
@@ -942,6 +945,11 @@ void FGeometryCollectionPhysicsProxy::CreateChildrenGeometry_External()
 			// The Implicits attributes from the Dynamic Collection are just used for initialization, after they can be removed and so free some memory.
 			GameThreadCollection.RemoveAttribute(FGeometryDynamicCollection::ImplicitsAttribute, FTransformCollection::TransformGroup);
 			bHasBuiltGeometryOnGT = true;
+
+			if (PostParticlesCreatedCallback)
+			{
+				PostParticlesCreatedCallback();
+			}
 		}
 	}
 }
@@ -1208,7 +1216,7 @@ void FGeometryCollectionPhysicsProxy::InitializeBodiesPT(Chaos::FPBDRigidsSolver
 
 	bIsInitializedOnPhysicsThread = true;
 	// Building geometry according to the cVar
-	bHasBuiltGeometryOnPT = bBuildGeometryForChildren;
+	bHasBuiltGeometryOnPT = bBuildGeometryForChildrenOnPT;
 
 	if (Parameters.Simulating && RestCollection)
 	{
@@ -1281,7 +1289,7 @@ void FGeometryCollectionPhysicsProxy::InitializeBodiesPT(Chaos::FPBDRigidsSolver
 					Handle,
 					Parameters.Shared,
 					Simplicials[TransformGroupIndex].Get(),
-					(bBuildGeometryForChildren || (TransformGroupIndex == Parameters.InitialRootIndex)) ? Implicits[TransformGroupIndex] : nullptr,
+					(bBuildGeometryForChildrenOnPT || (TransformGroupIndex == Parameters.InitialRootIndex)) ? Implicits[TransformGroupIndex] : nullptr,
 					SimFilter,
 					QueryFilter,
 					ScaledMass,
@@ -1949,7 +1957,7 @@ Chaos::TPBDGeometryCollectionParticleHandle<Chaos::FReal, 3>* FGeometryCollectio
 		Handle,
 		Parameters.Shared,
 		Simplicials[CollectionClusterIndex].Get(),
-		(bBuildGeometryForChildren || (CollectionClusterIndex == Parameters.InitialRootIndex)) ? Implicits[CollectionClusterIndex] : nullptr,
+		(bBuildGeometryForChildrenOnPT || (CollectionClusterIndex == Parameters.InitialRootIndex)) ? Implicits[CollectionClusterIndex] : nullptr,
 		SimFilter,
 		QueryFilter,
 		Mass,
@@ -4138,13 +4146,13 @@ bool FGeometryCollectionPhysicsProxy::PullNonInterpolatableDataFromSinglePhysics
 
 		const int32 ParticleIndex = FromTransformToParticleIndex[TransformGroupIndex];
 		const bool bIsRootIndex = (Parameters.InitialRootIndex == TransformGroupIndex);
+		const bool bIsActive = !StateData.State.DisabledState;
 
-		if (!bIsRootIndex)
+		if (!bIsRootIndex && bIsActive)
 		{
 			CreateChildrenGeometry_External();
 		}
 
-		const bool bIsActive = !StateData.State.DisabledState;
 		if (UpdateValue(GameThreadCollection.Active[TransformGroupIndex], bIsActive))
 		{
 			if (GTParticles[ParticleIndex].IsValid())
diff --git a/Engine/Source/Runtime/Experimental/Chaos/Public/PhysicsProxy/GeometryCollectionPhysicsProxy.h b/Engine/Source/Runtime/Experimental/Chaos/Public/PhysicsProxy/GeometryCollectionPhysicsProxy.h
--- a/Engine/Source/Runtime/Experimental/Chaos/Public/PhysicsProxy/GeometryCollectionPhysicsProxy.h
+++ b/Engine/Source/Runtime/Experimental/Chaos/Public/PhysicsProxy/GeometryCollectionPhysicsProxy.h
@@ -523,6 +523,11 @@ public:
 	{
 		PostPhysicsSyncCallback = Callback;
 	}
+
+	void SetPostParticlesCreatedCallback(TFunction<void()> Callback)
+	{
+		PostParticlesCreatedCallback = Callback;
+	}
 	
 	CHAOS_API FClusterHandle* GetInitialRootHandle_Internal() const;
 
@@ -707,6 +712,7 @@ private:
 
 	// called after we sync the physics thread data ( called on the game thread )
 	TFunction<void()> PostPhysicsSyncCallback;
+	TFunction<void()> PostParticlesCreatedCallback;
 	
 	// Per object collision fraction.
 	float CollisionParticlesPerObjectFraction;
diff --git a/Engine/Source/Runtime/Experimental/GeometryCollectionEngine/Private/GeometryCollection/GeometryCollectionComponent.cpp b/Engine/Source/Runtime/Experimental/GeometryCollectionEngine/Private/GeometryCollection/GeometryCollectionComponent.cpp
--- a/Engine/Source/Runtime/Experimental/GeometryCollectionEngine/Private/GeometryCollection/GeometryCollectionComponent.cpp
+++ b/Engine/Source/Runtime/Experimental/GeometryCollectionEngine/Private/GeometryCollection/GeometryCollectionComponent.cpp
@@ -3860,6 +3860,8 @@ void UGeometryCollectionComponent::RegisterAndInitializePhysicsProxy()
 #endif
 	PhysicsProxy = new FGeometryCollectionPhysicsProxy(this, *DynamicCollection, SimulationParameters, InitialSimFilter, InitialQueryFilter, CollectorGuid);
 	PhysicsProxy->SetPostPhysicsSyncCallback([this]() { OnPostPhysicsSync(); });
+	PhysicsProxy->SetPostParticlesCreatedCallback([this]() { OnPostCreateParticles(); });
+
 
 	if (GetIsReplicated())
 	{
@@ -4154,6 +4156,11 @@ void UGeometryCollectionComponent::OnTransformsDirty()
 	ComponentSpaceBounds.Init();
 }
 
+void UGeometryCollectionComponent::OnPostCreateParticles()
+{
+	LoadCollisionProfiles();
+}
+
 void UGeometryCollectionComponent::OnPostPhysicsSync()
 {
 	SCOPE_CYCLE_COUNTER(STAT_GCPostPhysicsSync);
diff --git a/Engine/Source/Runtime/Experimental/GeometryCollectionEngine/Public/GeometryCollection/GeometryCollectionComponent.h b/Engine/Source/Runtime/Experimental/GeometryCollectionEngine/Public/GeometryCollection/GeometryCollectionComponent.h
--- a/Engine/Source/Runtime/Experimental/GeometryCollectionEngine/Public/GeometryCollection/GeometryCollectionComponent.h
+++ b/Engine/Source/Runtime/Experimental/GeometryCollectionEngine/Public/GeometryCollection/GeometryCollectionComponent.h
@@ -1763,6 +1763,8 @@ private:
 
 	GEOMETRYCOLLECTIONENGINE_API void OnPostPhysicsSync();
 
+	GEOMETRYCOLLECTIONENGINE_API void OnPostCreateParticles();
+
 	GEOMETRYCOLLECTIONENGINE_API bool HasVisibleGeometry() const;
 
 	/** backward compatibility method, until we can remove GlobalMatrices */
