GT particle and PT particle were not always set with the same Id.

diff --git a/Engine/Source/Runtime/Experimental/Chaos/Private/PhysicsProxy/GeometryCollectionPhysicsProxy.cpp b/Engine/Source/Runtime/Experimental/Chaos/Private/PhysicsProxy/GeometryCollectionPhysicsProxy.cpp
--- a/Engine/Source/Runtime/Experimental/Chaos/Private/PhysicsProxy/GeometryCollectionPhysicsProxy.cpp
+++ b/Engine/Source/Runtime/Experimental/Chaos/Private/PhysicsProxy/GeometryCollectionPhysicsProxy.cpp
@@ -1063,7 +1063,7 @@ void FGeometryCollectionPhysicsProxy::CreateNonClusteredParticles(Chaos::FPBDRig
 			Chaos::FUniqueIdx ExistingIndex;
 			if (GTParticles[ParticleIndex] == nullptr)
 			{
-				ExistingIndex = RigidsSolver->GetEvolution()->GenerateUniqueIdx();
+				ExistingIndex = UniqueIdxs[ParticleIndex];
 			}
 			else
 			{
