Iris
- Temporary fix for execution order issue where dependentobject flaggad to serialize before parent ending up being dispatched on client after parent if being transmitted as a huge object.

diff --git a/Engine/Source/Runtime/Experimental/Iris/Core/Private/Iris/ReplicationSystem/ReplicationWriter.cpp b/Engine/Source/Runtime/Experimental/Iris/Core/Private/Iris/ReplicationSystem/ReplicationWriter.cpp
--- a/Engine/Source/Runtime/Experimental/Iris/Core/Private/Iris/ReplicationSystem/ReplicationWriter.cpp
+++ b/Engine/Source/Runtime/Experimental/Iris/Core/Private/Iris/ReplicationSystem/ReplicationWriter.cpp
@@ -1753,8 +1753,9 @@ bool FReplicationWriter::CanSendObject(uint32 InternalIndex) const
 		{
 			const FInternalNetRefIndex DependentInternalIndex = DependentObjectInfo.NetRefIndex;
 
-			// If the dependent object already has been written in this packet we are fine.
-			if (WriteContext.ObjectsWrittenThisPacket.GetBit(DependentInternalIndex))
+			// If the dependent object already has been written in this packet and is not part of a huge object we do not need to do any further checks.
+			// Note: To avoid waiting for ack of huge dependent object we could remove the special scheduling of dependent actors and instead handle this when we write the batch
+			if (WriteContext.ObjectsWrittenThisPacket.GetBit(DependentInternalIndex) && !IsActiveHugeObject(DependentInternalIndex))
 			{
 				continue;
 			}
diff --git a/Engine/Source/Runtime/Experimental/Iris/Core/Private/Iris/ReplicationSystem/ReplicationWriter.h b/Engine/Source/Runtime/Experimental/Iris/Core/Private/Iris/ReplicationSystem/ReplicationWriter.h
--- a/Engine/Source/Runtime/Experimental/Iris/Core/Private/Iris/ReplicationSystem/ReplicationWriter.h
+++ b/Engine/Source/Runtime/Experimental/Iris/Core/Private/Iris/ReplicationSystem/ReplicationWriter.h
@@ -505,7 +505,7 @@ private:
 	// Returns true object and subobjects can be created on remote
 	bool CanSendObject(uint32 InternalIndex) const;
 
-	inline bool IsInitialState(const EReplicatedObjectState State) const { return State == EReplicatedObjectState::PendingCreate || (bHighPrioCreate && State == EReplicatedObjectState::WaitOnCreateConfirmation); }
+	inline bool IsInitialState(const EReplicatedObjectState State) const { return State == EReplicatedObjectState::PendingCreate || State == EReplicatedObjectState::WaitOnCreateConfirmation; }
 
 	bool IsActiveHugeObject(uint32 InternalIndex) const;
 	bool IsObjectPartOfActiveHugeObject(uint32 InternalIndex, const FReplicationInfo& Info) const;
