Iris
- Fixed unexpected rpc execution ordering for new actor and rpc to already existing actor.

diff --git a/Engine/Source/Runtime/Experimental/Iris/Core/Private/Iris/ReplicationSystem/ReplicationReader.cpp b/Engine/Source/Runtime/Experimental/Iris/Core/Private/Iris/ReplicationSystem/ReplicationReader.cpp
--- a/Engine/Source/Runtime/Experimental/Iris/Core/Private/Iris/ReplicationSystem/ReplicationReader.cpp
+++ b/Engine/Source/Runtime/Experimental/Iris/Core/Private/Iris/ReplicationSystem/ReplicationReader.cpp
@@ -1488,6 +1488,16 @@ void FReplicationReader::DispatchStateData(FNetSerializationContext& Context)
 	{
 		FReplicatedObjectInfo* ReplicationInfo = GetReplicatedObjectInfo(Info.InternalIndex);
 
+		const FNetRefHandleManager::FReplicatedObjectData& ObjectData = NetRefHandleManager->GetReplicatedObjectDataNoCheck(Info.InternalIndex);
+
+		// Before starting a potentially new batch we want to flush rpc:s and legacy callbacks belonging to the previous batch
+		const FInternalNetRefIndex RootInternalIndex = ObjectData.SubObjectRootIndex == FNetRefHandleManager::InvalidInternalIndex ? Info.InternalIndex : ObjectData.SubObjectRootIndex;
+		if (RootInternalIndex != LastDispatchedRootInternalIndex && NumObjectsPendingPostDistpatch)
+		{
+			FlushPostDispatchForBatch();
+			LastDispatchedRootInternalIndex = RootInternalIndex;
+		}
+
 		FPostDispatchObjectInfo PostDispatchObjectInfo;
 		PostDispatchObjectInfo.ReplicationInfo = ReplicationInfo;
 		PostDispatchObjectInfo.Info = &Info;
@@ -1497,7 +1507,6 @@ void FReplicationReader::DispatchStateData(FNetSerializationContext& Context)
 		// For SubObjects we call must call this method after applying state data for the owner, in order to remain backwards compatible.
 		if (Info.bShouldCallSubObjectCreatedFromReplication)
 		{
-			const FNetRefHandleManager::FReplicatedObjectData& ObjectData = NetRefHandleManager->GetReplicatedObjectDataNoCheck(Info.InternalIndex);
 			if (ObjectData.SubObjectRootIndex != FNetRefHandleManager::InvalidInternalIndex)
 			{
 				ReplicationBridge->CallSubObjectCreatedFromReplication(ObjectData.RefHandle);
@@ -1516,15 +1525,6 @@ void FReplicationReader::DispatchStateData(FNetSerializationContext& Context)
 		// If we have any object references we want to update any unresolved ones, including previously unresolved references
 		if (Info.bHasState)
 		{
-			const FNetRefHandleManager::FReplicatedObjectData& ObjectData = NetRefHandleManager->GetReplicatedObjectDataNoCheck(Info.InternalIndex);
-
-			// We only need to flush if we are switching to a new root object with state data
-			const FInternalNetRefIndex RootInternalIndex = ObjectData.SubObjectRootIndex == FNetRefHandleManager::InvalidInternalIndex ? Info.InternalIndex : ObjectData.SubObjectRootIndex;
-			if (RootInternalIndex != LastDispatchedRootInternalIndex)
-			{
-				FlushPostDispatchForBatch();
-				LastDispatchedRootInternalIndex = RootInternalIndex;
-			}
 
 			const uint32 ChangeMaskBitCount = ReplicationInfo->ChangeMaskBitCount;
 
