Iris
- Reduced log spam
- Changed logging of groups from using GetRawValue to using GetGroupId

diff --git a/Engine/Source/Runtime/Engine/Private/Net/Iris/ReplicationSystem/ActorReplicationBridge.cpp b/Engine/Source/Runtime/Engine/Private/Net/Iris/ReplicationSystem/ActorReplicationBridge.cpp
--- a/Engine/Source/Runtime/Engine/Private/Net/Iris/ReplicationSystem/ActorReplicationBridge.cpp
+++ b/Engine/Source/Runtime/Engine/Private/Net/Iris/ReplicationSystem/ActorReplicationBridge.cpp
@@ -303,14 +303,14 @@ UE::Net::FNetRefHandle UActorReplicationBridge::BeginReplication(AActor* Actor,
 			{
 				LevelGroup = CreateLevelGroup(Level);
 
-				UE_LOG_ACTORREPLICATIONBRIDGE(Log, TEXT("Created new GroupIndex: %u for Level: %s"), LevelGroup.GetRawValue(), ToCStr(PackageName.ToString()));
+				UE_LOG_ACTORREPLICATIONBRIDGE(Log, TEXT("Created new GroupIndex: %u for Level: %s"), LevelGroup.GetGroupIndex(), ToCStr(PackageName.ToString()));
 
 				// Update the filtering status of the group based on current level visibility for all connections
 				NetDriver->UpdateGroupFilterStatusForLevel(Level, LevelGroup);
 			}
 
 			// Add object to group
-			UE_LOG_ACTORREPLICATIONBRIDGE(Verbose, TEXT("Added %s to GroupIndex: %u Level: %s"), *ActorRefHandle.ToString(), LevelGroup.GetRawValue(), ToCStr(PackageName.ToString()));
+			UE_LOG_ACTORREPLICATIONBRIDGE(Verbose, TEXT("Added %s to GroupIndex: %u Level: %s"), *ActorRefHandle.ToString(), LevelGroup.GetGroupIndex(), ToCStr(PackageName.ToString()));
 			GetReplicationSystem()->AddToGroup(LevelGroup, ActorRefHandle);	
 		}
 	}
diff --git a/Engine/Source/Runtime/Experimental/Iris/Core/Private/Iris/ReplicationSystem/Filtering/ReplicationFiltering.cpp b/Engine/Source/Runtime/Experimental/Iris/Core/Private/Iris/ReplicationSystem/Filtering/ReplicationFiltering.cpp
--- a/Engine/Source/Runtime/Experimental/Iris/Core/Private/Iris/ReplicationSystem/Filtering/ReplicationFiltering.cpp
+++ b/Engine/Source/Runtime/Experimental/Iris/Core/Private/Iris/ReplicationSystem/Filtering/ReplicationFiltering.cpp
@@ -1764,7 +1764,7 @@ void FReplicationFiltering::SetSubObjectFilterStatus(FNetObjectGroupHandle Group
 
 	if (ensure(ValidConnections.GetBit(ConnectionId) && SubObjectFilterGroups.GetBit(GroupIndex)))
 	{
-		UE_LOG(LogIrisFiltering, Verbose, TEXT("ReplicationFiltering::SetSubObjectFilterStatus GroupIndex: %u, ConnectionId: %u, FilterStatus: %u"), GroupHandle.GetRawValue(), ConnectionId, ReplicationStatus == ENetFilterStatus::Allow ? 1U : 0U);
+		UE_LOG(LogIrisFiltering, Verbose, TEXT("ReplicationFiltering::SetSubObjectFilterStatus GroupIndex: %u, ConnectionId: %u, FilterStatus: %u"), GroupHandle.GetGroupIndex(), ConnectionId, ReplicationStatus == ENetFilterStatus::Allow ? 1U : 0U);
 		FPerObjectInfo* FilterInfo = GetPerObjectInfo(GroupInfos[GroupIndex].ConnectionStateIndex);
 		SetConnectionFilterStatus(*FilterInfo, ConnectionId, ReplicationStatus);
 		if (!IsAnyConnectionFilterStatusAllowed(*FilterInfo))
diff --git a/Engine/Source/Runtime/Experimental/Iris/Core/Private/Iris/ReplicationSystem/ReplicationBridge.cpp b/Engine/Source/Runtime/Experimental/Iris/Core/Private/Iris/ReplicationSystem/ReplicationBridge.cpp
--- a/Engine/Source/Runtime/Experimental/Iris/Core/Private/Iris/ReplicationSystem/ReplicationBridge.cpp
+++ b/Engine/Source/Runtime/Experimental/Iris/Core/Private/Iris/ReplicationSystem/ReplicationBridge.cpp
@@ -868,7 +868,7 @@ UE::Net::FNetRefHandle UReplicationBridge::InternalAddDestructionInfo(FNetRefHan
 		GetReplicationSystem()->SetPrioritizer(DestructionInfoHandle, DefaultSpatialNetObjectPrioritizerHandle);
 	}
 
-	UE_LOG_REPLICATIONBRIDGE(Verbose, TEXT("UReplicationBridge::InternalAddDestructionInfo %s ( InternalIndex: %u ) for %s GroupIndex: %u"), *DestructionInfoHandle.ToString(), InternalReplicationIndex,  *Handle.ToString(), LevelGroupHandle.GetRawValue());
+	UE_LOG_REPLICATIONBRIDGE(Verbose, TEXT("UReplicationBridge::InternalAddDestructionInfo %s ( InternalIndex: %u ) for %s GroupIndex: %u"), *DestructionInfoHandle.ToString(), InternalReplicationIndex,  *Handle.ToString(), LevelGroupHandle.GetGroupIndex());
 	
 	return DestructionInfoHandle;
 }
diff --git a/Engine/Source/Runtime/Experimental/Iris/Core/Private/Iris/ReplicationSystem/ReplicationWriter.cpp b/Engine/Source/Runtime/Experimental/Iris/Core/Private/Iris/ReplicationSystem/ReplicationWriter.cpp
--- a/Engine/Source/Runtime/Experimental/Iris/Core/Private/Iris/ReplicationSystem/ReplicationWriter.cpp
+++ b/Engine/Source/Runtime/Experimental/Iris/Core/Private/Iris/ReplicationSystem/ReplicationWriter.cpp
@@ -1658,7 +1658,7 @@ uint32 FReplicationWriter::WriteObjectsPendingDestroy(FNetSerializationContext&
 		// We do not support destroying an object that is currently being sent as a huge object.
 		if (IsObjectPartOfActiveHugeObject(InternalIndex, Info))
 		{
-			UE_LOG(LogIris, Log, TEXT("Skipping writing destroy for object ( InternalIndex: %u ) which is part of active huge object."), InternalIndex);
+			UE_LOG(LogIris, Verbose, TEXT("Skipping writing destroy for object ( InternalIndex: %u ) which is part of active huge object."), InternalIndex);
 			bWroteAllDestroyedObjects = false;
 			continue;
 		}
@@ -1766,14 +1766,14 @@ bool FReplicationWriter::CanSendObject(uint32 InternalIndex) const
 				// if we cannot send the initial dependent object we must wait until we can.
 				if (!CanSendObject(DependentInternalIndex))
 				{
-					UE_LOG(LogIris, Log, TEXT("ReplicationWriter: Cannot send internal index (%u) due to waiting on init dependency internal index (%d)"), InternalIndex, DependentInternalIndex);
+					UE_LOG(LogIris, Verbose, TEXT("ReplicationWriter: Cannot send internal index (%u) due to waiting on init dependency internal index (%d)"), InternalIndex, DependentInternalIndex);
 					return false;
 				}
 
 				// if the dependent object are scheduled before parent and did not fit in this packet, we cannot write the parent either and have to wait until creation is confirmed
 				if ((DependentObjectInfo.SchedulingHint == EDependentObjectSchedulingHint::ScheduleBeforeParent) && ObjectsWithDirtyChanges.GetBit(DependentInternalIndex))
 				{
-					UE_LOG(LogIris, Log, TEXT("ReplicationWriter: Cannot send internal index (%u) due to waiting on ScheduleBefore dependency internal index (%d)"), InternalIndex, DependentInternalIndex);
+					UE_LOG(LogIris, Verbose, TEXT("ReplicationWriter: Cannot send internal index (%u) due to waiting on ScheduleBefore dependency internal index (%d)"), InternalIndex, DependentInternalIndex);
 					return false;
 				}
 			}
