Optimize particle and geometry creation when destruction occurs
- Reserve memory for map instead allocated it while iterating
- Precompute effective particle not to compute it while destroying

diff --git a/Engine/Source/Runtime/Experimental/Chaos/Private/PhysicsProxy/GeometryCollectionPhysicsProxy.cpp b/Engine/Source/Runtime/Experimental/Chaos/Private/PhysicsProxy/GeometryCollectionPhysicsProxy.cpp
--- a/Engine/Source/Runtime/Experimental/Chaos/Private/PhysicsProxy/GeometryCollectionPhysicsProxy.cpp
+++ b/Engine/Source/Runtime/Experimental/Chaos/Private/PhysicsProxy/GeometryCollectionPhysicsProxy.cpp
@@ -575,7 +575,6 @@ void FGeometryCollectionPhysicsProxy::Initialize(Chaos::FPBDRigidsEvolutionBase
 	NumTransforms = DynamicCollection.NumElements(FGeometryCollection::TransformGroup);
 	BaseParticleIndex = 0; // Are we always zero indexed now?
 	
-	TBitArray<> EffectiveParticles;
 	NumEffectiveParticles = CalculateEffectiveParticles(DynamicCollection, NumTransforms, Parameters.MaxSimulatedLevel, Parameters.EnableClustering, GetOwner(), EffectiveParticles);
 
 	SolverClusterID.Init(nullptr, NumEffectiveParticles);
@@ -644,7 +643,7 @@ void FGeometryCollectionPhysicsProxy::Initialize(Chaos::FPBDRigidsEvolutionBase
 	{
 		constexpr bool bInitializationTime = true;
 		bHasBuiltGeometryOnGT = bBuildGeometryForChildrenOnGT;
-		CreateGTParticles(EffectiveParticles, Implicits, Evolution, bInitializationTime);
+		CreateGTParticles(Implicits, Evolution, bInitializationTime);
 
 		// Skip simplicials, as they're owned by unique pointers.
 		static const FAttributeAndGroupId SkipList[] =
@@ -716,7 +715,7 @@ void FGeometryCollectionPhysicsProxy::Initialize(Chaos::FPBDRigidsEvolutionBase
 	}
 }
 
-void FGeometryCollectionPhysicsProxy::CreateGTParticles(const TBitArray<>& EffectiveParticles, TManagedArray<Chaos::FImplicitObjectPtr>& Implicits, Chaos::FPBDRigidsEvolutionBase* Evolution, bool bInitializationTime)
+void FGeometryCollectionPhysicsProxy::CreateGTParticles(TManagedArray<Chaos::FImplicitObjectPtr>& Implicits, Chaos::FPBDRigidsEvolutionBase* Evolution, bool bInitializationTime)
 {
 	const Chaos::Facades::FCollectionAnchoringFacade DynamicCollectionAnchoringFacade(GameThreadCollection);
 	const FVector Scale = Parameters.WorldTransform.GetScale3D();
@@ -725,6 +724,10 @@ void FGeometryCollectionPhysicsProxy::CreateGTParticles(const TBitArray<>& Effec
 	const TManagedArray<int32>& Level = Parameters.RestCollection->GetAttribute<int32>(LevelAttributeName, FTransformCollection::TransformGroup);
 
 	TArray<int32> ChildrenToCheckForParentFix;
+	if (!bInitializationTime && !bCreateGTParticleForChildren && !bBuildGeometryForChildrenOnGT)
+	{
+		GTParticlesToTransformGroupIndex.Reserve(NumEffectiveParticles);
+	}
 
 	for (int32 ParticleIndex = 0; ParticleIndex < NumEffectiveParticles; ++ParticleIndex)
 	{
@@ -936,9 +939,7 @@ void FGeometryCollectionPhysicsProxy::CreateChildrenGeometry_External()
 			}
 
 			TManagedArray<Chaos::FImplicitObjectPtr>& Implicits = GameThreadCollection.ModifyAttribute<Chaos::FImplicitObjectPtr>(FGeometryDynamicCollection::ImplicitsAttribute, FTransformCollection::TransformGroup);
-			TBitArray<> EffectiveParticles;
-			NumEffectiveParticles = CalculateEffectiveParticles(GameThreadCollection, NumTransforms, Parameters.MaxSimulatedLevel, Parameters.EnableClustering, GetOwner(), EffectiveParticles);
-			CreateGTParticles(EffectiveParticles, Implicits, RBDSolver->GetEvolution(), /*bInitializationTime*/false);
+			CreateGTParticles(Implicits, RBDSolver->GetEvolution(), /*bInitializationTime*/false);
 			SyncParticles_External();
 			UniqueIdxs.Empty();
 
@@ -1038,7 +1039,7 @@ void FGeometryCollectionPhysicsProxy::InitializeDynamicCollection(FGeometryDynam
 int32 ReportTooManyChildrenNum = -1;
 FAutoConsoleVariableRef CVarReportTooManyChildrenNum(TEXT("p.ReportTooManyChildrenNum"), ReportTooManyChildrenNum, TEXT("Issue warning if more than this many children exist in a single cluster"));
 
-void FGeometryCollectionPhysicsProxy::CreateNonClusteredParticles(Chaos::FPBDRigidsSolver* RigidsSolver, const FGeometryCollection& RestCollection, const FGeometryDynamicCollection& DynamicCollection, const TBitArray<>& EffectiveParticles)
+void FGeometryCollectionPhysicsProxy::CreateNonClusteredParticles(Chaos::FPBDRigidsSolver* RigidsSolver, const FGeometryCollection& RestCollection, const FGeometryDynamicCollection& DynamicCollection)
 {
 	LLM_SCOPE_BYNAME(TEXT("Physics/NonClusteredParticles"));
 	const TManagedArray<bool>& SimulatableParticles = DynamicCollection.SimulatableParticles;
@@ -1244,11 +1245,7 @@ void FGeometryCollectionPhysicsProxy::InitializeBodiesPT(Chaos::FPBDRigidsSolver
 		TArray<FTransform> Transform;
 		GeometryCollectionAlgo::Private::GlobalMatrices(DynamicCollection, Transform);
 
-		// Here Clean up Additional particles
-		TBitArray<> EffectiveParticles;
-		NumEffectiveParticles = CalculateEffectiveParticles(DynamicCollection, DynamicCollection.NumElements(FGeometryCollection::TransformGroup), Parameters.MaxSimulatedLevel, Parameters.EnableClustering, GetOwner(), EffectiveParticles);
-
-		CreateNonClusteredParticles(RigidsSolver, *RestCollection, DynamicCollection, EffectiveParticles);
+		CreateNonClusteredParticles(RigidsSolver, *RestCollection, DynamicCollection);
 
 		const float StrainDefault = Parameters.DamageThreshold.Num() ? Parameters.DamageThreshold[0] : 0;
 		// Add the rigid bodies
diff --git a/Engine/Source/Runtime/Experimental/Chaos/Public/PhysicsProxy/GeometryCollectionPhysicsProxy.h b/Engine/Source/Runtime/Experimental/Chaos/Public/PhysicsProxy/GeometryCollectionPhysicsProxy.h
--- a/Engine/Source/Runtime/Experimental/Chaos/Public/PhysicsProxy/GeometryCollectionPhysicsProxy.h
+++ b/Engine/Source/Runtime/Experimental/Chaos/Public/PhysicsProxy/GeometryCollectionPhysicsProxy.h
@@ -596,7 +596,7 @@ protected:
 
 	CHAOS_API void PrepareBufferData(Chaos::FDirtyGeometryCollectionData& BufferData, const FGeometryDynamicCollection& ThreadCollection,  Chaos::FReal SolverLastDt = 0.0);
 
-	CHAOS_API void CreateNonClusteredParticles(Chaos::FPBDRigidsSolver* RigidsSolver,	const FGeometryCollection& RestCollection, const FGeometryDynamicCollection& DynamicCollection, const TBitArray<>& EffectiveParticles);
+	CHAOS_API void CreateNonClusteredParticles(Chaos::FPBDRigidsSolver* RigidsSolver,	const FGeometryCollection& RestCollection, const FGeometryDynamicCollection& DynamicCollection);
 
 	CHAOS_API Chaos::FPBDRigidClusteredParticleHandle* FindClusteredParticleHandleByItemIndex_Internal(FGeometryCollectionItemIndex ItemIndex) const;
 
@@ -622,7 +622,7 @@ private:
 	static TBitArray<> CalculateClustersToCreateFromChildren(const FGeometryDynamicCollection& DynamicCollection, int32 NumTransforms);
 	static int32 CalculateEffectiveParticles(const FGeometryDynamicCollection& DynamicCollection, int32 NumTransform, int32 MaxSimulatedLevel, bool bEnableClustering, const UObject* Owner, TBitArray<>& EffectiveParticles);
 
-	void CreateGTParticles(const TBitArray<>& EffectiveParticles, TManagedArray<Chaos::FImplicitObjectPtr>& Implicits, Chaos::FPBDRigidsEvolutionBase* Evolution, bool bInitializeRootOnly);
+	void CreateGTParticles(TManagedArray<Chaos::FImplicitObjectPtr>& Implicits, Chaos::FPBDRigidsEvolutionBase* Evolution, bool bInitializeRootOnly);
 	void CreateChildrenGeometry_External();
 	void SyncParticles_External();
 	
@@ -656,13 +656,14 @@ private:
 	int32 NumEffectiveParticles;
 	int32 BaseParticleIndex;
 	TArray<FParticleHandle*> SolverClusterID;
-	TArray<FClusterHandle*> SolverClusterHandles; // make a TArray of the base class with type
+	TArray<FClusterHandle*> SolverClusterHandles; // make a TArray of the base clase with type
 	TArray<FClusterHandle*> SolverParticleHandles;// make a TArray of base class and join with above
 	TMap<FParticleHandle*, int32> HandleToTransformGroupIndex;
 	TMap<int32, FClusterHandle*> UniqueIdxToInternalClusterHandle;
 	TArray<Chaos::FUniqueIdx> UniqueIdxs;
 	TArray<int32> FromParticleToTransformIndex;
 	TArray<int32> FromTransformToParticleIndex;
+	TBitArray<> EffectiveParticles;
 
 	//
 	// Buffer Results State Information
