Iris
- Changed NetObjectGroups to allow for dynamic storage

diff --git a/Engine/Source/Runtime/Experimental/Iris/Core/Private/Iris/ReplicationSystem/Filtering/NetObjectGroups.cpp b/Engine/Source/Runtime/Experimental/Iris/Core/Private/Iris/ReplicationSystem/Filtering/NetObjectGroups.cpp
--- a/Engine/Source/Runtime/Experimental/Iris/Core/Private/Iris/ReplicationSystem/Filtering/NetObjectGroups.cpp
+++ b/Engine/Source/Runtime/Experimental/Iris/Core/Private/Iris/ReplicationSystem/Filtering/NetObjectGroups.cpp
@@ -2,6 +2,7 @@
 
 #include "NetObjectGroups.h"
 #include "Math/UnrealMathUtility.h"
+#include "HAL/IConsoleManager.h"
 #include "Iris/Core/IrisLog.h"
 #include "Iris/ReplicationSystem/Filtering/NetObjectFilter.h"
 #include "Iris/ReplicationSystem/NetRefHandleManager.h"
@@ -12,6 +13,9 @@ DEFINE_LOG_CATEGORY_STATIC(LogIrisGroup, Log, All)
 namespace UE::Net::Private
 {
 
+static int32 CVarEnsureIfNumGroupMembershipsExceedsNum = 128;
+static FAutoConsoleVariableRef CVarEnsureIfNumGroupMembershipsExceeds(TEXT("net.Iris.EnsureIfNumGroupMembershipsExceeds"), CVarEnsureIfNumGroupMembershipsExceedsNum, TEXT("If set to a positive number we will warn and ensure if an object is added to a high number of groups."), ECVF_Default );
+
 FNetObjectGroups::FNetObjectGroups()
 : CurrentEpoch(++NextEpoch)
 {
@@ -23,30 +27,16 @@ FNetObjectGroups::~FNetObjectGroups()
 
 bool FNetObjectGroups::IsMemberOf(const FNetObjectGroupMembership& Target, FNetObjectGroupHandle Group)
 {
-	for (uint32 It = 0; It < FNetObjectGroupMembership::MaxAssignedGroupCount; ++It)
-	{
-		if (Target.Groups[It] == Group)
-		{
-			return true;
-		}
-	}
-
-	return false;
+	return Target.Groups.Contains(Group);
 }
 
 bool FNetObjectGroups::AddGroupMembership(FNetObjectGroupMembership& Target, FNetObjectGroupHandle Group)
 {
-	for (uint32 It = 0; It < FNetObjectGroupMembership::MaxAssignedGroupCount; ++It)
+	if (!Target.Groups.Contains(Group))
 	{
-		if (Target.Groups[It] == Group)
-		{
-			return false;
-		}
-		else if (!Target.Groups[It].IsValid())
-		{
-			Target.Groups[It] = Group;
-			return true;
-		}
+		Target.Groups.Add(Group);
+
+		return true;
 	}
 
 	return false;
@@ -54,22 +44,12 @@ bool FNetObjectGroups::AddGroupMembership(FNetObjectGroupMembership& Target, FNe
 
 void FNetObjectGroups::RemoveGroupMembership(FNetObjectGroupMembership& Target, FNetObjectGroupHandle Group)
 {
-	const uint32 LastIndex = FNetObjectGroupMembership::MaxAssignedGroupCount - 1U;
-	for (uint32 It = 0U; It < FNetObjectGroupMembership::MaxAssignedGroupCount; ++It)
-	{
-		if (Target.Groups[It] == Group)
-		{
-			FPlatformMemory::Memmove(&Target.Groups[It], &Target.Groups[It + 1], (LastIndex - It) * sizeof(FNetObjectGroupHandle));
-			Target.Groups[LastIndex] = FNetObjectGroupHandle();
-
-			return;
-		}
-	}
+	Target.Groups.RemoveSingleSwap(Group);
 }
 
 void FNetObjectGroups::ResetGroupMembership(FNetObjectGroupMembership& Target)
 {
-	FMemory::Memzero(Target);
+	Target.Groups.Empty();
 }
 
 void FNetObjectGroups::Init(const FNetObjectGroupInitParams& Params)
@@ -82,7 +62,7 @@ void FNetObjectGroups::Init(const FNetObjectGroupInitParams& Params)
 	// Reserve first as invalid group
 	Groups.Add(FNetObjectGroup());
 
-	GroupMemberships.SetNumZeroed(Params.MaxObjectCount);
+	GroupMemberships.SetNum(Params.MaxObjectCount);
 	GroupFilteredOutObjects.Init(Params.MaxObjectCount);
 }
 
@@ -232,7 +212,8 @@ void FNetObjectGroups::AddToGroup(FNetObjectGroupHandle GroupHandle, FInternalNe
 		UE_LOG(LogIrisGroup, Verbose, TEXT("FNetObjectGroups::AddToGroup Adding %s to Group %s (GroupIndex: %u)"),
 			*NetRefHandleManager->PrintObjectFromIndex(InternalIndex), *GetGroupName(GroupHandle).ToString(), GroupHandle.GetGroupIndex());
 
-		if (AddGroupMembership(GroupMemberships[InternalIndex], GroupHandle))
+		FNetObjectGroupMembership& Membership = GroupMemberships[InternalIndex];
+		if (AddGroupMembership(Membership, GroupHandle))
 		{
 			Group->Members.AddUnique(InternalIndex);
 
@@ -240,12 +221,13 @@ void FNetObjectGroups::AddToGroup(FNetObjectGroupHandle GroupHandle, FInternalNe
 			{
 				GroupFilteredOutObjects.SetBit(InternalIndex);
 			}
-		}
-		else
-		{
-			UE_LOG(LogIrisGroup, Error, TEXT("FNetObjectGroups::AddToGroup Failed to add %s to Group %s (GroupIndex: %u) A NetObject can only be a member of %u groups."),
-				*NetRefHandleManager->PrintObjectFromIndex(InternalIndex), *GetGroupName(GroupHandle).ToString(), GroupHandle.GetGroupIndex(), FNetObjectGroupMembership::MaxAssignedGroupCount);
-			ensure(false);
+
+			if (CVarEnsureIfNumGroupMembershipsExceedsNum > 0 && Membership.Groups.Num() > CVarEnsureIfNumGroupMembershipsExceedsNum)
+			{
+				UE_LOG(LogIrisGroup, Error, TEXT("FNetObjectGroups::AddGroupMembership Unexpected high num groupmemberships for group %s (GroupIndex: %u) NetObject %s is member of %d groups."),
+					*GetGroupName(GroupHandle).ToString(), GroupHandle.GetGroupIndex(), *NetRefHandleManager->PrintObjectFromIndex(InternalIndex), Membership.Groups.Num());
+				ensure(false);			
+			}
 		}
 	}
 }
@@ -370,33 +352,16 @@ bool FNetObjectGroups::IsInAnyFilterGroup(const FNetObjectGroupMembership& Group
 	return false;
 }
 
-uint32 FNetObjectGroups::GetNumGroupMemberships(FInternalNetRefIndex InternalIndex) const
-{
-	if (InternalIndex >= (uint32)GroupMemberships.Num())
-	{
-		return 0U;
-	}
-
-	uint32 Count = 0U;
-	for (uint32 It=0; It < FNetObjectGroupMembership::MaxAssignedGroupCount; ++It)
-	{
-		Count += GroupMemberships[InternalIndex].Groups[It].IsValid() ? 1U : 0U;
-	}
-	
-	return Count;
-}
-
 const FNetObjectGroupHandle* FNetObjectGroups::GetGroupMemberships(FInternalNetRefIndex InternalIndex, uint32& GroupCount) const
 {
-	GroupCount = GetNumGroupMemberships(InternalIndex);
-	if (GroupCount)
-	{
-		return GroupMemberships[InternalIndex].Groups;
-	}
-	else
+	if (InternalIndex >= (uint32)GroupMemberships.Num())
 	{
 		return nullptr;
 	}
+
+	const FNetObjectGroupMembership& MemberShip = GroupMemberships[InternalIndex];
+	GroupCount = MemberShip.Groups.Num();
+	return GroupCount ? MemberShip.Groups.GetData() : nullptr;
 }
 
 }
diff --git a/Engine/Source/Runtime/Experimental/Iris/Core/Private/Iris/ReplicationSystem/Filtering/NetObjectGroups.h b/Engine/Source/Runtime/Experimental/Iris/Core/Private/Iris/ReplicationSystem/Filtering/NetObjectGroups.h
--- a/Engine/Source/Runtime/Experimental/Iris/Core/Private/Iris/ReplicationSystem/Filtering/NetObjectGroups.h
+++ b/Engine/Source/Runtime/Experimental/Iris/Core/Private/Iris/ReplicationSystem/Filtering/NetObjectGroups.h
@@ -99,9 +99,7 @@ public:
 	/** Does the group have the inclusion filter trait */
 	bool IsInclusionFilterGroup(FNetObjectGroupHandle GroupHandle) const;
 
-	/** Returns how many groups the given handle is a member of */
-	uint32 GetNumGroupMemberships(FInternalNetRefIndex InternalIndex) const;
-
+	/** Get the array of all groups that the NetObject is a member of */
 	const FNetObjectGroupHandle* GetGroupMemberships(FInternalNetRefIndex InternalIndex, uint32& GroupCount) const;
 
 	/** Create and manage named groups, only groups created as a named group will be findable by name */
@@ -122,8 +120,8 @@ public:
 private:
 	struct FNetObjectGroupMembership
 	{
-		enum { MaxAssignedGroupCount = 8 };
-		FNetObjectGroupHandle Groups[MaxAssignedGroupCount];		
+		enum { NumInlinedGroupHandles = 2 };
+		TArray<FNetObjectGroupHandle, TInlineAllocator<NumInlinedGroupHandles>> Groups;
 	};
 
 	static bool AddGroupMembership(FNetObjectGroupMembership& Target, FNetObjectGroupHandle Group);
