Iris
* ReplicationWriter sentinels are no longer enabled by compil target.  They should be enabled for all targets of a project instead.
* Fix connection issues between debug and non-debug targets.

#rnx

diff --git a/Engine/Source/Runtime/Experimental/Iris/Core/Private/Iris/IrisConfigInternal.h b/Engine/Source/Runtime/Experimental/Iris/Core/Private/Iris/IrisConfigInternal.h
--- a/Engine/Source/Runtime/Experimental/Iris/Core/Private/Iris/IrisConfigInternal.h
+++ b/Engine/Source/Runtime/Experimental/Iris/Core/Private/Iris/IrisConfigInternal.h
@@ -31,12 +31,10 @@
 #endif 
 #endif
 
+// When enabled will add bit stream guards to detect overflow or serialization errors.
+// Increases bandwidth and replication time so should only be enabled for local or unit testing.
 #ifndef UE_NET_USE_READER_WRITER_SENTINEL
-#if (UE_BUILD_SHIPPING || UE_BUILD_TEST || UE_BUILD_DEVELOPMENT)
 #	define UE_NET_USE_READER_WRITER_SENTINEL 0
-#else
-#	define UE_NET_USE_READER_WRITER_SENTINEL 1
-#endif 
 #endif
 
 #ifndef UE_NET_VALIDATE_DC_BASELINES
