Iris
* Added more logs to ReplicationReader errors
* Improved NetHandle reported error by printing net connection's description.

#rnx

diff --git a/Engine/Source/Runtime/Engine/Private/Net/Iris/ReplicationSystem/ActorReplicationBridge.cpp b/Engine/Source/Runtime/Engine/Private/Net/Iris/ReplicationSystem/ActorReplicationBridge.cpp
--- a/Engine/Source/Runtime/Engine/Private/Net/Iris/ReplicationSystem/ActorReplicationBridge.cpp
+++ b/Engine/Source/Runtime/Engine/Private/Net/Iris/ReplicationSystem/ActorReplicationBridge.cpp
@@ -1229,6 +1229,26 @@ void UActorReplicationBridge::ReportErrorWithNetRefHandle(uint32 ErrorType, FNet
 		}
 	}
 }
+
+FString UActorReplicationBridge::PrintConnectionInfo(uint32 ConnectionId)
+{
+	if (NetDriver)
+	{
+		if (UNetConnection* ClientConnection = NetDriver->GetConnectionById(ConnectionId))
+		{
+			return FString::Printf(TEXT("ConnectionId:%u %s"), ConnectionId, *ClientConnection->Describe());
+		}
+		else
+		{
+			return FString::Printf(TEXT("ConnectionId:%u no NetConnection found"), ConnectionId);
+		}
+	}
+	else
+	{
+		return FString::Printf(TEXT("ConnectionId:%u no NetDriver attached"), ConnectionId);
+	}
+}
+
 #else //!UE_WITH_IRIS
 
 UActorReplicationBridge::UActorReplicationBridge() = default;
diff --git a/Engine/Source/Runtime/Engine/Public/Net/Iris/ReplicationSystem/ActorReplicationBridge.h b/Engine/Source/Runtime/Engine/Public/Net/Iris/ReplicationSystem/ActorReplicationBridge.h
--- a/Engine/Source/Runtime/Engine/Public/Net/Iris/ReplicationSystem/ActorReplicationBridge.h
+++ b/Engine/Source/Runtime/Engine/Public/Net/Iris/ReplicationSystem/ActorReplicationBridge.h
@@ -106,6 +106,8 @@ protected:
 
 	virtual float GetPollFrequencyOfRootObject(const UObject* ReplicatedObject) const override;
 
+	[[nodiscard]] virtual FString PrintConnectionInfo(uint32 ConnectionId) override;
+
 private:
 	void GetActorCreationHeader(const AActor* Actor, UE::Net::Private::FActorCreationHeader& Header) const;
 	void GetSubObjectCreationHeader(const UObject* Object, const UObject* RootObject, UE::Net::Private::FSubObjectCreationHeader& Header) const;
diff --git a/Engine/Source/Runtime/Experimental/Iris/Core/Private/Iris/ReplicationSystem/ObjectReplicationBridge.cpp b/Engine/Source/Runtime/Experimental/Iris/Core/Private/Iris/ReplicationSystem/ObjectReplicationBridge.cpp
--- a/Engine/Source/Runtime/Experimental/Iris/Core/Private/Iris/ReplicationSystem/ObjectReplicationBridge.cpp
+++ b/Engine/Source/Runtime/Experimental/Iris/Core/Private/Iris/ReplicationSystem/ObjectReplicationBridge.cpp
@@ -1446,6 +1446,11 @@ bool UObjectReplicationBridge::IsClassCritical(const UClass* Class)
 	return false;
 }
 
+FString UObjectReplicationBridge::PrintConnectionInfo(uint32 ConnectionId)
+{
+	return FString::Printf(TEXT("ConnectionId:%u"), ConnectionId);
+}
+
 void UObjectReplicationBridge::OptionallySetObjectRequiresFrequentWorldLocationUpdate(FNetRefHandle RefHandle, bool bDesiresFrequentWorldLocationUpdate)
 {
 	using namespace UE::Net;
@@ -1973,35 +1978,18 @@ void UObjectReplicationBridge::OnErrorWithNetRefHandleReported(uint32 ErrorType,
 	// Ensure at the end so the log contains all the relevant information
 	ON_SCOPE_EXIT
 	{
-		ensureMsgf(false, TEXT("NetRefHandle error: %u reported. Look at the log for important information on the object tied to the handle."), ErrorType);
+		ensureMsgf(false, TEXT("NetRefHandle error: %u reported for %s. Look at the log for important information on the object tied to the handle."), ErrorType, *RefHandle.ToString());
 	};
 
 	const FInternalNetRefIndex ObjectInternalIndex = NetRefHandleManager->GetInternalIndex(RefHandle);
 	if (ObjectInternalIndex == FNetRefHandleManager::InvalidInternalIndex)
 	{
-		UE_LOG(LogIris, Warning, TEXT("OnErrorWithNetRefHandleReported: %u from Connection:%u for %s but object has no InternalIndex."), ErrorType, ConnectionId, *RefHandle.ToString());
+		UE_LOG(LogIris, Warning, TEXT("OnErrorWithNetRefHandleReported: %u from client %s for %s but object has no InternalIndex."), ErrorType, *PrintConnectionInfo(ConnectionId), *RefHandle.ToString());
 		return;
 	}
 
-	UObject* ObjInstance = NetRefHandleManager->GetReplicatedObjectInstance(ObjectInternalIndex);
-	const FNetRefHandleManager::FReplicatedObjectData& ObjData = NetRefHandleManager->GetReplicatedObjectData(ObjectInternalIndex);
-
-	if (ObjData.IsSubObject())
-	{
-		UObject* RootObjInstance = NetRefHandleManager->GetReplicatedObjectInstance(ObjData.SubObjectRootIndex);
-		const FNetRefHandle RootObjNetHandle = NetRefHandleManager->GetNetRefHandleFromInternalIndex(ObjData.SubObjectRootIndex);
-		
-		UE_LOG(LogIris, Error, TEXT("OnErrorWithNetRefHandleReported: %u from client:%u. %s maps to SubObject: %s owned by RootObject: %s using %s"), 
-			ErrorType, ConnectionId, 
-			*RefHandle.ToString(), *GetPathNameSafe(ObjInstance),
-			*GetNameSafe(RootObjInstance), *RootObjNetHandle.ToString()
-		);
-	}
-	else
-	{
-		UE_LOG(LogIris, Error, TEXT("OnErrorWithNetRefHandleReported: %u from client:%u. %s maps to RootObject: %s"), 
-			ErrorType, ConnectionId, 
-			*RefHandle.ToString(), *GetPathNameSafe(ObjInstance)
-		);
-	}
+	UE_LOG(LogIris, Error, TEXT("OnErrorWithNetRefHandleReported: %u from client %s. Problematic object was %s"), 
+		ErrorType, *PrintConnectionInfo(ConnectionId),
+		*NetRefHandleManager->PrintObjectFromIndex(ObjectInternalIndex)
+	);
 }
diff --git a/Engine/Source/Runtime/Experimental/Iris/Core/Private/Iris/ReplicationSystem/ReplicationReader.cpp b/Engine/Source/Runtime/Experimental/Iris/Core/Private/Iris/ReplicationSystem/ReplicationReader.cpp
--- a/Engine/Source/Runtime/Experimental/Iris/Core/Private/Iris/ReplicationSystem/ReplicationReader.cpp
+++ b/Engine/Source/Runtime/Experimental/Iris/Core/Private/Iris/ReplicationSystem/ReplicationReader.cpp
@@ -34,9 +34,7 @@
 #	define UE_LOG_REPLICATIONWRITER_CONN(Format, ...)
 #endif
 
-#define UE_LOG_REPLICATIONREADER_WARNING(Format, ...)  UE_LOG(LogIris, Warning, Format, ##__VA_ARGS__)
 #define UE_LOG_REPLICATIONREADER_CONN_WARNING(Format, ...) UE_LOG(LogIris, Warning, TEXT("Conn: %u ") Format, Parameters.ConnectionId, ##__VA_ARGS__)
-#define UE_LOG_REPLICATIONREADER_ERROR(Format, ...)  UE_LOG(LogIris, Error, Format, ##__VA_ARGS__)
 
 namespace UE::Net::Private
 {
@@ -239,7 +237,7 @@ void FReplicationReader::Deinit()
 {
 	for (FPendingBatchData& PendingBatchData : PendingBatches.PendingBatches)
 	{
-		UE_LOG_REPLICATIONREADER_WARNING(TEXT("FReplicationReader::Deinit NetHandle %s has %d unprocessed data batches"), *PendingBatchData.Handle.ToString(), PendingBatchData.QueuedDataChunks.Num());
+		UE_LOG(LogIris, Warning, TEXT("FReplicationReader::Deinit NetHandle %s has %d unprocessed data batches"), *PendingBatchData.Handle.ToString(), PendingBatchData.QueuedDataChunks.Num());
 
 		// Make sure to release all references that we are holding on to
 		if (ObjectReferenceCache)
@@ -443,7 +441,7 @@ void FReplicationReader::DeserializeObjectStateDelta(FNetSerializationContext& C
 
 		if (Reader.IsOverflown())
 		{
-			UE_LOG_REPLICATIONREADER_ERROR(TEXT("FReplicationReader::DeserializeObjectStateDelta Bitstream corrupted."));
+			UE_LOG(LogIris, Error, TEXT("FReplicationReader::DeserializeObjectStateDelta Bitstream corrupted."));
 			return;
 		}
 
@@ -470,7 +468,7 @@ void FReplicationReader::DeserializeObjectStateDelta(FNetSerializationContext& C
 		const uint32 NewBaselineIndex = Reader.ReadBits(FDeltaCompressionBaselineManager::BaselineIndexBitCount);
 		if (Reader.IsOverflown())
 		{
-			UE_LOG_REPLICATIONREADER_ERROR(TEXT("FReplicationReader::DeserializeObjectStateDelta Bitstream corrupted."));
+			UE_LOG(LogIris, Error, TEXT("FReplicationReader::DeserializeObjectStateDelta Bitstream corrupted."));
 			return;
 		}
 		OutNewBaselineIndex = NewBaselineIndex;
@@ -732,7 +730,7 @@ uint32 FReplicationReader::ReadObjectBatch(FNetSerializationContext& Context)
 				// $TODO: Report this to the server so it knows that the state of data in the batch is unknown
 
 				// Log error and try to recover, if get more incoming data for an object in the broken state we will skip it.
-				UE_LOG_REPLICATIONREADER_ERROR(TEXT("FReplicationReader::ReadObject Failed to read object batch handle: %s skipping batch data"), ToCStr(IncompleteHandle.ToString()));
+				UE_LOG(LogIris, Error, TEXT("FReplicationReader::ReadObject Failed to read object batch handle: %s skipping batch data"), ToCStr(IncompleteHandle.ToString()));
 					
 				BrokenObjects.AddUnique(IncompleteHandle);
 
@@ -802,7 +800,7 @@ void FReplicationReader::ReadObjectInBatch(FNetSerializationContext& Context, FN
 			FInternalNetRefIndex RootObjectInternalIndex = NetRefHandleManager->GetInternalIndex(IncompleteOwnerHandle);
 			if (Reader.IsOverflown() || RootObjectInternalIndex == FNetRefHandleManager::InvalidInternalIndex)
 			{
-				UE_LOG_REPLICATIONREADER_ERROR(TEXT("FReplicationReader::ReadObject Invalid subobjectowner handle. %s"), ToCStr(IncompleteOwnerHandle.ToString()));
+				UE_LOG(LogIris, Error, TEXT("FReplicationReader::ReadObject Invalid subobjectowner handle. %s"), ToCStr(IncompleteOwnerHandle.ToString()));
 				const FName& NetError = (Reader.IsOverflown() ? GNetError_BitStreamOverflow : GNetError_InvalidNetHandle);
 				Context.SetError(NetError);
 				return;			
@@ -821,7 +819,7 @@ void FReplicationReader::ReadObjectInBatch(FNetSerializationContext& Context, FN
 		// We got a read error
 		if (Reader.IsOverflown() || !IncompleteHandle.IsValid())
 		{
-			UE_LOG_REPLICATIONREADER_ERROR(TEXT("FReplicationReader::ReadObject Bitstream corrupted."));
+			UE_LOG(LogIris, Error, TEXT("FReplicationReader::ReadObject Bitstream corrupted."));
 			const FName& NetError = (Reader.IsOverflown() ? GNetError_BitStreamOverflow : GNetError_BitStreamError);
 			Context.SetError(NetError);
 			return;
@@ -834,7 +832,7 @@ void FReplicationReader::ReadObjectInBatch(FNetSerializationContext& Context, FN
 		FNetRefHandle NetRefHandle = CreateResult.NetRefHandle;
 		if (!NetRefHandle.IsValid())
 		{	
-			UE_LOG_REPLICATIONREADER_ERROR(TEXT("FReplicationReader::ReadObject Unable to create handle for %s."), *IncompleteHandle.ToString());
+			UE_LOG(LogIris, Error, TEXT("FReplicationReader::ReadObject Unable to create handle for %s."), *IncompleteHandle.ToString());
 
 			// Mark error, but do not mark the bitstream as overflown as we want to handle this error.
 			Context.SetError(GNetError_BrokenNetHandle, false);
@@ -854,7 +852,9 @@ void FReplicationReader::ReadObjectInBatch(FNetSerializationContext& Context, FN
 	}
 	else
 	{
-		bHasErrors = bHasErrors || Context.HasErrorOrOverflow();
+		bHasErrors = Context.HasErrorOrOverflow();
+		UE_CLOG(bHasErrors, LogIris, Error, TEXT("FReplicationReader::ReadObject ErrorOrOverFlow after reading object header"))
+
 		if (bHasErrors || !IncompleteHandle.IsValid())
 		{
 			InternalIndex = ObjectIndexForOOBAttachment;
@@ -869,6 +869,7 @@ void FReplicationReader::ReadObjectInBatch(FNetSerializationContext& Context, FN
 				if ((ReplicatedDestroyHeaderFlags & ReplicatedDestroyHeaderFlags_EndReplication) == 0U)
 				{
 					bHasErrors = true;
+					UE_LOG(LogIris, Error, TEXT("FReplicationReader::ReadObject Handle %s not bound to any InternalIndex"), *IncompleteHandle.ToString());
 				}
 				else
 				{
@@ -900,10 +901,10 @@ void FReplicationReader::ReadObjectInBatch(FNetSerializationContext& Context, FN
 
 		if (bHasState)
 		{
-			bHasErrors = IsObjectIndexForOOBAttachment(InternalIndex) || bIsReplicatedDestroyForInvalidObject;
-			if (bHasErrors)
+			if (IsObjectIndexForOOBAttachment(InternalIndex) || bIsReplicatedDestroyForInvalidObject)
 			{
-				UE_LOG_REPLICATIONREADER_WARNING(TEXT("FReplicationReader::ReadObject Bitstream corrupted. Getting state when not expecting state data."));
+				bHasErrors = true;
+				UE_LOG(LogIris, Error, TEXT("FReplicationReader::ReadObject Bitstream corrupted. Getting state when not expecting state data."));
 				Context.SetError(GNetError_BitStreamError);
 				goto ErrorHandling;
 			}
@@ -931,9 +932,10 @@ void FReplicationReader::ReadObjectInBatch(FNetSerializationContext& Context, FN
 					FReplicationProtocolOperations::DeserializeWithMask(Context, Info.ChangeMaskOrPointer.GetPointer(ChangeMaskBitCount), ObjectData.ReceiveStateBuffer, ObjectData.Protocol);
 				}
 			}
-		#if UE_NET_USE_READER_WRITER_SENTINEL
-				ReadAndVerifySentinelBits(&Reader, TEXT("HasStateEnd"), 8);
-		#endif
+
+#if UE_NET_USE_READER_WRITER_SENTINEL
+			ReadAndVerifySentinelBits(&Reader, TEXT("HasStateEnd"), 8);
+#endif
 
 			// Should we store a new baseline?
 			if (NewBaselineIndex != FDeltaCompressionBaselineManager::InvalidBaselineIndex)
@@ -974,13 +976,15 @@ void FReplicationReader::ReadObjectInBatch(FNetSerializationContext& Context, FN
 		{
 			if (IsObjectIndexForOOBAttachment(InternalIndex))
 			{
-				bHasErrors = bHasErrors || bIsReplicatedDestroyForInvalidObject;			
+				bHasErrors = bIsReplicatedDestroyForInvalidObject;
+				UE_CLOG(bHasErrors, LogIris, Error, TEXT("FReplicationReader::ReadObject Bitstream corrupted. Reading attachments when this was a destroy info message."));
 
 				if (!bHasErrors)
 				{
 					const bool bIsHugeObject = Reader.ReadBool();
 					AttachmentType = (bIsHugeObject ? ENetObjectAttachmentType::HugeObject : ENetObjectAttachmentType::OutOfBand);
-					bHasErrors = bHasErrors || (!Parameters.bAllowReceivingAttachmentsFromRemoteObjectsNotInScope && AttachmentType == ENetObjectAttachmentType::OutOfBand);
+					bHasErrors = (!Parameters.bAllowReceivingAttachmentsFromRemoteObjectsNotInScope && AttachmentType == ENetObjectAttachmentType::OutOfBand);
+					UE_CLOG(bHasErrors, LogIris, Error, TEXT("FReplicationReader::ReadObject Bitstream corrupted. Reading OutOfBand attachment for object not in scope."));
 				}
 
 				if (bHasErrors)
@@ -993,9 +997,10 @@ void FReplicationReader::ReadObjectInBatch(FNetSerializationContext& Context, FN
 			Attachments.Deserialize(Context, AttachmentType, InternalIndex, ObjectData.RefHandle);
 		}
 
-		bHasErrors = bHasErrors || Context.HasErrorOrOverflow();
-		if (bHasErrors)
+		if (Context.HasErrorOrOverflow())
 		{
+			bHasErrors = true;
+			UE_LOG(LogIris, Error, TEXT("FReplicationReader::ReadObject ErrorOrOverflow after reading bitstream."));
 			goto ErrorHandling;
 		}
 
@@ -1016,7 +1021,7 @@ ErrorHandling:
 	if (bHasErrors)
 	{
 		Context.SetErrorHandleContext(IncompleteHandle);
-		UE_LOG_REPLICATIONREADER_ERROR(TEXT("FReplicationReader::ReadObject Failed to read replicated object with %s. Error '%s'."), *IncompleteHandle.ToString(), (Context.HasError() ? ToCStr(Context.GetError().ToString()) : TEXT("BitStream Overflow")));
+		UE_LOG(LogIris, Error, TEXT("FReplicationReader::ReadObject Failed to read replicated object with %s. Error '%s'."), *IncompleteHandle.ToString(), (Context.HasError() ? ToCStr(Context.GetError().ToString()) : TEXT("BitStream Overflow")));
 	}
 }
 
@@ -1401,7 +1406,7 @@ void FReplicationReader::ResolveAndDispatchUnresolvedReferencesForObject(FNetSer
 		{
 			// $IRIS: $TODO: Figure out how to handle this, currently we do not crash but we probably want to
 			// handle this properly by accumulating changemask for later instantiation
-			UE_LOG_REPLICATIONREADER_WARNING(TEXT("Cannot dispatch state data for not instantiated %s"), *ObjectData.RefHandle.ToString());
+			UE_LOG(LogIris, Warning, TEXT("Cannot dispatch state data for not instantiated %s"), *ObjectData.RefHandle.ToString());
 		}
 	}
 
@@ -1618,7 +1623,7 @@ void FReplicationReader::DispatchStateData(FNetSerializationContext& Context)
 			{
 				// $IRIS: $TODO: Figure out how to handle this, currently we do not crash but we probably want to
 				// handle this properly by accumulating changemask for later instantiation
-				UE_LOG_REPLICATIONREADER_WARNING(TEXT("Cannot dispatch state data for not instantiated %s"), *(ObjectData.RefHandle.ToString()));
+				UE_LOG(LogIris, Warning, TEXT("Cannot dispatch state data for not instantiated %s"), *(ObjectData.RefHandle.ToString()));
 			}
 		}
 
@@ -1966,7 +1971,7 @@ void FReplicationReader::ProcessQueuedBatches()
 						// $TODO: Report this to the server so it knows that the state of data in the batch is unknown
 
 						// Log error and try to recover, if get more incoming data for an object in the broken state we will skip it.
-						UE_LOG_REPLICATIONREADER_ERROR(TEXT("FReplicationReader::ProcessQueuedBatches Failed to process object batch handle: %s skipping batch data"), ToCStr(PendingBatchData.Handle.ToString()));
+						UE_LOG(LogIris, Error, TEXT("FReplicationReader::ProcessQueuedBatches Failed to process object batch handle: %s skipping batch data"), ToCStr(PendingBatchData.Handle.ToString()));
 					
 						BrokenObjects.AddUnique(PendingBatchData.Handle);
 
diff --git a/Engine/Source/Runtime/Experimental/Iris/Core/Public/Iris/ReplicationSystem/ObjectReplicationBridge.h b/Engine/Source/Runtime/Experimental/Iris/Core/Public/Iris/ReplicationSystem/ObjectReplicationBridge.h
--- a/Engine/Source/Runtime/Experimental/Iris/Core/Public/Iris/ReplicationSystem/ObjectReplicationBridge.h
+++ b/Engine/Source/Runtime/Experimental/Iris/Core/Public/Iris/ReplicationSystem/ObjectReplicationBridge.h
@@ -301,6 +301,9 @@ protected:
 	/** Returns true if the class is considered critical and we force a disconnection if a protocol mismatch prevents instances of this class from replicating. */
 	IRISCORE_API bool IsClassCritical(const UClass* Class);
 
+	/** Returns the most relevant description of the client tied to this connection id. */
+	[[nodiscard]] IRISCORE_API virtual FString PrintConnectionInfo(uint32 ConnectionId);
+
 	/** Current max tick rate set by the engine */
 	float GetMaxTickRate() const { return MaxTickRate; }
 
