Geometry collection : expose EnableRootProxyForCustomRenderer function to blueprint

diff --git a/Engine/Source/Runtime/Experimental/GeometryCollectionEngine/Public/GeometryCollection/GeometryCollectionComponent.h b/Engine/Source/Runtime/Experimental/GeometryCollectionEngine/Public/GeometryCollection/GeometryCollectionComponent.h
--- a/Engine/Source/Runtime/Experimental/GeometryCollectionEngine/Public/GeometryCollection/GeometryCollectionComponent.h
+++ b/Engine/Source/Runtime/Experimental/GeometryCollectionEngine/Public/GeometryCollection/GeometryCollectionComponent.h
@@ -1325,7 +1325,9 @@ public:
 
 	/** Get any custom renderer. Returns nullptr if none is set. */
 	GEOMETRYCOLLECTIONENGINE_API IGeometryCollectionExternalRenderInterface* GetCustomRenderer() { return CustomRenderer.GetInterface(); }
+
 	/** Enable or disable root proxy for custom rendering - this can be set at runtime */
+	UFUNCTION(BlueprintCallable, Category = "Physics")
 	GEOMETRYCOLLECTIONENGINE_API void EnableRootProxyForCustomRenderer(bool bEnable);
 
 	/** Force all GC components to reregister their custom renderer objects. */
