Release unused slack memory in geometry collection rep data.

diff --git a/Engine/Source/Runtime/Experimental/GeometryCollectionEngine/Private/GeometryCollection/GeometryCollectionComponent.cpp b/Engine/Source/Runtime/Experimental/GeometryCollectionEngine/Private/GeometryCollection/GeometryCollectionComponent.cpp
--- a/Engine/Source/Runtime/Experimental/GeometryCollectionEngine/Private/GeometryCollection/GeometryCollectionComponent.cpp
+++ b/Engine/Source/Runtime/Experimental/GeometryCollectionEngine/Private/GeometryCollection/GeometryCollectionComponent.cpp
@@ -2401,6 +2401,9 @@ void UGeometryCollectionComponent::UpdateRepStateAndDynamicData()
 				}
 			}
 
+			//reclaim unused slack
+			LocalRepStateData.ReleasedData.Shrink();
+
 			// let's update the data for the tracked clusters 
 			LocalRepDynamicData.ClusterData.Reserve(RootsToTrack.Num());
 			for (const FRootHandle& Root: RootsToTrack)
