GeometryCollection Indeterminism:Add missing initialization of TBoundingVolumeHierarchy::MMaxLevels in the default constructor.
#rnx

diff --git a/Engine/Source/Runtime/Experimental/Chaos/Public/Chaos/BoundingVolumeHierarchy.h b/Engine/Source/Runtime/Experimental/Chaos/Public/Chaos/BoundingVolumeHierarchy.h
--- a/Engine/Source/Runtime/Experimental/Chaos/Public/Chaos/BoundingVolumeHierarchy.h
+++ b/Engine/Source/Runtime/Experimental/Chaos/Public/Chaos/BoundingVolumeHierarchy.h
@@ -67,6 +67,7 @@ class TBoundingVolumeHierarchy final : public ISpatialAcceleration<int32, T,d>
 
 	TBoundingVolumeHierarchy()
 		: MObjects(nullptr)
+		, MMaxLevels(DefaultMaxLevels)
 	{
 	}
 
