GeometryCollection : ManagedArrayCollection
- Explicit deletion of assignment operator and copy for the FValueType

diff --git a/Engine/Source/Runtime/Experimental/Chaos/Public/GeometryCollection/ManagedArrayCollection.h b/Engine/Source/Runtime/Experimental/Chaos/Public/GeometryCollection/ManagedArrayCollection.h
--- a/Engine/Source/Runtime/Experimental/Chaos/Public/GeometryCollection/ManagedArrayCollection.h
+++ b/Engine/Source/Runtime/Experimental/Chaos/Public/GeometryCollection/ManagedArrayCollection.h
@@ -673,24 +673,8 @@ private:
 			}
 		}
 
-
-		FValueType& operator=(const FValueType& Other) 
-		{
-			this->ArrayType = Other.ArrayType;
-			this->GroupIndexDependency = Other.GroupIndexDependency;
-			this->Saved = Other.Saved;
-			this->bExternalValue = false;
-			this->Value = nullptr;
-			if (Other.Value)
-			{
-				this->Value = NewManagedTypedArray(this->ArrayType);
-				this->Value->Resize(Other.Value->Num());
-				this->Value->Init(*Other.Value);
-			}
-			return *this;
-		}
-
-
+		FValueType& operator=(const FValueType& Other) = delete;
+		FValueType& operator=(FValueType&& Other) = delete;
 	};
 
 
