Chaos Cloth optimizations to XPBD constraints.

A few different optimizations to improve ISPC versions of the XPBD constraints:

* Remove branching where possible, mostly by eliminating checks against InvM != 0  (requires trimming kinematic constraints) and MinStiffness (no longer needed since reformulating to never divide by stiffness)
* Bending constraints--switch to storing Constraint indices to SOA instead of AOS.
* Bending constraints--store a copy of ParticlesX by constraint indexing. The copy is done while computing IsBuckled since we're touching ParticlesX there anyways.
* Switch weight maps to be stored per constraint rather than a lookup table.

Together, seeing simulation times drop from 30ms --> 20ms per frame on test examples with these constraints.

diff --git a/Engine/Source/Runtime/Experimental/Chaos/Private/Chaos/PBDBendingConstraints.cpp b/Engine/Source/Runtime/Experimental/Chaos/Private/Chaos/PBDBendingConstraints.cpp
--- a/Engine/Source/Runtime/Experimental/Chaos/Private/Chaos/PBDBendingConstraints.cpp
+++ b/Engine/Source/Runtime/Experimental/Chaos/Private/Chaos/PBDBendingConstraints.cpp
@@ -10,9 +10,6 @@
 DECLARE_CYCLE_STAT(TEXT("Chaos PBD Bending Constraint"), STAT_PBD_Bending, STATGROUP_Chaos);
 
 #if INTEL_ISPC && !UE_BUILD_SHIPPING
-bool bChaos_Bending_ISPC_Enabled = true;
-FAutoConsoleVariableRef CVarChaosBendingISPCEnabled(TEXT("p.Chaos.Bending.ISPC"), bChaos_Bending_ISPC_Enabled, TEXT("Whether to use ISPC optimizations in Bending constraints"));
-
 static_assert(sizeof(ispc::FVector4f) == sizeof(Chaos::Softs::FPAndInvM), "sizeof(ispc::FVector4f) != sizeof(Chaos::Softs::FPAndInvM");
 static_assert(sizeof(ispc::FIntVector4) == sizeof(Chaos::TVec4<int32>), "sizeof(ispc::FIntVector4) != sizeof(Chaos::TVec4<int32>");
 #endif
diff --git a/Engine/Source/Runtime/Experimental/Chaos/Private/Chaos/PBDBendingConstraintsBase.cpp b/Engine/Source/Runtime/Experimental/Chaos/Private/Chaos/PBDBendingConstraintsBase.cpp
--- a/Engine/Source/Runtime/Experimental/Chaos/Private/Chaos/PBDBendingConstraintsBase.cpp
+++ b/Engine/Source/Runtime/Experimental/Chaos/Private/Chaos/PBDBendingConstraintsBase.cpp
@@ -1,5 +1,16 @@
 // Copyright Epic Games, Inc. All Rights Reserved.
 #include "Chaos/PBDBendingConstraintsBase.h"
+#if INTEL_ISPC
+#include "PBDBendingConstraints.ispc.generated.h"
+#endif
+
+#if INTEL_ISPC && !UE_BUILD_SHIPPING
+bool bChaos_Bending_ISPC_Enabled = true;
+FAutoConsoleVariableRef CVarChaosBendingISPCEnabled(TEXT("p.Chaos.Bending.ISPC"), bChaos_Bending_ISPC_Enabled, TEXT("Whether to use ISPC optimizations in Bending constraints"));
+
+static_assert(sizeof(ispc::FVector4f) == sizeof(Chaos::Softs::FPAndInvM), "sizeof(ispc::FVector4f) != sizeof(Chaos::Softs::FPAndInvM");
+static_assert(sizeof(ispc::FIntVector4) == sizeof(Chaos::TVec4<int32>), "sizeof(ispc::FIntVector4) != sizeof(Chaos::TVec4<int32>");
+#endif
 
 namespace Chaos::Softs 
 {
@@ -134,4 +145,45 @@ namespace Chaos::Softs
 		const TConstArrayView<FRealSingle>& RestAngleMap,
 		const FSolverVec2& RestAngleValue,
 		ERestAngleConstructionType RestAngleConstructionType);
+
+
+	template<typename SolverParticlesOrRange>
+	void FPBDBendingConstraintsBase::Init(const SolverParticlesOrRange& InParticles)
+	{
+		TRACE_CPUPROFILER_EVENT_SCOPE(FPBDBendingConstraintsBase_Init);
+
+		IsBuckled.SetNumUninitialized(Constraints.Num());
+#if INTEL_ISPC
+		if (bRealTypeCompatibleWithISPC && bChaos_Bending_ISPC_Enabled)
+		{
+			ispc::InitBendingConstraintsIsBuckled(
+				(const ispc::FVector3f*)InParticles.XArray().GetData(),
+				(const ispc::FIntVector4*)Constraints.GetData(),
+				RestAngles.GetData(),
+				IsBuckled.GetData(),
+				BucklingRatio,
+				Constraints.Num()
+			);
+		}
+		else
+#endif
+		{
+			for (int32 ConstraintIndex = 0; ConstraintIndex < Constraints.Num(); ++ConstraintIndex)
+			{
+				const TVec4<int32>& Constraint = Constraints[ConstraintIndex];
+				const int32 Index1 = Constraint[0];
+				const int32 Index2 = Constraint[1];
+				const int32 Index3 = Constraint[2];
+				const int32 Index4 = Constraint[3];
+				const FSolverVec3& P1 = InParticles.X(Index1);
+				const FSolverVec3& P2 = InParticles.X(Index2);
+				const FSolverVec3& P3 = InParticles.X(Index3);
+				const FSolverVec3& P4 = InParticles.X(Index4);
+				const FSolverReal Angle = CalcAngle(P1, P2, P3, P4);
+				IsBuckled[ConstraintIndex] = AngleIsBuckled(Angle, RestAngles[ConstraintIndex]);
+			}
+		}
+	}
+	template CHAOS_API void FPBDBendingConstraintsBase::Init(const FSolverParticles& InParticles);
+	template CHAOS_API void FPBDBendingConstraintsBase::Init(const FSolverParticlesRange& InParticles);
 }  // End namespace Chaos::Softs
diff --git a/Engine/Source/Runtime/Experimental/Chaos/Private/Chaos/SoftsSpring.isph b/Engine/Source/Runtime/Experimental/Chaos/Private/Chaos/SoftsSpring.isph
--- a/Engine/Source/Runtime/Experimental/Chaos/Private/Chaos/SoftsSpring.isph
+++ b/Engine/Source/Runtime/Experimental/Chaos/Private/Chaos/SoftsSpring.isph
@@ -19,7 +19,7 @@ inline float SafeNormalize(FVector3f &Direction)
 // Spring without damping
 template <typename StiffnessFloatType>
 inline FVector3f GetXPBDSpringDelta(const FVector3f& P1, const FVector3f& P2, const varying float InvM1, const varying float InvM2,
-	 const uniform float Dt, const varying float Dist, const varying float Lambda, const StiffnessFloatType Stiffness, const uniform float MinStiffness, float& DLambda)
+	 const uniform float Dt, const varying float Dist, const varying float Lambda, const StiffnessFloatType Stiffness, float& DLambda)
 {
 	FVector3f Direction = P1 - P2;
 	const float Distance = SafeNormalize(Direction);
@@ -29,19 +29,16 @@ inline FVector3f GetXPBDSpringDelta(const FVector3f& P1, const FVector3f& P2, co
 	
 	const float CombinedInvMass = InvM2 + InvM1;
 
-	const float ComputedDLambda = (-AlphaInv * Offset - Lambda) / (AlphaInv * CombinedInvMass + (uniform float)1.f);
-	DLambda = select(Stiffness < MinStiffness, 0.f, ComputedDLambda);
+	DLambda = (-AlphaInv * Offset - Lambda) / (AlphaInv * CombinedInvMass + (uniform float)1.f);
 
-	const FVector3f ComputedDelta = DLambda * Direction;
-
-	return VectorSelect((InvM2 == FLOAT_ZERO && InvM1 == FLOAT_ZERO), FloatZeroVector, ComputedDelta);
+	return DLambda * Direction;
 }
 
 // This is a following the original XPBD paper using a single lambda for stretch and damping.
 template <typename StiffnessFloatType>
 inline FVector3f GetXPBDSpringDeltaWithDamping(const FVector3f& P1, const FVector3f& P2, const FVector3f& X1, const FVector3f& X2,
 	const varying float InvM1, const varying float InvM2,
-	const uniform float Dt, const varying float Dist, const varying float Lambda, const StiffnessFloatType Stiffness, const StiffnessFloatType DampingRatio, const uniform float MinStiffness, float& DLambda)
+	const uniform float Dt, const varying float Dist, const varying float Lambda, const StiffnessFloatType Stiffness, const StiffnessFloatType DampingRatio, float& DLambda)
 {
 	FVector3f Direction = P1 - P2;
 	const float Distance = SafeNormalize(Direction);
@@ -55,19 +52,16 @@ inline FVector3f GetXPBDSpringDeltaWithDamping(const FVector3f& P1, const FVecto
 	const varying float BetaDt = Damping * Dt;
 	const varying FVector3f RelativeVelocityTimesDt = P1 - X1 - P2 + X2;
 
-	const float ComputedDLambda = (-AlphaInv * Offset - Lambda - BetaDt * VectorDot(Direction, RelativeVelocityTimesDt)) / ((AlphaInv + BetaDt)*CombinedInvMass + (uniform float)1.f);
-	DLambda = select(Stiffness < MinStiffness, 0.f, ComputedDLambda);
-
-	const FVector3f ComputedDelta = DLambda * Direction;
+	DLambda = (-AlphaInv * Offset - Lambda - BetaDt * VectorDot(Direction, RelativeVelocityTimesDt)) / ((AlphaInv + BetaDt)*CombinedInvMass + (uniform float)1.f);
 
-	return VectorSelect(CombinedInvMass < FLOAT_SMALL_NUMBER, FloatZeroVector, ComputedDelta);
+	return DLambda * Direction;
 }
 
 // Spring damping constraint (separate from spring stretching)
 template <typename StiffnessFloatType>
 inline FVector3f GetXPBDSpringDampingDelta(const FVector3f& P1, const FVector3f& P2, const FVector3f& X1, const FVector3f& X2,
 	const varying float InvM1, const varying float InvM2,
-	const uniform float Dt, const varying float Dist, const varying float Lambda, const StiffnessFloatType Stiffness, const StiffnessFloatType DampingRatio, const uniform float MinStiffness, float& DLambda)
+	const uniform float Dt, const varying float Dist, const varying float Lambda, const StiffnessFloatType Stiffness, const StiffnessFloatType DampingRatio, float& DLambda)
 {
 	FVector3f Direction = P1 - P2;
 	const float Distance = SafeNormalize(Direction);
@@ -79,18 +73,15 @@ inline FVector3f GetXPBDSpringDampingDelta(const FVector3f& P1, const FVector3f&
 	const varying float BetaDt = Damping * Dt;
 	const varying FVector3f RelativeVelocityTimesDt = P1 - X1 - P2 + X2;
 
-	const float ComputedDLambda = ( -BetaDt * VectorDot(Direction, RelativeVelocityTimesDt) - Lambda) / (BetaDt * CombinedInvMass + (uniform float)1.f);
-	DLambda = select(Stiffness < MinStiffness, 0.f, ComputedDLambda);
-
-	const FVector3f ComputedDelta = DLambda * Direction;
+	DLambda = ( -BetaDt * VectorDot(Direction, RelativeVelocityTimesDt) - Lambda) / (BetaDt * CombinedInvMass + (uniform float)1.f);
 
-	return VectorSelect(CombinedInvMass < FLOAT_SMALL_NUMBER, FloatZeroVector, ComputedDelta);
+	return DLambda * Direction;
 }
 
 // Spring without damping
 template <typename StiffnessFloatType>
 inline FVector3f GetXPBDAxialSpringDelta(const FVector3f& P1, const FVector3f& P2, const FVector3f& P3, const varying float InvM1, const varying float InvM2, const varying float InvM3,
-	 const uniform float Dt, const varying float Bary, const varying float Dist, const varying float Lambda, const StiffnessFloatType Stiffness, const uniform float MinStiffness, float& DLambda)
+	 const uniform float Dt, const varying float Bary, const varying float Dist, const varying float Lambda, const StiffnessFloatType Stiffness, float& DLambda)
 {
 	FVector3f P = (P2 - P3) * Bary + P3;
 	FVector3f Direction = P1 - P;
@@ -101,12 +92,9 @@ inline FVector3f GetXPBDAxialSpringDelta(const FVector3f& P1, const FVector3f& P
 	
 	const float CombinedInvMass = (1.f - Bary) * InvM3 + Bary * InvM2 + InvM1;
 
-	const float ComputedDLambda = (-AlphaInv * Offset - Lambda) / (AlphaInv * CombinedInvMass + (uniform float)1.f);
-	DLambda = select(Stiffness < MinStiffness, 0.f, ComputedDLambda);
+	DLambda = (-AlphaInv * Offset - Lambda) / (AlphaInv * CombinedInvMass + (uniform float)1.f);
 
-	const FVector3f ComputedDelta = DLambda * Direction;
-
-	return VectorSelect((CombinedInvMass == FLOAT_ZERO), FloatZeroVector, ComputedDelta);
+	return DLambda * Direction;
 }
 
 // This is a following the original XPBD paper using a single lambda for stretch and damping.
@@ -115,7 +103,7 @@ inline FVector3f GetXPBDAxialSpringDeltaWithDamping(const FVector3f& P1, const F
 	const FVector3f& X1, const FVector3f& X2, const FVector3f& X3,
 	const varying float InvM1, const varying float InvM2, const varying float InvM3,
 	const uniform float Dt, const varying float Bary, const varying float Dist, const varying float Lambda, 
-	const StiffnessFloatType Stiffness, const StiffnessFloatType DampingRatio, const uniform float MinStiffness, float& DLambda)
+	const StiffnessFloatType Stiffness, const StiffnessFloatType DampingRatio, float& DLambda)
 {
 	FVector3f P = (P2 - P3) * Bary + P3;
 	FVector3f Direction = P1 - P;
@@ -130,12 +118,9 @@ inline FVector3f GetXPBDAxialSpringDeltaWithDamping(const FVector3f& P1, const F
 	const varying float BetaDt = Damping * Dt;
 	const varying FVector3f RelativeVelocityTimesDt = P1 - X1 - P2 + X2;
 
-	const float ComputedDLambda = (-AlphaInv * Offset - Lambda - BetaDt * VectorDot(Direction, RelativeVelocityTimesDt)) / ((AlphaInv + BetaDt)*CombinedInvMass + (uniform float)1.f);
-	DLambda = select(Stiffness < MinStiffness, 0.f, ComputedDLambda);
-
-	const FVector3f ComputedDelta = DLambda * Direction;
+	DLambda = (-AlphaInv * Offset - Lambda - BetaDt * VectorDot(Direction, RelativeVelocityTimesDt)) / ((AlphaInv + BetaDt)*CombinedInvMass + (uniform float)1.f);
 
-	return VectorSelect((CombinedInvMass < FLOAT_SMALL_NUMBER), FloatZeroVector, ComputedDelta);
+	return DLambda * Direction;
 }
 
 // Spring damping constraint (separate from spring stretching)
@@ -144,7 +129,7 @@ inline FVector3f GetXPBDAxialSpringDampingDelta(const FVector3f& P1, const FVect
 	const FVector3f& X1, const FVector3f& X2, const FVector3f& X3,
 	const varying float InvM1, const varying float InvM2, const varying float InvM3,
 	const uniform float Dt, const varying float Bary, const varying float Dist, const varying float Lambda, 
-	const StiffnessFloatType Stiffness, const StiffnessFloatType DampingRatio, const uniform float MinStiffness, float& DLambda)
+	const StiffnessFloatType Stiffness, const StiffnessFloatType DampingRatio, float& DLambda)
 {
 	FVector3f P = (P2 - P3) * Bary + P3;
 	FVector3f Direction = P1 - P;
@@ -157,11 +142,8 @@ inline FVector3f GetXPBDAxialSpringDampingDelta(const FVector3f& P1, const FVect
 	const varying float BetaDt = Damping * Dt;
 	const varying FVector3f RelativeVelocityTimesDt = P1 - X1 - P2 + X2;
 
-	const float ComputedDLambda = (- BetaDt * VectorDot(Direction, RelativeVelocityTimesDt) - Lambda ) / (BetaDt*CombinedInvMass + (uniform float)1.f);
-	DLambda = select(Stiffness < MinStiffness, 0.f, ComputedDLambda);
-
-	const FVector3f ComputedDelta = DLambda * Direction;
+	DLambda = (- BetaDt * VectorDot(Direction, RelativeVelocityTimesDt) - Lambda ) / (BetaDt*CombinedInvMass + (uniform float)1.f);
 
-	return VectorSelect((CombinedInvMass < FLOAT_SMALL_NUMBER), FloatZeroVector, ComputedDelta);
+	return DLambda * Direction;
 }
 #endif
diff --git a/Engine/Source/Runtime/Experimental/Chaos/Private/Chaos/XPBDAnisotropicBendingConstraints.cpp b/Engine/Source/Runtime/Experimental/Chaos/Private/Chaos/XPBDAnisotropicBendingConstraints.cpp
--- a/Engine/Source/Runtime/Experimental/Chaos/Private/Chaos/XPBDAnisotropicBendingConstraints.cpp
+++ b/Engine/Source/Runtime/Experimental/Chaos/Private/Chaos/XPBDAnisotropicBendingConstraints.cpp
@@ -7,6 +7,7 @@
 #include "ChaosStats.h"
 
 #if INTEL_ISPC
+#include "XPBDBendingConstraints.ispc.generated.h"
 #include "XPBDAnisotropicBendingConstraints.ispc.generated.h"
 #endif
 
@@ -14,7 +15,8 @@ DECLARE_CYCLE_STAT(TEXT("Chaos XPBD Anisotropic Bending Constraint"), STAT_XPBD_
 
 #if INTEL_ISPC && !UE_BUILD_SHIPPING
 static_assert(sizeof(ispc::FVector4f) == sizeof(Chaos::Softs::FPAndInvM), "sizeof(ispc::FVector4f) != sizeof(Chaos::Softs::FPAndInvM");
-static_assert(sizeof(ispc::FIntVector4) == sizeof(Chaos::TVec4<int32>), "sizeof(ispc::FIntVector4) != sizeof(Chaos::TVec4<int32>");
+static_assert(sizeof(ispc::FVector3f) == sizeof(Chaos::Softs::FSolverVec3), "sizeof(ispc::FVector3f) != sizeof(Chaos::Softs::FSolverVec3");
+static_assert(sizeof(ispc::FVector2f) == sizeof(Chaos::Softs::FSolverVec2), "sizeof(ispc::FVector2f) != sizeof(Chaos::Softs::FSolverVec2");
 #endif
 
 namespace Chaos::Softs {
@@ -22,61 +24,65 @@ namespace Chaos::Softs {
 // @todo(chaos): the parallel threshold (or decision to run parallel) should probably be owned by the solver and passed to the constraint container
 extern int32 Chaos_XPBDBending_ParallelConstraintCount;
 
+static int32 Chaos_XPBDBending_ISPC_ParallelBatchSize = 1028;
+static int32 Chaos_XPBDBending_ISPC_MinNumParallelBatches = 1028;  // effectively disabled for now
+FAutoConsoleVariableRef CVarChaosXPBDBendingISPCParallelBatchSize(TEXT("p.Chaos.XPBDBending.ISPC.ParallelBatchSize"), Chaos_XPBDBending_ISPC_ParallelBatchSize, TEXT("Parallel batch size for ISPC XPBDBending constraints"));
+FAutoConsoleVariableRef CVarChaosXPBDBendingISPCMinNumParallelBatches(TEXT("p.Chaos.XPBDBending.ISPC.MinNumParallelBatches"), Chaos_XPBDBending_ISPC_MinNumParallelBatches, TEXT("Min number of batches to invoke parallelFor ISPC XPBDBending constraints"));
+
 FXPBDAnisotropicBendingConstraints::FXPBDAnisotropicBendingConstraints(const FSolverParticlesRange& InParticles,
 	const FTriangleMesh& TriangleMesh,
 	const TArray<TVec3<FVec2f>>& FaceVertexPatternPositions,
 	const TMap<FString, TConstArrayView<FRealSingle>>& WeightMaps,
-	const FCollectionPropertyConstFacade& PropertyCollection,
-	bool bTrimKinematicConstraints)
+	const FCollectionPropertyConstFacade& PropertyCollection)
 	: Base(
 		InParticles,
 		TriangleMesh.GetUniqueAdjacentElements(),
-		WeightMaps.FindRef(GetXPBDAnisoBendingStiffnessWarpString(PropertyCollection, XPBDAnisoBendingStiffnessWarpName.ToString())),
-		WeightMaps.FindRef(GetXPBDAnisoBucklingStiffnessWarpString(PropertyCollection, XPBDAnisoBucklingStiffnessWarpName.ToString())),
+		TConstArrayView<FRealSingle>(), // We don't use base stiffness weight maps
+		TConstArrayView<FRealSingle>(), // We don't use base stiffness weight maps
 		GetRestAngleMapFromCollection(WeightMaps, PropertyCollection),
 		FSolverVec2(GetWeightedFloatXPBDAnisoBendingStiffnessWarp(PropertyCollection, MaxStiffness)),
 		(FSolverReal)GetXPBDAnisoBucklingRatio(PropertyCollection, 0.f),
 		FSolverVec2(GetWeightedFloatXPBDAnisoBucklingStiffnessWarp(PropertyCollection, MaxStiffness)),
 		GetRestAngleValueFromCollection(PropertyCollection),
 		(ERestAngleConstructionType)GetXPBDAnisoRestAngleType(PropertyCollection, (int32)ERestAngleConstructionType::Use3DRestAngles),
-		bTrimKinematicConstraints,
+		true /*bTrimKinematicConstraints*/,
 		MaxStiffness)
+	, StiffnessWarp(
+		FSolverVec2(GetWeightedFloatXPBDAnisoBendingStiffnessWarp(PropertyCollection, MaxStiffness)).ClampAxes(0, MaxStiffness),
+		WeightMaps.FindRef(GetXPBDAnisoBendingStiffnessWarpString(PropertyCollection, XPBDAnisoBendingStiffnessWarpName.ToString())),
+		TConstArrayView<TVec2<int32>>(ConstraintSharedEdges),
+		ParticleOffset,
+		ParticleCount)
 	, StiffnessWeft(
-		FSolverVec2(GetWeightedFloatXPBDAnisoBendingStiffnessWeft(PropertyCollection, MaxStiffness)),
+		FSolverVec2(GetWeightedFloatXPBDAnisoBendingStiffnessWeft(PropertyCollection, MaxStiffness)).ClampAxes(0, MaxStiffness),
 		WeightMaps.FindRef(GetXPBDAnisoBendingStiffnessWeftString(PropertyCollection, XPBDAnisoBendingStiffnessWeftName.ToString())),
 		TConstArrayView<TVec2<int32>>(ConstraintSharedEdges),
 		ParticleOffset,
-		ParticleCount,
-		FPBDStiffness::DefaultTableSize,
-		FPBDStiffness::DefaultParameterFitBase,
-		MaxStiffness)
+		ParticleCount)
 	, StiffnessBias(
-		FSolverVec2(GetWeightedFloatXPBDAnisoBendingStiffnessBias(PropertyCollection, MaxStiffness)),
+		FSolverVec2(GetWeightedFloatXPBDAnisoBendingStiffnessBias(PropertyCollection, MaxStiffness)).ClampAxes(0, MaxStiffness),
 		WeightMaps.FindRef(GetXPBDAnisoBendingStiffnessBiasString(PropertyCollection, XPBDAnisoBendingStiffnessBiasName.ToString())),
 		TConstArrayView<TVec2<int32>>(ConstraintSharedEdges),
 		ParticleOffset,
-		ParticleCount,
-		FPBDStiffness::DefaultTableSize,
-		FPBDStiffness::DefaultParameterFitBase,
-		MaxStiffness)
+		ParticleCount)
+	, BucklingStiffnessWarp(
+		FSolverVec2(GetWeightedFloatXPBDAnisoBucklingStiffnessWarp(PropertyCollection, MaxStiffness)).ClampAxes(0, MaxStiffness),
+		WeightMaps.FindRef(GetXPBDAnisoBucklingStiffnessWarpString(PropertyCollection, XPBDAnisoBucklingStiffnessWarpName.ToString())),
+		TConstArrayView<TVec2<int32>>(ConstraintSharedEdges),
+		ParticleOffset,
+		ParticleCount)
 	, BucklingStiffnessWeft(
-		FSolverVec2(GetWeightedFloatXPBDAnisoBucklingStiffnessWeft(PropertyCollection, MaxStiffness)),
+		FSolverVec2(GetWeightedFloatXPBDAnisoBucklingStiffnessWeft(PropertyCollection, MaxStiffness)).ClampAxes(0, MaxStiffness),
 		WeightMaps.FindRef(GetXPBDAnisoBucklingStiffnessWeftString(PropertyCollection, XPBDAnisoBucklingStiffnessWeftName.ToString())),
 		TConstArrayView<TVec2<int32>>(ConstraintSharedEdges),
 		ParticleOffset,
-		ParticleCount,
-		FPBDStiffness::DefaultTableSize,
-		FPBDStiffness::DefaultParameterFitBase,
-		MaxStiffness)
+		ParticleCount)
 	, BucklingStiffnessBias(
-		FSolverVec2(GetWeightedFloatXPBDAnisoBucklingStiffnessBias(PropertyCollection, MaxStiffness)),
+		FSolverVec2(GetWeightedFloatXPBDAnisoBucklingStiffnessBias(PropertyCollection, MaxStiffness)).ClampAxes(0, MaxStiffness),
 		WeightMaps.FindRef(GetXPBDAnisoBucklingStiffnessBiasString(PropertyCollection, XPBDAnisoBucklingStiffnessBiasName.ToString())),
 		TConstArrayView<TVec2<int32>>(ConstraintSharedEdges),
 		ParticleOffset,
-		ParticleCount,
-		FPBDStiffness::DefaultTableSize,
-		FPBDStiffness::DefaultParameterFitBase,
-		MaxStiffness)
+		ParticleCount)
 	, DampingRatio(
 		FSolverVec2(GetWeightedFloatXPBDAnisoBendingDamping(PropertyCollection, MinDamping)).ClampAxes(MinDamping, MaxDamping),
 		WeightMaps.FindRef(GetXPBDAnisoBendingDampingString(PropertyCollection, XPBDAnisoBendingDampingName.ToString())),
@@ -106,59 +112,58 @@ FXPBDAnisotropicBendingConstraints::FXPBDAnisotropicBendingConstraints(const FSo
 	const FTriangleMesh& TriangleMesh,
 	const TArray<TVec3<FVec2f>>& FaceVertexPatternPositions,
 	const TMap<FString, TConstArrayView<FRealSingle>>& WeightMaps,
-	const FCollectionPropertyConstFacade& PropertyCollection,
-	bool bTrimKinematicConstraints)
+	const FCollectionPropertyConstFacade& PropertyCollection)
 	: Base(
 		InParticles,
 		InParticleOffset,
 		InParticleCount,
 		TriangleMesh.GetUniqueAdjacentElements(),
-		WeightMaps.FindRef(GetXPBDAnisoBendingStiffnessWarpString(PropertyCollection, XPBDAnisoBendingStiffnessWarpName.ToString())),
-		WeightMaps.FindRef(GetXPBDAnisoBucklingStiffnessWarpString(PropertyCollection, XPBDAnisoBucklingStiffnessWarpName.ToString())),
+		TConstArrayView<FRealSingle>(), // We don't use base stiffness weight maps
+		TConstArrayView<FRealSingle>(), // We don't use base stiffness weight maps
 		GetRestAngleMapFromCollection(WeightMaps, PropertyCollection),
 		FSolverVec2(GetWeightedFloatXPBDAnisoBendingStiffnessWarp(PropertyCollection, MaxStiffness)),
 		(FSolverReal)GetXPBDAnisoBucklingRatio(PropertyCollection, 0.f),
 		FSolverVec2(GetWeightedFloatXPBDAnisoBucklingStiffnessWarp(PropertyCollection, MaxStiffness)),
 		GetRestAngleValueFromCollection(PropertyCollection),
 		(ERestAngleConstructionType)GetXPBDAnisoRestAngleType(PropertyCollection, (int32)ERestAngleConstructionType::Use3DRestAngles),
-		bTrimKinematicConstraints,
+		true /*bTrimKinematicConstraints*/,
 		MaxStiffness)
+	, StiffnessWarp(
+		FSolverVec2(GetWeightedFloatXPBDAnisoBendingStiffnessWarp(PropertyCollection, MaxStiffness)).ClampAxes(0, MaxStiffness),
+		WeightMaps.FindRef(GetXPBDAnisoBendingStiffnessWarpString(PropertyCollection, XPBDAnisoBendingStiffnessWarpName.ToString())),
+		TConstArrayView<TVec2<int32>>(ConstraintSharedEdges),
+		ParticleOffset,
+		ParticleCount)
 	, StiffnessWeft(
-		FSolverVec2(GetWeightedFloatXPBDAnisoBendingStiffnessWeft(PropertyCollection, MaxStiffness)),
+		FSolverVec2(GetWeightedFloatXPBDAnisoBendingStiffnessWeft(PropertyCollection, MaxStiffness)).ClampAxes(0, MaxStiffness),
 		WeightMaps.FindRef(GetXPBDAnisoBendingStiffnessWeftString(PropertyCollection, XPBDAnisoBendingStiffnessWeftName.ToString())),
 		TConstArrayView<TVec2<int32>>(ConstraintSharedEdges),
 		ParticleOffset,
-		ParticleCount,
-		FPBDStiffness::DefaultTableSize,
-		FPBDStiffness::DefaultParameterFitBase,
-		MaxStiffness)
+		ParticleCount)
 	, StiffnessBias(
-		FSolverVec2(GetWeightedFloatXPBDAnisoBendingStiffnessBias(PropertyCollection, MaxStiffness)),
+		FSolverVec2(GetWeightedFloatXPBDAnisoBendingStiffnessBias(PropertyCollection, MaxStiffness)).ClampAxes(0, MaxStiffness),
 		WeightMaps.FindRef(GetXPBDAnisoBendingStiffnessBiasString(PropertyCollection, XPBDAnisoBendingStiffnessBiasName.ToString())),
 		TConstArrayView<TVec2<int32>>(ConstraintSharedEdges),
 		ParticleOffset,
-		ParticleCount,
-		FPBDStiffness::DefaultTableSize,
-		FPBDStiffness::DefaultParameterFitBase,
-		MaxStiffness)
+		ParticleCount)
+	, BucklingStiffnessWarp(
+		FSolverVec2(GetWeightedFloatXPBDAnisoBucklingStiffnessWarp(PropertyCollection, MaxStiffness)).ClampAxes(0, MaxStiffness),
+		WeightMaps.FindRef(GetXPBDAnisoBucklingStiffnessWarpString(PropertyCollection, XPBDAnisoBucklingStiffnessWarpName.ToString())),
+		TConstArrayView<TVec2<int32>>(ConstraintSharedEdges),
+		ParticleOffset,
+		ParticleCount)
 	, BucklingStiffnessWeft(
-		FSolverVec2(GetWeightedFloatXPBDAnisoBucklingStiffnessWeft(PropertyCollection, MaxStiffness)),
+		FSolverVec2(GetWeightedFloatXPBDAnisoBucklingStiffnessWeft(PropertyCollection, MaxStiffness)).ClampAxes(0, MaxStiffness),
 		WeightMaps.FindRef(GetXPBDAnisoBucklingStiffnessWeftString(PropertyCollection, XPBDAnisoBucklingStiffnessWeftName.ToString())),
 		TConstArrayView<TVec2<int32>>(ConstraintSharedEdges),
 		ParticleOffset,
-		ParticleCount,
-		FPBDStiffness::DefaultTableSize,
-		FPBDStiffness::DefaultParameterFitBase,
-		MaxStiffness)
+		ParticleCount)
 	, BucklingStiffnessBias(
-		FSolverVec2(GetWeightedFloatXPBDAnisoBucklingStiffnessBias(PropertyCollection, MaxStiffness)),
+		FSolverVec2(GetWeightedFloatXPBDAnisoBucklingStiffnessBias(PropertyCollection, MaxStiffness)).ClampAxes(0, MaxStiffness),
 		WeightMaps.FindRef(GetXPBDAnisoBucklingStiffnessBiasString(PropertyCollection, XPBDAnisoBucklingStiffnessBiasName.ToString())),
 		TConstArrayView<TVec2<int32>>(ConstraintSharedEdges),
 		ParticleOffset,
-		ParticleCount,
-		FPBDStiffness::DefaultTableSize,
-		FPBDStiffness::DefaultParameterFitBase,
-		MaxStiffness)
+		ParticleCount)
 	, DampingRatio(
 		FSolverVec2(GetWeightedFloatXPBDAnisoBendingDamping(PropertyCollection, MinDamping)).ClampAxes(MinDamping, MaxDamping),
 		WeightMaps.FindRef(GetXPBDAnisoBendingDampingString(PropertyCollection, XPBDAnisoBendingDampingName.ToString())),
@@ -202,56 +207,55 @@ FXPBDAnisotropicBendingConstraints::FXPBDAnisotropicBendingConstraints(const FSo
 	const FSolverVec2& InBucklingStiffnessWarp,
 	const FSolverVec2& InBucklingStiffnessWeft,
 	const FSolverVec2& InBucklingStiffnessBias,
-	const FSolverVec2& InDampingRatio,
-	bool bTrimKinematicConstraints)
+	const FSolverVec2& InDampingRatio)
 	: Base(
 		InParticles,
 		ParticleOffset,
 		ParticleCount,
 		TriangleMesh.GetUniqueAdjacentElements(),
-		StiffnessWarpMultipliers,
-		BucklingStiffnessWarpMultipliers,
+		TConstArrayView<FRealSingle>(), // We don't use base stiffness weight maps
+		TConstArrayView<FRealSingle>(), // We don't use base stiffness weight maps
 		InStiffnessWarp,
 		InBucklingRatio,
 		InBucklingStiffnessWarp,
-		bTrimKinematicConstraints,
+		true /*bTrimKinematicConstraints*/,
 		MaxStiffness)
+	, StiffnessWarp(
+		InStiffnessWarp.ClampAxes(0, MaxStiffness),
+		StiffnessWarpMultipliers,
+		TConstArrayView<TVec2<int32>>(ConstraintSharedEdges),
+		ParticleOffset,
+		ParticleCount)
 	, StiffnessWeft(
-		InStiffnessWeft,
+		InStiffnessWeft.ClampAxes(0, MaxStiffness),
 		StiffnessWeftMultipliers,
 		TConstArrayView<TVec2<int32>>(ConstraintSharedEdges),
 		ParticleOffset,
-		ParticleCount,
-		FPBDStiffness::DefaultTableSize,
-		FPBDStiffness::DefaultParameterFitBase,
-		MaxStiffness)
+		ParticleCount)
 	, StiffnessBias(
-		InStiffnessBias,
+		InStiffnessBias.ClampAxes(0, MaxStiffness),
 		StiffnessBiasMultipliers,
 		TConstArrayView<TVec2<int32>>(ConstraintSharedEdges),
 		ParticleOffset,
-		ParticleCount,
-		FPBDStiffness::DefaultTableSize,
-		FPBDStiffness::DefaultParameterFitBase,
-		MaxStiffness)
+		ParticleCount)
+	, BucklingStiffnessWarp(
+		InBucklingStiffnessWarp.ClampAxes(0, MaxStiffness),
+		BucklingStiffnessWarpMultipliers,
+		TConstArrayView<TVec2<int32>>(ConstraintSharedEdges),
+		ParticleOffset,
+		ParticleCount)
 	, BucklingStiffnessWeft(
-		InBucklingStiffnessWeft,
+		InBucklingStiffnessWeft.ClampAxes(0, MaxStiffness),
 		BucklingStiffnessWeftMultipliers,
 		TConstArrayView<TVec2<int32>>(ConstraintSharedEdges),
 		ParticleOffset,
-		ParticleCount,
-		FPBDStiffness::DefaultTableSize,
-		FPBDStiffness::DefaultParameterFitBase,
-		MaxStiffness)
+		ParticleCount)
 	, BucklingStiffnessBias(
-		InBucklingStiffnessBias,
+		InBucklingStiffnessBias.ClampAxes(0, MaxStiffness),
 		BucklingStiffnessBiasMultipliers,
 		TConstArrayView<TVec2<int32>>(ConstraintSharedEdges),
 		ParticleOffset,
-		ParticleCount,
-		FPBDStiffness::DefaultTableSize,
-		FPBDStiffness::DefaultParameterFitBase,
-		MaxStiffness)
+		ParticleCount)
 	, DampingRatio(
 		InDampingRatio.ClampAxes((FSolverReal)0., (FSolverReal)1.),
 		DampingMultipliers,
@@ -377,14 +381,28 @@ void FXPBDAnisotropicBendingConstraints::InitColor(const SolverParticlesOrRange&
 		ConstraintSharedEdges = MoveTemp(ReorderedConstraintSharedEdges);
 		RestAngles = MoveTemp(ReorderedRestAngles);
 		WarpWeftBiasBaseMultipliers = MoveTemp(ReorderedWarpWeftBiasBaseMultipliers);
-		Stiffness.ReorderIndices(OrigToReorderedIndices);
+		StiffnessWarp.ReorderIndices(OrigToReorderedIndices);
 		StiffnessWeft.ReorderIndices(OrigToReorderedIndices);
 		StiffnessBias.ReorderIndices(OrigToReorderedIndices);
-		BucklingStiffness.ReorderIndices(OrigToReorderedIndices);
+		BucklingStiffnessWarp.ReorderIndices(OrigToReorderedIndices);
 		BucklingStiffnessWeft.ReorderIndices(OrigToReorderedIndices);
 		BucklingStiffnessBias.ReorderIndices(OrigToReorderedIndices);
 		DampingRatio.ReorderIndices(OrigToReorderedIndices);
+
+#if INTEL_ISPC
+		ConstraintsIndex1.SetNumUninitialized(Constraints.Num());
+		ConstraintsIndex2.SetNumUninitialized(Constraints.Num());
+		ConstraintsIndex3.SetNumUninitialized(Constraints.Num());
+		ConstraintsIndex4.SetNumUninitialized(Constraints.Num());
+		for (int32 ConstraintIndex = 0; ConstraintIndex < Constraints.Num(); ++ConstraintIndex)
+		{
+			ConstraintsIndex1[ConstraintIndex] = Constraints[ConstraintIndex][0];
+			ConstraintsIndex2[ConstraintIndex] = Constraints[ConstraintIndex][1];
+			ConstraintsIndex3[ConstraintIndex] = Constraints[ConstraintIndex][2];
+			ConstraintsIndex4[ConstraintIndex] = Constraints[ConstraintIndex][3];
+		}
 	}
+#endif
 }
 
 void FXPBDAnisotropicBendingConstraints::SetProperties(
@@ -393,65 +411,56 @@ void FXPBDAnisotropicBendingConstraints::SetProperties(
 {
 	if (IsXPBDAnisoBendingStiffnessWarpMutable(PropertyCollection))
 	{
-		const FSolverVec2 WeightedValue(GetWeightedFloatXPBDAnisoBendingStiffnessWarp(PropertyCollection));
+		const FSolverVec2 WeightedValue(GetWeightedFloatXPBDAnisoBendingStiffnessWarp(PropertyCollection).ClampAxes(0, MaxStiffness));
 		if (IsXPBDAnisoBendingStiffnessWarpStringDirty(PropertyCollection))
 		{
 			const FString& WeightMapName = GetXPBDAnisoBendingStiffnessWarpString(PropertyCollection);
-			Stiffness = FPBDStiffness(
+			StiffnessWarp = FPBDFlatWeightMap(
 				WeightedValue,
 				WeightMaps.FindRef(WeightMapName),
 				TConstArrayView<TVec2<int32>>(ConstraintSharedEdges),
 				ParticleOffset,
-				ParticleCount,
-				FPBDStiffness::DefaultTableSize,
-				FPBDStiffness::DefaultParameterFitBase,
-				MaxStiffness);
+				ParticleCount);
 		}
 		else
 		{
-			Stiffness.SetWeightedValue(WeightedValue, MaxStiffness);
+			StiffnessWarp.SetWeightedValue(WeightedValue);
 		}
 	}
 	if (IsXPBDAnisoBendingStiffnessWeftMutable(PropertyCollection))
 	{
-		const FSolverVec2 WeightedValue(GetWeightedFloatXPBDAnisoBendingStiffnessWeft(PropertyCollection));
+		const FSolverVec2 WeightedValue(GetWeightedFloatXPBDAnisoBendingStiffnessWeft(PropertyCollection).ClampAxes(0, MaxStiffness));
 		if (IsXPBDAnisoBendingStiffnessWeftStringDirty(PropertyCollection))
 		{
 			const FString& WeightMapName = GetXPBDAnisoBendingStiffnessWeftString(PropertyCollection);
-			StiffnessWeft = FPBDStiffness(
+			StiffnessWeft = FPBDFlatWeightMap(
 				WeightedValue,
 				WeightMaps.FindRef(WeightMapName),
 				TConstArrayView<TVec2<int32>>(ConstraintSharedEdges),
 				ParticleOffset,
-				ParticleCount,
-				FPBDStiffness::DefaultTableSize,
-				FPBDStiffness::DefaultParameterFitBase,
-				MaxStiffness);
+				ParticleCount);
 		}
 		else
 		{
-			StiffnessWeft.SetWeightedValue(WeightedValue, MaxStiffness);
+			StiffnessWeft.SetWeightedValue(WeightedValue);
 		}
 	}
 	if (IsXPBDAnisoBendingStiffnessBiasMutable(PropertyCollection))
 	{
-		const FSolverVec2 WeightedValue(GetWeightedFloatXPBDAnisoBendingStiffnessBias(PropertyCollection));
+		const FSolverVec2 WeightedValue(GetWeightedFloatXPBDAnisoBendingStiffnessBias(PropertyCollection).ClampAxes(0, MaxStiffness));
 		if (IsXPBDAnisoBendingStiffnessBiasStringDirty(PropertyCollection))
 		{
 			const FString& WeightMapName = GetXPBDAnisoBendingStiffnessBiasString(PropertyCollection);
-			StiffnessBias = FPBDStiffness(
+			StiffnessBias = FPBDFlatWeightMap(
 				WeightedValue,
 				WeightMaps.FindRef(WeightMapName),
 				TConstArrayView<TVec2<int32>>(ConstraintSharedEdges),
 				ParticleOffset,
-				ParticleCount,
-				FPBDStiffness::DefaultTableSize,
-				FPBDStiffness::DefaultParameterFitBase,
-				MaxStiffness);
+				ParticleCount);
 		}
 		else
 		{
-			StiffnessBias.SetWeightedValue(WeightedValue, MaxStiffness);
+			StiffnessBias.SetWeightedValue(WeightedValue);
 		}
 	}
 	if (IsXPBDAnisoBucklingRatioMutable(PropertyCollection))
@@ -460,65 +469,56 @@ void FXPBDAnisotropicBendingConstraints::SetProperties(
 	}
 	if (IsXPBDAnisoBucklingStiffnessWarpMutable(PropertyCollection))
 	{
-		const FSolverVec2 WeightedValue(GetWeightedFloatXPBDAnisoBucklingStiffnessWarp(PropertyCollection));
+		const FSolverVec2 WeightedValue(GetWeightedFloatXPBDAnisoBucklingStiffnessWarp(PropertyCollection).ClampAxes(0, MaxStiffness));
 		if (IsXPBDAnisoBucklingStiffnessWarpStringDirty(PropertyCollection))
 		{
 			const FString& WeightMapName = GetXPBDAnisoBucklingStiffnessWarpString(PropertyCollection);
-			BucklingStiffness = FPBDStiffness(
+			BucklingStiffnessWarp = FPBDFlatWeightMap(
 				WeightedValue,
 				WeightMaps.FindRef(WeightMapName),
 				TConstArrayView<TVec2<int32>>(ConstraintSharedEdges),
 				ParticleOffset,
-				ParticleCount,
-				FPBDStiffness::DefaultTableSize,
-				FPBDStiffness::DefaultParameterFitBase,
-				MaxStiffness);
+				ParticleCount);
 		}
 		else
 		{
-			BucklingStiffness.SetWeightedValue(WeightedValue, MaxStiffness);
+			BucklingStiffnessWarp.SetWeightedValue(WeightedValue);
 		}
 	}
 	if (IsXPBDAnisoBucklingStiffnessWeftMutable(PropertyCollection))
 	{
-		const FSolverVec2 WeightedValue(GetWeightedFloatXPBDAnisoBucklingStiffnessWeft(PropertyCollection));
+		const FSolverVec2 WeightedValue(GetWeightedFloatXPBDAnisoBucklingStiffnessWeft(PropertyCollection).ClampAxes(0, MaxStiffness));
 		if (IsXPBDAnisoBucklingStiffnessWeftStringDirty(PropertyCollection))
 		{
 			const FString& WeightMapName = GetXPBDAnisoBucklingStiffnessWeftString(PropertyCollection);
-			BucklingStiffnessWeft = FPBDStiffness(
+			BucklingStiffnessWeft = FPBDFlatWeightMap(
 				WeightedValue,
 				WeightMaps.FindRef(WeightMapName),
 				TConstArrayView<TVec2<int32>>(ConstraintSharedEdges),
 				ParticleOffset,
-				ParticleCount,
-				FPBDStiffness::DefaultTableSize,
-				FPBDStiffness::DefaultParameterFitBase,
-				MaxStiffness);
+				ParticleCount);
 		}
 		else
 		{
-			BucklingStiffnessWeft.SetWeightedValue(WeightedValue, MaxStiffness);
+			BucklingStiffnessWeft.SetWeightedValue(WeightedValue);
 		}
 	}
 	if (IsXPBDAnisoBucklingStiffnessBiasMutable(PropertyCollection))
 	{
-		const FSolverVec2 WeightedValue(GetWeightedFloatXPBDAnisoBucklingStiffnessBias(PropertyCollection));
+		const FSolverVec2 WeightedValue(GetWeightedFloatXPBDAnisoBucklingStiffnessBias(PropertyCollection).ClampAxes(0, MaxStiffness));
 		if (IsXPBDAnisoBucklingStiffnessBiasStringDirty(PropertyCollection))
 		{
 			const FString& WeightMapName = GetXPBDAnisoBucklingStiffnessBiasString(PropertyCollection);
-			BucklingStiffnessBias = FPBDStiffness(
+			BucklingStiffnessBias = FPBDFlatWeightMap(
 				WeightedValue,
 				WeightMaps.FindRef(WeightMapName),
 				TConstArrayView<TVec2<int32>>(ConstraintSharedEdges),
 				ParticleOffset,
-				ParticleCount,
-				FPBDStiffness::DefaultTableSize,
-				FPBDStiffness::DefaultParameterFitBase,
-				MaxStiffness);
+				ParticleCount);
 		}
 		else
 		{
-			BucklingStiffnessBias.SetWeightedValue(WeightedValue, MaxStiffness);
+			BucklingStiffnessBias.SetWeightedValue(WeightedValue);
 		}
 	}
 	if (IsXPBDAnisoBendingDampingMutable(PropertyCollection))
@@ -527,7 +527,7 @@ void FXPBDAnisotropicBendingConstraints::SetProperties(
 		if (IsXPBDAnisoBendingDampingStringDirty(PropertyCollection))
 		{
 			const FString& WeightMapName = GetXPBDAnisoBendingDampingString(PropertyCollection);
-			DampingRatio = FPBDWeightMap(
+			DampingRatio = FPBDFlatWeightMap(
 				WeightedValue,
 				WeightMaps.FindRef(WeightMapName),
 				TConstArrayView<TVec2<int32>>(ConstraintSharedEdges),
@@ -545,59 +545,122 @@ template<bool bDampingOnly, bool bElasticOnly, typename SolverParticlesOrRange>
 void FXPBDAnisotropicBendingConstraints::ApplyHelper(SolverParticlesOrRange& Particles, const FSolverReal Dt, const int32 ConstraintIndex, const FSolverVec3& ExpStiffnessValues,
 	const FSolverVec3& ExpBucklingStiffnessValues, const FSolverReal DampingRatioValue) const
 {
-	const TVec4<int32>& Constraint = Constraints[ConstraintIndex];
+	const TVec4<int32>& Constraint = Constraints.GetData()[ConstraintIndex];
 	const int32 Index1 = Constraint[0];
 	const int32 Index2 = Constraint[1];
 	const int32 Index3 = Constraint[2];
 	const int32 Index4 = Constraint[3];
-	const FSolverVec3& WarpWeftBiasBaseMultiplier = WarpWeftBiasBaseMultipliers[ConstraintIndex];
+	const FSolverVec3& WarpWeftBiasBaseMultiplier = WarpWeftBiasBaseMultipliers.GetData()[ConstraintIndex];
 
-	const FSolverReal BiphasicStiffnessValue = IsBuckled[ConstraintIndex] ?
+	const FSolverReal BiphasicStiffnessValue = IsBuckled.GetData()[ConstraintIndex] ?
 		WarpWeftBiasBaseMultiplier.Dot(ExpBucklingStiffnessValues): WarpWeftBiasBaseMultiplier.Dot(ExpStiffnessValues);
 
-	const FSolverReal CombinedInvMass = Particles.InvM(Index1) + Particles.InvM(Index2) + Particles.InvM(Index3) + Particles.InvM(Index4);
-	if (BiphasicStiffnessValue < MinStiffness || CombinedInvMass < UE_SMALL_NUMBER)
+	const FSolverReal InvM1 = Particles.GetPAndInvM().GetData()[Index1].InvM;
+	const FSolverReal InvM2 = Particles.GetPAndInvM().GetData()[Index2].InvM;
+	const FSolverReal InvM3 = Particles.GetPAndInvM().GetData()[Index3].InvM;
+	const FSolverReal InvM4 = Particles.GetPAndInvM().GetData()[Index4].InvM;
+
+	if (BiphasicStiffnessValue <= MinStiffness)
 	{
 		return;
 	}
+	const FSolverReal CombinedInvMass = InvM1 + InvM2 + InvM3 + InvM4;
+	check(CombinedInvMass > 0);
 
-	const FSolverReal Damping = (FSolverReal)2.f * DampingRatioValue * FMath::Sqrt(BiphasicStiffnessValue / CombinedInvMass);
+	FSolverReal ElasticTerm(0.f), DampingTerm(0.f), DenomScale(0.f);
+	const FSolverVec3 P1 = Particles.GetPAndInvM().GetData()[Index1].P;
+	const FSolverVec3 P2 = Particles.GetPAndInvM().GetData()[Index2].P;
+	const FSolverVec3 P3 = Particles.GetPAndInvM().GetData()[Index3].P;
+	const FSolverVec3 P4 = Particles.GetPAndInvM().GetData()[Index4].P;
 
-	const FSolverReal Angle = CalcAngle(Particles.P(Index1), Particles.P(Index2), Particles.P(Index3), Particles.P(Index4));
-	const TStaticArray<FSolverVec3, 4> Grads = Base::GetGradients(Particles, ConstraintIndex);
+	FSolverReal Angle(0.f);
+	const TStaticArray<FSolverVec3, 4> Grads = Base::CalcGradients(P1, P2, P3, P4, bDampingOnly ? nullptr : &Angle);
 
-	const FSolverVec3 V1TimesDt = Particles.P(Index1) - Particles.X(Index1);
-	const FSolverVec3 V2TimesDt = Particles.P(Index2) - Particles.X(Index2);
-	const FSolverVec3 V3TimesDt = Particles.P(Index3) - Particles.X(Index3);
-	const FSolverVec3 V4TimesDt = Particles.P(Index4) - Particles.X(Index4);
+	if constexpr (!bDampingOnly)
+	{
+		const FSolverReal AlphaInv = BiphasicStiffnessValue * Dt * Dt;
+		ElasticTerm = AlphaInv * (Angle - RestAngles.GetData()[ConstraintIndex]);
+
+		DenomScale += AlphaInv;
+	}
+	if constexpr (!bElasticOnly)
+	{
+		const FSolverReal Damping = (FSolverReal)2.f * DampingRatioValue * FMath::Sqrt(BiphasicStiffnessValue / CombinedInvMass);
+		const FSolverReal BetaDt = Damping * Dt;
 
-	FSolverReal& Lambda = bDampingOnly ? LambdasDamping[ConstraintIndex] : Lambdas[ConstraintIndex];
-	const FSolverReal AlphaInv = bDampingOnly ? (FSolverReal)0.f : BiphasicStiffnessValue * Dt * Dt;
-	const FSolverReal BetaDt = bElasticOnly ? (FSolverReal)0.f : Damping * Dt;
+		const FSolverVec3 V1TimesDt = P1 - Particles.XArray().GetData()[Index1];
+		const FSolverVec3 V2TimesDt = P2 - Particles.XArray().GetData()[Index2];
+		const FSolverVec3 V3TimesDt = P3 - Particles.XArray().GetData()[Index3];
+		const FSolverVec3 V4TimesDt = P4 - Particles.XArray().GetData()[Index4];
 
-	const FSolverReal DampingTerm = bElasticOnly ? (FSolverReal)0.f : BetaDt * (FSolverVec3::DotProduct(V1TimesDt, Grads[0]) + FSolverVec3::DotProduct(V2TimesDt, Grads[1]) + FSolverVec3::DotProduct(V3TimesDt, Grads[2]) + FSolverVec3::DotProduct(V4TimesDt, Grads[3]));
-	const FSolverReal ElasticTerm = bDampingOnly ? (FSolverReal)0.f : AlphaInv * (Angle - RestAngles[ConstraintIndex]);
+		DampingTerm = BetaDt * (FSolverVec3::DotProduct(V1TimesDt, Grads[0]) + FSolverVec3::DotProduct(V2TimesDt, Grads[1]) + FSolverVec3::DotProduct(V3TimesDt, Grads[2]) + FSolverVec3::DotProduct(V4TimesDt, Grads[3]));
 
-	const FSolverReal Denom = (AlphaInv + BetaDt) * (Particles.InvM(Index1) * Grads[0].SizeSquared() + Particles.InvM(Index2) * Grads[1].SizeSquared() + Particles.InvM(Index3) * Grads[2].SizeSquared() + Particles.InvM(Index4) * Grads[3].SizeSquared()) + (FSolverReal)1.;
+		DenomScale += BetaDt;
+	}
+
+	FSolverReal& Lambda = bDampingOnly ? LambdasDamping.GetData()[ConstraintIndex] : Lambdas.GetData()[ConstraintIndex];
+
+	const FSolverReal Denom = (DenomScale) * (InvM1 * Grads[0].SizeSquared() + InvM2 * Grads[1].SizeSquared() + InvM3 * Grads[2].SizeSquared() + InvM4 * Grads[3].SizeSquared()) + (FSolverReal)1.;
 	const FSolverReal DLambda = (-ElasticTerm - DampingTerm - Lambda) / Denom;
 
-	Particles.P(Index1) += DLambda * Particles.InvM(Index1) * Grads[0];
-	Particles.P(Index2) += DLambda * Particles.InvM(Index2) * Grads[1];
-	Particles.P(Index3) += DLambda * Particles.InvM(Index3) * Grads[2];
-	Particles.P(Index4) += DLambda * Particles.InvM(Index4) * Grads[3];
+	Particles.GetPAndInvM().GetData()[Index1].P += DLambda * InvM1 * Grads[0];
+	Particles.GetPAndInvM().GetData()[Index2].P += DLambda * InvM2 * Grads[1];
+	Particles.GetPAndInvM().GetData()[Index3].P += DLambda * InvM3 * Grads[2];
+	Particles.GetPAndInvM().GetData()[Index4].P += DLambda * InvM4 * Grads[3];
 	Lambda += DLambda;
 }
 
+template<typename SolverParticlesOrRange>
+void FXPBDAnisotropicBendingConstraints::Init(const SolverParticlesOrRange& Particles)
+{
+	TRACE_CPUPROFILER_EVENT_SCOPE(FXPBDAnisotropicBendingConstraints_Init);
+	Lambdas.Reset();
+	Lambdas.AddZeroed(Constraints.Num());
+	LambdasDamping.Reset();
+	LambdasDamping.AddZeroed(Constraints.Num());
+#if INTEL_ISPC
+	IsBuckled.SetNumUninitialized(Constraints.Num());
+	X1Array.SetNumUninitialized(Constraints.Num());
+	X2Array.SetNumUninitialized(Constraints.Num());
+	X3Array.SetNumUninitialized(Constraints.Num());
+	X4Array.SetNumUninitialized(Constraints.Num());
+	if (bRealTypeCompatibleWithISPC && bChaos_Bending_ISPC_Enabled && ConstraintsIndex1.Num() == Constraints.Num())
+	{
+		ispc::InitXPBDBendingConstraintsIsBuckled(
+			(const ispc::FVector3f*)Particles.XArray().GetData(),
+			ConstraintsIndex1.GetData(),
+			ConstraintsIndex2.GetData(),
+			ConstraintsIndex3.GetData(),
+			ConstraintsIndex4.GetData(),
+			RestAngles.GetData(),
+			IsBuckled.GetData(),
+			(ispc::FVector3f*)X1Array.GetData(),
+			(ispc::FVector3f*)X2Array.GetData(),
+			(ispc::FVector3f*)X3Array.GetData(),
+			(ispc::FVector3f*)X4Array.GetData(),
+			BucklingRatio,
+			Constraints.Num()
+		);
+	}
+	else
+#endif
+	{
+		FPBDBendingConstraintsBase::Init(Particles);
+	}
+}
+template CHAOS_API void FXPBDAnisotropicBendingConstraints::Init(const FSolverParticles& Particles);
+template CHAOS_API void FXPBDAnisotropicBendingConstraints::Init(const FSolverParticlesRange& Particles);
+
 template<typename SolverParticlesOrRange>
 void FXPBDAnisotropicBendingConstraints::Apply(SolverParticlesOrRange& Particles, const FSolverReal Dt) const
 {
 	TRACE_CPUPROFILER_EVENT_SCOPE(FXPBDAnisotropicBendingConstraints_Apply);
 	SCOPE_CYCLE_COUNTER(STAT_XPBD_AnisoBending);
 
-	const bool StiffnessHasWeightMap = Stiffness.HasWeightMap();
+	const bool StiffnessHasWeightMap = StiffnessWarp.HasWeightMap();
 	const bool StiffnessWeftHasWeightMap = StiffnessWeft.HasWeightMap();
 	const bool StiffnessBiasHasWeightMap = StiffnessBias.HasWeightMap();
-	const bool BucklingStiffnessHasWeightMap = BucklingStiffness.HasWeightMap();
+	const bool BucklingStiffnessHasWeightMap = BucklingStiffnessWarp.HasWeightMap();
 	const bool BucklingStiffnessWeftHasWeightMap = BucklingStiffnessWeft.HasWeightMap();
 	const bool BucklingStiffnessBiasHasWeightMap = BucklingStiffnessBias.HasWeightMap();
 	const bool DampingHasWeightMap = DampingRatio.HasWeightMap();
@@ -612,11 +675,11 @@ void FXPBDAnisotropicBendingConstraints::Apply(SolverParticlesOrRange& Particles
 				!BucklingStiffnessHasWeightMap && !BucklingStiffnessWeftHasWeightMap && !BucklingStiffnessBiasHasWeightMap &&
 				!DampingHasWeightMap)
 			{
-				const FSolverVec3 ExpStiffnessValue((FSolverReal)StiffnessWeft, (FSolverReal)Stiffness, (FSolverReal)StiffnessBias);
-				const FSolverVec3 ExpBucklingValue((FSolverReal)BucklingStiffnessWeft, (FSolverReal)BucklingStiffness, (FSolverReal)BucklingStiffnessBias);
+				const FSolverVec3 ExpStiffnessValue((FSolverReal)StiffnessWeft, (FSolverReal)StiffnessWarp, (FSolverReal)StiffnessBias);
+				const FSolverVec3 ExpBucklingValue((FSolverReal)BucklingStiffnessWeft, (FSolverReal)BucklingStiffnessWarp, (FSolverReal)BucklingStiffnessBias);
 				const FSolverReal DampingRatioValue = (FSolverReal)DampingRatio;
 
-				if (ExpStiffnessValue.Max() < MinStiffness && ExpBucklingValue.Max() < MinStiffness)
+				if (ExpStiffnessValue.Max() <= MinStiffness && ExpBucklingValue.Max() <= MinStiffness)
 				{
 					return;
 				}
@@ -628,21 +691,59 @@ void FXPBDAnisotropicBendingConstraints::Apply(SolverParticlesOrRange& Particles
 						for (int32 ConstraintColorIndex = 0; ConstraintColorIndex < ConstraintColorNum; ++ConstraintColorIndex)
 						{
 							const int32 ColorStart = ConstraintsPerColorStartIndex[ConstraintColorIndex];
-							const int32 ColorSize = ConstraintsPerColorStartIndex[ConstraintColorIndex + 1] - ColorStart;
-							ispc::ApplyXPBDAnisotropicBendingDampingConstraints(
-								(ispc::FVector4f*)Particles.GetPAndInvM().GetData(),
-								(const ispc::FVector3f*)Particles.XArray().GetData(),
-								(ispc::FIntVector4*)&Constraints.GetData()[ColorStart],
-								&RestAngles.GetData()[ColorStart],
-								&IsBuckled.GetData()[ColorStart],
-								(const ispc::FVector3f*)&WarpWeftBiasBaseMultipliers.GetData()[ColorStart],
-								&LambdasDamping.GetData()[ColorStart],
-								Dt,
-								MinStiffness,
-								reinterpret_cast<const ispc::FVector3f&>(ExpStiffnessValue),
-								reinterpret_cast<const ispc::FVector3f&>(ExpBucklingValue),
-								DampingRatioValue,
-								ColorSize);
+							const int32 ColorSize = ConstraintsPerColorStartIndex[ConstraintColorIndex + 1] - ColorStart;	
+							if (ColorSize >= Chaos_XPBDBending_ISPC_ParallelBatchSize * Chaos_XPBDBending_ISPC_MinNumParallelBatches)
+							{
+								const int32 NumBatches = FMath::DivideAndRoundUp(ColorSize, Chaos_XPBDBending_ISPC_ParallelBatchSize);
+								PhysicsParallelFor(NumBatches, [this, &Particles, ColorStart, ColorSize,
+									ParallelBatchSize = Chaos_XPBDBending_ISPC_ParallelBatchSize,
+									Dt, &ExpStiffnessValue, &ExpBucklingValue, DampingRatioValue](const int32 BatchIndex)
+								{
+									const int32 BatchStart = BatchIndex * ParallelBatchSize;
+									const int32 BatchEnd = FMath::Min((BatchIndex + 1) * ParallelBatchSize, ColorSize);
+									ispc::ApplyXPBDAnisotropicBendingDampingConstraints(
+										(ispc::FVector4f*)Particles.GetPAndInvM().GetData(),
+										(const ispc::FVector3f*)&X1Array.GetData()[ColorStart + BatchStart],
+										(const ispc::FVector3f*)&X2Array.GetData()[ColorStart + BatchStart],
+										(const ispc::FVector3f*)&X3Array.GetData()[ColorStart + BatchStart],
+										(const ispc::FVector3f*)&X4Array.GetData()[ColorStart + BatchStart],
+										&ConstraintsIndex1.GetData()[ColorStart + BatchStart],
+										&ConstraintsIndex2.GetData()[ColorStart + BatchStart],
+										&ConstraintsIndex3.GetData()[ColorStart + BatchStart],
+										&ConstraintsIndex4.GetData()[ColorStart + BatchStart],
+										&RestAngles.GetData()[ColorStart + BatchStart],
+										&IsBuckled.GetData()[ColorStart + BatchStart],
+										(const ispc::FVector3f*)&WarpWeftBiasBaseMultipliers.GetData()[ColorStart + BatchStart],
+										&LambdasDamping.GetData()[ColorStart + BatchStart],
+										Dt,
+										reinterpret_cast<const ispc::FVector3f&>(ExpStiffnessValue),
+										reinterpret_cast<const ispc::FVector3f&>(ExpBucklingValue),
+										DampingRatioValue,
+										BatchEnd - BatchStart);
+								});
+							}
+							else
+							{
+								ispc::ApplyXPBDAnisotropicBendingDampingConstraints(
+									(ispc::FVector4f*)Particles.GetPAndInvM().GetData(),
+									(const ispc::FVector3f*)&X1Array.GetData()[ColorStart],
+									(const ispc::FVector3f*)&X2Array.GetData()[ColorStart],
+									(const ispc::FVector3f*)&X3Array.GetData()[ColorStart],
+									(const ispc::FVector3f*)&X4Array.GetData()[ColorStart],
+									&ConstraintsIndex1.GetData()[ColorStart],
+									&ConstraintsIndex2.GetData()[ColorStart],
+									&ConstraintsIndex3.GetData()[ColorStart],
+									&ConstraintsIndex4.GetData()[ColorStart],
+									&RestAngles.GetData()[ColorStart],
+									&IsBuckled.GetData()[ColorStart],
+									(const ispc::FVector3f*)&WarpWeftBiasBaseMultipliers.GetData()[ColorStart],
+									&LambdasDamping.GetData()[ColorStart],
+									Dt,
+									reinterpret_cast<const ispc::FVector3f&>(ExpStiffnessValue),
+									reinterpret_cast<const ispc::FVector3f&>(ExpBucklingValue),
+									DampingRatioValue,
+									ColorSize);
+							}
 						}
 					}
 					else
@@ -651,20 +752,58 @@ void FXPBDAnisotropicBendingConstraints::Apply(SolverParticlesOrRange& Particles
 						{
 							const int32 ColorStart = ConstraintsPerColorStartIndex[ConstraintColorIndex];
 							const int32 ColorSize = ConstraintsPerColorStartIndex[ConstraintColorIndex + 1] - ColorStart;
-							ispc::ApplyXPBDAnisotropicBendingConstraintsWithDamping(
-								(ispc::FVector4f*)Particles.GetPAndInvM().GetData(),
-								(const ispc::FVector3f*)Particles.XArray().GetData(),
-								(ispc::FIntVector4*)&Constraints.GetData()[ColorStart],
-								&RestAngles.GetData()[ColorStart],
-								&IsBuckled.GetData()[ColorStart],
-								(const ispc::FVector3f*)&WarpWeftBiasBaseMultipliers.GetData()[ColorStart],
-								&Lambdas.GetData()[ColorStart],
-								Dt,
-								MinStiffness,
-								reinterpret_cast<const ispc::FVector3f&>(ExpStiffnessValue),
-								reinterpret_cast<const ispc::FVector3f&>(ExpBucklingValue),
-								DampingRatioValue,
-								ColorSize);
+							if (ColorSize >= Chaos_XPBDBending_ISPC_ParallelBatchSize * Chaos_XPBDBending_ISPC_MinNumParallelBatches)
+							{
+								const int32 NumBatches = FMath::DivideAndRoundUp(ColorSize, Chaos_XPBDBending_ISPC_ParallelBatchSize);
+								PhysicsParallelFor(NumBatches, [this, &Particles, ColorStart, ColorSize,
+									ParallelBatchSize = Chaos_XPBDBending_ISPC_ParallelBatchSize,
+									Dt, &ExpStiffnessValue, &ExpBucklingValue, DampingRatioValue](const int32 BatchIndex)
+								{
+									const int32 BatchStart = BatchIndex * ParallelBatchSize;
+									const int32 BatchEnd = FMath::Min((BatchIndex + 1) * ParallelBatchSize, ColorSize);
+									ispc::ApplyXPBDAnisotropicBendingConstraintsWithDamping(
+										(ispc::FVector4f*)Particles.GetPAndInvM().GetData(),
+										(const ispc::FVector3f*)&X1Array.GetData()[ColorStart + BatchStart],
+										(const ispc::FVector3f*)&X2Array.GetData()[ColorStart + BatchStart],
+										(const ispc::FVector3f*)&X3Array.GetData()[ColorStart + BatchStart],
+										(const ispc::FVector3f*)&X4Array.GetData()[ColorStart + BatchStart],
+										&ConstraintsIndex1.GetData()[ColorStart + BatchStart],
+										&ConstraintsIndex2.GetData()[ColorStart + BatchStart],
+										&ConstraintsIndex3.GetData()[ColorStart + BatchStart],
+										&ConstraintsIndex4.GetData()[ColorStart + BatchStart],
+										&RestAngles.GetData()[ColorStart + BatchStart],
+										&IsBuckled.GetData()[ColorStart + BatchStart],
+										(const ispc::FVector3f*)&WarpWeftBiasBaseMultipliers.GetData()[ColorStart + BatchStart],
+										&Lambdas.GetData()[ColorStart + BatchStart],
+										Dt,
+										reinterpret_cast<const ispc::FVector3f&>(ExpStiffnessValue),
+										reinterpret_cast<const ispc::FVector3f&>(ExpBucklingValue),
+										DampingRatioValue,
+										BatchEnd - BatchStart);
+								});
+							}
+							else
+							{
+								ispc::ApplyXPBDAnisotropicBendingConstraintsWithDamping(
+									(ispc::FVector4f*)Particles.GetPAndInvM().GetData(),
+									(const ispc::FVector3f*)&X1Array.GetData()[ColorStart],
+									(const ispc::FVector3f*)&X2Array.GetData()[ColorStart],
+									(const ispc::FVector3f*)&X3Array.GetData()[ColorStart],
+									(const ispc::FVector3f*)&X4Array.GetData()[ColorStart],
+									&ConstraintsIndex1.GetData()[ColorStart],
+									&ConstraintsIndex2.GetData()[ColorStart],
+									&ConstraintsIndex3.GetData()[ColorStart],
+									&ConstraintsIndex4.GetData()[ColorStart],
+									&RestAngles.GetData()[ColorStart],
+									&IsBuckled.GetData()[ColorStart],
+									(const ispc::FVector3f*)&WarpWeftBiasBaseMultipliers.GetData()[ColorStart],
+									&Lambdas.GetData()[ColorStart],
+									Dt,
+									reinterpret_cast<const ispc::FVector3f&>(ExpStiffnessValue),
+									reinterpret_cast<const ispc::FVector3f&>(ExpBucklingValue),
+									DampingRatioValue,
+									ColorSize);
+							}
 						}
 						return;
 					}
@@ -673,18 +812,49 @@ void FXPBDAnisotropicBendingConstraints::Apply(SolverParticlesOrRange& Particles
 				{
 					const int32 ColorStart = ConstraintsPerColorStartIndex[ConstraintColorIndex];
 					const int32 ColorSize = ConstraintsPerColorStartIndex[ConstraintColorIndex + 1] - ColorStart;
-					ispc::ApplyXPBDAnisotropicBendingConstraints(
-						(ispc::FVector4f*)Particles.GetPAndInvM().GetData(),
-						(ispc::FIntVector4*)&Constraints.GetData()[ColorStart],
-						&RestAngles.GetData()[ColorStart],
-						&IsBuckled.GetData()[ColorStart],
-						(const ispc::FVector3f*)&WarpWeftBiasBaseMultipliers.GetData()[ColorStart],
-						&Lambdas.GetData()[ColorStart],
-						Dt,
-						MinStiffness,
-						reinterpret_cast<const ispc::FVector3f&>(ExpStiffnessValue),
-						reinterpret_cast<const ispc::FVector3f&>(ExpBucklingValue),
-						ColorSize);
+					if (ColorSize >= Chaos_XPBDBending_ISPC_ParallelBatchSize * Chaos_XPBDBending_ISPC_MinNumParallelBatches)
+					{
+						const int32 NumBatches = FMath::DivideAndRoundUp(ColorSize, Chaos_XPBDBending_ISPC_ParallelBatchSize);
+						PhysicsParallelFor(NumBatches, [this, &Particles, ColorStart, ColorSize,
+							ParallelBatchSize = Chaos_XPBDBending_ISPC_ParallelBatchSize,
+							Dt, &ExpStiffnessValue, &ExpBucklingValue](const int32 BatchIndex)
+						{
+							const int32 BatchStart = BatchIndex * ParallelBatchSize;
+							const int32 BatchEnd = FMath::Min((BatchIndex + 1) * ParallelBatchSize, ColorSize);
+							ispc::ApplyXPBDAnisotropicBendingConstraints(
+								(ispc::FVector4f*)Particles.GetPAndInvM().GetData(),
+								&ConstraintsIndex1.GetData()[ColorStart + BatchStart],
+								&ConstraintsIndex2.GetData()[ColorStart + BatchStart],
+								&ConstraintsIndex3.GetData()[ColorStart + BatchStart],
+								&ConstraintsIndex4.GetData()[ColorStart + BatchStart],
+								&RestAngles.GetData()[ColorStart + BatchStart],
+								&IsBuckled.GetData()[ColorStart + BatchStart],
+								(const ispc::FVector3f*)&WarpWeftBiasBaseMultipliers.GetData()[ColorStart + BatchStart],
+								&Lambdas.GetData()[ColorStart + BatchStart],
+								Dt,
+								reinterpret_cast<const ispc::FVector3f&>(ExpStiffnessValue),
+								reinterpret_cast<const ispc::FVector3f&>(ExpBucklingValue),
+								BatchEnd - BatchStart);
+						});
+					}
+					else
+					{
+						ispc::ApplyXPBDAnisotropicBendingConstraints(
+							(ispc::FVector4f*)Particles.GetPAndInvM().GetData(),
+							&ConstraintsIndex1.GetData()[ColorStart],
+							&ConstraintsIndex2.GetData()[ColorStart],
+							&ConstraintsIndex3.GetData()[ColorStart],
+							&ConstraintsIndex4.GetData()[ColorStart],
+							&RestAngles.GetData()[ColorStart],
+							&IsBuckled.GetData()[ColorStart],
+							(const ispc::FVector3f*)&WarpWeftBiasBaseMultipliers.GetData()[ColorStart],
+							&Lambdas.GetData()[ColorStart],
+							Dt,
+							reinterpret_cast<const ispc::FVector3f&>(ExpStiffnessValue),
+							reinterpret_cast<const ispc::FVector3f&>(ExpBucklingValue),
+							ColorSize);
+
+					}
 				}
 			}
 			else
@@ -700,35 +870,40 @@ void FXPBDAnisotropicBendingConstraints::Apply(SolverParticlesOrRange& Particles
 							const int32 ColorSize = ConstraintsPerColorStartIndex[ConstraintColorIndex + 1] - ColorStart;
 							ispc::ApplyXPBDAnisotropicBendingDampingConstraintsWithMaps(
 								(ispc::FVector4f*)Particles.GetPAndInvM().GetData(),
-								(const ispc::FVector3f*)Particles.XArray().GetData(),
-								(ispc::FIntVector4*)&Constraints.GetData()[ColorStart],
+								(const ispc::FVector3f*)&X1Array.GetData()[ColorStart],
+								(const ispc::FVector3f*)&X2Array.GetData()[ColorStart],
+								(const ispc::FVector3f*)&X3Array.GetData()[ColorStart],
+								(const ispc::FVector3f*)&X4Array.GetData()[ColorStart],
+								&ConstraintsIndex1.GetData()[ColorStart],
+								&ConstraintsIndex2.GetData()[ColorStart],
+								&ConstraintsIndex3.GetData()[ColorStart],
+								&ConstraintsIndex4.GetData()[ColorStart],
 								&RestAngles.GetData()[ColorStart],
 								&IsBuckled.GetData()[ColorStart],
 								(const ispc::FVector3f*)&WarpWeftBiasBaseMultipliers.GetData()[ColorStart],
 								&LambdasDamping.GetData()[ColorStart],
 								Dt,
-								MinStiffness,
 								StiffnessHasWeightMap,
-								StiffnessHasWeightMap ? &Stiffness.GetIndices().GetData()[ColorStart] : nullptr,
-								&Stiffness.GetTable().GetData()[0],
+								reinterpret_cast<const ispc::FVector2f&>(StiffnessWarp.GetOffsetRange()),
+								StiffnessHasWeightMap ? &StiffnessWarp.GetMapValues().GetData()[ColorStart] : nullptr,
 								StiffnessWeftHasWeightMap,
-								StiffnessWeftHasWeightMap ? &StiffnessWeft.GetIndices().GetData()[ColorStart] : nullptr,
-								&StiffnessWeft.GetTable().GetData()[0],
+								reinterpret_cast<const ispc::FVector2f&>(StiffnessWeft.GetOffsetRange()),
+								StiffnessWeftHasWeightMap ? &StiffnessWeft.GetMapValues().GetData()[ColorStart] : nullptr,
 								StiffnessBiasHasWeightMap,
-								StiffnessBiasHasWeightMap ? &StiffnessBias.GetIndices().GetData()[ColorStart] : nullptr,
-								&StiffnessBias.GetTable().GetData()[0],
+								reinterpret_cast<const ispc::FVector2f&>(StiffnessBias.GetOffsetRange()),
+								StiffnessBiasHasWeightMap ? &StiffnessBias.GetMapValues().GetData()[ColorStart] : nullptr,
 								BucklingStiffnessHasWeightMap,
-								BucklingStiffnessHasWeightMap ? &BucklingStiffness.GetIndices().GetData()[ColorStart] : nullptr,
-								&BucklingStiffness.GetTable().GetData()[0],
+								reinterpret_cast<const ispc::FVector2f&>(BucklingStiffnessWarp.GetOffsetRange()),
+								BucklingStiffnessHasWeightMap ? &BucklingStiffnessWarp.GetMapValues().GetData()[ColorStart] : nullptr,
 								BucklingStiffnessWeftHasWeightMap,
-								BucklingStiffnessWeftHasWeightMap ? &BucklingStiffnessWeft.GetIndices().GetData()[ColorStart] : nullptr,
-								&BucklingStiffnessWeft.GetTable().GetData()[0],
+								reinterpret_cast<const ispc::FVector2f&>(BucklingStiffnessWeft.GetOffsetRange()),
+								BucklingStiffnessWeftHasWeightMap ? &BucklingStiffnessWeft.GetMapValues().GetData()[ColorStart] : nullptr,
 								BucklingStiffnessBiasHasWeightMap,
-								BucklingStiffnessBiasHasWeightMap ? &BucklingStiffnessBias.GetIndices().GetData()[ColorStart] : nullptr,
-								&BucklingStiffnessBias.GetTable().GetData()[0],
+								reinterpret_cast<const ispc::FVector2f&>(BucklingStiffnessBias.GetOffsetRange()),
+								BucklingStiffnessBiasHasWeightMap ? &BucklingStiffnessBias.GetMapValues().GetData()[ColorStart] : nullptr,
 								DampingHasWeightMap,
-								DampingHasWeightMap ? &DampingRatio.GetIndices().GetData()[ColorStart] : nullptr,
-								&DampingRatio.GetTable().GetData()[0],
+								reinterpret_cast<const ispc::FVector2f&>(DampingRatio.GetOffsetRange()),
+								DampingHasWeightMap ? &DampingRatio.GetMapValues()[ColorStart] : nullptr,
 								ColorSize);
 						}
 					}
@@ -740,35 +915,40 @@ void FXPBDAnisotropicBendingConstraints::Apply(SolverParticlesOrRange& Particles
 							const int32 ColorSize = ConstraintsPerColorStartIndex[ConstraintColorIndex + 1] - ColorStart;
 							ispc::ApplyXPBDAnisotropicBendingConstraintsWithDampingAndMaps(
 								(ispc::FVector4f*)Particles.GetPAndInvM().GetData(),
-								(const ispc::FVector3f*)Particles.XArray().GetData(),
-								(ispc::FIntVector4*)&Constraints.GetData()[ColorStart],
+								(const ispc::FVector3f*)&X1Array.GetData()[ColorStart],
+								(const ispc::FVector3f*)&X2Array.GetData()[ColorStart],
+								(const ispc::FVector3f*)&X3Array.GetData()[ColorStart],
+								(const ispc::FVector3f*)&X4Array.GetData()[ColorStart],
+								&ConstraintsIndex1.GetData()[ColorStart],
+								&ConstraintsIndex2.GetData()[ColorStart],
+								&ConstraintsIndex3.GetData()[ColorStart],
+								&ConstraintsIndex4.GetData()[ColorStart],
 								&RestAngles.GetData()[ColorStart],
 								&IsBuckled.GetData()[ColorStart],
 								(const ispc::FVector3f*)&WarpWeftBiasBaseMultipliers.GetData()[ColorStart],
 								&Lambdas.GetData()[ColorStart],
 								Dt,
-								MinStiffness,
 								StiffnessHasWeightMap,
-								StiffnessHasWeightMap ? &Stiffness.GetIndices().GetData()[ColorStart] : nullptr,
-								&Stiffness.GetTable().GetData()[0],
+								reinterpret_cast<const ispc::FVector2f&>(StiffnessWarp.GetOffsetRange()),
+								StiffnessHasWeightMap ? &StiffnessWarp.GetMapValues().GetData()[ColorStart] : nullptr,
 								StiffnessWeftHasWeightMap,
-								StiffnessWeftHasWeightMap ? &StiffnessWeft.GetIndices().GetData()[ColorStart] : nullptr,
-								&StiffnessWeft.GetTable().GetData()[0],
+								reinterpret_cast<const ispc::FVector2f&>(StiffnessWeft.GetOffsetRange()),
+								StiffnessWeftHasWeightMap ? &StiffnessWeft.GetMapValues().GetData()[ColorStart] : nullptr,
 								StiffnessBiasHasWeightMap,
-								StiffnessBiasHasWeightMap ? &StiffnessBias.GetIndices().GetData()[ColorStart] : nullptr,
-								&StiffnessBias.GetTable().GetData()[0],
+								reinterpret_cast<const ispc::FVector2f&>(StiffnessBias.GetOffsetRange()),
+								StiffnessBiasHasWeightMap ? &StiffnessBias.GetMapValues().GetData()[ColorStart] : nullptr,
 								BucklingStiffnessHasWeightMap,
-								BucklingStiffnessHasWeightMap ? &BucklingStiffness.GetIndices().GetData()[ColorStart] : nullptr,
-								&BucklingStiffness.GetTable().GetData()[0],
+								reinterpret_cast<const ispc::FVector2f&>(BucklingStiffnessWarp.GetOffsetRange()),
+								BucklingStiffnessHasWeightMap ? &BucklingStiffnessWarp.GetMapValues().GetData()[ColorStart] : nullptr,
 								BucklingStiffnessWeftHasWeightMap,
-								BucklingStiffnessWeftHasWeightMap ? &BucklingStiffnessWeft.GetIndices().GetData()[ColorStart] : nullptr,
-								&BucklingStiffnessWeft.GetTable().GetData()[0],
+								reinterpret_cast<const ispc::FVector2f&>(BucklingStiffnessWeft.GetOffsetRange()),
+								BucklingStiffnessWeftHasWeightMap ? &BucklingStiffnessWeft.GetMapValues().GetData()[ColorStart] : nullptr,
 								BucklingStiffnessBiasHasWeightMap,
-								BucklingStiffnessBiasHasWeightMap ? &BucklingStiffnessBias.GetIndices().GetData()[ColorStart] : nullptr,
-								&BucklingStiffnessBias.GetTable().GetData()[0],
+								reinterpret_cast<const ispc::FVector2f&>(BucklingStiffnessBias.GetOffsetRange()),
+								BucklingStiffnessBiasHasWeightMap ? &BucklingStiffnessBias.GetMapValues().GetData()[ColorStart] : nullptr,
 								DampingHasWeightMap,
-								DampingHasWeightMap ? &DampingRatio.GetIndices().GetData()[ColorStart] : nullptr,
-								&DampingRatio.GetTable().GetData()[0],
+								reinterpret_cast<const ispc::FVector2f&>(DampingRatio.GetOffsetRange()),
+								DampingHasWeightMap ? &DampingRatio.GetMapValues()[ColorStart] : nullptr,
 								ColorSize);
 						}
 						return;
@@ -780,31 +960,33 @@ void FXPBDAnisotropicBendingConstraints::Apply(SolverParticlesOrRange& Particles
 					const int32 ColorSize = ConstraintsPerColorStartIndex[ConstraintColorIndex + 1] - ColorStart;
 					ispc::ApplyXPBDAnisotropicBendingConstraintsWithMaps(
 						(ispc::FVector4f*)Particles.GetPAndInvM().GetData(),
-						(ispc::FIntVector4*)&Constraints.GetData()[ColorStart],
+						&ConstraintsIndex1.GetData()[ColorStart],
+						&ConstraintsIndex2.GetData()[ColorStart],
+						&ConstraintsIndex3.GetData()[ColorStart],
+						&ConstraintsIndex4.GetData()[ColorStart],
 						&RestAngles.GetData()[ColorStart],
 						&IsBuckled.GetData()[ColorStart],
 						(const ispc::FVector3f*)&WarpWeftBiasBaseMultipliers.GetData()[ColorStart],
 						&Lambdas.GetData()[ColorStart],
 						Dt,
-						MinStiffness,
 						StiffnessHasWeightMap,
-						StiffnessHasWeightMap ? &Stiffness.GetIndices().GetData()[ColorStart] : nullptr,
-						&Stiffness.GetTable().GetData()[0],
+						reinterpret_cast<const ispc::FVector2f&>(StiffnessWarp.GetOffsetRange()),
+						StiffnessHasWeightMap ? &StiffnessWarp.GetMapValues().GetData()[ColorStart] : nullptr,
 						StiffnessWeftHasWeightMap,
-						StiffnessWeftHasWeightMap ? &StiffnessWeft.GetIndices().GetData()[ColorStart] : nullptr,
-						&StiffnessWeft.GetTable().GetData()[0],
+						reinterpret_cast<const ispc::FVector2f&>(StiffnessWeft.GetOffsetRange()),
+						StiffnessWeftHasWeightMap ? &StiffnessWeft.GetMapValues().GetData()[ColorStart] : nullptr,
 						StiffnessBiasHasWeightMap,
-						StiffnessBiasHasWeightMap ? &StiffnessBias.GetIndices().GetData()[ColorStart] : nullptr,
-						&StiffnessBias.GetTable().GetData()[0],
+						reinterpret_cast<const ispc::FVector2f&>(StiffnessBias.GetOffsetRange()),
+						StiffnessBiasHasWeightMap ? &StiffnessBias.GetMapValues().GetData()[ColorStart] : nullptr,
 						BucklingStiffnessHasWeightMap,
-						BucklingStiffnessHasWeightMap ? &BucklingStiffness.GetIndices().GetData()[ColorStart] : nullptr,
-						&BucklingStiffness.GetTable().GetData()[0],
+						reinterpret_cast<const ispc::FVector2f&>(BucklingStiffnessWarp.GetOffsetRange()),
+						BucklingStiffnessHasWeightMap ? &BucklingStiffnessWarp.GetMapValues().GetData()[ColorStart] : nullptr,
 						BucklingStiffnessWeftHasWeightMap,
-						BucklingStiffnessWeftHasWeightMap ? &BucklingStiffnessWeft.GetIndices().GetData()[ColorStart] : nullptr,
-						&BucklingStiffnessWeft.GetTable().GetData()[0],
+						reinterpret_cast<const ispc::FVector2f&>(BucklingStiffnessWeft.GetOffsetRange()),
+						BucklingStiffnessWeftHasWeightMap ? &BucklingStiffnessWeft.GetMapValues().GetData()[ColorStart] : nullptr,
 						BucklingStiffnessBiasHasWeightMap,
-						BucklingStiffnessBiasHasWeightMap ? &BucklingStiffnessBias.GetIndices().GetData()[ColorStart] : nullptr,
-						&BucklingStiffnessBias.GetTable().GetData()[0],
+						reinterpret_cast<const ispc::FVector2f&>(BucklingStiffnessBias.GetOffsetRange()),
+						BucklingStiffnessBiasHasWeightMap ? &BucklingStiffnessBias.GetMapValues().GetData()[ColorStart] : nullptr,
 						ColorSize);
 				}
 			}
@@ -813,10 +995,10 @@ void FXPBDAnisotropicBendingConstraints::Apply(SolverParticlesOrRange& Particles
 #endif
 		{
 			// Parallel non-ispc
-			const FSolverReal StiffnessNoMap = (FSolverReal)Stiffness;
+			const FSolverReal StiffnessNoMap = (FSolverReal)StiffnessWarp;
 			const FSolverReal StiffnessWeftNoMap = (FSolverReal)StiffnessWeft;
 			const FSolverReal StiffnessBiasNoMap = (FSolverReal)StiffnessBias;
-			const FSolverReal BucklingStiffnessNoMap = (FSolverReal)BucklingStiffness;
+			const FSolverReal BucklingStiffnessNoMap = (FSolverReal)BucklingStiffnessWarp;
 			const FSolverReal BucklingStiffnessWeftNoMap = (FSolverReal)BucklingStiffnessWeft;
 			const FSolverReal BucklingStiffnessBiasNoMap = (FSolverReal)BucklingStiffnessBias;
 			const FSolverReal DampingNoMap = (FSolverReal)DampingRatio;
@@ -832,10 +1014,10 @@ void FXPBDAnisotropicBendingConstraints::Apply(SolverParticlesOrRange& Particles
 						PhysicsParallelFor(ColorSize, [&](const int32 Index)
 						{
 							const int32 ConstraintIndex = ColorStart + Index;
-							const FSolverReal ExpStiffnessValue = StiffnessHasWeightMap ? Stiffness[ConstraintIndex] : StiffnessNoMap;
+							const FSolverReal ExpStiffnessValue = StiffnessHasWeightMap ? StiffnessWarp[ConstraintIndex] : StiffnessNoMap;
 							const FSolverReal ExpStiffnessWeftValue = StiffnessWeftHasWeightMap ? StiffnessWeft[ConstraintIndex] : StiffnessWeftNoMap;
 							const FSolverReal ExpStiffnessBiasValue = StiffnessBiasHasWeightMap ? StiffnessBias[ConstraintIndex] : StiffnessBiasNoMap;
-							const FSolverReal ExpBucklingValue = BucklingStiffnessHasWeightMap ? BucklingStiffness[ConstraintIndex] : BucklingStiffnessNoMap;
+							const FSolverReal ExpBucklingValue = BucklingStiffnessHasWeightMap ? BucklingStiffnessWarp[ConstraintIndex] : BucklingStiffnessNoMap;
 							const FSolverReal ExpBucklingWeftValue = BucklingStiffnessWeftHasWeightMap ? BucklingStiffnessWeft[ConstraintIndex] : BucklingStiffnessWeftNoMap;
 							const FSolverReal ExpBucklingBiasValue = BucklingStiffnessBiasHasWeightMap ? BucklingStiffnessBias[ConstraintIndex] : BucklingStiffnessBiasNoMap;
 							const FSolverReal DampingRatioValue = DampingHasWeightMap ? DampingRatio[ConstraintIndex] : DampingNoMap;
@@ -855,10 +1037,10 @@ void FXPBDAnisotropicBendingConstraints::Apply(SolverParticlesOrRange& Particles
 						PhysicsParallelFor(ColorSize, [&](const int32 Index)
 						{
 							const int32 ConstraintIndex = ColorStart + Index;
-							const FSolverReal ExpStiffnessValue = StiffnessHasWeightMap ? Stiffness[ConstraintIndex] : StiffnessNoMap;
+							const FSolverReal ExpStiffnessValue = StiffnessHasWeightMap ? StiffnessWarp[ConstraintIndex] : StiffnessNoMap;
 							const FSolverReal ExpStiffnessWeftValue = StiffnessWeftHasWeightMap ? StiffnessWeft[ConstraintIndex] : StiffnessWeftNoMap;
 							const FSolverReal ExpStiffnessBiasValue = StiffnessBiasHasWeightMap ? StiffnessBias[ConstraintIndex] : StiffnessBiasNoMap;
-							const FSolverReal ExpBucklingValue = BucklingStiffnessHasWeightMap ? BucklingStiffness[ConstraintIndex] : BucklingStiffnessNoMap;
+							const FSolverReal ExpBucklingValue = BucklingStiffnessHasWeightMap ? BucklingStiffnessWarp[ConstraintIndex] : BucklingStiffnessNoMap;
 							const FSolverReal ExpBucklingWeftValue = BucklingStiffnessWeftHasWeightMap ? BucklingStiffnessWeft[ConstraintIndex] : BucklingStiffnessWeftNoMap;
 							const FSolverReal ExpBucklingBiasValue = BucklingStiffnessBiasHasWeightMap ? BucklingStiffnessBias[ConstraintIndex] : BucklingStiffnessBiasNoMap;
 							const FSolverReal DampingRatioValue = DampingHasWeightMap ? DampingRatio[ConstraintIndex] : DampingNoMap;
@@ -878,10 +1060,10 @@ void FXPBDAnisotropicBendingConstraints::Apply(SolverParticlesOrRange& Particles
 					PhysicsParallelFor(ColorSize, [&](const int32 Index)
 					{
 						const int32 ConstraintIndex = ColorStart + Index;
-						const FSolverReal ExpStiffnessValue = StiffnessHasWeightMap ? Stiffness[ConstraintIndex] : StiffnessNoMap;
+						const FSolverReal ExpStiffnessValue = StiffnessHasWeightMap ? StiffnessWarp[ConstraintIndex] : StiffnessNoMap;
 						const FSolverReal ExpStiffnessWeftValue = StiffnessWeftHasWeightMap ? StiffnessWeft[ConstraintIndex] : StiffnessWeftNoMap;
 						const FSolverReal ExpStiffnessBiasValue = StiffnessBiasHasWeightMap ? StiffnessBias[ConstraintIndex] : StiffnessBiasNoMap;
-						const FSolverReal ExpBucklingValue = BucklingStiffnessHasWeightMap ? BucklingStiffness[ConstraintIndex] : BucklingStiffnessNoMap;
+						const FSolverReal ExpBucklingValue = BucklingStiffnessHasWeightMap ? BucklingStiffnessWarp[ConstraintIndex] : BucklingStiffnessNoMap;
 						const FSolverReal ExpBucklingWeftValue = BucklingStiffnessWeftHasWeightMap ? BucklingStiffnessWeft[ConstraintIndex] : BucklingStiffnessWeftNoMap;
 						const FSolverReal ExpBucklingBiasValue = BucklingStiffnessBiasHasWeightMap ? BucklingStiffnessBias[ConstraintIndex] : BucklingStiffnessBiasNoMap;
 						const FSolverReal DampingRatioValue = DampingHasWeightMap ? DampingRatio[ConstraintIndex] : DampingNoMap;
@@ -910,10 +1092,10 @@ void FXPBDAnisotropicBendingConstraints::Apply(SolverParticlesOrRange& Particles
 			{
 				for (int32 ConstraintIndex = 0; ConstraintIndex < Constraints.Num(); ++ConstraintIndex)
 				{
-					const FSolverReal ExpStiffnessValue = StiffnessHasWeightMap ? Stiffness[ConstraintIndex] : StiffnessNoMap;
+					const FSolverReal ExpStiffnessValue = StiffnessHasWeightMap ? StiffnessWarp[ConstraintIndex] : StiffnessNoMap;
 					const FSolverReal ExpStiffnessWeftValue = StiffnessWeftHasWeightMap ? StiffnessWeft[ConstraintIndex] : StiffnessWeftNoMap;
 					const FSolverReal ExpStiffnessBiasValue = StiffnessBiasHasWeightMap ? StiffnessBias[ConstraintIndex] : StiffnessBiasNoMap;
-					const FSolverReal ExpBucklingValue = BucklingStiffnessHasWeightMap ? BucklingStiffness[ConstraintIndex] : BucklingStiffnessNoMap;
+					const FSolverReal ExpBucklingValue = BucklingStiffnessHasWeightMap ? BucklingStiffnessWarp[ConstraintIndex] : BucklingStiffnessNoMap;
 					const FSolverReal ExpBucklingWeftValue = BucklingStiffnessWeftHasWeightMap ? BucklingStiffnessWeft[ConstraintIndex] : BucklingStiffnessWeftNoMap;
 					const FSolverReal ExpBucklingBiasValue = BucklingStiffnessBiasHasWeightMap ? BucklingStiffnessBias[ConstraintIndex] : BucklingStiffnessBiasNoMap;
 					const FSolverReal DampingRatioValue = DampingHasWeightMap ? DampingRatio[ConstraintIndex] : DampingNoMap;
@@ -927,10 +1109,10 @@ void FXPBDAnisotropicBendingConstraints::Apply(SolverParticlesOrRange& Particles
 			{
 				for (int32 ConstraintIndex = 0; ConstraintIndex < Constraints.Num(); ++ConstraintIndex)
 				{
-					const FSolverReal ExpStiffnessValue = StiffnessHasWeightMap ? Stiffness[ConstraintIndex] : StiffnessNoMap;
+					const FSolverReal ExpStiffnessValue = StiffnessHasWeightMap ? StiffnessWarp[ConstraintIndex] : StiffnessNoMap;
 					const FSolverReal ExpStiffnessWeftValue = StiffnessWeftHasWeightMap ? StiffnessWeft[ConstraintIndex] : StiffnessWeftNoMap;
 					const FSolverReal ExpStiffnessBiasValue = StiffnessBiasHasWeightMap ? StiffnessBias[ConstraintIndex] : StiffnessBiasNoMap;
-					const FSolverReal ExpBucklingValue = BucklingStiffnessHasWeightMap ? BucklingStiffness[ConstraintIndex] : BucklingStiffnessNoMap;
+					const FSolverReal ExpBucklingValue = BucklingStiffnessHasWeightMap ? BucklingStiffnessWarp[ConstraintIndex] : BucklingStiffnessNoMap;
 					const FSolverReal ExpBucklingWeftValue = BucklingStiffnessWeftHasWeightMap ? BucklingStiffnessWeft[ConstraintIndex] : BucklingStiffnessWeftNoMap;
 					const FSolverReal ExpBucklingBiasValue = BucklingStiffnessBiasHasWeightMap ? BucklingStiffnessBias[ConstraintIndex] : BucklingStiffnessBiasNoMap;
 					const FSolverReal DampingRatioValue = DampingHasWeightMap ? DampingRatio[ConstraintIndex] : DampingNoMap;
@@ -945,10 +1127,10 @@ void FXPBDAnisotropicBendingConstraints::Apply(SolverParticlesOrRange& Particles
 
 		for (int32 ConstraintIndex = 0; ConstraintIndex < Constraints.Num(); ++ConstraintIndex)
 		{
-			const FSolverReal ExpStiffnessValue = StiffnessHasWeightMap ? Stiffness[ConstraintIndex] : StiffnessNoMap;
+			const FSolverReal ExpStiffnessValue = StiffnessHasWeightMap ? StiffnessWarp[ConstraintIndex] : StiffnessNoMap;
 			const FSolverReal ExpStiffnessWeftValue = StiffnessWeftHasWeightMap ? StiffnessWeft[ConstraintIndex] : StiffnessWeftNoMap;
 			const FSolverReal ExpStiffnessBiasValue = StiffnessBiasHasWeightMap ? StiffnessBias[ConstraintIndex] : StiffnessBiasNoMap;
-			const FSolverReal ExpBucklingValue = BucklingStiffnessHasWeightMap ? BucklingStiffness[ConstraintIndex] : BucklingStiffnessNoMap;
+			const FSolverReal ExpBucklingValue = BucklingStiffnessHasWeightMap ? BucklingStiffnessWarp[ConstraintIndex] : BucklingStiffnessNoMap;
 			const FSolverReal ExpBucklingWeftValue = BucklingStiffnessWeftHasWeightMap ? BucklingStiffnessWeft[ConstraintIndex] : BucklingStiffnessWeftNoMap;
 			const FSolverReal ExpBucklingBiasValue = BucklingStiffnessBiasHasWeightMap ? BucklingStiffnessBias[ConstraintIndex] : BucklingStiffnessBiasNoMap;
 			const FSolverReal DampingRatioValue = DampingHasWeightMap ? DampingRatio[ConstraintIndex] : DampingNoMap;
@@ -1000,8 +1182,8 @@ void FXPBDAnisotropicBendingConstraints::ComputeGradTheta(const FSolverVec3& X0,
 
 void FXPBDAnisotropicBendingConstraints::AddAnisotropicBendingResidualAndHessian(const FSolverParticles& Particles, const int32 ConstraintIndex, const int32 ConstraintIndexLocal, const FSolverReal Dt, TVec3<FSolverReal>& ParticleResidual, Chaos::PMatrix<FSolverReal, 3, 3>& ParticleHessian)
 {
-	const FSolverVec3 ExpStiffnessValue((FSolverReal)StiffnessWeft, (FSolverReal)Stiffness, (FSolverReal)StiffnessBias);
-	const FSolverVec3 ExpBucklingValue((FSolverReal)BucklingStiffnessWeft, (FSolverReal)BucklingStiffness, (FSolverReal)BucklingStiffnessBias);
+	const FSolverVec3 ExpStiffnessValue((FSolverReal)StiffnessWeft, (FSolverReal)StiffnessWarp, (FSolverReal)StiffnessBias);
+	const FSolverVec3 ExpBucklingValue((FSolverReal)BucklingStiffnessWeft, (FSolverReal)BucklingStiffnessWarp, (FSolverReal)BucklingStiffnessBias);
 	const FSolverReal DampingRatioValue = (FSolverReal)DampingRatio;
 	const FSolverVec3& WarpWeftBiasBaseMultiplier = WarpWeftBiasBaseMultipliers[ConstraintIndex];
 
@@ -1041,8 +1223,8 @@ void FXPBDAnisotropicBendingConstraints::AddAnisotropicBendingResidualAndHessian
 
 void FXPBDAnisotropicBendingConstraints::AddInternalForceDifferential(const FSolverParticles& InParticles, const TArray<TVector<FSolverReal, 3>>& DeltaParticles, TArray<TVector<FSolverReal, 3>>& ndf)
 {
-	const FSolverVec3 ExpStiffnessValue((FSolverReal)StiffnessWeft, (FSolverReal)Stiffness, (FSolverReal)StiffnessBias);
-	const FSolverVec3 ExpBucklingValue((FSolverReal)BucklingStiffnessWeft, (FSolverReal)BucklingStiffness, (FSolverReal)BucklingStiffnessBias);
+	const FSolverVec3 ExpStiffnessValue((FSolverReal)StiffnessWeft, (FSolverReal)StiffnessWarp, (FSolverReal)StiffnessBias);
+	const FSolverVec3 ExpBucklingValue((FSolverReal)BucklingStiffnessWeft, (FSolverReal)BucklingStiffnessWarp, (FSolverReal)BucklingStiffnessBias);
 	const FSolverReal DampingRatioValue = (FSolverReal)DampingRatio;
 
 	int32 ParticleStart = 0;
