Chaos Cloth optimizations to XPBD constraints.

diff --git a/Engine/Source/Runtime/Experimental/Chaos/Private/Chaos/XPBDAnisotropicBendingConstraints.ispc b/Engine/Source/Runtime/Experimental/Chaos/Private/Chaos/XPBDAnisotropicBendingConstraints.ispc
--- a/Engine/Source/Runtime/Experimental/Chaos/Private/Chaos/XPBDAnisotropicBendingConstraints.ispc
+++ b/Engine/Source/Runtime/Experimental/Chaos/Private/Chaos/XPBDAnisotropicBendingConstraints.ispc
@@ -7,13 +7,15 @@
 #include "Chaos/XPBDBendingElement.isph"
 
 export void ApplyXPBDAnisotropicBendingConstraints(uniform FVector4f PandInvM[],
-									const uniform FIntVector4 Constraints[],
+									const uniform int32 ConstraintIndex1[],
+									const uniform int32 ConstraintIndex2[],
+									const uniform int32 ConstraintIndex3[],
+									const uniform int32 ConstraintIndex4[],
 									const uniform float RestAngles[],
 									const uniform bool IsBuckled[],
 									const uniform FVector3f WarpWeftBiasBaseMultiplers[],
 									uniform float Lambdas[],
 									const uniform float Dt,
-									const uniform float MinStiffness,
 									const uniform FVector3f& Stiffness,
 									const uniform FVector3f& BucklingStiffness,
 									const uniform int32 NumConstraints)
@@ -21,64 +23,54 @@ export void ApplyXPBDAnisotropicBendingConstraints(uniform FVector4f PandInvM[],
 
 	foreach(i = 0 ... NumConstraints)
 	{
-		const varying FIntVector4 Constraint = VectorLoad(&Constraints[extract(i,0)]);
+		const varying int32 i1 = ConstraintIndex1[i];
+		const varying int32 i2 = ConstraintIndex2[i];
+		const varying int32 i3 = ConstraintIndex3[i];
+		const varying int32 i4 = ConstraintIndex4[i];
 
-		const varying int32 i1 = Constraint.V[0];
-		const varying int32 i2 = Constraint.V[1];
-		const varying int32 i3 = Constraint.V[2];
-		const varying int32 i4 = Constraint.V[3];
-		
 		const varying FVector4f PandInvM1 = VectorGather(&PandInvM[i1]);
 		const varying FVector4f PandInvM2 = VectorGather(&PandInvM[i2]);
 		const varying FVector4f PandInvM3 = VectorGather(&PandInvM[i3]);
 		const varying FVector4f PandInvM4 = VectorGather(&PandInvM[i4]);
 
 		varying FVector3f P1, P2, P3, P4;
-		varying float M1, M2, M3, M4;
-
-		UnzipPandInvM(PandInvM1, P1, M1);
-		UnzipPandInvM(PandInvM2, P2, M2);
-		UnzipPandInvM(PandInvM3, P3, M3);
-		UnzipPandInvM(PandInvM4, P4, M4);
+		varying float InvM1, InvM2, InvM3, InvM4;
 
-		const varying FVector3f WarpWeftBiasBaseMultiplier = VectorGather(&WarpWeftBiasBaseMultiplers[i]);
+		UnzipPandInvM(PandInvM1, P1, InvM1);
+		UnzipPandInvM(PandInvM2, P2, InvM2);
+		UnzipPandInvM(PandInvM3, P3, InvM3);
+		UnzipPandInvM(PandInvM4, P4, InvM4);
+		
+		const varying FVector3f WarpWeftBiasBaseMultiplier = VectorLoad(&WarpWeftBiasBaseMultiplers[extract(i,0)]);
 		const varying float StiffnessValue = select(IsBuckled[i], VectorDot(WarpWeftBiasBaseMultiplier, BucklingStiffness), VectorDot(WarpWeftBiasBaseMultiplier, Stiffness));
 		
 		const varying float Lambda = Lambdas[i];
 
 		varying FVector3f Delta1, Delta2, Delta3, Delta4;
-		const varying float DLambda = GetXPBDBendingElementDeltas(P1, P2, P3, P4, M1, M2, M3, M4, Dt, RestAngles[i], Lambda, StiffnessValue, MinStiffness, Delta1, Delta2, Delta3, Delta4);
+		const varying float DLambda = GetXPBDBendingElementDeltas(P1, P2, P3, P4, InvM1, InvM2, InvM3, InvM4, Dt, RestAngles[i], Lambda, StiffnessValue, Delta1, Delta2, Delta3, Delta4);
 
 		Lambdas[i] = Lambda + DLambda;
-
-		if (M1 > 0)
-		{
-			VectorScatter(&PandInvM[i1], SetVector4(P1 + Delta1, M1));
-		}
-		if (M2 > 0)
-		{
-			VectorScatter(&PandInvM[i2], SetVector4(P2 + Delta2, M2));
-		}
-		if (M3 > 0)
-		{
-			VectorScatter(&PandInvM[i3], SetVector4(P3 + Delta3, M3));
-		}
-		if (M4 > 0)
-		{
-			VectorScatter(&PandInvM[i4], SetVector4(P4 + Delta4, M4));
-		}
+		VectorScatter(&PandInvM[i1], SetVector4(P1 + Delta1, InvM1));
+		VectorScatter(&PandInvM[i2], SetVector4(P2 + Delta2, InvM2));
+		VectorScatter(&PandInvM[i3], SetVector4(P3 + Delta3, InvM3));
+		VectorScatter(&PandInvM[i4], SetVector4(P4 + Delta4, InvM4));
 	}
 }
 
 export void ApplyXPBDAnisotropicBendingConstraintsWithDamping(uniform FVector4f PandInvM[],
-									const uniform FVector3f X[],
-									const uniform FIntVector4 Constraints[],
+									const uniform FVector3f X1Array[],
+									const uniform FVector3f X2Array[],
+									const uniform FVector3f X3Array[],
+									const uniform FVector3f X4Array[],
+									const uniform int32 ConstraintIndex1[],
+									const uniform int32 ConstraintIndex2[],
+									const uniform int32 ConstraintIndex3[],
+									const uniform int32 ConstraintIndex4[],
 									const uniform float RestAngles[],
 									const uniform bool IsBuckled[],
 									const uniform FVector3f WarpWeftBiasBaseMultiplers[],
 									uniform float Lambdas[],
 									const uniform float Dt,
-									const uniform float MinStiffness,
 									const uniform FVector3f& Stiffness,
 									const uniform FVector3f& BucklingStiffness,
 									const uniform float DampingRatio,
@@ -87,12 +79,10 @@ export void ApplyXPBDAnisotropicBendingConstraintsWithDamping(uniform FVector4f
 
 	foreach(i = 0 ... NumConstraints)
 	{
-		const varying FIntVector4 Constraint = VectorLoad(&Constraints[extract(i,0)]);
-
-		const varying int32 i1 = Constraint.V[0];
-		const varying int32 i2 = Constraint.V[1];
-		const varying int32 i3 = Constraint.V[2];
-		const varying int32 i4 = Constraint.V[3];
+		const varying int32 i1 = ConstraintIndex1[i];
+		const varying int32 i2 = ConstraintIndex2[i];
+		const varying int32 i3 = ConstraintIndex3[i];
+		const varying int32 i4 = ConstraintIndex4[i];
 		
 		const varying FVector4f PandInvM1 = VectorGather(&PandInvM[i1]);
 		const varying FVector4f PandInvM2 = VectorGather(&PandInvM[i2]);
@@ -101,54 +91,44 @@ export void ApplyXPBDAnisotropicBendingConstraintsWithDamping(uniform FVector4f
 
 		varying FVector3f P1, P2, P3, P4;
 		varying float M1, M2, M3, M4;
+		
+		const uniform int32 i0 = extract(i,0);
 
-		UnzipPandInvM(PandInvM1, P1, M1);
-		UnzipPandInvM(PandInvM2, P2, M2);
-		UnzipPandInvM(PandInvM3, P3, M3);
-		UnzipPandInvM(PandInvM4, P4, M4);
-
-		const varying FVector3f X1 = VectorGather(&X[i1]);
-		const varying FVector3f X2 = VectorGather(&X[i2]);
-		const varying FVector3f X3 = VectorGather(&X[i3]);
-		const varying FVector3f X4 = VectorGather(&X[i4]);
-
-		const varying FVector3f WarpWeftBiasBaseMultiplier = VectorGather(&WarpWeftBiasBaseMultiplers[i]);
+		const varying FVector3f X1 = VectorLoad(&X1Array[i0]);
+		const varying FVector3f X2 = VectorLoad(&X2Array[i0]);
+		const varying FVector3f X3 = VectorLoad(&X3Array[i0]);
+		const varying FVector3f X4 = VectorLoad(&X4Array[i0]);
+		
+		const varying FVector3f WarpWeftBiasBaseMultiplier = VectorLoad(&WarpWeftBiasBaseMultiplers[i0]);
 		const varying float StiffnessValue = select(IsBuckled[i], VectorDot(WarpWeftBiasBaseMultiplier, BucklingStiffness), VectorDot(WarpWeftBiasBaseMultiplier, Stiffness));
 		const varying float Lambda = Lambdas[i];
 
 		varying FVector3f Delta1, Delta2, Delta3, Delta4;
-		const varying float DLambda = GetXPBDBendingElementDeltasWithDamping(P1, P2, P3, P4, X1, X2, X3, X4, M1, M2, M3, M4, Dt, RestAngles[i], Lambda, StiffnessValue, DampingRatio, MinStiffness, Delta1, Delta2, Delta3, Delta4);
+		const varying float DLambda = GetXPBDBendingElementDeltasWithDamping(P1, P2, P3, P4, X1, X2, X3, X4, M1, M2, M3, M4, Dt, RestAngles[i], Lambda, StiffnessValue, DampingRatio, Delta1, Delta2, Delta3, Delta4);
 
 		Lambdas[i] = Lambda + DLambda;
 
-		if (M1 > 0)
-		{
-			VectorScatter(&PandInvM[i1], SetVector4(P1 + Delta1, M1));
-		}
-		if (M2 > 0)
-		{
-			VectorScatter(&PandInvM[i2], SetVector4(P2 + Delta2, M2));
-		}
-		if (M3 > 0)
-		{
-			VectorScatter(&PandInvM[i3], SetVector4(P3 + Delta3, M3));
-		}
-		if (M4 > 0)
-		{
-			VectorScatter(&PandInvM[i4], SetVector4(P4 + Delta4, M4));
-		}
+		VectorScatter(&PandInvM[i1], SetVector4(P1 + Delta1, M1));
+		VectorScatter(&PandInvM[i2], SetVector4(P2 + Delta2, M2));
+		VectorScatter(&PandInvM[i3], SetVector4(P3 + Delta3, M3));
+		VectorScatter(&PandInvM[i4], SetVector4(P4 + Delta4, M4));
 	}
 }
 
 export void ApplyXPBDAnisotropicBendingDampingConstraints(uniform FVector4f PandInvM[],
-									const uniform FVector3f X[],
-									const uniform FIntVector4 Constraints[],
+									const uniform FVector3f X1Array[],
+									const uniform FVector3f X2Array[],
+									const uniform FVector3f X3Array[],
+									const uniform FVector3f X4Array[],
+									const uniform int32 ConstraintIndex1[],
+									const uniform int32 ConstraintIndex2[],
+									const uniform int32 ConstraintIndex3[],
+									const uniform int32 ConstraintIndex4[],
 									const uniform float RestAngles[],
 									const uniform bool IsBuckled[],
 									const uniform FVector3f WarpWeftBiasBaseMultiplers[],
 									uniform float Lambdas[],
 									const uniform float Dt,
-									const uniform float MinStiffness,
 									const uniform FVector3f& Stiffness,
 									const uniform FVector3f& BucklingStiffness,
 									const uniform float DampingRatio,
@@ -157,12 +137,10 @@ export void ApplyXPBDAnisotropicBendingDampingConstraints(uniform FVector4f Pand
 
 	foreach(i = 0 ... NumConstraints)
 	{
-		const varying FIntVector4 Constraint = VectorLoad(&Constraints[extract(i,0)]);
-
-		const varying int32 i1 = Constraint.V[0];
-		const varying int32 i2 = Constraint.V[1];
-		const varying int32 i3 = Constraint.V[2];
-		const varying int32 i4 = Constraint.V[3];
+		const varying int32 i1 = ConstraintIndex1[i];
+		const varying int32 i2 = ConstraintIndex2[i];
+		const varying int32 i3 = ConstraintIndex3[i];
+		const varying int32 i4 = ConstraintIndex4[i];
 		
 		const varying FVector4f PandInvM1 = VectorGather(&PandInvM[i1]);
 		const varying FVector4f PandInvM2 = VectorGather(&PandInvM[i2]);
@@ -176,83 +154,66 @@ export void ApplyXPBDAnisotropicBendingDampingConstraints(uniform FVector4f Pand
 		UnzipPandInvM(PandInvM2, P2, M2);
 		UnzipPandInvM(PandInvM3, P3, M3);
 		UnzipPandInvM(PandInvM4, P4, M4);
+		
+		const uniform int32 i0 = extract(i,0);
 
-		const varying FVector3f X1 = VectorGather(&X[i1]);
-		const varying FVector3f X2 = VectorGather(&X[i2]);
-		const varying FVector3f X3 = VectorGather(&X[i3]);
-		const varying FVector3f X4 = VectorGather(&X[i4]);
-
-		const varying FVector3f WarpWeftBiasBaseMultiplier = VectorGather(&WarpWeftBiasBaseMultiplers[i]);
+		const varying FVector3f X1 = VectorLoad(&X1Array[i0]);
+		const varying FVector3f X2 = VectorLoad(&X2Array[i0]);
+		const varying FVector3f X3 = VectorLoad(&X3Array[i0]);
+		const varying FVector3f X4 = VectorLoad(&X4Array[i0]);
+		
+		const varying FVector3f WarpWeftBiasBaseMultiplier = VectorLoad(&WarpWeftBiasBaseMultiplers[i0]);
 		const varying float StiffnessValue = select(IsBuckled[i], VectorDot(WarpWeftBiasBaseMultiplier, BucklingStiffness), VectorDot(WarpWeftBiasBaseMultiplier, Stiffness));
 		const varying float Lambda = Lambdas[i];
 
 		varying FVector3f Delta1, Delta2, Delta3, Delta4;
-		const varying float DLambda = GetXPBDBendingElementDampingDeltas(P1, P2, P3, P4, X1, X2, X3, X4, M1, M2, M3, M4, Dt, RestAngles[i], Lambda, StiffnessValue, DampingRatio, MinStiffness, Delta1, Delta2, Delta3, Delta4);
+		const varying float DLambda = GetXPBDBendingElementDampingDeltas(P1, P2, P3, P4, X1, X2, X3, X4, M1, M2, M3, M4, Dt, RestAngles[i], Lambda, StiffnessValue, DampingRatio, Delta1, Delta2, Delta3, Delta4);
 
 		Lambdas[i] = Lambda + DLambda;
-
-		if (M1 > 0)
-		{
-			VectorScatter(&PandInvM[i1], SetVector4(P1 + Delta1, M1));
-		}
-		if (M2 > 0)
-		{
-			VectorScatter(&PandInvM[i2], SetVector4(P2 + Delta2, M2));
-		}
-		if (M3 > 0)
-		{
-			VectorScatter(&PandInvM[i3], SetVector4(P3 + Delta3, M3));
-		}
-		if (M4 > 0)
-		{
-			VectorScatter(&PandInvM[i4], SetVector4(P4 + Delta4, M4));
-		}
+		VectorScatter(&PandInvM[i1], SetVector4(P1 + Delta1, M1));
+		VectorScatter(&PandInvM[i2], SetVector4(P2 + Delta2, M2));
+		VectorScatter(&PandInvM[i3], SetVector4(P3 + Delta3, M3));
+		VectorScatter(&PandInvM[i4], SetVector4(P4 + Delta4, M4));
 	}
 }
 
 export void ApplyXPBDAnisotropicBendingConstraintsWithMaps(uniform FVector4f PandInvM[],
-											const uniform FIntVector4 Constraints[],
+											const uniform int32 ConstraintIndex1[],
+											const uniform int32 ConstraintIndex2[],
+											const uniform int32 ConstraintIndex3[],
+											const uniform int32 ConstraintIndex4[],
 											const uniform float RestAngles[],
 											const uniform bool IsBuckled[],
 											const uniform FVector3f WarpWeftBiasBaseMultiplers[],
 											uniform float Lambdas[],
 											const uniform float Dt,
-											const uniform float MinStiffness,
 											const uniform bool StiffnessHasMap,
-											const uniform uint8 StiffnessIndices[],
-											const uniform float StiffnessTable[],
+											const uniform FVector2f& StiffnessOffsetRange,
+											const uniform float StiffnessMapValues[],
 											const uniform bool StiffnessWeftHasMap,
-											const uniform uint8 StiffnessWeftIndices[],
-											const uniform float StiffnessWeftTable[],
+											const uniform FVector2f& StiffnessWeftOffsetRange,
+											const uniform float StiffnessWeftMapValues[],
 											const uniform bool StiffnessBiasHasMap,
-											const uniform uint8 StiffnessBiasIndices[],
-											const uniform float StiffnessBiasTable[],
+											const uniform FVector2f& StiffnessBiasOffsetRange,
+											const uniform float StiffnessBiasMapValues[],
 											const uniform bool BucklingStiffnessHasMap,
-											const uniform uint8 BucklingStiffnessIndices[],
-											const uniform float BucklingStiffnessTable[],
+											const uniform FVector2f& BucklingStiffnessOffsetRange,
+											const uniform float BucklingStiffnessMapValues[],
 											const uniform bool BucklingStiffnessWeftHasMap,
-											const uniform uint8 BucklingStiffnessWeftIndices[],
-											const uniform float BucklingStiffnessWeftTable[],
+											const uniform FVector2f& BucklingStiffnessWeftOffsetRange,
+											const uniform float BucklingStiffnessWeftMapValues[],
 											const uniform bool BucklingStiffnessBiasHasMap,
-											const uniform uint8 BucklingStiffnessBiasIndices[],
-											const uniform float BucklingStiffnessBiasTable[],
+											const uniform FVector2f& BucklingStiffnessBiasOffsetRange,
+											const uniform float BucklingStiffnessBiasMapValues[],
 											const uniform int32 NumConstraints)
 {
 
 	foreach(i = 0 ... NumConstraints)
 	{
-		const varying FIntVector4 Constraint = VectorLoad(&Constraints[extract(i,0)]);
-
-		const varying int32 i1 = Constraint.V[0];
-		const varying int32 i2 = Constraint.V[1];
-		const varying int32 i3 = Constraint.V[2];
-		const varying int32 i4 = Constraint.V[3];
-		const varying int8 StiffnessIndex = StiffnessHasMap ? StiffnessIndices[i] : 0;
-		const varying int8 StiffnessWeftIndex = StiffnessWeftHasMap ? StiffnessWeftIndices[i] : 0;
-		const varying int8 StiffnessBiasIndex = StiffnessBiasHasMap ? StiffnessBiasIndices[i] : 0;
-		const varying int8 BucklingStiffnessIndex = BucklingStiffnessHasMap ? BucklingStiffnessIndices[i] : 0;
-		const varying int8 BucklingStiffnessWeftIndex = BucklingStiffnessWeftHasMap ? BucklingStiffnessWeftIndices[i] : 0;
-		const varying int8 BucklingStiffnessBiasIndex = BucklingStiffnessBiasHasMap ? BucklingStiffnessBiasIndices[i] : 0;
+		const varying int32 i1 = ConstraintIndex1[i];
+		const varying int32 i2 = ConstraintIndex2[i];
+		const varying int32 i3 = ConstraintIndex3[i];
+		const varying int32 i4 = ConstraintIndex4[i];
 
 		const varying FVector4f PandInvM2 = VectorGather(&PandInvM[i2]);
 		const varying FVector4f PandInvM3 = VectorGather(&PandInvM[i3]);
@@ -266,102 +227,80 @@ export void ApplyXPBDAnisotropicBendingConstraintsWithMaps(uniform FVector4f Pan
 		UnzipPandInvM(PandInvM2, P2, M2);
 		UnzipPandInvM(PandInvM3, P3, M3);
 		UnzipPandInvM(PandInvM4, P4, M4);
-				
-		#pragma ignore warning(perf)
-		const varying float Stiffness = StiffnessTable[StiffnessIndex];
-		#pragma ignore warning(perf)
-		const varying float StiffnessWeft = StiffnessWeftTable[StiffnessWeftIndex];
-		#pragma ignore warning(perf)
-		const varying float StiffnessBias = StiffnessBiasTable[StiffnessBiasIndex];
 
-		const varying FVector3f StiffnessVector = {{StiffnessWeft, Stiffness, StiffnessBias}};
+		const varying float Stiffness = StiffnessHasMap ? StiffnessOffsetRange.V[0] + StiffnessMapValues[i]*StiffnessOffsetRange.V[1] : StiffnessOffsetRange.V[0];
+		const varying float StiffnessWeft = StiffnessWeftHasMap ? StiffnessWeftOffsetRange.V[0] + StiffnessWeftMapValues[i]*StiffnessWeftOffsetRange.V[1] : StiffnessWeftOffsetRange.V[0];
+		const varying float StiffnessBias = StiffnessBiasHasMap ? StiffnessBiasOffsetRange.V[0] + StiffnessBiasMapValues[i]*StiffnessBiasOffsetRange.V[1] : StiffnessBiasOffsetRange.V[0];
 
-		#pragma ignore warning(perf)
-		const varying float BucklingStiffness = BucklingStiffnessTable[BucklingStiffnessIndex];
-		#pragma ignore warning(perf)
-		const varying float BucklingStiffnessWeft = BucklingStiffnessWeftTable[BucklingStiffnessWeftIndex];
-		#pragma ignore warning(perf)
-		const varying float BucklingStiffnessBias = BucklingStiffnessBiasTable[BucklingStiffnessBiasIndex];
+		const varying FVector3f StiffnessVector = {{StiffnessWeft, Stiffness, StiffnessBias}};
+		
+		const varying float BucklingStiffness = BucklingStiffnessHasMap ? BucklingStiffnessOffsetRange.V[0] + BucklingStiffnessMapValues[i]*BucklingStiffnessOffsetRange.V[1] : BucklingStiffnessOffsetRange.V[0];
+		const varying float BucklingStiffnessWeft = BucklingStiffnessWeftHasMap ? BucklingStiffnessWeftOffsetRange.V[0] + BucklingStiffnessWeftMapValues[i]*BucklingStiffnessWeftOffsetRange.V[1] : BucklingStiffnessWeftOffsetRange.V[0];
+		const varying float BucklingStiffnessBias = BucklingStiffnessBiasHasMap ? BucklingStiffnessBiasOffsetRange.V[0] + BucklingStiffnessBiasMapValues[i]*BucklingStiffnessBiasOffsetRange.V[1] : BucklingStiffnessBiasOffsetRange.V[0];
 		
 		const varying FVector3f BucklingStiffnessVector = {{BucklingStiffnessWeft, BucklingStiffness, BucklingStiffnessBias}};
 		
-		const varying FVector3f WarpWeftBiasBaseMultiplier = VectorGather(&WarpWeftBiasBaseMultiplers[i]);
+		const varying FVector3f WarpWeftBiasBaseMultiplier = VectorLoad(&WarpWeftBiasBaseMultiplers[extract(i,0)]);
 		const varying float StiffnessValue = select(IsBuckled[i], VectorDot(WarpWeftBiasBaseMultiplier, BucklingStiffnessVector), VectorDot(WarpWeftBiasBaseMultiplier, StiffnessVector));
 		
 		const varying float Lambda = Lambdas[i];
 
 		varying FVector3f Delta1, Delta2, Delta3, Delta4;
-		const varying float DLambda = GetXPBDBendingElementDeltas(P1, P2, P3, P4, M1, M2, M3, M4, Dt, RestAngles[i], Lambda, StiffnessValue, MinStiffness, Delta1, Delta2, Delta3, Delta4);
+		const varying float DLambda = GetXPBDBendingElementDeltas(P1, P2, P3, P4, M1, M2, M3, M4, Dt, RestAngles[i], Lambda, StiffnessValue, Delta1, Delta2, Delta3, Delta4);
 
 		Lambdas[i] = Lambda + DLambda;
-
-		if (M1 > 0)
-		{
-			VectorScatter(&PandInvM[i1], SetVector4(P1 + Delta1, M1));
-		}
-		if (M2 > 0)
-		{
-			VectorScatter(&PandInvM[i2], SetVector4(P2 + Delta2, M2));
-		}
-		if (M3 > 0)
-		{
-			VectorScatter(&PandInvM[i3], SetVector4(P3 + Delta3, M3));
-		}
-		if (M4 > 0)
-		{
-			VectorScatter(&PandInvM[i4], SetVector4(P4 + Delta4, M4));
-		}
+		
+		VectorScatter(&PandInvM[i1], SetVector4(P1 + Delta1, M1));
+		VectorScatter(&PandInvM[i2], SetVector4(P2 + Delta2, M2));
+		VectorScatter(&PandInvM[i3], SetVector4(P3 + Delta3, M3));
+		VectorScatter(&PandInvM[i4], SetVector4(P4 + Delta4, M4));
 	}
 }
 
 export void ApplyXPBDAnisotropicBendingConstraintsWithDampingAndMaps(uniform FVector4f PandInvM[],
-											const uniform FVector3f X[],
-											const uniform FIntVector4 Constraints[],
+											const uniform FVector3f X1Array[],
+											const uniform FVector3f X2Array[],
+											const uniform FVector3f X3Array[],
+											const uniform FVector3f X4Array[],
+											const uniform int32 ConstraintIndex1[],
+											const uniform int32 ConstraintIndex2[],
+											const uniform int32 ConstraintIndex3[],
+											const uniform int32 ConstraintIndex4[],
 											const uniform float RestAngles[],
 											const uniform bool IsBuckled[],
 											const uniform FVector3f WarpWeftBiasBaseMultiplers[],
 											uniform float Lambdas[],
 											const uniform float Dt,
-											const uniform float MinStiffness,
 											const uniform bool StiffnessHasMap,
-											const uniform uint8 StiffnessIndices[],
-											const uniform float StiffnessTable[],
+											const uniform FVector2f& StiffnessOffsetRange,
+											const uniform float StiffnessMapValues[],
 											const uniform bool StiffnessWeftHasMap,
-											const uniform uint8 StiffnessWeftIndices[],
-											const uniform float StiffnessWeftTable[],
+											const uniform FVector2f& StiffnessWeftOffsetRange,
+											const uniform float StiffnessWeftMapValues[],
 											const uniform bool StiffnessBiasHasMap,
-											const uniform uint8 StiffnessBiasIndices[],
-											const uniform float StiffnessBiasTable[],
+											const uniform FVector2f& StiffnessBiasOffsetRange,
+											const uniform float StiffnessBiasMapValues[],
 											const uniform bool BucklingStiffnessHasMap,
-											const uniform uint8 BucklingStiffnessIndices[],
-											const uniform float BucklingStiffnessTable[],
+											const uniform FVector2f& BucklingStiffnessOffsetRange,
+											const uniform float BucklingStiffnessMapValues[],
 											const uniform bool BucklingStiffnessWeftHasMap,
-											const uniform uint8 BucklingStiffnessWeftIndices[],
-											const uniform float BucklingStiffnessWeftTable[],
+											const uniform FVector2f& BucklingStiffnessWeftOffsetRange,
+											const uniform float BucklingStiffnessWeftMapValues[],
 											const uniform bool BucklingStiffnessBiasHasMap,
-											const uniform uint8 BucklingStiffnessBiasIndices[],
-											const uniform float BucklingStiffnessBiasTable[],
+											const uniform FVector2f& BucklingStiffnessBiasOffsetRange,
+											const uniform float BucklingStiffnessBiasMapValues[],
 											const uniform bool DampingHasMap,
-											const uniform uint8 DampingIndices[],
-											const uniform float DampingTable[],
+											const uniform FVector2f& DampingOffsetRange,
+											const uniform float DampingMapValues[],
 											const uniform int32 NumConstraints)
 {
 
 	foreach(i = 0 ... NumConstraints)
 	{
-		const varying FIntVector4 Constraint = VectorLoad(&Constraints[extract(i,0)]);
-
-		const varying int32 i1 = Constraint.V[0];
-		const varying int32 i2 = Constraint.V[1];
-		const varying int32 i3 = Constraint.V[2];
-		const varying int32 i4 = Constraint.V[3];
-		const varying int8 StiffnessIndex = StiffnessHasMap ? StiffnessIndices[i] : 0;
-		const varying int8 StiffnessWeftIndex = StiffnessWeftHasMap ? StiffnessWeftIndices[i] : 0;
-		const varying int8 StiffnessBiasIndex = StiffnessBiasHasMap ? StiffnessBiasIndices[i] : 0;
-		const varying int8 BucklingStiffnessIndex = BucklingStiffnessHasMap ? BucklingStiffnessIndices[i] : 0;
-		const varying int8 BucklingStiffnessWeftIndex = BucklingStiffnessWeftHasMap ? BucklingStiffnessWeftIndices[i] : 0;
-		const varying int8 BucklingStiffnessBiasIndex = BucklingStiffnessBiasHasMap ? BucklingStiffnessBiasIndices[i] : 0;
-		const varying int8 DampingIndex = DampingHasMap ? DampingIndices[i] : 0;
+		const varying int32 i1 = ConstraintIndex1[i];
+		const varying int32 i2 = ConstraintIndex2[i];
+		const varying int32 i3 = ConstraintIndex3[i];
+		const varying int32 i4 = ConstraintIndex4[i];
 
 		const varying FVector4f PandInvM2 = VectorGather(&PandInvM[i2]);
 		const varying FVector4f PandInvM3 = VectorGather(&PandInvM[i3]);
@@ -376,110 +315,88 @@ export void ApplyXPBDAnisotropicBendingConstraintsWithDampingAndMaps(uniform FVe
 		UnzipPandInvM(PandInvM3, P3, M3);
 		UnzipPandInvM(PandInvM4, P4, M4);
 
-		#pragma ignore warning(perf)
-		const varying float Stiffness = StiffnessTable[StiffnessIndex];
-		#pragma ignore warning(perf)
-		const varying float StiffnessWeft = StiffnessWeftTable[StiffnessWeftIndex];
-		#pragma ignore warning(perf)
-		const varying float StiffnessBias = StiffnessBiasTable[StiffnessBiasIndex];
+		const varying float Stiffness = StiffnessHasMap ? StiffnessOffsetRange.V[0] + StiffnessMapValues[i]*StiffnessOffsetRange.V[1] : StiffnessOffsetRange.V[0];
+		const varying float StiffnessWeft = StiffnessWeftHasMap ? StiffnessWeftOffsetRange.V[0] + StiffnessWeftMapValues[i]*StiffnessWeftOffsetRange.V[1] : StiffnessWeftOffsetRange.V[0];
+		const varying float StiffnessBias = StiffnessBiasHasMap ? StiffnessBiasOffsetRange.V[0] + StiffnessBiasMapValues[i]*StiffnessBiasOffsetRange.V[1] : StiffnessBiasOffsetRange.V[0];
 
 		const varying FVector3f StiffnessVector = {{StiffnessWeft, Stiffness, StiffnessBias}};
-
-		#pragma ignore warning(perf)
-		const varying float BucklingStiffness = BucklingStiffnessTable[BucklingStiffnessIndex];
-		#pragma ignore warning(perf)
-		const varying float BucklingStiffnessWeft = BucklingStiffnessWeftTable[BucklingStiffnessWeftIndex];
-		#pragma ignore warning(perf)
-		const varying float BucklingStiffnessBias = BucklingStiffnessBiasTable[BucklingStiffnessBiasIndex];
+		
+		const varying float BucklingStiffness = BucklingStiffnessHasMap ? BucklingStiffnessOffsetRange.V[0] + BucklingStiffnessMapValues[i]*BucklingStiffnessOffsetRange.V[1] : BucklingStiffnessOffsetRange.V[0];
+		const varying float BucklingStiffnessWeft = BucklingStiffnessWeftHasMap ? BucklingStiffnessWeftOffsetRange.V[0] + BucklingStiffnessWeftMapValues[i]*BucklingStiffnessWeftOffsetRange.V[1] : BucklingStiffnessWeftOffsetRange.V[0];
+		const varying float BucklingStiffnessBias = BucklingStiffnessBiasHasMap ? BucklingStiffnessBiasOffsetRange.V[0] + BucklingStiffnessBiasMapValues[i]*BucklingStiffnessBiasOffsetRange.V[1] : BucklingStiffnessBiasOffsetRange.V[0];
 		
 		const varying FVector3f BucklingStiffnessVector = {{BucklingStiffnessWeft, BucklingStiffness, BucklingStiffnessBias}};
 		
-		#pragma ignore warning(perf)
-		const varying float DampingRatio = DampingTable[DampingIndex];
+		const varying float DampingRatio = DampingHasMap ? DampingOffsetRange.V[0] + DampingMapValues[i]*DampingOffsetRange.V[1] : DampingOffsetRange.V[0];
 		
-		const varying float RestAngle = RestAngles[i];
-		const varying FVector3f WarpWeftBiasBaseMultiplier = VectorGather(&WarpWeftBiasBaseMultiplers[i]);
-		const varying float StiffnessValue = select(IsBuckled[i], VectorDot(WarpWeftBiasBaseMultiplier, BucklingStiffnessVector), VectorDot(WarpWeftBiasBaseMultiplier, StiffnessVector));
+		const uniform int32 i0 = extract(i,0);
 
-		const varying FVector3f X1 = VectorGather(&X[i1]);
-		const varying FVector3f X2 = VectorGather(&X[i2]);
-		const varying FVector3f X3 = VectorGather(&X[i3]);
-		const varying FVector3f X4 = VectorGather(&X[i4]);
+		const varying FVector3f X1 = VectorLoad(&X1Array[i0]);
+		const varying FVector3f X2 = VectorLoad(&X2Array[i0]);
+		const varying FVector3f X3 = VectorLoad(&X3Array[i0]);
+		const varying FVector3f X4 = VectorLoad(&X4Array[i0]);
+		
+		const varying FVector3f WarpWeftBiasBaseMultiplier = VectorLoad(&WarpWeftBiasBaseMultiplers[i0]);
+		const varying float StiffnessValue = select(IsBuckled[i], VectorDot(WarpWeftBiasBaseMultiplier, BucklingStiffnessVector), VectorDot(WarpWeftBiasBaseMultiplier, StiffnessVector));
 		
 		const varying float Lambda = Lambdas[i];
 
 		varying FVector3f Delta1, Delta2, Delta3, Delta4;
-		const varying float DLambda = GetXPBDBendingElementDeltasWithDamping(P1, P2, P3, P4, X1, X2, X3, X4, M1, M2, M3, M4, Dt, RestAngles[i], Lambda, StiffnessValue, DampingRatio, MinStiffness, Delta1, Delta2, Delta3, Delta4);
+		const varying float DLambda = GetXPBDBendingElementDeltasWithDamping(P1, P2, P3, P4, X1, X2, X3, X4, M1, M2, M3, M4, Dt, RestAngles[i], Lambda, StiffnessValue, DampingRatio, Delta1, Delta2, Delta3, Delta4);
 
 		Lambdas[i] = Lambda + DLambda;
 
-		if (M1 > 0)
-		{
-			VectorScatter(&PandInvM[i1], SetVector4(P1 + Delta1, M1));
-		}
-		if (M2 > 0)
-		{
-			VectorScatter(&PandInvM[i2], SetVector4(P2 + Delta2, M2));
-		}
-		if (M3 > 0)
-		{
-			VectorScatter(&PandInvM[i3], SetVector4(P3 + Delta3, M3));
-		}
-		if (M4 > 0)
-		{
-			VectorScatter(&PandInvM[i4], SetVector4(P4 + Delta4, M4));
-		}
+		VectorScatter(&PandInvM[i1], SetVector4(P1 + Delta1, M1));
+		VectorScatter(&PandInvM[i2], SetVector4(P2 + Delta2, M2));
+		VectorScatter(&PandInvM[i3], SetVector4(P3 + Delta3, M3));
+		VectorScatter(&PandInvM[i4], SetVector4(P4 + Delta4, M4));
 	}
 }
 
 export void ApplyXPBDAnisotropicBendingDampingConstraintsWithMaps(uniform FVector4f PandInvM[],
-											const uniform FVector3f X[],
-											const uniform FIntVector4 Constraints[],
+											const uniform FVector3f X1Array[],
+											const uniform FVector3f X2Array[],
+											const uniform FVector3f X3Array[],
+											const uniform FVector3f X4Array[],
+											const uniform int32 ConstraintIndex1[],
+											const uniform int32 ConstraintIndex2[],
+											const uniform int32 ConstraintIndex3[],
+											const uniform int32 ConstraintIndex4[],
 											const uniform float RestAngles[],
 											const uniform bool IsBuckled[],
 											const uniform FVector3f WarpWeftBiasBaseMultiplers[],
 											uniform float Lambdas[],
 											const uniform float Dt,
-											const uniform float MinStiffness,
 											const uniform bool StiffnessHasMap,
-											const uniform uint8 StiffnessIndices[],
-											const uniform float StiffnessTable[],
+											const uniform FVector2f& StiffnessOffsetRange,
+											const uniform float StiffnessMapValues[],
 											const uniform bool StiffnessWeftHasMap,
-											const uniform uint8 StiffnessWeftIndices[],
-											const uniform float StiffnessWeftTable[],
+											const uniform FVector2f& StiffnessWeftOffsetRange,
+											const uniform float StiffnessWeftMapValues[],
 											const uniform bool StiffnessBiasHasMap,
-											const uniform uint8 StiffnessBiasIndices[],
-											const uniform float StiffnessBiasTable[],
+											const uniform FVector2f& StiffnessBiasOffsetRange,
+											const uniform float StiffnessBiasMapValues[],
 											const uniform bool BucklingStiffnessHasMap,
-											const uniform uint8 BucklingStiffnessIndices[],
-											const uniform float BucklingStiffnessTable[],
+											const uniform FVector2f& BucklingStiffnessOffsetRange,
+											const uniform float BucklingStiffnessMapValues[],
 											const uniform bool BucklingStiffnessWeftHasMap,
-											const uniform uint8 BucklingStiffnessWeftIndices[],
-											const uniform float BucklingStiffnessWeftTable[],
+											const uniform FVector2f& BucklingStiffnessWeftOffsetRange,
+											const uniform float BucklingStiffnessWeftMapValues[],
 											const uniform bool BucklingStiffnessBiasHasMap,
-											const uniform uint8 BucklingStiffnessBiasIndices[],
-											const uniform float BucklingStiffnessBiasTable[],
+											const uniform FVector2f& BucklingStiffnessBiasOffsetRange,
+											const uniform float BucklingStiffnessBiasMapValues[],
 											const uniform bool DampingHasMap,
-											const uniform uint8 DampingIndices[],
-											const uniform float DampingTable[],
+											const uniform FVector2f& DampingOffsetRange,
+											const uniform float DampingMapValues[],
 											const uniform int32 NumConstraints)
 {
 
 	foreach(i = 0 ... NumConstraints)
 	{
-		const varying FIntVector4 Constraint = VectorLoad(&Constraints[extract(i,0)]);
-
-		const varying int32 i1 = Constraint.V[0];
-		const varying int32 i2 = Constraint.V[1];
-		const varying int32 i3 = Constraint.V[2];
-		const varying int32 i4 = Constraint.V[3];
-		const varying int8 StiffnessIndex = StiffnessHasMap ? StiffnessIndices[i] : 0;
-		const varying int8 StiffnessWeftIndex = StiffnessWeftHasMap ? StiffnessWeftIndices[i] : 0;
-		const varying int8 StiffnessBiasIndex = StiffnessBiasHasMap ? StiffnessBiasIndices[i] : 0;
-		const varying int8 BucklingStiffnessIndex = BucklingStiffnessHasMap ? BucklingStiffnessIndices[i] : 0;
-		const varying int8 BucklingStiffnessWeftIndex = BucklingStiffnessWeftHasMap ? BucklingStiffnessWeftIndices[i] : 0;
-		const varying int8 BucklingStiffnessBiasIndex = BucklingStiffnessBiasHasMap ? BucklingStiffnessBiasIndices[i] : 0;
-		const varying int8 DampingIndex = DampingHasMap ? DampingIndices[i] : 0;
+		const varying int32 i1 = ConstraintIndex1[i];
+		const varying int32 i2 = ConstraintIndex2[i];
+		const varying int32 i3 = ConstraintIndex3[i];
+		const varying int32 i4 = ConstraintIndex4[i];
 
 		const varying FVector4f PandInvM2 = VectorGather(&PandInvM[i2]);
 		const varying FVector4f PandInvM3 = VectorGather(&PandInvM[i3]);
@@ -494,58 +411,40 @@ export void ApplyXPBDAnisotropicBendingDampingConstraintsWithMaps(uniform FVecto
 		UnzipPandInvM(PandInvM3, P3, M3);
 		UnzipPandInvM(PandInvM4, P4, M4);
 
-		#pragma ignore warning(perf)
-		const varying float Stiffness = StiffnessTable[StiffnessIndex];
-		#pragma ignore warning(perf)
-		const varying float StiffnessWeft = StiffnessWeftTable[StiffnessWeftIndex];
-		#pragma ignore warning(perf)
-		const varying float StiffnessBias = StiffnessBiasTable[StiffnessBiasIndex];
+		const varying float Stiffness = StiffnessHasMap ? StiffnessOffsetRange.V[0] + StiffnessMapValues[i]*StiffnessOffsetRange.V[1] : StiffnessOffsetRange.V[0];
+		const varying float StiffnessWeft = StiffnessWeftHasMap ? StiffnessWeftOffsetRange.V[0] + StiffnessWeftMapValues[i]*StiffnessWeftOffsetRange.V[1] : StiffnessWeftOffsetRange.V[0];
+		const varying float StiffnessBias = StiffnessBiasHasMap ? StiffnessBiasOffsetRange.V[0] + StiffnessBiasMapValues[i]*StiffnessBiasOffsetRange.V[1] : StiffnessBiasOffsetRange.V[0];
 
 		const varying FVector3f StiffnessVector = {{StiffnessWeft, Stiffness, StiffnessBias}};
-
-		#pragma ignore warning(perf)
-		const varying float BucklingStiffness = BucklingStiffnessTable[BucklingStiffnessIndex];
-		#pragma ignore warning(perf)
-		const varying float BucklingStiffnessWeft = BucklingStiffnessWeftTable[BucklingStiffnessWeftIndex];
-		#pragma ignore warning(perf)
-		const varying float BucklingStiffnessBias = BucklingStiffnessBiasTable[BucklingStiffnessBiasIndex];
+		
+		const varying float BucklingStiffness = BucklingStiffnessHasMap ? BucklingStiffnessOffsetRange.V[0] + BucklingStiffnessMapValues[i]*BucklingStiffnessOffsetRange.V[1] : BucklingStiffnessOffsetRange.V[0];
+		const varying float BucklingStiffnessWeft = BucklingStiffnessWeftHasMap ? BucklingStiffnessWeftOffsetRange.V[0] + BucklingStiffnessWeftMapValues[i]*BucklingStiffnessWeftOffsetRange.V[1] : BucklingStiffnessWeftOffsetRange.V[0];
+		const varying float BucklingStiffnessBias = BucklingStiffnessBiasHasMap ? BucklingStiffnessBiasOffsetRange.V[0] + BucklingStiffnessBiasMapValues[i]*BucklingStiffnessBiasOffsetRange.V[1] : BucklingStiffnessBiasOffsetRange.V[0];
 		
 		const varying FVector3f BucklingStiffnessVector = {{BucklingStiffnessWeft, BucklingStiffness, BucklingStiffnessBias}};
 		
-		#pragma ignore warning(perf)
-		const varying float DampingRatio = DampingTable[DampingIndex];
+		const varying float DampingRatio = DampingHasMap ? DampingOffsetRange.V[0] + DampingMapValues[i]*DampingOffsetRange.V[1] : DampingOffsetRange.V[0];
+
+		const uniform int32 i0 = extract(i,0);
+
+		const varying FVector3f X1 = VectorLoad(&X1Array[i0]);
+		const varying FVector3f X2 = VectorLoad(&X2Array[i0]);
+		const varying FVector3f X3 = VectorLoad(&X3Array[i0]);
+		const varying FVector3f X4 = VectorLoad(&X4Array[i0]);
 		
-		const varying float RestAngle = RestAngles[i];
-		const varying FVector3f WarpWeftBiasBaseMultiplier = VectorGather(&WarpWeftBiasBaseMultiplers[i]);
+		const varying FVector3f WarpWeftBiasBaseMultiplier = VectorLoad(&WarpWeftBiasBaseMultiplers[i0]);
 		const varying float StiffnessValue = select(IsBuckled[i], VectorDot(WarpWeftBiasBaseMultiplier, BucklingStiffnessVector), VectorDot(WarpWeftBiasBaseMultiplier, StiffnessVector));
-
-		const varying FVector3f X1 = VectorGather(&X[i1]);
-		const varying FVector3f X2 = VectorGather(&X[i2]);
-		const varying FVector3f X3 = VectorGather(&X[i3]);
-		const varying FVector3f X4 = VectorGather(&X[i4]);
 		
 		const varying float Lambda = Lambdas[i];
 
 		varying FVector3f Delta1, Delta2, Delta3, Delta4;
-		const varying float DLambda = GetXPBDBendingElementDampingDeltas(P1, P2, P3, P4, X1, X2, X3, X4, M1, M2, M3, M4, Dt, RestAngles[i], Lambda, StiffnessValue, DampingRatio, MinStiffness, Delta1, Delta2, Delta3, Delta4);
+		const varying float DLambda = GetXPBDBendingElementDampingDeltas(P1, P2, P3, P4, X1, X2, X3, X4, M1, M2, M3, M4, Dt, RestAngles[i], Lambda, StiffnessValue, DampingRatio, Delta1, Delta2, Delta3, Delta4);
 
 		Lambdas[i] = Lambda + DLambda;
-
-		if (M1 > 0)
-		{
-			VectorScatter(&PandInvM[i1], SetVector4(P1 + Delta1, M1));
-		}
-		if (M2 > 0)
-		{
-			VectorScatter(&PandInvM[i2], SetVector4(P2 + Delta2, M2));
-		}
-		if (M3 > 0)
-		{
-			VectorScatter(&PandInvM[i3], SetVector4(P3 + Delta3, M3));
-		}
-		if (M4 > 0)
-		{
-			VectorScatter(&PandInvM[i4], SetVector4(P4 + Delta4, M4));
-		}
+		
+		VectorScatter(&PandInvM[i1], SetVector4(P1 + Delta1, M1));
+		VectorScatter(&PandInvM[i2], SetVector4(P2 + Delta2, M2));
+		VectorScatter(&PandInvM[i3], SetVector4(P3 + Delta3, M3));
+		VectorScatter(&PandInvM[i4], SetVector4(P4 + Delta4, M4));
 	}
 }
