Chaos Cloth optimizations to XPBD constraints.

diff --git a/Engine/Source/Runtime/Experimental/Chaos/Private/Chaos/XPBDAnisotropicSpringConstraints.cpp b/Engine/Source/Runtime/Experimental/Chaos/Private/Chaos/XPBDAnisotropicSpringConstraints.cpp
--- a/Engine/Source/Runtime/Experimental/Chaos/Private/Chaos/XPBDAnisotropicSpringConstraints.cpp
+++ b/Engine/Source/Runtime/Experimental/Chaos/Private/Chaos/XPBDAnisotropicSpringConstraints.cpp
@@ -59,30 +59,28 @@ FXPBDAnisotropicEdgeSpringConstraints::FXPBDAnisotropicEdgeSpringConstraints(
 	const FSolverVec2& InStiffnessBias,
 	const FSolverVec2& InDampingRatio,
 	const FSolverVec2& InWarpScale,
-	const FSolverVec2& InWeftScale,
-	bool bTrimKinematicConstraints)
+	const FSolverVec2& InWeftScale)
 	: Base(Particles,
 		TriangleMesh.GetElements(),
-		StiffnessWarpMultipliers,
+		TConstArrayView<FRealSingle>(), // We don't use base stiffness maps
 		InStiffnessWarp,
-		bTrimKinematicConstraints,
+		true /*bTrimKinematicConstraints*/,
 		MaxStiffness)
-	, StiffnessWeft(InStiffnessWeft,
+	, StiffnessWarp(InStiffnessWarp.ClampAxes(0, MaxStiffness),
+		StiffnessWarpMultipliers,
+		TConstArrayView<TVec2<int32>>(Constraints),
+		ParticleOffset,
+		ParticleCount)
+	, StiffnessWeft(InStiffnessWeft.ClampAxes(0, MaxStiffness),
 		StiffnessWeftMultipliers,
 		TConstArrayView<TVec2<int32>>(Constraints),
 		ParticleOffset,
-		ParticleCount,
-		FPBDStiffness::DefaultTableSize,
-		FPBDStiffness::DefaultParameterFitBase,
-		MaxStiffness)
-	, StiffnessBias(InStiffnessBias,
+		ParticleCount)
+	, StiffnessBias(InStiffnessBias.ClampAxes(0, MaxStiffness),
 		StiffnessBiasMultipliers,
 		TConstArrayView<TVec2<int32>>(Constraints),
 		ParticleOffset,
-		ParticleCount,
-		FPBDStiffness::DefaultTableSize,
-		FPBDStiffness::DefaultParameterFitBase,
-		MaxStiffness)
+		ParticleCount)
 	, DampingRatio(
 		InDampingRatio.ClampAxes(MinDampingRatio, MaxDampingRatio),
 		DampingMultipliers,
@@ -125,32 +123,30 @@ FXPBDAnisotropicEdgeSpringConstraints::FXPBDAnisotropicEdgeSpringConstraints(
 	const FSolverVec2& InStiffnessBias,
 	const FSolverVec2& InDampingRatio,
 	const FSolverVec2& InWarpScale,
-	const FSolverVec2& InWeftScale,
-	bool bTrimKinematicConstraints)
+	const FSolverVec2& InWeftScale)
 	: Base(Particles,
 		InParticleOffset,
 		InParticleCount,
 		TriangleMesh.GetElements(),
-		StiffnessWarpMultipliers,
+		TConstArrayView<FRealSingle>(), // We don't use base stiffness maps
 		InStiffnessWarp,
-		bTrimKinematicConstraints,
+		true /*bTrimKinematicConstraints*/,
 		MaxStiffness)
-	, StiffnessWeft(InStiffnessWeft,
+	, StiffnessWarp(InStiffnessWarp.ClampAxes(0, MaxStiffness),
+		StiffnessWarpMultipliers,
+		TConstArrayView<TVec2<int32>>(Constraints),
+		ParticleOffset,
+		ParticleCount)
+	, StiffnessWeft(InStiffnessWeft.ClampAxes(0, MaxStiffness),
 		StiffnessWeftMultipliers,
 		TConstArrayView<TVec2<int32>>(Constraints),
 		ParticleOffset,
-		ParticleCount,
-		FPBDStiffness::DefaultTableSize,
-		FPBDStiffness::DefaultParameterFitBase,
-		MaxStiffness)
-	, StiffnessBias(InStiffnessBias,
+		ParticleCount)
+	, StiffnessBias(InStiffnessBias.ClampAxes(0, MaxStiffness),
 		StiffnessBiasMultipliers,
 		TConstArrayView<TVec2<int32>>(Constraints),
 		ParticleOffset,
-		ParticleCount,
-		FPBDStiffness::DefaultTableSize,
-		FPBDStiffness::DefaultParameterFitBase,
-		MaxStiffness)
+		ParticleCount)
 	, DampingRatio(
 		InDampingRatio.ClampAxes(MinDampingRatio, MaxDampingRatio),
 		DampingMultipliers,
@@ -335,35 +331,29 @@ void FXPBDAnisotropicEdgeSpringConstraints::ApplyHelper(SolverParticlesOrRange&
 	if constexpr (bDampingBefore)
 	{
 		Delta += Spring::GetXPBDSpringDampingDelta(Particles, Dt, Constraint, Dists[ConstraintIndex], LambdasDamping[ConstraintIndex],
-			ExpStiffnessValue, MinStiffness, DampingRatioValue);
+			ExpStiffnessValue, DampingRatioValue);
 	}
 
 	if constexpr (bSingleLambda)
 	{
 		Delta += Spring::GetXPBDSpringDeltaWithDamping(Particles, Dt, Constraint, Dists[ConstraintIndex], Lambdas[ConstraintIndex],
-			ExpStiffnessValue, MinStiffness, DampingRatioValue);
+			ExpStiffnessValue, DampingRatioValue);
 	}
 
 	if constexpr (bSeparateStretch)
 	{
 		Delta += Spring::GetXPBDSpringDelta(Particles, Dt, Constraint, Dists[ConstraintIndex], Lambdas[ConstraintIndex],
-			ExpStiffnessValue, MinStiffness);
+			ExpStiffnessValue);
 	}
 
 	if constexpr (bDampingAfter)
 	{
 		Delta += Spring::GetXPBDSpringDampingDelta(Particles, Dt, Constraint, Dists[ConstraintIndex], LambdasDamping[ConstraintIndex],
-			ExpStiffnessValue, MinStiffness, DampingRatioValue);
+			ExpStiffnessValue, DampingRatioValue);
 	}
 
-	if (Particles.InvM(Index1) > (FSolverReal)0.)
-	{
-		Particles.P(Index1) += Particles.InvM(Index1) * Delta;
-	}
-	if (Particles.InvM(Index2) > (FSolverReal)0.)
-	{
-		Particles.P(Index2) -= Particles.InvM(Index2) * Delta;
-	}
+	Particles.P(Index1) += Particles.InvM(Index1) * Delta;
+	Particles.P(Index2) -= Particles.InvM(Index2) * Delta;
 }
 
 template<typename SolverParticlesOrRange>
@@ -372,7 +362,7 @@ void FXPBDAnisotropicEdgeSpringConstraints::Apply(SolverParticlesOrRange& Partic
 	TRACE_CPUPROFILER_EVENT_SCOPE(FXPBDAnisotropicEdgeSpringConstraints_Apply);
 	SCOPE_CYCLE_COUNTER(STAT_XPBD_AnisoSpring);
 
-	const bool StiffnessHasWeightMap = Stiffness.HasWeightMap();
+	const bool StiffnessHasWeightMap = StiffnessWarp.HasWeightMap();
 	const bool StiffnessWeftHasWeightMap = StiffnessWeft.HasWeightMap();
 	const bool StiffnessBiasHasWeightMap = StiffnessBias.HasWeightMap();
 	const bool DampingHasWeightMap = DampingRatio.HasWeightMap();
@@ -387,9 +377,9 @@ void FXPBDAnisotropicEdgeSpringConstraints::Apply(SolverParticlesOrRange& Partic
 
 			if (!StiffnessHasWeightMap && !StiffnessWeftHasWeightMap && !StiffnessBiasHasWeightMap && !DampingHasWeightMap)
 			{
-				const FSolverVec3 ExpStiffnessValue((FSolverReal)StiffnessWeft, (FSolverReal)Stiffness, (FSolverReal)StiffnessBias);
+				const FSolverVec3 ExpStiffnessValue((FSolverReal)StiffnessWeft, (FSolverReal)StiffnessWarp, (FSolverReal)StiffnessBias);
 				const FSolverReal DampingRatioValue = (FSolverReal)DampingRatio;
-				if (ExpStiffnessValue.Max() < MinStiffness)
+				if (ExpStiffnessValue.Max() <= MinStiffness)
 				{
 					return;
 				}
@@ -469,19 +459,18 @@ void FXPBDAnisotropicEdgeSpringConstraints::Apply(SolverParticlesOrRange& Partic
 								(const ispc::FVector3f*)&WarpWeftBiasBaseMultipliers.GetData()[ColorStart],
 								&Lambdas.GetData()[ColorStart],
 								Dt,
-								MinStiffness,
 								StiffnessHasWeightMap,
-								StiffnessHasWeightMap ? &Stiffness.GetIndices().GetData()[ColorStart] : nullptr,
-								&Stiffness.GetTable().GetData()[0],
+								reinterpret_cast<const ispc::FVector2f&>(StiffnessWarp.GetOffsetRange()),
+								StiffnessHasWeightMap ? &StiffnessWarp.GetMapValues().GetData()[ColorStart] : nullptr,
 								StiffnessWeftHasWeightMap,
-								StiffnessWeftHasWeightMap ? &StiffnessWeft.GetIndices().GetData()[ColorStart] : nullptr,
-								&StiffnessWeft.GetTable().GetData()[0],
+								reinterpret_cast<const ispc::FVector2f&>(StiffnessWeft.GetOffsetRange()),
+								StiffnessWeftHasWeightMap ? &StiffnessWeft.GetMapValues().GetData()[ColorStart] : nullptr,
 								StiffnessBiasHasWeightMap,
-								StiffnessBiasHasWeightMap ? &StiffnessBias.GetIndices().GetData()[ColorStart] : nullptr,
-								&StiffnessBias.GetTable().GetData()[0],
+								reinterpret_cast<const ispc::FVector2f&>(StiffnessBias.GetOffsetRange()),
+								StiffnessBiasHasWeightMap ? &StiffnessBias.GetMapValues().GetData()[ColorStart] : nullptr,
 								DampingHasWeightMap,
-								DampingHasWeightMap ? &DampingRatio.GetIndices().GetData()[ColorStart] : nullptr,
-								&DampingRatio.GetTable().GetData()[0],
+								reinterpret_cast<const ispc::FVector2f&>(DampingRatio.GetOffsetRange()),
+								DampingHasWeightMap ? &DampingRatio.GetMapValues()[ColorStart] : nullptr,
 								ColorSize);
 						}
 						return;
@@ -500,19 +489,18 @@ void FXPBDAnisotropicEdgeSpringConstraints::Apply(SolverParticlesOrRange& Partic
 								(const ispc::FVector3f*)&WarpWeftBiasBaseMultipliers.GetData()[ColorStart],
 								&LambdasDamping.GetData()[ColorStart],
 								Dt,
-								MinStiffness,
 								StiffnessHasWeightMap,
-								StiffnessHasWeightMap ? &Stiffness.GetIndices().GetData()[ColorStart] : nullptr,
-								&Stiffness.GetTable().GetData()[0],
+								reinterpret_cast<const ispc::FVector2f&>(StiffnessWarp.GetOffsetRange()),
+								StiffnessHasWeightMap ? &StiffnessWarp.GetMapValues().GetData()[ColorStart] : nullptr,
 								StiffnessWeftHasWeightMap,
-								StiffnessWeftHasWeightMap ? &StiffnessWeft.GetIndices().GetData()[ColorStart] : nullptr,
-								&StiffnessWeft.GetTable().GetData()[0],
+								reinterpret_cast<const ispc::FVector2f&>(StiffnessWeft.GetOffsetRange()),
+								StiffnessWeftHasWeightMap ? &StiffnessWeft.GetMapValues().GetData()[ColorStart] : nullptr,
 								StiffnessBiasHasWeightMap,
-								StiffnessBiasHasWeightMap ? &StiffnessBias.GetIndices().GetData()[ColorStart] : nullptr,
-								&StiffnessBias.GetTable().GetData()[0],
+								reinterpret_cast<const ispc::FVector2f&>(StiffnessBias.GetOffsetRange()),
+								StiffnessBiasHasWeightMap ? &StiffnessBias.GetMapValues().GetData()[ColorStart] : nullptr,
 								DampingHasWeightMap,
-								DampingHasWeightMap ? &DampingRatio.GetIndices().GetData()[ColorStart] : nullptr,
-								&DampingRatio.GetTable().GetData()[0],
+								reinterpret_cast<const ispc::FVector2f&>(DampingRatio.GetOffsetRange()),
+								DampingHasWeightMap ? &DampingRatio.GetMapValues()[ColorStart] : nullptr,
 								ColorSize);
 						}
 					}
@@ -528,16 +516,15 @@ void FXPBDAnisotropicEdgeSpringConstraints::Apply(SolverParticlesOrRange& Partic
 						(const ispc::FVector3f*)&WarpWeftBiasBaseMultipliers.GetData()[ColorStart],
 						&Lambdas.GetData()[ColorStart],
 						Dt,
-						MinStiffness,
 						StiffnessHasWeightMap,
-						&Stiffness.GetIndices().GetData()[ColorStart],
-						&Stiffness.GetTable().GetData()[0],
+						reinterpret_cast<const ispc::FVector2f&>(StiffnessWarp.GetOffsetRange()),
+						StiffnessHasWeightMap ? &StiffnessWarp.GetMapValues().GetData()[ColorStart] : nullptr,
 						StiffnessWeftHasWeightMap,
-						StiffnessWeftHasWeightMap ? &StiffnessWeft.GetIndices().GetData()[ColorStart] : nullptr,
-						&StiffnessWeft.GetTable().GetData()[0],
+						reinterpret_cast<const ispc::FVector2f&>(StiffnessWeft.GetOffsetRange()),
+						StiffnessWeftHasWeightMap ? &StiffnessWeft.GetMapValues().GetData()[ColorStart] : nullptr,
 						StiffnessBiasHasWeightMap,
-						StiffnessBiasHasWeightMap ? &StiffnessBias.GetIndices().GetData()[ColorStart] : nullptr,
-						&StiffnessBias.GetTable().GetData()[0],
+						reinterpret_cast<const ispc::FVector2f&>(StiffnessBias.GetOffsetRange()),
+						StiffnessBiasHasWeightMap ? &StiffnessBias.GetMapValues().GetData()[ColorStart] : nullptr,
 						ColorSize);
 				}
 			}
@@ -546,7 +533,7 @@ void FXPBDAnisotropicEdgeSpringConstraints::Apply(SolverParticlesOrRange& Partic
 #endif
 		{
 			// Parallel non-ispc
-			const FSolverReal StiffnessNoMap = (FSolverReal)Stiffness;
+			const FSolverReal StiffnessNoMap = (FSolverReal)StiffnessWarp;
 			const FSolverReal StiffnessWeftNoMap = (FSolverReal)StiffnessWeft;
 			const FSolverReal StiffnessBiasNoMap = (FSolverReal)StiffnessBias;
 			const FSolverReal DampingNoMap = (FSolverReal)DampingRatio;
@@ -565,7 +552,7 @@ void FXPBDAnisotropicEdgeSpringConstraints::Apply(SolverParticlesOrRange& Partic
 							constexpr bool bSingleLambda = false;
 							constexpr bool bSeparateStretch = false;
 							constexpr bool bDampingAfter = false;
-							const FSolverReal ExpStiffnessValue = StiffnessHasWeightMap ? Stiffness[ConstraintIndex] : StiffnessNoMap;
+							const FSolverReal ExpStiffnessValue = StiffnessHasWeightMap ? StiffnessWarp[ConstraintIndex] : StiffnessNoMap;
 							const FSolverReal ExpStiffnessWeftValue = StiffnessWeftHasWeightMap ? StiffnessWeft[ConstraintIndex] : StiffnessWeftNoMap;
 							const FSolverReal ExpStiffnessBiasValue = StiffnessBiasHasWeightMap ? StiffnessBias[ConstraintIndex] : StiffnessBiasNoMap;
 							const FSolverReal DampingRatioValue = DampingHasWeightMap ? DampingRatio[ConstraintIndex] : DampingNoMap;
@@ -593,7 +580,7 @@ void FXPBDAnisotropicEdgeSpringConstraints::Apply(SolverParticlesOrRange& Partic
 							constexpr bool bSingleLambda = false;
 							constexpr bool bSeparateStretch = true;
 							constexpr bool bDampingAfter = false;
-							const FSolverReal ExpStiffnessValue = StiffnessHasWeightMap ? Stiffness[ConstraintIndex] : StiffnessNoMap;
+							const FSolverReal ExpStiffnessValue = StiffnessHasWeightMap ? StiffnessWarp[ConstraintIndex] : StiffnessNoMap;
 							const FSolverReal ExpStiffnessWeftValue = StiffnessWeftHasWeightMap ? StiffnessWeft[ConstraintIndex] : StiffnessWeftNoMap;
 							const FSolverReal ExpStiffnessBiasValue = StiffnessBiasHasWeightMap ? StiffnessBias[ConstraintIndex] : StiffnessBiasNoMap;
 							const FSolverReal DampingRatioValue = DampingHasWeightMap ? DampingRatio[ConstraintIndex] : DampingNoMap;
@@ -616,7 +603,7 @@ void FXPBDAnisotropicEdgeSpringConstraints::Apply(SolverParticlesOrRange& Partic
 							constexpr bool bSingleLambda = true;
 							constexpr bool bSeparateStretch = false;
 							constexpr bool bDampingAfter = false;
-							const FSolverReal ExpStiffnessValue = StiffnessHasWeightMap ? Stiffness[ConstraintIndex] : StiffnessNoMap;
+							const FSolverReal ExpStiffnessValue = StiffnessHasWeightMap ? StiffnessWarp[ConstraintIndex] : StiffnessNoMap;
 							const FSolverReal ExpStiffnessWeftValue = StiffnessWeftHasWeightMap ? StiffnessWeft[ConstraintIndex] : StiffnessWeftNoMap;
 							const FSolverReal ExpStiffnessBiasValue = StiffnessBiasHasWeightMap ? StiffnessBias[ConstraintIndex] : StiffnessBiasNoMap;
 							const FSolverReal DampingRatioValue = DampingHasWeightMap ? DampingRatio[ConstraintIndex] : DampingNoMap;
@@ -639,7 +626,7 @@ void FXPBDAnisotropicEdgeSpringConstraints::Apply(SolverParticlesOrRange& Partic
 							constexpr bool bSingleLambda = false;
 							constexpr bool bSeparateStretch = true;
 							constexpr bool bDampingAfter = false;
-							const FSolverReal ExpStiffnessValue = StiffnessHasWeightMap ? Stiffness[ConstraintIndex] : StiffnessNoMap;
+							const FSolverReal ExpStiffnessValue = StiffnessHasWeightMap ? StiffnessWarp[ConstraintIndex] : StiffnessNoMap;
 							const FSolverReal ExpStiffnessWeftValue = StiffnessWeftHasWeightMap ? StiffnessWeft[ConstraintIndex] : StiffnessWeftNoMap;
 							const FSolverReal ExpStiffnessBiasValue = StiffnessBiasHasWeightMap ? StiffnessBias[ConstraintIndex] : StiffnessBiasNoMap;
 							const FSolverReal DampingRatioValue = DampingHasWeightMap ? DampingRatio[ConstraintIndex] : DampingNoMap;
@@ -662,7 +649,7 @@ void FXPBDAnisotropicEdgeSpringConstraints::Apply(SolverParticlesOrRange& Partic
 							constexpr bool bSingleLambda = false;
 							constexpr bool bSeparateStretch = true;
 							constexpr bool bDampingAfter = true;
-							const FSolverReal ExpStiffnessValue = StiffnessHasWeightMap ? Stiffness[ConstraintIndex] : StiffnessNoMap;
+							const FSolverReal ExpStiffnessValue = StiffnessHasWeightMap ? StiffnessWarp[ConstraintIndex] : StiffnessNoMap;
 							const FSolverReal ExpStiffnessWeftValue = StiffnessWeftHasWeightMap ? StiffnessWeft[ConstraintIndex] : StiffnessWeftNoMap;
 							const FSolverReal ExpStiffnessBiasValue = StiffnessBiasHasWeightMap ? StiffnessBias[ConstraintIndex] : StiffnessBiasNoMap;
 							const FSolverReal DampingRatioValue = DampingHasWeightMap ? DampingRatio[ConstraintIndex] : DampingNoMap;
@@ -685,7 +672,7 @@ void FXPBDAnisotropicEdgeSpringConstraints::Apply(SolverParticlesOrRange& Partic
 							constexpr bool bSingleLambda = false;
 							constexpr bool bSeparateStretch = false;
 							constexpr bool bDampingAfter = true;
-							const FSolverReal ExpStiffnessValue = StiffnessHasWeightMap ? Stiffness[ConstraintIndex] : StiffnessNoMap;
+							const FSolverReal ExpStiffnessValue = StiffnessHasWeightMap ? StiffnessWarp[ConstraintIndex] : StiffnessNoMap;
 							const FSolverReal ExpStiffnessWeftValue = StiffnessWeftHasWeightMap ? StiffnessWeft[ConstraintIndex] : StiffnessWeftNoMap;
 							const FSolverReal ExpStiffnessBiasValue = StiffnessBiasHasWeightMap ? StiffnessBias[ConstraintIndex] : StiffnessBiasNoMap;
 							const FSolverReal DampingRatioValue = DampingHasWeightMap ? DampingRatio[ConstraintIndex] : DampingNoMap;
@@ -708,7 +695,7 @@ void FXPBDAnisotropicEdgeSpringConstraints::Apply(SolverParticlesOrRange& Partic
 						constexpr bool bSingleLambda = false;
 						constexpr bool bSeparateStretch = true;
 						constexpr bool bDampingAfter = false;
-						const FSolverReal ExpStiffnessValue = StiffnessHasWeightMap ? Stiffness[ConstraintIndex] : StiffnessNoMap;
+						const FSolverReal ExpStiffnessValue = StiffnessHasWeightMap ? StiffnessWarp[ConstraintIndex] : StiffnessNoMap;
 						const FSolverReal ExpStiffnessWeftValue = StiffnessWeftHasWeightMap ? StiffnessWeft[ConstraintIndex] : StiffnessWeftNoMap;
 						const FSolverReal ExpStiffnessBiasValue = StiffnessBiasHasWeightMap ? StiffnessBias[ConstraintIndex] : StiffnessBiasNoMap;
 						ApplyHelper<bDampingBefore, bSingleLambda, bSeparateStretch, bDampingAfter>(Particles, Dt, ConstraintIndex, FSolverVec3(ExpStiffnessWeftValue, ExpStiffnessValue, ExpStiffnessBiasValue), (FSolverReal)0.);
@@ -720,7 +707,7 @@ void FXPBDAnisotropicEdgeSpringConstraints::Apply(SolverParticlesOrRange& Partic
 	else
 	{
 		// Single threaded
-		const FSolverReal StiffnessNoMap = (FSolverReal)Stiffness;
+		const FSolverReal StiffnessNoMap = (FSolverReal)StiffnessWarp;
 		const FSolverReal StiffnessWeftNoMap = (FSolverReal)StiffnessWeft;
 		const FSolverReal StiffnessBiasNoMap = (FSolverReal)StiffnessBias;
 		const FSolverReal DampingNoMap = (FSolverReal)DampingRatio;
@@ -734,7 +721,7 @@ void FXPBDAnisotropicEdgeSpringConstraints::Apply(SolverParticlesOrRange& Partic
 					constexpr bool bSingleLambda = false;
 					constexpr bool bSeparateStretch = false;
 					constexpr bool bDampingAfter = false;
-					const FSolverReal ExpStiffnessValue = StiffnessHasWeightMap ? Stiffness[ConstraintIndex] : StiffnessNoMap;
+					const FSolverReal ExpStiffnessValue = StiffnessHasWeightMap ? StiffnessWarp[ConstraintIndex] : StiffnessNoMap;
 					const FSolverReal ExpStiffnessWeftValue = StiffnessWeftHasWeightMap ? StiffnessWeft[ConstraintIndex] : StiffnessWeftNoMap;
 					const FSolverReal ExpStiffnessBiasValue = StiffnessBiasHasWeightMap ? StiffnessBias[ConstraintIndex] : StiffnessBiasNoMap;
 					const FSolverReal DampingRatioValue = DampingHasWeightMap ? DampingRatio[ConstraintIndex] : DampingNoMap;
@@ -756,7 +743,7 @@ void FXPBDAnisotropicEdgeSpringConstraints::Apply(SolverParticlesOrRange& Partic
 					constexpr bool bSingleLambda = false;
 					constexpr bool bSeparateStretch = true;
 					constexpr bool bDampingAfter = false;
-					const FSolverReal ExpStiffnessValue = StiffnessHasWeightMap ? Stiffness[ConstraintIndex] : StiffnessNoMap;
+					const FSolverReal ExpStiffnessValue = StiffnessHasWeightMap ? StiffnessWarp[ConstraintIndex] : StiffnessNoMap;
 					const FSolverReal ExpStiffnessWeftValue = StiffnessWeftHasWeightMap ? StiffnessWeft[ConstraintIndex] : StiffnessWeftNoMap;
 					const FSolverReal ExpStiffnessBiasValue = StiffnessBiasHasWeightMap ? StiffnessBias[ConstraintIndex] : StiffnessBiasNoMap;
 					const FSolverReal DampingRatioValue = DampingHasWeightMap ? DampingRatio[ConstraintIndex] : DampingNoMap;
@@ -773,7 +760,7 @@ void FXPBDAnisotropicEdgeSpringConstraints::Apply(SolverParticlesOrRange& Partic
 					constexpr bool bSingleLambda = true;
 					constexpr bool bSeparateStretch = false;
 					constexpr bool bDampingAfter = false;
-					const FSolverReal ExpStiffnessValue = StiffnessHasWeightMap ? Stiffness[ConstraintIndex] : StiffnessNoMap;
+					const FSolverReal ExpStiffnessValue = StiffnessHasWeightMap ? StiffnessWarp[ConstraintIndex] : StiffnessNoMap;
 					const FSolverReal ExpStiffnessWeftValue = StiffnessWeftHasWeightMap ? StiffnessWeft[ConstraintIndex] : StiffnessWeftNoMap;
 					const FSolverReal ExpStiffnessBiasValue = StiffnessBiasHasWeightMap ? StiffnessBias[ConstraintIndex] : StiffnessBiasNoMap;
 					const FSolverReal DampingRatioValue = DampingHasWeightMap ? DampingRatio[ConstraintIndex] : DampingNoMap;
@@ -790,7 +777,7 @@ void FXPBDAnisotropicEdgeSpringConstraints::Apply(SolverParticlesOrRange& Partic
 					constexpr bool bSingleLambda = false;
 					constexpr bool bSeparateStretch = true;
 					constexpr bool bDampingAfter = false;
-					const FSolverReal ExpStiffnessValue = StiffnessHasWeightMap ? Stiffness[ConstraintIndex] : StiffnessNoMap;
+					const FSolverReal ExpStiffnessValue = StiffnessHasWeightMap ? StiffnessWarp[ConstraintIndex] : StiffnessNoMap;
 					const FSolverReal ExpStiffnessWeftValue = StiffnessWeftHasWeightMap ? StiffnessWeft[ConstraintIndex] : StiffnessWeftNoMap;
 					const FSolverReal ExpStiffnessBiasValue = StiffnessBiasHasWeightMap ? StiffnessBias[ConstraintIndex] : StiffnessBiasNoMap;
 					const FSolverReal DampingRatioValue = DampingHasWeightMap ? DampingRatio[ConstraintIndex] : DampingNoMap;
@@ -807,7 +794,7 @@ void FXPBDAnisotropicEdgeSpringConstraints::Apply(SolverParticlesOrRange& Partic
 					constexpr bool bSingleLambda = false;
 					constexpr bool bSeparateStretch = true;
 					constexpr bool bDampingAfter = true;
-					const FSolverReal ExpStiffnessValue = StiffnessHasWeightMap ? Stiffness[ConstraintIndex] : StiffnessNoMap;
+					const FSolverReal ExpStiffnessValue = StiffnessHasWeightMap ? StiffnessWarp[ConstraintIndex] : StiffnessNoMap;
 					const FSolverReal ExpStiffnessWeftValue = StiffnessWeftHasWeightMap ? StiffnessWeft[ConstraintIndex] : StiffnessWeftNoMap;
 					const FSolverReal ExpStiffnessBiasValue = StiffnessBiasHasWeightMap ? StiffnessBias[ConstraintIndex] : StiffnessBiasNoMap;
 					const FSolverReal DampingRatioValue = DampingHasWeightMap ? DampingRatio[ConstraintIndex] : DampingNoMap;
@@ -824,7 +811,7 @@ void FXPBDAnisotropicEdgeSpringConstraints::Apply(SolverParticlesOrRange& Partic
 					constexpr bool bSingleLambda = false;
 					constexpr bool bSeparateStretch = false;
 					constexpr bool bDampingAfter = true;
-					const FSolverReal ExpStiffnessValue = StiffnessHasWeightMap ? Stiffness[ConstraintIndex] : StiffnessNoMap;
+					const FSolverReal ExpStiffnessValue = StiffnessHasWeightMap ? StiffnessWarp[ConstraintIndex] : StiffnessNoMap;
 					const FSolverReal ExpStiffnessWeftValue = StiffnessWeftHasWeightMap ? StiffnessWeft[ConstraintIndex] : StiffnessWeftNoMap;
 					const FSolverReal ExpStiffnessBiasValue = StiffnessBiasHasWeightMap ? StiffnessBias[ConstraintIndex] : StiffnessBiasNoMap;
 					const FSolverReal DampingRatioValue = DampingHasWeightMap ? DampingRatio[ConstraintIndex] : DampingNoMap;
@@ -841,7 +828,7 @@ void FXPBDAnisotropicEdgeSpringConstraints::Apply(SolverParticlesOrRange& Partic
 				constexpr bool bSingleLambda = false;
 				constexpr bool bSeparateStretch = true;
 				constexpr bool bDampingAfter = false;
-				const FSolverReal ExpStiffnessValue = StiffnessHasWeightMap ? Stiffness[ConstraintIndex] : StiffnessNoMap;
+				const FSolverReal ExpStiffnessValue = StiffnessHasWeightMap ? StiffnessWarp[ConstraintIndex] : StiffnessNoMap;
 				const FSolverReal ExpStiffnessWeftValue = StiffnessWeftHasWeightMap ? StiffnessWeft[ConstraintIndex] : StiffnessWeftNoMap;
 				const FSolverReal ExpStiffnessBiasValue = StiffnessBiasHasWeightMap ? StiffnessBias[ConstraintIndex] : StiffnessBiasNoMap;
 				ApplyHelper<bDampingBefore, bSingleLambda, bSeparateStretch, bDampingAfter>(Particles, Dt, ConstraintIndex, FSolverVec3(ExpStiffnessWeftValue, ExpStiffnessValue, ExpStiffnessBiasValue), (FSolverReal)0.);
@@ -858,7 +845,7 @@ void FXPBDAnisotropicEdgeSpringConstraints::UpdateLinearSystem(const FSolverPart
 	TRACE_CPUPROFILER_EVENT_SCOPE(FXPBDAnisotropicEdgeSpringConstraints_UpdateLinearSystem);
 	LinearSystem.ReserveForParallelAdd(Constraints.Num() * 2, Constraints.Num());
 
-	const bool StiffnessHasWeightMap = Stiffness.HasWeightMap();
+	const bool StiffnessHasWeightMap = StiffnessWarp.HasWeightMap();
 	const bool StiffnessWeftHasWeightMap = StiffnessWeft.HasWeightMap();
 	const bool StiffnessBiasHasWeightMap = StiffnessBias.HasWeightMap();
 	const bool DampingHasWeightMap = DampingRatio.HasWeightMap();
@@ -867,7 +854,7 @@ void FXPBDAnisotropicEdgeSpringConstraints::UpdateLinearSystem(const FSolverPart
 		const int32 ConstraintColorNum = ConstraintsPerColorStartIndex.Num() - 1;
 		if (!StiffnessHasWeightMap && !StiffnessWeftHasWeightMap && !StiffnessBiasHasWeightMap && !DampingHasWeightMap)
 		{
-			const FSolverVec3 ExpStiffnessValue((FSolverReal)StiffnessWeft, (FSolverReal)Stiffness, (FSolverReal)StiffnessBias);
+			const FSolverVec3 ExpStiffnessValue((FSolverReal)StiffnessWeft, (FSolverReal)StiffnessWarp, (FSolverReal)StiffnessBias);
 			const FSolverReal DampingRatioValue = (FSolverReal)DampingRatio;
 			if (ExpStiffnessValue.Max() < MinStiffness)
 			{
@@ -888,7 +875,7 @@ void FXPBDAnisotropicEdgeSpringConstraints::UpdateLinearSystem(const FSolverPart
 		}
 		else // Has weight maps
 		{
-			const FSolverReal StiffnessNoMap = (FSolverReal)Stiffness;
+			const FSolverReal StiffnessNoMap = (FSolverReal)StiffnessWarp;
 			const FSolverReal StiffnessWeftNoMap = (FSolverReal)StiffnessWeft;
 			const FSolverReal StiffnessBiasNoMap = (FSolverReal)StiffnessBias;
 			const FSolverReal DampingNoMap = (FSolverReal)DampingRatio;
@@ -900,7 +887,7 @@ void FXPBDAnisotropicEdgeSpringConstraints::UpdateLinearSystem(const FSolverPart
 					StiffnessBiasNoMap, DampingHasWeightMap, DampingNoMap](const int32 Index)
 				{
 					const int32 ConstraintIndex = ColorStart + Index;
-					const FSolverReal ExpStiffnessValue = StiffnessHasWeightMap ? Stiffness[ConstraintIndex] : StiffnessNoMap;
+					const FSolverReal ExpStiffnessValue = StiffnessHasWeightMap ? StiffnessWarp[ConstraintIndex] : StiffnessNoMap;
 					const FSolverReal ExpStiffnessWeftValue = StiffnessWeftHasWeightMap ? StiffnessWeft[ConstraintIndex] : StiffnessWeftNoMap;
 					const FSolverReal ExpStiffnessBiasValue = StiffnessBiasHasWeightMap ? StiffnessBias[ConstraintIndex] : StiffnessBiasNoMap;
 					const FSolverReal DampingRatioValue = DampingHasWeightMap ? DampingRatio[ConstraintIndex] : DampingNoMap;
@@ -915,7 +902,7 @@ void FXPBDAnisotropicEdgeSpringConstraints::UpdateLinearSystem(const FSolverPart
 	{
 		if (!StiffnessHasWeightMap && !StiffnessWeftHasWeightMap && !StiffnessBiasHasWeightMap && !DampingHasWeightMap)
 		{
-			const FSolverVec3 ExpStiffnessValue((FSolverReal)StiffnessWeft, (FSolverReal)Stiffness, (FSolverReal)StiffnessBias);
+			const FSolverVec3 ExpStiffnessValue((FSolverReal)StiffnessWeft, (FSolverReal)StiffnessWarp, (FSolverReal)StiffnessBias);
 			const FSolverReal DampingRatioValue = (FSolverReal)DampingRatio;
 			if (ExpStiffnessValue.Max() < MinStiffness)
 			{
@@ -930,13 +917,13 @@ void FXPBDAnisotropicEdgeSpringConstraints::UpdateLinearSystem(const FSolverPart
 		}
 		else
 		{
-			const FSolverReal StiffnessNoMap = (FSolverReal)Stiffness;
+			const FSolverReal StiffnessNoMap = (FSolverReal)StiffnessWarp;
 			const FSolverReal StiffnessWeftNoMap = (FSolverReal)StiffnessWeft;
 			const FSolverReal StiffnessBiasNoMap = (FSolverReal)StiffnessBias;
 			const FSolverReal DampingNoMap = (FSolverReal)DampingRatio;
 			for (int32 ConstraintIndex = 0; ConstraintIndex < Constraints.Num(); ++ConstraintIndex)
 			{
-				const FSolverReal ExpStiffnessValue = StiffnessHasWeightMap ? Stiffness[ConstraintIndex] : StiffnessNoMap;
+				const FSolverReal ExpStiffnessValue = StiffnessHasWeightMap ? StiffnessWarp[ConstraintIndex] : StiffnessNoMap;
 				const FSolverReal ExpStiffnessWeftValue = StiffnessWeftHasWeightMap ? StiffnessWeft[ConstraintIndex] : StiffnessWeftNoMap;
 				const FSolverReal ExpStiffnessBiasValue = StiffnessBiasHasWeightMap ? StiffnessBias[ConstraintIndex] : StiffnessBiasNoMap;
 				const FSolverReal DampingRatioValue = DampingHasWeightMap ? DampingRatio[ConstraintIndex] : DampingNoMap;
@@ -964,30 +951,28 @@ FXPBDAnisotropicAxialSpringConstraints::FXPBDAnisotropicAxialSpringConstraints(
 	const FSolverVec2& InStiffnessBias,
 	const FSolverVec2& InDampingRatio,
 	const FSolverVec2& InWarpScale,
-	const FSolverVec2& InWeftScale,
-	bool bTrimKinematicConstraints)
+	const FSolverVec2& InWeftScale)
 	: Base(Particles,
 		TriangleMesh.GetElements(),
-		StiffnessWarpMultipliers,
+		TConstArrayView<FRealSingle>(), // We don't use base stiffness maps
 		InStiffnessWarp,
-		bTrimKinematicConstraints,
+		true /*bTrimKinematicConstraints*/,
 		MaxStiffness)
-	, StiffnessWeft(InStiffnessWeft,
+	, StiffnessWarp(InStiffnessWarp.ClampAxes(0, MaxStiffness),
+		StiffnessWarpMultipliers,
+		TConstArrayView<TVec3<int32>>(Constraints),
+		ParticleOffset,
+		ParticleCount)
+	, StiffnessWeft(InStiffnessWeft.ClampAxes(0, MaxStiffness),
 		StiffnessWeftMultipliers,
 		TConstArrayView<TVec3<int32>>(Constraints),
 		ParticleOffset,
-		ParticleCount,
-		FPBDStiffness::DefaultTableSize,
-		FPBDStiffness::DefaultParameterFitBase,
-		MaxStiffness)
-	, StiffnessBias(InStiffnessBias,
+		ParticleCount)
+	, StiffnessBias(InStiffnessBias.ClampAxes(0, MaxStiffness),
 		StiffnessBiasMultipliers,
 		TConstArrayView<TVec3<int32>>(Constraints),
 		ParticleOffset,
-		ParticleCount,
-		FPBDStiffness::DefaultTableSize,
-		FPBDStiffness::DefaultParameterFitBase,
-		MaxStiffness)
+		ParticleCount)
 	, DampingRatio(
 		InDampingRatio.ClampAxes(MinDampingRatio, MaxDampingRatio),
 		DampingMultipliers,
@@ -1030,32 +1015,30 @@ FXPBDAnisotropicAxialSpringConstraints::FXPBDAnisotropicAxialSpringConstraints(
 	const FSolverVec2& InStiffnessBias,
 	const FSolverVec2& InDampingRatio,
 	const FSolverVec2& InWarpScale,
-	const FSolverVec2& InWeftScale,
-	bool bTrimKinematicConstraints)
+	const FSolverVec2& InWeftScale)
 	: Base(Particles,
 		InParticleOffset,
 		InParticleCount,
 		TriangleMesh.GetElements(),
-		StiffnessWarpMultipliers,
+		TConstArrayView<FRealSingle>(), // We don't use base stiffness maps
 		InStiffnessWarp,
-		bTrimKinematicConstraints,
+		true /*bTrimKinematicConstraints*/,
 		MaxStiffness)
-	, StiffnessWeft(InStiffnessWeft,
+	, StiffnessWarp(InStiffnessWarp.ClampAxes(0, MaxStiffness),
+		StiffnessWarpMultipliers,
+		TConstArrayView<TVec3<int32>>(Constraints),
+		ParticleOffset,
+		ParticleCount)
+	, StiffnessWeft(InStiffnessWeft.ClampAxes(0, MaxStiffness),
 		StiffnessWeftMultipliers,
 		TConstArrayView<TVec3<int32>>(Constraints),
 		ParticleOffset,
-		ParticleCount,
-		FPBDStiffness::DefaultTableSize,
-		FPBDStiffness::DefaultParameterFitBase,
-		MaxStiffness)
-	, StiffnessBias(InStiffnessBias,
+		ParticleCount)
+	, StiffnessBias(InStiffnessBias.ClampAxes(0, MaxStiffness),
 		StiffnessBiasMultipliers,
 		TConstArrayView<TVec3<int32>>(Constraints),
 		ParticleOffset,
-		ParticleCount,
-		FPBDStiffness::DefaultTableSize,
-		FPBDStiffness::DefaultParameterFitBase,
-		MaxStiffness)
+		ParticleCount)
 	, DampingRatio(
 		InDampingRatio.ClampAxes(MinDampingRatio, MaxDampingRatio),
 		DampingMultipliers,
@@ -1241,39 +1224,29 @@ void FXPBDAnisotropicAxialSpringConstraints::ApplyHelper(SolverParticlesOrRange&
 	if constexpr (bDampingBefore)
 	{
 		Delta += Spring::GetXPBDAxialSpringDampingDelta(Particles, Dt, Constraint, Barys[ConstraintIndex], Dists[ConstraintIndex], LambdasDamping[ConstraintIndex],
-			ExpStiffnessValue, MinStiffness, DampingRatioValue);
+			ExpStiffnessValue, DampingRatioValue);
 	}
 
 	if constexpr (bSingleLambda)
 	{
 		Delta += Spring::GetXPBDAxialSpringDeltaWithDamping(Particles, Dt, Constraint, Barys[ConstraintIndex], Dists[ConstraintIndex], Lambdas[ConstraintIndex],
-			ExpStiffnessValue, MinStiffness, DampingRatioValue);
+			ExpStiffnessValue, DampingRatioValue);
 	}
 
 	if constexpr (bSeparateStretch)
 	{
 		Delta += Spring::GetXPBDAxialSpringDelta(Particles, Dt, Constraint, Barys[ConstraintIndex], Dists[ConstraintIndex], Lambdas[ConstraintIndex],
-			ExpStiffnessValue, MinStiffness);
+			ExpStiffnessValue);
 	}
 
 	if constexpr (bDampingAfter)
 	{
 		Delta += Spring::GetXPBDAxialSpringDampingDelta(Particles, Dt, Constraint, Barys[ConstraintIndex], Dists[ConstraintIndex], LambdasDamping[ConstraintIndex],
-			ExpStiffnessValue, MinStiffness, DampingRatioValue);
-	}
-
-	if (Particles.InvM(Index1) > (FSolverReal)0.)
-	{
-		Particles.P(Index1) += Particles.InvM(Index1) * Delta;
-	}
-	if (Particles.InvM(Index2) > (FSolverReal)0.)
-	{
-		Particles.P(Index2) -= Particles.InvM(Index2) * Barys[ConstraintIndex] * Delta;
-	}
-	if (Particles.InvM(Index3) > (FSolverReal)0.)
-	{
-		Particles.P(Index3) -= Particles.InvM(Index3) * ((FSolverReal)1. - Barys[ConstraintIndex]) * Delta;
+			ExpStiffnessValue, DampingRatioValue);
 	}
+	Particles.P(Index1) += Particles.InvM(Index1) * Delta;
+	Particles.P(Index2) -= Particles.InvM(Index2) * Barys[ConstraintIndex] * Delta;
+	Particles.P(Index3) -= Particles.InvM(Index3) * ((FSolverReal)1. - Barys[ConstraintIndex]) * Delta;
 }
 
 template<typename SolverParticlesOrRange>
@@ -1282,7 +1255,7 @@ void FXPBDAnisotropicAxialSpringConstraints::Apply(SolverParticlesOrRange& Parti
 	TRACE_CPUPROFILER_EVENT_SCOPE(FXPBDAnisotropicAxialSpringConstraints_Apply);
 	SCOPE_CYCLE_COUNTER(STAT_XPBD_AnisoSpring);
 
-	const bool StiffnessHasWeightMap = Stiffness.HasWeightMap();
+	const bool StiffnessHasWeightMap = StiffnessWarp.HasWeightMap();
 	const bool StiffnessWeftHasWeightMap = StiffnessWeft.HasWeightMap();
 	const bool StiffnessBiasHasWeightMap = StiffnessBias.HasWeightMap();
 	const bool DampingHasWeightMap = DampingRatio.HasWeightMap();
@@ -1293,34 +1266,59 @@ void FXPBDAnisotropicAxialSpringConstraints::Apply(SolverParticlesOrRange& Parti
 #if INTEL_ISPC
 		if (bRealTypeCompatibleWithISPC && bChaos_XPBDSpring_ISPC_Enabled)
 		{
+			const bool bSingleLambda = Chaos_XPBDSpring_SplitDampingMode == (int32)EXPBDSplitDampingMode::SingleLambda;
+
 			if (!StiffnessHasWeightMap && !StiffnessWeftHasWeightMap && !StiffnessBiasHasWeightMap && !DampingHasWeightMap)
 			{
-				const FSolverVec3 ExpStiffnessValue((FSolverReal)StiffnessWeft, (FSolverReal)Stiffness, (FSolverReal)StiffnessBias);
+				const FSolverVec3 ExpStiffnessValue((FSolverReal)StiffnessWeft, (FSolverReal)StiffnessWarp, (FSolverReal)StiffnessBias);
 				const FSolverReal DampingRatioValue = (FSolverReal)DampingRatio;
-				if (ExpStiffnessValue.Max() < MinStiffness)
+				if (ExpStiffnessValue.Max() <= MinStiffness)
 				{
 					return;
 				}
 
-				// ISPC always does two pass damping before stretch
 				if (DampingRatioValue > 0)
 				{
-					for (int32 ConstraintColorIndex = 0; ConstraintColorIndex < ConstraintColorNum; ++ConstraintColorIndex)
+					if (bSingleLambda)
 					{
-						const int32 ColorStart = ConstraintsPerColorStartIndex[ConstraintColorIndex];
-						const int32 ColorSize = ConstraintsPerColorStartIndex[ConstraintColorIndex + 1] - ColorStart;
-						ispc::ApplyXPBDAnisoAxialSpringDampingConstraints(
-							(ispc::FVector4f*)Particles.GetPAndInvM().GetData(),
-							(const ispc::FVector3f*)Particles.XArray().GetData(),
-							(ispc::FIntVector*)&Constraints.GetData()[ColorStart],
-							&Barys.GetData()[ColorStart],
-							&Dists.GetData()[ColorStart],
-							(const ispc::FVector3f*)&WarpWeftBiasBaseMultipliers.GetData()[ColorStart],
-							&LambdasDamping.GetData()[ColorStart],
-							Dt,
-							reinterpret_cast<const ispc::FVector3f&>(ExpStiffnessValue),
-							DampingRatioValue,
-							ColorSize);
+						for (int32 ConstraintColorIndex = 0; ConstraintColorIndex < ConstraintColorNum; ++ConstraintColorIndex)
+						{
+							const int32 ColorStart = ConstraintsPerColorStartIndex[ConstraintColorIndex];
+							const int32 ColorSize = ConstraintsPerColorStartIndex[ConstraintColorIndex + 1] - ColorStart;
+							ispc::ApplyXPBDAnisoAxialSpringConstraintsWithDamping(
+								(ispc::FVector4f*)Particles.GetPAndInvM().GetData(),
+								(const ispc::FVector3f*)Particles.XArray().GetData(),
+								(ispc::FIntVector*)&Constraints.GetData()[ColorStart],
+								&Barys.GetData()[ColorStart],
+								&Dists.GetData()[ColorStart],
+								(const ispc::FVector3f*)&WarpWeftBiasBaseMultipliers.GetData()[ColorStart],
+								&LambdasDamping.GetData()[ColorStart],
+								Dt,
+								reinterpret_cast<const ispc::FVector3f&>(ExpStiffnessValue),
+								DampingRatioValue,
+								ColorSize);
+						}
+						return;
+					}
+					else
+					{
+						for (int32 ConstraintColorIndex = 0; ConstraintColorIndex < ConstraintColorNum; ++ConstraintColorIndex)
+						{
+							const int32 ColorStart = ConstraintsPerColorStartIndex[ConstraintColorIndex];
+							const int32 ColorSize = ConstraintsPerColorStartIndex[ConstraintColorIndex + 1] - ColorStart;
+							ispc::ApplyXPBDAnisoAxialSpringDampingConstraints(
+								(ispc::FVector4f*)Particles.GetPAndInvM().GetData(),
+								(const ispc::FVector3f*)Particles.XArray().GetData(),
+								(ispc::FIntVector*)&Constraints.GetData()[ColorStart],
+								&Barys.GetData()[ColorStart],
+								&Dists.GetData()[ColorStart],
+								(const ispc::FVector3f*)&WarpWeftBiasBaseMultipliers.GetData()[ColorStart],
+								&LambdasDamping.GetData()[ColorStart],
+								Dt,
+								reinterpret_cast<const ispc::FVector3f&>(ExpStiffnessValue),
+								DampingRatioValue,
+								ColorSize);
+						}
 					}
 				}
 				for (int32 ConstraintColorIndex = 0; ConstraintColorIndex < ConstraintColorNum; ++ConstraintColorIndex)
@@ -1341,36 +1339,68 @@ void FXPBDAnisotropicAxialSpringConstraints::Apply(SolverParticlesOrRange& Parti
 			}
 			else // ISPC with maps
 			{
-				// ISPC always does two pass damping before stretch
 				if (DampingHasWeightMap || (FSolverReal)DampingRatio > 0)
 				{
-					for (int32 ConstraintColorIndex = 0; ConstraintColorIndex < ConstraintColorNum; ++ConstraintColorIndex)
+					if (bSingleLambda)
 					{
-						const int32 ColorStart = ConstraintsPerColorStartIndex[ConstraintColorIndex];
-						const int32 ColorSize = ConstraintsPerColorStartIndex[ConstraintColorIndex + 1] - ColorStart;
-						ispc::ApplyXPBDAnisoAxialSpringDampingConstraintsWithWeightMaps(
-							(ispc::FVector4f*)Particles.GetPAndInvM().GetData(),
-							(const ispc::FVector3f*)Particles.XArray().GetData(),
-							(ispc::FIntVector*)&Constraints.GetData()[ColorStart],
-							&Barys.GetData()[ColorStart],
-							&Dists.GetData()[ColorStart],
-							(const ispc::FVector3f*)&WarpWeftBiasBaseMultipliers.GetData()[ColorStart],
-							&LambdasDamping.GetData()[ColorStart],
-							Dt,
-							MinStiffness,
-							StiffnessHasWeightMap,
-							StiffnessHasWeightMap ? &Stiffness.GetIndices().GetData()[ColorStart] : nullptr,
-							&Stiffness.GetTable().GetData()[0],
-							StiffnessWeftHasWeightMap,
-							StiffnessWeftHasWeightMap ? &StiffnessWeft.GetIndices().GetData()[ColorStart] : nullptr,
-							&StiffnessWeft.GetTable().GetData()[0],
-							StiffnessBiasHasWeightMap,
-							StiffnessBiasHasWeightMap ? &StiffnessBias.GetIndices().GetData()[ColorStart] : nullptr,
-							&StiffnessBias.GetTable().GetData()[0],
-							DampingHasWeightMap,
-							DampingHasWeightMap ? &DampingRatio.GetIndices().GetData()[ColorStart] : nullptr,
-							&DampingRatio.GetTable().GetData()[0],
-							ColorSize);
+						for (int32 ConstraintColorIndex = 0; ConstraintColorIndex < ConstraintColorNum; ++ConstraintColorIndex)
+						{
+							const int32 ColorStart = ConstraintsPerColorStartIndex[ConstraintColorIndex];
+							const int32 ColorSize = ConstraintsPerColorStartIndex[ConstraintColorIndex + 1] - ColorStart;
+							ispc::ApplyXPBDAnisoAxialSpringConstraintsWithDampingAndWeightMaps(
+								(ispc::FVector4f*)Particles.GetPAndInvM().GetData(),
+								(const ispc::FVector3f*)Particles.XArray().GetData(),
+								(ispc::FIntVector*)&Constraints.GetData()[ColorStart],
+								&Barys.GetData()[ColorStart],
+								&Dists.GetData()[ColorStart],
+								(const ispc::FVector3f*)&WarpWeftBiasBaseMultipliers.GetData()[ColorStart],
+								&LambdasDamping.GetData()[ColorStart],
+								Dt,
+								StiffnessHasWeightMap,
+								reinterpret_cast<const ispc::FVector2f&>(StiffnessWarp.GetOffsetRange()),
+								StiffnessHasWeightMap ? &StiffnessWarp.GetMapValues().GetData()[ColorStart] : nullptr,
+								StiffnessWeftHasWeightMap,
+								reinterpret_cast<const ispc::FVector2f&>(StiffnessWeft.GetOffsetRange()),
+								StiffnessWeftHasWeightMap ? &StiffnessWeft.GetMapValues().GetData()[ColorStart] : nullptr,
+								StiffnessBiasHasWeightMap,
+								reinterpret_cast<const ispc::FVector2f&>(StiffnessBias.GetOffsetRange()),
+								StiffnessBiasHasWeightMap ? &StiffnessBias.GetMapValues().GetData()[ColorStart] : nullptr,
+								DampingHasWeightMap,
+								reinterpret_cast<const ispc::FVector2f&>(DampingRatio.GetOffsetRange()),
+								DampingHasWeightMap ? &DampingRatio.GetMapValues()[ColorStart] : nullptr,
+								ColorSize);
+						}
+						return;
+					}
+					else
+					{
+						for (int32 ConstraintColorIndex = 0; ConstraintColorIndex < ConstraintColorNum; ++ConstraintColorIndex)
+						{
+							const int32 ColorStart = ConstraintsPerColorStartIndex[ConstraintColorIndex];
+							const int32 ColorSize = ConstraintsPerColorStartIndex[ConstraintColorIndex + 1] - ColorStart;
+							ispc::ApplyXPBDAnisoAxialSpringDampingConstraintsWithWeightMaps(
+								(ispc::FVector4f*)Particles.GetPAndInvM().GetData(),
+								(const ispc::FVector3f*)Particles.XArray().GetData(),
+								(ispc::FIntVector*)&Constraints.GetData()[ColorStart],
+								&Barys.GetData()[ColorStart],
+								&Dists.GetData()[ColorStart],
+								(const ispc::FVector3f*)&WarpWeftBiasBaseMultipliers.GetData()[ColorStart],
+								&LambdasDamping.GetData()[ColorStart],
+								Dt,
+								StiffnessHasWeightMap,
+								reinterpret_cast<const ispc::FVector2f&>(StiffnessWarp.GetOffsetRange()),
+								StiffnessHasWeightMap ? &StiffnessWarp.GetMapValues().GetData()[ColorStart] : nullptr,
+								StiffnessWeftHasWeightMap,
+								reinterpret_cast<const ispc::FVector2f&>(StiffnessWeft.GetOffsetRange()),
+								StiffnessWeftHasWeightMap ? &StiffnessWeft.GetMapValues().GetData()[ColorStart] : nullptr,
+								StiffnessBiasHasWeightMap,
+								reinterpret_cast<const ispc::FVector2f&>(StiffnessBias.GetOffsetRange()),
+								StiffnessBiasHasWeightMap ? &StiffnessBias.GetMapValues().GetData()[ColorStart] : nullptr,
+								DampingHasWeightMap,
+								reinterpret_cast<const ispc::FVector2f&>(DampingRatio.GetOffsetRange()),
+								DampingHasWeightMap ? &DampingRatio.GetMapValues()[ColorStart] : nullptr,
+								ColorSize);
+						}
 					}
 				}
 				for (int32 ConstraintColorIndex = 0; ConstraintColorIndex < ConstraintColorNum; ++ConstraintColorIndex)
@@ -1385,16 +1415,15 @@ void FXPBDAnisotropicAxialSpringConstraints::Apply(SolverParticlesOrRange& Parti
 						(const ispc::FVector3f*)&WarpWeftBiasBaseMultipliers.GetData()[ColorStart],
 						&Lambdas.GetData()[ColorStart],
 						Dt,
-						MinStiffness,
 						StiffnessHasWeightMap,
-						&Stiffness.GetIndices().GetData()[ColorStart],
-						&Stiffness.GetTable().GetData()[0],
+						reinterpret_cast<const ispc::FVector2f&>(StiffnessWarp.GetOffsetRange()),
+						StiffnessHasWeightMap ? &StiffnessWarp.GetMapValues().GetData()[ColorStart] : nullptr,
 						StiffnessWeftHasWeightMap,
-						StiffnessWeftHasWeightMap ? &StiffnessWeft.GetIndices().GetData()[ColorStart] : nullptr,
-						&StiffnessWeft.GetTable().GetData()[0],
+						reinterpret_cast<const ispc::FVector2f&>(StiffnessWeft.GetOffsetRange()),
+						StiffnessWeftHasWeightMap ? &StiffnessWeft.GetMapValues().GetData()[ColorStart] : nullptr,
 						StiffnessBiasHasWeightMap,
-						StiffnessBiasHasWeightMap ? &StiffnessBias.GetIndices().GetData()[ColorStart] : nullptr,
-						&StiffnessBias.GetTable().GetData()[0],
+						reinterpret_cast<const ispc::FVector2f&>(StiffnessBias.GetOffsetRange()),
+						StiffnessBiasHasWeightMap ? &StiffnessBias.GetMapValues().GetData()[ColorStart] : nullptr,
 						ColorSize);
 				}
 			}
@@ -1403,7 +1432,7 @@ void FXPBDAnisotropicAxialSpringConstraints::Apply(SolverParticlesOrRange& Parti
 #endif
 		{
 			// Parallel non-ispc
-			const FSolverReal StiffnessNoMap = (FSolverReal)Stiffness;
+			const FSolverReal StiffnessNoMap = (FSolverReal)StiffnessWarp;
 			const FSolverReal StiffnessWeftNoMap = (FSolverReal)StiffnessWeft;
 			const FSolverReal StiffnessBiasNoMap = (FSolverReal)StiffnessBias;
 			const FSolverReal DampingNoMap = (FSolverReal)DampingRatio;
@@ -1422,7 +1451,7 @@ void FXPBDAnisotropicAxialSpringConstraints::Apply(SolverParticlesOrRange& Parti
 							constexpr bool bSingleLambda = false;
 							constexpr bool bSeparateStretch = false;
 							constexpr bool bDampingAfter = false;
-							const FSolverReal ExpStiffnessValue = StiffnessHasWeightMap ? Stiffness[ConstraintIndex] : StiffnessNoMap;
+							const FSolverReal ExpStiffnessValue = StiffnessHasWeightMap ? StiffnessWarp[ConstraintIndex] : StiffnessNoMap;
 							const FSolverReal ExpStiffnessWeftValue = StiffnessWeftHasWeightMap ? StiffnessWeft[ConstraintIndex] : StiffnessWeftNoMap;
 							const FSolverReal ExpStiffnessBiasValue = StiffnessBiasHasWeightMap ? StiffnessBias[ConstraintIndex] : StiffnessBiasNoMap;
 							const FSolverReal DampingRatioValue = DampingHasWeightMap ? DampingRatio[ConstraintIndex] : DampingNoMap;
@@ -1450,7 +1479,7 @@ void FXPBDAnisotropicAxialSpringConstraints::Apply(SolverParticlesOrRange& Parti
 							constexpr bool bSingleLambda = false;
 							constexpr bool bSeparateStretch = true;
 							constexpr bool bDampingAfter = false;
-							const FSolverReal ExpStiffnessValue = StiffnessHasWeightMap ? Stiffness[ConstraintIndex] : StiffnessNoMap;
+							const FSolverReal ExpStiffnessValue = StiffnessHasWeightMap ? StiffnessWarp[ConstraintIndex] : StiffnessNoMap;
 							const FSolverReal ExpStiffnessWeftValue = StiffnessWeftHasWeightMap ? StiffnessWeft[ConstraintIndex] : StiffnessWeftNoMap;
 							const FSolverReal ExpStiffnessBiasValue = StiffnessBiasHasWeightMap ? StiffnessBias[ConstraintIndex] : StiffnessBiasNoMap;
 							const FSolverReal DampingRatioValue = DampingHasWeightMap ? DampingRatio[ConstraintIndex] : DampingNoMap;
@@ -1473,7 +1502,7 @@ void FXPBDAnisotropicAxialSpringConstraints::Apply(SolverParticlesOrRange& Parti
 							constexpr bool bSingleLambda = true;
 							constexpr bool bSeparateStretch = false;
 							constexpr bool bDampingAfter = false;
-							const FSolverReal ExpStiffnessValue = StiffnessHasWeightMap ? Stiffness[ConstraintIndex] : StiffnessNoMap;
+							const FSolverReal ExpStiffnessValue = StiffnessHasWeightMap ? StiffnessWarp[ConstraintIndex] : StiffnessNoMap;
 							const FSolverReal ExpStiffnessWeftValue = StiffnessWeftHasWeightMap ? StiffnessWeft[ConstraintIndex] : StiffnessWeftNoMap;
 							const FSolverReal ExpStiffnessBiasValue = StiffnessBiasHasWeightMap ? StiffnessBias[ConstraintIndex] : StiffnessBiasNoMap;
 							const FSolverReal DampingRatioValue = DampingHasWeightMap ? DampingRatio[ConstraintIndex] : DampingNoMap;
@@ -1496,7 +1525,7 @@ void FXPBDAnisotropicAxialSpringConstraints::Apply(SolverParticlesOrRange& Parti
 							constexpr bool bSingleLambda = false;
 							constexpr bool bSeparateStretch = true;
 							constexpr bool bDampingAfter = false;
-							const FSolverReal ExpStiffnessValue = StiffnessHasWeightMap ? Stiffness[ConstraintIndex] : StiffnessNoMap;
+							const FSolverReal ExpStiffnessValue = StiffnessHasWeightMap ? StiffnessWarp[ConstraintIndex] : StiffnessNoMap;
 							const FSolverReal ExpStiffnessWeftValue = StiffnessWeftHasWeightMap ? StiffnessWeft[ConstraintIndex] : StiffnessWeftNoMap;
 							const FSolverReal ExpStiffnessBiasValue = StiffnessBiasHasWeightMap ? StiffnessBias[ConstraintIndex] : StiffnessBiasNoMap;
 							const FSolverReal DampingRatioValue = DampingHasWeightMap ? DampingRatio[ConstraintIndex] : DampingNoMap;
@@ -1519,7 +1548,7 @@ void FXPBDAnisotropicAxialSpringConstraints::Apply(SolverParticlesOrRange& Parti
 							constexpr bool bSingleLambda = false;
 							constexpr bool bSeparateStretch = true;
 							constexpr bool bDampingAfter = true;
-							const FSolverReal ExpStiffnessValue = StiffnessHasWeightMap ? Stiffness[ConstraintIndex] : StiffnessNoMap;
+							const FSolverReal ExpStiffnessValue = StiffnessHasWeightMap ? StiffnessWarp[ConstraintIndex] : StiffnessNoMap;
 							const FSolverReal ExpStiffnessWeftValue = StiffnessWeftHasWeightMap ? StiffnessWeft[ConstraintIndex] : StiffnessWeftNoMap;
 							const FSolverReal ExpStiffnessBiasValue = StiffnessBiasHasWeightMap ? StiffnessBias[ConstraintIndex] : StiffnessBiasNoMap;
 							const FSolverReal DampingRatioValue = DampingHasWeightMap ? DampingRatio[ConstraintIndex] : DampingNoMap;
@@ -1542,7 +1571,7 @@ void FXPBDAnisotropicAxialSpringConstraints::Apply(SolverParticlesOrRange& Parti
 							constexpr bool bSingleLambda = false;
 							constexpr bool bSeparateStretch = false;
 							constexpr bool bDampingAfter = true;
-							const FSolverReal ExpStiffnessValue = StiffnessHasWeightMap ? Stiffness[ConstraintIndex] : StiffnessNoMap;
+							const FSolverReal ExpStiffnessValue = StiffnessHasWeightMap ? StiffnessWarp[ConstraintIndex] : StiffnessNoMap;
 							const FSolverReal ExpStiffnessWeftValue = StiffnessWeftHasWeightMap ? StiffnessWeft[ConstraintIndex] : StiffnessWeftNoMap;
 							const FSolverReal ExpStiffnessBiasValue = StiffnessBiasHasWeightMap ? StiffnessBias[ConstraintIndex] : StiffnessBiasNoMap;
 							const FSolverReal DampingRatioValue = DampingHasWeightMap ? DampingRatio[ConstraintIndex] : DampingNoMap;
@@ -1565,7 +1594,7 @@ void FXPBDAnisotropicAxialSpringConstraints::Apply(SolverParticlesOrRange& Parti
 						constexpr bool bSingleLambda = false;
 						constexpr bool bSeparateStretch = true;
 						constexpr bool bDampingAfter = false;
-						const FSolverReal ExpStiffnessValue = StiffnessHasWeightMap ? Stiffness[ConstraintIndex] : StiffnessNoMap;
+						const FSolverReal ExpStiffnessValue = StiffnessHasWeightMap ? StiffnessWarp[ConstraintIndex] : StiffnessNoMap;
 						const FSolverReal ExpStiffnessWeftValue = StiffnessWeftHasWeightMap ? StiffnessWeft[ConstraintIndex] : StiffnessWeftNoMap;
 						const FSolverReal ExpStiffnessBiasValue = StiffnessBiasHasWeightMap ? StiffnessBias[ConstraintIndex] : StiffnessBiasNoMap;
 						ApplyHelper<bDampingBefore, bSingleLambda, bSeparateStretch, bDampingAfter>(Particles, Dt, ConstraintIndex, FSolverVec3(ExpStiffnessWeftValue, ExpStiffnessValue, ExpStiffnessBiasValue), (FSolverReal)0.);
@@ -1577,7 +1606,7 @@ void FXPBDAnisotropicAxialSpringConstraints::Apply(SolverParticlesOrRange& Parti
 	else
 	{
 		// Single threaded
-		const FSolverReal StiffnessNoMap = (FSolverReal)Stiffness;
+		const FSolverReal StiffnessNoMap = (FSolverReal)StiffnessWarp;
 		const FSolverReal StiffnessWeftNoMap = (FSolverReal)StiffnessWeft;
 		const FSolverReal StiffnessBiasNoMap = (FSolverReal)StiffnessBias;
 		const FSolverReal DampingNoMap = (FSolverReal)DampingRatio;
@@ -1591,7 +1620,7 @@ void FXPBDAnisotropicAxialSpringConstraints::Apply(SolverParticlesOrRange& Parti
 					constexpr bool bSingleLambda = false;
 					constexpr bool bSeparateStretch = false;
 					constexpr bool bDampingAfter = false;
-					const FSolverReal ExpStiffnessValue = StiffnessHasWeightMap ? Stiffness[ConstraintIndex] : StiffnessNoMap;
+					const FSolverReal ExpStiffnessValue = StiffnessHasWeightMap ? StiffnessWarp[ConstraintIndex] : StiffnessNoMap;
 					const FSolverReal ExpStiffnessWeftValue = StiffnessWeftHasWeightMap ? StiffnessWeft[ConstraintIndex] : StiffnessWeftNoMap;
 					const FSolverReal ExpStiffnessBiasValue = StiffnessBiasHasWeightMap ? StiffnessBias[ConstraintIndex] : StiffnessBiasNoMap;
 					const FSolverReal DampingRatioValue = DampingHasWeightMap ? DampingRatio[ConstraintIndex] : DampingNoMap;
@@ -1613,7 +1642,7 @@ void FXPBDAnisotropicAxialSpringConstraints::Apply(SolverParticlesOrRange& Parti
 					constexpr bool bSingleLambda = false;
 					constexpr bool bSeparateStretch = true;
 					constexpr bool bDampingAfter = false;
-					const FSolverReal ExpStiffnessValue = StiffnessHasWeightMap ? Stiffness[ConstraintIndex] : StiffnessNoMap;
+					const FSolverReal ExpStiffnessValue = StiffnessHasWeightMap ? StiffnessWarp[ConstraintIndex] : StiffnessNoMap;
 					const FSolverReal ExpStiffnessWeftValue = StiffnessWeftHasWeightMap ? StiffnessWeft[ConstraintIndex] : StiffnessWeftNoMap;
 					const FSolverReal ExpStiffnessBiasValue = StiffnessBiasHasWeightMap ? StiffnessBias[ConstraintIndex] : StiffnessBiasNoMap;
 					const FSolverReal DampingRatioValue = DampingHasWeightMap ? DampingRatio[ConstraintIndex] : DampingNoMap;
@@ -1630,7 +1659,7 @@ void FXPBDAnisotropicAxialSpringConstraints::Apply(SolverParticlesOrRange& Parti
 					constexpr bool bSingleLambda = true;
 					constexpr bool bSeparateStretch = false;
 					constexpr bool bDampingAfter = false;
-					const FSolverReal ExpStiffnessValue = StiffnessHasWeightMap ? Stiffness[ConstraintIndex] : StiffnessNoMap;
+					const FSolverReal ExpStiffnessValue = StiffnessHasWeightMap ? StiffnessWarp[ConstraintIndex] : StiffnessNoMap;
 					const FSolverReal ExpStiffnessWeftValue = StiffnessWeftHasWeightMap ? StiffnessWeft[ConstraintIndex] : StiffnessWeftNoMap;
 					const FSolverReal ExpStiffnessBiasValue = StiffnessBiasHasWeightMap ? StiffnessBias[ConstraintIndex] : StiffnessBiasNoMap;
 					const FSolverReal DampingRatioValue = DampingHasWeightMap ? DampingRatio[ConstraintIndex] : DampingNoMap;
@@ -1647,7 +1676,7 @@ void FXPBDAnisotropicAxialSpringConstraints::Apply(SolverParticlesOrRange& Parti
 					constexpr bool bSingleLambda = false;
 					constexpr bool bSeparateStretch = true;
 					constexpr bool bDampingAfter = false;
-					const FSolverReal ExpStiffnessValue = StiffnessHasWeightMap ? Stiffness[ConstraintIndex] : StiffnessNoMap;
+					const FSolverReal ExpStiffnessValue = StiffnessHasWeightMap ? StiffnessWarp[ConstraintIndex] : StiffnessNoMap;
 					const FSolverReal ExpStiffnessWeftValue = StiffnessWeftHasWeightMap ? StiffnessWeft[ConstraintIndex] : StiffnessWeftNoMap;
 					const FSolverReal ExpStiffnessBiasValue = StiffnessBiasHasWeightMap ? StiffnessBias[ConstraintIndex] : StiffnessBiasNoMap;
 					const FSolverReal DampingRatioValue = DampingHasWeightMap ? DampingRatio[ConstraintIndex] : DampingNoMap;
@@ -1664,7 +1693,7 @@ void FXPBDAnisotropicAxialSpringConstraints::Apply(SolverParticlesOrRange& Parti
 					constexpr bool bSingleLambda = false;
 					constexpr bool bSeparateStretch = true;
 					constexpr bool bDampingAfter = true;
-					const FSolverReal ExpStiffnessValue = StiffnessHasWeightMap ? Stiffness[ConstraintIndex] : StiffnessNoMap;
+					const FSolverReal ExpStiffnessValue = StiffnessHasWeightMap ? StiffnessWarp[ConstraintIndex] : StiffnessNoMap;
 					const FSolverReal ExpStiffnessWeftValue = StiffnessWeftHasWeightMap ? StiffnessWeft[ConstraintIndex] : StiffnessWeftNoMap;
 					const FSolverReal ExpStiffnessBiasValue = StiffnessBiasHasWeightMap ? StiffnessBias[ConstraintIndex] : StiffnessBiasNoMap;
 					const FSolverReal DampingRatioValue = DampingHasWeightMap ? DampingRatio[ConstraintIndex] : DampingNoMap;
@@ -1681,7 +1710,7 @@ void FXPBDAnisotropicAxialSpringConstraints::Apply(SolverParticlesOrRange& Parti
 					constexpr bool bSingleLambda = false;
 					constexpr bool bSeparateStretch = false;
 					constexpr bool bDampingAfter = true;
-					const FSolverReal ExpStiffnessValue = StiffnessHasWeightMap ? Stiffness[ConstraintIndex] : StiffnessNoMap;
+					const FSolverReal ExpStiffnessValue = StiffnessHasWeightMap ? StiffnessWarp[ConstraintIndex] : StiffnessNoMap;
 					const FSolverReal ExpStiffnessWeftValue = StiffnessWeftHasWeightMap ? StiffnessWeft[ConstraintIndex] : StiffnessWeftNoMap;
 					const FSolverReal ExpStiffnessBiasValue = StiffnessBiasHasWeightMap ? StiffnessBias[ConstraintIndex] : StiffnessBiasNoMap;
 					const FSolverReal DampingRatioValue = DampingHasWeightMap ? DampingRatio[ConstraintIndex] : DampingNoMap;
@@ -1698,7 +1727,7 @@ void FXPBDAnisotropicAxialSpringConstraints::Apply(SolverParticlesOrRange& Parti
 				constexpr bool bSingleLambda = false;
 				constexpr bool bSeparateStretch = true;
 				constexpr bool bDampingAfter = false;
-				const FSolverReal ExpStiffnessValue = StiffnessHasWeightMap ? Stiffness[ConstraintIndex] : StiffnessNoMap;
+				const FSolverReal ExpStiffnessValue = StiffnessHasWeightMap ? StiffnessWarp[ConstraintIndex] : StiffnessNoMap;
 				const FSolverReal ExpStiffnessWeftValue = StiffnessWeftHasWeightMap ? StiffnessWeft[ConstraintIndex] : StiffnessWeftNoMap;
 				const FSolverReal ExpStiffnessBiasValue = StiffnessBiasHasWeightMap ? StiffnessBias[ConstraintIndex] : StiffnessBiasNoMap;
 				ApplyHelper<bDampingBefore, bSingleLambda, bSeparateStretch, bDampingAfter>(Particles, Dt, ConstraintIndex, FSolverVec3(ExpStiffnessWeftValue, ExpStiffnessValue, ExpStiffnessBiasValue), (FSolverReal)0.);
@@ -1715,7 +1744,7 @@ void FXPBDAnisotropicAxialSpringConstraints::UpdateLinearSystem(const FSolverPar
 	TRACE_CPUPROFILER_EVENT_SCOPE(FXPBDAnisotropicAxialSpringConstraints_UpdateLinearSystem);
 	LinearSystem.ReserveForParallelAdd(Constraints.Num() * 3, Constraints.Num() * 2);
 
-	const bool StiffnessHasWeightMap = Stiffness.HasWeightMap();
+	const bool StiffnessHasWeightMap = StiffnessWarp.HasWeightMap();
 	const bool StiffnessWeftHasWeightMap = StiffnessWeft.HasWeightMap();
 	const bool StiffnessBiasHasWeightMap = StiffnessBias.HasWeightMap();
 	const bool DampingHasWeightMap = DampingRatio.HasWeightMap();
@@ -1724,7 +1753,7 @@ void FXPBDAnisotropicAxialSpringConstraints::UpdateLinearSystem(const FSolverPar
 		const int32 ConstraintColorNum = ConstraintsPerColorStartIndex.Num() - 1;
 		if (!StiffnessHasWeightMap && !StiffnessWeftHasWeightMap && !StiffnessBiasHasWeightMap && !DampingHasWeightMap)
 		{
-			const FSolverVec3 ExpStiffnessValue((FSolverReal)StiffnessWeft, (FSolverReal)Stiffness, (FSolverReal)StiffnessBias);
+			const FSolverVec3 ExpStiffnessValue((FSolverReal)StiffnessWeft, (FSolverReal)StiffnessWarp, (FSolverReal)StiffnessBias);
 			const FSolverReal DampingRatioValue = (FSolverReal)DampingRatio;
 			if (ExpStiffnessValue.Max() < MinStiffness)
 			{
@@ -1745,7 +1774,7 @@ void FXPBDAnisotropicAxialSpringConstraints::UpdateLinearSystem(const FSolverPar
 		}
 		else // Has weight maps
 		{
-			const FSolverReal StiffnessNoMap = (FSolverReal)Stiffness;
+			const FSolverReal StiffnessNoMap = (FSolverReal)StiffnessWarp;
 			const FSolverReal StiffnessWeftNoMap = (FSolverReal)StiffnessWeft;
 			const FSolverReal StiffnessBiasNoMap = (FSolverReal)StiffnessBias;
 			const FSolverReal DampingNoMap = (FSolverReal)DampingRatio;
@@ -1757,7 +1786,7 @@ void FXPBDAnisotropicAxialSpringConstraints::UpdateLinearSystem(const FSolverPar
 					StiffnessBiasNoMap, DampingHasWeightMap, DampingNoMap](const int32 Index)
 				{
 					const int32 ConstraintIndex = ColorStart + Index;
-					const FSolverReal ExpStiffnessValue = StiffnessHasWeightMap ? Stiffness[ConstraintIndex] : StiffnessNoMap;
+					const FSolverReal ExpStiffnessValue = StiffnessHasWeightMap ? StiffnessWarp[ConstraintIndex] : StiffnessNoMap;
 					const FSolverReal ExpStiffnessWeftValue = StiffnessWeftHasWeightMap ? StiffnessWeft[ConstraintIndex] : StiffnessWeftNoMap;
 					const FSolverReal ExpStiffnessBiasValue = StiffnessBiasHasWeightMap ? StiffnessBias[ConstraintIndex] : StiffnessBiasNoMap;
 					const FSolverReal DampingRatioValue = DampingHasWeightMap ? DampingRatio[ConstraintIndex] : DampingNoMap;
@@ -1772,7 +1801,7 @@ void FXPBDAnisotropicAxialSpringConstraints::UpdateLinearSystem(const FSolverPar
 	{
 		if (!StiffnessHasWeightMap && !StiffnessWeftHasWeightMap && !StiffnessBiasHasWeightMap && !DampingHasWeightMap)
 		{
-			const FSolverVec3 ExpStiffnessValue((FSolverReal)StiffnessWeft, (FSolverReal)Stiffness, (FSolverReal)StiffnessBias);
+			const FSolverVec3 ExpStiffnessValue((FSolverReal)StiffnessWeft, (FSolverReal)StiffnessWarp, (FSolverReal)StiffnessBias);
 			const FSolverReal DampingRatioValue = (FSolverReal)DampingRatio;
 			if (ExpStiffnessValue.Max() < MinStiffness)
 			{
@@ -1787,13 +1816,13 @@ void FXPBDAnisotropicAxialSpringConstraints::UpdateLinearSystem(const FSolverPar
 		}
 		else
 		{
-			const FSolverReal StiffnessNoMap = (FSolverReal)Stiffness;
+			const FSolverReal StiffnessNoMap = (FSolverReal)StiffnessWarp;
 			const FSolverReal StiffnessWeftNoMap = (FSolverReal)StiffnessWeft;
 			const FSolverReal StiffnessBiasNoMap = (FSolverReal)StiffnessBias;
 			const FSolverReal DampingNoMap = (FSolverReal)DampingRatio;
 			for (int32 ConstraintIndex = 0; ConstraintIndex < Constraints.Num(); ++ConstraintIndex)
 			{
-				const FSolverReal ExpStiffnessValue = StiffnessHasWeightMap ? Stiffness[ConstraintIndex] : StiffnessNoMap;
+				const FSolverReal ExpStiffnessValue = StiffnessHasWeightMap ? StiffnessWarp[ConstraintIndex] : StiffnessNoMap;
 				const FSolverReal ExpStiffnessWeftValue = StiffnessWeftHasWeightMap ? StiffnessWeft[ConstraintIndex] : StiffnessWeftNoMap;
 				const FSolverReal ExpStiffnessBiasValue = StiffnessBiasHasWeightMap ? StiffnessBias[ConstraintIndex] : StiffnessBiasNoMap;
 				const FSolverReal DampingRatioValue = DampingHasWeightMap ? DampingRatio[ConstraintIndex] : DampingNoMap;
@@ -1815,29 +1844,23 @@ void FXPBDAnisotropicSpringConstraints::SetProperties(
 		if (IsXPBDAnisoSpringStiffnessWarpStringDirty(PropertyCollection))
 		{
 			const FString& WeightMapName = GetXPBDAnisoSpringStiffnessWarpString(PropertyCollection);
-			EdgeConstraints.Stiffness = FPBDStiffness(
-				WeightedValue,
+			EdgeConstraints.StiffnessWarp = FPBDFlatWeightMap(
+				WeightedValue.ClampAxes(0, FXPBDAnisotropicEdgeSpringConstraints::MaxStiffness),
 				WeightMaps.FindRef(WeightMapName),
 				TConstArrayView<TVec2<int32>>(EdgeConstraints.Constraints),
 				EdgeConstraints.ParticleOffset,
-				EdgeConstraints.ParticleCount,
-				FPBDStiffness::DefaultTableSize,
-				FPBDStiffness::DefaultParameterFitBase,
-				FXPBDAnisotropicEdgeSpringConstraints::MaxStiffness);
-			AxialConstraints.Stiffness = FPBDStiffness(
-				WeightedValue,
+				EdgeConstraints.ParticleCount);
+			AxialConstraints.StiffnessWarp = FPBDFlatWeightMap(
+				WeightedValue.ClampAxes(0, FXPBDAnisotropicAxialSpringConstraints::MaxStiffness),
 				WeightMaps.FindRef(WeightMapName),
 				TConstArrayView<TVec3<int32>>(AxialConstraints.Constraints),
 				AxialConstraints.ParticleOffset,
-				AxialConstraints.ParticleCount,
-				FPBDStiffness::DefaultTableSize,
-				FPBDStiffness::DefaultParameterFitBase,
-				FXPBDAnisotropicEdgeSpringConstraints::MaxStiffness);
+				AxialConstraints.ParticleCount);
 		}
 		else
 		{
-			EdgeConstraints.Stiffness.SetWeightedValue(WeightedValue, FXPBDAnisotropicEdgeSpringConstraints::MaxStiffness);
-			AxialConstraints.Stiffness.SetWeightedValue(WeightedValue, FXPBDAnisotropicAxialSpringConstraints::MaxStiffness);
+			EdgeConstraints.StiffnessWarp.SetWeightedValue(WeightedValue.ClampAxes(0, FXPBDAnisotropicEdgeSpringConstraints::MaxStiffness));
+			AxialConstraints.StiffnessWarp.SetWeightedValue(WeightedValue.ClampAxes(0, FXPBDAnisotropicAxialSpringConstraints::MaxStiffness));
 		}
 	}
 	if (IsXPBDAnisoSpringStiffnessWeftMutable(PropertyCollection))
@@ -1846,29 +1869,23 @@ void FXPBDAnisotropicSpringConstraints::SetProperties(
 		if (IsXPBDAnisoSpringStiffnessWeftStringDirty(PropertyCollection))
 		{
 			const FString& WeightMapName = GetXPBDAnisoSpringStiffnessWeftString(PropertyCollection);
-			EdgeConstraints.StiffnessWeft = FPBDStiffness(
-				WeightedValue,
+			EdgeConstraints.StiffnessWeft = FPBDFlatWeightMap(
+				WeightedValue.ClampAxes(0, FXPBDAnisotropicEdgeSpringConstraints::MaxStiffness),
 				WeightMaps.FindRef(WeightMapName),
 				TConstArrayView<TVec2<int32>>(EdgeConstraints.Constraints),
 				EdgeConstraints.ParticleOffset,
-				EdgeConstraints.ParticleCount,
-				FPBDStiffness::DefaultTableSize,
-				FPBDStiffness::DefaultParameterFitBase,
-				FXPBDAnisotropicEdgeSpringConstraints::MaxStiffness);
-			AxialConstraints.StiffnessWeft = FPBDStiffness(
-				WeightedValue,
+				EdgeConstraints.ParticleCount);
+			AxialConstraints.StiffnessWeft = FPBDFlatWeightMap(
+				WeightedValue.ClampAxes(0, FXPBDAnisotropicAxialSpringConstraints::MaxStiffness),
 				WeightMaps.FindRef(WeightMapName),
 				TConstArrayView<TVec3<int32>>(AxialConstraints.Constraints),
 				AxialConstraints.ParticleOffset,
-				AxialConstraints.ParticleCount,
-				FPBDStiffness::DefaultTableSize,
-				FPBDStiffness::DefaultParameterFitBase,
-				FXPBDAnisotropicEdgeSpringConstraints::MaxStiffness);
+				AxialConstraints.ParticleCount);
 		}
 		else
 		{
-			EdgeConstraints.StiffnessWeft.SetWeightedValue(WeightedValue, FXPBDAnisotropicEdgeSpringConstraints::MaxStiffness);
-			AxialConstraints.StiffnessWeft.SetWeightedValue(WeightedValue, FXPBDAnisotropicAxialSpringConstraints::MaxStiffness);
+			EdgeConstraints.StiffnessWeft.SetWeightedValue(WeightedValue.ClampAxes(0, FXPBDAnisotropicEdgeSpringConstraints::MaxStiffness));
+			AxialConstraints.StiffnessWeft.SetWeightedValue(WeightedValue.ClampAxes(0, FXPBDAnisotropicAxialSpringConstraints::MaxStiffness));
 		}
 	}
 	if (IsXPBDAnisoSpringStiffnessBiasMutable(PropertyCollection))
@@ -1877,29 +1894,23 @@ void FXPBDAnisotropicSpringConstraints::SetProperties(
 		if (IsXPBDAnisoSpringStiffnessBiasStringDirty(PropertyCollection))
 		{
 			const FString& WeightMapName = GetXPBDAnisoSpringStiffnessBiasString(PropertyCollection);
-			EdgeConstraints.StiffnessBias = FPBDStiffness(
-				WeightedValue,
+			EdgeConstraints.StiffnessBias = FPBDFlatWeightMap(
+				WeightedValue.ClampAxes(0, FXPBDAnisotropicEdgeSpringConstraints::MaxStiffness),
 				WeightMaps.FindRef(WeightMapName),
 				TConstArrayView<TVec2<int32>>(EdgeConstraints.Constraints),
 				EdgeConstraints.ParticleOffset,
-				EdgeConstraints.ParticleCount,
-				FPBDStiffness::DefaultTableSize,
-				FPBDStiffness::DefaultParameterFitBase,
-				FXPBDAnisotropicEdgeSpringConstraints::MaxStiffness);
-			AxialConstraints.StiffnessBias = FPBDStiffness(
-				WeightedValue,
+				EdgeConstraints.ParticleCount);
+			AxialConstraints.StiffnessBias = FPBDFlatWeightMap(
+				WeightedValue.ClampAxes(0, FXPBDAnisotropicAxialSpringConstraints::MaxStiffness),
 				WeightMaps.FindRef(WeightMapName),
 				TConstArrayView<TVec3<int32>>(AxialConstraints.Constraints),
 				AxialConstraints.ParticleOffset,
-				AxialConstraints.ParticleCount,
-				FPBDStiffness::DefaultTableSize,
-				FPBDStiffness::DefaultParameterFitBase,
-				FXPBDAnisotropicEdgeSpringConstraints::MaxStiffness);
+				AxialConstraints.ParticleCount);
 		}
 		else
 		{
-			EdgeConstraints.StiffnessBias.SetWeightedValue(WeightedValue, FXPBDAnisotropicEdgeSpringConstraints::MaxStiffness);
-			AxialConstraints.StiffnessBias.SetWeightedValue(WeightedValue, FXPBDAnisotropicAxialSpringConstraints::MaxStiffness);
+			EdgeConstraints.StiffnessBias.SetWeightedValue(WeightedValue.ClampAxes(0, FXPBDAnisotropicEdgeSpringConstraints::MaxStiffness));
+			AxialConstraints.StiffnessBias.SetWeightedValue(WeightedValue.ClampAxes(0, FXPBDAnisotropicAxialSpringConstraints::MaxStiffness));
 		}
 	}
 	if (IsXPBDAnisoSpringDampingMutable(PropertyCollection))
@@ -1908,13 +1919,13 @@ void FXPBDAnisotropicSpringConstraints::SetProperties(
 		if (IsXPBDAnisoSpringDampingStringDirty(PropertyCollection))
 		{
 			const FString& WeightMapName = GetXPBDAnisoSpringDampingString(PropertyCollection);
-			AxialConstraints.DampingRatio = FPBDWeightMap(
+			AxialConstraints.DampingRatio = FPBDFlatWeightMap(
 				WeightedValue,
 				WeightMaps.FindRef(WeightMapName),
 				TConstArrayView<TVec3<int32>>(AxialConstraints.Constraints),
 				AxialConstraints.ParticleOffset,
 				AxialConstraints.ParticleCount);
-			AxialConstraints.DampingRatio = FPBDWeightMap(
+			AxialConstraints.DampingRatio = FPBDFlatWeightMap(
 				WeightedValue,
 				WeightMaps.FindRef(WeightMapName),
 				TConstArrayView<TVec2<int32>>(EdgeConstraints.Constraints),
