Chaos Cloth optimizations to XPBD constraints.

diff --git a/Engine/Source/Runtime/Experimental/Chaos/Private/Chaos/XPBDAnisotropicSpringConstraints.ispc b/Engine/Source/Runtime/Experimental/Chaos/Private/Chaos/XPBDAnisotropicSpringConstraints.ispc
--- a/Engine/Source/Runtime/Experimental/Chaos/Private/Chaos/XPBDAnisotropicSpringConstraints.ispc
+++ b/Engine/Source/Runtime/Experimental/Chaos/Private/Chaos/XPBDAnisotropicSpringConstraints.ispc
@@ -31,22 +31,16 @@ export void ApplyXPBDAnisoSpringConstraints(uniform FVector4f PandInvM[],
 		
 		const varying float Dist = Dists[Index];
 		const varying float Lambda = Lambdas[Index];
-		const varying FVector3f WarpWeftBiasBaseMultiplier = VectorGather(&WarpWeftBiasBaseMultiplers[Index]);
+		const varying FVector3f WarpWeftBiasBaseMultiplier = VectorLoad(&WarpWeftBiasBaseMultiplers[extract(Index,0)]);
 		const varying float StiffnessValue = VectorDot(WarpWeftBiasBaseMultiplier, Stiffness);
 
 		varying float DLambda;
-		const varying FVector3f Delta = GetXPBDSpringDelta(P1, P2, InvM1, InvM2, Dt, Dist, Lambda, StiffnessValue, FLOAT_SMALL_NUMBER, DLambda);
+		const varying FVector3f Delta = GetXPBDSpringDelta(P1, P2, InvM1, InvM2, Dt, Dist, Lambda, StiffnessValue, DLambda);
 
 		Lambdas[Index] = Lambda + DLambda;
 
-		if (InvM1 > FLOAT_ZERO)
-		{
-			VectorScatter(&PandInvM[Index1], SetVector4( P1 + (InvM1 * Delta), InvM1 ));
-		}
-		if (InvM2 > FLOAT_ZERO)
-		{
-			VectorScatter(&PandInvM[Index2], SetVector4( P2 - (InvM2 * Delta), InvM2 ));
-		}
+		VectorScatter(&PandInvM[Index1], SetVector4( P1 + (InvM1 * Delta), InvM1 ));
+		VectorScatter(&PandInvM[Index2], SetVector4( P2 - (InvM2 * Delta), InvM2 ));
 	}
 }
 
@@ -80,22 +74,16 @@ export void ApplyXPBDAnisoSpringConstraintsWithDamping(uniform FVector4f PandInv
 
 		const varying float Dist = Dists[Index];
 		const varying float Lambda = Lambdas[Index];
-		const varying FVector3f WarpWeftBiasBaseMultiplier = VectorGather(&WarpWeftBiasBaseMultiplers[Index]);
+		const varying FVector3f WarpWeftBiasBaseMultiplier = VectorLoad(&WarpWeftBiasBaseMultiplers[extract(Index,0)]);
 		const varying float StiffnessValue = VectorDot(WarpWeftBiasBaseMultiplier, Stiffness);
 
 		varying float DLambda;
-		const FVector3f Delta = GetXPBDSpringDeltaWithDamping(P1, P2, X1, X2, InvM1, InvM2, Dt, Dist, Lambda, StiffnessValue, DampingRatio, FLOAT_SMALL_NUMBER, DLambda);
+		const FVector3f Delta = GetXPBDSpringDeltaWithDamping(P1, P2, X1, X2, InvM1, InvM2, Dt, Dist, Lambda, StiffnessValue, DampingRatio, DLambda);
 
 		Lambdas[Index] = Lambda + DLambda;
-
-		if (InvM1 > FLOAT_ZERO)
-		{
-			VectorScatter(&PandInvM[Index1], SetVector4( P1 + (InvM1 * Delta), InvM1 ));
-		}
-		if (InvM2 > FLOAT_ZERO)
-		{
-			VectorScatter(&PandInvM[Index2], SetVector4( P2 - (InvM2 * Delta), InvM2 ));
-		}
+		
+		VectorScatter(&PandInvM[Index1], SetVector4( P1 + (InvM1 * Delta), InvM1 ));
+		VectorScatter(&PandInvM[Index2], SetVector4( P2 - (InvM2 * Delta), InvM2 ));
 	}
 }
 
@@ -129,22 +117,16 @@ export void ApplyXPBDAnisoSpringDampingConstraints(uniform FVector4f PandInvM[],
 
 		const varying float Dist = Dists[Index];
 		const varying float Lambda = Lambdas[Index];
-		const varying FVector3f WarpWeftBiasBaseMultiplier = VectorGather(&WarpWeftBiasBaseMultiplers[Index]);
+		const varying FVector3f WarpWeftBiasBaseMultiplier = VectorLoad(&WarpWeftBiasBaseMultiplers[extract(Index,0)]);
 		const varying float StiffnessValue = VectorDot(WarpWeftBiasBaseMultiplier, Stiffness);
 
 		varying float DLambda;
-		const FVector3f Delta = GetXPBDSpringDampingDelta(P1, P2, X1, X2, InvM1, InvM2, Dt, Dist, Lambda, StiffnessValue, DampingRatio, FLOAT_SMALL_NUMBER, DLambda);
+		const FVector3f Delta = GetXPBDSpringDampingDelta(P1, P2, X1, X2, InvM1, InvM2, Dt, Dist, Lambda, StiffnessValue, DampingRatio, DLambda);
 
 		Lambdas[Index] = Lambda + DLambda;
-
-		if (InvM1 > FLOAT_ZERO)
-		{
-			VectorScatter(&PandInvM[Index1], SetVector4( P1 + (InvM1 * Delta), InvM1 ));
-		}
-		if (InvM2 > FLOAT_ZERO)
-		{
-			VectorScatter(&PandInvM[Index2], SetVector4( P2 - (InvM2 * Delta), InvM2 ));
-		}
+		
+		VectorScatter(&PandInvM[Index1], SetVector4( P1 + (InvM1 * Delta), InvM1 ));
+		VectorScatter(&PandInvM[Index2], SetVector4( P2 - (InvM2 * Delta), InvM2 ));
 	}
 }
 
@@ -154,16 +136,15 @@ export void ApplyXPBDAnisoSpringConstraintsWithWeightMaps(uniform FVector4f Pand
 									const uniform FVector3f WarpWeftBiasBaseMultiplers[],
 									uniform float Lambdas[],
 									const uniform float Dt,
-									const uniform float MinStiffness,
 									const uniform bool StiffnessHasMap,
-									const uniform uint8 StiffnessIndices[],
-									const uniform float StiffnessTable[],
+									const uniform FVector2f& StiffnessOffsetRange,
+									const uniform float StiffnessMapValues[],
 									const uniform bool StiffnessWeftHasMap,
-									const uniform uint8 StiffnessWeftIndices[],
-									const uniform float StiffnessWeftTable[],
+									const uniform FVector2f& StiffnessWeftOffsetRange,
+									const uniform float StiffnessWeftMapValues[],
 									const uniform bool StiffnessBiasHasMap,
-									const uniform uint8 StiffnessBiasIndices[],
-									const uniform float StiffnessBiasTable[],
+									const uniform FVector2f& StiffnessBiasOffsetRange,
+									const uniform float StiffnessBiasMapValues[],
 									const uniform int32 NumConstraints)
 {
 
@@ -172,10 +153,6 @@ export void ApplyXPBDAnisoSpringConstraintsWithWeightMaps(uniform FVector4f Pand
 		const varying FIntVector2 Constraint = VectorLoad(&Constraints[extract(Index,0)]);
 		const varying int32 Index1 = Constraint.V[0];
 		const varying int32 Index2 = Constraint.V[1];
-		
-		const varying int8 StiffnessIndex = StiffnessHasMap ? StiffnessIndices[Index] : 0;
-		const varying int8 StiffnessWeftIndex = StiffnessWeftHasMap ? StiffnessWeftIndices[Index] : 0;
-		const varying int8 StiffnessBiasIndex = StiffnessBiasHasMap ? StiffnessBiasIndices[Index] : 0;
 
 		const varying FVector4f PandInvM1 = VectorGather(&PandInvM[Index1]);
 		const varying FVector4f PandInvM2 = VectorGather(&PandInvM[Index2]);
@@ -187,31 +164,21 @@ export void ApplyXPBDAnisoSpringConstraintsWithWeightMaps(uniform FVector4f Pand
 
 		const varying float Dist = Dists[Index];
 		const varying float Lambda = Lambdas[Index];
-		
-		#pragma ignore warning(perf)
-		const varying float Stiffness = StiffnessTable[StiffnessIndex];
-		#pragma ignore warning(perf)
-		const varying float StiffnessWeft = StiffnessWeftTable[StiffnessWeftIndex];
-		#pragma ignore warning(perf)
-		const varying float StiffnessBias = StiffnessBiasTable[StiffnessBiasIndex];
+
+		const varying float Stiffness = StiffnessHasMap ? StiffnessOffsetRange.V[0] + StiffnessMapValues[Index]*StiffnessOffsetRange.V[1] : StiffnessOffsetRange.V[0];
+		const varying float StiffnessWeft = StiffnessWeftHasMap ? StiffnessWeftOffsetRange.V[0] + StiffnessWeftMapValues[Index]*StiffnessWeftOffsetRange.V[1] : StiffnessWeftOffsetRange.V[0];
+		const varying float StiffnessBias = StiffnessBiasHasMap ? StiffnessBiasOffsetRange.V[0] + StiffnessBiasMapValues[Index]*StiffnessBiasOffsetRange.V[1] : StiffnessBiasOffsetRange.V[0];
 		
 		const varying FVector3f StiffnessVector = {{StiffnessWeft, Stiffness, StiffnessBias}};
-		const varying FVector3f WarpWeftBiasBaseMultiplier = VectorGather(&WarpWeftBiasBaseMultiplers[Index]);
+		const varying FVector3f WarpWeftBiasBaseMultiplier = VectorLoad(&WarpWeftBiasBaseMultiplers[extract(Index,0)]);
 		const varying float StiffnessValue = VectorDot(WarpWeftBiasBaseMultiplier, StiffnessVector);
 
 		varying float DLambda;
-		const varying FVector3f Delta = GetXPBDSpringDelta(P1, P2, InvM1, InvM2, Dt, Dist, Lambda, StiffnessValue, MinStiffness, DLambda);
+		const varying FVector3f Delta = GetXPBDSpringDelta(P1, P2, InvM1, InvM2, Dt, Dist, Lambda, StiffnessValue, DLambda);
 
 		Lambdas[Index] = Lambda + DLambda;
-
-		if (InvM1 > FLOAT_ZERO)
-		{
-			VectorScatter(&PandInvM[Index1], SetVector4( P1 + (InvM1 * Delta), InvM1 ));
-		}
-		if (InvM2 > FLOAT_ZERO)
-		{
-			VectorScatter(&PandInvM[Index2], SetVector4( P2 - (InvM2 * Delta), InvM2 ));
-		}
+		VectorScatter(&PandInvM[Index1], SetVector4( P1 + (InvM1 * Delta), InvM1 ));
+		VectorScatter(&PandInvM[Index2], SetVector4( P2 - (InvM2 * Delta), InvM2 ));
 	}
 }
 
@@ -222,19 +189,18 @@ export void ApplyXPBDAnisoSpringConstraintsWithDampingAndWeightMaps(uniform FVec
 															const uniform FVector3f WarpWeftBiasBaseMultiplers[],
 															uniform float Lambdas[],
 															const uniform float Dt,
-															const uniform float MinStiffness,
 															const uniform bool StiffnessHasMap,
-															const uniform uint8 StiffnessIndices[],
-															const uniform float StiffnessTable[],
+															const uniform FVector2f& StiffnessOffsetRange,
+															const uniform float StiffnessMapValues[],
 															const uniform bool StiffnessWeftHasMap,
-															const uniform uint8 StiffnessWeftIndices[],
-															const uniform float StiffnessWeftTable[],
+															const uniform FVector2f& StiffnessWeftOffsetRange,
+															const uniform float StiffnessWeftMapValues[],
 															const uniform bool StiffnessBiasHasMap,
-															const uniform uint8 StiffnessBiasIndices[],
-															const uniform float StiffnessBiasTable[],
+															const uniform FVector2f& StiffnessBiasOffsetRange,
+															const uniform float StiffnessBiasMapValues[],
 															const uniform bool DampingHasMap,
-															const uniform uint8 DampingIndices[],
-															const uniform float DampingTable[],
+															const uniform FVector2f& DampingOffsetRange,
+															const uniform float DampingMapValues[],
 															const uniform int32 NumConstraints)
 {
 
@@ -243,10 +209,6 @@ export void ApplyXPBDAnisoSpringConstraintsWithDampingAndWeightMaps(uniform FVec
 		const varying FIntVector2 Constraint = VectorLoad(&Constraints[extract(Index,0)]);
 		const varying int32 Index1 = Constraint.V[0];
 		const varying int32 Index2 = Constraint.V[1];
-		const varying int8 StiffnessIndex = StiffnessHasMap ? StiffnessIndices[Index] : 0;
-		const varying int8 StiffnessWeftIndex = StiffnessWeftHasMap ? StiffnessWeftIndices[Index] : 0;
-		const varying int8 StiffnessBiasIndex = StiffnessBiasHasMap ? StiffnessBiasIndices[Index] : 0;
-		const varying int8 DampingIndex = DampingHasMap ? DampingIndices[Index] : 0;
 
 		const varying FVector4f PandInvM1 = VectorGather(&PandInvM[Index1]);
 		const varying FVector4f PandInvM2 = VectorGather(&PandInvM[Index2]);
@@ -261,33 +223,24 @@ export void ApplyXPBDAnisoSpringConstraintsWithDampingAndWeightMaps(uniform FVec
 
 		const varying float Dist = Dists[Index];
 		const varying float Lambda = Lambdas[Index];
-
-		#pragma ignore warning(perf)
-		const varying float Stiffness = StiffnessTable[StiffnessIndex];
-		#pragma ignore warning(perf)
-		const varying float StiffnessWeft = StiffnessWeftTable[StiffnessWeftIndex];
-		#pragma ignore warning(perf)
-		const varying float StiffnessBias = StiffnessBiasTable[StiffnessBiasIndex];
-		#pragma ignore warning(perf)
-		const varying float DampingRatio = DampingTable[DampingIndex];
+		
+		const varying float Stiffness = StiffnessHasMap ? StiffnessOffsetRange.V[0] + StiffnessMapValues[Index]*StiffnessOffsetRange.V[1] : StiffnessOffsetRange.V[0];
+		const varying float StiffnessWeft = StiffnessWeftHasMap ? StiffnessWeftOffsetRange.V[0] + StiffnessWeftMapValues[Index]*StiffnessWeftOffsetRange.V[1] : StiffnessWeftOffsetRange.V[0];
+		const varying float StiffnessBias = StiffnessBiasHasMap ? StiffnessBiasOffsetRange.V[0] + StiffnessBiasMapValues[Index]*StiffnessBiasOffsetRange.V[1] : StiffnessBiasOffsetRange.V[0];
 
 		const varying FVector3f StiffnessVector = {{StiffnessWeft, Stiffness, StiffnessBias}};
-		const varying FVector3f WarpWeftBiasBaseMultiplier = VectorGather(&WarpWeftBiasBaseMultiplers[Index]);
+		const varying FVector3f WarpWeftBiasBaseMultiplier = VectorLoad(&WarpWeftBiasBaseMultiplers[extract(Index,0)]);
 		const varying float StiffnessValue = VectorDot(WarpWeftBiasBaseMultiplier, StiffnessVector);
 
+		const varying float DampingRatio = DampingHasMap ? DampingOffsetRange.V[0] + DampingMapValues[Index]*DampingOffsetRange.V[1] : DampingOffsetRange.V[0];
+
 		varying float DLambda;
-		const FVector3f Delta = GetXPBDSpringDeltaWithDamping(P1, P2, X1, X2, InvM1, InvM2, Dt, Dist, Lambda, StiffnessValue, DampingRatio, MinStiffness, DLambda);
+		const FVector3f Delta = GetXPBDSpringDeltaWithDamping(P1, P2, X1, X2, InvM1, InvM2, Dt, Dist, Lambda, StiffnessValue, DampingRatio, DLambda);
 
 		Lambdas[Index] = Lambda + DLambda;
 
-		if (InvM1 > FLOAT_ZERO)
-		{
-			VectorScatter(&PandInvM[Index1], SetVector4( P1 + (InvM1 * Delta), InvM1 ));
-		}
-		if (InvM2 > FLOAT_ZERO)
-		{
-			VectorScatter(&PandInvM[Index2], SetVector4( P2 - (InvM2 * Delta), InvM2 ));
-		}
+		VectorScatter(&PandInvM[Index1], SetVector4( P1 + (InvM1 * Delta), InvM1 ));
+		VectorScatter(&PandInvM[Index2], SetVector4( P2 - (InvM2 * Delta), InvM2 ));
 	}
 }
 
@@ -298,19 +251,18 @@ export void ApplyXPBDAnisoSpringDampingConstraintsWithWeightMaps(uniform FVector
 															const uniform FVector3f WarpWeftBiasBaseMultiplers[],
 															uniform float Lambdas[],
 															const uniform float Dt,
-															const uniform float MinStiffness,
 															const uniform bool StiffnessHasMap,
-															const uniform uint8 StiffnessIndices[],
-															const uniform float StiffnessTable[],
+															const uniform FVector2f& StiffnessOffsetRange,
+															const uniform float StiffnessMapValues[],
 															const uniform bool StiffnessWeftHasMap,
-															const uniform uint8 StiffnessWeftIndices[],
-															const uniform float StiffnessWeftTable[],
+															const uniform FVector2f& StiffnessWeftOffsetRange,
+															const uniform float StiffnessWeftMapValues[],
 															const uniform bool StiffnessBiasHasMap,
-															const uniform uint8 StiffnessBiasIndices[],
-															const uniform float StiffnessBiasTable[],
+															const uniform FVector2f& StiffnessBiasOffsetRange,
+															const uniform float StiffnessBiasMapValues[],
 															const uniform bool DampingHasMap,
-															const uniform uint8 DampingIndices[],
-															const uniform float DampingTable[],
+															const uniform FVector2f& DampingOffsetRange,
+															const uniform float DampingMapValues[],
 															const uniform int32 NumConstraints)
 {
 
@@ -319,10 +271,6 @@ export void ApplyXPBDAnisoSpringDampingConstraintsWithWeightMaps(uniform FVector
 		const varying FIntVector2 Constraint = VectorLoad(&Constraints[extract(Index,0)]);
 		const varying int32 Index1 = Constraint.V[0];
 		const varying int32 Index2 = Constraint.V[1];
-		const varying int8 StiffnessIndex = StiffnessHasMap ? StiffnessIndices[Index] : 0;
-		const varying int8 StiffnessWeftIndex = StiffnessWeftHasMap ? StiffnessWeftIndices[Index] : 0;
-		const varying int8 StiffnessBiasIndex = StiffnessBiasHasMap ? StiffnessBiasIndices[Index] : 0;
-		const varying int8 DampingIndex = DampingHasMap ? DampingIndices[Index] : 0;
 
 		const varying FVector4f PandInvM1 = VectorGather(&PandInvM[Index1]);
 		const varying FVector4f PandInvM2 = VectorGather(&PandInvM[Index2]);
@@ -338,32 +286,23 @@ export void ApplyXPBDAnisoSpringDampingConstraintsWithWeightMaps(uniform FVector
 		const varying float Dist = Dists[Index];
 		const varying float Lambda = Lambdas[Index];
 
-		#pragma ignore warning(perf)
-		const varying float Stiffness = StiffnessTable[StiffnessIndex];
-		#pragma ignore warning(perf)
-		const varying float StiffnessWeft = StiffnessWeftTable[StiffnessWeftIndex];
-		#pragma ignore warning(perf)
-		const varying float StiffnessBias = StiffnessBiasTable[StiffnessBiasIndex];
-		#pragma ignore warning(perf)
-		const varying float DampingRatio = DampingTable[DampingIndex];
+		const varying float Stiffness = StiffnessHasMap ? StiffnessOffsetRange.V[0] + StiffnessMapValues[Index]*StiffnessOffsetRange.V[1] : StiffnessOffsetRange.V[0];
+		const varying float StiffnessWeft = StiffnessWeftHasMap ? StiffnessWeftOffsetRange.V[0] + StiffnessWeftMapValues[Index]*StiffnessWeftOffsetRange.V[1] : StiffnessWeftOffsetRange.V[0];
+		const varying float StiffnessBias = StiffnessBiasHasMap ? StiffnessBiasOffsetRange.V[0] + StiffnessBiasMapValues[Index]*StiffnessBiasOffsetRange.V[1] : StiffnessBiasOffsetRange.V[0];
 
 		const varying FVector3f StiffnessVector = {{StiffnessWeft, Stiffness, StiffnessBias}};
-		const varying FVector3f WarpWeftBiasBaseMultiplier = VectorGather(&WarpWeftBiasBaseMultiplers[Index]);
+		const varying FVector3f WarpWeftBiasBaseMultiplier = VectorLoad(&WarpWeftBiasBaseMultiplers[extract(Index,0)]);
 		const varying float StiffnessValue = VectorDot(WarpWeftBiasBaseMultiplier, StiffnessVector);
+		
+		const varying float DampingRatio = DampingHasMap ? DampingOffsetRange.V[0] + DampingMapValues[Index]*DampingOffsetRange.V[1] : DampingOffsetRange.V[0];
 
 		varying float DLambda;
-		const FVector3f Delta = GetXPBDSpringDampingDelta(P1, P2, X1, X2, InvM1, InvM2, Dt, Dist, Lambda, StiffnessValue, DampingRatio, MinStiffness, DLambda);
+		const FVector3f Delta = GetXPBDSpringDampingDelta(P1, P2, X1, X2, InvM1, InvM2, Dt, Dist, Lambda, StiffnessValue, DampingRatio, DLambda);
 
 		Lambdas[Index] = Lambda + DLambda;
-
-		if (InvM1 > FLOAT_ZERO)
-		{
-			VectorScatter(&PandInvM[Index1], SetVector4( P1 + (InvM1 * Delta), InvM1 ));
-		}
-		if (InvM2 > FLOAT_ZERO)
-		{
-			VectorScatter(&PandInvM[Index2], SetVector4( P2 - (InvM2 * Delta), InvM2 ));
-		}
+		
+		VectorScatter(&PandInvM[Index1], SetVector4( P1 + (InvM1 * Delta), InvM1 ));
+		VectorScatter(&PandInvM[Index2], SetVector4( P2 - (InvM2 * Delta), InvM2 ));
 	}
 }
 
@@ -397,26 +336,17 @@ export void ApplyXPBDAnisoAxialSpringConstraints(uniform FVector4f PandInvM[],
 		const varying float Bary = Barys[Index];
 		const varying float Dist = Dists[Index];
 		const varying float Lambda = Lambdas[Index];
-		const varying FVector3f WarpWeftBiasBaseMultiplier = VectorGather(&WarpWeftBiasBaseMultiplers[Index]);
+		const varying FVector3f WarpWeftBiasBaseMultiplier = VectorLoad(&WarpWeftBiasBaseMultiplers[extract(Index,0)]);
 		const varying float StiffnessValue = VectorDot(WarpWeftBiasBaseMultiplier, Stiffness);
 
 		varying float DLambda;
-		const varying FVector3f Delta = GetXPBDAxialSpringDelta(P1, P2, P3, InvM1, InvM2, InvM3, Dt, Bary, Dist, Lambda, StiffnessValue, FLOAT_SMALL_NUMBER, DLambda);
+		const varying FVector3f Delta = GetXPBDAxialSpringDelta(P1, P2, P3, InvM1, InvM2, InvM3, Dt, Bary, Dist, Lambda, StiffnessValue, DLambda);
 
 		Lambdas[Index] = Lambda + DLambda;
 
-		if (InvM1 > FLOAT_ZERO)
-		{
-			VectorScatter(&PandInvM[Index1], SetVector4( P1 + (InvM1 * Delta), InvM1 ));
-		}
-		if (InvM2 > FLOAT_ZERO)
-		{
-			VectorScatter(&PandInvM[Index2], SetVector4( P2 - (InvM2 * Bary * Delta), InvM2 ));
-		}
-		if (InvM3 > FLOAT_ZERO)
-		{
-			VectorScatter(&PandInvM[Index3], SetVector4( P3 - (InvM3 * (1.f - Bary) * Delta), InvM3 ));
-		}
+		VectorScatter(&PandInvM[Index1], SetVector4( P1 + (InvM1 * Delta), InvM1 ));
+		VectorScatter(&PandInvM[Index2], SetVector4( P2 - (InvM2 * Bary * Delta), InvM2 ));
+		VectorScatter(&PandInvM[Index3], SetVector4( P3 - (InvM3 * (1.f - Bary) * Delta), InvM3 ));
 	}
 }
 
@@ -456,26 +386,17 @@ export void ApplyXPBDAnisoAxialSpringConstraintsWithDamping(uniform FVector4f Pa
 		const varying float Bary = Barys[Index];
 		const varying float Dist = Dists[Index];
 		const varying float Lambda = Lambdas[Index];
-		const varying FVector3f WarpWeftBiasBaseMultiplier = VectorGather(&WarpWeftBiasBaseMultiplers[Index]);
+		const varying FVector3f WarpWeftBiasBaseMultiplier = VectorLoad(&WarpWeftBiasBaseMultiplers[extract(Index,0)]);
 		const varying float StiffnessValue = VectorDot(WarpWeftBiasBaseMultiplier, Stiffness);
 
 		varying float DLambda;
-		const FVector3f Delta = GetXPBDAxialSpringDeltaWithDamping(P1, P2, P3, X1, X2, X3, InvM1, InvM2, InvM3, Dt, Bary, Dist, Lambda, StiffnessValue, DampingRatio, FLOAT_SMALL_NUMBER, DLambda);
+		const FVector3f Delta = GetXPBDAxialSpringDeltaWithDamping(P1, P2, P3, X1, X2, X3, InvM1, InvM2, InvM3, Dt, Bary, Dist, Lambda, StiffnessValue, DampingRatio, DLambda);
 
 		Lambdas[Index] = Lambda + DLambda;
 
-		if (InvM1 > FLOAT_ZERO)
-		{
-			VectorScatter(&PandInvM[Index1], SetVector4( P1 + (InvM1 * Delta), InvM1 ));
-		}
-		if (InvM2 > FLOAT_ZERO)
-		{
-			VectorScatter(&PandInvM[Index2], SetVector4( P2 - (InvM2 * Bary * Delta), InvM2 ));
-		}
-		if (InvM3 > FLOAT_ZERO)
-		{
-			VectorScatter(&PandInvM[Index3], SetVector4( P3 - (InvM3 * (1.f - Bary) * Delta), InvM3 ));
-		}
+		VectorScatter(&PandInvM[Index1], SetVector4( P1 + (InvM1 * Delta), InvM1 ));
+		VectorScatter(&PandInvM[Index2], SetVector4( P2 - (InvM2 * Bary * Delta), InvM2 ));
+		VectorScatter(&PandInvM[Index3], SetVector4( P3 - (InvM3 * (1.f - Bary) * Delta), InvM3 ));
 	}
 }
 
@@ -515,26 +436,17 @@ export void ApplyXPBDAnisoAxialSpringDampingConstraints(uniform FVector4f PandIn
 		const varying float Bary = Barys[Index];
 		const varying float Dist = Dists[Index];
 		const varying float Lambda = Lambdas[Index];
-		const varying FVector3f WarpWeftBiasBaseMultiplier = VectorGather(&WarpWeftBiasBaseMultiplers[Index]);
+		const varying FVector3f WarpWeftBiasBaseMultiplier = VectorLoad(&WarpWeftBiasBaseMultiplers[extract(Index,0)]);
 		const varying float StiffnessValue = VectorDot(WarpWeftBiasBaseMultiplier, Stiffness);
 
 		varying float DLambda;
-		const FVector3f Delta = GetXPBDAxialSpringDampingDelta(P1, P2, P3, X1, X2, X3, InvM1, InvM2, InvM3, Dt, Bary, Dist, Lambda, StiffnessValue, DampingRatio, FLOAT_SMALL_NUMBER, DLambda);
+		const FVector3f Delta = GetXPBDAxialSpringDampingDelta(P1, P2, P3, X1, X2, X3, InvM1, InvM2, InvM3, Dt, Bary, Dist, Lambda, StiffnessValue, DampingRatio, DLambda);
 
 		Lambdas[Index] = Lambda + DLambda;
 
-		if (InvM1 > FLOAT_ZERO)
-		{
-			VectorScatter(&PandInvM[Index1], SetVector4( P1 + (InvM1 * Delta), InvM1 ));
-		}
-		if (InvM2 > FLOAT_ZERO)
-		{
-			VectorScatter(&PandInvM[Index2], SetVector4( P2 - (InvM2 * Bary * Delta), InvM2 ));
-		}
-		if (InvM3 > FLOAT_ZERO)
-		{
-			VectorScatter(&PandInvM[Index3], SetVector4( P3 - (InvM3 * (1.f - Bary) * Delta), InvM3 ));
-		}
+		VectorScatter(&PandInvM[Index1], SetVector4( P1 + (InvM1 * Delta), InvM1 ));
+		VectorScatter(&PandInvM[Index2], SetVector4( P2 - (InvM2 * Bary * Delta), InvM2 ));
+		VectorScatter(&PandInvM[Index3], SetVector4( P3 - (InvM3 * (1.f - Bary) * Delta), InvM3 ));
 	}
 }
 
@@ -545,16 +457,15 @@ export void ApplyXPBDAnisoAxialSpringConstraintsWithWeightMaps(uniform FVector4f
 									const uniform FVector3f WarpWeftBiasBaseMultiplers[],
 									uniform float Lambdas[],
 									const uniform float Dt,
-									const uniform float MinStiffness,
 									const uniform bool StiffnessHasMap,
-									const uniform uint8 StiffnessIndices[],
-									const uniform float StiffnessTable[],
+									const uniform FVector2f& StiffnessOffsetRange,
+									const uniform float StiffnessMapValues[],
 									const uniform bool StiffnessWeftHasMap,
-									const uniform uint8 StiffnessWeftIndices[],
-									const uniform float StiffnessWeftTable[],
+									const uniform FVector2f& StiffnessWeftOffsetRange,
+									const uniform float StiffnessWeftMapValues[],
 									const uniform bool StiffnessBiasHasMap,
-									const uniform uint8 StiffnessBiasIndices[],
-									const uniform float StiffnessBiasTable[],
+									const uniform FVector2f& StiffnessBiasOffsetRange,
+									const uniform float StiffnessBiasMapValues[],
 									const uniform int32 NumConstraints)
 {
 
@@ -565,10 +476,6 @@ export void ApplyXPBDAnisoAxialSpringConstraintsWithWeightMaps(uniform FVector4f
 		const varying int32 Index2 = Constraint.V[1];
 		const varying int32 Index3 = Constraint.V[2];
 		
-		const varying int8 StiffnessIndex = StiffnessHasMap ? StiffnessIndices[Index] : 0;
-		const varying int8 StiffnessWeftIndex = StiffnessWeftHasMap ? StiffnessWeftIndices[Index] : 0;
-		const varying int8 StiffnessBiasIndex = StiffnessBiasHasMap ? StiffnessBiasIndices[Index] : 0;
-		
 		const varying FVector4f PandInvM1 = VectorGather(&PandInvM[Index1]);
 		const varying FVector4f PandInvM2 = VectorGather(&PandInvM[Index2]);
 		const varying FVector4f PandInvM3 = VectorGather(&PandInvM[Index3]);
@@ -583,34 +490,22 @@ export void ApplyXPBDAnisoAxialSpringConstraintsWithWeightMaps(uniform FVector4f
 		const varying float Dist = Dists[Index];
 		const varying float Lambda = Lambdas[Index];
 		
-		#pragma ignore warning(perf)
-		const varying float Stiffness = StiffnessTable[StiffnessIndex];
-		#pragma ignore warning(perf)
-		const varying float StiffnessWeft = StiffnessWeftTable[StiffnessWeftIndex];
-		#pragma ignore warning(perf)
-		const varying float StiffnessBias = StiffnessBiasTable[StiffnessBiasIndex];
+		const varying float Stiffness = StiffnessHasMap ? StiffnessOffsetRange.V[0] + StiffnessMapValues[Index]*StiffnessOffsetRange.V[1] : StiffnessOffsetRange.V[0];
+		const varying float StiffnessWeft = StiffnessWeftHasMap ? StiffnessWeftOffsetRange.V[0] + StiffnessWeftMapValues[Index]*StiffnessWeftOffsetRange.V[1] : StiffnessWeftOffsetRange.V[0];
+		const varying float StiffnessBias = StiffnessBiasHasMap ? StiffnessBiasOffsetRange.V[0] + StiffnessBiasMapValues[Index]*StiffnessBiasOffsetRange.V[1] : StiffnessBiasOffsetRange.V[0];
 		
 		const varying FVector3f StiffnessVector = {{StiffnessWeft, Stiffness, StiffnessBias}};
-		const varying FVector3f WarpWeftBiasBaseMultiplier = VectorGather(&WarpWeftBiasBaseMultiplers[Index]);
+		const varying FVector3f WarpWeftBiasBaseMultiplier = VectorLoad(&WarpWeftBiasBaseMultiplers[extract(Index,0)]);
 		const varying float StiffnessValue = VectorDot(WarpWeftBiasBaseMultiplier, StiffnessVector);
 
 		varying float DLambda;
-		const varying FVector3f Delta = GetXPBDAxialSpringDelta(P1, P2, P3, InvM1, InvM2, InvM3, Dt, Bary, Dist, Lambda, StiffnessValue, FLOAT_SMALL_NUMBER, DLambda);
+		const varying FVector3f Delta = GetXPBDAxialSpringDelta(P1, P2, P3, InvM1, InvM2, InvM3, Dt, Bary, Dist, Lambda, StiffnessValue, DLambda);
 
 		Lambdas[Index] = Lambda + DLambda;
-
-		if (InvM1 > FLOAT_ZERO)
-		{
-			VectorScatter(&PandInvM[Index1], SetVector4( P1 + (InvM1 * Delta), InvM1 ));
-		}
-		if (InvM2 > FLOAT_ZERO)
-		{
-			VectorScatter(&PandInvM[Index2], SetVector4( P2 - (InvM2 * Bary * Delta), InvM2 ));
-		}
-		if (InvM3 > FLOAT_ZERO)
-		{
-			VectorScatter(&PandInvM[Index3], SetVector4( P3 - (InvM3 * (1.f - Bary) * Delta), InvM3 ));
-		}
+		
+		VectorScatter(&PandInvM[Index1], SetVector4( P1 + (InvM1 * Delta), InvM1 ));
+		VectorScatter(&PandInvM[Index2], SetVector4( P2 - (InvM2 * Bary * Delta), InvM2 ));
+		VectorScatter(&PandInvM[Index3], SetVector4( P3 - (InvM3 * (1.f - Bary) * Delta), InvM3 ));
 	}
 }
 
@@ -622,19 +517,18 @@ export void ApplyXPBDAnisoAxialSpringConstraintsWithDampingAndWeightMaps(uniform
 															const uniform FVector3f WarpWeftBiasBaseMultiplers[],
 															uniform float Lambdas[],
 															const uniform float Dt,
-															const uniform float MinStiffness,
 															const uniform bool StiffnessHasMap,
-															const uniform uint8 StiffnessIndices[],
-															const uniform float StiffnessTable[],
+															const uniform FVector2f& StiffnessOffsetRange,
+															const uniform float StiffnessMapValues[],
 															const uniform bool StiffnessWeftHasMap,
-															const uniform uint8 StiffnessWeftIndices[],
-															const uniform float StiffnessWeftTable[],
+															const uniform FVector2f& StiffnessWeftOffsetRange,
+															const uniform float StiffnessWeftMapValues[],
 															const uniform bool StiffnessBiasHasMap,
-															const uniform uint8 StiffnessBiasIndices[],
-															const uniform float StiffnessBiasTable[],
+															const uniform FVector2f& StiffnessBiasOffsetRange,
+															const uniform float StiffnessBiasMapValues[],
 															const uniform bool DampingHasMap,
-															const uniform uint8 DampingIndices[],
-															const uniform float DampingTable[],
+															const uniform FVector2f& DampingOffsetRange,
+															const uniform float DampingMapValues[],
 															const uniform int32 NumConstraints)
 {
 
@@ -644,11 +538,6 @@ export void ApplyXPBDAnisoAxialSpringConstraintsWithDampingAndWeightMaps(uniform
 		const varying int32 Index1 = Constraint.V[0];
 		const varying int32 Index2 = Constraint.V[1];
 		const varying int32 Index3 = Constraint.V[2];
-
-		const varying int8 StiffnessIndex = StiffnessHasMap ? StiffnessIndices[Index] : 0;
-		const varying int8 StiffnessWeftIndex = StiffnessWeftHasMap ? StiffnessWeftIndices[Index] : 0;
-		const varying int8 StiffnessBiasIndex = StiffnessBiasHasMap ? StiffnessBiasIndices[Index] : 0;
-		const varying int8 DampingIndex = DampingHasMap ? DampingIndices[Index] : 0;
 		
 		const varying FVector4f PandInvM1 = VectorGather(&PandInvM[Index1]);
 		const varying FVector4f PandInvM2 = VectorGather(&PandInvM[Index2]);
@@ -666,38 +555,26 @@ export void ApplyXPBDAnisoAxialSpringConstraintsWithDampingAndWeightMaps(uniform
 		
 		const varying float Bary = Barys[Index];
 		const varying float Dist = Dists[Index];
-		const varying float Lambda = Lambdas[Index];
-
-		#pragma ignore warning(perf)
-		const varying float Stiffness = StiffnessTable[StiffnessIndex];
-		#pragma ignore warning(perf)
-		const varying float StiffnessWeft = StiffnessWeftTable[StiffnessWeftIndex];
-		#pragma ignore warning(perf)
-		const varying float StiffnessBias = StiffnessBiasTable[StiffnessBiasIndex];
-		#pragma ignore warning(perf)
-		const varying float DampingRatio = DampingTable[DampingIndex];
+		const varying float Lambda = Lambdas[Index];		
+		
+		const varying float Stiffness = StiffnessHasMap ? StiffnessOffsetRange.V[0] + StiffnessMapValues[Index]*StiffnessOffsetRange.V[1] : StiffnessOffsetRange.V[0];
+		const varying float StiffnessWeft = StiffnessWeftHasMap ? StiffnessWeftOffsetRange.V[0] + StiffnessWeftMapValues[Index]*StiffnessWeftOffsetRange.V[1] : StiffnessWeftOffsetRange.V[0];
+		const varying float StiffnessBias = StiffnessBiasHasMap ? StiffnessBiasOffsetRange.V[0] + StiffnessBiasMapValues[Index]*StiffnessBiasOffsetRange.V[1] : StiffnessBiasOffsetRange.V[0];
 
 		const varying FVector3f StiffnessVector = {{StiffnessWeft, Stiffness, StiffnessBias}};
-		const varying FVector3f WarpWeftBiasBaseMultiplier = VectorGather(&WarpWeftBiasBaseMultiplers[Index]);
+		const varying FVector3f WarpWeftBiasBaseMultiplier = VectorLoad(&WarpWeftBiasBaseMultiplers[extract(Index,0)]);
 		const varying float StiffnessValue = VectorDot(WarpWeftBiasBaseMultiplier, StiffnessVector);
+		
+		const varying float DampingRatio = DampingHasMap ? DampingOffsetRange.V[0] + DampingMapValues[Index]*DampingOffsetRange.V[1] : DampingOffsetRange.V[0];
 
 		varying float DLambda;
-		const FVector3f Delta = GetXPBDAxialSpringDeltaWithDamping(P1, P2, P3, X1, X2, X3, InvM1, InvM2, InvM3, Dt, Bary, Dist, Lambda, StiffnessValue, DampingRatio, FLOAT_SMALL_NUMBER, DLambda);
+		const FVector3f Delta = GetXPBDAxialSpringDeltaWithDamping(P1, P2, P3, X1, X2, X3, InvM1, InvM2, InvM3, Dt, Bary, Dist, Lambda, StiffnessValue, DampingRatio, DLambda);
 
 		Lambdas[Index] = Lambda + DLambda;
-
-		if (InvM1 > FLOAT_ZERO)
-		{
-			VectorScatter(&PandInvM[Index1], SetVector4( P1 + (InvM1 * Delta), InvM1 ));
-		}
-		if (InvM2 > FLOAT_ZERO)
-		{
-			VectorScatter(&PandInvM[Index2], SetVector4( P2 - (InvM2 * Bary * Delta), InvM2 ));
-		}
-		if (InvM3 > FLOAT_ZERO)
-		{
-			VectorScatter(&PandInvM[Index3], SetVector4( P3 - (InvM3 * (1.f - Bary) * Delta), InvM3 ));
-		}
+		
+		VectorScatter(&PandInvM[Index1], SetVector4( P1 + (InvM1 * Delta), InvM1 ));
+		VectorScatter(&PandInvM[Index2], SetVector4( P2 - (InvM2 * Bary * Delta), InvM2 ));
+		VectorScatter(&PandInvM[Index3], SetVector4( P3 - (InvM3 * (1.f - Bary) * Delta), InvM3 ));
 	}
 }
 
@@ -709,19 +586,18 @@ export void ApplyXPBDAnisoAxialSpringDampingConstraintsWithWeightMaps(uniform FV
 															const uniform FVector3f WarpWeftBiasBaseMultiplers[],
 															uniform float Lambdas[],
 															const uniform float Dt,
-															const uniform float MinStiffness,
 															const uniform bool StiffnessHasMap,
-															const uniform uint8 StiffnessIndices[],
-															const uniform float StiffnessTable[],
+															const uniform FVector2f& StiffnessOffsetRange,
+															const uniform float StiffnessMapValues[],
 															const uniform bool StiffnessWeftHasMap,
-															const uniform uint8 StiffnessWeftIndices[],
-															const uniform float StiffnessWeftTable[],
+															const uniform FVector2f& StiffnessWeftOffsetRange,
+															const uniform float StiffnessWeftMapValues[],
 															const uniform bool StiffnessBiasHasMap,
-															const uniform uint8 StiffnessBiasIndices[],
-															const uniform float StiffnessBiasTable[],
+															const uniform FVector2f& StiffnessBiasOffsetRange,
+															const uniform float StiffnessBiasMapValues[],
 															const uniform bool DampingHasMap,
-															const uniform uint8 DampingIndices[],
-															const uniform float DampingTable[],
+															const uniform FVector2f& DampingOffsetRange,
+															const uniform float DampingMapValues[],
 															const uniform int32 NumConstraints)
 {
 
@@ -731,11 +607,6 @@ export void ApplyXPBDAnisoAxialSpringDampingConstraintsWithWeightMaps(uniform FV
 		const varying int32 Index1 = Constraint.V[0];
 		const varying int32 Index2 = Constraint.V[1];
 		const varying int32 Index3 = Constraint.V[2];
-
-		const varying int8 StiffnessIndex = StiffnessHasMap ? StiffnessIndices[Index] : 0;
-		const varying int8 StiffnessWeftIndex = StiffnessWeftHasMap ? StiffnessWeftIndices[Index] : 0;
-		const varying int8 StiffnessBiasIndex = StiffnessBiasHasMap ? StiffnessBiasIndices[Index] : 0;
-		const varying int8 DampingIndex = DampingHasMap ? DampingIndices[Index] : 0;
 		
 		const varying FVector4f PandInvM1 = VectorGather(&PandInvM[Index1]);
 		const varying FVector4f PandInvM2 = VectorGather(&PandInvM[Index2]);
@@ -754,36 +625,24 @@ export void ApplyXPBDAnisoAxialSpringDampingConstraintsWithWeightMaps(uniform FV
 		const varying float Bary = Barys[Index];
 		const varying float Dist = Dists[Index];
 		const varying float Lambda = Lambdas[Index];
-
-		#pragma ignore warning(perf)
-		const varying float Stiffness = StiffnessTable[StiffnessIndex];
-		#pragma ignore warning(perf)
-		const varying float StiffnessWeft = StiffnessWeftTable[StiffnessWeftIndex];
-		#pragma ignore warning(perf)
-		const varying float StiffnessBias = StiffnessBiasTable[StiffnessBiasIndex];
-		#pragma ignore warning(perf)
-		const varying float DampingRatio = DampingTable[DampingIndex];
+		
+		const varying float Stiffness = StiffnessHasMap ? StiffnessOffsetRange.V[0] + StiffnessMapValues[Index]*StiffnessOffsetRange.V[1] : StiffnessOffsetRange.V[0];
+		const varying float StiffnessWeft = StiffnessWeftHasMap ? StiffnessWeftOffsetRange.V[0] + StiffnessWeftMapValues[Index]*StiffnessWeftOffsetRange.V[1] : StiffnessWeftOffsetRange.V[0];
+		const varying float StiffnessBias = StiffnessBiasHasMap ? StiffnessBiasOffsetRange.V[0] + StiffnessBiasMapValues[Index]*StiffnessBiasOffsetRange.V[1] : StiffnessBiasOffsetRange.V[0];
 
 		const varying FVector3f StiffnessVector = {{StiffnessWeft, Stiffness, StiffnessBias}};
-		const varying FVector3f WarpWeftBiasBaseMultiplier = VectorGather(&WarpWeftBiasBaseMultiplers[Index]);
+		const varying FVector3f WarpWeftBiasBaseMultiplier = VectorLoad(&WarpWeftBiasBaseMultiplers[extract(Index,0)]);
 		const varying float StiffnessValue = VectorDot(WarpWeftBiasBaseMultiplier, StiffnessVector);
 
+		const varying float DampingRatio = DampingHasMap ? DampingOffsetRange.V[0] + DampingMapValues[Index]*DampingOffsetRange.V[1] : DampingOffsetRange.V[0];
+
 		varying float DLambda;
-		const FVector3f Delta = GetXPBDAxialSpringDampingDelta(P1, P2, P3, X1, X2, X3, InvM1, InvM2, InvM3, Dt, Bary, Dist, Lambda, StiffnessValue, DampingRatio, FLOAT_SMALL_NUMBER, DLambda);
+		const FVector3f Delta = GetXPBDAxialSpringDampingDelta(P1, P2, P3, X1, X2, X3, InvM1, InvM2, InvM3, Dt, Bary, Dist, Lambda, StiffnessValue, DampingRatio, DLambda);
 
 		Lambdas[Index] = Lambda + DLambda;
-
-		if (InvM1 > FLOAT_ZERO)
-		{
-			VectorScatter(&PandInvM[Index1], SetVector4( P1 + (InvM1 * Delta), InvM1 ));
-		}
-		if (InvM2 > FLOAT_ZERO)
-		{
-			VectorScatter(&PandInvM[Index2], SetVector4( P2 - (InvM2 * Bary * Delta), InvM2 ));
-		}
-		if (InvM3 > FLOAT_ZERO)
-		{
-			VectorScatter(&PandInvM[Index3], SetVector4( P3 - (InvM3 * (1.f - Bary) * Delta), InvM3 ));
-		}
+		
+		VectorScatter(&PandInvM[Index1], SetVector4( P1 + (InvM1 * Delta), InvM1 ));
+		VectorScatter(&PandInvM[Index2], SetVector4( P2 - (InvM2 * Bary * Delta), InvM2 ));
+		VectorScatter(&PandInvM[Index3], SetVector4( P3 - (InvM3 * (1.f - Bary) * Delta), InvM3 ));
 	}
 }
diff --git a/Engine/Source/Runtime/Experimental/Chaos/Private/Chaos/XPBDBendingConstraints.cpp b/Engine/Source/Runtime/Experimental/Chaos/Private/Chaos/XPBDBendingConstraints.cpp
--- a/Engine/Source/Runtime/Experimental/Chaos/Private/Chaos/XPBDBendingConstraints.cpp
+++ b/Engine/Source/Runtime/Experimental/Chaos/Private/Chaos/XPBDBendingConstraints.cpp
@@ -14,7 +14,8 @@ bool bChaos_XPBDBending_ISPC_Enabled = true;
 FAutoConsoleVariableRef CVarChaosXPBDBendingISPCEnabled(TEXT("p.Chaos.XPBDBending.ISPC"), bChaos_XPBDBending_ISPC_Enabled, TEXT("Whether to use ISPC optimizations in XPBD Bending constraints"));
 
 static_assert(sizeof(ispc::FVector4f) == sizeof(Chaos::Softs::FPAndInvM), "sizeof(ispc::FVector4f) != sizeof(Chaos::Softs::FPAndInvM");
-static_assert(sizeof(ispc::FIntVector4) == sizeof(Chaos::TVec4<int32>), "sizeof(ispc::FIntVector4) != sizeof(Chaos::TVec4<int32>");
+static_assert(sizeof(ispc::FVector3f) == sizeof(Chaos::Softs::FSolverVec3), "sizeof(ispc::FVector3f) != sizeof(Chaos::Softs::FSolverVec3");
+static_assert(sizeof(ispc::FVector2f) == sizeof(Chaos::Softs::FSolverVec2), "sizeof(ispc::FVector2f) != sizeof(Chaos::Softs::FSolverVec2");
 #endif
 
 namespace Chaos::Softs {
@@ -69,36 +70,88 @@ void FXPBDBendingConstraints::InitColor(const SolverParticlesOrRange& InParticle
 		Constraints = MoveTemp(ReorderedConstraints);
 		ConstraintSharedEdges = MoveTemp(ReorderedConstraintSharedEdges);
 		RestAngles = MoveTemp(ReorderedRestAngles);
-		Stiffness.ReorderIndices(OrigToReorderedIndices);
-		BucklingStiffness.ReorderIndices(OrigToReorderedIndices);
+		XPBDStiffness.ReorderIndices(OrigToReorderedIndices);
+		XPBDBucklingStiffness.ReorderIndices(OrigToReorderedIndices);
+
+#if INTEL_ISPC
+		ConstraintsIndex1.SetNumUninitialized(Constraints.Num());
+		ConstraintsIndex2.SetNumUninitialized(Constraints.Num());
+		ConstraintsIndex3.SetNumUninitialized(Constraints.Num());
+		ConstraintsIndex4.SetNumUninitialized(Constraints.Num());
+		for (int32 ConstraintIndex = 0; ConstraintIndex < Constraints.Num(); ++ConstraintIndex)
+		{
+			ConstraintsIndex1[ConstraintIndex] = Constraints[ConstraintIndex][0];
+			ConstraintsIndex2[ConstraintIndex] = Constraints[ConstraintIndex][1];
+			ConstraintsIndex3[ConstraintIndex] = Constraints[ConstraintIndex][2];
+			ConstraintsIndex4[ConstraintIndex] = Constraints[ConstraintIndex][3];
+		}
+#endif
 	}
 }
 template CHAOS_API void FXPBDBendingConstraints::InitColor(const FSolverParticles& InParticles);
 template CHAOS_API void FXPBDBendingConstraints::InitColor(const FSolverParticlesRange& InParticles);
 
+
+template<typename SolverParticlesOrRange>
+void FXPBDBendingConstraints::Init(const SolverParticlesOrRange& InParticles)
+{
+	Lambdas.Reset();
+	Lambdas.AddZeroed(Constraints.Num());
+	LambdasDamping.Reset();
+	LambdasDamping.AddZeroed(Constraints.Num());
+#if INTEL_ISPC
+	IsBuckled.SetNumUninitialized(Constraints.Num());
+	X1Array.SetNumUninitialized(Constraints.Num());
+	X2Array.SetNumUninitialized(Constraints.Num());
+	X3Array.SetNumUninitialized(Constraints.Num());
+	X4Array.SetNumUninitialized(Constraints.Num());
+	if (bRealTypeCompatibleWithISPC && bChaos_Bending_ISPC_Enabled && ConstraintsIndex1.Num() == Constraints.Num())
+	{
+		ispc::InitXPBDBendingConstraintsIsBuckled(
+			(const ispc::FVector3f*)InParticles.XArray().GetData(),
+			ConstraintsIndex1.GetData(),
+			ConstraintsIndex2.GetData(),
+			ConstraintsIndex3.GetData(),
+			ConstraintsIndex4.GetData(),
+			RestAngles.GetData(),
+			IsBuckled.GetData(),
+			(ispc::FVector3f*)X1Array.GetData(),
+			(ispc::FVector3f*)X2Array.GetData(),
+			(ispc::FVector3f*)X3Array.GetData(),
+			(ispc::FVector3f*)X4Array.GetData(),
+			BucklingRatio,
+			Constraints.Num()
+		);
+	}
+	else
+#endif
+	{
+		FPBDBendingConstraintsBase::Init(InParticles);
+	}
+}
+template void CHAOS_API FXPBDBendingConstraints::Init(const FSolverParticles& InParticles);
+template void CHAOS_API FXPBDBendingConstraints::Init(const FSolverParticlesRange& InParticles);
+
 void FXPBDBendingConstraints::SetProperties(
 	const FCollectionPropertyConstFacade& PropertyCollection,
 	const TMap<FString, TConstArrayView<FRealSingle>>& WeightMaps)
 {
 	if (IsXPBDBendingElementStiffnessMutable(PropertyCollection))
 	{
-		const FSolverVec2 WeightedValue(GetWeightedFloatXPBDBendingElementStiffness(PropertyCollection));
+		const FSolverVec2 WeightedValue(GetWeightedFloatXPBDBendingElementStiffness(PropertyCollection).ClampAxes(0, MaxStiffness));
 		if (IsXPBDBendingElementStiffnessStringDirty(PropertyCollection))
 		{
 			const FString& WeightMapName = GetXPBDBendingElementStiffnessString(PropertyCollection);
-			Stiffness = FPBDStiffness(
+			XPBDStiffness = FPBDFlatWeightMap(
 				WeightedValue,
 				WeightMaps.FindRef(WeightMapName),
 				TConstArrayView<TVec2<int32>>(ConstraintSharedEdges),
 				ParticleOffset,
-				ParticleCount,
-				FPBDStiffness::DefaultTableSize,
-				FPBDStiffness::DefaultParameterFitBase,
-				MaxStiffness);
+				ParticleCount);
 		}
 		else
 		{
-			Stiffness.SetWeightedValue(WeightedValue, MaxStiffness);
+			XPBDStiffness.SetWeightedValue(WeightedValue);
 		}
 	}
 	if (IsXPBDBucklingRatioMutable(PropertyCollection))
@@ -107,23 +160,20 @@ void FXPBDBendingConstraints::SetProperties(
 	}
 	if (IsXPBDBucklingStiffnessMutable(PropertyCollection))
 	{
-		const FSolverVec2 WeightedValue(GetWeightedFloatXPBDBucklingStiffness(PropertyCollection));
+		const FSolverVec2 WeightedValue(GetWeightedFloatXPBDBucklingStiffness(PropertyCollection).ClampAxes(0, MaxStiffness));
 		if (IsXPBDBucklingStiffnessStringDirty(PropertyCollection))
 		{
 			const FString& WeightMapName = GetXPBDBucklingStiffnessString(PropertyCollection);
-			BucklingStiffness = FPBDStiffness(
+			XPBDBucklingStiffness = FPBDFlatWeightMap(
 				WeightedValue,
 				WeightMaps.FindRef(WeightMapName),
 				TConstArrayView<TVec2<int32>>(ConstraintSharedEdges),
 				ParticleOffset,
-				ParticleCount,
-				FPBDStiffness::DefaultTableSize,
-				FPBDStiffness::DefaultParameterFitBase,
-				MaxStiffness);
+				ParticleCount);
 		}
 		else
 		{
-			BucklingStiffness.SetWeightedValue(WeightedValue, MaxStiffness);
+			XPBDBucklingStiffness.SetWeightedValue(WeightedValue);
 		}
 	}
 	if (IsXPBDBendingElementDampingMutable(PropertyCollection))
@@ -132,7 +182,7 @@ void FXPBDBendingConstraints::SetProperties(
 		if (IsXPBDBendingElementDampingStringDirty(PropertyCollection))
 		{
 			const FString& WeightMapName = GetXPBDBendingElementDampingString(PropertyCollection);
-			DampingRatio = FPBDWeightMap(
+			DampingRatio = FPBDFlatWeightMap(
 				WeightedValue,
 				WeightMaps.FindRef(WeightMapName),
 				TConstArrayView<TVec2<int32>>(ConstraintSharedEdges),
@@ -159,11 +209,7 @@ void FXPBDBendingConstraints::ApplyHelper(SolverParticlesOrRange& Particles, con
 
 	const FSolverReal BiphasicStiffnessValue = IsBuckled[ConstraintIndex] ? BucklingValue : StiffnessValue;
 	const FSolverReal CombinedInvMass = Particles.InvM(Index1) + Particles.InvM(Index2) + Particles.InvM(Index3) + Particles.InvM(Index4);
-	if (BiphasicStiffnessValue < MinStiffness || CombinedInvMass < UE_SMALL_NUMBER)
-	{
-		return;
-	}
-
+	check(CombinedInvMass > 0);
 	const FSolverReal Damping = (FSolverReal)2.f * DampingRatioValue * FMath::Sqrt(BiphasicStiffnessValue / CombinedInvMass);
 
 	const FSolverReal Angle = CalcAngle(Particles.P(Index1), Particles.P(Index2), Particles.P(Index3), Particles.P(Index4));
@@ -197,8 +243,8 @@ void FXPBDBendingConstraints::Apply(SolverParticlesOrRange& Particles, const FSo
 	TRACE_CPUPROFILER_EVENT_SCOPE(FXPBDBendingConstraints_Apply);
 	SCOPE_CYCLE_COUNTER(STAT_XPBD_Bending);
 
-	const bool StiffnessHasWeightMap = Stiffness.HasWeightMap();
-	const bool BucklingStiffnessHasWeightMap = BucklingStiffness.HasWeightMap();
+	const bool StiffnessHasWeightMap = XPBDStiffness.HasWeightMap();
+	const bool BucklingStiffnessHasWeightMap = XPBDBucklingStiffness.HasWeightMap();
 	const bool DampingHasWeightMap = DampingRatio.HasWeightMap();
 
 	if (ConstraintsPerColorStartIndex.Num() > 0 && Constraints.Num() > Chaos_XPBDBending_ParallelConstraintCount)
@@ -210,11 +256,11 @@ void FXPBDBendingConstraints::Apply(SolverParticlesOrRange& Particles, const FSo
 		{
 			if (!StiffnessHasWeightMap && !BucklingStiffnessHasWeightMap && !DampingHasWeightMap)
 			{
-				const FSolverReal ExpStiffnessValue = (FSolverReal)Stiffness;
-				const FSolverReal ExpBucklingValue = (FSolverReal)BucklingStiffness;
+				const FSolverReal ExpStiffnessValue = (FSolverReal)XPBDStiffness;
+				const FSolverReal ExpBucklingValue = (FSolverReal)XPBDBucklingStiffness;
 				const FSolverReal DampingRatioValue = (FSolverReal)DampingRatio;
 
-				if (ExpStiffnessValue < MinStiffness && ExpBucklingValue < MinStiffness)
+				if (ExpStiffnessValue <= MinStiffness && ExpBucklingValue <= MinStiffness)
 				{
 					return;
 				}
@@ -228,13 +274,18 @@ void FXPBDBendingConstraints::Apply(SolverParticlesOrRange& Particles, const FSo
 							const int32 ColorSize = ConstraintsPerColorStartIndex[ConstraintColorIndex + 1] - ColorStart;
 							ispc::ApplyXPBDBendingDampingConstraints(
 								(ispc::FVector4f*)Particles.GetPAndInvM().GetData(),
-								(const ispc::FVector3f*)Particles.XArray().GetData(),
-								(ispc::FIntVector4*)&Constraints.GetData()[ColorStart],
+								(const ispc::FVector3f*)&X1Array.GetData()[ColorStart],
+								(const ispc::FVector3f*)&X2Array.GetData()[ColorStart],
+								(const ispc::FVector3f*)&X3Array.GetData()[ColorStart],
+								(const ispc::FVector3f*)&X4Array.GetData()[ColorStart],
+								&ConstraintsIndex1.GetData()[ColorStart],
+								&ConstraintsIndex2.GetData()[ColorStart],
+								&ConstraintsIndex3.GetData()[ColorStart],
+								&ConstraintsIndex4.GetData()[ColorStart],
 								&RestAngles.GetData()[ColorStart],
 								&IsBuckled.GetData()[ColorStart],
 								&LambdasDamping.GetData()[ColorStart],
 								Dt,
-								MinStiffness,
 								ExpStiffnessValue,
 								ExpBucklingValue,
 								DampingRatioValue,
@@ -249,13 +300,18 @@ void FXPBDBendingConstraints::Apply(SolverParticlesOrRange& Particles, const FSo
 							const int32 ColorSize = ConstraintsPerColorStartIndex[ConstraintColorIndex + 1] - ColorStart;
 							ispc::ApplyXPBDBendingConstraintsWithDamping(
 								(ispc::FVector4f*)Particles.GetPAndInvM().GetData(),
-								(const ispc::FVector3f*)Particles.XArray().GetData(),
-								(ispc::FIntVector4*)&Constraints.GetData()[ColorStart],
+								(const ispc::FVector3f*)&X1Array.GetData()[ColorStart],
+								(const ispc::FVector3f*)&X2Array.GetData()[ColorStart],
+								(const ispc::FVector3f*)&X3Array.GetData()[ColorStart],
+								(const ispc::FVector3f*)&X4Array.GetData()[ColorStart],
+								&ConstraintsIndex1.GetData()[ColorStart],
+								&ConstraintsIndex2.GetData()[ColorStart],
+								&ConstraintsIndex3.GetData()[ColorStart],
+								&ConstraintsIndex4.GetData()[ColorStart],
 								&RestAngles.GetData()[ColorStart],
 								&IsBuckled.GetData()[ColorStart],
 								&Lambdas.GetData()[ColorStart],
 								Dt,
-								MinStiffness,
 								ExpStiffnessValue,
 								ExpBucklingValue,
 								DampingRatioValue,
@@ -270,12 +326,14 @@ void FXPBDBendingConstraints::Apply(SolverParticlesOrRange& Particles, const FSo
 					const int32 ColorSize = ConstraintsPerColorStartIndex[ConstraintColorIndex + 1] - ColorStart;
 					ispc::ApplyXPBDBendingConstraints(
 						(ispc::FVector4f*)Particles.GetPAndInvM().GetData(),
-						(ispc::FIntVector4*)&Constraints.GetData()[ColorStart],
+						&ConstraintsIndex1.GetData()[ColorStart],
+						&ConstraintsIndex2.GetData()[ColorStart],
+						&ConstraintsIndex3.GetData()[ColorStart],
+						&ConstraintsIndex4.GetData()[ColorStart],
 						&RestAngles.GetData()[ColorStart],
 						&IsBuckled.GetData()[ColorStart],
 						&Lambdas.GetData()[ColorStart],
 						Dt,
-						MinStiffness,
 						ExpStiffnessValue,
 						ExpBucklingValue,
 						ColorSize);
@@ -294,22 +352,27 @@ void FXPBDBendingConstraints::Apply(SolverParticlesOrRange& Particles, const FSo
 							const int32 ColorSize = ConstraintsPerColorStartIndex[ConstraintColorIndex + 1] - ColorStart;
 							ispc::ApplyXPBDBendingDampingConstraintsWithMaps(
 								(ispc::FVector4f*)Particles.GetPAndInvM().GetData(),
-								(const ispc::FVector3f*)Particles.XArray().GetData(),
-								(ispc::FIntVector4*)&Constraints.GetData()[ColorStart],
+								(const ispc::FVector3f*)&X1Array.GetData()[ColorStart],
+								(const ispc::FVector3f*)&X2Array.GetData()[ColorStart],
+								(const ispc::FVector3f*)&X3Array.GetData()[ColorStart],
+								(const ispc::FVector3f*)&X4Array.GetData()[ColorStart],
+								&ConstraintsIndex1.GetData()[ColorStart],
+								&ConstraintsIndex2.GetData()[ColorStart],
+								&ConstraintsIndex3.GetData()[ColorStart],
+								&ConstraintsIndex4.GetData()[ColorStart],
 								&RestAngles.GetData()[ColorStart],
 								&IsBuckled.GetData()[ColorStart],
 								&LambdasDamping.GetData()[ColorStart],
 								Dt,
-								MinStiffness,
 								StiffnessHasWeightMap,
-								StiffnessHasWeightMap ? &Stiffness.GetIndices().GetData()[ColorStart] : nullptr,
-								&Stiffness.GetTable().GetData()[0],
+								reinterpret_cast<const ispc::FVector2f&>(XPBDStiffness.GetOffsetRange()),
+								StiffnessHasWeightMap ? &XPBDStiffness.GetMapValues().GetData()[ColorStart] : nullptr,
 								BucklingStiffnessHasWeightMap,
-								BucklingStiffnessHasWeightMap ? &BucklingStiffness.GetIndices().GetData()[ColorStart] : nullptr,
-								&BucklingStiffness.GetTable().GetData()[0],
+								reinterpret_cast<const ispc::FVector2f&>(XPBDBucklingStiffness.GetOffsetRange()),
+								BucklingStiffnessHasWeightMap ? &XPBDBucklingStiffness.GetMapValues().GetData()[ColorStart] : nullptr,
 								DampingHasWeightMap,
-								DampingHasWeightMap ? &DampingRatio.GetIndices().GetData()[ColorStart] : nullptr,
-								&DampingRatio.GetTable().GetData()[0],
+								reinterpret_cast<const ispc::FVector2f&>(DampingRatio.GetOffsetRange()),
+								DampingHasWeightMap ? &DampingRatio.GetMapValues()[ColorStart] : nullptr,
 								ColorSize);
 						}
 					}
@@ -321,22 +384,27 @@ void FXPBDBendingConstraints::Apply(SolverParticlesOrRange& Particles, const FSo
 							const int32 ColorSize = ConstraintsPerColorStartIndex[ConstraintColorIndex + 1] - ColorStart;
 							ispc::ApplyXPBDBendingConstraintsWithDampingAndMaps(
 								(ispc::FVector4f*)Particles.GetPAndInvM().GetData(),
-								(const ispc::FVector3f*)Particles.XArray().GetData(),
-								(ispc::FIntVector4*)&Constraints.GetData()[ColorStart],
+								(const ispc::FVector3f*)&X1Array.GetData()[ColorStart],
+								(const ispc::FVector3f*)&X2Array.GetData()[ColorStart],
+								(const ispc::FVector3f*)&X3Array.GetData()[ColorStart],
+								(const ispc::FVector3f*)&X4Array.GetData()[ColorStart],
+								&ConstraintsIndex1.GetData()[ColorStart],
+								&ConstraintsIndex2.GetData()[ColorStart],
+								&ConstraintsIndex3.GetData()[ColorStart],
+								&ConstraintsIndex4.GetData()[ColorStart],
 								&RestAngles.GetData()[ColorStart],
 								&IsBuckled.GetData()[ColorStart],
 								&Lambdas.GetData()[ColorStart],
 								Dt,
-								MinStiffness,
 								StiffnessHasWeightMap,
-								StiffnessHasWeightMap ? &Stiffness.GetIndices().GetData()[ColorStart] : nullptr,
-								&Stiffness.GetTable().GetData()[0],
+								reinterpret_cast<const ispc::FVector2f&>(XPBDStiffness.GetOffsetRange()),
+								StiffnessHasWeightMap ? &XPBDStiffness.GetMapValues().GetData()[ColorStart] : nullptr,
 								BucklingStiffnessHasWeightMap,
-								BucklingStiffnessHasWeightMap ? &BucklingStiffness.GetIndices().GetData()[ColorStart] : nullptr,
-								&BucklingStiffness.GetTable().GetData()[0],
+								reinterpret_cast<const ispc::FVector2f&>(XPBDBucklingStiffness.GetOffsetRange()),
+								BucklingStiffnessHasWeightMap ? &XPBDBucklingStiffness.GetMapValues().GetData()[ColorStart] : nullptr,
 								DampingHasWeightMap,
-								DampingHasWeightMap ? &DampingRatio.GetIndices().GetData()[ColorStart] : nullptr,
-								&DampingRatio.GetTable().GetData()[0],
+								reinterpret_cast<const ispc::FVector2f&>(DampingRatio.GetOffsetRange()),
+								DampingHasWeightMap ? &DampingRatio.GetMapValues()[ColorStart] : nullptr,
 								ColorSize);
 						}
 						return;
@@ -348,18 +416,20 @@ void FXPBDBendingConstraints::Apply(SolverParticlesOrRange& Particles, const FSo
 					const int32 ColorSize = ConstraintsPerColorStartIndex[ConstraintColorIndex + 1] - ColorStart;
 					ispc::ApplyXPBDBendingConstraintsWithMaps(
 						(ispc::FVector4f*)Particles.GetPAndInvM().GetData(),
-						(ispc::FIntVector4*)&Constraints.GetData()[ColorStart],
+						&ConstraintsIndex1.GetData()[ColorStart],
+						&ConstraintsIndex2.GetData()[ColorStart],
+						&ConstraintsIndex3.GetData()[ColorStart],
+						&ConstraintsIndex4.GetData()[ColorStart],
 						&RestAngles.GetData()[ColorStart],
 						&IsBuckled.GetData()[ColorStart],
 						&Lambdas.GetData()[ColorStart],
 						Dt,
-						MinStiffness,
 						StiffnessHasWeightMap,
-						StiffnessHasWeightMap ? &Stiffness.GetIndices().GetData()[ColorStart] : nullptr,
-						&Stiffness.GetTable().GetData()[0],
+						reinterpret_cast<const ispc::FVector2f&>(XPBDStiffness.GetOffsetRange()),
+						StiffnessHasWeightMap ? &XPBDStiffness.GetMapValues().GetData()[ColorStart] : nullptr,
 						BucklingStiffnessHasWeightMap,
-						BucklingStiffnessHasWeightMap ? &BucklingStiffness.GetIndices().GetData()[ColorStart] : nullptr,
-						&BucklingStiffness.GetTable().GetData()[0],
+						reinterpret_cast<const ispc::FVector2f&>(XPBDBucklingStiffness.GetOffsetRange()),
+						BucklingStiffnessHasWeightMap ? &XPBDBucklingStiffness.GetMapValues().GetData()[ColorStart] : nullptr,
 						ColorSize);
 				}
 			}
@@ -368,8 +438,8 @@ void FXPBDBendingConstraints::Apply(SolverParticlesOrRange& Particles, const FSo
 #endif
 		{
 			// Parallel non-ispc
-			const FSolverReal StiffnessNoMap = (FSolverReal)Stiffness;
-			const FSolverReal BucklingStiffnessNoMap = (FSolverReal)BucklingStiffness;
+			const FSolverReal StiffnessNoMap = (FSolverReal)XPBDStiffness;
+			const FSolverReal BucklingStiffnessNoMap = (FSolverReal)XPBDBucklingStiffness;
 			const FSolverReal DampingNoMap = (FSolverReal)DampingRatio;
 
 			if (DampingHasWeightMap || (FSolverReal)DampingRatio > 0)
@@ -383,8 +453,8 @@ void FXPBDBendingConstraints::Apply(SolverParticlesOrRange& Particles, const FSo
 						PhysicsParallelFor(ColorSize, [&](const int32 Index)
 						{
 							const int32 ConstraintIndex = ColorStart + Index;
-							const FSolverReal ExpStiffnessValue = StiffnessHasWeightMap ? Stiffness[ConstraintIndex] : StiffnessNoMap;
-							const FSolverReal ExpBucklingValue = BucklingStiffnessHasWeightMap ? BucklingStiffness[ConstraintIndex] : BucklingStiffnessNoMap;
+							const FSolverReal ExpStiffnessValue = StiffnessHasWeightMap ? XPBDStiffness[ConstraintIndex] : StiffnessNoMap;
+							const FSolverReal ExpBucklingValue = BucklingStiffnessHasWeightMap ? XPBDBucklingStiffness[ConstraintIndex] : BucklingStiffnessNoMap;
 							const FSolverReal DampingRatioValue = DampingHasWeightMap ? DampingRatio[ConstraintIndex] : DampingNoMap;
 							constexpr bool bDampingOnly = true;
 							constexpr bool bElasticOnly = false;
@@ -401,8 +471,8 @@ void FXPBDBendingConstraints::Apply(SolverParticlesOrRange& Particles, const FSo
 						PhysicsParallelFor(ColorSize, [&](const int32 Index)
 						{
 							const int32 ConstraintIndex = ColorStart + Index;
-							const FSolverReal ExpStiffnessValue = StiffnessHasWeightMap ? Stiffness[ConstraintIndex] : StiffnessNoMap;
-							const FSolverReal ExpBucklingValue = BucklingStiffnessHasWeightMap ? BucklingStiffness[ConstraintIndex] : BucklingStiffnessNoMap;
+							const FSolverReal ExpStiffnessValue = StiffnessHasWeightMap ? XPBDStiffness[ConstraintIndex] : StiffnessNoMap;
+							const FSolverReal ExpBucklingValue = BucklingStiffnessHasWeightMap ? XPBDBucklingStiffness[ConstraintIndex] : BucklingStiffnessNoMap;
 							const FSolverReal DampingRatioValue = DampingHasWeightMap ? DampingRatio[ConstraintIndex] : DampingNoMap;
 							constexpr bool bDampingOnly = false;
 							constexpr bool bElasticOnly = false;
@@ -419,8 +489,8 @@ void FXPBDBendingConstraints::Apply(SolverParticlesOrRange& Particles, const FSo
 				PhysicsParallelFor(ColorSize, [&](const int32 Index)
 				{
 					const int32 ConstraintIndex = ColorStart + Index;
-					const FSolverReal ExpStiffnessValue = StiffnessHasWeightMap ? Stiffness[ConstraintIndex] : StiffnessNoMap;
-					const FSolverReal ExpBucklingValue = BucklingStiffnessHasWeightMap ? BucklingStiffness[ConstraintIndex] : BucklingStiffnessNoMap;
+					const FSolverReal ExpStiffnessValue = StiffnessHasWeightMap ? XPBDStiffness[ConstraintIndex] : StiffnessNoMap;
+					const FSolverReal ExpBucklingValue = BucklingStiffnessHasWeightMap ? XPBDBucklingStiffness[ConstraintIndex] : BucklingStiffnessNoMap;
 					const FSolverReal DampingRatioValue = DampingHasWeightMap ? DampingRatio[ConstraintIndex] : DampingNoMap;
 					constexpr bool bDampingOnly = false;
 					constexpr bool bElasticOnly = true;
@@ -432,8 +502,8 @@ void FXPBDBendingConstraints::Apply(SolverParticlesOrRange& Particles, const FSo
 	else
 	{
 		// Single-threaded
-		const FSolverReal StiffnessNoMap = (FSolverReal)Stiffness;
-		const FSolverReal BucklingStiffnessNoMap = (FSolverReal)BucklingStiffness;
+		const FSolverReal StiffnessNoMap = (FSolverReal)XPBDStiffness;
+		const FSolverReal BucklingStiffnessNoMap = (FSolverReal)XPBDBucklingStiffness;
 		const FSolverReal DampingNoMap = (FSolverReal)DampingRatio;
 
 		if (DampingHasWeightMap || (FSolverReal)DampingRatio > 0)
@@ -442,8 +512,8 @@ void FXPBDBendingConstraints::Apply(SolverParticlesOrRange& Particles, const FSo
 			{
 				for (int32 ConstraintIndex = 0; ConstraintIndex < Constraints.Num(); ++ConstraintIndex)
 				{
-					const FSolverReal ExpStiffnessValue = StiffnessHasWeightMap ? Stiffness[ConstraintIndex] : StiffnessNoMap;
-					const FSolverReal ExpBucklingValue = BucklingStiffnessHasWeightMap ? BucklingStiffness[ConstraintIndex] : BucklingStiffnessNoMap;
+					const FSolverReal ExpStiffnessValue = StiffnessHasWeightMap ? XPBDStiffness[ConstraintIndex] : StiffnessNoMap;
+					const FSolverReal ExpBucklingValue = BucklingStiffnessHasWeightMap ? XPBDBucklingStiffness[ConstraintIndex] : BucklingStiffnessNoMap;
 					const FSolverReal DampingRatioValue = DampingHasWeightMap ? DampingRatio[ConstraintIndex] : DampingNoMap;
 					constexpr bool bDampingOnly = true;
 					constexpr bool bElasticOnly = false;
@@ -454,8 +524,8 @@ void FXPBDBendingConstraints::Apply(SolverParticlesOrRange& Particles, const FSo
 			{
 				for (int32 ConstraintIndex = 0; ConstraintIndex < Constraints.Num(); ++ConstraintIndex)
 				{
-					const FSolverReal ExpStiffnessValue = StiffnessHasWeightMap ? Stiffness[ConstraintIndex] : StiffnessNoMap;
-					const FSolverReal ExpBucklingValue = BucklingStiffnessHasWeightMap ? BucklingStiffness[ConstraintIndex] : BucklingStiffnessNoMap;
+					const FSolverReal ExpStiffnessValue = StiffnessHasWeightMap ? XPBDStiffness[ConstraintIndex] : StiffnessNoMap;
+					const FSolverReal ExpBucklingValue = BucklingStiffnessHasWeightMap ? XPBDBucklingStiffness[ConstraintIndex] : BucklingStiffnessNoMap;
 					const FSolverReal DampingRatioValue = DampingHasWeightMap ? DampingRatio[ConstraintIndex] : DampingNoMap;
 					constexpr bool bDampingOnly = false;
 					constexpr bool bElasticOnly = false;
@@ -466,8 +536,8 @@ void FXPBDBendingConstraints::Apply(SolverParticlesOrRange& Particles, const FSo
 		}
 		for (int32 ConstraintIndex = 0; ConstraintIndex < Constraints.Num(); ++ConstraintIndex)
 		{
-			const FSolverReal ExpStiffnessValue = StiffnessHasWeightMap ? Stiffness[ConstraintIndex] : StiffnessNoMap;
-			const FSolverReal ExpBucklingValue = BucklingStiffnessHasWeightMap ? BucklingStiffness[ConstraintIndex] : BucklingStiffnessNoMap;
+			const FSolverReal ExpStiffnessValue = StiffnessHasWeightMap ? XPBDStiffness[ConstraintIndex] : StiffnessNoMap;
+			const FSolverReal ExpBucklingValue = BucklingStiffnessHasWeightMap ? XPBDBucklingStiffness[ConstraintIndex] : BucklingStiffnessNoMap;
 			const FSolverReal DampingRatioValue = DampingHasWeightMap ? DampingRatio[ConstraintIndex] : DampingNoMap;
 			constexpr bool bDampingOnly = false;
 			constexpr bool bElasticOnly = true;
@@ -480,7 +550,7 @@ template CHAOS_API void FXPBDBendingConstraints::Apply(FSolverParticlesRange& Pa
 
 FSolverReal FXPBDBendingConstraints::ComputeTotalEnergy(const FSolverParticles& InParticles, const FSolverReal ExplicitStiffness)
 {
-	FSolverReal StiffnessValue = (FSolverReal)Stiffness;
+	FSolverReal StiffnessValue = (FSolverReal)XPBDStiffness;
 	if (ExplicitStiffness > 0.f)
 	{
 		StiffnessValue = ExplicitStiffness;
@@ -585,8 +655,8 @@ void ComputeGradTheta(const FSolverVec3& X0, const FSolverVec3& X1, const FSolve
 
 void FXPBDBendingConstraints::AddBendingResidualAndHessian(const FSolverParticles& Particles, const int32 ConstraintIndex, const int32 ConstraintIndexLocal, const FSolverReal Dt, TVec3<FSolverReal>& ParticleResidual, Chaos::PMatrix<FSolverReal, 3, 3>& ParticleHessian)
 {
-	FSolverReal ExpStiffnessValue = (FSolverReal)Stiffness;
-	FSolverReal ExpBucklingValue = (FSolverReal)BucklingStiffness;
+	FSolverReal ExpStiffnessValue = (FSolverReal)XPBDStiffness;
+	FSolverReal ExpBucklingValue = (FSolverReal)XPBDBucklingStiffness;
 	const FSolverReal DampingRatioValue = (FSolverReal)DampingRatio;
 
 	const TVec4<int32>& Constraint = Constraints[ConstraintIndex];
@@ -622,8 +692,8 @@ void FXPBDBendingConstraints::AddBendingResidualAndHessian(const FSolverParticle
 
 void FXPBDBendingConstraints::AddInternalForceDifferential(const FSolverParticles& InParticles, const TArray<TVector<FSolverReal, 3>>& DeltaParticles, TArray<TVector<FSolverReal, 3>>& ndf)
 {
-	FSolverReal ExpStiffnessValue = (FSolverReal)Stiffness;
-	FSolverReal ExpBucklingValue = (FSolverReal)BucklingStiffness;
+	FSolverReal ExpStiffnessValue = (FSolverReal)XPBDStiffness;
+	FSolverReal ExpBucklingValue = (FSolverReal)XPBDBucklingStiffness;
 	const FSolverReal DampingRatioValue = (FSolverReal)DampingRatio;
 
 
diff --git a/Engine/Source/Runtime/Experimental/Chaos/Private/Chaos/XPBDBendingConstraints.ispc b/Engine/Source/Runtime/Experimental/Chaos/Private/Chaos/XPBDBendingConstraints.ispc
--- a/Engine/Source/Runtime/Experimental/Chaos/Private/Chaos/XPBDBendingConstraints.ispc
+++ b/Engine/Source/Runtime/Experimental/Chaos/Private/Chaos/XPBDBendingConstraints.ispc
@@ -6,13 +6,71 @@
 #include "Chaos/PBDSofts.isph"
 #include "Chaos/XPBDBendingElement.isph"
 
+static inline float CalcAngle(const FVector3f& P1, const FVector3f& P2, const FVector3f& P3, const FVector3f& P4)
+{
+	const FVector3f SharedEdgeNormalized = VectorGetSafeNormal(P2 - P1);
+	const FVector3f P13CrossP23 = VectorCross(P1 - P3, P2 - P3);
+	const float Normal1Len = VectorSize(P13CrossP23);
+	const FVector3f Normal1 = SafeDivide(P13CrossP23, Normal1Len);
+	const FVector3f P24CrossP14 = VectorCross(P2 - P4, P1 - P4);
+	const float Normal2Len = VectorSize(P24CrossP14);
+	const FVector3f Normal2 = SafeDivide(P24CrossP14, Normal2Len);
+
+	const FVector3f N2CrossN1 = VectorCross(Normal2, Normal1);
+
+	const float CosPhi = clamp(VectorDot(Normal1, Normal2), (uniform float)(-1), (uniform float)(1));
+	const float SinPhi = clamp(VectorDot(N2CrossN1, SharedEdgeNormalized), (uniform float)(-1), (uniform float)(1));
+	return atan2(SinPhi, CosPhi);
+}
+
+export void InitXPBDBendingConstraintsIsBuckled(const uniform FVector3f X[],
+							const uniform int32 ConstraintIndex1[],
+							const uniform int32 ConstraintIndex2[],
+							const uniform int32 ConstraintIndex3[],
+							const uniform int32 ConstraintIndex4[],
+							const uniform float RestAngles[],
+							uniform bool IsBuckled[],
+							uniform FVector3f X1Out[],
+							uniform FVector3f X2Out[],
+							uniform FVector3f X3Out[],
+							uniform FVector3f X4Out[],
+							const uniform float BucklingRatio,
+							const uniform int32 NumConstraints)
+{
+	foreach(i = 0 ... NumConstraints)
+	{
+		const varying int32 Index1 = ConstraintIndex1[i];
+		const varying int32 Index2 = ConstraintIndex2[i];
+		const varying int32 Index3 = ConstraintIndex3[i];
+		const varying int32 Index4 = ConstraintIndex4[i];
+
+		const varying FVector3f X1 = VectorGather(&X[Index1]);
+		const varying FVector3f X2 = VectorGather(&X[Index2]);
+		const varying FVector3f X3 = VectorGather(&X[Index3]);
+		const varying FVector3f X4 = VectorGather(&X[Index4]);
+
+		const varying float RestAngle = RestAngles[i];
+		const varying float BucklingAngle = BucklingRatio * (FLOAT_PI - abs(RestAngle));
+
+		const varying float Angle = CalcAngle(X1, X2, X3, X4);
+		IsBuckled[i] = FLOAT_PI - abs(Angle) < BucklingAngle;
+
+		VectorStore(&X1Out[extract(i,0)], X1);
+		VectorStore(&X2Out[extract(i,0)], X2);
+		VectorStore(&X3Out[extract(i,0)], X3);
+		VectorStore(&X4Out[extract(i,0)], X4);
+	}
+}
+
 export void ApplyXPBDBendingConstraints(uniform FVector4f PandInvM[],
-									const uniform FIntVector4 Constraints[],
+									const uniform int32 ConstraintIndex1[],
+									const uniform int32 ConstraintIndex2[],
+									const uniform int32 ConstraintIndex3[],
+									const uniform int32 ConstraintIndex4[],
 									const uniform float RestAngles[],
 									const uniform bool IsBuckled[],
 									uniform float Lambdas[],
 									const uniform float Dt,
-									const uniform float MinStiffness,
 									const uniform float Stiffness,
 									const uniform float BucklingStiffness,
 									const uniform int32 NumConstraints)
@@ -20,12 +78,10 @@ export void ApplyXPBDBendingConstraints(uniform FVector4f PandInvM[],
 
 	foreach(i = 0 ... NumConstraints)
 	{
-		const varying FIntVector4 Constraint = VectorLoad(&Constraints[extract(i,0)]);
-
-		const varying int32 i1 = Constraint.V[0];
-		const varying int32 i2 = Constraint.V[1];
-		const varying int32 i3 = Constraint.V[2];
-		const varying int32 i4 = Constraint.V[3];
+		const varying int32 i1 = ConstraintIndex1[i];
+		const varying int32 i2 = ConstraintIndex2[i];
+		const varying int32 i3 = ConstraintIndex3[i];
+		const varying int32 i4 = ConstraintIndex4[i];
 		
 		const varying FVector4f PandInvM1 = VectorGather(&PandInvM[i1]);
 		const varying FVector4f PandInvM2 = VectorGather(&PandInvM[i2]);
@@ -45,37 +101,29 @@ export void ApplyXPBDBendingConstraints(uniform FVector4f PandInvM[],
 		const varying float Lambda = Lambdas[i];
 
 		varying FVector3f Delta1, Delta2, Delta3, Delta4;
-		const varying float DLambda = GetXPBDBendingElementDeltas(P1, P2, P3, P4, M1, M2, M3, M4, Dt, RestAngles[i], Lambda, StiffnessValue, MinStiffness, Delta1, Delta2, Delta3, Delta4);
+		const varying float DLambda = GetXPBDBendingElementDeltas(P1, P2, P3, P4, M1, M2, M3, M4, Dt, RestAngles[i], Lambda, StiffnessValue, Delta1, Delta2, Delta3, Delta4);
 
 		Lambdas[i] = Lambda + DLambda;
-
-		if (M1 > 0)
-		{
-			VectorScatter(&PandInvM[i1], SetVector4(P1 + Delta1, M1));
-		}
-		if (M2 > 0)
-		{
-			VectorScatter(&PandInvM[i2], SetVector4(P2 + Delta2, M2));
-		}
-		if (M3 > 0)
-		{
-			VectorScatter(&PandInvM[i3], SetVector4(P3 + Delta3, M3));
-		}
-		if (M4 > 0)
-		{
-			VectorScatter(&PandInvM[i4], SetVector4(P4 + Delta4, M4));
-		}
+		VectorScatter(&PandInvM[i1], SetVector4(P1 + Delta1, M1));
+		VectorScatter(&PandInvM[i2], SetVector4(P2 + Delta2, M2));
+		VectorScatter(&PandInvM[i3], SetVector4(P3 + Delta3, M3));
+		VectorScatter(&PandInvM[i4], SetVector4(P4 + Delta4, M4));
 	}
 }
 
 export void ApplyXPBDBendingConstraintsWithDamping(uniform FVector4f PandInvM[],
-									const uniform FVector3f X[],
-									const uniform FIntVector4 Constraints[],
+									const uniform FVector3f X1Array[],
+									const uniform FVector3f X2Array[],
+									const uniform FVector3f X3Array[],
+									const uniform FVector3f X4Array[],
+									const uniform int32 ConstraintIndex1[],
+									const uniform int32 ConstraintIndex2[],
+									const uniform int32 ConstraintIndex3[],
+									const uniform int32 ConstraintIndex4[],
 									const uniform float RestAngles[],
 									const uniform bool IsBuckled[],
 									uniform float Lambdas[],
 									const uniform float Dt,
-									const uniform float MinStiffness,
 									const uniform float Stiffness,
 									const uniform float BucklingStiffness,
 									const uniform float DampingRatio,
@@ -84,12 +132,10 @@ export void ApplyXPBDBendingConstraintsWithDamping(uniform FVector4f PandInvM[],
 
 	foreach(i = 0 ... NumConstraints)
 	{
-		const varying FIntVector4 Constraint = VectorLoad(&Constraints[extract(i,0)]);
-
-		const varying int32 i1 = Constraint.V[0];
-		const varying int32 i2 = Constraint.V[1];
-		const varying int32 i3 = Constraint.V[2];
-		const varying int32 i4 = Constraint.V[3];
+		const varying int32 i1 = ConstraintIndex1[i];
+		const varying int32 i2 = ConstraintIndex2[i];
+		const varying int32 i3 = ConstraintIndex3[i];
+		const varying int32 i4 = ConstraintIndex4[i];
 		
 		const varying FVector4f PandInvM1 = VectorGather(&PandInvM[i1]);
 		const varying FVector4f PandInvM2 = VectorGather(&PandInvM[i2]);
@@ -103,47 +149,42 @@ export void ApplyXPBDBendingConstraintsWithDamping(uniform FVector4f PandInvM[],
 		UnzipPandInvM(PandInvM2, P2, M2);
 		UnzipPandInvM(PandInvM3, P3, M3);
 		UnzipPandInvM(PandInvM4, P4, M4);
+		
+		const uniform int32 i0 = extract(i,0);
 
-		const varying FVector3f X1 = VectorGather(&X[i1]);
-		const varying FVector3f X2 = VectorGather(&X[i2]);
-		const varying FVector3f X3 = VectorGather(&X[i3]);
-		const varying FVector3f X4 = VectorGather(&X[i4]);
+		const varying FVector3f X1 = VectorLoad(&X1Array[i0]);
+		const varying FVector3f X2 = VectorLoad(&X2Array[i0]);
+		const varying FVector3f X3 = VectorLoad(&X3Array[i0]);
+		const varying FVector3f X4 = VectorLoad(&X4Array[i0]);
 		
 		const varying float StiffnessValue = select(IsBuckled[i], BucklingStiffness, Stiffness);
 		const varying float Lambda = Lambdas[i];
 
 		varying FVector3f Delta1, Delta2, Delta3, Delta4;
-		const varying float DLambda = GetXPBDBendingElementDeltasWithDamping(P1, P2, P3, P4, X1, X2, X3, X4, M1, M2, M3, M4, Dt, RestAngles[i], Lambda, StiffnessValue, DampingRatio, MinStiffness, Delta1, Delta2, Delta3, Delta4);
+		const varying float DLambda = GetXPBDBendingElementDeltasWithDamping(P1, P2, P3, P4, X1, X2, X3, X4, M1, M2, M3, M4, Dt, RestAngles[i], Lambda, StiffnessValue, DampingRatio, Delta1, Delta2, Delta3, Delta4);
 
 		Lambdas[i] = Lambda + DLambda;
 
-		if (M1 > 0)
-		{
-			VectorScatter(&PandInvM[i1], SetVector4(P1 + Delta1, M1));
-		}
-		if (M2 > 0)
-		{
-			VectorScatter(&PandInvM[i2], SetVector4(P2 + Delta2, M2));
-		}
-		if (M3 > 0)
-		{
-			VectorScatter(&PandInvM[i3], SetVector4(P3 + Delta3, M3));
-		}
-		if (M4 > 0)
-		{
-			VectorScatter(&PandInvM[i4], SetVector4(P4 + Delta4, M4));
-		}
+		VectorScatter(&PandInvM[i1], SetVector4(P1 + Delta1, M1));
+		VectorScatter(&PandInvM[i2], SetVector4(P2 + Delta2, M2));
+		VectorScatter(&PandInvM[i3], SetVector4(P3 + Delta3, M3));
+		VectorScatter(&PandInvM[i4], SetVector4(P4 + Delta4, M4));
 	}
 }
 
 export void ApplyXPBDBendingDampingConstraints(uniform FVector4f PandInvM[],
-									const uniform FVector3f X[],
-									const uniform FIntVector4 Constraints[],
+									const uniform FVector3f X1Array[],
+									const uniform FVector3f X2Array[],
+									const uniform FVector3f X3Array[],
+									const uniform FVector3f X4Array[],
+									const uniform int32 ConstraintIndex1[],
+									const uniform int32 ConstraintIndex2[],
+									const uniform int32 ConstraintIndex3[],
+									const uniform int32 ConstraintIndex4[],
 									const uniform float RestAngles[],
 									const uniform bool IsBuckled[],
 									uniform float Lambdas[],
 									const uniform float Dt,
-									const uniform float MinStiffness,
 									const uniform float Stiffness,
 									const uniform float BucklingStiffness,
 									const uniform float DampingRatio,
@@ -152,12 +193,10 @@ export void ApplyXPBDBendingDampingConstraints(uniform FVector4f PandInvM[],
 
 	foreach(i = 0 ... NumConstraints)
 	{
-		const varying FIntVector4 Constraint = VectorLoad(&Constraints[extract(i,0)]);
-
-		const varying int32 i1 = Constraint.V[0];
-		const varying int32 i2 = Constraint.V[1];
-		const varying int32 i3 = Constraint.V[2];
-		const varying int32 i4 = Constraint.V[3];
+		const varying int32 i1 = ConstraintIndex1[i];
+		const varying int32 i2 = ConstraintIndex2[i];
+		const varying int32 i3 = ConstraintIndex3[i];
+		const varying int32 i4 = ConstraintIndex4[i];
 		
 		const varying FVector4f PandInvM1 = VectorGather(&PandInvM[i1]);
 		const varying FVector4f PandInvM2 = VectorGather(&PandInvM[i2]);
@@ -172,64 +211,51 @@ export void ApplyXPBDBendingDampingConstraints(uniform FVector4f PandInvM[],
 		UnzipPandInvM(PandInvM3, P3, M3);
 		UnzipPandInvM(PandInvM4, P4, M4);
 
-		const varying FVector3f X1 = VectorGather(&X[i1]);
-		const varying FVector3f X2 = VectorGather(&X[i2]);
-		const varying FVector3f X3 = VectorGather(&X[i3]);
-		const varying FVector3f X4 = VectorGather(&X[i4]);
+		const uniform int32 i0 = extract(i,0);
+
+		const varying FVector3f X1 = VectorLoad(&X1Array[i0]);
+		const varying FVector3f X2 = VectorLoad(&X2Array[i0]);
+		const varying FVector3f X3 = VectorLoad(&X3Array[i0]);
+		const varying FVector3f X4 = VectorLoad(&X4Array[i0]);
 		
 		const varying float StiffnessValue = select(IsBuckled[i], BucklingStiffness, Stiffness);
 		const varying float Lambda = Lambdas[i];
 
 		varying FVector3f Delta1, Delta2, Delta3, Delta4;
-		const varying float DLambda = GetXPBDBendingElementDampingDeltas(P1, P2, P3, P4, X1, X2, X3, X4, M1, M2, M3, M4, Dt, RestAngles[i], Lambda, StiffnessValue, DampingRatio, MinStiffness, Delta1, Delta2, Delta3, Delta4);
+		const varying float DLambda = GetXPBDBendingElementDampingDeltas(P1, P2, P3, P4, X1, X2, X3, X4, M1, M2, M3, M4, Dt, RestAngles[i], Lambda, StiffnessValue, DampingRatio, Delta1, Delta2, Delta3, Delta4);
 
 		Lambdas[i] = Lambda + DLambda;
-
-		if (M1 > 0)
-		{
-			VectorScatter(&PandInvM[i1], SetVector4(P1 + Delta1, M1));
-		}
-		if (M2 > 0)
-		{
-			VectorScatter(&PandInvM[i2], SetVector4(P2 + Delta2, M2));
-		}
-		if (M3 > 0)
-		{
-			VectorScatter(&PandInvM[i3], SetVector4(P3 + Delta3, M3));
-		}
-		if (M4 > 0)
-		{
-			VectorScatter(&PandInvM[i4], SetVector4(P4 + Delta4, M4));
-		}
+		VectorScatter(&PandInvM[i1], SetVector4(P1 + Delta1, M1));
+		VectorScatter(&PandInvM[i2], SetVector4(P2 + Delta2, M2));
+		VectorScatter(&PandInvM[i3], SetVector4(P3 + Delta3, M3));
+		VectorScatter(&PandInvM[i4], SetVector4(P4 + Delta4, M4));
 	}
 }
 
 export void ApplyXPBDBendingConstraintsWithMaps(uniform FVector4f PandInvM[],
-											const uniform FIntVector4 Constraints[],
+											const uniform int32 ConstraintIndex1[],
+											const uniform int32 ConstraintIndex2[],
+											const uniform int32 ConstraintIndex3[],
+											const uniform int32 ConstraintIndex4[],
 											const uniform float RestAngles[],
 											const uniform bool IsBuckled[],
 											uniform float Lambdas[],
 											const uniform float Dt,
-											const uniform float MinStiffness,
 											const uniform bool StiffnessHasMap,
-											const uniform uint8 StiffnessIndices[],
-											const uniform float StiffnessTable[],
+											const uniform FVector2f& StiffnessOffsetRange,
+											const uniform float StiffnessMapValues[],
 											const uniform bool BucklingStiffnessHasMap,
-											const uniform uint8 BucklingStiffnessIndices[],
-											const uniform float BucklingStiffnessTable[],
+											const uniform FVector2f& BucklingStiffnessOffsetRange,
+											const uniform float BucklingStiffnessMapValues[],
 											const uniform int32 NumConstraints)
 {
 
 	foreach(i = 0 ... NumConstraints)
 	{
-		const varying FIntVector4 Constraint = VectorLoad(&Constraints[extract(i,0)]);
-
-		const varying int32 i1 = Constraint.V[0];
-		const varying int32 i2 = Constraint.V[1];
-		const varying int32 i3 = Constraint.V[2];
-		const varying int32 i4 = Constraint.V[3];
-		const varying int8 StiffnessIndex = StiffnessHasMap ? StiffnessIndices[i] : 0;
-		const varying int8 BucklingStiffnessIndex = BucklingStiffnessHasMap ? BucklingStiffnessIndices[i] : 0;
+		const varying int32 i1 = ConstraintIndex1[i];
+		const varying int32 i2 = ConstraintIndex2[i];
+		const varying int32 i3 = ConstraintIndex3[i];
+		const varying int32 i4 = ConstraintIndex4[i];
 
 		const varying FVector4f PandInvM2 = VectorGather(&PandInvM[i2]);
 		const varying FVector4f PandInvM3 = VectorGather(&PandInvM[i3]);
@@ -244,70 +270,56 @@ export void ApplyXPBDBendingConstraintsWithMaps(uniform FVector4f PandInvM[],
 		UnzipPandInvM(PandInvM3, P3, M3);
 		UnzipPandInvM(PandInvM4, P4, M4);
 		
-		#pragma ignore warning(perf)
-		const varying float Stiffness = StiffnessTable[StiffnessIndex];
-		#pragma ignore warning(perf)
-		const varying float BucklingStiffness = BucklingStiffnessTable[BucklingStiffnessIndex];
+		const varying float Stiffness = StiffnessHasMap ? StiffnessOffsetRange.V[0] + StiffnessMapValues[i]*StiffnessOffsetRange.V[1] : StiffnessOffsetRange.V[0];
+		const varying float BucklingStiffness = BucklingStiffnessHasMap ? BucklingStiffnessOffsetRange.V[0] + BucklingStiffnessMapValues[i]*BucklingStiffnessOffsetRange.V[1] : BucklingStiffnessOffsetRange.V[0];
 		
 		const varying float StiffnessValue = select(IsBuckled[i], BucklingStiffness, Stiffness);
 		
 		const varying float Lambda = Lambdas[i];
 
 		varying FVector3f Delta1, Delta2, Delta3, Delta4;
-		const varying float DLambda = GetXPBDBendingElementDeltas(P1, P2, P3, P4, M1, M2, M3, M4, Dt, RestAngles[i], Lambda, StiffnessValue, MinStiffness, Delta1, Delta2, Delta3, Delta4);
+		const varying float DLambda = GetXPBDBendingElementDeltas(P1, P2, P3, P4, M1, M2, M3, M4, Dt, RestAngles[i], Lambda, StiffnessValue, Delta1, Delta2, Delta3, Delta4);
 
 		Lambdas[i] = Lambda + DLambda;
 
-		if (M1 > 0)
-		{
-			VectorScatter(&PandInvM[i1], SetVector4(P1 + Delta1, M1));
-		}
-		if (M2 > 0)
-		{
-			VectorScatter(&PandInvM[i2], SetVector4(P2 + Delta2, M2));
-		}
-		if (M3 > 0)
-		{
-			VectorScatter(&PandInvM[i3], SetVector4(P3 + Delta3, M3));
-		}
-		if (M4 > 0)
-		{
-			VectorScatter(&PandInvM[i4], SetVector4(P4 + Delta4, M4));
-		}
+		VectorScatter(&PandInvM[i1], SetVector4(P1 + Delta1, M1));
+		VectorScatter(&PandInvM[i2], SetVector4(P2 + Delta2, M2));
+		VectorScatter(&PandInvM[i3], SetVector4(P3 + Delta3, M3));
+		VectorScatter(&PandInvM[i4], SetVector4(P4 + Delta4, M4));
 	}
 }
 
 export void ApplyXPBDBendingConstraintsWithDampingAndMaps(uniform FVector4f PandInvM[],
-											const uniform FVector3f X[],
-											const uniform FIntVector4 Constraints[],
+											const uniform FVector3f X1Array[],
+											const uniform FVector3f X2Array[],
+											const uniform FVector3f X3Array[],
+											const uniform FVector3f X4Array[],
+											const uniform int32 ConstraintIndex1[],
+											const uniform int32 ConstraintIndex2[],
+											const uniform int32 ConstraintIndex3[],
+											const uniform int32 ConstraintIndex4[],
 											const uniform float RestAngles[],
 											const uniform bool IsBuckled[],
 											uniform float Lambdas[],
 											const uniform float Dt,
-											const uniform float MinStiffness,
 											const uniform bool StiffnessHasMap,
-											const uniform uint8 StiffnessIndices[],
-											const uniform float StiffnessTable[],
+											const uniform FVector2f& StiffnessOffsetRange,
+											const uniform float StiffnessMapValues[],
 											const uniform bool BucklingStiffnessHasMap,
-											const uniform uint8 BucklingStiffnessIndices[],
-											const uniform float BucklingStiffnessTable[],
+											const uniform FVector2f& BucklingStiffnessOffsetRange,
+											const uniform float BucklingStiffnessMapValues[],
 											const uniform bool DampingHasMap,
-											const uniform uint8 DampingIndices[],
-											const uniform float DampingTable[],
+											const uniform FVector2f& DampingOffsetRange,
+											const uniform float DampingMapValues[],
 											const uniform int32 NumConstraints)
 {
 
 	foreach(i = 0 ... NumConstraints)
 	{
-		const varying FIntVector4 Constraint = VectorLoad(&Constraints[extract(i,0)]);
-
-		const varying int32 i1 = Constraint.V[0];
-		const varying int32 i2 = Constraint.V[1];
-		const varying int32 i3 = Constraint.V[2];
-		const varying int32 i4 = Constraint.V[3];
-		const varying int8 StiffnessIndex = StiffnessHasMap ? StiffnessIndices[i] : 0;
-		const varying int8 BucklingStiffnessIndex = BucklingStiffnessHasMap ? BucklingStiffnessIndices[i] : 0;
-		const varying int8 DampingIndex = DampingHasMap ? DampingIndices[i] : 0;
+		const varying int32 i1 = ConstraintIndex1[i];
+		const varying int32 i2 = ConstraintIndex2[i];
+		const varying int32 i3 = ConstraintIndex3[i];
+		const varying int32 i4 = ConstraintIndex4[i];
 
 		const varying FVector4f PandInvM2 = VectorGather(&PandInvM[i2]);
 		const varying FVector4f PandInvM3 = VectorGather(&PandInvM[i3]);
@@ -322,78 +334,64 @@ export void ApplyXPBDBendingConstraintsWithDampingAndMaps(uniform FVector4f Pand
 		UnzipPandInvM(PandInvM3, P3, M3);
 		UnzipPandInvM(PandInvM4, P4, M4);
 
-
-		#pragma ignore warning(perf)
-		const varying float Stiffness = StiffnessTable[StiffnessIndex];
-		#pragma ignore warning(perf)
-		const varying float BucklingStiffness = BucklingStiffnessTable[BucklingStiffnessIndex];		
-		#pragma ignore warning(perf)
-		const varying float DampingRatio = DampingTable[DampingIndex];
+		const varying float Stiffness = StiffnessHasMap ? StiffnessOffsetRange.V[0] + StiffnessMapValues[i]*StiffnessOffsetRange.V[1] : StiffnessOffsetRange.V[0];
+		const varying float BucklingStiffness = BucklingStiffnessHasMap ? BucklingStiffnessOffsetRange.V[0] + BucklingStiffnessMapValues[i]*BucklingStiffnessOffsetRange.V[1] : BucklingStiffnessOffsetRange.V[0];
+		const varying float DampingRatio = DampingHasMap ? DampingOffsetRange.V[0] + DampingMapValues[i]*DampingOffsetRange.V[1] : DampingOffsetRange.V[0];
 		
 		const varying float StiffnessValue = select(IsBuckled[i], BucklingStiffness, Stiffness);
 
-		const varying FVector3f X1 = VectorGather(&X[i1]);
-		const varying FVector3f X2 = VectorGather(&X[i2]);
-		const varying FVector3f X3 = VectorGather(&X[i3]);
-		const varying FVector3f X4 = VectorGather(&X[i4]);
+		const uniform int32 i0 = extract(i,0);
+
+		const varying FVector3f X1 = VectorLoad(&X1Array[i0]);
+		const varying FVector3f X2 = VectorLoad(&X2Array[i0]);
+		const varying FVector3f X3 = VectorLoad(&X3Array[i0]);
+		const varying FVector3f X4 = VectorLoad(&X4Array[i0]);
 		
 		const varying float Lambda = Lambdas[i];
 
 		varying FVector3f Delta1, Delta2, Delta3, Delta4;
-		const varying float DLambda = GetXPBDBendingElementDeltasWithDamping(P1, P2, P3, P4, X1, X2, X3, X4, M1, M2, M3, M4, Dt, RestAngles[i], Lambda, StiffnessValue, DampingRatio, MinStiffness, Delta1, Delta2, Delta3, Delta4);
+		const varying float DLambda = GetXPBDBendingElementDeltasWithDamping(P1, P2, P3, P4, X1, X2, X3, X4, M1, M2, M3, M4, Dt, RestAngles[i], Lambda, StiffnessValue, DampingRatio, Delta1, Delta2, Delta3, Delta4);
 
 		Lambdas[i] = Lambda + DLambda;
 
-		if (M1 > 0)
-		{
-			VectorScatter(&PandInvM[i1], SetVector4(P1 + Delta1, M1));
-		}
-		if (M2 > 0)
-		{
-			VectorScatter(&PandInvM[i2], SetVector4(P2 + Delta2, M2));
-		}
-		if (M3 > 0)
-		{
-			VectorScatter(&PandInvM[i3], SetVector4(P3 + Delta3, M3));
-		}
-		if (M4 > 0)
-		{
-			VectorScatter(&PandInvM[i4], SetVector4(P4 + Delta4, M4));
-		}
+		VectorScatter(&PandInvM[i1], SetVector4(P1 + Delta1, M1));
+		VectorScatter(&PandInvM[i2], SetVector4(P2 + Delta2, M2));
+		VectorScatter(&PandInvM[i3], SetVector4(P3 + Delta3, M3));
+		VectorScatter(&PandInvM[i4], SetVector4(P4 + Delta4, M4));
 	}
 }
 
 export void ApplyXPBDBendingDampingConstraintsWithMaps(uniform FVector4f PandInvM[],
-											const uniform FVector3f X[],
-											const uniform FIntVector4 Constraints[],
+											const uniform FVector3f X1Array[],
+											const uniform FVector3f X2Array[],
+											const uniform FVector3f X3Array[],
+											const uniform FVector3f X4Array[],
+											const uniform int32 ConstraintIndex1[],
+											const uniform int32 ConstraintIndex2[],
+											const uniform int32 ConstraintIndex3[],
+											const uniform int32 ConstraintIndex4[],
 											const uniform float RestAngles[],
 											const uniform bool IsBuckled[],
 											uniform float Lambdas[],
 											const uniform float Dt,
-											const uniform float MinStiffness,
 											const uniform bool StiffnessHasMap,
-											const uniform uint8 StiffnessIndices[],
-											const uniform float StiffnessTable[],
+											const uniform FVector2f& StiffnessOffsetRange,
+											const uniform float StiffnessMapValues[],
 											const uniform bool BucklingStiffnessHasMap,
-											const uniform uint8 BucklingStiffnessIndices[],
-											const uniform float BucklingStiffnessTable[],
+											const uniform FVector2f& BucklingStiffnessOffsetRange,
+											const uniform float BucklingStiffnessMapValues[],
 											const uniform bool DampingHasMap,
-											const uniform uint8 DampingIndices[],
-											const uniform float DampingTable[],
+											const uniform FVector2f& DampingOffsetRange,
+											const uniform float DampingMapValues[],
 											const uniform int32 NumConstraints)
 {
 
 	foreach(i = 0 ... NumConstraints)
 	{
-		const varying FIntVector4 Constraint = VectorLoad(&Constraints[extract(i,0)]);
-
-		const varying int32 i1 = Constraint.V[0];
-		const varying int32 i2 = Constraint.V[1];
-		const varying int32 i3 = Constraint.V[2];
-		const varying int32 i4 = Constraint.V[3];
-		const varying int8 StiffnessIndex = StiffnessHasMap ? StiffnessIndices[i] : 0;
-		const varying int8 BucklingStiffnessIndex = BucklingStiffnessHasMap ? BucklingStiffnessIndices[i] : 0;
-		const varying int8 DampingIndex = DampingHasMap ? DampingIndices[i] : 0;
+		const varying int32 i1 = ConstraintIndex1[i];
+		const varying int32 i2 = ConstraintIndex2[i];
+		const varying int32 i3 = ConstraintIndex3[i];
+		const varying int32 i4 = ConstraintIndex4[i];
 
 		const varying FVector4f PandInvM2 = VectorGather(&PandInvM[i2]);
 		const varying FVector4f PandInvM3 = VectorGather(&PandInvM[i3]);
@@ -408,43 +406,29 @@ export void ApplyXPBDBendingDampingConstraintsWithMaps(uniform FVector4f PandInv
 		UnzipPandInvM(PandInvM3, P3, M3);
 		UnzipPandInvM(PandInvM4, P4, M4);
 
-
-		#pragma ignore warning(perf)
-		const varying float Stiffness = StiffnessTable[StiffnessIndex];
-		#pragma ignore warning(perf)
-		const varying float BucklingStiffness = BucklingStiffnessTable[BucklingStiffnessIndex];		
-		#pragma ignore warning(perf)
-		const varying float DampingRatio = DampingTable[DampingIndex];
+		const varying float Stiffness = StiffnessHasMap ? StiffnessOffsetRange.V[0] + StiffnessMapValues[i]*StiffnessOffsetRange.V[1] : StiffnessOffsetRange.V[0];
+		const varying float BucklingStiffness = BucklingStiffnessHasMap ? BucklingStiffnessOffsetRange.V[0] + BucklingStiffnessMapValues[i]*BucklingStiffnessOffsetRange.V[1] : BucklingStiffnessOffsetRange.V[0];
+		const varying float DampingRatio = DampingHasMap ? DampingOffsetRange.V[0] + DampingMapValues[i]*DampingOffsetRange.V[1] : DampingOffsetRange.V[0];
 		
 		const varying float StiffnessValue = select(IsBuckled[i], BucklingStiffness, Stiffness);
 
-		const varying FVector3f X1 = VectorGather(&X[i1]);
-		const varying FVector3f X2 = VectorGather(&X[i2]);
-		const varying FVector3f X3 = VectorGather(&X[i3]);
-		const varying FVector3f X4 = VectorGather(&X[i4]);
+		const uniform int32 i0 = extract(i,0);
+
+		const varying FVector3f X1 = VectorLoad(&X1Array[i0]);
+		const varying FVector3f X2 = VectorLoad(&X2Array[i0]);
+		const varying FVector3f X3 = VectorLoad(&X3Array[i0]);
+		const varying FVector3f X4 = VectorLoad(&X4Array[i0]);
 		
 		const varying float Lambda = Lambdas[i];
 
 		varying FVector3f Delta1, Delta2, Delta3, Delta4;
-		const varying float DLambda = GetXPBDBendingElementDampingDeltas(P1, P2, P3, P4, X1, X2, X3, X4, M1, M2, M3, M4, Dt, RestAngles[i], Lambda, StiffnessValue, DampingRatio, MinStiffness, Delta1, Delta2, Delta3, Delta4);
+		const varying float DLambda = GetXPBDBendingElementDampingDeltas(P1, P2, P3, P4, X1, X2, X3, X4, M1, M2, M3, M4, Dt, RestAngles[i], Lambda, StiffnessValue, DampingRatio, Delta1, Delta2, Delta3, Delta4);
 
 		Lambdas[i] = Lambda + DLambda;
-
-		if (M1 > 0)
-		{
-			VectorScatter(&PandInvM[i1], SetVector4(P1 + Delta1, M1));
-		}
-		if (M2 > 0)
-		{
-			VectorScatter(&PandInvM[i2], SetVector4(P2 + Delta2, M2));
-		}
-		if (M3 > 0)
-		{
-			VectorScatter(&PandInvM[i3], SetVector4(P3 + Delta3, M3));
-		}
-		if (M4 > 0)
-		{
-			VectorScatter(&PandInvM[i4], SetVector4(P4 + Delta4, M4));
-		}
+		
+		VectorScatter(&PandInvM[i1], SetVector4(P1 + Delta1, M1));
+		VectorScatter(&PandInvM[i2], SetVector4(P2 + Delta2, M2));
+		VectorScatter(&PandInvM[i3], SetVector4(P3 + Delta3, M3));
+		VectorScatter(&PandInvM[i4], SetVector4(P4 + Delta4, M4));
 	}
 }
diff --git a/Engine/Source/Runtime/Experimental/Chaos/Private/Chaos/XPBDBendingElement.isph b/Engine/Source/Runtime/Experimental/Chaos/Private/Chaos/XPBDBendingElement.isph
--- a/Engine/Source/Runtime/Experimental/Chaos/Private/Chaos/XPBDBendingElement.isph
+++ b/Engine/Source/Runtime/Experimental/Chaos/Private/Chaos/XPBDBendingElement.isph
@@ -49,18 +49,19 @@ inline void CalcGradients(const FVector3f& P1, const FVector3f& P2, const FVecto
 // Returns Angle
 inline float CalcGradientsAndAngle(const FVector3f& P1, const FVector3f& P2, const FVector3f& P3, const FVector3f& P4, FVector3f& Grad1, FVector3f& Grad2, FVector3f& Grad3, FVector3f& Grad4)
 {
-	const FVector3f SharedEdgeNormalized = VectorGetSafeNormal(P2 - P1);
-	const FVector3f P13CrossP23 = VectorCross(P1 - P3, P2 - P3);
-	const float Normal1Len = VectorSize(P13CrossP23);
-	const FVector3f Normal1 = SafeDivide(P13CrossP23, Normal1Len);
-	const FVector3f P24CrossP14 = VectorCross(P2 - P4, P1 - P4);
-	const float Normal2Len = VectorSize(P24CrossP14);
-	const FVector3f Normal2 = SafeDivide(P24CrossP14, Normal2Len);
+	const varying FVector3f SharedEdgeNormalized = VectorGetSafeNormal(P2 - P1);
+	const varying FVector3f P13CrossP23 = VectorCross(P1 - P3, P2 - P3);
+	const varying float Normal1Len = VectorSize(P13CrossP23);
+	const varying FVector3f Normal1 = SafeDivide(P13CrossP23, Normal1Len);
+	const varying FVector3f P24CrossP14 = VectorCross(P2 - P4, P1 - P4);
+	const varying float Normal2Len = VectorSize(P24CrossP14);
+	const varying FVector3f Normal2 = SafeDivide(P24CrossP14, Normal2Len);
 
-	const FVector3f N2CrossN1 = VectorCross(Normal2, Normal1);
+	const varying FVector3f N2CrossN1 = VectorCross(Normal2, Normal1);
 
 	const float CosPhi = clamp(VectorDot(Normal1, Normal2), (uniform float)(-1), (uniform float)(1));
 	const float SinPhi = clamp(VectorDot(N2CrossN1, SharedEdgeNormalized), (uniform float)(-1), (uniform float)(1));
+	const float Angle = atan2(SinPhi, CosPhi);
 
 	const FVector3f DPhiDN1_OverNormal1Len = SafeDivide(CosPhi * VectorCross(SharedEdgeNormalized, Normal2) - SinPhi * Normal2, Normal1Len);
 	const FVector3f DPhiDN2_OverNormal2Len = SafeDivide(CosPhi * VectorCross(Normal1, SharedEdgeNormalized) - SinPhi * Normal1, Normal2Len);
@@ -75,13 +76,13 @@ inline float CalcGradientsAndAngle(const FVector3f& P1, const FVector3f& P2, con
 	Grad3 = -1 * DPhiDP13 - DPhiDP23;
 	Grad4 = -1 * DPhiDP14 - DPhiDP24;
 
-	return atan2(SinPhi, CosPhi);
+	return Angle;
 }
 
 // Elastic part only
 inline float GetXPBDBendingElementDeltas(const FVector3f& P1, const FVector3f& P2, const FVector3f& P3, const FVector3f& P4,
 	const varying float InvM1, const varying float InvM2, const varying float InvM3, const varying float InvM4,
-	const uniform float Dt, const varying float RestAngle, const varying float Lambda, const varying float StiffnessValue, const uniform float MinStiffness, 
+	const uniform float Dt, const varying float RestAngle, const varying float Lambda, const varying float StiffnessValue,
 	FVector3f& Delta1, FVector3f& Delta2, FVector3f& Delta3, FVector3f& Delta4)
 {
 	varying FVector3f Grad1, Grad2, Grad3, Grad4;
@@ -91,8 +92,7 @@ inline float GetXPBDBendingElementDeltas(const FVector3f& P1, const FVector3f& P
 	const varying float ElasticTerm = AlphaInv * (Angle - RestAngle);
 	const varying float Denom = AlphaInv*(InvM1 * VectorSizeSquared(Grad1) + InvM2 * VectorSizeSquared(Grad2) + InvM3 * VectorSizeSquared(Grad3) + InvM4 * VectorSizeSquared(Grad4)) + 1.f;
 
-	const varying float ComputedDLambda = (-ElasticTerm - Lambda) / Denom;
-	const varying float DLambda = select(StiffnessValue < MinStiffness, (uniform float)0, ComputedDLambda);
+	const varying float DLambda = (-ElasticTerm - Lambda) / Denom;
 
 	Delta1 = DLambda * InvM1 * Grad1;
 	Delta2 = DLambda * InvM2 * Grad2;
@@ -107,7 +107,7 @@ inline float GetXPBDBendingElementDeltasWithDamping(const FVector3f& P1, const F
 	const FVector3f& X1, const FVector3f& X2, const FVector3f& X3, const FVector3f& X4,
 	const varying float InvM1, const varying float InvM2, const varying float InvM3, const varying float InvM4,
 	const uniform float Dt, const varying float RestAngle, const varying float Lambda, 
-	const varying float StiffnessValue, const DampingFloatType DampingRatio, const uniform float MinStiffness, 
+	const varying float StiffnessValue, const DampingFloatType DampingRatio,
 	FVector3f& Delta1, FVector3f& Delta2, FVector3f& Delta3, FVector3f& Delta4)
 {
 	varying FVector3f Grad1, Grad2, Grad3, Grad4;
@@ -127,9 +127,7 @@ inline float GetXPBDBendingElementDeltasWithDamping(const FVector3f& P1, const F
 
 	const varying float Denom = (AlphaInv + BetaDt)*(InvM1 * VectorSizeSquared(Grad1) + InvM2 * VectorSizeSquared(Grad2) + InvM3 * VectorSizeSquared(Grad3) + InvM4 * VectorSizeSquared(Grad4)) + 1.f;
 
-	const varying float ComputedDLambda = (-ElasticTerm - DampingTerm - Lambda) / Denom;
-
-	const varying float DLambda = select(StiffnessValue < MinStiffness, (uniform float)0, ComputedDLambda);
+	const varying float DLambda = (-ElasticTerm - DampingTerm - Lambda) / Denom;
 
 	Delta1 = DLambda * InvM1 * Grad1;
 	Delta2 = DLambda * InvM2 * Grad2;
@@ -144,7 +142,7 @@ inline float GetXPBDBendingElementDampingDeltas(const FVector3f& P1, const FVect
 	const FVector3f& X1, const FVector3f& X2, const FVector3f& X3, const FVector3f& X4,
 	const varying float InvM1, const varying float InvM2, const varying float InvM3, const varying float InvM4,
 	const uniform float Dt, const varying float RestAngle, const varying float Lambda, 
-	const varying float StiffnessValue, const DampingFloatType DampingRatio, const uniform float MinStiffness, 
+	const varying float StiffnessValue, const DampingFloatType DampingRatio,
 	FVector3f& Delta1, FVector3f& Delta2, FVector3f& Delta3, FVector3f& Delta4)
 {
 	varying FVector3f Grad1, Grad2, Grad3, Grad4;
@@ -162,9 +160,7 @@ inline float GetXPBDBendingElementDampingDeltas(const FVector3f& P1, const FVect
 
 	const varying float Denom = BetaDt*(InvM1 * VectorSizeSquared(Grad1) + InvM2 * VectorSizeSquared(Grad2) + InvM3 * VectorSizeSquared(Grad3) + InvM4 * VectorSizeSquared(Grad4)) + 1.f;
 
-	const varying float ComputedDLambda = ( -DampingTerm - Lambda) / Denom;
-
-	const varying float DLambda = select(StiffnessValue < MinStiffness, (uniform float)0, ComputedDLambda);
+	const varying float DLambda = ( -DampingTerm - Lambda) / Denom;
 
 	Delta1 = DLambda * InvM1 * Grad1;
 	Delta2 = DLambda * InvM2 * Grad2;
