Chaos Cloth optimizations to XPBD constraints.

diff --git a/Engine/Source/Runtime/Experimental/Chaos/Private/Chaos/XPBDSpringConstraints.cpp b/Engine/Source/Runtime/Experimental/Chaos/Private/Chaos/XPBDSpringConstraints.cpp
--- a/Engine/Source/Runtime/Experimental/Chaos/Private/Chaos/XPBDSpringConstraints.cpp
+++ b/Engine/Source/Runtime/Experimental/Chaos/Private/Chaos/XPBDSpringConstraints.cpp
@@ -92,35 +92,29 @@ void FXPBDSpringConstraints::ApplyHelper(SolverParticlesOrRange& Particles, cons
 	if constexpr (bDampingBefore)
 	{
 		Delta += Spring::GetXPBDSpringDampingDelta(Particles, Dt, Constraint, Dists[ConstraintIndex], LambdasDamping[ConstraintIndex],
-			ExpStiffnessValue, MinStiffness, DampingRatioValue);
+			ExpStiffnessValue, DampingRatioValue);
 	}
 
 	if constexpr (bSingleLambda)
 	{
 		Delta += Spring::GetXPBDSpringDeltaWithDamping(Particles, Dt, Constraint, Dists[ConstraintIndex], Lambdas[ConstraintIndex],
-			ExpStiffnessValue, MinStiffness, DampingRatioValue);
+			ExpStiffnessValue, DampingRatioValue);
 	}
 
 	if constexpr (bSeparateStretch)
 	{
 		Delta += Spring::GetXPBDSpringDelta(Particles, Dt, Constraint, Dists[ConstraintIndex], Lambdas[ConstraintIndex],
-			ExpStiffnessValue, MinStiffness);
+			ExpStiffnessValue);
 	}
 
 	if constexpr (bDampingAfter)
 	{
 		Delta += Spring::GetXPBDSpringDampingDelta(Particles, Dt, Constraint, Dists[ConstraintIndex], LambdasDamping[ConstraintIndex],
-			ExpStiffnessValue, MinStiffness, DampingRatioValue);
-	}
-
-	if (Particles.InvM(Index1) > (FSolverReal)0.)
-	{
-		Particles.P(Index1) += Particles.InvM(Index1) * Delta;
-	}
-	if (Particles.InvM(Index2) > (FSolverReal)0.)
-	{
-		Particles.P(Index2) -= Particles.InvM(Index2) * Delta;
+			ExpStiffnessValue, DampingRatioValue);
 	}
+	
+	Particles.P(Index1) += Particles.InvM(Index1) * Delta;
+	Particles.P(Index2) -= Particles.InvM(Index2) * Delta;
 }
 
 template<typename SolverParticlesOrRange>
@@ -220,7 +214,6 @@ void FXPBDSpringConstraints::Apply(SolverParticlesOrRange& Particles, const FSol
 								&Dists.GetData()[ColorStart],
 								&Lambdas.GetData()[ColorStart],
 								Dt,
-								MinStiffness,
 								StiffnessHasWeightMap,
 								StiffnessHasWeightMap ? &Stiffness.GetIndices().GetData()[ColorStart] : nullptr,
 								&Stiffness.GetTable().GetData()[0],
@@ -244,7 +237,6 @@ void FXPBDSpringConstraints::Apply(SolverParticlesOrRange& Particles, const FSol
 								&Dists.GetData()[ColorStart],
 								&LambdasDamping.GetData()[ColorStart],
 								Dt,
-								MinStiffness,
 								StiffnessHasWeightMap,
 								StiffnessHasWeightMap ? &Stiffness.GetIndices().GetData()[ColorStart] : nullptr,
 								&Stiffness.GetTable().GetData()[0],
@@ -265,7 +257,6 @@ void FXPBDSpringConstraints::Apply(SolverParticlesOrRange& Particles, const FSol
 						&Dists.GetData()[ColorStart],
 						&Lambdas.GetData()[ColorStart],
 						Dt,
-						MinStiffness,
 						&Stiffness.GetIndices().GetData()[ColorStart],
 						&Stiffness.GetTable().GetData()[0],
 						ColorSize);
diff --git a/Engine/Source/Runtime/Experimental/Chaos/Private/Chaos/XPBDSpringConstraints.ispc b/Engine/Source/Runtime/Experimental/Chaos/Private/Chaos/XPBDSpringConstraints.ispc
--- a/Engine/Source/Runtime/Experimental/Chaos/Private/Chaos/XPBDSpringConstraints.ispc
+++ b/Engine/Source/Runtime/Experimental/Chaos/Private/Chaos/XPBDSpringConstraints.ispc
@@ -32,18 +32,11 @@ export void ApplyXPBDSpringConstraints(uniform FVector4f PandInvM[],
 		const varying float Lambda = Lambdas[i];
 
 		varying float DLambda;
-		const varying FVector3f Delta = GetXPBDSpringDelta(P1, P2, InvM1, InvM2, Dt, Dist, Lambda, Stiffness, FLOAT_SMALL_NUMBER, DLambda);
+		const varying FVector3f Delta = GetXPBDSpringDelta(P1, P2, InvM1, InvM2, Dt, Dist, Lambda, Stiffness, DLambda);
 
 		Lambdas[i] = Lambda + DLambda;
-
-		if (InvM1 > FLOAT_ZERO)
-		{
-			VectorScatter(&PandInvM[i1], SetVector4( P1 + (InvM1 * Delta), InvM1 ));
-		}
-		if (InvM2 > FLOAT_ZERO)
-		{
-			VectorScatter(&PandInvM[i2], SetVector4( P2 - (InvM2 * Delta), InvM2 ));
-		}
+		VectorScatter(&PandInvM[i1], SetVector4( P1 + (InvM1 * Delta), InvM1 ));
+		VectorScatter(&PandInvM[i2], SetVector4( P2 - (InvM2 * Delta), InvM2 ));
 	}
 }
 
@@ -78,18 +71,11 @@ export void ApplyXPBDSpringConstraintsWithDamping(uniform FVector4f PandInvM[],
 		const varying float Lambda = Lambdas[i];
 
 		varying float DLambda;
-		const FVector3f Delta = GetXPBDSpringDeltaWithDamping(P1, P2, X1, X2, InvM1, InvM2, Dt, Dist, Lambda, Stiffness, DampingRatio, FLOAT_SMALL_NUMBER, DLambda);
+		const FVector3f Delta = GetXPBDSpringDeltaWithDamping(P1, P2, X1, X2, InvM1, InvM2, Dt, Dist, Lambda, Stiffness, DampingRatio, DLambda);
 
 		Lambdas[i] = Lambda + DLambda;
-
-		if (InvM1 > FLOAT_ZERO)
-		{
-			VectorScatter(&PandInvM[i1], SetVector4( P1 + (InvM1 * Delta), InvM1 ));
-		}
-		if (InvM2 > FLOAT_ZERO)
-		{
-			VectorScatter(&PandInvM[i2], SetVector4( P2 - (InvM2 * Delta), InvM2 ));
-		}
+		VectorScatter(&PandInvM[i1], SetVector4( P1 + (InvM1 * Delta), InvM1 ));
+		VectorScatter(&PandInvM[i2], SetVector4( P2 - (InvM2 * Delta), InvM2 ));
 	}
 }
 
@@ -124,18 +110,12 @@ export void ApplyXPBDSpringDampingConstraints(uniform FVector4f PandInvM[],
 		const varying float Lambda = Lambdas[i];
 
 		varying float DLambda;
-		const FVector3f Delta = GetXPBDSpringDampingDelta(P1, P2, X1, X2, InvM1, InvM2, Dt, Dist, Lambda, Stiffness, DampingRatio, FLOAT_SMALL_NUMBER, DLambda);
+		const FVector3f Delta = GetXPBDSpringDampingDelta(P1, P2, X1, X2, InvM1, InvM2, Dt, Dist, Lambda, Stiffness, DampingRatio, DLambda);
 
 		Lambdas[i] = Lambda + DLambda;
-
-		if (InvM1 > FLOAT_ZERO)
-		{
-			VectorScatter(&PandInvM[i1], SetVector4( P1 + (InvM1 * Delta), InvM1 ));
-		}
-		if (InvM2 > FLOAT_ZERO)
-		{
-			VectorScatter(&PandInvM[i2], SetVector4( P2 - (InvM2 * Delta), InvM2 ));
-		}
+		
+		VectorScatter(&PandInvM[i1], SetVector4( P1 + (InvM1 * Delta), InvM1 ));
+		VectorScatter(&PandInvM[i2], SetVector4( P2 - (InvM2 * Delta), InvM2 ));
 	}
 }
 
@@ -144,7 +124,6 @@ export void ApplyXPBDSpringConstraintsWithWeightMaps(uniform FVector4f PandInvM[
 									const uniform float Dists[],
 									uniform float Lambdas[],
 									const uniform float Dt,
-									const uniform float MinStiffness,
 									const uniform uint8 StiffnessIndices[],
 									const uniform float StiffnessTable[],
 									const uniform int32 NumConstraints)
@@ -173,18 +152,12 @@ export void ApplyXPBDSpringConstraintsWithWeightMaps(uniform FVector4f PandInvM[
 		const varying float Stiffness = StiffnessTable[StiffnessIndex];
 
 		varying float DLambda;
-		const varying FVector3f Delta = GetXPBDSpringDelta(P1, P2, InvM1, InvM2, Dt, Dist, Lambda, Stiffness, MinStiffness, DLambda);
+		const varying FVector3f Delta = GetXPBDSpringDelta(P1, P2, InvM1, InvM2, Dt, Dist, Lambda, Stiffness, DLambda);
 
 		Lambdas[i] = Lambda + DLambda;
-
-		if (InvM1 > FLOAT_ZERO)
-		{
-			VectorScatter(&PandInvM[i1], SetVector4( P1 + (InvM1 * Delta), InvM1 ));
-		}
-		if (InvM2 > FLOAT_ZERO)
-		{
-			VectorScatter(&PandInvM[i2], SetVector4( P2 + (InvM2 * Delta), InvM2 ));
-		}
+		
+		VectorScatter(&PandInvM[i1], SetVector4( P1 + (InvM1 * Delta), InvM1 ));
+		VectorScatter(&PandInvM[i2], SetVector4( P2 + (InvM2 * Delta), InvM2 ));
 	}
 }
 
@@ -194,7 +167,6 @@ export void ApplyXPBDSpringConstraintsWithDampingAndWeightMaps(uniform FVector4f
 															const uniform float Dists[],
 															uniform float Lambdas[],
 															const uniform float Dt,
-															const uniform float MinStiffness,
 															const uniform bool StiffnessHasMap,
 															const uniform uint8 StiffnessIndices[],
 															const uniform float StiffnessTable[],
@@ -232,18 +204,12 @@ export void ApplyXPBDSpringConstraintsWithDampingAndWeightMaps(uniform FVector4f
 		const varying float DampingRatio = DampingTable[DampingIndex];
 
 		varying float DLambda;
-		const FVector3f Delta = GetXPBDSpringDeltaWithDamping(P1, P2, X1, X2, InvM1, InvM2, Dt, Dist, Lambda, Stiffness, DampingRatio, MinStiffness, DLambda);
+		const FVector3f Delta = GetXPBDSpringDeltaWithDamping(P1, P2, X1, X2, InvM1, InvM2, Dt, Dist, Lambda, Stiffness, DampingRatio, DLambda);
 
 		Lambdas[i] = Lambda + DLambda;
-
-		if (InvM1 > FLOAT_ZERO)
-		{
-			VectorScatter(&PandInvM[i1], SetVector4( P1 + (InvM1 * Delta), InvM1 ));
-		}
-		if (InvM2 > FLOAT_ZERO)
-		{
-			VectorScatter(&PandInvM[i2], SetVector4( P2 - (InvM2 * Delta), InvM2 ));
-		}
+		
+		VectorScatter(&PandInvM[i1], SetVector4( P1 + (InvM1 * Delta), InvM1 ));
+		VectorScatter(&PandInvM[i2], SetVector4( P2 - (InvM2 * Delta), InvM2 ));
 	}
 }
 
@@ -253,7 +219,6 @@ export void ApplyXPBDSpringDampingConstraintsWithWeightMaps(uniform FVector4f Pa
 															const uniform float Dists[],
 															uniform float Lambdas[],
 															const uniform float Dt,
-															const uniform float MinStiffness,
 															const uniform bool StiffnessHasMap,
 															const uniform uint8 StiffnessIndices[],
 															const uniform float StiffnessTable[],
@@ -291,17 +256,11 @@ export void ApplyXPBDSpringDampingConstraintsWithWeightMaps(uniform FVector4f Pa
 		const varying float DampingRatio = DampingTable[DampingIndex];
 
 		varying float DLambda;
-		const FVector3f Delta = GetXPBDSpringDampingDelta(P1, P2, X1, X2, InvM1, InvM2, Dt, Dist, Lambda, Stiffness, DampingRatio, MinStiffness, DLambda);
+		const FVector3f Delta = GetXPBDSpringDampingDelta(P1, P2, X1, X2, InvM1, InvM2, Dt, Dist, Lambda, Stiffness, DampingRatio, DLambda);
 
 		Lambdas[i] = Lambda + DLambda;
-
-		if (InvM1 > FLOAT_ZERO)
-		{
-			VectorScatter(&PandInvM[i1], SetVector4( P1 + (InvM1 * Delta), InvM1 ));
-		}
-		if (InvM2 > FLOAT_ZERO)
-		{
-			VectorScatter(&PandInvM[i2], SetVector4( P2 - (InvM2 * Delta), InvM2 ));
-		}
+		
+		VectorScatter(&PandInvM[i1], SetVector4( P1 + (InvM1 * Delta), InvM1 ));
+		VectorScatter(&PandInvM[i2], SetVector4( P2 - (InvM2 * Delta), InvM2 ));
 	}
 }
\ No newline at end of file
diff --git a/Engine/Source/Runtime/Experimental/Chaos/Public/Chaos/PBDAxialSpringConstraintsBase.h b/Engine/Source/Runtime/Experimental/Chaos/Public/Chaos/PBDAxialSpringConstraintsBase.h
--- a/Engine/Source/Runtime/Experimental/Chaos/Public/Chaos/PBDAxialSpringConstraintsBase.h
+++ b/Engine/Source/Runtime/Experimental/Chaos/Public/Chaos/PBDAxialSpringConstraintsBase.h
@@ -40,7 +40,7 @@ public:
 			FPBDStiffness::DefaultParameterFitBase,
 			MaxStiffness)
 	{
-		Init(Particles);
+		Init(Particles, bTrimKinematicConstraints);
 	}
 
 	FPBDAxialSpringConstraintsBase(
@@ -69,7 +69,7 @@ public:
 			FPBDStiffness::DefaultParameterFitBase,
 			MaxStiffness)
 	{
-		Init(Particles);
+		Init(Particles, bTrimKinematicConstraints);
 	}
 
 	virtual ~FPBDAxialSpringConstraintsBase() {}
@@ -152,13 +152,14 @@ private:
 	}
 
 	template<typename SolverParticlesOrRange>
-	void Init(const SolverParticlesOrRange& Particles)
+	void Init(const SolverParticlesOrRange& Particles, bool bTrimKinematicConstraints)
 	{
 		Barys.Reset(Constraints.Num());
 		Dists.Reset(Constraints.Num());
 
-		for (TVec3<int32>& Constraint : Constraints)
+		for (int32 Index = 0; Index < Constraints.Num(); ++Index)
 		{
+			TVec3<int32>& Constraint = Constraints[Index];
 			int32 i1 = Constraint[0];
 			int32 i2 = Constraint[1];
 			int32 i3 = Constraint[2];
@@ -192,6 +193,18 @@ private:
 			const FSolverVec3& P2 = Particles.X(i2);
 			const FSolverVec3& P3 = Particles.X(i3);
 			const FSolverVec3 P = (P2 - P3) * Bary + P3;
+
+			if (bTrimKinematicConstraints)
+			{
+				const FSolverReal PInvMass = Particles.InvM(i3) * ((FSolverReal)1. - Bary) + Particles.InvM(i2) * Bary;
+				const FSolverReal CombinedInvMass = PInvMass + Particles.InvM(i1);
+				if (CombinedInvMass < UE_SMALL_NUMBER)
+				{
+					Constraints.RemoveAtSwap(Index);
+					continue;
+				}
+			}
+
 			Barys.Add(Bary);
 			Dists.Add((P1 - P).Size());
 		}
diff --git a/Engine/Source/Runtime/Experimental/Chaos/Public/Chaos/PBDBendingConstraints.h b/Engine/Source/Runtime/Experimental/Chaos/Public/Chaos/PBDBendingConstraints.h
--- a/Engine/Source/Runtime/Experimental/Chaos/Public/Chaos/PBDBendingConstraints.h
+++ b/Engine/Source/Runtime/Experimental/Chaos/Public/Chaos/PBDBendingConstraints.h
@@ -222,12 +222,3 @@ private:
 };
 
 }  // End namespace Chaos::Softs
-
-// Support ISPC enable/disable in non-shipping builds
-#if !INTEL_ISPC
-const bool bChaos_Bending_ISPC_Enabled = false;
-#elif UE_BUILD_SHIPPING
-const bool bChaos_Bending_ISPC_Enabled = true;
-#else
-extern CHAOS_API bool bChaos_Bending_ISPC_Enabled;
-#endif
diff --git a/Engine/Source/Runtime/Experimental/Chaos/Public/Chaos/PBDBendingConstraintsBase.h b/Engine/Source/Runtime/Experimental/Chaos/Public/Chaos/PBDBendingConstraintsBase.h
--- a/Engine/Source/Runtime/Experimental/Chaos/Public/Chaos/PBDBendingConstraintsBase.h
+++ b/Engine/Source/Runtime/Experimental/Chaos/Public/Chaos/PBDBendingConstraintsBase.h
@@ -220,30 +220,13 @@ public:
 	}
 
 	template<typename SolverParticlesOrRange>
-	void Init(const SolverParticlesOrRange& InParticles)
-	{
-		IsBuckled.SetNumUninitialized(Constraints.Num());
-		for (int32 ConstraintIndex = 0; ConstraintIndex < Constraints.Num(); ++ConstraintIndex)
-		{
-			const TVec4<int32>& Constraint = Constraints[ConstraintIndex];
-			const int32 i1 = Constraint[0];
-			const int32 i2 = Constraint[1];
-			const int32 i3 = Constraint[2];
-			const int32 i4 = Constraint[3];
-			const FSolverVec3& P1 = InParticles.X(i1);
-			const FSolverVec3& P2 = InParticles.X(i2);
-			const FSolverVec3& P3 = InParticles.X(i3);
-			const FSolverVec3& P4 = InParticles.X(i4);
-			const FSolverReal Angle = CalcAngle(P1, P2, P3, P4);
-			IsBuckled[ConstraintIndex] = AngleIsBuckled(Angle, RestAngles[ConstraintIndex]);
-		}
-	}
+	CHAOS_API void Init(const SolverParticlesOrRange& InParticles);
 
 	const TArray<FSolverReal>& GetRestAngles() const { return RestAngles; }
 	const TArray<TVec4<int32>>& GetConstraints() const { return Constraints; }
 	const TArray<bool>& GetIsBuckled() const { return IsBuckled; }
 
-private:
+protected:
 	template<class TNum>
 	static TNum SafeDivide(const TNum& Numerator, const FSolverReal& Denominator)
 	{
@@ -252,7 +235,7 @@ private:
 		return TNum(0);
 	}
 	
-	static TStaticArray<FSolverVec3, 4> CalcGradients(const FSolverVec3& P1, const FSolverVec3& P2, const FSolverVec3& P3, const FSolverVec3& P4)
+	static TStaticArray<FSolverVec3, 4> CalcGradients(const FSolverVec3& P1, const FSolverVec3& P2, const FSolverVec3& P3, const FSolverVec3& P4, FSolverReal* OutAngle = nullptr)
 	{
 		TStaticArray<FSolverVec3, 4> Grads;
 		// Calculated using Phi = atan2(SinPhi, CosPhi)
@@ -284,9 +267,15 @@ private:
 		Grads[2] = -DPhiDP13 - DPhiDP23;
 		Grads[3] = -DPhiDP14 - DPhiDP24;
 
+		if (OutAngle)
+		{
+			*OutAngle = FMath::Atan2(SinPhi, CosPhi);
+		}
+
 		return Grads;
 	}
 
+private:
 	template<typename SolverParticlesOrRange>
 	static TArray<TVec4<int32>> TrimKinematicConstraints(const TArray<TVec4<int32>>& InConstraints, const SolverParticlesOrRange& InParticles)
 	{
@@ -338,3 +327,12 @@ protected:
 };
 
 }  // End namespace Chaos::Softs
+
+// Support ISPC enable/disable in non-shipping builds
+#if !INTEL_ISPC
+const bool bChaos_Bending_ISPC_Enabled = false;
+#elif UE_BUILD_SHIPPING
+const bool bChaos_Bending_ISPC_Enabled = true;
+#else
+extern CHAOS_API bool bChaos_Bending_ISPC_Enabled;
+#endif
diff --git a/Engine/Source/Runtime/Experimental/Chaos/Public/Chaos/SoftsSpring.h b/Engine/Source/Runtime/Experimental/Chaos/Public/Chaos/SoftsSpring.h
--- a/Engine/Source/Runtime/Experimental/Chaos/Public/Chaos/SoftsSpring.h
+++ b/Engine/Source/Runtime/Experimental/Chaos/Public/Chaos/SoftsSpring.h
@@ -12,16 +12,11 @@ namespace Spring
 template<typename SolverParticlesOrRange>
 FSolverVec3 GetXPBDSpringDelta(const SolverParticlesOrRange& Particles, const FSolverReal Dt,
 	const TVec2<int32>& Constraint, const FSolverReal RestLength, FSolverReal& Lambda,
-	const FSolverReal StiffnessValue, const FSolverReal MinStiffness)
+	const FSolverReal StiffnessValue)
 {
 	const int32 Index1 = Constraint[0];
 	const int32 Index2 = Constraint[1];
 
-	if (StiffnessValue < MinStiffness || (Particles.InvM(Index2) == (FSolverReal)0. && Particles.InvM(Index1) == (FSolverReal)0.))
-	{
-		return FSolverVec3((FSolverReal)0.);
-	}
-
 	const FSolverReal CombinedInvMass = Particles.InvM(Index2) + Particles.InvM(Index1);
 
 	const FSolverVec3& P1 = Particles.P(Index1);
@@ -43,16 +38,13 @@ FSolverVec3 GetXPBDSpringDelta(const SolverParticlesOrRange& Particles, const FS
 template<typename SolverParticlesOrRange>
 FSolverVec3 GetXPBDSpringDeltaWithDamping(const SolverParticlesOrRange& Particles, const FSolverReal Dt,
 	const TVec2<int32>& Constraint, const FSolverReal RestLength, FSolverReal& Lambda,
-	const FSolverReal StiffnessValue, const FSolverReal MinStiffness, const FSolverReal DampingRatioValue)
+	const FSolverReal StiffnessValue, const FSolverReal DampingRatioValue)
 {
 	const int32 Index1 = Constraint[0];
 	const int32 Index2 = Constraint[1];
 
 	const FSolverReal CombinedInvMass = Particles.InvM(Index2) + Particles.InvM(Index1);
-	if (StiffnessValue < MinStiffness || CombinedInvMass < UE_SMALL_NUMBER)
-	{
-		return FSolverVec3((FSolverReal)0.);
-	}
+	check(CombinedInvMass > (FSolverReal)0);
 
 	const FSolverReal Damping = DampingRatioValue * 2.f * FMath::Sqrt(StiffnessValue / CombinedInvMass) * (RestLength > UE_SMALL_NUMBER ? (FSolverReal)1. / RestLength : (FSolverReal)1.);
 
@@ -81,15 +73,12 @@ FSolverVec3 GetXPBDSpringDeltaWithDamping(const SolverParticlesOrRange& Particle
 template<typename SolverParticlesOrRange>
 FSolverVec3 GetXPBDSpringDampingDelta(const SolverParticlesOrRange& Particles, const FSolverReal Dt,
 	const TVec2<int32>& Constraint, const FSolverReal RestLength, FSolverReal& Lambda,
-	const FSolverReal StiffnessValue, const FSolverReal MinStiffness, const FSolverReal DampingRatioValue)
+	const FSolverReal StiffnessValue, const FSolverReal DampingRatioValue)
 {
 	const int32 Index1 = Constraint[0];
 	const int32 Index2 = Constraint[1];
 	const FSolverReal CombinedInvMass = Particles.InvM(Index2) + Particles.InvM(Index1);
-	if (StiffnessValue < MinStiffness || CombinedInvMass < UE_SMALL_NUMBER)
-	{
-		return FSolverVec3((FSolverReal)0.);
-	}
+	check(CombinedInvMass > (FSolverReal)0);
 
 	const FSolverReal Damping = DampingRatioValue * 2.f * FMath::Sqrt(StiffnessValue / CombinedInvMass) * (RestLength > UE_SMALL_NUMBER ? (FSolverReal)1. / RestLength : (FSolverReal)1.);
 
@@ -112,19 +101,15 @@ FSolverVec3 GetXPBDSpringDampingDelta(const SolverParticlesOrRange& Particles, c
 template<typename SolverParticlesOrRange>
 FSolverVec3 GetXPBDAxialSpringDelta(const SolverParticlesOrRange& Particles, const FSolverReal Dt,
 	const TVec3<int32>& Constraint, const FSolverReal Bary, const FSolverReal RestLength, FSolverReal& Lambda,
-	const FSolverReal StiffnessValue, const FSolverReal MinStiffness)
+	const FSolverReal StiffnessValue)
 {
 	const int32 Index1 = Constraint[0];
 	const int32 Index2 = Constraint[1];
 	const int32 Index3 = Constraint[2];
 
 	const FSolverReal PInvMass = Particles.InvM(Index3) * ((FSolverReal)1. - Bary) + Particles.InvM(Index2) * Bary;
-	if (StiffnessValue < MinStiffness || (Particles.InvM(Index2) == (FSolverReal)0. && PInvMass == (FSolverReal)0.))
-	{
-		return FSolverVec3((FSolverReal)0.);
-	}
-
 	const FSolverReal CombinedInvMass = PInvMass + Particles.InvM(Index1);
+	check(CombinedInvMass > (FSolverReal)0);
 
 	const FSolverVec3& P1 = Particles.P(Index1);
 	const FSolverVec3& P2 = Particles.P(Index2);
@@ -147,7 +132,7 @@ FSolverVec3 GetXPBDAxialSpringDelta(const SolverParticlesOrRange& Particles, con
 template<typename SolverParticlesOrRange>
 FSolverVec3 GetXPBDAxialSpringDeltaWithDamping(const SolverParticlesOrRange& Particles, const FSolverReal Dt,
 	const TVec3<int32>& Constraint, const FSolverReal Bary, const FSolverReal RestLength, FSolverReal& Lambda,
-	const FSolverReal StiffnessValue, const FSolverReal MinStiffness, const FSolverReal DampingRatioValue)
+	const FSolverReal StiffnessValue, const FSolverReal DampingRatioValue)
 {
 	const int32 Index1 = Constraint[0];
 	const int32 Index2 = Constraint[1];
@@ -155,10 +140,7 @@ FSolverVec3 GetXPBDAxialSpringDeltaWithDamping(const SolverParticlesOrRange& Par
 
 	const FSolverReal PInvMass = Particles.InvM(Index3) * ((FSolverReal)1. - Bary) + Particles.InvM(Index2) * Bary;
 	const FSolverReal CombinedInvMass = PInvMass + Particles.InvM(Index1);
-	if (StiffnessValue < MinStiffness || CombinedInvMass < UE_SMALL_NUMBER)
-	{
-		return FSolverVec3((FSolverReal)0.);
-	}
+	check(CombinedInvMass > (FSolverReal)0);
 
 	const FSolverReal Damping = DampingRatioValue * 2.f * FMath::Sqrt(StiffnessValue / CombinedInvMass) * (RestLength > UE_SMALL_NUMBER ? (FSolverReal)1. / RestLength : (FSolverReal)1.);
 
@@ -191,7 +173,7 @@ FSolverVec3 GetXPBDAxialSpringDeltaWithDamping(const SolverParticlesOrRange& Par
 template<typename SolverParticlesOrRange>
 FSolverVec3 GetXPBDAxialSpringDampingDelta(const SolverParticlesOrRange& Particles, const FSolverReal Dt,
 	const TVec3<int32>& Constraint, const FSolverReal Bary, const FSolverReal RestLength, FSolverReal& Lambda,
-	const FSolverReal StiffnessValue, const FSolverReal MinStiffness, const FSolverReal DampingRatioValue)
+	const FSolverReal StiffnessValue, const FSolverReal DampingRatioValue)
 {
 	const int32 Index1 = Constraint[0];
 	const int32 Index2 = Constraint[1];
@@ -199,10 +181,7 @@ FSolverVec3 GetXPBDAxialSpringDampingDelta(const SolverParticlesOrRange& Particl
 
 	const FSolverReal PInvMass = Particles.InvM(Index3) * ((FSolverReal)1. - Bary) + Particles.InvM(Index2) * Bary;
 	const FSolverReal CombinedInvMass = PInvMass + Particles.InvM(Index1);
-	if (StiffnessValue < MinStiffness || CombinedInvMass < UE_SMALL_NUMBER)
-	{
-		return FSolverVec3((FSolverReal)0.);
-	}
+	check(CombinedInvMass > (FSolverReal)0);
 
 	const FSolverReal Damping = DampingRatioValue * 2.f * FMath::Sqrt(StiffnessValue / CombinedInvMass) * (RestLength > UE_SMALL_NUMBER ? (FSolverReal)1. / RestLength : (FSolverReal)1.);
 
@@ -235,7 +214,7 @@ inline void UpdateSpringLinearSystem(const FSolverParticlesRange& Particles, con
 	const int32 Index1 = Constraint[0];
 	const int32 Index2 = Constraint[1];
 
-	if (StiffnessValue < MinStiffness || (Particles.InvM(Index2) == (FSolverReal)0. && Particles.InvM(Index1) == (FSolverReal)0.))
+	if (StiffnessValue <= MinStiffness || (Particles.InvM(Index2) == (FSolverReal)0. && Particles.InvM(Index1) == (FSolverReal)0.))
 	{
 		return;
 	}
diff --git a/Engine/Source/Runtime/Experimental/Chaos/Public/Chaos/XPBDAnisotropicBendingConstraints.h b/Engine/Source/Runtime/Experimental/Chaos/Public/Chaos/XPBDAnisotropicBendingConstraints.h
--- a/Engine/Source/Runtime/Experimental/Chaos/Public/Chaos/XPBDAnisotropicBendingConstraints.h
+++ b/Engine/Source/Runtime/Experimental/Chaos/Public/Chaos/XPBDAnisotropicBendingConstraints.h
@@ -3,6 +3,7 @@
 
 #include "Chaos/PBDBendingConstraintsBase.h"
 #include "Chaos/CollectionPropertyFacade.h"
+#include "Chaos/PBDFlatWeightMap.h"
 #include "ChaosStats.h"
 
 namespace Chaos
@@ -19,7 +20,7 @@ class FXPBDAnisotropicBendingConstraints final : public FPBDBendingConstraintsBa
 
 public:
 	// Stiffness is in kg cm^2 / rad^2 s^2
-	static constexpr FSolverReal MinStiffness = (FSolverReal)1e-4; // Stiffness below this will be considered 0 since all of our calculations are actually based on 1 / stiffness.
+	static constexpr FSolverReal MinStiffness = (FSolverReal)0; // We're not checking against MinStiffness (except when it's constant and == 0)
 	static constexpr FSolverReal MaxStiffness = (FSolverReal)1e7;
 	static constexpr FSolverReal MinDamping = (FSolverReal)0.;
 	static constexpr FSolverReal MaxDamping = (FSolverReal)1000.;
@@ -30,22 +31,64 @@ public:
 	}
 
 	CHAOS_API FXPBDAnisotropicBendingConstraints(const FSolverParticlesRange& InParticles,
+		const FTriangleMesh& TriangleMesh,
+		const TArray<TVec3<FVec2f>>& FaceVertexPatternPositions,
+		const TMap<FString, TConstArrayView<FRealSingle>>& WeightMaps,
+		const FCollectionPropertyConstFacade& PropertyCollection);
+
+	UE_DEPRECATED(5.4, "XPBD Constraints must always trim kinematic constraints")
+	FXPBDAnisotropicBendingConstraints(const FSolverParticlesRange& InParticles,
 		const FTriangleMesh& TriangleMesh,
 		const TArray<TVec3<FVec2f>>& FaceVertexPatternPositions,
 		const TMap<FString, TConstArrayView<FRealSingle>>& WeightMaps,
 		const FCollectionPropertyConstFacade& PropertyCollection,
-		bool bTrimKinematicConstraints = false);
+		bool bTrimKinematicConstraints)
+		: FXPBDAnisotropicBendingConstraints(InParticles, TriangleMesh, FaceVertexPatternPositions, WeightMaps, PropertyCollection)
+	{}
 
 	CHAOS_API FXPBDAnisotropicBendingConstraints(const FSolverParticles& InParticles,
+		int32 InParticleOffset,
+		int32 InParticleCount,
+		const FTriangleMesh& TriangleMesh,
+		const TArray<TVec3<FVec2f>>& FaceVertexPatternPositions,
+		const TMap<FString, TConstArrayView<FRealSingle>>& WeightMaps,
+		const FCollectionPropertyConstFacade& PropertyCollection);
+
+	UE_DEPRECATED(5.4, "XPBD Constraints must always trim kinematic constraints")
+	FXPBDAnisotropicBendingConstraints(const FSolverParticles& InParticles,
 		int32 InParticleOffset,
 		int32 InParticleCount,
 		const FTriangleMesh& TriangleMesh,
 		const TArray<TVec3<FVec2f>>& FaceVertexPatternPositions,
 		const TMap<FString, TConstArrayView<FRealSingle>>& WeightMaps,
 		const FCollectionPropertyConstFacade& PropertyCollection,
-		bool bTrimKinematicConstraints = false);
+		bool bTrimKinematicConstraints)
+		: FXPBDAnisotropicBendingConstraints(InParticles, InParticleOffset, InParticleCount, TriangleMesh, FaceVertexPatternPositions, WeightMaps, PropertyCollection)
+	{}
 
 	CHAOS_API FXPBDAnisotropicBendingConstraints(const FSolverParticles& InParticles,
+		int32 InParticleOffset,
+		int32 InParticleCount,
+		const FTriangleMesh& TriangleMesh,
+		const TArray<TVec3<FVec2f>>& FaceVertexPatternPositions,
+		const TConstArrayView<FRealSingle>& StiffnessWarpMultipliers,
+		const TConstArrayView<FRealSingle>& StiffnessWeftMultipliers,
+		const TConstArrayView<FRealSingle>& StiffnessBiasMultipliers,
+		const TConstArrayView<FRealSingle>& BucklingStiffnessWarpMultipliers,
+		const TConstArrayView<FRealSingle>& BucklingStiffnessWeftMultipliers,
+		const TConstArrayView<FRealSingle>& BucklingStiffnessBiasMultipliers,
+		const TConstArrayView<FRealSingle>& DampingMultipliers,
+		const FSolverVec2& InStiffnessWarp,
+		const FSolverVec2& InStiffnessWeft,
+		const FSolverVec2& InStiffnessBias,
+		const FSolverReal InBucklingRatio,
+		const FSolverVec2& InBucklingStiffnessWarp,
+		const FSolverVec2& InBucklingStiffnessWeft,
+		const FSolverVec2& InBucklingStiffnessBias,
+		const FSolverVec2& InDampingRatio);
+
+	UE_DEPRECATED(5.4, "XPBD Constraints must always trim kinematic constraints")
+	FXPBDAnisotropicBendingConstraints(const FSolverParticles& InParticles,
 		int32 InParticleOffset,
 		int32 InParticleCount,
 		const FTriangleMesh& TriangleMesh,
@@ -65,19 +108,16 @@ public:
 		const FSolverVec2& InBucklingStiffnessWeft,
 		const FSolverVec2& InBucklingStiffnessBias,
 		const FSolverVec2& InDampingRatio,
-		bool bTrimKinematicConstraints = false);
+		bool bTrimKinematicConstraints)
+		: FXPBDAnisotropicBendingConstraints(InParticles, InParticleOffset, InParticleCount, TriangleMesh, FaceVertexPatternPositions, StiffnessWarpMultipliers, StiffnessWeftMultipliers,
+			StiffnessBiasMultipliers, BucklingStiffnessWarpMultipliers, BucklingStiffnessWeftMultipliers, BucklingStiffnessBiasMultipliers, DampingMultipliers, InStiffnessWarp,
+			InStiffnessWeft, InStiffnessBias, InBucklingRatio, InBucklingStiffnessWarp, InBucklingStiffnessWeft, InBucklingStiffnessBias, InDampingRatio)
+	{}
 
 	virtual ~FXPBDAnisotropicBendingConstraints() override {}
 
 	template<typename SolverParticlesOrRange>
-	void Init(const SolverParticlesOrRange& InParticles)
-	{ 
-		Lambdas.Reset();
-		Lambdas.AddZeroed(Constraints.Num());
-		LambdasDamping.Reset();
-		LambdasDamping.AddZeroed(Constraints.Num());
-		FPBDBendingConstraintsBase::Init(InParticles);
-	}
+	CHAOS_API void Init(const SolverParticlesOrRange& InParticles);
 
 	CHAOS_API void SetProperties(
 		const FCollectionPropertyConstFacade& PropertyCollection,
@@ -89,16 +129,9 @@ public:
 		SetProperties(PropertyCollection, TMap<FString, TConstArrayView<FRealSingle>>());
 	}
 
-	// Update stiffness table, as well as the simulation stiffness exponent
 	void ApplyProperties(const FSolverReal /*Dt*/, const int32 /*NumIterations*/)
 	{
-		Stiffness.ApplyXPBDValues(MaxStiffness);
-		StiffnessWeft.ApplyXPBDValues(MaxStiffness);
-		StiffnessBias.ApplyXPBDValues(MaxStiffness);
-		BucklingStiffness.ApplyXPBDValues(MaxStiffness);
-		BucklingStiffnessWeft.ApplyXPBDValues(MaxStiffness);
-		BucklingStiffnessBias.ApplyXPBDValues(MaxStiffness);
-		DampingRatio.ApplyValues();
+		// Nothing to be done here for flat weight maps. Want to avoid base class from being called instead.
 	}
 
 	template<typename SolverParticlesOrRange>
@@ -175,19 +208,35 @@ private:
 	using Base::ParticleOffset;
 	using Base::ParticleCount;
 	using Base::RestAngles;
-	using Base::Stiffness; // Warp
-	using Base::BucklingStiffness; // Warp
-	
-	FPBDStiffness StiffnessWeft;
-	FPBDStiffness StiffnessBias;
-	FPBDStiffness BucklingStiffnessWeft;
-	FPBDStiffness BucklingStiffnessBias;
 
-	FPBDWeightMap DampingRatio;
+	FPBDFlatWeightMap StiffnessWarp;
+	FPBDFlatWeightMap StiffnessWeft;
+	FPBDFlatWeightMap StiffnessBias;
+	FPBDFlatWeightMap BucklingStiffnessWarp;
+	FPBDFlatWeightMap BucklingStiffnessWeft;
+	FPBDFlatWeightMap BucklingStiffnessBias;
+
+	FPBDFlatWeightMap DampingRatio;
+	
 	mutable TArray<FSolverReal> Lambdas;
 	mutable TArray<FSolverReal> LambdasDamping;
 	TArray<int32> ConstraintsPerColorStartIndex; // Constraints are ordered so each batch is contiguous. This is ColorNum + 1 length so it can be used as start and end.
 
+#if INTEL_ISPC
+	// Constraint SOA. InitColor will initialize these. Only used if using ISPC
+	TArray<int32> ConstraintsIndex1;
+	TArray<int32> ConstraintsIndex2;
+	TArray<int32> ConstraintsIndex3;
+	TArray<int32> ConstraintsIndex4;
+
+	// Particles.X but stored per constraint.
+	// These are only copied over if using ISPC.
+	TArray<FSolverVec3> X1Array;
+	TArray<FSolverVec3> X2Array;
+	TArray<FSolverVec3> X3Array;
+	TArray<FSolverVec3> X4Array;
+#endif
+
 	TArray<FSolverVec3> WarpWeftBiasBaseMultipliers;
 
 	UE_CHAOS_DECLARE_PROPERTYCOLLECTION_NAME(XPBDAnisoBendingStiffnessWarp, float);
diff --git a/Engine/Source/Runtime/Experimental/Chaos/Public/Chaos/XPBDAnisotropicSpringConstraints.h b/Engine/Source/Runtime/Experimental/Chaos/Public/Chaos/XPBDAnisotropicSpringConstraints.h
--- a/Engine/Source/Runtime/Experimental/Chaos/Public/Chaos/XPBDAnisotropicSpringConstraints.h
+++ b/Engine/Source/Runtime/Experimental/Chaos/Public/Chaos/XPBDAnisotropicSpringConstraints.h
@@ -5,6 +5,7 @@
 #include "Chaos/PBDAxialSpringConstraintsBase.h"
 #include "Chaos/PBDSpringConstraintsBase.h"
 #include "Chaos/PBDWeightMap.h"
+#include "Chaos/PBDFlatWeightMap.h"
 
 namespace Chaos
 {
@@ -22,7 +23,7 @@ class FXPBDAnisotropicEdgeSpringConstraints : public FPBDSpringConstraintsBase
 
 public:
 	// Stiffness is in kg cm /s^2
-	static constexpr FSolverReal MinStiffness = (FSolverReal)UE_SMALL_NUMBER; // Stiffness below this will be considered 0 
+	static constexpr FSolverReal MinStiffness = (FSolverReal)0; // We're not checking against MinStiffness (except when it's constant and == 0)
 	static constexpr FSolverReal MaxStiffness = (FSolverReal)1e9;
 	static constexpr FSolverReal MinDampingRatio = (FSolverReal)0.;
 	static constexpr FSolverReal MaxDampingRatio = (FSolverReal)1000.;
@@ -45,8 +46,7 @@ public:
 		const FSolverVec2& InStiffnessBias,
 		const FSolverVec2& InDampingRatio,
 		const FSolverVec2& InWarpScale,
-		const FSolverVec2& InWeftScale,
-		bool bTrimKinematicConstraints = false);
+		const FSolverVec2& InWeftScale);
 
 	CHAOS_API FXPBDAnisotropicEdgeSpringConstraints(
 		const FSolverParticles& InParticles,
@@ -66,8 +66,7 @@ public:
 		const FSolverVec2& InStiffnessBias,
 		const FSolverVec2& InDampingRatio,
 		const FSolverVec2& InWarpScale,
-		const FSolverVec2& InWeftScale,
-		bool bTrimKinematicConstraints = false);
+		const FSolverVec2& InWeftScale);
 
 	virtual ~FXPBDAnisotropicEdgeSpringConstraints() override {}
 
@@ -79,13 +78,9 @@ public:
 		LambdasDamping.AddZeroed(Constraints.Num());
 	}
 
-	// Update stiffness table, as well as the simulation stiffness exponent
+	// Update rest lengths from warp/weft scale
 	void ApplyProperties(const FSolverReal /*Dt*/, const int32 /*NumIterations*/)
 	{
-		Stiffness.ApplyXPBDValues(MaxStiffness);
-		StiffnessWeft.ApplyXPBDValues(MaxStiffness);
-		StiffnessBias.ApplyXPBDValues(MaxStiffness);
-		DampingRatio.ApplyValues();
 		bool bWarpScaleChanged = false;
 		WarpScale.ApplyValues(&bWarpScaleChanged);
 		bool bWeftScaleChanged = false;
@@ -120,11 +115,11 @@ private:
 	using Base::ParticleOffset;
 	using Base::ParticleCount;
 	using Base::Dists; // These will be updated when Warp/Weft Scales change.
-	using Base::Stiffness; // Warp stiffness
 
-	FPBDStiffness StiffnessWeft;
-	FPBDStiffness StiffnessBias;
-	FPBDWeightMap DampingRatio;
+	FPBDFlatWeightMap StiffnessWarp;
+	FPBDFlatWeightMap StiffnessWeft;
+	FPBDFlatWeightMap StiffnessBias;
+	FPBDFlatWeightMap DampingRatio;
 	FPBDWeightMap WarpScale;
 	FPBDWeightMap WeftScale;
 
@@ -144,7 +139,7 @@ class FXPBDAnisotropicAxialSpringConstraints : public FPBDAxialSpringConstraints
 
 public:
 	// Stiffness is in kg cm /s^2
-	static constexpr FSolverReal MinStiffness = (FSolverReal)UE_SMALL_NUMBER; // Stiffness below this will be considered 0 
+	static constexpr FSolverReal MinStiffness = (FSolverReal)0; // We're not checking against MinStiffness (except when it's constant and == 0)
 	static constexpr FSolverReal MaxStiffness = (FSolverReal)1e9;
 	static constexpr FSolverReal MinDampingRatio = (FSolverReal)0.;
 	static constexpr FSolverReal MaxDampingRatio = (FSolverReal)1000.;
@@ -167,8 +162,7 @@ public:
 		const FSolverVec2& InStiffnessBias,
 		const FSolverVec2& InDampingRatio,
 		const FSolverVec2& InWarpScale,
-		const FSolverVec2& InWeftScale,
-		bool bTrimKinematicConstraints = false);
+		const FSolverVec2& InWeftScale);
 
 		CHAOS_API FXPBDAnisotropicAxialSpringConstraints(
 			const FSolverParticles& InParticles,
@@ -188,8 +182,7 @@ public:
 			const FSolverVec2& InStiffnessBias,
 			const FSolverVec2& InDampingRatio,
 			const FSolverVec2& InWarpScale,
-			const FSolverVec2& InWeftScale,
-			bool bTrimKinematicConstraints = false);
+			const FSolverVec2& InWeftScale);
 
 		virtual ~FXPBDAnisotropicAxialSpringConstraints() override {}
 
@@ -201,13 +194,9 @@ public:
 			LambdasDamping.AddZeroed(Constraints.Num());
 		}
 
-		// Update stiffness table, as well as the simulation stiffness exponent
+		// Update rest lengths from warp/weft scale
 		void ApplyProperties(const FSolverReal /*Dt*/, const int32 /*NumIterations*/)
 		{
-			Stiffness.ApplyXPBDValues(MaxStiffness);
-			StiffnessWeft.ApplyXPBDValues(MaxStiffness);
-			StiffnessBias.ApplyXPBDValues(MaxStiffness);
-			DampingRatio.ApplyValues();
 			bool bWarpScaleChanged = false;
 			WarpScale.ApplyValues(&bWarpScaleChanged);
 			bool bWeftScaleChanged = false;
@@ -243,11 +232,11 @@ private:
 	using Base::ParticleCount;
 	using Base::Barys;
 	using Base::Dists; // These will be updated when Warp/Weft Scales change.
-	using Base::Stiffness; // Warp stiffness
 
-	FPBDStiffness StiffnessWeft;
-	FPBDStiffness StiffnessBias;
-	FPBDWeightMap DampingRatio;
+	FPBDFlatWeightMap StiffnessWarp;
+	FPBDFlatWeightMap StiffnessWeft;
+	FPBDFlatWeightMap StiffnessBias;
+	FPBDFlatWeightMap DampingRatio;
 	FPBDWeightMap WarpScale;
 	FPBDWeightMap WeftScale;
 
@@ -280,8 +269,7 @@ public:
 		const FTriangleMesh& TriangleMesh,
 		const TArray<TVec3<FVec2f>>& FaceVertexPatternPositions,
 		const TMap<FString, TConstArrayView<FRealSingle>>& WeightMaps,
-		const FCollectionPropertyConstFacade& PropertyCollection,
-		bool bTrimKinematicConstraints = false)
+		const FCollectionPropertyConstFacade& PropertyCollection)
 		: EdgeConstraints(
 			Particles,
 			TriangleMesh,
@@ -298,8 +286,7 @@ public:
 			FSolverVec2(GetWeightedFloatXPBDAnisoSpringStiffnessBias(PropertyCollection, DefaultStiffness)),
 			FSolverVec2(GetWeightedFloatXPBDAnisoSpringDamping(PropertyCollection, DefaultDamping)),
 			FSolverVec2(GetWeightedFloatXPBDAnisoSpringWarpScale(PropertyCollection, DefaultWarpWeftScale)),
-			FSolverVec2(GetWeightedFloatXPBDAnisoSpringWeftScale(PropertyCollection, DefaultWarpWeftScale)),
-			bTrimKinematicConstraints)
+			FSolverVec2(GetWeightedFloatXPBDAnisoSpringWeftScale(PropertyCollection, DefaultWarpWeftScale)))
 		, AxialConstraints(
 			Particles,
 			TriangleMesh,
@@ -316,8 +303,7 @@ public:
 			FSolverVec2(GetWeightedFloatXPBDAnisoSpringStiffnessBias(PropertyCollection, DefaultStiffness)),
 			FSolverVec2(GetWeightedFloatXPBDAnisoSpringDamping(PropertyCollection, DefaultDamping)),
 			FSolverVec2(GetWeightedFloatXPBDAnisoSpringWarpScale(PropertyCollection, DefaultWarpWeftScale)),
-			FSolverVec2(GetWeightedFloatXPBDAnisoSpringWeftScale(PropertyCollection, DefaultWarpWeftScale)),
-			bTrimKinematicConstraints)
+			FSolverVec2(GetWeightedFloatXPBDAnisoSpringWeftScale(PropertyCollection, DefaultWarpWeftScale)))
 		, XPBDAnisoSpringUse3dRestLengthsIndex(PropertyCollection)
 		, XPBDAnisoSpringStiffnessWarpIndex(PropertyCollection)
 		, XPBDAnisoSpringStiffnessWeftIndex(PropertyCollection)
@@ -334,8 +320,7 @@ public:
 		const FTriangleMesh& TriangleMesh,
 		const TArray<TVec3<FVec2f>>& FaceVertexPatternPositions,
 		const TMap<FString, TConstArrayView<FRealSingle>>& WeightMaps,
-		const FCollectionPropertyConstFacade& PropertyCollection,
-		bool bTrimKinematicConstraints = false)
+		const FCollectionPropertyConstFacade& PropertyCollection)
 		: EdgeConstraints(
 			InParticles,
 			InParticleOffset,
@@ -354,8 +339,7 @@ public:
 			FSolverVec2(GetWeightedFloatXPBDAnisoSpringStiffnessBias(PropertyCollection, DefaultStiffness)),
 			FSolverVec2(GetWeightedFloatXPBDAnisoSpringDamping(PropertyCollection, DefaultDamping)),
 			FSolverVec2(GetWeightedFloatXPBDAnisoSpringWarpScale(PropertyCollection, DefaultWarpWeftScale)),
-			FSolverVec2(GetWeightedFloatXPBDAnisoSpringWeftScale(PropertyCollection, DefaultWarpWeftScale)),
-			bTrimKinematicConstraints)
+			FSolverVec2(GetWeightedFloatXPBDAnisoSpringWeftScale(PropertyCollection, DefaultWarpWeftScale)))
 		, AxialConstraints(
 			InParticles,
 			InParticleOffset,
@@ -374,8 +358,7 @@ public:
 			FSolverVec2(GetWeightedFloatXPBDAnisoSpringStiffnessBias(PropertyCollection, DefaultStiffness)),
 			FSolverVec2(GetWeightedFloatXPBDAnisoSpringDamping(PropertyCollection, DefaultDamping)),
 			FSolverVec2(GetWeightedFloatXPBDAnisoSpringWarpScale(PropertyCollection, DefaultWarpWeftScale)),
-			FSolverVec2(GetWeightedFloatXPBDAnisoSpringWeftScale(PropertyCollection, DefaultWarpWeftScale)),
-			bTrimKinematicConstraints)
+			FSolverVec2(GetWeightedFloatXPBDAnisoSpringWeftScale(PropertyCollection, DefaultWarpWeftScale)))
 		, XPBDAnisoSpringUse3dRestLengthsIndex(PropertyCollection)
 		, XPBDAnisoSpringStiffnessWarpIndex(PropertyCollection)
 		, XPBDAnisoSpringStiffnessWeftIndex(PropertyCollection)
@@ -389,7 +372,7 @@ public:
 		const FCollectionPropertyConstFacade& PropertyCollection,
 		const TMap<FString, TConstArrayView<FRealSingle>>& WeightMaps);
 
-	// Update stiffness table, as well as the simulation stiffness exponent
+	// Update rest lengths from warp/weft scale
 	void ApplyProperties(const FSolverReal Dt, const int32 NumIterations)
 	{
 		EdgeConstraints.ApplyProperties(Dt, NumIterations);
diff --git a/Engine/Source/Runtime/Experimental/Chaos/Public/Chaos/XPBDBendingConstraints.h b/Engine/Source/Runtime/Experimental/Chaos/Public/Chaos/XPBDBendingConstraints.h
--- a/Engine/Source/Runtime/Experimental/Chaos/Public/Chaos/XPBDBendingConstraints.h
+++ b/Engine/Source/Runtime/Experimental/Chaos/Public/Chaos/XPBDBendingConstraints.h
@@ -4,13 +4,14 @@
 #include "Chaos/PBDBendingConstraintsBase.h"
 #include "Chaos/CollectionPropertyFacade.h"
 #include "ChaosStats.h"
+#include "Chaos/PBDFlatWeightMap.h"
 
 namespace Chaos::Softs
 {
 
 // Stiffness is in kg cm^2 / rad^2 s^2
 UE_DEPRECATED(5.2, "Use FXPBDBendingConstraints::MinStiffness instead.")
-static const FSolverReal XPBDBendMinStiffness = (FSolverReal)1e-4; // Stiffness below this will be considered 0 since all of our calculations are actually based on 1 / stiffness.
+static const FSolverReal XPBDBendMinStiffness = (FSolverReal)0;  // We're not checking against MinStiffness (except when it's constant and == 0)
 UE_DEPRECATED(5.2, "Use FXPBDBendingConstraints::MaxStiffness instead.")
 static const FSolverReal XPBDBendMaxStiffness = (FSolverReal)1e7;
 
@@ -20,7 +21,7 @@ class FXPBDBendingConstraints final : public FPBDBendingConstraintsBase
 
 public:
 	// Stiffness is in kg cm^2 / rad^2 s^2
-	static constexpr FSolverReal MinStiffness = (FSolverReal)1e-4; // Stiffness below this will be considered 0 since all of our calculations are actually based on 1 / stiffness.
+	static constexpr FSolverReal MinStiffness = (FSolverReal)0;  // We're not checking against MinStiffness (except when it's constant and == 0)
 	static constexpr FSolverReal MaxStiffness = (FSolverReal)1e7;
 	static constexpr FSolverReal MinDamping = (FSolverReal)0.;
 	static constexpr FSolverReal MaxDamping = (FSolverReal)1000.;
@@ -33,21 +34,32 @@ public:
 	FXPBDBendingConstraints(const FSolverParticlesRange& InParticles,
 		TArray<TVec4<int32>>&& InConstraints,
 		const TMap<FString, TConstArrayView<FRealSingle>>& WeightMaps,
-		const FCollectionPropertyConstFacade& PropertyCollection,
-		bool bTrimKinematicConstraints = false)
+		const FCollectionPropertyConstFacade& PropertyCollection)
 		: Base(
 			InParticles,
 			MoveTemp(InConstraints),
-			WeightMaps.FindRef(GetXPBDBendingElementStiffnessString(PropertyCollection, XPBDBendingElementStiffnessName.ToString())),
-			WeightMaps.FindRef(GetXPBDBucklingStiffnessString(PropertyCollection, XPBDBucklingStiffnessName.ToString())),
+			TConstArrayView<FRealSingle>(), // We don't use base stiffness weight maps
+			TConstArrayView<FRealSingle>(), // We don't use base stiffness weight maps
 			GetRestAngleMapFromCollection(WeightMaps, PropertyCollection),
 			FSolverVec2(GetWeightedFloatXPBDBendingElementStiffness(PropertyCollection, MaxStiffness)),
 			(FSolverReal)GetXPBDBucklingRatio(PropertyCollection, 0.f),
 			FSolverVec2(GetWeightedFloatXPBDBucklingStiffness(PropertyCollection, MaxStiffness)),
 			GetRestAngleValueFromCollection(PropertyCollection),
 			(ERestAngleConstructionType)GetXPBDRestAngleType(PropertyCollection, (int32)ERestAngleConstructionType::Use3DRestAngles),
-			bTrimKinematicConstraints,
+			true /*bTrimKinematicConstraints*/,
 			MaxStiffness)
+		, XPBDStiffness(
+			FSolverVec2(GetWeightedFloatXPBDBendingElementStiffness(PropertyCollection, MaxStiffness)).ClampAxes(0, MaxStiffness),
+			WeightMaps.FindRef(GetXPBDBendingElementStiffnessString(PropertyCollection, XPBDBendingElementStiffnessName.ToString())),
+			TConstArrayView<TVec2<int32>>(ConstraintSharedEdges),
+			ParticleOffset,
+			ParticleCount)
+		, XPBDBucklingStiffness(
+			FSolverVec2(GetWeightedFloatXPBDBucklingStiffness(PropertyCollection, MaxStiffness)).ClampAxes(0, MaxStiffness),
+			WeightMaps.FindRef(GetXPBDBucklingStiffnessString(PropertyCollection, XPBDBucklingStiffnessName.ToString())),
+			TConstArrayView<TVec2<int32>>(ConstraintSharedEdges),
+			ParticleOffset,
+			ParticleCount)
 		, DampingRatio(
 			FSolverVec2(GetWeightedFloatXPBDBendingElementDamping(PropertyCollection, MinDamping)).ClampAxes(MinDamping, MaxDamping),
 			WeightMaps.FindRef(GetXPBDBendingElementDampingString(PropertyCollection, XPBDBendingElementDampingName.ToString())),
@@ -67,28 +79,48 @@ public:
 		InitColor(InParticles);
 	}
 
+	UE_DEPRECATED(5.4, "XPBD Constraints must always trim kinematic constraints")
+	FXPBDBendingConstraints(const FSolverParticlesRange& InParticles,
+		TArray<TVec4<int32>>&& InConstraints,
+		const TMap<FString, TConstArrayView<FRealSingle>>& WeightMaps,
+		const FCollectionPropertyConstFacade& PropertyCollection,
+		bool bTrimKinematicConstraints)
+		: FXPBDBendingConstraints(InParticles, MoveTemp(InConstraints), WeightMaps, PropertyCollection)
+	{}
+
 	FXPBDBendingConstraints(const FSolverParticles& InParticles,
 		int32 InParticleOffset,
 		int32 InParticleCount,
 		TArray<TVec4<int32>>&& InConstraints,
 		const TMap<FString, TConstArrayView<FRealSingle>>& WeightMaps,
-		const FCollectionPropertyConstFacade& PropertyCollection,
-		bool bTrimKinematicConstraints = false)
+		const FCollectionPropertyConstFacade& PropertyCollection)
 		: Base(
 			InParticles,
 			InParticleOffset,
 			InParticleCount,
 			MoveTemp(InConstraints),
-			WeightMaps.FindRef(GetXPBDBendingElementStiffnessString(PropertyCollection, XPBDBendingElementStiffnessName.ToString())),
-			WeightMaps.FindRef(GetXPBDBucklingStiffnessString(PropertyCollection, XPBDBucklingStiffnessName.ToString())),
+			TConstArrayView<FRealSingle>(), // We don't use base stiffness weight maps
+			TConstArrayView<FRealSingle>(), // We don't use base stiffness weight maps
 			GetRestAngleMapFromCollection(WeightMaps, PropertyCollection),
 			FSolverVec2(GetWeightedFloatXPBDBendingElementStiffness(PropertyCollection, MaxStiffness)),
 			(FSolverReal)GetXPBDBucklingRatio(PropertyCollection, 0.f),
 			FSolverVec2(GetWeightedFloatXPBDBucklingStiffness(PropertyCollection, MaxStiffness)),
 			GetRestAngleValueFromCollection(PropertyCollection),
 			(ERestAngleConstructionType)GetXPBDRestAngleType(PropertyCollection, (int32)ERestAngleConstructionType::Use3DRestAngles),
-			bTrimKinematicConstraints,
+			true /*bTrimKinematicConstraints*/,
 			MaxStiffness)
+		, XPBDStiffness(
+			FSolverVec2(GetWeightedFloatXPBDBendingElementStiffness(PropertyCollection, MaxStiffness)).ClampAxes(0, MaxStiffness),
+			WeightMaps.FindRef(GetXPBDBendingElementStiffnessString(PropertyCollection, XPBDBendingElementStiffnessName.ToString())),
+			TConstArrayView<TVec2<int32>>(ConstraintSharedEdges),
+			ParticleOffset,
+			ParticleCount)
+		, XPBDBucklingStiffness(
+			FSolverVec2(GetWeightedFloatXPBDBucklingStiffness(PropertyCollection, MaxStiffness)).ClampAxes(0, MaxStiffness),
+			WeightMaps.FindRef(GetXPBDBucklingStiffnessString(PropertyCollection, XPBDBucklingStiffnessName.ToString())),
+			TConstArrayView<TVec2<int32>>(ConstraintSharedEdges),
+			ParticleOffset,
+			ParticleCount)
 		, DampingRatio(
 			FSolverVec2(GetWeightedFloatXPBDBendingElementDamping(PropertyCollection, MinDamping)).ClampAxes(MinDamping, MaxDamping),
 			WeightMaps.FindRef(GetXPBDBendingElementDampingString(PropertyCollection, XPBDBendingElementDampingName.ToString())),
@@ -108,6 +140,17 @@ public:
 		InitColor(InParticles);
 	}
 
+	UE_DEPRECATED(5.4, "XPBD Constraints must always trim kinematic constraints")
+	FXPBDBendingConstraints(const FSolverParticles& InParticles,
+		int32 InParticleOffset,
+		int32 InParticleCount,
+		TArray<TVec4<int32>>&& InConstraints,
+		const TMap<FString, TConstArrayView<FRealSingle>>& WeightMaps,
+		const FCollectionPropertyConstFacade& PropertyCollection,
+		bool bTrimKinematicConstraints)
+		: FXPBDBendingConstraints(InParticles, InParticleOffset, InParticleCount, MoveTemp(InConstraints), WeightMaps, PropertyCollection)
+	{}
+
 	UE_DEPRECATED(5.3, "Use weight map constructor instead.")
 	FXPBDBendingConstraints(const FSolverParticles& InParticles,
 		int32 InParticleOffset,
@@ -123,13 +166,25 @@ public:
 			InParticleOffset,
 			InParticleCount,
 			MoveTemp(InConstraints),
-			StiffnessMultipliers,
-			BucklingStiffnessMultipliers,
+			TConstArrayView<FRealSingle>(), // We don't use base stiffness weight maps
+			TConstArrayView<FRealSingle>(), // We don't use base stiffness weight maps
 			FSolverVec2(GetWeightedFloatXPBDBendingElementStiffness(PropertyCollection, MaxStiffness)),
 			(FSolverReal)GetXPBDBucklingRatio(PropertyCollection, 0.f),
 			FSolverVec2(GetWeightedFloatXPBDBucklingStiffness(PropertyCollection, MaxStiffness)),
-			bTrimKinematicConstraints,
+			true /*bTrimKinematicConstraints*/,
 			MaxStiffness)
+		, XPBDStiffness(
+			FSolverVec2(GetWeightedFloatXPBDBendingElementStiffness(PropertyCollection, MaxStiffness)).ClampAxes(0, MaxStiffness),
+			StiffnessMultipliers,
+			TConstArrayView<TVec2<int32>>(ConstraintSharedEdges),
+			ParticleOffset,
+			ParticleCount)
+		, XPBDBucklingStiffness(
+			FSolverVec2(GetWeightedFloatXPBDBucklingStiffness(PropertyCollection, MaxStiffness)).ClampAxes(0, MaxStiffness),
+			BucklingStiffnessMultipliers,
+			TConstArrayView<TVec2<int32>>(ConstraintSharedEdges),
+			ParticleOffset,
+			ParticleCount)
 		, DampingRatio(
 			FSolverVec2(GetWeightedFloatXPBDBendingElementDamping(PropertyCollection, MinDamping)).ClampAxes(MinDamping, MaxDamping),
 			DampingMultipliers,
@@ -165,13 +220,15 @@ public:
 			InParticleOffset,
 			InParticleCount,
 			MoveTemp(InConstraints),
-			StiffnessMultipliers,
-			BucklingStiffnessMultipliers,
+			TConstArrayView<FRealSingle>(), // We don't use base stiffness weight maps
+			TConstArrayView<FRealSingle>(), // We don't use base stiffness weight maps
 			InStiffness,
 			InBucklingRatio,
 			InBucklingStiffness,
-			bTrimKinematicConstraints,
+			true /*bTrimKinematicConstraints*/,
 			MaxStiffness)
+		, XPBDStiffness(InStiffness.ClampAxes(0, MaxStiffness))
+		, XPBDBucklingStiffness(InBucklingStiffness.ClampAxes(0, MaxStiffness))
 		, DampingRatio(FSolverVec2::ZeroVector)
 		, XPBDBendingElementStiffnessIndex(ForceInit)
 		, XPBDBendingElementDampingIndex(ForceInit)
@@ -196,20 +253,31 @@ public:
 		const FSolverVec2& InStiffness,
 		const FSolverReal InBucklingRatio,
 		const FSolverVec2& InBucklingStiffness,
-		const FSolverVec2& InDampingRatio,
-		bool bTrimKinematicConstraints = false)
+		const FSolverVec2& InDampingRatio)
 		: Base(
 			InParticles,
 			ParticleOffset,
 			ParticleCount,
 			MoveTemp(InConstraints),
-			StiffnessMultipliers,
-			BucklingStiffnessMultipliers,
+			TConstArrayView<FRealSingle>(), // We don't use base stiffness weight maps
+			TConstArrayView<FRealSingle>(), // We don't use base stiffness weight maps
 			InStiffness,
 			InBucklingRatio,
 			InBucklingStiffness,
-			bTrimKinematicConstraints,
+			true /*bTrimKinematicConstraints*/,
 			MaxStiffness)
+		, XPBDStiffness(
+			InStiffness.ClampAxes(0, MaxStiffness),
+			StiffnessMultipliers,
+			TConstArrayView<TVec2<int32>>(ConstraintSharedEdges),
+			ParticleOffset,
+			ParticleCount)
+		, XPBDBucklingStiffness(
+			InBucklingStiffness.ClampAxes(0, MaxStiffness),
+			BucklingStiffnessMultipliers,
+			TConstArrayView<TVec2<int32>>(ConstraintSharedEdges),
+			ParticleOffset,
+			ParticleCount)
 		, DampingRatio(
 			InDampingRatio.ClampAxes((FSolverReal)0., (FSolverReal)1.),
 			InDampingMultipliers,
@@ -229,17 +297,26 @@ public:
 		InitColor(InParticles);
 	}
 
+	UE_DEPRECATED(5.4, "XPBD Constraints must always trim kinematic constraints")
+	FXPBDBendingConstraints(const FSolverParticles& InParticles,
+		int32 ParticleOffset,
+		int32 ParticleCount,
+		TArray<TVec4<int32>>&& InConstraints,
+		const TConstArrayView<FRealSingle>& StiffnessMultipliers,
+		const TConstArrayView<FRealSingle>& BucklingStiffnessMultipliers,
+		const TConstArrayView<FRealSingle>& InDampingMultipliers,
+		const FSolverVec2& InStiffness,
+		const FSolverReal InBucklingRatio,
+		const FSolverVec2& InBucklingStiffness,
+		const FSolverVec2& InDampingRatio,
+		bool bTrimKinematicConstraints)
+		: FXPBDBendingConstraints(InParticles, ParticleOffset, ParticleCount, MoveTemp(InConstraints), StiffnessMultipliers, BucklingStiffnessMultipliers, InDampingMultipliers, InStiffness, InBucklingRatio, InBucklingStiffness, InDampingRatio)
+	{}
+
 	virtual ~FXPBDBendingConstraints() override {}
 
 	template<typename SolverParticlesOrRange>
-	void Init(const SolverParticlesOrRange& InParticles)
-	{ 
-		Lambdas.Reset();
-		Lambdas.AddZeroed(Constraints.Num());
-		LambdasDamping.Reset();
-		LambdasDamping.AddZeroed(Constraints.Num());
-		FPBDBendingConstraintsBase::Init(InParticles);
-	}
+	void CHAOS_API Init(const SolverParticlesOrRange& InParticles);
 
 	CHAOS_API void SetProperties(
 		const FCollectionPropertyConstFacade& PropertyCollection,
@@ -253,18 +330,15 @@ public:
 
 	void SetProperties(const FSolverVec2& InStiffness, const FSolverReal InBucklingRatio, const FSolverVec2& InBucklingStiffness, const FSolverVec2& InDampingRatio = FSolverVec2::ZeroVector)
 	{
-		Stiffness.SetWeightedValue(InStiffness, MaxStiffness);
+		XPBDStiffness.SetWeightedValue(InStiffness.ClampAxes(0, MaxStiffness));
 		BucklingRatio = FMath::Clamp(InBucklingRatio, (FSolverReal)0., (FSolverReal)1.);
-		BucklingStiffness.SetWeightedValue(InBucklingStiffness, MaxStiffness);
+		XPBDBucklingStiffness.SetWeightedValue(InBucklingStiffness.ClampAxes(0, MaxStiffness));
 		DampingRatio.SetWeightedValue(InDampingRatio.ClampAxes(MinDamping, MaxDamping));
 	}
 
-	// Update stiffness table, as well as the simulation stiffness exponent
 	void ApplyProperties(const FSolverReal /*Dt*/, const int32 /*NumIterations*/)
 	{
-		Stiffness.ApplyXPBDValues(MaxStiffness);
-		BucklingStiffness.ApplyXPBDValues(MaxStiffness);
-		DampingRatio.ApplyValues();
+		// Nothing to be done here for flat weight maps. Want to avoid base class from being called instead.
 	}
 
 	template<typename SolverParticlesOrRange>
@@ -336,14 +410,29 @@ private:
 	using Base::ParticleOffset;
 	using Base::ParticleCount;
 	using Base::RestAngles;
-	using Base::Stiffness;
-	using Base::BucklingStiffness;
 
-	FPBDWeightMap DampingRatio;
+	FPBDFlatWeightMap XPBDStiffness;
+	FPBDFlatWeightMap XPBDBucklingStiffness;
+	FPBDFlatWeightMap DampingRatio;
 	mutable TArray<FSolverReal> Lambdas;
 	mutable TArray<FSolverReal> LambdasDamping;
 	TArray<int32> ConstraintsPerColorStartIndex; // Constraints are ordered so each batch is contiguous. This is ColorNum + 1 length so it can be used as start and end.
 
+#if INTEL_ISPC
+	// Constraint SOA. InitColor will initialize these. Only used by ISPC
+	TArray<int32> ConstraintsIndex1;
+	TArray<int32> ConstraintsIndex2;
+	TArray<int32> ConstraintsIndex3;
+	TArray<int32> ConstraintsIndex4;
+
+	// Particles.X but stored per constraint.
+	// These are only copied over if using ISPC.
+	TArray<FSolverVec3> X1Array;
+	TArray<FSolverVec3> X2Array;
+	TArray<FSolverVec3> X3Array;
+	TArray<FSolverVec3> X4Array;
+#endif 
+
 	UE_CHAOS_DECLARE_PROPERTYCOLLECTION_NAME(XPBDBendingElementStiffness, float);
 	UE_CHAOS_DECLARE_PROPERTYCOLLECTION_NAME(XPBDBendingElementDamping, float);
 	UE_CHAOS_DECLARE_PROPERTYCOLLECTION_NAME(XPBDBucklingRatio, float);
diff --git a/Engine/Source/Runtime/Experimental/Chaos/Public/Chaos/XPBDSpringConstraints.h b/Engine/Source/Runtime/Experimental/Chaos/Public/Chaos/XPBDSpringConstraints.h
--- a/Engine/Source/Runtime/Experimental/Chaos/Public/Chaos/XPBDSpringConstraints.h
+++ b/Engine/Source/Runtime/Experimental/Chaos/Public/Chaos/XPBDSpringConstraints.h
@@ -12,7 +12,7 @@ namespace Chaos::Softs
 
 // Stiffness is in kg cm /s^2
 UE_DEPRECATED(5.2, "Use FXPBDSpringConstraints::MinStiffness instead.")
-static const FSolverReal XPBDSpringMinStiffness = (FSolverReal)UE_SMALL_NUMBER; // Stiffness below this will be considered 0 
+static const FSolverReal XPBDSpringMinStiffness = (FSolverReal)0; // We're not checking against MinStiffness (except when it's constant and == 0)
 UE_DEPRECATED(5.2, "Use FXPBDSpringConstraints::MaxStiffness instead.")
 static const FSolverReal XPBDSpringMaxStiffness = (FSolverReal)1e9;
 
@@ -21,7 +21,7 @@ class FXPBDSpringConstraints : public FPBDSpringConstraintsBase
 	typedef FPBDSpringConstraintsBase Base;
 
 public:
-	static constexpr FSolverReal MinStiffness = (FSolverReal)UE_SMALL_NUMBER; // Stiffness below this will be considered 0 
+	static constexpr FSolverReal MinStiffness = (FSolverReal)0; // We're not checking against MinStiffness (except when it's constant and == 0)
 	static constexpr FSolverReal MaxStiffness = (FSolverReal)1e9;
 	static constexpr FSolverReal MinDampingRatio = (FSolverReal)0.;
 	static constexpr FSolverReal MaxDampingRatio = (FSolverReal)1000.;
@@ -43,7 +43,7 @@ public:
 			InConstraints,
 			StiffnessMultipliers,
 			InStiffness,
-			bTrimKinematicConstraints,
+			true /*bTrimKinematicConstraints*/,
 			MaxStiffness)
 		, DampingRatio(FSolverVec2::ZeroVector)
 	{
@@ -52,6 +52,33 @@ public:
 	}	
 	
 	template<int32 Valence, TEMPLATE_REQUIRES(Valence >= 2 && Valence <= 4)>
+	FXPBDSpringConstraints(
+		const FSolverParticlesRange& Particles,
+		const TArray<TVector<int32, Valence>>& InConstraints,
+		const TConstArrayView<FRealSingle>& StiffnessMultipliers,
+		const TConstArrayView<FRealSingle>& DampingMultipliers,
+		const FSolverVec2& InStiffness,
+		const FSolverVec2& InDampingRatio)
+		: Base(
+			Particles,
+			InConstraints,
+			StiffnessMultipliers,
+			InStiffness,
+			true /*bTrimKinematicConstraints*/,
+			MaxStiffness)
+		, DampingRatio(
+			InDampingRatio.ClampAxes(MinDampingRatio, MaxDampingRatio),
+			DampingMultipliers,
+			TConstArrayView<TVec2<int32>>(Constraints),
+			ParticleOffset,
+			ParticleCount)
+	{
+		Lambdas.Init((FSolverReal)0., Constraints.Num());
+		InitColor(Particles);
+	}
+
+	template<int32 Valence, TEMPLATE_REQUIRES(Valence >= 2 && Valence <= 4)>
+	UE_DEPRECATED(5.4, "XPBD Constraints must always trim kinematic constraints")
 	FXPBDSpringConstraints(
 		const FSolverParticlesRange& Particles,
 		const TArray<TVector<int32, Valence>>& InConstraints,
@@ -59,13 +86,13 @@ public:
 		const TConstArrayView<FRealSingle>& DampingMultipliers,
 		const FSolverVec2& InStiffness,
 		const FSolverVec2& InDampingRatio,
-		bool bTrimKinematicConstraints = false)
+		bool bTrimKinematicConstraints)
 		: Base(
 			Particles,
 			InConstraints,
 			StiffnessMultipliers,
 			InStiffness,
-			bTrimKinematicConstraints,
+			true /*bTrimKinematicConstraints*/,
 			MaxStiffness)
 		, DampingRatio(
 			InDampingRatio.ClampAxes(MinDampingRatio, MaxDampingRatio),
@@ -87,8 +114,7 @@ public:
 		const TConstArrayView<FRealSingle>& StiffnessMultipliers,
 		const TConstArrayView<FRealSingle>& DampingMultipliers,
 		const FSolverVec2& InStiffness,
-		const FSolverVec2& InDampingRatio,
-		bool bTrimKinematicConstraints = false)
+		const FSolverVec2& InDampingRatio)
 		: Base(
 			Particles,
 			ParticleOffset,
@@ -96,7 +122,7 @@ public:
 			InConstraints,
 			StiffnessMultipliers,
 			InStiffness,
-			bTrimKinematicConstraints,
+			true /*bTrimKinematicConstraints*/,
 			MaxStiffness)
 		, DampingRatio(
 			InDampingRatio.ClampAxes(MinDampingRatio, MaxDampingRatio),
@@ -112,6 +138,41 @@ public:
 		InitColor(Particles);
 	}
 
+	template<int32 Valence, TEMPLATE_REQUIRES(Valence >= 2 && Valence <= 4)>
+	UE_DEPRECATED(5.4, "XPBD Constraints must always trim kinematic constraints")
+	FXPBDSpringConstraints(
+		const FSolverParticles& Particles,
+		int32 ParticleOffset,
+		int32 ParticleCount,
+		const TArray<TVector<int32, Valence>>& InConstraints,
+		const TConstArrayView<FRealSingle>& StiffnessMultipliers,
+		const TConstArrayView<FRealSingle>& DampingMultipliers,
+		const FSolverVec2& InStiffness,
+		const FSolverVec2& InDampingRatio,
+		bool bTrimKinematicConstraints)
+	: Base(
+		Particles,
+		ParticleOffset,
+		ParticleCount,
+		InConstraints,
+		StiffnessMultipliers,
+		InStiffness,
+		true /*bTrimKinematicConstraints*/,
+		MaxStiffness)
+	, DampingRatio(
+		InDampingRatio.ClampAxes(MinDampingRatio, MaxDampingRatio),
+		DampingMultipliers,
+		TConstArrayView<TVec2<int32>>(Constraints),
+		ParticleOffset,
+		ParticleCount)
+	{
+		Lambdas.Reset();
+		Lambdas.SetNumZeroed(Constraints.Num());
+		LambdasDamping.Reset();
+		LambdasDamping.SetNumZeroed(Constraints.Num());
+		InitColor(Particles);
+	}
+
 	virtual ~FXPBDSpringConstraints() override {}
 
 	void Init() const 
@@ -176,16 +237,32 @@ public:
 		const FSolverParticlesRange& Particles,
 		const TArray<TVec3<int32>>& InConstraints,
 		const TMap<FString, TConstArrayView<FRealSingle>>& WeightMaps,
-		const FCollectionPropertyConstFacade& PropertyCollection,
-		bool bTrimKinematicConstraints = false)
+		const FCollectionPropertyConstFacade& PropertyCollection)
+		: FXPBDSpringConstraints(
+			Particles,
+			InConstraints,
+			WeightMaps.FindRef(GetXPBDEdgeSpringStiffnessString(PropertyCollection, XPBDEdgeSpringStiffnessName.ToString())),
+			WeightMaps.FindRef(GetXPBDEdgeSpringDampingString(PropertyCollection, XPBDEdgeSpringDampingName.ToString())),
+			FSolverVec2(GetWeightedFloatXPBDEdgeSpringStiffness(PropertyCollection, MaxStiffness)),
+			FSolverVec2(GetWeightedFloatXPBDEdgeSpringDamping(PropertyCollection, MinDampingRatio)))
+		, XPBDEdgeSpringStiffnessIndex(PropertyCollection)
+		, XPBDEdgeSpringDampingIndex(PropertyCollection)
+	{}
+
+	UE_DEPRECATED(5.4, "XPBD Constraints must always trim kinematic constraints")
+	FXPBDEdgeSpringConstraints(
+		const FSolverParticlesRange & Particles,
+		const TArray<TVec3<int32>>&InConstraints,
+		const TMap<FString, TConstArrayView<FRealSingle>>&WeightMaps,
+		const FCollectionPropertyConstFacade & PropertyCollection,
+		bool bTrimKinematicConstraints)
 		: FXPBDSpringConstraints(
 			Particles,
 			InConstraints,
 			WeightMaps.FindRef(GetXPBDEdgeSpringStiffnessString(PropertyCollection, XPBDEdgeSpringStiffnessName.ToString())),
 			WeightMaps.FindRef(GetXPBDEdgeSpringDampingString(PropertyCollection, XPBDEdgeSpringDampingName.ToString())),
 			FSolverVec2(GetWeightedFloatXPBDEdgeSpringStiffness(PropertyCollection, MaxStiffness)),
-			FSolverVec2(GetWeightedFloatXPBDEdgeSpringDamping(PropertyCollection, MinDampingRatio)),
-			bTrimKinematicConstraints)
+			FSolverVec2(GetWeightedFloatXPBDEdgeSpringDamping(PropertyCollection, MinDampingRatio)))
 		, XPBDEdgeSpringStiffnessIndex(PropertyCollection)
 		, XPBDEdgeSpringDampingIndex(PropertyCollection)
 	{}
@@ -196,8 +273,29 @@ public:
 		int32 ParticleCount,
 		const TArray<TVec3<int32>>& InConstraints,
 		const TMap<FString, TConstArrayView<FRealSingle>>& WeightMaps,
-		const FCollectionPropertyConstFacade& PropertyCollection,
-		bool bTrimKinematicConstraints = false)
+		const FCollectionPropertyConstFacade& PropertyCollection)
+		: FXPBDSpringConstraints(
+			Particles,
+			ParticleOffset,
+			ParticleCount,
+			InConstraints,
+			WeightMaps.FindRef(GetXPBDEdgeSpringStiffnessString(PropertyCollection, XPBDEdgeSpringStiffnessName.ToString())),
+			WeightMaps.FindRef(GetXPBDEdgeSpringDampingString(PropertyCollection, XPBDEdgeSpringDampingName.ToString())),
+			FSolverVec2(GetWeightedFloatXPBDEdgeSpringStiffness(PropertyCollection, MaxStiffness)),
+			FSolverVec2(GetWeightedFloatXPBDEdgeSpringDamping(PropertyCollection, MinDampingRatio)))
+		, XPBDEdgeSpringStiffnessIndex(PropertyCollection)
+		, XPBDEdgeSpringDampingIndex(PropertyCollection)
+	{}
+
+	UE_DEPRECATED(5.4, "XPBD Constraints must always trim kinematic constraints")
+	FXPBDEdgeSpringConstraints(
+		const FSolverParticles & Particles,
+		int32 ParticleOffset,
+		int32 ParticleCount,
+		const TArray<TVec3<int32>>&InConstraints,
+		const TMap<FString, TConstArrayView<FRealSingle>>&WeightMaps,
+		const FCollectionPropertyConstFacade & PropertyCollection,
+		bool bTrimKinematicConstraints)
 		: FXPBDSpringConstraints(
 			Particles,
 			ParticleOffset,
@@ -206,8 +304,7 @@ public:
 			WeightMaps.FindRef(GetXPBDEdgeSpringStiffnessString(PropertyCollection, XPBDEdgeSpringStiffnessName.ToString())),
 			WeightMaps.FindRef(GetXPBDEdgeSpringDampingString(PropertyCollection, XPBDEdgeSpringDampingName.ToString())),
 			FSolverVec2(GetWeightedFloatXPBDEdgeSpringStiffness(PropertyCollection, MaxStiffness)),
-			FSolverVec2(GetWeightedFloatXPBDEdgeSpringDamping(PropertyCollection, MinDampingRatio)),
-			bTrimKinematicConstraints)
+			FSolverVec2(GetWeightedFloatXPBDEdgeSpringDamping(PropertyCollection, MinDampingRatio)))
 		, XPBDEdgeSpringStiffnessIndex(PropertyCollection)
 		, XPBDEdgeSpringDampingIndex(PropertyCollection)
 	{}
@@ -230,8 +327,7 @@ public:
 			StiffnessMultipliers,
 			DampingMultipliers,
 			FSolverVec2(GetWeightedFloatXPBDEdgeSpringStiffness(PropertyCollection, MaxStiffness)),
-			FSolverVec2(GetWeightedFloatXPBDEdgeSpringDamping(PropertyCollection, MinDampingRatio)),
-			bTrimKinematicConstraints)
+			FSolverVec2(GetWeightedFloatXPBDEdgeSpringDamping(PropertyCollection, MinDampingRatio)))
 		, XPBDEdgeSpringStiffnessIndex(PropertyCollection)
 		, XPBDEdgeSpringDampingIndex(PropertyCollection)
 	{}
@@ -271,20 +367,57 @@ public:
 		const FSolverParticlesRange& Particles,
 		const TArray<TVec2<int32>>& InConstraints,
 		const TMap<FString, TConstArrayView<FRealSingle>>& WeightMaps,
-		const FCollectionPropertyConstFacade& PropertyCollection,
-		bool bTrimKinematicConstraints = false)
+		const FCollectionPropertyConstFacade& PropertyCollection)
+		: FXPBDSpringConstraints(
+			Particles,
+			InConstraints,
+			WeightMaps.FindRef(GetXPBDBendingSpringStiffnessString(PropertyCollection, XPBDBendingSpringStiffnessName.ToString())),
+			WeightMaps.FindRef(GetXPBDBendingSpringDampingString(PropertyCollection, XPBDBendingSpringDampingName.ToString())),
+			FSolverVec2(GetWeightedFloatXPBDBendingSpringStiffness(PropertyCollection, MaxStiffness)),
+			FSolverVec2(GetWeightedFloatXPBDBendingSpringDamping(PropertyCollection, MinDampingRatio)))
+		, XPBDBendingSpringStiffnessIndex(PropertyCollection)
+		, XPBDBendingSpringDampingIndex(PropertyCollection)
+	{}
+
+	UE_DEPRECATED(5.4, "XPBD Constraints must always trim kinematic constraints")
+	FXPBDBendingSpringConstraints(
+		const FSolverParticlesRange & Particles,
+		const TArray<TVec2<int32>>&InConstraints,
+		const TMap<FString, TConstArrayView<FRealSingle>>&WeightMaps,
+		const FCollectionPropertyConstFacade & PropertyCollection,
+		bool bTrimKinematicConstraints)
+		: FXPBDSpringConstraints(
+			Particles,
+			InConstraints,
+			WeightMaps.FindRef(GetXPBDBendingSpringStiffnessString(PropertyCollection, XPBDBendingSpringStiffnessName.ToString())),
+			WeightMaps.FindRef(GetXPBDBendingSpringDampingString(PropertyCollection, XPBDBendingSpringDampingName.ToString())),
+			FSolverVec2(GetWeightedFloatXPBDBendingSpringStiffness(PropertyCollection, MaxStiffness)),
+			FSolverVec2(GetWeightedFloatXPBDBendingSpringDamping(PropertyCollection, MinDampingRatio)))
+		, XPBDBendingSpringStiffnessIndex(PropertyCollection)
+		, XPBDBendingSpringDampingIndex(PropertyCollection)
+	{}
+
+	FXPBDBendingSpringConstraints(
+		const FSolverParticles & Particles,
+		int32 ParticleOffset,
+		int32 ParticleCount,
+		const TArray<TVec2<int32>>&InConstraints,
+		const TMap<FString, TConstArrayView<FRealSingle>>&WeightMaps,
+		const FCollectionPropertyConstFacade & PropertyCollection)
 		: FXPBDSpringConstraints(
 			Particles,
+			ParticleOffset,
+			ParticleCount,
 			InConstraints,
 			WeightMaps.FindRef(GetXPBDBendingSpringStiffnessString(PropertyCollection, XPBDBendingSpringStiffnessName.ToString())),
 			WeightMaps.FindRef(GetXPBDBendingSpringDampingString(PropertyCollection, XPBDBendingSpringDampingName.ToString())),
 			FSolverVec2(GetWeightedFloatXPBDBendingSpringStiffness(PropertyCollection, MaxStiffness)),
-			FSolverVec2(GetWeightedFloatXPBDBendingSpringDamping(PropertyCollection, MinDampingRatio)),
-			bTrimKinematicConstraints)
+			FSolverVec2(GetWeightedFloatXPBDBendingSpringDamping(PropertyCollection, MinDampingRatio)))
 		, XPBDBendingSpringStiffnessIndex(PropertyCollection)
 		, XPBDBendingSpringDampingIndex(PropertyCollection)
 	{}
 
+	UE_DEPRECATED(5.4, "XPBD Constraints must always trim kinematic constraints")
 	FXPBDBendingSpringConstraints(
 		const FSolverParticles& Particles,
 		int32 ParticleOffset,
@@ -292,7 +425,7 @@ public:
 		const TArray<TVec2<int32>>& InConstraints,
 		const TMap<FString, TConstArrayView<FRealSingle>>& WeightMaps,
 		const FCollectionPropertyConstFacade& PropertyCollection,
-		bool bTrimKinematicConstraints = false)
+		bool bTrimKinematicConstraints)
 		: FXPBDSpringConstraints(
 			Particles,
 			ParticleOffset,
@@ -301,8 +434,7 @@ public:
 			WeightMaps.FindRef(GetXPBDBendingSpringStiffnessString(PropertyCollection, XPBDBendingSpringStiffnessName.ToString())),
 			WeightMaps.FindRef(GetXPBDBendingSpringDampingString(PropertyCollection, XPBDBendingSpringDampingName.ToString())),
 			FSolverVec2(GetWeightedFloatXPBDBendingSpringStiffness(PropertyCollection, MaxStiffness)),
-			FSolverVec2(GetWeightedFloatXPBDBendingSpringDamping(PropertyCollection, MinDampingRatio)),
-			bTrimKinematicConstraints)
+			FSolverVec2(GetWeightedFloatXPBDBendingSpringDamping(PropertyCollection, MinDampingRatio)))
 		, XPBDBendingSpringStiffnessIndex(PropertyCollection)
 		, XPBDBendingSpringDampingIndex(PropertyCollection)
 	{}
@@ -325,8 +457,7 @@ public:
 			StiffnessMultipliers,
 			DampingMultipliers,
 			FSolverVec2(GetWeightedFloatXPBDBendingSpringStiffness(PropertyCollection, MaxStiffness)),
-			FSolverVec2(GetWeightedFloatXPBDBendingSpringDamping(PropertyCollection, MinDampingRatio)),
-			bTrimKinematicConstraints)
+			FSolverVec2(GetWeightedFloatXPBDBendingSpringDamping(PropertyCollection, MinDampingRatio)))
 		, XPBDBendingSpringStiffnessIndex(PropertyCollection)
 		, XPBDBendingSpringDampingIndex(PropertyCollection)
 	{}
