Iris - Fix use after free.

#rnx

diff --git a/Engine/Source/Runtime/Experimental/Iris/Core/Private/Iris/ReplicationSystem/ObjectReplicationBridge.cpp b/Engine/Source/Runtime/Experimental/Iris/Core/Private/Iris/ReplicationSystem/ObjectReplicationBridge.cpp
--- a/Engine/Source/Runtime/Experimental/Iris/Core/Private/Iris/ReplicationSystem/ObjectReplicationBridge.cpp
+++ b/Engine/Source/Runtime/Experimental/Iris/Core/Private/Iris/ReplicationSystem/ObjectReplicationBridge.cpp
@@ -1633,13 +1633,14 @@ UE::Net::FNetObjectPrioritizerHandle UObjectReplicationBridge::GetPrioritizer(co
 		{
 			const FName SuperClassName = GetConfigClassPathName(SuperClass);
 
-			if (const FClassPrioritizerInfo* PrioritizerInfo = ClassesWithPrioritizer.Find(SuperClassName))
+			if (const FClassPrioritizerInfo* PrioritizerInfoPtr = ClassesWithPrioritizer.Find(SuperClassName))
 			{
 				// Copy info to this class
-				ClassesWithPrioritizer.Add(ClassName, *PrioritizerInfo);
+				FClassPrioritizerInfo PrioritizerInfo = *PrioritizerInfoPtr;
+				ClassesWithPrioritizer.Add(ClassName, PrioritizerInfo);
 
-				const bool bUsePrioritizer = !bRequireForceEnabled || PrioritizerInfo->bForceEnable;
-				return bUsePrioritizer ? PrioritizerInfo->PrioritizerHandle : InvalidNetObjectPrioritizerHandle;
+				const bool bUsePrioritizer = !bRequireForceEnabled || PrioritizerInfo.bForceEnable;
+				return bUsePrioritizer ? PrioritizerInfo.PrioritizerHandle : InvalidNetObjectPrioritizerHandle;
 			}
 		}
 	}
