Chaos - fix use of memmark in FImplicitBVH

diff --git a/Engine/Source/Runtime/Experimental/Chaos/Public/Chaos/ImplicitObjectBVH.h b/Engine/Source/Runtime/Experimental/Chaos/Public/Chaos/ImplicitObjectBVH.h
--- a/Engine/Source/Runtime/Experimental/Chaos/Public/Chaos/ImplicitObjectBVH.h
+++ b/Engine/Source/Runtime/Experimental/Chaos/Public/Chaos/ImplicitObjectBVH.h
@@ -242,8 +242,8 @@ namespace Chaos
 				}
 
 				FMemMark Mark(FMemStack::Get());
-				TArray<int32> NodeStack;
-				NodeStack.Reserve(GetNumObjects());
+				TArray<int32, TMemStackAllocator<>> NodeStack;
+				NodeStack.Reserve(GetDepth());
 
 				int32 NodeIndex = 0;
 				while (true)
@@ -258,6 +258,8 @@ namespace Chaos
 						}
 						else
 						{
+							check(NodeStack.Num() < GetDepth());
+
 							NodeIndex = Node.ChildNodeIndices[0];
 							NodeStack.Push(Node.ChildNodeIndices[1]);
 							continue;
@@ -282,8 +284,8 @@ namespace Chaos
 				}
 
 				FMemMark Mark(FMemStack::Get());
-				TArray<TPair<int32, int32>> NodeStack;	// NodeIndex, NodeDepth
-				NodeStack.Reserve(GetNumObjects());
+				TArray<TPair<int32, int32>, TMemStackAllocator<>> NodeStack;	// NodeIndex, NodeDepth
+				NodeStack.Reserve(GetDepth());
 
 				int32 NodeIndex = 0;
 				int32 NodeDepth = 0;
@@ -295,6 +297,8 @@ namespace Chaos
 
 					if (bVisitChildren && !Node.IsLeaf())
 					{
+						check(NodeStack.Num() < GetDepth());
+
 						const int32 ChildNodeIndexL = Node.ChildNodeIndices[0];
 						const int32 ChildNodeIndexR = Node.ChildNodeIndices[1];
 						NodeDepth = NodeDepth + 1;
@@ -324,8 +328,9 @@ namespace Chaos
 
 				// The node pair stack
 				FMemMark Mark(FMemStack::Get());
-				TArray<TVec2<int32>> NodePairStack;
-				NodePairStack.Reserve(FMath::Max(BVHA.GetNumObjects(), BVHB.GetNumObjects()));	// Can we make a better estimate of required size?
+				TArray<TVec2<int32>, TMemStackAllocator<>> NodePairStack;
+				const int32 NodeStackMax = BVHA.GetDepth() + BVHB.GetDepth();
+				NodePairStack.Reserve(NodeStackMax);
 
 				int32 NodeIndexA = 0;
 				int32 NodeIndexB = 0;
@@ -346,6 +351,8 @@ namespace Chaos
 						}
 						else
 						{
+							check(NodePairStack.Num() < NodeStackMax);
+
 							// @todo(chaos): rule to choose whether to descend into A or B first
 							// Descend into B first, until we reach its leaf nodes
 							const bool bDescendA = NodeB.IsLeaf();
@@ -386,8 +393,8 @@ namespace Chaos
 				}
 
 				FMemMark Mark(FMemStack::Get());
-				TArray<int32> NodeStack;
-				NodeStack.Reserve(GetNumObjects());
+				TArray<int32, TMemStackAllocator<>> NodeStack;
+				NodeStack.Reserve(GetDepth());
 
 				int32 NodeIndex = 0;
 				while (true)
@@ -400,6 +407,7 @@ namespace Chaos
 						{
 							return true;
 						}
+						check(NodeStack.Num() < GetDepth());
 
 						NodeIndex = Node.ChildNodeIndices[0];
 						NodeStack.Push(Node.ChildNodeIndices[1]);
