Chaos - add more solver callbacks

diff --git a/Engine/Source/Runtime/Experimental/Chaos/Private/Chaos/PBDRigidsEvolutionGBF.cpp b/Engine/Source/Runtime/Experimental/Chaos/Private/Chaos/PBDRigidsEvolutionGBF.cpp
--- a/Engine/Source/Runtime/Experimental/Chaos/Private/Chaos/PBDRigidsEvolutionGBF.cpp
+++ b/Engine/Source/Runtime/Experimental/Chaos/Private/Chaos/PBDRigidsEvolutionGBF.cpp
@@ -133,7 +133,10 @@ namespace Chaos
 		DECLARE_CYCLE_STAT(TEXT("FPBDRigidsEvolutionGBF::UnclusterUnions"), STAT_Evolution_UnclusterUnions, STATGROUP_Chaos);
 		DECLARE_CYCLE_STAT(TEXT("FPBDRigidsEvolutionGBF::Integrate"), STAT_Evolution_Integrate, STATGROUP_Chaos);
 		DECLARE_CYCLE_STAT(TEXT("FPBDRigidsEvolutionGBF::KinematicTargets"), STAT_Evolution_KinematicTargets, STATGROUP_Chaos);
+		DECLARE_CYCLE_STAT(TEXT("FPBDRigidsEvolutionGBF::PreIntegrateCallback"), STAT_Evolution_PreIntegrateCallback, STATGROUP_Chaos);
 		DECLARE_CYCLE_STAT(TEXT("FPBDRigidsEvolutionGBF::PostIntegrateCallback"), STAT_Evolution_PostIntegrateCallback, STATGROUP_Chaos);
+		DECLARE_CYCLE_STAT(TEXT("FPBDRigidsEvolutionGBF::PreSolveCallback"), STAT_Evolution_PreSolveCallback, STATGROUP_Chaos);
+		DECLARE_CYCLE_STAT(TEXT("FPBDRigidsEvolutionGBF::PostSolveCallback"), STAT_Evolution_PostSolveCallback, STATGROUP_Chaos);
 		DECLARE_CYCLE_STAT(TEXT("FPBDRigidsEvolutionGBF::CCDModifierCallback"), STAT_Evolution_CCDModifierCallback, STATGROUP_Chaos);
 		DECLARE_CYCLE_STAT(TEXT("FPBDRigidsEvolutionGBF::CollisionModifierCallback"), STAT_Evolution_CollisionModifierCallback, STATGROUP_Chaos);
 		DECLARE_CYCLE_STAT(TEXT("FPBDRigidsEvolutionGBF::MidPhaseModifierCallback"), STAT_Evolution_MidPhaseModifierCallback, STATGROUP_Chaos);
@@ -150,7 +153,6 @@ namespace Chaos
 		DECLARE_CYCLE_STAT(TEXT("FPBDRigidsEvolutionGBF::CreateIslands"), STAT_Evolution_CreateIslands, STATGROUP_Chaos);
 		DECLARE_CYCLE_STAT(TEXT("FPBDRigidsEvolutionGBF::PruneCollisions"), STAT_Evolution_PruneCollisions, STATGROUP_Chaos);
 		DECLARE_CYCLE_STAT(TEXT("FPBDRigidsEvolutionGBF::AddSleepingContacts"), STAT_Evolution_AddSleepingContacts, STATGROUP_Chaos);
-		DECLARE_CYCLE_STAT(TEXT("FPBDRigidsEvolutionGBF::PreApplyCallback"), STAT_Evolution_PreApplyCallback, STATGROUP_Chaos);
 		DECLARE_CYCLE_STAT(TEXT("FPBDRigidsEvolutionGBF::ParallelSolve"), STAT_Evolution_ParallelSolve, STATGROUP_Chaos);
 		DECLARE_CYCLE_STAT(TEXT("FPBDRigidsEvolutionGBF::SaveParticlePostSolve"), STAT_Evolution_SavePostSolve, STATGROUP_Chaos);
 		DECLARE_CYCLE_STAT(TEXT("FPBDRigidsEvolutionGBF::DeactivateSleep"), STAT_Evolution_DeactivateSleep, STATGROUP_Chaos);
@@ -440,6 +442,12 @@ void FPBDRigidsEvolutionGBF::AdvanceOneTimeStepImpl(const FReal Dt, const FSubSt
 	// Update the collision solver type (used to support runtime comparisons of solver types for debugging/testing)
 	UpdateCollisionSolverType();
 
+	if (PreIntegrateCallback != nullptr)
+	{
+		SCOPE_CYCLE_COUNTER(STAT_Evolution_PreIntegrateCallback);
+		PreIntegrateCallback();
+	}
+
 #if CHAOS_EVOLUTION_COLLISION_TESTMODE
 	{
 		TestModeSaveParticles();
@@ -568,10 +576,10 @@ void FPBDRigidsEvolutionGBF::AdvanceOneTimeStepImpl(const FReal Dt, const FSubSt
 		CollisionConstraints.GetConstraintAllocator().PruneExpiredItems();
 	}
 
-	if (PreApplyCallback != nullptr)
+	if (PreSolveCallback != nullptr)
 	{
-		SCOPE_CYCLE_COUNTER(STAT_Evolution_PreApplyCallback);
-		PreApplyCallback();
+		SCOPE_CYCLE_COUNTER(STAT_Evolution_PreSolveCallback);
+		PreSolveCallback();
 	}
 
 	// todo(chaos) : we are using the main gravity ( index 0 ) we should revise this to account for the various gravities based on the constraint ? 
@@ -620,6 +628,12 @@ void FPBDRigidsEvolutionGBF::AdvanceOneTimeStepImpl(const FReal Dt, const FSubSt
 		}
 	}
 
+	if (PostSolveCallback != nullptr)
+	{
+		SCOPE_CYCLE_COUNTER(STAT_Evolution_PostSolveCallback);
+		PostSolveCallback();
+	}
+
 	{
 		SCOPE_CYCLE_COUNTER(STAT_Evolution_SavePostSolve);
 		CSV_SCOPED_TIMING_STAT(PhysicsVerbose, StepSolver_SavePostSolve);
@@ -751,8 +765,10 @@ FPBDRigidsEvolutionGBF::FPBDRigidsEvolutionGBF(
 	, CollisionConstraints(InParticles, Collided, PhysicsMaterials, PerParticlePhysicsMaterials, &SolverPhysicsMaterials, CalculateNumCollisionsPerBlock(), DefaultRestitutionThreshold)
 	, BroadPhase(InParticles)
 	, CollisionDetector(BroadPhase, CollisionConstraints)
+	, PreIntegrateCallback(nullptr)
 	, PostIntegrateCallback(nullptr)
-	, PreApplyCallback(nullptr)
+	, PreSolveCallback(nullptr)
+	, PostSolveCallback(nullptr)
 	, CurrentStepResimCacheImp(nullptr)
 	, MidPhaseModifiers(InMidPhaseModifiers)
 	, CCDModifiers(InCCDModifiers)
diff --git a/Engine/Source/Runtime/Experimental/Chaos/Private/PBDRigidsSolver.cpp b/Engine/Source/Runtime/Experimental/Chaos/Private/PBDRigidsSolver.cpp
--- a/Engine/Source/Runtime/Experimental/Chaos/Private/PBDRigidsSolver.cpp
+++ b/Engine/Source/Runtime/Experimental/Chaos/Private/PBDRigidsSolver.cpp
@@ -611,11 +611,59 @@ namespace Chaos
 				NewParticle->SetPhysicsProxy(Proxy);
 			});
 
-		MEvolution->SetPreApplyCallback(
+		MEvolution->SetPreIntegrateCallback(
 			[this]()
 			{
+				for (ISimCallbackObject* Callback : SimCallbackObjects)
+				{
+					if (Callback->HasOption(ESimCallbackOptions::PreIntegrate))
+					{
+						FScopedTraceSolverCallback TraceCallback(Callback);
+						Callback->PreIntegrate_Internal();
+					}
+				}
+			});
+
+		MEvolution->SetPostIntegrateCallback(
+			[this]()
+			{
+				for (ISimCallbackObject* Callback : SimCallbackObjects)
+				{
+					if (Callback->HasOption(ESimCallbackOptions::PostIntegrate))
+					{
+						FScopedTraceSolverCallback TraceCallback(Callback);
+						Callback->PostIntegrate_Internal();
+					}
+				}
+			});
+
+		MEvolution->SetPreSolveCallback(
+			[this]()
+			{
+				for (ISimCallbackObject* Callback : SimCallbackObjects)
+				{
+					if (Callback->HasOption(ESimCallbackOptions::PreSolve))
+					{
+						FScopedTraceSolverCallback TraceCallback(Callback);
+						Callback->PreSolve_Internal();
+					}
+				}
+
 				PreSolveDebugDraw();
 			});
+
+		MEvolution->SetPostSolveCallback(
+			[this]()
+			{
+				for (ISimCallbackObject* Callback : SimCallbackObjects)
+				{
+					if (Callback->HasOption(ESimCallbackOptions::PostSolve))
+					{
+						FScopedTraceSolverCallback TraceCallback(Callback);
+						Callback->PostSolve_Internal();
+					}
+				}
+			});
 	}
 
 #if CHAOS_DEBUG_NAME
diff --git a/Engine/Source/Runtime/Experimental/Chaos/Public/Chaos/PBDRigidsEvolutionGBF.h b/Engine/Source/Runtime/Experimental/Chaos/Public/Chaos/PBDRigidsEvolutionGBF.h
--- a/Engine/Source/Runtime/Experimental/Chaos/Public/Chaos/PBDRigidsEvolutionGBF.h
+++ b/Engine/Source/Runtime/Experimental/Chaos/Public/Chaos/PBDRigidsEvolutionGBF.h
@@ -77,21 +77,34 @@ namespace Chaos
 			bool InIsSingleThreaded = false);
 		CHAOS_API ~FPBDRigidsEvolutionGBF();
 
-		FORCEINLINE void SetPostIntegrateCallback(const FPBDRigidsEvolutionCallback& Cb)
+		void SetPreIntegrateCallback(const FPBDRigidsEvolutionCallback& Cb)
+		{
+			PreIntegrateCallback = Cb;
+		}
+
+		void SetPostIntegrateCallback(const FPBDRigidsEvolutionCallback& Cb)
 		{
 			PostIntegrateCallback = Cb;
 		}
 
-		FORCEINLINE void SetPostDetectCollisionsCallback(const FPBDRigidsEvolutionCallback& Cb)
+		void SetPreSolveCallback(const FPBDRigidsEvolutionCallback& Cb)
 		{
-			PostDetectCollisionsCallback = Cb;
+			PreSolveCallback = Cb;
 		}
 
-		FORCEINLINE void SetPreApplyCallback(const FPBDRigidsEvolutionCallback& Cb)
+		void SetPostSolveCallback(const FPBDRigidsEvolutionCallback& Cb)
 		{
-			PreApplyCallback = Cb;
+			PostSolveCallback = Cb;
 		}
 
+		void SetPostDetectCollisionsCallback(const FPBDRigidsEvolutionCallback& Cb)
+		{
+			PostDetectCollisionsCallback = Cb;
+		}
+
+		UE_DEPRECATED(5.4, "Use SetPreSolveCallback")
+		void SetPreApplyCallback(const FPBDRigidsEvolutionCallback& Cb) { SetPreSolveCallback(Cb); }
+
 		FORCEINLINE void SetInternalParticleInitilizationFunction(const FPBDRigidsEvolutionInternalHandleCallback& Cb)
 		{ 
 			InternalParticleInitilization = Cb;
@@ -398,9 +411,11 @@ namespace Chaos
 		FSpatialAccelerationBroadPhase BroadPhase;
 		FSpatialAccelerationCollisionDetector CollisionDetector;
 
+		FPBDRigidsEvolutionCallback PreIntegrateCallback;
 		FPBDRigidsEvolutionCallback PostIntegrateCallback;
 		FPBDRigidsEvolutionCallback PostDetectCollisionsCallback;
-		FPBDRigidsEvolutionCallback PreApplyCallback;
+		FPBDRigidsEvolutionCallback PreSolveCallback;
+		FPBDRigidsEvolutionCallback PostSolveCallback;
 		FPBDRigidsEvolutionInternalHandleCallback InternalParticleInitilization;
 		FEvolutionResimCache* CurrentStepResimCacheImp;
 
diff --git a/Engine/Source/Runtime/Experimental/Chaos/Public/Chaos/SimCallbackObject.h b/Engine/Source/Runtime/Experimental/Chaos/Public/Chaos/SimCallbackObject.h
--- a/Engine/Source/Runtime/Experimental/Chaos/Public/Chaos/SimCallbackObject.h
+++ b/Engine/Source/Runtime/Experimental/Chaos/Public/Chaos/SimCallbackObject.h
@@ -40,7 +40,11 @@ enum class ESimCallbackOptions : uint16
 	ParticleUnregister		= 1 << 7,
 	RunOnFrozenGameThread	= 1 << 8,
 	Rewind					= 1 << 9,
-	PhysicsObjectUnregister	= 1 << 10
+	PhysicsObjectUnregister	= 1 << 10,
+	PreIntegrate			= 1 << 11,
+	PostIntegrate			= 1 << 12,
+	PreSolve				= 1 << 13,
+	PostSolve				= 1 << 14,
 };
 ENUM_CLASS_FLAGS(ESimCallbackOptions)
 
@@ -80,6 +84,16 @@ public:
 		OnPreSimulate_Internal();
 	}
 
+	void PreIntegrate_Internal()
+	{
+		OnPreIntegrate_Internal();
+	}
+
+	void PostIntegrate_Internal()
+	{
+		OnPostIntegrate_Internal();
+	}
+
 	void MidPhaseModification_Internal(FMidPhaseModifierAccessor& Modifier)
 	{
 		OnMidPhaseModification_Internal(Modifier);
@@ -100,6 +114,16 @@ public:
 		OnContactModification_Internal(Modifier);
 	}
 
+	void PreSolve_Internal()
+	{
+		OnPreSolve_Internal();
+	}
+
+	void PostSolve_Internal()
+	{
+		OnPostSolve_Internal();
+	}
+
 	void FinalizeOutputData_Internal()
 	{
 		if(CurrentOutput_Internal)
@@ -224,9 +248,29 @@ private:
 	virtual FSimCallbackInput* AllocateInputData_External() = 0;
 
 	/**
-	* Called before simulation step
+	* Called before simulation step (NOTE: not once per sub-step when sub-stepping is enabled)
 	*/
 	virtual void OnPreSimulate_Internal() = 0;
+	
+	/**
+	* Called once per simulation sub-step, before Integrate. Can be used to modify particle positions, velocities etc.
+	*
+	* NOTE: you must explicitly request PreIntegrate when registering the callback for this to be called
+	*/
+	virtual void OnPreIntegrate_Internal()
+	{ 
+		check(false);
+	}
+
+	/**
+	* Called once per simulation sub-step. Can be used to modify particle positions, velocities etc.
+	*
+	* NOTE: you must explicitly request PostIntegrate when registering the callback for this to be called
+	*/
+	virtual void OnPostIntegrate_Internal()
+	{
+		check(false);
+	}
 
 	/**
 	* Called once per simulation step. Allows user to modify midphase pairs
@@ -238,7 +282,6 @@ private:
 		check(false);
 	}
 
-
 	/**
 	* Called once per simulation step. Allows user to modify CCD results
 	*
@@ -265,6 +308,26 @@ private:
 		check(false);
 	}
 
+	/**
+	* Called once per simulation sub-step. Can be used to modify particle positions, velocities etc.
+	*
+	* NOTE: you must explicitly request PreSolve when registering the callback for this to be called
+	*/
+	virtual void OnPreSolve_Internal()
+	{
+		check(false);
+	}
+
+	/**
+	* Called once per simulation sub-step. Can be used to modify particle positions, velocities etc.
+	*
+	* NOTE: you must explicitly request PostSolve when registering the callback for this to be called
+	*/
+	virtual void OnPostSolve_Internal()
+	{
+		check(false);
+	}
+
 	/**
 	* Called once in a simulation step if any new particles were registered. Occurs after UniqueIdxs
 	* are valid.
