Port Destruction builder pipeline to Atom plugin

#virtualized

diff --git a/Engine/Source/Runtime/Experimental/Dataflow/Core/Public/Dataflow/DataflowNodeFactory.h b/Engine/Source/Runtime/Experimental/Dataflow/Core/Public/Dataflow/DataflowNodeFactory.h
--- a/Engine/Source/Runtime/Experimental/Dataflow/Core/Public/Dataflow/DataflowNodeFactory.h
+++ b/Engine/Source/Runtime/Experimental/Dataflow/Core/Public/Dataflow/DataflowNodeFactory.h
@@ -63,7 +63,7 @@ namespace Dataflow
 			return Instance;
 		}
 
-		void RegisterNode(const FFactoryParameters& Parameters, FNewNodeFunction NewFunction)
+		DATAFLOWCORE_API void RegisterNode(const FFactoryParameters& Parameters, FNewNodeFunction NewFunction)
 		{
 			bool bRegisterNode = true;
 			if (ClassMap.Contains(Parameters.TypeName) || DisplayMap.Contains(Parameters.DisplayName))
