Dataflow : Rendering
- Make context rendering argument const

diff --git a/Engine/Source/Runtime/Experimental/Dataflow/Engine/Private/Dataflow/DataflowEdNode.cpp b/Engine/Source/Runtime/Experimental/Dataflow/Engine/Private/Dataflow/DataflowEdNode.cpp
--- a/Engine/Source/Runtime/Experimental/Dataflow/Engine/Private/Dataflow/DataflowEdNode.cpp
+++ b/Engine/Source/Runtime/Experimental/Dataflow/Engine/Private/Dataflow/DataflowEdNode.cpp
@@ -490,7 +490,7 @@ TArray<Dataflow::FRenderingParameter> UDataflowEdNode::GetRenderParameters() con
 }
 
 
-bool UDataflowEdNode::Render(GeometryCollection::Facades::FRenderingFacade& RenderData, TSharedPtr<Dataflow::FContext> Context) const
+bool UDataflowEdNode::Render(GeometryCollection::Facades::FRenderingFacade& RenderData, const TSharedPtr<Dataflow::FContext> Context) const
 {
 	bool bNeedsRefresh = false;
 	if (DataflowGraph)
diff --git a/Engine/Source/Runtime/Experimental/Dataflow/Engine/Public/Dataflow/DataflowEdNode.h b/Engine/Source/Runtime/Experimental/Dataflow/Engine/Public/Dataflow/DataflowEdNode.h
--- a/Engine/Source/Runtime/Experimental/Dataflow/Engine/Public/Dataflow/DataflowEdNode.h
+++ b/Engine/Source/Runtime/Experimental/Dataflow/Engine/Public/Dataflow/DataflowEdNode.h
@@ -79,7 +79,7 @@ public:
 	DATAFLOWENGINE_API void SetAssetRender(bool bInRender);
 	bool DoAssetRender() { return bRenderInAssetEditor; }
 	DATAFLOWENGINE_API TArray<Dataflow::FRenderingParameter> GetRenderParameters() const;
-	DATAFLOWENGINE_API virtual bool Render(GeometryCollection::Facades::FRenderingFacade& RenderData, TSharedPtr<Dataflow::FContext> Context) const;
+	DATAFLOWENGINE_API virtual bool Render(GeometryCollection::Facades::FRenderingFacade& RenderData, const TSharedPtr<Dataflow::FContext> Context) const;
 
 	UPROPERTY()
 	bool bRenderInAssetEditor = false;
