Fix compilation error related to using LogChaos in dataflow module

diff --git a/Engine/Source/Runtime/Experimental/Dataflow/Core/Private/Dataflow/DataflowNodeFactory.cpp b/Engine/Source/Runtime/Experimental/Dataflow/Core/Private/Dataflow/DataflowNodeFactory.cpp
--- a/Engine/Source/Runtime/Experimental/Dataflow/Core/Private/Dataflow/DataflowNodeFactory.cpp
+++ b/Engine/Source/Runtime/Experimental/Dataflow/Core/Private/Dataflow/DataflowNodeFactory.cpp
@@ -5,6 +5,8 @@
 #include "Dataflow/DataflowNode.h"
 #include "Misc/MessageDialog.h"
 
+DEFINE_LOG_CATEGORY(LogDataflowFactory);
+
 namespace Dataflow
 {
 	FNodeFactory* FNodeFactory::Instance = nullptr;
diff --git a/Engine/Source/Runtime/Experimental/Dataflow/Core/Public/Dataflow/DataflowNodeFactory.h b/Engine/Source/Runtime/Experimental/Dataflow/Core/Public/Dataflow/DataflowNodeFactory.h
--- a/Engine/Source/Runtime/Experimental/Dataflow/Core/Public/Dataflow/DataflowNodeFactory.h
+++ b/Engine/Source/Runtime/Experimental/Dataflow/Core/Public/Dataflow/DataflowNodeFactory.h
@@ -4,12 +4,14 @@
 
 #include "CoreMinimal.h"
 #include "Dataflow/DataflowGraph.h"
-#include "ChaosLog.h"
+#include "Logging/LogMacros.h"
 
 struct FDataflowNode;
 struct FDataflowConnection;
 
 
+DATAFLOWCORE_API DECLARE_LOG_CATEGORY_EXTERN(LogDataflowFactory, Warning, All);
+
 namespace Dataflow
 {
 
@@ -70,7 +72,7 @@ namespace Dataflow
 			{
 				if (ParametersMap[Parameters.TypeName].DisplayName.IsEqual(Parameters.DisplayName) )
 				{
-					UE_LOG(LogChaos, Warning, 
+					UE_LOG(LogDataflowFactory, Warning,
 						TEXT("Warning : Dataflow node registration mismatch with type(%s).The \
 						nodes have inconsistent display names(%s) vs(%s).There are two nodes \
 						with the same type being registered."), *Parameters.TypeName.ToString(),
@@ -79,7 +81,7 @@ namespace Dataflow
 				}
 				if (ParametersMap[Parameters.TypeName].Category.IsEqual(Parameters.Category))
 				{
-					UE_LOG(LogChaos, Warning, 
+					UE_LOG(LogDataflowFactory, Warning,
 						TEXT("Warning : Dataflow node registration mismatch with type (%s). The nodes \
 						have inconsistent categories names (%s) vs (%s). There are two different nodes \
 						with the same type being registered. "), *Parameters.TypeName.ToString(),
@@ -88,7 +90,7 @@ namespace Dataflow
 				}
 				if (!ClassMap.Contains(Parameters.TypeName))
 				{
-					UE_LOG(LogChaos, Warning,
+					UE_LOG(LogDataflowFactory, Warning,
 						TEXT("Warning: Attempted to register node type(%s) with display name (%s) \
 						that conflicts with an existing nodes display name (%s)."), 
 						*Parameters.TypeName.ToString(),*Parameters.DisplayName.ToString(), 
