Dataflow : Fix CLang compilation issue

diff --git a/Engine/Source/Runtime/Experimental/Dataflow/Core/Private/Dataflow/DataflowNodeFactory.cpp b/Engine/Source/Runtime/Experimental/Dataflow/Core/Private/Dataflow/DataflowNodeFactory.cpp
--- a/Engine/Source/Runtime/Experimental/Dataflow/Core/Private/Dataflow/DataflowNodeFactory.cpp
+++ b/Engine/Source/Runtime/Experimental/Dataflow/Core/Private/Dataflow/DataflowNodeFactory.cpp
@@ -4,8 +4,9 @@
 
 #include "Dataflow/DataflowNode.h"
 #include "Misc/MessageDialog.h"
+#include "Logging/LogMacros.h"
 
-DEFINE_LOG_CATEGORY(LogDataflowFactory);
+DEFINE_LOG_CATEGORY_STATIC(LogDataflowFactory, Warning, All);
 
 namespace Dataflow
 {
@@ -30,5 +31,45 @@ namespace Dataflow
 		}
 		return TSharedPtr<FDataflowNode>(nullptr);
 	}
+
+	void FNodeFactory::RegisterNode(const FFactoryParameters& Parameters, FNewNodeFunction NewFunction)
+	{
+		bool bRegisterNode = true;
+		if (ClassMap.Contains(Parameters.TypeName) || DisplayMap.Contains(Parameters.DisplayName))
+		{
+			if (ParametersMap[Parameters.TypeName].DisplayName.IsEqual(Parameters.DisplayName))
+			{
+				UE_LOG(LogDataflowFactory, Warning,
+					TEXT("Warning : Dataflow node registration mismatch with type(%s).The \
+						nodes have inconsistent display names(%s) vs(%s).There are two nodes \
+						with the same type being registered."), *Parameters.TypeName.ToString(),
+					*ParametersMap[Parameters.TypeName].DisplayName.ToString(),
+					*Parameters.DisplayName.ToString(), *Parameters.TypeName.ToString());
+			}
+			if (ParametersMap[Parameters.TypeName].Category.IsEqual(Parameters.Category))
+			{
+				UE_LOG(LogDataflowFactory, Warning,
+					TEXT("Warning : Dataflow node registration mismatch with type (%s). The nodes \
+						have inconsistent categories names (%s) vs (%s). There are two different nodes \
+						with the same type being registered. "), *Parameters.TypeName.ToString(),
+					*ParametersMap[Parameters.TypeName].DisplayName.ToString(),
+					*Parameters.DisplayName.ToString(), *Parameters.TypeName.ToString());
+			}
+			if (!ClassMap.Contains(Parameters.TypeName))
+			{
+				UE_LOG(LogDataflowFactory, Warning,
+					TEXT("Warning: Attempted to register node type(%s) with display name (%s) \
+						that conflicts with an existing nodes display name (%s)."),
+					*Parameters.TypeName.ToString(), *Parameters.DisplayName.ToString(),
+					*ParametersMap[Parameters.TypeName].DisplayName.ToString());
+			}
+		}
+		else
+		{
+			ClassMap.Add(Parameters.TypeName, NewFunction);
+			ParametersMap.Add(Parameters.TypeName, Parameters);
+			DisplayMap.Add(Parameters.DisplayName, Parameters.TypeName);
+		}
+	}
 }
 
diff --git a/Engine/Source/Runtime/Experimental/Dataflow/Core/Public/Dataflow/DataflowNodeFactory.h b/Engine/Source/Runtime/Experimental/Dataflow/Core/Public/Dataflow/DataflowNodeFactory.h
--- a/Engine/Source/Runtime/Experimental/Dataflow/Core/Public/Dataflow/DataflowNodeFactory.h
+++ b/Engine/Source/Runtime/Experimental/Dataflow/Core/Public/Dataflow/DataflowNodeFactory.h
@@ -4,14 +4,10 @@
 
 #include "CoreMinimal.h"
 #include "Dataflow/DataflowGraph.h"
-#include "Logging/LogMacros.h"
 
 struct FDataflowNode;
 struct FDataflowConnection;
 
-
-DATAFLOWCORE_API DECLARE_LOG_CATEGORY_EXTERN(LogDataflowFactory, Warning, All);
-
 namespace Dataflow
 {
 
@@ -65,45 +61,7 @@ namespace Dataflow
 			return Instance;
 		}
 
-		DATAFLOWCORE_API void RegisterNode(const FFactoryParameters& Parameters, FNewNodeFunction NewFunction)
-		{
-			bool bRegisterNode = true;
-			if (ClassMap.Contains(Parameters.TypeName) || DisplayMap.Contains(Parameters.DisplayName))
-			{
-				if (ParametersMap[Parameters.TypeName].DisplayName.IsEqual(Parameters.DisplayName) )
-				{
-					UE_LOG(LogDataflowFactory, Warning,
-						TEXT("Warning : Dataflow node registration mismatch with type(%s).The \
-						nodes have inconsistent display names(%s) vs(%s).There are two nodes \
-						with the same type being registered."), *Parameters.TypeName.ToString(),
-						*ParametersMap[Parameters.TypeName].DisplayName.ToString(), 
-						*Parameters.DisplayName.ToString(), *Parameters.TypeName.ToString());
-				}
-				if (ParametersMap[Parameters.TypeName].Category.IsEqual(Parameters.Category))
-				{
-					UE_LOG(LogDataflowFactory, Warning,
-						TEXT("Warning : Dataflow node registration mismatch with type (%s). The nodes \
-						have inconsistent categories names (%s) vs (%s). There are two different nodes \
-						with the same type being registered. "), *Parameters.TypeName.ToString(),
-						*ParametersMap[Parameters.TypeName].DisplayName.ToString(), 
-						*Parameters.DisplayName.ToString(),*Parameters.TypeName.ToString());
-				}
-				if (!ClassMap.Contains(Parameters.TypeName))
-				{
-					UE_LOG(LogDataflowFactory, Warning,
-						TEXT("Warning: Attempted to register node type(%s) with display name (%s) \
-						that conflicts with an existing nodes display name (%s)."), 
-						*Parameters.TypeName.ToString(),*Parameters.DisplayName.ToString(), 
-						*ParametersMap[Parameters.TypeName].DisplayName.ToString());
-				}
-			}
-			else
-			{
-				ClassMap.Add(Parameters.TypeName, NewFunction);
-				ParametersMap.Add(Parameters.TypeName, Parameters);
-				DisplayMap.Add(Parameters.DisplayName, Parameters.TypeName);
-			}
-		}
+		DATAFLOWCORE_API void RegisterNode(const FFactoryParameters& Parameters, FNewNodeFunction NewFunction);
 
 		FName TypeNameFromDisplayName(const FName& DisplayName)
 		{
