Fix potential server out of bounds array access crash

diff --git a/Engine/Source/Runtime/Experimental/Chaos/Private/PhysicsProxy/GeometryCollectionPhysicsProxy.cpp b/Engine/Source/Runtime/Experimental/Chaos/Private/PhysicsProxy/GeometryCollectionPhysicsProxy.cpp
--- a/Engine/Source/Runtime/Experimental/Chaos/Private/PhysicsProxy/GeometryCollectionPhysicsProxy.cpp
+++ b/Engine/Source/Runtime/Experimental/Chaos/Private/PhysicsProxy/GeometryCollectionPhysicsProxy.cpp
@@ -1120,10 +1120,13 @@ Chaos::FPBDRigidClusteredParticleHandle* FGeometryCollectionPhysicsProxy::FindCl
 	else
 	{
 		const int32 TransformIndex = ItemIndex.GetTransformIndex();
-		int32 ParticleIndex = FromTransformToParticleIndex[TransformIndex];
-		if (SolverParticleHandles.IsValidIndex(ParticleIndex))
+		if (FromTransformToParticleIndex.IsValidIndex(TransformIndex))
 		{
-			ResultHandle = SolverParticleHandles[ParticleIndex];
+			const int32 ParticleIndex = FromTransformToParticleIndex[TransformIndex];
+			if (SolverParticleHandles.IsValidIndex(ParticleIndex))
+			{
+				ResultHandle = SolverParticleHandles[ParticleIndex];
+			}
 		}
 	}
 	return ResultHandle;
