Fix warning V1099: Using the 'ConcreteType' function of uninitialized derived class while initializing the 'IPhysicsProxyBase' base class will lead to undefined behavior.

diff --git a/Engine/Source/Runtime/Experimental/Chaos/Public/Chaos/Framework/PhysicsProxy.h b/Engine/Source/Runtime/Experimental/Chaos/Public/Chaos/Framework/PhysicsProxy.h
--- a/Engine/Source/Runtime/Experimental/Chaos/Public/Chaos/Framework/PhysicsProxy.h
+++ b/Engine/Source/Runtime/Experimental/Chaos/Public/Chaos/Framework/PhysicsProxy.h
@@ -81,7 +81,7 @@ public:
 	void BindParticleCallbackMapping(Chaos::TArrayCollectionArray<PhysicsProxyWrapper> & PhysicsProxyReverseMap, Chaos::TArrayCollectionArray<int32> & ParticleIDReverseMap) {static_cast<Concrete*>(this)->BindParticleCallbackMapping(PhysicsProxyReverseMap, ParticleIDReverseMap);}
 
 	/** Returns the concrete type of the derived class*/
-	EPhysicsProxyType ConcreteType() { return static_cast<Concrete*>(this)->ConcreteType(); }
+	static constexpr EPhysicsProxyType ConcreteType() { return Concrete::ConcreteType(); }
 	
 	/**
 	 * CONTEXT: GAMETHREAD
diff --git a/Engine/Source/Runtime/Experimental/Chaos/Public/PhysicsProxy/ClusterUnionPhysicsProxy.h b/Engine/Source/Runtime/Experimental/Chaos/Public/PhysicsProxy/ClusterUnionPhysicsProxy.h
--- a/Engine/Source/Runtime/Experimental/Chaos/Public/PhysicsProxy/ClusterUnionPhysicsProxy.h
+++ b/Engine/Source/Runtime/Experimental/Chaos/Public/PhysicsProxy/ClusterUnionPhysicsProxy.h
@@ -175,7 +175,7 @@ namespace Chaos
 		void DisableCollisionsCallback(TSet<TTuple<int32, int32>>& InPairs) {}
 		void AddForceCallback(FParticlesType& InParticles, const float InDt, const int32 InIndex) {}
 		void BindParticleCallbackMapping(Chaos::TArrayCollectionArray<PhysicsProxyWrapper>& PhysicsProxyReverseMap, Chaos::TArrayCollectionArray<int32>& ParticleIDReverseMap) {}
-		static EPhysicsProxyType ConcreteType() { return EPhysicsProxyType::ClusterUnionProxy; }
+		static constexpr EPhysicsProxyType ConcreteType() { return EPhysicsProxyType::ClusterUnionProxy; }
 		void SyncBeforeDestroy() {}
 		void OnRemoveFromScene() {}
 		bool IsDirty() { return false; }
diff --git a/Engine/Source/Runtime/Experimental/Chaos/Public/PhysicsProxy/GeometryCollectionPhysicsProxy.h b/Engine/Source/Runtime/Experimental/Chaos/Public/PhysicsProxy/GeometryCollectionPhysicsProxy.h
--- a/Engine/Source/Runtime/Experimental/Chaos/Public/PhysicsProxy/GeometryCollectionPhysicsProxy.h
+++ b/Engine/Source/Runtime/Experimental/Chaos/Public/PhysicsProxy/GeometryCollectionPhysicsProxy.h
@@ -228,7 +228,7 @@ public:
 
 	bool IsDirty() { return false; }
 
-	static EPhysicsProxyType ConcreteType() { return EPhysicsProxyType::GeometryCollectionType; }
+	static constexpr EPhysicsProxyType ConcreteType() { return EPhysicsProxyType::GeometryCollectionType; }
 
 	CHAOS_API void SyncBeforeDestroy();
 	CHAOS_API void OnRemoveFromSolver(Chaos::FPBDRigidsSolver *RBDSolver);
diff --git a/Engine/Source/Runtime/Experimental/Chaos/Public/PhysicsProxy/SkeletalMeshPhysicsProxy.h b/Engine/Source/Runtime/Experimental/Chaos/Public/PhysicsProxy/SkeletalMeshPhysicsProxy.h
--- a/Engine/Source/Runtime/Experimental/Chaos/Public/PhysicsProxy/SkeletalMeshPhysicsProxy.h
+++ b/Engine/Source/Runtime/Experimental/Chaos/Public/PhysicsProxy/SkeletalMeshPhysicsProxy.h
@@ -151,7 +151,7 @@ public:
 	CHAOS_API void FlipBuffer();
 	CHAOS_API bool PullFromPhysicsState(const int32 SolverSyncTimestamp);
 	bool IsDirty() { return false; }
-	EPhysicsProxyType ConcreteType() { return EPhysicsProxyType::SkeletalMeshType; }
+	static constexpr EPhysicsProxyType ConcreteType() { return EPhysicsProxyType::SkeletalMeshType; }
 	/** ----------------------- */
 
 	/**
diff --git a/Engine/Source/Runtime/Experimental/Chaos/Public/PhysicsProxy/StaticMeshPhysicsProxy.h b/Engine/Source/Runtime/Experimental/Chaos/Public/PhysicsProxy/StaticMeshPhysicsProxy.h
--- a/Engine/Source/Runtime/Experimental/Chaos/Public/PhysicsProxy/StaticMeshPhysicsProxy.h
+++ b/Engine/Source/Runtime/Experimental/Chaos/Public/PhysicsProxy/StaticMeshPhysicsProxy.h
@@ -108,7 +108,7 @@ public:
 	CHAOS_API bool PullFromPhysicsState(const int32 SolverSyncTimestamp);
 	bool IsDirty() { return false; }
 	FStubSkeletalMeshData* NewData() { return nullptr; }
-	EPhysicsProxyType ConcreteType() { return EPhysicsProxyType::StaticMeshType; }
+	static constexpr EPhysicsProxyType ConcreteType() { return EPhysicsProxyType::StaticMeshType; }
 	/** ----------------------- */
 
 private:
