Iris - Fix off by one issue.

diff --git a/Engine/Source/Runtime/Experimental/Iris/Core/Public/Iris/Serialization/PolymorphicNetSerializer.h b/Engine/Source/Runtime/Experimental/Iris/Core/Public/Iris/Serialization/PolymorphicNetSerializer.h
--- a/Engine/Source/Runtime/Experimental/Iris/Core/Public/Iris/Serialization/PolymorphicNetSerializer.h
+++ b/Engine/Source/Runtime/Experimental/Iris/Core/Public/Iris/Serialization/PolymorphicNetSerializer.h
@@ -62,7 +62,7 @@ namespace UE::Net
 const FPolymorphicNetSerializerScriptStructCache::FTypeInfo* FPolymorphicNetSerializerScriptStructCache::GetTypeInfo(uint32 TypeIndex) const
 {
 	static_assert(InvalidTypeIndex == 0U, "Expected InvalidTypeIndex to be 0");
-	if ((TypeIndex - 1U) > static_cast<uint32>(RegisteredTypes.Num()))
+	if ((TypeIndex - 1U) >= static_cast<uint32>(RegisteredTypes.Num()))
 	{
 		return nullptr;
 	}
