Chaos - fix midphase modifier disable CCD

Disabling CCD via a midphase modifier would leave the CCD ApplyCorrections active on the constraint which it should not.

diff --git a/Engine/Source/Runtime/Experimental/Chaos/Private/Chaos/Collision/ParticlePairMidPhase.cpp b/Engine/Source/Runtime/Experimental/Chaos/Private/Chaos/Collision/ParticlePairMidPhase.cpp
--- a/Engine/Source/Runtime/Experimental/Chaos/Private/Chaos/Collision/ParticlePairMidPhase.cpp
+++ b/Engine/Source/Runtime/Experimental/Chaos/Private/Chaos/Collision/ParticlePairMidPhase.cpp
@@ -462,6 +462,10 @@ namespace Chaos
 
 			Constraint->SetCullDistance(CullDistance);
 
+			// Constraint may have been previously used with CCD enabled (e.g., a midphase modifier)
+			// so we need to make sure that the CCD flag is disabled
+			Constraint->SetCCDEnabled(false);
+
 			// If the constraint was not used last frame, it needs to be reset, otherwise we will try to reuse
 			if (!bWasUpdatedLastTick || (Constraint->GetManifoldPoints().Num() == 0))
 			{
@@ -531,14 +535,17 @@ namespace Chaos
 		{
 			// Lazy creation of the constraint. 
 			CreateConstraint(CullDistance, Context);
+		}
 
-			// Flag this contact as requiring CCD
+		// Constraint may have been previously used with CCD disabled (e.g., a midphase modifier)
+		// so we need to make sure that the CCD flags are set appropriately
+		if (Constraint.IsValid())
+		{
+			check(!Constraint->IsEnabled());
 			Constraint->SetCCDEnabled(true);
+			Constraint->SetCCDSweepEnabled(bEnableCCDSweep);
 		}
-		check(!Constraint->IsEnabled());
 
-		// Do we want to enable the CCD sweep? If not, we fall back to the standard collision detection for this tick
-		Constraint->SetCCDSweepEnabled(bEnableCCDSweep);
 		if (!bEnableCCDSweep)
 		{
 			return GenerateCollision(CullDistance, Dt, Context);
@@ -860,13 +867,13 @@ namespace Chaos
 		InitThresholds();
 	}
 
-	bool FParticlePairMidPhase::ShouldEnableCCD(const FReal Dt)
+	bool FParticlePairMidPhase::ShouldEnableCCDSweep(const FReal Dt)
 	{
 		// bIsCCDActive is set to bIsCCD at the beginning of every frame, but may be
 		// overridden in midphase modification or potentially other systems which run
 		// in between mid and narrow phase. bIsCCDActive indicates the final
 		// overridden value so we use that here instead of bIsCCD.
-		if (Flags.bIsCCDActive)
+		if (Flags.bIsCCDActive != 0)
 		{
 			FConstGenericParticleHandle ConstParticle0 = FConstGenericParticleHandle(Particle0);
 			FConstGenericParticleHandle ConstParticle1 = FConstGenericParticleHandle(Particle1);
@@ -923,7 +930,7 @@ namespace Chaos
 			FReal CullDistance = InCullDistance * CullDistanceScale;
 
 			// If CCD is enabled, did we move far enough to require a sweep?
-			Flags.bUseSweep = Flags.bIsCCD && ShouldEnableCCD(Dt);
+			Flags.bUseSweep = (Flags.bIsCCDActive != 0) && ShouldEnableCCDSweep(Dt);
 
 			// We increase CullDistance based on velocity (up to a limit for perf with large velocities).
 			// NOTE: This somewhat matches the bounds expansion in FPBDRigidsEvolutionGBF::Integrate
@@ -1106,7 +1113,7 @@ namespace Chaos
 		//TRACE_COUNTER_INCREMENT(ChaosTraceCounter_MidPhase_NumShapePair);
 
 		int32 NumActive = 0;
-		if (Flags.bIsCCD)
+		if (Flags.bIsCCDActive != 0)
 		{
 			for (FSingleShapePairCollisionDetector& ShapePair : ShapePairDetectors)
 			{
@@ -1725,9 +1732,6 @@ namespace Chaos
 
 		Constraint->SetCollisionSortKey(CollisionSortKey);
 
-		// Is this a CCD constraint?
-		Constraint->SetCCDEnabled(IsCCD());
-
 		return Constraints.Add(CollisionKey.GetKey(), MoveTemp(Constraint)).Get();
 	}
 
@@ -1747,7 +1751,7 @@ namespace Chaos
 		CHAOS_MIDPHASE_SCOPE_CYCLE_TIMER(FGenericParticlePairMidPhase_ProcessNewConstraints);
 
 		int32 NumActive = 0;
-		const bool bUseCCDSweep = IsCCD() && Flags.bUseSweep;
+		const bool bUseCCDSweep = Flags.bIsCCDActive && Flags.bUseSweep;
 
 		const int32 NumNewConstraints = NewConstraints.Num();
 		for (int32 ConstraintIndex = 0; ConstraintIndex < NumNewConstraints; ++ConstraintIndex)
@@ -1755,12 +1759,9 @@ namespace Chaos
 			FPBDCollisionConstraint* Constraint = NewConstraints[ConstraintIndex];
 			PrefetchConstraint(NewConstraints, ConstraintIndex + 1);
 
-			// Do we want to sweep on this tick? 
-			// CCD may be temporarily disabled by the user or because we are moving slowly.
-			if (IsCCD())
-			{
-				Constraint->SetCCDSweepEnabled(bUseCCDSweep);
-			}
+			// CCD may be temporarily disabled by the user (via a midphase modifier) or because we are moving slowly.
+			Constraint->SetCCDEnabled(Flags.bIsCCDActive);
+			Constraint->SetCCDSweepEnabled(bUseCCDSweep);
 
 			// NOTE: Probe constraints are always active and we run collision detection for them at the end opf the frame
 			bool bIsActive = true;
diff --git a/Engine/Source/Runtime/Experimental/Chaos/Public/Chaos/Collision/ParticlePairMidPhase.h b/Engine/Source/Runtime/Experimental/Chaos/Public/Chaos/Collision/ParticlePairMidPhase.h
--- a/Engine/Source/Runtime/Experimental/Chaos/Public/Chaos/Collision/ParticlePairMidPhase.h
+++ b/Engine/Source/Runtime/Experimental/Chaos/Public/Chaos/Collision/ParticlePairMidPhase.h
@@ -379,7 +379,7 @@ namespace Chaos
 		 * @return true if CCD is enabled this tick, false otherwise
 		 * This may return false, even for collisions on CCD-enabled bodies when the bodies are moving slowly
 		 */
-		CHAOS_API bool ShouldEnableCCD(const FReal Dt);
+		CHAOS_API bool ShouldEnableCCDSweep(const FReal Dt);
 
 		CHAOS_API void InitThresholds();
 
