Chaos - sleeping fixes and behaviour improvement

- add debug draw to visualize sleep state
- implement SleepFamily system that was missing from Chaos (SleepThesholdMultiplier)
- SetV and SetW now wakes particles
- reduced default velocity smoothing used by sleep thresholds
- fix smoothed velocity for particles in awake islands (minor issue that made it slightly harder to go to sleep)
- don't use smoothed velocity for floating particles
- system top alloe reduced sleep thresholds for floating particles (cvar controlled, disable dby default)

diff --git a/Engine/Source/Runtime/Engine/Private/PhysicsEngine/BodyInstance.cpp b/Engine/Source/Runtime/Engine/Private/PhysicsEngine/BodyInstance.cpp
--- a/Engine/Source/Runtime/Engine/Private/PhysicsEngine/BodyInstance.cpp
+++ b/Engine/Source/Runtime/Engine/Private/PhysicsEngine/BodyInstance.cpp
@@ -96,6 +96,12 @@ FAutoConsoleVariableRef CVarSkipShapeCreationForEmptyBodySetup(
 	bSkipShapeCreationForEmptyBodySetup,
 	TEXT("If true, CreateShapesAndActors will not try to create actors and shapes for all instances if the body setup doesn't have any geometry."));
 
+float SensitiveSleepThresholdMultiplier = 1.0f/20.0f;
+FAutoConsoleVariableRef CVarSensitiveSleepThresholdMultiplier(
+	TEXT("p.SensitiveSleepThresholdMultiplier"),
+	SensitiveSleepThresholdMultiplier,
+	TEXT("The sleep threshold multiplier to use for bodies using the Sensitive sleep family."));
+
 using namespace PhysicsInterfaceTypes;
 
 bool IsRigidBodyKinematic_AssumesLocked(const FPhysicsActorHandle& InActorRef)
@@ -3232,7 +3238,7 @@ float FBodyInstance::GetSleepThresholdMultiplier() const
 {
 	if (SleepFamily == ESleepFamily::Sensitive)
 	{
-		return 1 / 20.0f;
+		return SensitiveSleepThresholdMultiplier;
 	}
 	else if (SleepFamily == ESleepFamily::Custom)
 	{
@@ -4162,6 +4168,9 @@ void FBodyInstance::InitDynamicProperties_AssumesLocked()
 			FPhysicsInterface::SetLinearVelocity_AssumesLocked(ActorHandle, InitialLinVel);
 		}
 
+		FPhysicsInterface::SetSleepThresholdMultiplier_AssumesLocked(ActorHandle, GetSleepThresholdMultiplier());
+
+		// @todo: implement sleep energy threshold
 		float SleepEnergyThresh = FPhysicsInterface::GetSleepEnergyThreshold_AssumesLocked(ActorHandle);
 		SleepEnergyThresh *= GetSleepThresholdMultiplier();
 		FPhysicsInterface::SetSleepEnergyThreshold_AssumesLocked(ActorHandle, SleepEnergyThresh);
diff --git a/Engine/Source/Runtime/Experimental/Chaos/Private/Chaos/ClusterUnionManager.cpp b/Engine/Source/Runtime/Experimental/Chaos/Private/Chaos/ClusterUnionManager.cpp
--- a/Engine/Source/Runtime/Experimental/Chaos/Private/Chaos/ClusterUnionManager.cpp
+++ b/Engine/Source/Runtime/Experimental/Chaos/Private/Chaos/ClusterUnionManager.cpp
@@ -646,6 +646,9 @@ namespace Chaos
 		// Cluster has one-way interaction only if all children are also one-way
 		bool bIsOneWayInteraction = (bIsNewCluster) || Cluster->InternalCluster->OneWayInteraction();
 
+		// Use the minimum sleep multiplier of all member particles
+		FRealSingle MinSleepThresholdMultiplier = TNumericLimits<FRealSingle>::Max();
+
 		for (FPBDRigidParticleHandle* Particle : FinalParticlesToAdd)
 		{
 			if (Particle->GetGeometry() == nullptr)
@@ -664,6 +667,7 @@ namespace Chaos
 			}
 
 			bIsOneWayInteraction &= Particle->OneWayInteraction();
+			MinSleepThresholdMultiplier = FMath::Min(MinSleepThresholdMultiplier, Particle->SleepThresholdMultiplier());
 
 			if (!Cluster->ChildProperties.Contains(Particle))
 			{
@@ -735,6 +739,11 @@ namespace Chaos
 		}
 
 		Cluster->InternalCluster->SetOneWayInteraction(bIsOneWayInteraction);
+
+		if (MinSleepThresholdMultiplier != TNumericLimits<FRealSingle>::Max())
+		{
+			Cluster->InternalCluster->SetSleepThresholdMultiplier(MinSleepThresholdMultiplier);
+		}
 	}
 
 	DECLARE_CYCLE_STAT(TEXT("FClusterUnionManager::DeferredClusterUnionUpdate"), STAT_DeferredClusterUnionUpdate, STATGROUP_Chaos);
diff --git a/Engine/Source/Runtime/Experimental/Chaos/Private/Chaos/Island/IslandManager.cpp b/Engine/Source/Runtime/Experimental/Chaos/Private/Chaos/Island/IslandManager.cpp
--- a/Engine/Source/Runtime/Experimental/Chaos/Private/Chaos/Island/IslandManager.cpp
+++ b/Engine/Source/Runtime/Experimental/Chaos/Private/Chaos/Island/IslandManager.cpp
@@ -1,5 +1,6 @@
 // Copyright Epic Games, Inc. All Rights Reserved.
 #include "Chaos/Island/IslandManager.h"
+#include "Chaos/ChaosDebugDraw.h"
 #include "Chaos/ParticleHandle.h"
 #include "Chaos/ParticleIterator.h"
 #include "Chaos/PBDConstraintContainer.h"
@@ -46,26 +47,38 @@ namespace Chaos::CVars
 
 	/** Cvar to enable/disable the island sleeping */
 	bool bChaosSolverSleepEnabled = true;
-	FAutoConsoleVariableRef CVarChaosSolverSleepEnabled(TEXT("p.Chaos.Solver.SleepEnabled"), bChaosSolverSleepEnabled, TEXT(""));
+	FAutoConsoleVariableRef CVarChaosSolverSleepEnabled(TEXT("p.Chaos.Solver.Sleep.Enabled"), bChaosSolverSleepEnabled, TEXT(""));
 
 	/** Cvar to override the sleep counter threshold if necessary */
 	int32 ChaosSolverCollisionDefaultSleepCounterThreshold = 20;
-	FAutoConsoleVariableRef CVarChaosSolverCollisionDefaultSleepCounterThreshold(TEXT("p.ChaosSolverCollisionDefaultSleepCounterThreshold"), ChaosSolverCollisionDefaultSleepCounterThreshold, TEXT("Default counter threshold for sleeping.[def:20]"));
+	FAutoConsoleVariableRef CVarChaosSolverCollisionDefaultSleepCounterThreshold(TEXT("p.Chaos.Solver.Sleep.Defaults.SleepCounterThreshold"), ChaosSolverCollisionDefaultSleepCounterThreshold, TEXT("Default counter threshold for sleeping.[def:20]"));
 
 	/** Cvar to override the sleep linear threshold if necessary */
 	FRealSingle ChaosSolverCollisionDefaultLinearSleepThreshold = 0.001f; // .001 unit mass cm
-	FAutoConsoleVariableRef CVarChaosSolverCollisionDefaultLinearSleepThreshold(TEXT("p.ChaosSolverCollisionDefaultLinearSleepThreshold"), ChaosSolverCollisionDefaultLinearSleepThreshold, TEXT("Default linear threshold for sleeping.[def:0.001]"));
+	FAutoConsoleVariableRef CVarChaosSolverCollisionDefaultLinearSleepThreshold(TEXT("p.Chaos.Solver.Sleep.Defaults.LinearSleepThreshold"), ChaosSolverCollisionDefaultLinearSleepThreshold, TEXT("Default linear threshold for sleeping.[def:0.001]"));
 
 	/** Cvar to override the sleep angular threshold if necessary */
 	FRealSingle ChaosSolverCollisionDefaultAngularSleepThreshold = 0.0087f;  //~1/2 unit mass degree
-	FAutoConsoleVariableRef CVarChaosSolverCollisionDefaultAngularSleepThreshold(TEXT("p.ChaosSolverCollisionDefaultAngularSleepThreshold"), ChaosSolverCollisionDefaultAngularSleepThreshold, TEXT("Default angular threshold for sleeping.[def:0.0087]"));
+	FAutoConsoleVariableRef CVarChaosSolverCollisionDefaultAngularSleepThreshold(TEXT("p.Chaos.Solver.Sleep.Defaults.AngularSleepThreshold"), ChaosSolverCollisionDefaultAngularSleepThreshold, TEXT("Default angular threshold for sleeping.[def:0.0087]"));
 
 	/** The size of object for which the angular sleep threshold is defined. Large objects reduce the threshold propertionally. 0 means do not apply size scale. */
 	// E.g., if ChaosSolverCollisionAngularSleepThresholdSize=100, an objects with a bounds of 500 will have 1/5x the sleep threshold.
 	// We are effectively converting the angular threshold into a linear threshold calculated at the object extents.
 	// @todo(chaos): male this a project setting or something
 	FRealSingle ChaosSolverCollisionAngularSleepThresholdSize = 0;
-	FAutoConsoleVariableRef CVarChaosSolverCollisionAngularSleepThresholdSize(TEXT("p.ChaosSolverCollisionAngularSleepThresholdSize"), ChaosSolverCollisionAngularSleepThresholdSize, TEXT("Scales the angular threshold based on size (0 to disable size based scaling)"));
+	FAutoConsoleVariableRef CVarChaosSolverCollisionAngularSleepThresholdSize(TEXT("p.Chaos.Solver.Sleep.AngularSleepThresholdSize"), ChaosSolverCollisionAngularSleepThresholdSize, TEXT("Scales the angular threshold based on size (0 to disable size based scaling)"));
+
+	/* Cvar to increase the sleep counter threshold for floating particles */
+	int32 IsolatedParticleSleepCounterThresholdMultiplier = 1;
+	FAutoConsoleVariableRef CVarChaosSolverIsolatedParticleSleepCounterThresholdMultiplier(TEXT("p.Chaos.Solver.Sleep.IsolatedParticle.CounterMultiplier"), IsolatedParticleSleepCounterThresholdMultiplier, TEXT("A multiplier applied to SleepCounterThreshold for floating particles"));
+
+	/* Cvar to adjust the sleep linear threshold for floating particles */
+	FRealSingle IsolatedParticleSleepLinearThresholdMultiplier = 1.0f;
+	FAutoConsoleVariableRef CVarChaosSolverIsolatedParticleSleepLinearThresholdMultiplier(TEXT("p.Chaos.Solver.Sleep.IsolatedParticle.LinearMultiplier"), IsolatedParticleSleepLinearThresholdMultiplier, TEXT("A multiplier applied to SleepLinearThreshold for floating particles"));
+
+	/* Cvar to adjust the sleep angular threshold for floating particles */
+	FRealSingle IsolatedParticleSleepAngularThresholdMultiplier = 1.0f;
+	FAutoConsoleVariableRef CVarChaosSolverIsolatedParticleSleepAngularThresholdMultiplier(TEXT("p.Chaos.Solver.Sleep.IsolatedParticle.AngularMultiplier"), IsolatedParticleSleepAngularThresholdMultiplier, TEXT("A multiplier applied to SleepAngularThreshold for floating particles"));
 }
 
 
@@ -132,7 +145,7 @@ namespace Chaos::Private
 		return (Particle != nullptr) && (Particle->SyncState() != ESyncState::InSync);
 	}
 
-	bool GetParticleSleepThresholds(
+	bool GetIslandParticleSleepThresholds(
 		const FGeometryParticleHandle* Particle,
 		const TArrayCollectionArray<TSerializablePtr<FChaosPhysicsMaterial>>* PhysicsMaterials,
 		const TArrayCollectionArray<TUniquePtr<FChaosPhysicsMaterial>>* PerParticlePhysicsMaterials,
@@ -148,17 +161,19 @@ namespace Chaos::Private
 		const FPBDRigidParticleHandle* Rigid = Particle->CastToRigidParticle();
 		if ((Rigid != nullptr) && (Rigid->SleepType() != ESleepType::NeverSleep))
 		{
+			const FRealSingle ParticleSleepThresholdMultiplier = Rigid->SleepThresholdMultiplier();
+
 			const FChaosPhysicsMaterial* PhysicsMaterial = Private::GetFirstPhysicsMaterial(Rigid, PhysicsMaterials, PerParticlePhysicsMaterials, SimMaterials);
 			if (PhysicsMaterial != nullptr)
 			{
-				OutSleepLinearThreshold = FRealSingle(PhysicsMaterial->SleepingLinearThreshold);
-				OutSleepAngularThreshold = FRealSingle(PhysicsMaterial->SleepingAngularThreshold);
+				OutSleepLinearThreshold = ParticleSleepThresholdMultiplier * FRealSingle(PhysicsMaterial->SleepingLinearThreshold);
+				OutSleepAngularThreshold = ParticleSleepThresholdMultiplier * FRealSingle(PhysicsMaterial->SleepingAngularThreshold);
 				OutSleepCounterThreshold = PhysicsMaterial->SleepCounterThreshold;
 			}
 			else
 			{
-				OutSleepLinearThreshold = CVars::ChaosSolverCollisionDefaultLinearSleepThreshold;
-				OutSleepAngularThreshold = CVars::ChaosSolverCollisionDefaultAngularSleepThreshold;
+				OutSleepLinearThreshold = ParticleSleepThresholdMultiplier * CVars::ChaosSolverCollisionDefaultLinearSleepThreshold;
+				OutSleepAngularThreshold = ParticleSleepThresholdMultiplier * CVars::ChaosSolverCollisionDefaultAngularSleepThreshold;
 				OutSleepCounterThreshold = CVars::ChaosSolverCollisionDefaultSleepCounterThreshold;
 			}
 
@@ -181,6 +196,30 @@ namespace Chaos::Private
 		return false;
 	}
 
+	bool GetIsolatedParticleSleepThresholds(
+		const FGeometryParticleHandle* Particle,
+		const TArrayCollectionArray<TSerializablePtr<FChaosPhysicsMaterial>>* PhysicsMaterials,
+		const TArrayCollectionArray<TUniquePtr<FChaosPhysicsMaterial>>* PerParticlePhysicsMaterials,
+		const THandleArray<FChaosPhysicsMaterial>* SimMaterials,
+		FRealSingle& OutSleepLinearThreshold,
+		FRealSingle& OutSleepAngularThreshold,
+		int& OutSleepCounterThreshold)
+	{
+		if (GetIslandParticleSleepThresholds(Particle, PhysicsMaterials, PerParticlePhysicsMaterials, SimMaterials, OutSleepLinearThreshold, OutSleepAngularThreshold, OutSleepCounterThreshold))
+		{
+			// Sleep thresholds are tuned to hide minor collision jitter which isn't a problem for floating particles
+			// so we scale the thesholds to make sleeping harder. E.g., to avoid going to sleep at the apex of an
+			// upwards ballistic trajectory, or when a floating oscillating body reverses rotation, etc.
+			OutSleepCounterThreshold *= FMath::Max(1, CVars::IsolatedParticleSleepCounterThresholdMultiplier);
+			OutSleepLinearThreshold *= FMath::Max(0.0f, CVars::IsolatedParticleSleepLinearThresholdMultiplier);
+			OutSleepAngularThreshold *= FMath::Max(0.0f, CVars::IsolatedParticleSleepAngularThresholdMultiplier);
+
+			return true;
+		}
+
+		return false;
+	}
+
 	bool GetParticleDisableThresholds(
 		const FGeometryParticleHandle* Particle,
 		const TArrayCollectionArray<TSerializablePtr<FChaosPhysicsMaterial>>* PhysicsMaterials,
@@ -207,6 +246,16 @@ namespace Chaos::Private
 		return false;
 	}
 
+	template<typename TRigidParticleHandle>
+	void InitParticleSleepMetrics(TRigidParticleHandle& Rigid, FReal Dt)
+	{
+		if (Dt > UE_SMALL_NUMBER)
+		{
+			Rigid.SetVSmooth(Rigid.V());
+			Rigid.SetWSmooth(Rigid.W());
+		}
+	}
+
 	template<typename TRigidParticleHandle>
 	void UpdateParticleSleepMetrics(TRigidParticleHandle& Rigid, FReal Dt)
 	{
@@ -1110,7 +1159,7 @@ namespace Chaos::Private
 			FRealSingle SleepLinearThreshold, SleepAngularThreshold;
 			FRealSingle DisableLinearThreshold, DisableAngularThreshold;
 			int32 SleepCounterThreshold;
-			GetParticleSleepThresholds(Node->Particle, PhysicsMaterials, PerParticlePhysicsMaterials, SimMaterials, SleepLinearThreshold, SleepAngularThreshold, SleepCounterThreshold);
+			GetIslandParticleSleepThresholds(Node->Particle, PhysicsMaterials, PerParticlePhysicsMaterials, SimMaterials, SleepLinearThreshold, SleepAngularThreshold, SleepCounterThreshold);
 			GetParticleDisableThresholds(Node->Particle, PhysicsMaterials, PerParticlePhysicsMaterials, SimMaterials, DisableLinearThreshold, DisableAngularThreshold);
 
 			Node->SleepLinearThresholdSq = FMath::Square(SleepLinearThreshold);
@@ -2034,12 +2083,14 @@ namespace Chaos::Private
 			{
 				FRealSingle SleepLinearThreshold, SleepAngularThreshold;
 				int32 SleepCounterThreshold;
-				GetParticleSleepThresholds(Rigid.Handle(), PhysicsMaterials, PerParticlePhysicsMaterials, SimMaterials, SleepLinearThreshold, SleepAngularThreshold, SleepCounterThreshold);
+				GetIsolatedParticleSleepThresholds(Rigid.Handle(), PhysicsMaterials, PerParticlePhysicsMaterials, SimMaterials, SleepLinearThreshold, SleepAngularThreshold, SleepCounterThreshold);
 
 				// Check for sleep
 				if ((SleepLinearThreshold > 0) || (SleepAngularThreshold > 0))
 				{
-					UpdateParticleSleepMetrics(Rigid, Dt);
+					// NOTE: We do not use smoothed velocity for isolated particles (smoothed velocity is used to hide
+					// minor collision/joint jitter and that won't be present) so we reset it for isolated particles
+					InitParticleSleepMetrics(Rigid, Dt);
 
 					int32 SleepCounter = 0;
 					if (SleepCounterThreshold < TNumericLimits<int32>::Max())
@@ -2102,7 +2153,11 @@ namespace Chaos::Private
 				|| (Rigid->WSmooth().SizeSquared() > Node->SleepAngularThresholdSq))
 			{
 				bWithinSleepThreshold = false;
-				break;
+
+				// NOTE: We will not sleep if any particle exceeds the threshold, so we could "break" here.
+				// However but we still want to update the SleepMetrics for all particles because they 
+				// currently use a moving average, so must continue to remaining particles
+				continue;
 			}
 
 			// Take the longest sleep time
@@ -2429,4 +2484,67 @@ namespace Chaos::Private
 		return true;
 	}
 
+#if CHAOS_DEBUG_DRAW
+	void FPBDIslandManager::DebugDrawSleepState(const DebugDraw::FChaosDebugDrawSettings* DebugDrawSettings) const
+	{
+		// Loop over isolated particles
+		for (FTransientPBDRigidParticleHandle& Rigid : Particles.GetActiveDynamicMovingKinematicParticlesView())
+		{
+			if (Rigid.IsDynamic() && !Rigid.IsInConstraintGraph())
+			{
+				FColor Color = FColor(128, 128, 128);
+
+				if (!Rigid.IsSleeping())
+				{
+					FRealSingle SleepLinearThreshold, SleepAngularThreshold;
+					int32 SleepCounterThreshold;
+					GetIsolatedParticleSleepThresholds(Rigid.Handle(), PhysicsMaterials, PerParticlePhysicsMaterials, SimMaterials, SleepLinearThreshold, SleepAngularThreshold, SleepCounterThreshold);
+
+					// Check for sleep
+					if ((SleepLinearThreshold > 0) || (SleepAngularThreshold > 0))
+					{
+						if (SleepCounterThreshold < TNumericLimits<int32>::Max())
+						{
+							// Isolated particles have a max sleep counter of 127 (to reduce counter space in the particle)
+							SleepCounterThreshold = FMath::Min(SleepCounterThreshold, TNumericLimits<int8>::Max());
+
+							// Did we exceed the velocity threshold?
+							const bool bIsParticlePreventingSleep = ((Rigid.VSmooth().SizeSquared() > FMath::Square(SleepLinearThreshold)) || (Rigid.WSmooth().SizeSquared() > FMath::Square(SleepAngularThreshold)));
+							Color = (bIsParticlePreventingSleep) ? FColor::Red : FColor::Green;
+						}
+					}
+				}
+
+				DebugDraw::DrawParticleShapes(FRigidTransform3::Identity, Rigid.Handle(), Color, DebugDrawSettings);
+			}
+		}
+
+		// Loop over particles in islands
+		for (FPBDIsland* Island : Islands)
+		{
+			for (FPBDIslandParticle* Node : Island->GetParticles())
+			{
+				FColor Color = FColor(128, 128, 128);
+
+				if (!Island->Flags.bIsSleeping)
+				{
+					// All zeroes means never sleep/disable
+					if ((Node->SleepLinearThresholdSq > 0) || (Node->SleepAngularThresholdSq > 0))
+					{
+						// Check the particle state against the thresholds
+						if (FPBDRigidParticleHandle* Rigid = Node->GetParticle()->CastToRigidParticle())
+						{
+							// Did we exceed the velocity threshold?
+							const bool bIsParticlePreventingSleep = ((Rigid->VSmooth().SizeSquared() > Node->SleepLinearThresholdSq) || (Rigid->WSmooth().SizeSquared() > Node->SleepAngularThresholdSq));
+							Color = (bIsParticlePreventingSleep) ? FColor::Red : FColor::Green;
+						}
+					}
+				}
+
+				DebugDraw::DrawParticleShapes(FRigidTransform3::Identity, Node->GetParticle(), Color, DebugDrawSettings);
+			}
+		}
+	}
+#endif
+
 } // namsepace Chaos::Private
diff --git a/Engine/Source/Runtime/Experimental/Chaos/Private/Chaos/PBDRigidsEvolutionGBF.cpp b/Engine/Source/Runtime/Experimental/Chaos/Private/Chaos/PBDRigidsEvolutionGBF.cpp
--- a/Engine/Source/Runtime/Experimental/Chaos/Private/Chaos/PBDRigidsEvolutionGBF.cpp
+++ b/Engine/Source/Runtime/Experimental/Chaos/Private/Chaos/PBDRigidsEvolutionGBF.cpp
@@ -51,7 +51,7 @@ namespace Chaos
 		int CollisionDisableCulledContacts = 0;
 		FAutoConsoleVariableRef CVarDisableCulledContacts(TEXT("p.CollisionDisableCulledContacts"), CollisionDisableCulledContacts, TEXT("Allow the PBDRigidsEvolutionGBF collision constraints to throw out contacts mid solve if they are culled."));
 
-		FRealSingle SmoothedPositionLerpRate = 0.1f;
+		FRealSingle SmoothedPositionLerpRate = 0.3f;
 		FAutoConsoleVariableRef CVarSmoothedPositionLerpRate(TEXT("p.Chaos.SmoothedPositionLerpRate"), SmoothedPositionLerpRate, TEXT("The interpolation rate for the smoothed position calculation. Used for sleeping."));
 
 		int DisableParticleUpdateVelocityParallelFor = 0;
@@ -1137,11 +1137,39 @@ void FPBDRigidsEvolutionGBF::OnParticleMoved(FGeometryParticleHandle* InParticle
 	}
 }
 
+void FPBDRigidsEvolutionGBF::SetParticleVelocities(FGeometryParticleHandle* InParticle, const FVec3& InV, const FVec3f& InW)
+{
+	if (FKinematicGeometryParticleHandle* Kinematic = InParticle->CastToKinematicParticle())
+	{
+		Kinematic->SetV(InV);
+		Kinematic->SetW(InW);
+	}
+
+	if (FPBDRigidParticleHandle* Rigid = InParticle->CastToRigidParticle())
+	{
+		// Reset the velocity-based sleepiness tracking properties
+		if (Rigid->IsDynamic())
+		{
+			Rigid->SetVSmooth(InV);
+			Rigid->SetWSmooth(InW);
+
+			// Wake the particle if the velocity is non-zero
+			// NOTE: We do this even when not sleeping because we want to reset the sleep accumulators
+			if (!InV.IsNearlyZero() || !InW.IsNearlyZero())
+			{
+				WakeParticle(Rigid);
+			}
+
+			// @todo(chaos): do we want to reset static friction an any existing contacts as well?
+		}
+	}
+}
+
 void FPBDRigidsEvolutionGBF::ParticleMaterialChanged(FGeometryParticleHandle* Particle)
 {
 	Particle->ParticleCollisions().VisitCollisions([this](FPBDCollisionConstraint& Collision)
 	{
-		// Reset the material - the material properties will get collected later (collision activation)
+		// Reset the material. This is a fast operation - the material properties will get collected later if the collision is activated
 		Collision.ClearMaterialProperties();
 		return ECollisionVisitorResult::Continue;
 
diff --git a/Engine/Source/Runtime/Experimental/Chaos/Private/PBDRigidsSolver.cpp b/Engine/Source/Runtime/Experimental/Chaos/Private/PBDRigidsSolver.cpp
--- a/Engine/Source/Runtime/Experimental/Chaos/Private/PBDRigidsSolver.cpp
+++ b/Engine/Source/Runtime/Experimental/Chaos/Private/PBDRigidsSolver.cpp
@@ -119,6 +119,7 @@ namespace Chaos
 		int32 ChaosSolverDebugDrawBounds = 0;
 		int32 ChaosSolverDrawTransforms = 0;
 		int32 ChaosSolverDrawIslands = 0;
+		int32 ChaosSolverDebugDrawIslandSleepState = 0;
 		int32 ChaosSolverDrawCCDInteractions = 0;
 		int32 ChaosSolverDrawCCDThresholds = 0;
 		int32 ChaosSolverDrawShapesShowStatic = 1;
@@ -145,6 +146,7 @@ namespace Chaos
 		FAutoConsoleVariableRef CVarChaosSolverDrawBounds(TEXT("p.Chaos.Solver.DebugDrawBounds"), ChaosSolverDebugDrawBounds, TEXT("Draw bounding volumes inside the broadphase (0 = never; 1 = end of frame)."));
 		FAutoConsoleVariableRef CVarChaosSolverDrawTransforms(TEXT("p.Chaos.Solver.DebugDrawTransforms"), ChaosSolverDrawTransforms, TEXT("Draw particle transforms (0 = never; 1 = end of frame)."));
 		FAutoConsoleVariableRef CVarChaosSolverDrawIslands(TEXT("p.Chaos.Solver.DebugDrawIslands"), ChaosSolverDrawIslands, TEXT("Draw solver islands (0 = never; 1 = end of frame)."));
+		FAutoConsoleVariableRef CVarChaosSolverDrawIslandSleepState(TEXT("p.Chaos.Solver.DebugDrawSleepState"), ChaosSolverDebugDrawIslandSleepState, TEXT("Draw island sleep state."));
 		FAutoConsoleVariableRef CVarChaosSolverDrawCCD(TEXT("p.Chaos.Solver.DebugDrawCCDInteractions"), ChaosSolverDrawCCDInteractions, TEXT("Draw CCD interactions."));
 		FAutoConsoleVariableRef CVarChaosSolverDrawCCDThresholds(TEXT("p.Chaos.Solver.DebugDrawCCDThresholds"), ChaosSolverDrawCCDThresholds, TEXT("Draw CCD swept thresholds."));
 		FAutoConsoleVariableRef CVarChaosSolverDrawShapesShapesStatic(TEXT("p.Chaos.Solver.DebugDraw.ShowStatics"), ChaosSolverDrawShapesShowStatic, TEXT("If DebugDrawShapes is enabled, whether to show static objects"));
@@ -2484,6 +2486,10 @@ TRACE_COUNTER_SET(ChaosTraceCounter_##Name, Value)
 		{
 			DebugDraw::DrawConstraintGraph(FRigidTransform3(), GetEvolution()->GetIslandManager(), &ChaosSolverDebugDebugDrawSettings);
 		}
+		if (ChaosSolverDebugDrawIslandSleepState == 1)
+		{
+			GetEvolution()->GetIslandManager().DebugDrawSleepState(&ChaosSolverDebugDebugDrawSettings);
+		}
 		if (ChaosSolverDrawClusterConstraints == 1)
 		{
 			DebugDraw::DrawConnectionGraph(MEvolution->GetRigidClustering(), &ChaosSolverDebugDebugDrawSettings);
@@ -2515,7 +2521,6 @@ TRACE_COUNTER_SET(ChaosTraceCounter_##Name, Value)
 		{
 			DebugDraw::DrawSuspensionConstraints(FRigidTransform3(), GetEvolution()->GetSuspensionConstraints(), &ChaosSolverDebugDebugDrawSettings);
 		}
-
 #endif
 	}
 
@@ -2661,6 +2666,7 @@ TRACE_COUNTER_SET(ChaosTraceCounter_##Name, Value)
 		Rigid->SetMaxLinearSpeedSq(DynamicMisc.MaxLinearSpeedSq());
 		Rigid->SetMaxAngularSpeedSq(DynamicMisc.MaxAngularSpeedSq());
 		Rigid->SetInitialOverlapDepenetrationVelocity(DynamicMisc.InitialOverlapDepenetrationVelocity());
+		Rigid->SetSleepThresholdMultiplier(DynamicMisc.SleepThresholdMultiplier());
 		Rigid->SetCollisionGroup(DynamicMisc.CollisionGroup());
 		Rigid->SetDisabled(DynamicMisc.Disabled());
 		Rigid->SetCollisionConstraintFlags(DynamicMisc.CollisionConstraintFlags());
diff --git a/Engine/Source/Runtime/Experimental/Chaos/Private/PhysicsProxy/ClusterUnionPhysicsProxy.cpp b/Engine/Source/Runtime/Experimental/Chaos/Private/PhysicsProxy/ClusterUnionPhysicsProxy.cpp
--- a/Engine/Source/Runtime/Experimental/Chaos/Private/PhysicsProxy/ClusterUnionPhysicsProxy.cpp
+++ b/Engine/Source/Runtime/Experimental/Chaos/Private/PhysicsProxy/ClusterUnionPhysicsProxy.cpp
@@ -408,7 +408,7 @@ namespace Chaos
 
 		if (const FParticleVelocities* NewVelocities = ParticleData.FindClusterVelocities(Manager, DataIdx))
 		{
-			Particle_Internal->SetVelocities(*NewVelocities);
+			Evolution.SetParticleVelocities(Particle_Internal, NewVelocities->V(), NewVelocities->W());
 		}
 	}
 
diff --git a/Engine/Source/Runtime/Experimental/Chaos/Private/PhysicsProxy/GeometryCollectionPhysicsProxy.cpp b/Engine/Source/Runtime/Experimental/Chaos/Private/PhysicsProxy/GeometryCollectionPhysicsProxy.cpp
--- a/Engine/Source/Runtime/Experimental/Chaos/Private/PhysicsProxy/GeometryCollectionPhysicsProxy.cpp
+++ b/Engine/Source/Runtime/Experimental/Chaos/Private/PhysicsProxy/GeometryCollectionPhysicsProxy.cpp
@@ -1670,6 +1670,7 @@ void FGeometryCollectionPhysicsProxy::InitializeBodiesPT(Chaos::FPBDRigidsSolver
 				Handle->SetLinearEtherDrag(Parameters.LinearDamping);
 				Handle->SetAngularEtherDrag(Parameters.AngularDamping);
 				Handle->SetInitialOverlapDepenetrationVelocity(Parameters.InitialOverlapDepenetrationVelocity);
+				Handle->SetSleepThresholdMultiplier(Parameters.SleepThresholdMultiplier);
 			}
 
 #if !(UE_BUILD_SHIPPING || UE_BUILD_TEST)
diff --git a/Engine/Source/Runtime/Experimental/Chaos/Private/PhysicsProxy/SingleParticlePhysicsProxy.cpp b/Engine/Source/Runtime/Experimental/Chaos/Private/PhysicsProxy/SingleParticlePhysicsProxy.cpp
--- a/Engine/Source/Runtime/Experimental/Chaos/Private/PhysicsProxy/SingleParticlePhysicsProxy.cpp
+++ b/Engine/Source/Runtime/Experimental/Chaos/Private/PhysicsProxy/SingleParticlePhysicsProxy.cpp
@@ -90,7 +90,7 @@ void PushToPhysicsStateImp(const Chaos::FDirtyPropertiesManager& Manager, Chaos:
 		auto NewVelocities = bHasKinematicData ? ParticleData.FindVelocities(Manager, DataIdx) : nullptr;
 		if(NewVelocities)
 		{
-			KinematicHandle->SetVelocities(*NewVelocities);
+			Evolution.SetParticleVelocities(KinematicHandle, NewVelocities->V(), NewVelocities->W());
 		}
 
 		auto NewKinematicTargetGT = bHasKinematicData ? ParticleData.FindKinematicTarget(Manager, DataIdx) : nullptr;
diff --git a/Engine/Source/Runtime/Experimental/Chaos/Private/RewindData.cpp b/Engine/Source/Runtime/Experimental/Chaos/Private/RewindData.cpp
--- a/Engine/Source/Runtime/Experimental/Chaos/Private/RewindData.cpp
+++ b/Engine/Source/Runtime/Experimental/Chaos/Private/RewindData.cpp
@@ -208,7 +208,11 @@ void FRewindData::ApplyTargets(const int32 Frame, const bool bResetSimulation)
 
 		const bool bResimAsFollower = DirtyParticleInfo.bResimAsFollower;
 		RewindHelper(PTParticle, bResimAsFollower, History.TargetPositions, [](auto Particle, const auto& Data) {Particle->SetXR(Data); });
-		RewindHelper(PTParticle->CastToKinematicParticle(), bResimAsFollower, History.TargetVelocities, [](auto Particle, const auto& Data) {Particle->SetVelocities(Data); });
+		RewindHelper(PTParticle->CastToKinematicParticle(), bResimAsFollower, History.TargetVelocities, [](auto Particle, const auto& Data)
+		{
+			Particle->SetV(Data.V());
+			Particle->SetW(Data.W());
+		});
 		RewindHelper(PTParticle->CastToRigidParticle(), bResimAsFollower, History.TargetStates, [this](auto Particle, const auto& Data) 
 		{ 
 			if (Particle == nullptr || Solver->GetEvolution() == nullptr)
@@ -332,7 +336,7 @@ bool FRewindData::RewindToFrame(int32 Frame)
 		const bool bResimAsFollower = DirtyParticleInfo.bResimAsFollower;
 
 		bool bAnyChange = RewindHelper(PTParticle, bResimAsFollower, History.ParticlePositionRotation, [](auto Particle, const auto& Data) {Particle->SetXR(Data); });
-		bAnyChange |= RewindHelper(PTParticle->CastToKinematicParticle(), bResimAsFollower, History.Velocities, [](auto Particle, const auto& Data) {Particle->SetVelocities(Data); });
+		bAnyChange |= RewindHelper(PTParticle->CastToKinematicParticle(), bResimAsFollower, History.Velocities, [](auto Particle, const auto& Data) { Particle->SetV(Data.V()); Particle->SetW(Data.W()); });
 		bAnyChange |= RewindHelper(PTParticle, bResimAsFollower, History.NonFrequentData, [this](auto Particle, const auto& Data)
 			{
 				Solver->GetEvolution()->InvalidateParticle(Particle); // Clear collision/constraints before updating NonFrequentData
diff --git a/Engine/Source/Runtime/Experimental/Chaos/Public/Chaos/Island/IslandManager.h b/Engine/Source/Runtime/Experimental/Chaos/Public/Chaos/Island/IslandManager.h
--- a/Engine/Source/Runtime/Experimental/Chaos/Public/Chaos/Island/IslandManager.h
+++ b/Engine/Source/Runtime/Experimental/Chaos/Public/Chaos/Island/IslandManager.h
@@ -13,6 +13,11 @@ namespace Chaos
 	class FPerParticleGravity;
 }
 
+namespace Chaos::DebugDraw
+{
+	struct FChaosDebugDrawSettings;
+}
+
 namespace Chaos::Private
 {
 	class FPBDIslandMergeSet;
@@ -463,6 +468,11 @@ namespace Chaos::Private
 		CHAOS_API TArray<const FConstraintHandle*> FindConstraintsInIslands(const TArray<const FPBDIsland*> Islands, int32 ContainerId) const;
 		CHAOS_API void SetParticleIslandIsSleeping(FGeometryParticleHandle* Particle, const bool bInIsSleeping);
 
+		// Debug draw
+#if CHAOS_DEBUG_DRAW
+		CHAOS_API void DebugDrawSleepState(const DebugDraw::FChaosDebugDrawSettings* DebugDrawSettings) const;
+#endif
+
 		// Deprecated API
 		UE_DEPRECATED(5.3, "Use Reset") void RemoveConstraints() { Reset(); }
 		UE_DEPRECATED(5.3, "Use UpdateParticles") void InitializeGraph(const TParticleView<FPBDRigidParticles>& InParticles) { UpdateParticles(); }
diff --git a/Engine/Source/Runtime/Experimental/Chaos/Public/Chaos/PBDRigidsEvolutionGBF.h b/Engine/Source/Runtime/Experimental/Chaos/Public/Chaos/PBDRigidsEvolutionGBF.h
--- a/Engine/Source/Runtime/Experimental/Chaos/Public/Chaos/PBDRigidsEvolutionGBF.h
+++ b/Engine/Source/Runtime/Experimental/Chaos/Public/Chaos/PBDRigidsEvolutionGBF.h
@@ -176,6 +176,11 @@ namespace Chaos
 		 */
 		CHAOS_API void OnParticleMoved(FGeometryParticleHandle* InParticle, const FVec3& PrevX, const FRotation3& PrevR, const bool bIsTeleport);
 
+		/**
+		 * User has changed particle velocity or angular velocity
+		 */
+		CHAOS_API void SetParticleVelocities(FGeometryParticleHandle* InParticle, const FVec3& InV, const FVec3f& InW);
+
 
 		/**
 		 * Reload the particles cache for all particles where appropriate
diff --git a/Engine/Source/Runtime/Experimental/Chaos/Public/Chaos/ParticleDirtyFlags.h b/Engine/Source/Runtime/Experimental/Chaos/Public/Chaos/ParticleDirtyFlags.h
--- a/Engine/Source/Runtime/Experimental/Chaos/Public/Chaos/ParticleDirtyFlags.h
+++ b/Engine/Source/Runtime/Experimental/Chaos/Public/Chaos/ParticleDirtyFlags.h
@@ -455,10 +455,12 @@ public:
 		}
 
 		// @todo(chaos): add this
-		//Ar << MInitialOverlapDepenetrationVelocity
+		//Ar << MInitialOverlapDepenetrationVelocity;
+		//Ar << MSleepThresholdMultiplier;
 		if (Ar.IsLoading())
 		{
 			MInitialOverlapDepenetrationVelocity = 0;
+			MSleepThresholdMultiplier = 1.0f;
 		}
 
 		if (bAddControlFlags)
@@ -475,6 +477,7 @@ public:
 		SetMaxLinearSpeedSq(Other.MaxLinearSpeedSq());
 		SetMaxAngularSpeedSq(Other.MaxAngularSpeedSq());
 		SetInitialOverlapDepenetrationVelocity(Other.InitialOverlapDepenetrationVelocity());
+		SetSleepThresholdMultiplier(Other.SleepThresholdMultiplier());
 		SetObjectState(Other.ObjectState());
 		SetCollisionGroup(Other.CollisionGroup());
 		SetSleepType(Other.SleepType());
@@ -492,6 +495,7 @@ public:
 			&& MMaxLinearSpeedSq == FRealSingle(Other.MaxLinearSpeedSq())
 			&& MMaxAngularSpeedSq == FRealSingle(Other.MaxAngularSpeedSq())
 			&& InitialOverlapDepenetrationVelocity() == Other.InitialOverlapDepenetrationVelocity()
+			&& SleepThresholdMultiplier() == Other.SleepThresholdMultiplier()
 			&& CollisionGroup() == Other.CollisionGroup()
 			&& SleepType() == Other.SleepType()
 			&& CollisionConstraintFlags() == Other.CollisionConstraintFlags()
@@ -519,6 +523,9 @@ public:
 	FRealSingle InitialOverlapDepenetrationVelocity() const { return MInitialOverlapDepenetrationVelocity; }
 	void SetInitialOverlapDepenetrationVelocity(FRealSingle InVel) { MInitialOverlapDepenetrationVelocity = InVel; }
 
+	FRealSingle SleepThresholdMultiplier() const { return MSleepThresholdMultiplier; }
+	void SetSleepThresholdMultiplier(FRealSingle InSleepThresholdMultiplier) { MSleepThresholdMultiplier = InSleepThresholdMultiplier; }
+
 	EObjectStateType ObjectState() const { return MObjectState; }
 	void SetObjectState(EObjectStateType InState){ MObjectState = InState; }
 
@@ -563,6 +570,7 @@ private:
 	FRealSingle MMaxLinearSpeedSq;
 	FRealSingle MMaxAngularSpeedSq;
 	FRealSingle MInitialOverlapDepenetrationVelocity = 0;
+	FRealSingle MSleepThresholdMultiplier = 1;
 	int32 MCollisionGroup;
 
 	EObjectStateType MObjectState;
diff --git a/Engine/Source/Runtime/Experimental/Chaos/Public/Chaos/ParticleHandle.h b/Engine/Source/Runtime/Experimental/Chaos/Public/Chaos/ParticleHandle.h
--- a/Engine/Source/Runtime/Experimental/Chaos/Public/Chaos/ParticleHandle.h
+++ b/Engine/Source/Runtime/Experimental/Chaos/Public/Chaos/ParticleHandle.h
@@ -113,6 +113,7 @@ void PBDRigidParticleDefaultConstruct(FConcrete& Concrete, const FPBDRigidPartic
 	Concrete.SetMaxLinearSpeedSq(TNumericLimits<T>::Max());
 	Concrete.SetMaxAngularSpeedSq(TNumericLimits<T>::Max());
 	Concrete.SetInitialOverlapDepenetrationVelocity(-1.0f);
+	Concrete.SetSleepThresholdMultiplier(1.0f);
 	Concrete.SetM(1);
 	Concrete.SetInvM(1);
 	Concrete.SetCenterOfMass(TVector<T,d>(0));
@@ -903,6 +904,7 @@ public:
 	TVector<T, d>& W() { return KinematicGeometryParticles->W(ParticleIdx); }
 	void SetW(const TVector<T, d>& InW, bool bInvalidate = false) { KinematicGeometryParticles->W(ParticleIdx) = InW; }
 
+	UE_DEPRECATED(5.4, "Use FPBDRigidsEvolutionGBF::SetParticleVelocitied or (SetV, SetW if that is not appropriate)")
 	void SetVelocities(const FParticleVelocities& Velocities)
 	{
 		SetV(Velocities.V());
@@ -1175,6 +1177,9 @@ public:
 	FRealSingle InitialOverlapDepenetrationVelocity() const { return PBDRigidParticles->InitialOverlapDepenetrationVelocity(ParticleIdx); }
 	void SetInitialOverlapDepenetrationVelocity(FRealSingle InVel) { PBDRigidParticles->InitialOverlapDepenetrationVelocity(ParticleIdx) = InVel; }
 
+	inline FRealSingle SleepThresholdMultiplier() const { return PBDRigidParticles->SleepThresholdMultiplier(ParticleIdx); }
+	inline void SetSleepThresholdMultiplier(FRealSingle Multiplier) { PBDRigidParticles->SleepThresholdMultiplier(ParticleIdx) = Multiplier; }
+
 	EObjectStateType ObjectState() const { return PBDRigidParticles->ObjectState(ParticleIdx); }
 	EObjectStateType PreObjectState() const { return PBDRigidParticles->PreObjectState(ParticleIdx); }
 
@@ -2280,6 +2285,15 @@ public:
 		return 0;
 	}
 
+	FRealSingle SleepThresholdMultiplier() const
+	{
+		if (auto RigidHandle = MHandle->CastToRigidParticle())
+		{
+			return RigidHandle->SleepThresholdMultiplier();
+		}
+		return 1.0f;
+	}
+
 	bool OneWayInteraction() const
 	{
 		if (auto RigidHandle = MHandle->CastToRigidParticle())
@@ -3400,7 +3414,12 @@ public:
 	void SetInitialOverlapDepenetrationVelocity(FRealSingle InVel)
 	{
 		MMiscData.Modify(true, MDirtyFlags, Proxy, [InVel](auto& Data) { Data.SetInitialOverlapDepenetrationVelocity(InVel); });
-		
+	}
+
+	FRealSingle SleepThresholdMultiplier() const { return MMiscData.Read().SleepThresholdMultiplier(); }
+	void SetSleepThresholdMultiplier(FRealSingle Multiplier)
+	{
+		MMiscData.Modify(true, MDirtyFlags, Proxy, [Multiplier](auto& Data) { Data.SetSleepThresholdMultiplier(Multiplier); });
 	}
 
 	int32 Island() const { return MIsland; }
diff --git a/Engine/Source/Runtime/Experimental/Chaos/Public/Chaos/RigidParticles.h b/Engine/Source/Runtime/Experimental/Chaos/Public/Chaos/RigidParticles.h
--- a/Engine/Source/Runtime/Experimental/Chaos/Public/Chaos/RigidParticles.h
+++ b/Engine/Source/Runtime/Experimental/Chaos/Public/Chaos/RigidParticles.h
@@ -108,6 +108,7 @@ public:
 		, MaxLinearSpeedsSq(MoveTemp(Other.MaxLinearSpeedsSq))
 		, MaxAngularSpeedsSq(MoveTemp(Other.MaxAngularSpeedsSq))
 		, MInitialOverlapDepenetrationVelocity(MoveTemp(Other.MInitialOverlapDepenetrationVelocity))
+		, MSleepThresholdMultiplier(MoveTemp(Other.MSleepThresholdMultiplier))
 		, MCollisionParticles(MoveTemp(Other.MCollisionParticles))
 		, MSleepType(MoveTemp(Other.MSleepType))
 		, MSleepCounter(MoveTemp(Other.MSleepCounter))
@@ -138,6 +139,7 @@ public:
 		TArrayCollection::AddArray(&MaxLinearSpeedsSq);
 		TArrayCollection::AddArray(&MaxAngularSpeedsSq);
 		TArrayCollection::AddArray(&MInitialOverlapDepenetrationVelocity);
+		TArrayCollection::AddArray(&MSleepThresholdMultiplier);
 		TArrayCollection::AddArray(&MCollisionParticles);
 		TArrayCollection::AddArray(&MSleepType);
 		TArrayCollection::AddArray(&MSleepCounter);
@@ -202,6 +204,9 @@ public:
 	FORCEINLINE const FRealSingle& InitialOverlapDepenetrationVelocity(const int32 index) const { return MInitialOverlapDepenetrationVelocity[index]; }
 	FORCEINLINE FRealSingle& InitialOverlapDepenetrationVelocity(const int32 index) { return MInitialOverlapDepenetrationVelocity[index]; }
 
+	FORCEINLINE const FRealSingle& SleepThresholdMultiplier(const int32 Index) const { return MSleepThresholdMultiplier[Index]; }
+	FORCEINLINE FRealSingle& SleepThresholdMultiplier(const int32 Index) { return MSleepThresholdMultiplier[Index]; }
+
 	FORCEINLINE int32 CollisionParticlesSize(int32 Index) const { return MCollisionParticles[Index] == nullptr ? 0 : MCollisionParticles[Index]->Size(); }
 
 	void CollisionParticlesInitIfNeeded(const int32 Index);
@@ -413,6 +418,7 @@ private:
 	TArrayCollectionArray<T> MaxLinearSpeedsSq;
 	TArrayCollectionArray<T> MaxAngularSpeedsSq;
 	TArrayCollectionArray<FRealSingle> MInitialOverlapDepenetrationVelocity;
+	TArrayCollectionArray<FRealSingle> MSleepThresholdMultiplier;
 	TArrayCollectionArray<TUniquePtr<TBVHParticles<T, d>>> MCollisionParticles;
 	TArrayCollectionArray<ESleepType> MSleepType;
 	TArrayCollectionArray<int8> MSleepCounter;
diff --git a/Engine/Source/Runtime/Experimental/Chaos/Public/GeometryCollection/GeometryCollectionSimulationCoreTypes.h b/Engine/Source/Runtime/Experimental/Chaos/Public/GeometryCollection/GeometryCollectionSimulationCoreTypes.h
--- a/Engine/Source/Runtime/Experimental/Chaos/Public/GeometryCollection/GeometryCollectionSimulationCoreTypes.h
+++ b/Engine/Source/Runtime/Experimental/Chaos/Public/GeometryCollection/GeometryCollectionSimulationCoreTypes.h
@@ -217,6 +217,7 @@ struct FSimulationParameters
 		, LinearDamping(0.01f)
 		, AngularDamping(0)
 		, InitialOverlapDepenetrationVelocity(-1.0f)
+		, SleepThresholdMultiplier(1.0f)
 		, bUseDamagePropagation(false)
 		, BreakDamagePropagationFactor(1.0f)
 		, ShockDamagePropagationFactor(0.0f)
@@ -279,6 +280,7 @@ struct FSimulationParameters
 		, LinearDamping(Other.LinearDamping)
 		, AngularDamping(Other.AngularDamping)
 		, InitialOverlapDepenetrationVelocity(Other.InitialOverlapDepenetrationVelocity)
+		, SleepThresholdMultiplier(Other.SleepThresholdMultiplier)
 		, bUseDamagePropagation(Other.bUseDamagePropagation)
 		, BreakDamagePropagationFactor(Other.BreakDamagePropagationFactor)
 		, ShockDamagePropagationFactor(Other.ShockDamagePropagationFactor)
@@ -371,6 +373,7 @@ struct FSimulationParameters
 	float LinearDamping;
 	float AngularDamping;
 	float InitialOverlapDepenetrationVelocity;
+	float SleepThresholdMultiplier;
 
 	bool bUseDamagePropagation;
 	float BreakDamagePropagationFactor;
diff --git a/Engine/Source/Runtime/Experimental/Chaos/Public/PhysicsProxy/SingleParticlePhysicsProxy.h b/Engine/Source/Runtime/Experimental/Chaos/Public/PhysicsProxy/SingleParticlePhysicsProxy.h
--- a/Engine/Source/Runtime/Experimental/Chaos/Public/PhysicsProxy/SingleParticlePhysicsProxy.h
+++ b/Engine/Source/Runtime/Experimental/Chaos/Public/PhysicsProxy/SingleParticlePhysicsProxy.h
@@ -1279,7 +1279,7 @@ public:
 	FRealSingle GetInitialOverlapDepenetrationVelocity()
 	{
 		VerifyContext();
-		if (const TPBDRigidParticle<FReal, 3>*Rigid = GetParticle_LowLevel()->CastToRigidParticle())
+		if (const TPBDRigidParticle<FReal, 3>* Rigid = GetParticle_LowLevel()->CastToRigidParticle())
 		{
 			return Rigid->InitialOverlapDepenetrationVelocity();
 		}
@@ -1290,12 +1290,21 @@ public:
 	void SetInitialOverlapDepenetrationVelocity(FRealSingle InNewSpeed)
 	{
 		VerifyContext();
-		if (TPBDRigidParticle<FReal, 3>*Rigid = GetParticle_LowLevel()->CastToRigidParticle())
+		if (TPBDRigidParticle<FReal, 3>* Rigid = GetParticle_LowLevel()->CastToRigidParticle())
 		{
 			Rigid->SetInitialOverlapDepenetrationVelocity(InNewSpeed);
 		}
 	}
 
+	void SetSleepThresholdMultiplier(FRealSingle Multiplier)
+	{
+		VerifyContext();
+		if (TPBDRigidParticle<FReal, 3>* Rigid = GetParticle_LowLevel()->CastToRigidParticle())
+		{
+			Rigid->SetSleepThresholdMultiplier(Multiplier);
+		}
+	}
+
 	void SetDisabled(bool bDisable)
 	{
 		VerifyContext();
diff --git a/Engine/Source/Runtime/Experimental/Chaos/Public/RewindData.h b/Engine/Source/Runtime/Experimental/Chaos/Public/RewindData.h
--- a/Engine/Source/Runtime/Experimental/Chaos/Public/RewindData.h
+++ b/Engine/Source/Runtime/Experimental/Chaos/Public/RewindData.h
@@ -1000,6 +1000,7 @@ public:
 	REWIND_PARTICLE_RIGID_PROPERTY(DynamicsMisc, MaxLinearSpeedSq)
 	REWIND_PARTICLE_RIGID_PROPERTY(DynamicsMisc, MaxAngularSpeedSq)
 	REWIND_PARTICLE_RIGID_PROPERTY(DynamicsMisc, InitialOverlapDepenetrationVelocity)
+	REWIND_PARTICLE_RIGID_PROPERTY(DynamicsMisc, SleepThresholdMultiplier)
 	REWIND_PARTICLE_RIGID_PROPERTY(DynamicsMisc, ObjectState)
 	REWIND_PARTICLE_RIGID_PROPERTY(DynamicsMisc, CollisionGroup)
 	REWIND_PARTICLE_RIGID_PROPERTY(DynamicsMisc, ControlFlags)
@@ -1063,6 +1064,7 @@ public:
 			REWIND_PARTICLE_TO_STR(MaxLinearSpeedSq)
 			REWIND_PARTICLE_TO_STR(MaxAngularSpeedSq)
 			REWIND_PARTICLE_TO_STR(InitialOverlapDepenetrationVelocity)
+			REWIND_PARTICLE_TO_STR(SleepThresholdMultiplier)
 
 			REWIND_PARTICLE_TO_STR(ObjectState)
 			REWIND_PARTICLE_TO_STR(CollisionGroup)
diff --git a/Engine/Source/Runtime/Experimental/ChaosCore/Public/Chaos/Core.h b/Engine/Source/Runtime/Experimental/ChaosCore/Public/Chaos/Core.h
--- a/Engine/Source/Runtime/Experimental/ChaosCore/Public/Chaos/Core.h
+++ b/Engine/Source/Runtime/Experimental/ChaosCore/Public/Chaos/Core.h
@@ -51,7 +51,7 @@ namespace Chaos
 	// NOTE: if you get a merge conflict on the GUID, you must replace it with a new GUID - do not accept the source or target
 	// or you will likely get DDC version conflicts resulting in crashes during load.
 	// Core version string for Chaos data. Any DDC builder dependent on Chaos for serialization should depend on this version
-	inline const TCHAR* const ChaosVersionGUID = TEXT("466C9323-59EB-4808-82F3-5A9E0F6EAC51");
+	inline const TCHAR* const ChaosVersionGUID = TEXT("5F81677D-8347-4E79-85D0-9F71C9A7DF03");
 
 	inline FString GetChaosVersionStringInner()
 	{
diff --git a/Engine/Source/Runtime/Experimental/ChaosVisualDebugger/Public/DataWrappers/ChaosVDParticleDataWrapper.h b/Engine/Source/Runtime/Experimental/ChaosVisualDebugger/Public/DataWrappers/ChaosVDParticleDataWrapper.h
--- a/Engine/Source/Runtime/Experimental/ChaosVisualDebugger/Public/DataWrappers/ChaosVDParticleDataWrapper.h
+++ b/Engine/Source/Runtime/Experimental/ChaosVisualDebugger/Public/DataWrappers/ChaosVDParticleDataWrapper.h
@@ -321,7 +321,7 @@ struct FChaosVDParticleDynamicMisc : public FChaosVDParticleDataBase
 {
 	GENERATED_BODY()
 
-	FChaosVDParticleDynamicMisc(): MAngularEtherDrag(0), MMaxLinearSpeedSq(0), MMaxAngularSpeedSq(0), MInitialOverlapDepenetrationVelocity(0),
+	FChaosVDParticleDynamicMisc(): MAngularEtherDrag(0), MMaxLinearSpeedSq(0), MMaxAngularSpeedSq(0), MInitialOverlapDepenetrationVelocity(0), MSleepThresholdMultiplier(1),
 	                               MCollisionGroup(0), MObjectState(), MSleepType(), bDisabled(false)
 	{
 	}
@@ -336,6 +336,7 @@ struct FChaosVDParticleDynamicMisc : public FChaosVDParticleDataBase
 		MMaxLinearSpeedSq = Other.MaxLinearSpeedSq();
 		MMaxAngularSpeedSq = Other.MaxAngularSpeedSq();
 		MInitialOverlapDepenetrationVelocity = Other.InitialOverlapDepenetrationVelocity();
+		MSleepThresholdMultiplier = Other.SleepThresholdMultiplier();
 		MObjectState = static_cast<EChaosVDObjectStateType>(Other.ObjectState());
 		MCollisionGroup = Other.CollisionGroup();
 		MSleepType =  static_cast<EChaosVDSleepType>(Other.SleepType());
@@ -360,6 +361,9 @@ struct FChaosVDParticleDynamicMisc : public FChaosVDParticleDataBase
 	UPROPERTY(EditAnywhere, Category = "Particle Dynamic Misc")
 	float MInitialOverlapDepenetrationVelocity;
 
+	UPROPERTY(EditAnywhere, Category = "Particle Dynamic Misc")
+	float MSleepThresholdMultiplier;
+
 	UPROPERTY(EditAnywhere, Category= "Particle Dynamic Misc")
 	int32 MCollisionGroup;
 
diff --git a/Engine/Source/Runtime/Experimental/GeometryCollectionEngine/Private/GeometryCollection/GeometryCollectionComponent.cpp b/Engine/Source/Runtime/Experimental/GeometryCollectionEngine/Private/GeometryCollection/GeometryCollectionComponent.cpp
--- a/Engine/Source/Runtime/Experimental/GeometryCollectionEngine/Private/GeometryCollection/GeometryCollectionComponent.cpp
+++ b/Engine/Source/Runtime/Experimental/GeometryCollectionEngine/Private/GeometryCollection/GeometryCollectionComponent.cpp
@@ -3828,6 +3828,7 @@ void UGeometryCollectionComponent::RegisterAndInitializePhysicsProxy()
 		SimulationParameters.LinearDamping = BodyInstance.LinearDamping;
 		SimulationParameters.AngularDamping = BodyInstance.AngularDamping;
 		SimulationParameters.InitialOverlapDepenetrationVelocity = BodyInstance.GetMaxDepenetrationVelocity();
+		SimulationParameters.SleepThresholdMultiplier = BodyInstance.GetSleepThresholdMultiplier();
 		SimulationParameters.bUseDamagePropagation = DamagePropagationData.bEnabled;
 		SimulationParameters.BreakDamagePropagationFactor = DamagePropagationData.BreakDamagePropagationFactor;
 		SimulationParameters.ShockDamagePropagationFactor = DamagePropagationData.ShockDamagePropagationFactor;
diff --git a/Engine/Source/Runtime/PhysicsCore/Private/ChaosEngineInterface.cpp b/Engine/Source/Runtime/PhysicsCore/Private/ChaosEngineInterface.cpp
--- a/Engine/Source/Runtime/PhysicsCore/Private/ChaosEngineInterface.cpp
+++ b/Engine/Source/Runtime/PhysicsCore/Private/ChaosEngineInterface.cpp
@@ -1161,6 +1161,11 @@ void FChaosEngineInterface::SetSleepEnergyThreshold_AssumesLocked(const FPhysics
 {
 }
 
+void FChaosEngineInterface::SetSleepThresholdMultiplier_AssumesLocked(const FPhysicsActorHandle& InActorReference, float InThresholdMultiplier)
+{
+	return InActorReference->GetGameThreadAPI().SetSleepThresholdMultiplier(InThresholdMultiplier);
+}
+
 void FChaosEngineInterface::SetMass_AssumesLocked(FPhysicsActorHandle& InActorReference,float InMass)
 {
 	Chaos::FRigidBodyHandle_External& Body_External = InActorReference->GetGameThreadAPI();
diff --git a/Engine/Source/Runtime/PhysicsCore/Public/Chaos/ChaosEngineInterface.h b/Engine/Source/Runtime/PhysicsCore/Public/Chaos/ChaosEngineInterface.h
--- a/Engine/Source/Runtime/PhysicsCore/Public/Chaos/ChaosEngineInterface.h
+++ b/Engine/Source/Runtime/PhysicsCore/Public/Chaos/ChaosEngineInterface.h
@@ -463,7 +463,9 @@ public:
 	static PHYSICSCORE_API void SetOneWayInteraction_AssumesLocked(const FPhysicsActorHandle& InHandle, bool InOneWayInteraction);
 
 	static PHYSICSCORE_API float GetSleepEnergyThreshold_AssumesLocked(const FPhysicsActorHandle& InActorReference);
-	static PHYSICSCORE_API void SetSleepEnergyThreshold_AssumesLocked(const FPhysicsActorHandle& InActorReference,float InEnergyThreshold);
+	static PHYSICSCORE_API void SetSleepEnergyThreshold_AssumesLocked(const FPhysicsActorHandle& InActorReference, float InEnergyThreshold);
+
+	static PHYSICSCORE_API void SetSleepThresholdMultiplier_AssumesLocked(const FPhysicsActorHandle& InActorReference, float ThresholdMultiplier);
 
 	static PHYSICSCORE_API void SetMass_AssumesLocked(FPhysicsActorHandle& InHandle,float InMass);
 	static PHYSICSCORE_API void SetMassSpaceInertiaTensor_AssumesLocked(FPhysicsActorHandle& InHandle,const FVector& InTensor);
