Constraints: fixed translation issues
    - note that this needs a bit more investigation

diff --git a/Engine/Source/Runtime/Experimental/Animation/Constraints/Private/TransformConstraint.cpp b/Engine/Source/Runtime/Experimental/Animation/Constraints/Private/TransformConstraint.cpp
--- a/Engine/Source/Runtime/Experimental/Animation/Constraints/Private/TransformConstraint.cpp
+++ b/Engine/Source/Runtime/Experimental/Animation/Constraints/Private/TransformConstraint.cpp
@@ -755,15 +755,12 @@ FConstraintTickFunction::ConstraintFunction UTickableTranslationConstraint::GetF
 
 		AxisFilter.FilterVector(NewTranslation, ChildTranslation);
 
-		if (!NewTranslation.Equals(ChildTranslation))
-		{
-			Transform.SetLocation(NewTranslation);
+		Transform.SetLocation(NewTranslation);
 			
-			SetChildGlobalTransform(Transform);
-			if (ChildTRSHandle)
-			{
-				ChildTRSHandle->TickTarget();
-			}
+		SetChildGlobalTransform(Transform);
+		if (ChildTRSHandle)
+		{
+			ChildTRSHandle->TickTarget();
 		}
 	};
 }
@@ -899,16 +896,14 @@ FConstraintTickFunction::ConstraintFunction UTickableRotationConstraint::GetFunc
 
 		AxisFilter.FilterQuat(NewRotation, ChildRotation);
 
-		if (!NewRotation.Equals(ChildRotation))
-		{
-			Transform.SetRotation(NewRotation);
+
+		Transform.SetRotation(NewRotation);
 		
-			SetChildGlobalTransform(Transform);
+		SetChildGlobalTransform(Transform);
 
-			if (ChildTRSHandle)
-			{
-				ChildTRSHandle->TickTarget();
-			}
+		if (ChildTRSHandle)
+		{
+			ChildTRSHandle->TickTarget();
 		}
 	};
 }
@@ -1048,15 +1043,12 @@ FConstraintTickFunction::ConstraintFunction UTickableScaleConstraint::GetFunctio
 
 		AxisFilter.FilterVector(NewScale, ChildScale);
 
-		if (!NewScale.Equals(ChildScale))
+		Transform.SetScale3D(NewScale);
+		
+		SetChildGlobalTransform(Transform);
+		if (ChildTRSHandle)
 		{
-			Transform.SetScale3D(NewScale);
-			
-			SetChildGlobalTransform(Transform);
-			if (ChildTRSHandle)
-			{
-				ChildTRSHandle->TickTarget();
-			}
+			ChildTRSHandle->TickTarget();
 		}
 	};
 }
@@ -1217,13 +1209,10 @@ FConstraintTickFunction::ConstraintFunction UTickableParentConstraint::GetFuncti
 			TargetTransform.SetScale3D(ChildGlobalTransform.GetScale3D());
 		}
 
-		if (!ChildGlobalTransform.Equals(TargetTransform))
+		SetChildGlobalTransform(TargetTransform);
+		if (ChildTRSHandle)
 		{
-			SetChildGlobalTransform(TargetTransform);
-			if (ChildTRSHandle)
-			{
-				ChildTRSHandle->TickTarget();
-			}
+			ChildTRSHandle->TickTarget();
 		}
 	};
 }
