Chaos - fix comment in IslandManager.cpp

diff --git a/Engine/Source/Runtime/Experimental/Chaos/Private/Chaos/Island/IslandManager.cpp b/Engine/Source/Runtime/Experimental/Chaos/Private/Chaos/Island/IslandManager.cpp
--- a/Engine/Source/Runtime/Experimental/Chaos/Private/Chaos/Island/IslandManager.cpp
+++ b/Engine/Source/Runtime/Experimental/Chaos/Private/Chaos/Island/IslandManager.cpp
@@ -2155,8 +2155,8 @@ namespace Chaos::Private
 				bWithinSleepThreshold = false;
 
 				// NOTE: We will not sleep if any particle exceeds the threshold, so we could "break" here.
-				// However but we still want to update the SleepMetrics for all particles because they 
-				// currently use a moving average, so must continue to remaining particles
+				// However we still want to update the SleepMetrics for all particles because we want to
+				// update the smoothed velocity based on current state, so we must continue to remaining particles
 				continue;
 			}
 
