Iris - Fix issue with subobjects not being destroyed as soon as possible. Validate object isn't pending destroy when marking dirty.

diff --git a/Engine/Plugins/Runtime/ReplicationSystemTestPlugin/Source/Private/Tests/ReplicationSystem/TestCancelPendingDestroy.cpp b/Engine/Plugins/Runtime/ReplicationSystemTestPlugin/Source/Private/Tests/ReplicationSystem/TestCancelPendingDestroy.cpp
--- a/Engine/Plugins/Runtime/ReplicationSystemTestPlugin/Source/Private/Tests/ReplicationSystem/TestCancelPendingDestroy.cpp
+++ b/Engine/Plugins/Runtime/ReplicationSystemTestPlugin/Source/Private/Tests/ReplicationSystem/TestCancelPendingDestroy.cpp
@@ -739,6 +739,47 @@ UE_NET_TEST_FIXTURE(FTestCancelPendingDestroyFixture, TestCancelingSubObjectDest
 	UE_NET_ASSERT_EQ(Client->GetReplicationBridge()->GetReplicatedObject(ServerSubObject2->NetRefHandle), nullptr);
 }
 
+UE_NET_TEST_FIXTURE(FTestCancelPendingDestroyFixture, TestCancelRootObjectDestroyAfterSubObjectDestroy)
+{
+	// Add a client
+	FReplicationSystemTestClient* Client = CreateClient();
+
+	// Spawn object on server
+	UTestReplicatedIrisObject* ServerObject = Server->CreateObject(UTestReplicatedIrisObject::FComponents());
+	UTestReplicatedIrisObject* ServerSubObject = Server->CreateSubObject(ServerObject->NetRefHandle, UTestReplicatedIrisObject::FComponents());
+
+	// Write and send packet
+	Server->PreSendUpdate();
+	Server->SendAndDeliverTo(Client, DeliverPacket);
+	Server->PostSendUpdate();
+
+	// Verify that the subobject now also exists on the client
+	UE_NET_ASSERT_NE(Client->GetReplicationBridge()->GetReplicatedObject(ServerSubObject->NetRefHandle), nullptr);
+
+	// Filter out subobject to cause a SubObjectPendingDestroy
+	Server->GetReplicationSystem()->AddToGroup(Server->GetReplicationSystem()->GetNotReplicatedNetObjectGroup(), ServerSubObject->NetRefHandle);
+	Server->PreSendUpdate();
+	Server->PostSendUpdate();
+
+	// Filter out root object to cause subobject to move to PendingDestroy instead
+	Server->GetReplicationSystem()->AddToGroup(Server->GetReplicationSystem()->GetNotReplicatedNetObjectGroup(), ServerObject->NetRefHandle);
+	Server->PreSendUpdate();
+	Server->PostSendUpdate();
+
+	// ... and cancel that thought! We don't want the object to be destroyed after all.
+	Server->GetReplicationSystem()->RemoveFromGroup(Server->GetReplicationSystem()->GetNotReplicatedNetObjectGroup(), ServerObject->NetRefHandle);
+	Server->PreSendUpdate();
+	Server->PostSendUpdate();
+
+	// Write and send a packet
+	Server->PreSendUpdate();
+	Server->SendAndDeliverTo(Client, DeliverPacket);
+	Server->PostSendUpdate();
+
+	// Verify that the subobject was destroyed on the client
+	UE_NET_ASSERT_EQ(Client->GetReplicationBridge()->GetReplicatedObject(ServerSubObject->NetRefHandle), nullptr);
+}
+
 // See TestObjectSplitting.cpp for cancel pending destroy on huge object
 
 }
diff --git a/Engine/Source/Runtime/Experimental/Iris/Core/Private/Iris/ReplicationSystem/ReplicationWriter.cpp b/Engine/Source/Runtime/Experimental/Iris/Core/Private/Iris/ReplicationSystem/ReplicationWriter.cpp
--- a/Engine/Source/Runtime/Experimental/Iris/Core/Private/Iris/ReplicationSystem/ReplicationWriter.cpp
+++ b/Engine/Source/Runtime/Experimental/Iris/Core/Private/Iris/ReplicationSystem/ReplicationWriter.cpp
@@ -471,7 +471,7 @@ void FReplicationWriter::StartReplication(uint32 InternalIndex)
 		if (ensure(RootObjectInternalIndex != FNetRefHandleManager::InvalidInternalIndex))
 		{
 			FReplicationInfo& OwnerInfo = ReplicatedObjects[RootObjectInternalIndex];
-			if (OwnerInfo.GetState() != EReplicatedObjectState::Invalid)
+			if (OwnerInfo.GetState() != EReplicatedObjectState::Invalid && ensureMsgf(OwnerInfo.GetState() < EReplicatedObjectState::PendingDestroy, TEXT("Unsupported state %s"), LexToString(OwnerInfo.GetState())))
 			{
 				ObjectsWithDirtyChanges.SetBit(RootObjectInternalIndex);
 				OwnerInfo.HasDirtySubObjects = 1U;
@@ -619,6 +619,7 @@ void FReplicationWriter::SetPendingDestroyOrSubObjectPendingDestroyState(uint32
 			{
 				SubObjectInfo.SetState(EReplicatedObjectState::PendingDestroy);
 				SubObjectInfo.SubObjectPendingDestroy = 0U;
+				ObjectsWithDirtyChanges.ClearBit(SubObjectIndex);
 			}
 		}
 	}
@@ -686,6 +687,30 @@ void FReplicationWriter::UpdateScope(const FNetBitArrayView& UpdatedScope)
 					ObjectsWithDirtyChanges.SetBitValue(ObjectData.SubObjectRootIndex, ObjectsWithDirtyChanges.GetBit(ObjectData.SubObjectRootIndex) || Info.HasDirtyChangeMask);
 				}
 			}
+			else if (!Info.IsSubObject)
+			{
+				// If there are subobjects pending destroy we should make sure they're once again resorting to getting destroyed via state replication.
+				bool bHasSubObjectsPendingDestroy = false;
+				for (uint32 SubObjectIndex : NetRefHandleManager->GetSubObjects(Index))
+				{
+					FReplicationInfo& SubObjectInfo = GetReplicationInfo(SubObjectIndex);
+					if (SubObjectInfo.GetState() == EReplicatedObjectState::PendingDestroy)
+					{
+						SubObjectInfo.SetState(EReplicatedObjectState::SubObjectPendingDestroy);
+						SubObjectInfo.SubObjectPendingDestroy = 1U;
+
+						ObjectsWithDirtyChanges.SetBit(SubObjectIndex);
+
+						bHasSubObjectsPendingDestroy = true;
+					}
+				}
+
+				if (bHasSubObjectsPendingDestroy)
+				{
+					ObjectsWithDirtyChanges.SetBit(Index);
+					Info.HasDirtySubObjects = 1U;
+				}
+			}
 		}
 		else
 		{
@@ -1647,6 +1672,7 @@ uint32 FReplicationWriter::WriteObjectsPendingDestroy(FNetSerializationContext&
 						OwnerInfo.HasDirtySubObjects = 1U;
 
 						ObjectsWithDirtyChanges.SetBit(ObjectData.SubObjectRootIndex);
+						ObjectsWithDirtyChanges.SetBit(InternalIndex);
 						continue;
 					}
 				}
@@ -3481,7 +3507,7 @@ void FReplicationWriter::MarkObjectDirty(FInternalNetRefIndex InternalIndex, con
 	if (bValidateObjectsWithDirtyChanges)
 	{
 		const FReplicationInfo& ObjectInfo = ReplicatedObjects[InternalIndex];
-		if (!ensureMsgf(ObjectInfo.GetState() != EReplicatedObjectState::Invalid, TEXT("Object ( InternalIndex: %u ) with Invalid state marked dirty. Caller: %hs"), InternalIndex, Caller))
+		if (!ensureMsgf(ObjectInfo.GetState() != EReplicatedObjectState::Invalid && ObjectInfo.GetState() < EReplicatedObjectState::PendingDestroy, TEXT("Object ( InternalIndex: %u ) with Invalid state marked dirty. Caller: %hs"), InternalIndex, Caller))
 		{
 			return;
 		}
