Iris
- FastArraySerializer will now only compare and apply modified states that are marked as dirty in the changemask

diff --git a/Engine/Source/Runtime/Experimental/Iris/Core/Public/Iris/ReplicationSystem/Private/FastArrayReplicationFragmentInternal.h b/Engine/Source/Runtime/Experimental/Iris/Core/Public/Iris/ReplicationSystem/Private/FastArrayReplicationFragmentInternal.h
--- a/Engine/Source/Runtime/Experimental/Iris/Core/Public/Iris/ReplicationSystem/Private/FastArrayReplicationFragmentInternal.h
+++ b/Engine/Source/Runtime/Experimental/Iris/Core/Public/Iris/ReplicationSystem/Private/FastArrayReplicationFragmentInternal.h
@@ -181,6 +181,17 @@ void FFastArrayReplicationFragmentHelper::ApplyReplicatedState(FastArrayType* Ds
 	ConditionalRebuildItemMap(*DstArraySerializer, *DstWrappedArray);
 	ConditionalRebuildItemMap(*SrcArraySerializer, *SrcWrappedArray);
 
+	const bool bIsNativeFastArraySerializer = EnumHasAnyFlags(Context.Descriptor->Traits, EReplicationStateTraits::IsNativeFastArrayReplicationState);
+	const uint32* ChangeMaskData = bIsNativeFastArraySerializer ?  Context.StateBufferData.ChangeMaskData : (const uint32*)(Context.StateBufferData.ExternalStateBuffer + Context.Descriptor->ChangeMasksExternalOffset);
+	FNetBitArrayView MemberChangeMask = MakeNetBitArrayView(ChangeMaskData, Context.Descriptor->ChangeMaskBitCount);
+
+	// We currently use a simple modulo scheme for bits in the changemask
+	// A single bit might represent several entries in the array which all will be considered dirty, it is up to the serializer to handle this
+	// The first bit is used by the owning property we need to offset by one and deduct one from the usable bits
+	const FReplicationStateMemberChangeMaskDescriptor& MemberChangeMaskDescriptor = Context.Descriptor->MemberChangeMaskDescriptors[0];
+	const uint32 ChangeMaskBitOffset = MemberChangeMaskDescriptor.BitOffset + FIrisFastArraySerializer::IrisFastArrayChangeMaskBitOffset;
+	const uint32 ChangeMaskBitCount = MemberChangeMaskDescriptor.BitCount - FIrisFastArraySerializer::IrisFastArrayChangeMaskBitOffset;
+
 	// Find removed elements in received data, that is elements that exist in old map but not in new map
 	TArray<int32> RemovedIndices;
 	{
@@ -211,8 +222,10 @@ void FFastArrayReplicationFragmentHelper::ApplyReplicatedState(FastArrayType* Ds
 		{			
 			if (int32* ExistingIndex = DstArraySerializer->ItemMap.Find(SrcItems[It].ReplicationID))
 			{
-				// As we currently always send the full array. To be correct, we must compare the element as well before issuing the callback
-				if (!InternalCompareArrayElement(ArrayElementDescriptor, &(*DstWrappedArray)[*ExistingIndex], &SrcItems[It]))
+				const bool bIsDirty = ChangeMaskBitCount == 0U || MemberChangeMask.GetBit((It % ChangeMaskBitCount) + ChangeMaskBitOffset);
+
+				// Only compare if the changemask indicate that this might be a dirty entry, the compare is required since we do share entries in the changemask.
+				if (bIsDirty && !InternalCompareArrayElement(ArrayElementDescriptor, &(*DstWrappedArray)[*ExistingIndex], &SrcItems[It]))
 				{
 					UE_LOG(LogNetFastTArray, Log, TEXT("   Changed. ID: %d -> Idx: %d"), SrcItems[It].ReplicationID, *ExistingIndex);
 
