Iris
* Clean up scopes in Iris Write phase.
* Added dynamic protocol scope when writing individual object

diff --git a/Engine/Source/Runtime/Experimental/Iris/Core/Private/Iris/ReplicationSystem/Conditionals/ReplicationConditionals.cpp b/Engine/Source/Runtime/Experimental/Iris/Core/Private/Iris/ReplicationSystem/Conditionals/ReplicationConditionals.cpp
--- a/Engine/Source/Runtime/Experimental/Iris/Core/Private/Iris/ReplicationSystem/Conditionals/ReplicationConditionals.cpp
+++ b/Engine/Source/Runtime/Experimental/Iris/Core/Private/Iris/ReplicationSystem/Conditionals/ReplicationConditionals.cpp
@@ -520,7 +520,7 @@ void FReplicationConditionals::GetChildSubObjectsToReplicate(uint32 ReplicatingC
 
 void FReplicationConditionals::GetSubObjectsToReplicate(uint32 ReplicationConnectionId, FInternalNetRefIndex RootObjectIndex, FSubObjectsToReplicateArray& OutSubObjectsToReplicate)
 {
-	IRIS_PROFILER_SCOPE(FReplicationConditionals_GetSubObjectsToReplicate);
+	//IRIS_PROFILER_SCOPE_VERBOSE(FReplicationConditionals_GetSubObjectsToReplicate);
 
 	// For now, we do nothing to detect if a conditional has changed on the RootParent, we simply defer this until the next
 	// time the subobjects are marked as dirty. We might want to consider to explicitly mark object and subobjects as dirty when 
@@ -532,7 +532,7 @@ void FReplicationConditionals::GetSubObjectsToReplicate(uint32 ReplicationConnec
 
 bool FReplicationConditionals::ApplyConditionalsToChangeMask(uint32 ReplicatingConnectionId, bool bIsInitialState, FInternalNetRefIndex ParentObjectIndex, FInternalNetRefIndex ObjectIndex, uint32* ChangeMaskData, const uint32* ConditionalChangeMaskData, const FReplicationProtocol* Protocol)
 {
-	IRIS_PROFILER_SCOPE(FReplicationConditionals_ApplyConditionalsToChangeMask);
+	//IRIS_PROFILER_SCOPE_VERBOSE(FReplicationConditionals_ApplyConditionalsToChangeMask);
 
 	bool bMaskWasModified = false;
 
diff --git a/Engine/Source/Runtime/Experimental/Iris/Core/Private/Iris/ReplicationSystem/ReplicationWriter.cpp b/Engine/Source/Runtime/Experimental/Iris/Core/Private/Iris/ReplicationSystem/ReplicationWriter.cpp
--- a/Engine/Source/Runtime/Experimental/Iris/Core/Private/Iris/ReplicationSystem/ReplicationWriter.cpp
+++ b/Engine/Source/Runtime/Experimental/Iris/Core/Private/Iris/ReplicationSystem/ReplicationWriter.cpp
@@ -1841,8 +1841,6 @@ void FReplicationWriter::SerializeObjectStateDelta(FNetSerializationContext& Con
 
 FReplicationWriter::EWriteObjectStatus FReplicationWriter::WriteObjectAndSubObjects(FNetSerializationContext& Context, uint32 InternalIndex, uint32 WriteObjectFlags, FBatchInfo& OutBatchInfo)
 {
-	IRIS_PROFILER_SCOPE(FReplicationWriter_WriteObjectAndSubObjects);
-
 	FNetBitStreamWriter& Writer = *Context.GetBitStreamWriter();
 	
 	FReplicationInfo& Info = GetReplicationInfo(InternalIndex);
@@ -1860,6 +1858,8 @@ FReplicationWriter::EWriteObjectStatus FReplicationWriter::WriteObjectAndSubObje
 	const FNetRefHandleManager::FReplicatedObjectData& ObjectData = NetRefHandleManager->GetReplicatedObjectDataNoCheck(InternalIndex);
 	const FNetRefHandle NetRefHandle = ObjectData.RefHandle;
 
+	IRIS_PROFILER_PROTOCOL_NAME(ObjectData.Protocol?ObjectData.Protocol->DebugName->Name:TEXT("NoProtocol"));
+
 #if UE_NET_TRACE_ENABLED
 	FNetRefHandle NetRefHandleForTraceScope = NetRefHandle;
 	if (WriteObjectFlags & EWriteObjectFlag::WriteObjectFlag_HugeObject)
@@ -1930,7 +1930,6 @@ FReplicationWriter::EWriteObjectStatus FReplicationWriter::WriteObjectAndSubObje
 	uint8* ReplicatedObjectStateBuffer = NetRefHandleManager->GetReplicatedObjectStateBufferNoCheck(InternalIndex);
 
 	const bool bIsInitialState = IsInitialState(State);
-	//IRIS_PROFILER_PROTOCOL_NAME(ObjectData.Protocol?ObjectData.Protocol->DebugName->Name:TEXT("NoProtocol"));
 
 	// Filter out changemasks that are not supposed to be replicated to this connection
 	const bool bNeedToFilterChangeMask = (bIsInitialState || Info.HasDirtyChangeMask) && Info.HasChangemaskFilter;
@@ -2318,8 +2317,6 @@ FReplicationWriter::EWriteObjectStatus FReplicationWriter::WriteObjectAndSubObje
 
 FReplicationWriter::EWriteObjectStatus FReplicationWriter::WriteObjectInBatch(FNetSerializationContext& Context, uint32 InternalIndex, uint32 WriteObjectFlags, FBatchInfo& OutBatchInfo)
 {
-	IRIS_PROFILER_SCOPE(FReplicationWriter_WriteObjectInBatch);
-
 	UE_NET_IRIS_STATS_TIMER(Timer, Context.GetNetStatsContext());
 
 	// Reset pending exports
