Maybe placate PVS

#rnx

diff --git a/Engine/Source/Runtime/Experimental/IoStoreOnDemand/Private/CoreHttp/Http.cpp b/Engine/Source/Runtime/Experimental/IoStoreOnDemand/Private/CoreHttp/Http.cpp
--- a/Engine/Source/Runtime/Experimental/IoStoreOnDemand/Private/CoreHttp/Http.cpp
+++ b/Engine/Source/Runtime/Experimental/IoStoreOnDemand/Private/CoreHttp/Http.cpp
@@ -1721,10 +1721,10 @@ static FActivity* Activity_Alloc(uint32 BufferSize)
 ////////////////////////////////////////////////////////////////////////////////
 static void Activity_Free(FActivity* Activity)
 {
+	Trace(Activity, ETrace::ActivityDestroy);
+
 	Activity->~FActivity();
 	FMemory::Free(Activity);
-
-	Trace(Activity, ETrace::ActivityDestroy);
 }
 
 ////////////////////////////////////////////////////////////////////////////////
@@ -2871,7 +2871,7 @@ bool FSocketGroup::Tick(FTickState& State)
 		RecvInternal(State);
 	}
 
-	return IsKeepAlive | !!(UPTRINT(Send) | UPTRINT(Recv));
+	return !!IsKeepAlive | !!(UPTRINT(Send) | UPTRINT(Recv));
 }
 
 ////////////////////////////////////////////////////////////////////////////////
@@ -3359,7 +3359,7 @@ uint32 FEventLoop::FImpl::Tick(int32 PollTimeoutMs)
 
 ////////////////////////////////////////////////////////////////////////////////
 FEventLoop::FEventLoop()						{ Impl = new FEventLoop::FImpl(); Trace(Impl, ETrace::LoopCreate); }
-FEventLoop::~FEventLoop()						{ delete Impl; Trace(Impl, ETrace::LoopDestroy); }
+FEventLoop::~FEventLoop()						{ Trace(Impl, ETrace::LoopDestroy); delete Impl; }
 uint32 FEventLoop::Tick(int32 PollTimeoutMs)	{ return Impl->Tick(PollTimeoutMs); }
 bool FEventLoop::IsIdle() const					{ return Impl->IsIdle(); }
 void FEventLoop::Cancel(FTicket Ticket)			{ return Impl->Cancel(Ticket); }
