IAS - use correct compression block size

diff --git a/Engine/Source/Runtime/Experimental/IoStoreOnDemand/Private/IoStoreOnDemand.cpp b/Engine/Source/Runtime/Experimental/IoStoreOnDemand/Private/IoStoreOnDemand.cpp
--- a/Engine/Source/Runtime/Experimental/IoStoreOnDemand/Private/IoStoreOnDemand.cpp
+++ b/Engine/Source/Runtime/Experimental/IoStoreOnDemand/Private/IoStoreOnDemand.cpp
@@ -1142,8 +1142,9 @@ TIoStatusOr<FIoStoreUploadResult> UploadContainerFiles(
 			uint64 EncodedChunkSize = 0;
 			for (const FIoStoreCompressedBlockInfo& BlockInfo : ReadResult.Blocks)
 			{
-				const uint64 EncodedBlockSize = Align(BlockInfo.CompressedSize, FAES::AESBlockSize);
-				ContainerEntry.BlockSizes.Add(uint32(EncodedBlockSize));
+				check(Align(BlockInfo.CompressedSize, FAES::AESBlockSize) == BlockInfo.AlignedSize);
+				const uint64 EncodedBlockSize = BlockInfo.AlignedSize;
+				ContainerEntry.BlockSizes.Add(uint32(BlockInfo.CompressedSize));
 
 				FMemoryView EncodedBlock = EncodedBlocks.Left(EncodedBlockSize);
 				EncodedBlocks += EncodedBlock.GetSize();
diff --git a/Engine/Source/Runtime/Experimental/IoStoreOnDemand/Private/OnDemandIoDispatcherBackend.cpp b/Engine/Source/Runtime/Experimental/IoStoreOnDemand/Private/OnDemandIoDispatcherBackend.cpp
--- a/Engine/Source/Runtime/Experimental/IoStoreOnDemand/Private/OnDemandIoDispatcherBackend.cpp
+++ b/Engine/Source/Runtime/Experimental/IoStoreOnDemand/Private/OnDemandIoDispatcherBackend.cpp
@@ -1101,8 +1101,9 @@ static TIoStatusOr<FOnDemandToc> GenerateOnDemandTocFromDisk(FStringView TocHash
 				uint64 EncodedChunkSize = 0;
 				for (const FIoStoreCompressedBlockInfo& BlockInfo : CompressedChunkInfo.Blocks)
 				{
-					const uint64 EncodedBlockSize = Align(BlockInfo.CompressedSize, FAES::AESBlockSize);
-					Container.BlockSizes.Add(uint32(EncodedBlockSize));
+					check(Align(BlockInfo.CompressedSize, FAES::AESBlockSize) == BlockInfo.AlignedSize);
+					const uint64 EncodedBlockSize = BlockInfo.AlignedSize;
+					Container.BlockSizes.Add(uint32(BlockInfo.CompressedSize));
 
 					FIoBlockHash BlockHash;
 					FMemory::Memcpy(&BlockHash, &BlockInfo.DiskHash, sizeof(FIoBlockHash));
