IAS - set correct range header

diff --git a/Engine/Source/Runtime/Experimental/IoStoreOnDemand/Private/OnDemandHttpClient.cpp b/Engine/Source/Runtime/Experimental/IoStoreOnDemand/Private/OnDemandHttpClient.cpp
--- a/Engine/Source/Runtime/Experimental/IoStoreOnDemand/Private/OnDemandHttpClient.cpp
+++ b/Engine/Source/Runtime/Experimental/IoStoreOnDemand/Private/OnDemandHttpClient.cpp
@@ -156,6 +156,7 @@ void FHttpClient::IssueRequest(FRequestParams&& Params)
 				const bool bSuccessful = StatusCode > 199 && StatusCode < 300;
 				if (bSuccessful && Content.GetSize() > 0)
 				{
+					check(Params.Range.GetLength() == 0 || Params.Range.GetLength() == Content.GetSize());
 					Params.Callback(Content, DurationMs);
 					if (Params.Attempt > 0 && Config.bChangeEndpointAfterSuccessfulRetry)
 					{
@@ -209,7 +210,7 @@ void FHttpClient::IssueRequest(FRequestParams&& Params)
 	if (Params.Range.GetOffset() > 0 || Params.Range.GetLength() > 0)
 	{
 		Request.Header(ANSITEXTVIEW("Range"),
-			WriteToAnsiString<64>(ANSITEXTVIEW("bytes="), Params.Range.GetOffset(), ANSITEXTVIEW("-"), Params.Range.GetOffset() + Params.Range.GetLength()));
+			WriteToAnsiString<64>(ANSITEXTVIEW("bytes="), Params.Range.GetOffset(), ANSITEXTVIEW("-"), Params.Range.GetOffset() + Params.Range.GetLength() - 1));
 	}
 
 	EventLoop.Send(MoveTemp(Request), MoveTemp(Sink));
