Dataflow : Convert transfoms from skeleton as a FTransform3f

diff --git a/Engine/Source/Runtime/Experimental/GeometryCollectionEngine/Private/GeometryCollection/GeometryCollectionEngineConversion.cpp b/Engine/Source/Runtime/Experimental/GeometryCollectionEngine/Private/GeometryCollection/GeometryCollectionEngineConversion.cpp
--- a/Engine/Source/Runtime/Experimental/GeometryCollectionEngine/Private/GeometryCollection/GeometryCollectionEngineConversion.cpp
+++ b/Engine/Source/Runtime/Experimental/GeometryCollectionEngine/Private/GeometryCollection/GeometryCollectionEngineConversion.cpp
@@ -1334,7 +1334,7 @@ void FGeometryCollectionEngineConversion::AppendSkeleton(const USkeleton* InSkel
 	FGeometryCollection::DefineTransformSchema(*InCollection);
 	GeometryCollection::Facades::FTransformSource TransformSourceFacade(*InCollection);
 
-	TManagedArray<FTransform>& Transform = InCollection->ModifyAttribute<FTransform>(FTransformCollection::TransformAttribute, FTransformCollection::TransformGroup);
+	TManagedArray<FTransform3f>& Transform = InCollection->ModifyAttribute<FTransform3f>(FTransformCollection::TransformAttribute, FTransformCollection::TransformGroup);
 	TManagedArray<FLinearColor>& BoneColor = InCollection->ModifyAttribute<FLinearColor>("BoneColor", FTransformCollection::TransformGroup);
 	TManagedArray<FString>& BoneName = InCollection->ModifyAttribute<FString>("BoneName", FTransformCollection::TransformGroup);
 	TManagedArray<int32>& Parent = InCollection->ModifyAttribute<int32>(FTransformCollection::ParentAttribute, FTransformCollection::TransformGroup);
@@ -1351,7 +1351,7 @@ void FGeometryCollectionEngineConversion::AppendSkeleton(const USkeleton* InSkel
 		int32 TransformBaseIndex = InCollection->AddElements(NumBones, FGeometryCollection::TransformGroup);
 		for (int i = 0, Idx = TransformBaseIndex; i < NumBones; i++, Idx++)
 		{
-			Transform[Idx] = RestTransform[i];
+			Transform[Idx] = FTransform3f(RestTransform[i]);
 			BoneColor[Idx] = FLinearColor(FColor(FMath::Rand() % 100 + 5, FMath::Rand() % 100 + 5, FMath::Rand() % 100 + 5, 255));
 			BoneName[Idx] = BoneInfo[i].Name.ToString();
 			Parent[Idx] = BoneInfo[i].ParentIndex;
