Iris - Force dirty RemoteRole when autonomous proxy condition changes.

diff --git a/Engine/Source/Runtime/Experimental/Iris/Core/Private/Iris/ReplicationSystem/Conditionals/ReplicationConditionals.cpp b/Engine/Source/Runtime/Experimental/Iris/Core/Private/Iris/ReplicationSystem/Conditionals/ReplicationConditionals.cpp
--- a/Engine/Source/Runtime/Experimental/Iris/Core/Private/Iris/ReplicationSystem/Conditionals/ReplicationConditionals.cpp
+++ b/Engine/Source/Runtime/Experimental/Iris/Core/Private/Iris/ReplicationSystem/Conditionals/ReplicationConditionals.cpp
@@ -12,8 +12,10 @@
 #include "Iris/ReplicationSystem/DeltaCompression/DeltaCompressionBaselineInvalidationTracker.h"
 #include "Iris/ReplicationSystem/Filtering/NetObjectGroups.h"
 #include "Iris/ReplicationSystem/Filtering/ReplicationFiltering.h"
+#include "Iris/Serialization/InternalNetSerializers.h"
 #include "Containers/ArrayView.h"
 #include "UObject/CoreNetTypes.h"
+#include "Net/Core/NetHandle/NetHandleManager.h"
 #include "Net/Core/PropertyConditions/RepChangedPropertyTracker.h"
 #include "Net/Core/PropertyConditions/PropertyConditions.h"
 #include "Net/Core/Trace/NetDebugName.h"
@@ -68,6 +70,8 @@ bool FReplicationConditionals::SetConditionConnectionFilter(FInternalNetRefIndex
 		ObjectInfo->AutonomousConnectionId = uint16(AutonomousConnectionId);
 
 		BaselineInvalidationTracker->InvalidateBaselines(ObjectIndex, ConnIdForBaselineInvalidation);
+
+		MarkRemoteRoleDirty(ObjectIndex);
 	}
 
 	return true;
@@ -822,4 +826,125 @@ bool FReplicationConditionals::DynamicConditionChangeRequiresBaselineInvalidatio
 	return OldConditionMayHaveBeenDisabled && NewConditionMayBeEnabled;
 }
 
+void FReplicationConditionals::MarkRemoteRoleDirty(FInternalNetRefIndex ObjectIndex)
+{
+	const FNetRefHandleManager::FReplicatedObjectData& ReplicatedObjectData = NetRefHandleManager->GetReplicatedObjectDataNoCheck(ObjectIndex);
+	const FReplicationProtocol* Protocol = ReplicatedObjectData.Protocol;
+
+	if (NetRefHandleManager->GetReplicatedObjectStateBufferNoCheck(ObjectIndex) == nullptr)
+	{
+		return;
+	}
+
+	if (!ReplicatedObjectData.NetHandle.IsValid())
+	{
+		return;
+	}
+
+	const uint16 RepIndex = GetRemoteRoleRepIndex(Protocol);
+	if (RepIndex == InvalidRepIndex)
+	{
+		return;
+	}
+
+	MarkPropertyDirty(ObjectIndex, RepIndex);
+}
+
+uint16 FReplicationConditionals::GetRemoteRoleRepIndex(const FReplicationProtocol* Protocol)
+{
+	if (CachedRemoteRoleRepIndex != InvalidRepIndex)
+	{
+		return CachedRemoteRoleRepIndex;
+	}
+	
+	const FNetSerializer* NetRoleNetSerializer = &UE_NET_GET_SERIALIZER(FNetRoleNetSerializer);
+
+	// Loop through all state descriptors end their properties to find the RemoteRole
+	for (const FReplicationStateDescriptor* StateDescriptor : MakeArrayView(Protocol->ReplicationStateDescriptors, static_cast<int32>(Protocol->ReplicationStateCount)))
+	{
+		for (const FReplicationStateMemberSerializerDescriptor& SerializerDescriptor : MakeArrayView(StateDescriptor->MemberSerializerDescriptors, StateDescriptor->MemberCount))
+		{
+			if (SerializerDescriptor.Serializer != NetRoleNetSerializer)
+			{
+				continue;
+			}
+
+			const SIZE_T MemberIndex = &SerializerDescriptor - StateDescriptor->MemberSerializerDescriptors;
+			const FProperty* Property = StateDescriptor->MemberProperties[MemberIndex];
+			if (Property && Property->GetFName() == NAME_RemoteRole)
+			{
+				CachedRemoteRoleRepIndex = Property->RepIndex;
+				return Property->RepIndex;
+			}
+		}
+	}
+
+	return InvalidRepIndex;
+}
+
+void FReplicationConditionals::MarkPropertyDirty(FInternalNetRefIndex ObjectIndex, uint16 RepIndex)
+{
+	const FNetRefHandleManager::FReplicatedObjectData& ReplicatedObjectData = NetRefHandleManager->GetReplicatedObjectDataNoCheck(ObjectIndex);
+
+	const FNetHandle OwnerHandle = ReplicatedObjectData.NetHandle;
+	if (!OwnerHandle.IsValid())
+	{
+		return;
+	}
+
+	if (NetRefHandleManager->GetReplicatedObjectStateBufferNoCheck(ObjectIndex) == nullptr)
+	{
+		return;
+	}
+
+	const FReplicationProtocol* Protocol = ReplicatedObjectData.Protocol;
+	const FReplicationInstanceProtocol* InstanceProtocol = ReplicatedObjectData.InstanceProtocol;
+
+	constexpr uint32 MaxFragmentOwnerCount = 1U;
+	UObject* FragmentOwners[MaxFragmentOwnerCount] = {};
+	FReplicationStateOwnerCollector FragmentOwnerCollector(FragmentOwners, MaxFragmentOwnerCount);
+
+	for (const FReplicationStateDescriptor*& StateDescriptor : MakeArrayView(Protocol->ReplicationStateDescriptors, static_cast<int32>(Protocol->ReplicationStateCount)))
+	{
+		const SIZE_T StateIndex = &StateDescriptor - Protocol->ReplicationStateDescriptors;
+
+		// Is the passed Owner the owner of the fragment?
+		{
+			const FReplicationFragment* ReplicationFragment = InstanceProtocol->Fragments[StateIndex];
+
+			FragmentOwnerCollector.Reset();
+			ReplicationFragment->CollectOwner(&FragmentOwnerCollector);
+			if (FragmentOwnerCollector.GetOwnerCount() == 0U || FNetHandleManager::GetNetHandle(FragmentOwnerCollector.GetOwners()[0]) != OwnerHandle)
+			{
+				// Not the right owner.
+				continue;
+			}
+		}
+
+		// Can this state contain this property?
+		if (RepIndex >= StateDescriptor->RepIndexCount)
+		{
+			continue;
+		}
+
+		// Does this state contain this property?
+		const FReplicationStateMemberRepIndexToMemberIndexDescriptor& RepIndexToMemberIndexDescriptor = StateDescriptor->MemberRepIndexToMemberIndexDescriptors[RepIndex];
+		if (RepIndexToMemberIndexDescriptor.MemberIndex == FReplicationStateMemberRepIndexToMemberIndexDescriptor::InvalidEntry)
+		{
+			continue;
+		}
+
+		// We found the relevant state. Modify the external state changemask.
+		const FReplicationInstanceProtocol::FFragmentData& Fragment = InstanceProtocol->FragmentData[StateIndex];
+		const FReplicationStateMemberChangeMaskDescriptor& ChangeMaskDescriptor = StateDescriptor->MemberChangeMaskDescriptors[RepIndexToMemberIndexDescriptor.MemberIndex];
+		FNetBitArrayView MemberChangeMask = GetMemberChangeMask(Fragment.ExternalSrcBuffer, StateDescriptor);
+		FReplicationStateHeader& Header = GetReplicationStateHeader(Fragment.ExternalSrcBuffer, StateDescriptor);
+		MarkDirty(Header, MemberChangeMask, ChangeMaskDescriptor);
+
+		return;
+	}
+
+	UE_LOG(LogIris, Warning, TEXT("Trying to mark non-existing property with RepIndex %u in protocol %s as dirty"), RepIndex, ToCStr(Protocol->DebugName));
+}
+
 }
diff --git a/Engine/Source/Runtime/Experimental/Iris/Core/Private/Iris/ReplicationSystem/Conditionals/ReplicationConditionals.h b/Engine/Source/Runtime/Experimental/Iris/Core/Private/Iris/ReplicationSystem/Conditionals/ReplicationConditionals.h
--- a/Engine/Source/Runtime/Experimental/Iris/Core/Private/Iris/ReplicationSystem/Conditionals/ReplicationConditionals.h
+++ b/Engine/Source/Runtime/Experimental/Iris/Core/Private/Iris/ReplicationSystem/Conditionals/ReplicationConditionals.h
@@ -62,6 +62,9 @@ public:
 	void InitPropertyCustomConditions(FInternalNetRefIndex ObjectIndex);
 	bool SetPropertyCustomCondition(FInternalNetRefIndex ObjectIndex, const void* Owner, uint16 RepIndex, bool bIsActive);
 	bool SetPropertyDynamicCondition(FInternalNetRefIndex ObjectIndex, const void* Owner, uint16 RepIndex, ELifetimeCondition Condition);
+	
+	/** Unconditionally marks a property as dirty, causing it to replicate with the object at the earliest convenience. */
+	void MarkPropertyDirty(FInternalNetRefIndex ObjectIndex, uint16 RepIndex);
 
 	void Update();
 
@@ -121,19 +124,25 @@ private:
 	void SetDynamicCondition(FInternalNetRefIndex ObjectIndex, uint16 RepIndex, ELifetimeCondition Condition);
 	bool DynamicConditionChangeRequiresBaselineInvalidation(ELifetimeCondition OldCondition, ELifetimeCondition NewCondition) const;
 
+	void MarkRemoteRoleDirty(FInternalNetRefIndex ObjectIndex);
+	uint16 GetRemoteRoleRepIndex(const FReplicationProtocol* Protocol);
+
 private:
+	static constexpr uint16 InvalidRepIndex = 65535U;
+
 	const FNetRefHandleManager* NetRefHandleManager = nullptr;
 	const FReplicationFiltering* ReplicationFiltering = nullptr;
 	const FReplicationConnections* ReplicationConnections = nullptr;
 	FDeltaCompressionBaselineInvalidationTracker* BaselineInvalidationTracker = nullptr;
 	const FNetObjectGroups* NetObjectGroups = nullptr;
 
-	uint32 MaxObjectCount = 0;
-	uint32 MaxConnectionCount = 0;
-
 	TArray<FPerObjectInfo> PerObjectInfos;
 	TArray<FPerConnectionInfo> ConnectionInfos;
 	TMap<FInternalNetRefIndex, FObjectDynamicConditions> DynamicConditions;
+
+	uint32 MaxObjectCount = 0;
+	uint32 MaxConnectionCount = 0;
+	uint16 CachedRemoteRoleRepIndex = InvalidRepIndex;
 };
 
 inline FReplicationConditionals::FPerObjectInfo* FReplicationConditionals::GetPerObjectInfo(FInternalNetRefIndex ObjectIndex)
