Add option (disabled by default) for IAS to write out one toc per container rather than a single all encompassing toc.

#rnx

### IAS
- The previous version of this system worked by appending IAS data to the end of .utoc files and using that data to generate the FOnDemandTocContainerEntry data that we needed, however we don't want to do this long term as we want to remove the IAS code from IoStoreUtilities. So instead we want to store a FOnDemandToc for each .utoc that would contain on demand data, where the FOnDemandToc contains a single FOnDemandTocContainerEntry detailing the chunks that it has access to. Then when we mount a .pak we can mount this new file (currently called .uondemandtoc) if it exists for that .pak file, in the exact same way that we would mount a .utoc.
- Note that for now we will continue to produce and use the single large .iochunktoc file as before. This new code path does work for basic games but will need more testing and work before it will function correctly in all cases.
-- Both the old format (.iochunktoc) and the new format (.uondemandtoc) are only produced if the cmdline arg '-WriteTocToDisk' is passed to UnrealPak.
- Added a new global string 'GIasOnDemandTocExt' to allow us to change the extension easily in the future. Currently set to .uondemandtoc.
- There are a number of places that require us to calculate paths based on the provided service url and other server configuration options. This code has been moved to it's own function 'ResolvePaths' so that it can be reused.
-- Previously this code was in IoStoreUtilities where it was used to write out the .ini file. As it has been moved to the IoStoreOnDemandModule we now return the paths as part of FIoStoreUploadResult so that the ini file can continue to use the info.
- Removed the cvar 's.IasGenerateOnDemandToc' and 'ias.LoadOnDemandToc and replaced it with a new cvar 'ias.TocMode which is an int rather than a bool and allows us to select the type of toc loading that we want to do (load from .iochunktoc, download a .iochunktoc or use the new .uondemandtoc files).
-- Using enums as cvars is discouraged so instead of using the cvar directly we use an accessor function ::GetTocMode which converts the cvar to an enum (ETocMode) and gives errors if the cvar is invalid.
- FOnDemandIoStore::AddToc has been extended to work with the traditional toc methods (where TocPath is a path to a file) and the new method (where TocPath is just the filename).
-- The old method needs to parse the prefix, but the new method will store the path as "ChunksDirectory".
- All code for generating toc info from .utoc files has been removed as well as the code to verify that a generated FOnDemandToc is the same as the version downloaded from the network. Code that was working on both loaded and generated tocs has been renamed to reflect that we only load tocs now.
- When mounting FOnDemandIoBackend and using the new method we first hook into the newly added GetPakMountOperationDelegate in the PakFile system so that we are notified about new pak files being mounted followed by an attempt to find the pakfiles that have already been mounted and catch up by mounting any corresponding .uondemandtoc files.

## IAS Misc
- Minor Optimization: Changed some const FString& parameters to FStringView in places that were commonly being called with a TCHAR* string to avoid the implicit FString conversion.
- Minor Optimization: Presize some arrays during UploadContainerFiles to avoid unnecessary resizing and copying.
- Fixed some typos.
- Removed any remaining reference to the pakfile module as all communicate is via core delegates.
- When uploading we no longer log the containers that we are skipping.
- Changed the logging of uploaded chunks from Display to Log to reduce output spam.

diff --git a/Engine/Source/Developer/IoStoreUtilities/Private/IoStoreUtilities.cpp b/Engine/Source/Developer/IoStoreUtilities/Private/IoStoreUtilities.cpp
--- a/Engine/Source/Developer/IoStoreUtilities/Private/IoStoreUtilities.cpp
+++ b/Engine/Source/Developer/IoStoreUtilities/Private/IoStoreUtilities.cpp
@@ -9686,48 +9686,15 @@ bool UploadIoStoreContainerFiles(const UE::IO::IAS::FIoStoreUploadParams& Upload
 				Sb << TEXT("FallbackUrl=\"") << FallbackUrl << TEXT("\"\r\n");
 			}
 			
-			Sb << TEXT("TocPath=\"") << UploadResult.TocPath << TEXT("\"\r\n");
+			
 		}
 		else
 		{
-			// The configuration file should specify a service URL without any trailing
-			// host path, i.e. http://{host:port}/{host-path}. Add the trailing path
-			// to the TOC path to form the complete path the TOC from the host, i.e
-			// TocPath={host-path}/{bucket}/{bucket-prefix}/{toc-hash}.uchunktoc
-
-			FStringView ServiceUrl = UploadParams.ServiceUrl;
-			FStringView TocPrefx;
-			{
-				int32 Sep = INDEX_NONE;
-				int32 Count = 0;
-				for (int32 Idx = 0, Len = ServiceUrl.Len(); Idx < Len; ++Idx)
-				{
-					if (ServiceUrl[Idx] == TEXT('/'))
-					{
-						if (++Count == 3)
-						{
-							Sep = Idx;
-							break;
-						}
-					}
-				}
-
-				if (Sep != INDEX_NONE)
-				{
-					TocPrefx = ServiceUrl.RightChop(Sep + 1);
-					ServiceUrl.LeftInline(Sep);
-				}
-			}
-
-			Sb << TEXT("ServiceUrl=\"") << ServiceUrl << TEXT("\"\r\n");
-			Sb << TEXT("TocPath=\"");
-			if (TocPrefx.IsEmpty() == false)
-			{
-				Sb << TocPrefx << TEXT('/');
-			}
-			Sb << UploadParams.Bucket / UploadResult.TocPath << TEXT("\"\r\n");
+			Sb << TEXT("ServiceUrl=\"") << UploadResult.ServiceUrl << TEXT("\"\r\n");
 		}
 
+		Sb << TEXT("TocPath=\"") << UploadResult.TocPath << TEXT("\"\r\n");
+
 		// Temporary solution to get replays working with encrypted on demand content
 		{
 			FString EncryptionKeyName;
diff --git a/Engine/Source/Runtime/Experimental/IoStoreOnDemand/IoStoreOnDemand.Build.cs b/Engine/Source/Runtime/Experimental/IoStoreOnDemand/IoStoreOnDemand.Build.cs
--- a/Engine/Source/Runtime/Experimental/IoStoreOnDemand/IoStoreOnDemand.Build.cs
+++ b/Engine/Source/Runtime/Experimental/IoStoreOnDemand/IoStoreOnDemand.Build.cs
@@ -24,17 +24,5 @@ public class IoStoreOnDemand : ModuleRules
 		{
 			PrivateDependencyModuleNames.AddRange(new string[] { "S3Client", "RSA" });
 		}
-
-		bool bFindUcasViaPakFileModule = false;
-
-		if (bFindUcasViaPakFileModule)
-		{
-			PrivateDependencyModuleNames.Add("PakFile");
-			PublicDefinitions.Add("UE_IAS_LINKPAKFILE=1");
-		}
-		else
-		{
-			PublicDefinitions.Add("UE_IAS_LINKPAKFILE=0");
-		}
 	}
 }
diff --git a/Engine/Source/Runtime/Experimental/IoStoreOnDemand/Private/IoStoreOnDemand.cpp b/Engine/Source/Runtime/Experimental/IoStoreOnDemand/Private/IoStoreOnDemand.cpp
--- a/Engine/Source/Runtime/Experimental/IoStoreOnDemand/Private/IoStoreOnDemand.cpp
+++ b/Engine/Source/Runtime/Experimental/IoStoreOnDemand/Private/IoStoreOnDemand.cpp
@@ -48,6 +48,8 @@ DEFINE_LOG_CATEGORY(LogIas);
 namespace UE::IO::IAS
 {
 
+FString GIasOnDemandTocExt = TEXT(".uondemandtoc");
+
 bool GIasSuspendSystem = false;
 static FAutoConsoleVariableRef CVar_SuspendSystemEnabled(
 	TEXT("ias.SuspendSystem"),
@@ -64,6 +66,9 @@ static FAutoConsoleVariableRef CVar_DistributedEndpointFallbackUrl(
 	TEXT("CDN url to be used if a distributed endpoint cannot be reached (overrides IoStoreOnDemand.ini)")
 );
 
+// When eanbled we will write out a toc (using extension GIasOnDemandTocExt) per container, currently disabled for testing
+static bool GIasWriteOnDemandTocs = false;
+
 ////////////////////////////////////////////////////////////////////////////////
 static int64 ParseSizeParam(FStringView Value)
 {
@@ -767,7 +772,7 @@ class FS3UploadQueue
 {
 public:
 	FS3UploadQueue(FS3Client& Client, const FString& Bucket, int32 ThreadCount);
-	bool Enqueue(const FString& Key, FIoBuffer Payload);
+	bool Enqueue(FStringView Key, FIoBuffer Payload);
 	bool Flush();
 
 private:
@@ -806,7 +811,7 @@ FS3UploadQueue::FS3UploadQueue(FS3Client& InClient, const FString& InBucket, int
 	}
 }
 
-bool FS3UploadQueue::Enqueue(const FString& Key, FIoBuffer Payload)
+bool FS3UploadQueue::Enqueue(FStringView Key, FIoBuffer Payload)
 {
 	if (ActiveThreadCount == 0)
 	{
@@ -820,7 +825,7 @@ bool FS3UploadQueue::Enqueue(const FString& Key, FIoBuffer Payload)
 			FScopeLock _(&CriticalSection);
 			if (ConcurrentUploads < Threads.Num())
 			{
-				bEnqueued = Queue.Enqueue(FQueueEntry {Key, Payload});
+				bEnqueued = Queue.Enqueue(FQueueEntry {FString(Key), Payload});
 			}
 		}
 
@@ -864,7 +869,7 @@ void FS3UploadQueue::ThreadEntry()
 
 		if (Response.IsOk())
 		{
-			UE_LOG(LogIas, Display, TEXT("Uploaded chunk '%s/%s/%s'"), *Client.GetConfig().ServiceUrl, *Bucket, *Entry.Key);
+			UE_LOG(LogIas, Log, TEXT("Uploaded chunk '%s/%s/%s'"), *Client.GetConfig().ServiceUrl, *Bucket, *Entry.Key);
 		}
 		else
 		{
@@ -897,6 +902,125 @@ bool FS3UploadQueue::Flush()
 }
 
 ////////////////////////////////////////////////////////////////////////////////
+
+struct FResolvedPaths
+{
+	FString ServiceUrl;
+	FString TocPath;
+	FString ChunkPrefix;
+};
+
+[[nodiscard]] static FResolvedPaths ResolvePaths(const FIoStoreUploadParams& UploadParams, FStringView TocFilename)
+{
+	FResolvedPaths Result;
+	if (!UploadParams.BucketPrefix.IsEmpty())
+	{
+		Result.ServiceUrl = UploadParams.ServiceUrl;
+		Result.TocPath = TocFilename;
+		Result.ChunkPrefix = UploadParams.BucketPrefix;
+	}
+	else
+	{
+		// The configuration file should specify a service URL without any trailing
+		// host path, i.e. http://{host:port}/{host-path}. Add the trailing path
+		// to the TOC path to form the complete path the TOC from the host, i.e
+		// TocPath={host-path}/{bucket}/{bucket-prefix}/{toc-hash}.uchunktoc
+
+		FStringView ServiceUrl = UploadParams.ServiceUrl;
+		FStringView TocPrefx;
+		{
+			// Find the first '//' then find the first '/' after that
+			int32 Sep = INDEX_NONE;
+
+			const int32 DoubleSlash = ServiceUrl.Find(TEXT("//"));
+			if (DoubleSlash != INDEX_NONE)
+			{
+				Sep = ServiceUrl.Find(TEXT("/"), DoubleSlash + 2);
+			}
+
+			if (Sep != INDEX_NONE)
+			{
+				TocPrefx = ServiceUrl.RightChop(Sep + 1);
+				ServiceUrl.LeftInline(Sep);
+			}
+		}
+
+		Result.ServiceUrl = ServiceUrl;
+
+		TStringBuilder<256> Prefix;
+		FPathViews::Append(Prefix, TocPrefx, UploadParams.Bucket);
+
+		Result.ChunkPrefix = Prefix;
+		
+		if (!TocFilename.IsEmpty())
+		{
+			TStringBuilder<256> TocPath;
+			FPathViews::Append(TocPath, Prefix, TocFilename);
+
+			Result.TocPath = TocPath;
+		}
+	}
+
+	return Result;
+}
+
+[[nodiscard]] static FString ResolveChunksDirectory(const FIoStoreUploadParams& UploadParams)
+{
+	FResolvedPaths Paths = ResolvePaths(UploadParams, FStringView());
+	return Paths.ChunkPrefix;
+}
+
+static FIoStatus WriteContainerFiles(FOnDemandToc& OnDemandToc, const TMap<FIoHash, FString>& UTocPaths)
+{
+	UE_LOG(LogIas, Display, TEXT("Attempting to write out %d '%s' files"), OnDemandToc.Containers.Num(), *GIasOnDemandTocExt);
+
+	for (FOnDemandTocContainerEntry& Container : OnDemandToc.Containers)
+	{
+		const FString* UTocPath = UTocPaths.Find(Container.UTocHash);
+		if (UTocPath == nullptr)
+		{
+			return FIoStatus(EIoErrorCode::Unknown, FString::Printf(TEXT("Could not find the original path for ondemand container '%s'"), *Container.ContainerName));
+		}
+
+		// Create a new FOnDemandToc to be written out and give it ownership of the FOnDemandTocContainerEntry that we want
+		// associated with it.This way we don't need to write any specialized serialization code to split up the input FOnDemandToc
+		// into one per container.
+		FOnDemandToc ContainerToc;
+		ContainerToc.Header = OnDemandToc.Header;
+		ContainerToc.Meta = OnDemandToc.Meta;
+		ContainerToc.Containers.Emplace(MoveTemp(Container));
+
+		const FString OutputPath = FPathViews::ChangeExtension(*UTocPath, GIasOnDemandTocExt);
+
+		TUniquePtr<FArchive> Ar(IFileManager::Get().CreateFileWriter(*OutputPath));
+
+		if (!Ar.IsValid())
+		{
+			return FIoStatus(EIoErrorCode::FileOpenFailed, FString::Printf(TEXT("Failed to open '%s' for write"), *OutputPath));
+		}
+
+		// TODO: We should consider adding a hash of the FOnDemandToc that can be computed at runtime on the loaded structure
+		// (to avoid running over the file twice) to verify that nothing was corrupted.
+
+		(*Ar) << ContainerToc;
+
+		if (Ar->IsError() || Ar->IsCriticalError())
+		{
+			return FIoStatus(EIoErrorCode::WriteError, FString::Printf(TEXT("Failed to write to '%s'"), *OutputPath));
+		}
+
+		// TODO: Should write an end of file sentinel here for a simple/easy file corruption check
+
+		UE_LOG(LogIas, Display, TEXT("Wrote ondemand container file '%s' to disk"), *OutputPath);
+
+		// Move the container entry back to OnDemandToc in case we want to use the data structure
+		// beyond this point in the future.
+		Container = MoveTemp(ContainerToc.Containers[0]);
+	}
+
+	return FIoStatus::Ok;
+}
+
 TIoStatusOr<FIoStoreUploadParams> FIoStoreUploadParams::Parse(const TCHAR* CommandLine)
 {
 	FIoStoreUploadParams Params;
@@ -1011,7 +1135,7 @@ TIoStatusOr<FIoStoreUploadResult> UploadContainerFiles(
 		TStringBuilder<256> TocsKey;
 		TocsKey << UploadParams.BucketPrefix << "/";
 
-		UE_LOG(LogIas, Display, TEXT("Fetching existing TOC's from '%s/%s/%s'"), *Client.GetConfig().ServiceUrl, *UploadParams.Bucket, TocsKey.ToString());
+		UE_LOG(LogIas, Display, TEXT("	 '%s/%s/%s'"), *Client.GetConfig().ServiceUrl, *UploadParams.Bucket, TocsKey.ToString());
 		FS3ListObjectResponse Response = Client.ListObjects(FS3ListObjectsRequest
 		{
 			UploadParams.Bucket,
@@ -1077,7 +1201,14 @@ TIoStatusOr<FIoStoreUploadResult> UploadContainerFiles(
 	uint64 TotalUploadedBytes = 0;
 
 	FOnDemandToc OnDemandToc;
-	OnDemandToc.Header.ChunksDirectory = TEXT("Chunks");
+	OnDemandToc.Header.ChunksDirectory = ResolveChunksDirectory(UploadParams);
+
+	OnDemandToc.Containers.Reserve(ContainerFiles.Num());
+
+	// Map of the .utoc paths that we have created ondemand containers for, indexed by their hash so that the paths can
+	// be looked up later. We do not rely on the filename as we cannot be sure that there won't be duplicate file names
+	// stored in different directories.
+	TMap<FIoHash, FString> UTocPaths;
 
 	TArray<FString> FilesToDelete;
 	for (const FString& Path : ContainerFiles)
@@ -1094,11 +1225,9 @@ TIoStatusOr<FIoStoreUploadResult> UploadContainerFiles(
 
 		if (EnumHasAnyFlags(ContainerFileReader.GetContainerFlags(), EIoContainerFlags::OnDemand) == false)
 		{
-			UE_LOG(LogIas, Display, TEXT("Skipping non ondemand container '%s'"), *Path);
 			continue;
 		}
-		
-		UE_LOG(LogIas, Display, TEXT("Uploading container '%s'"), *Path);
+		UE_LOG(LogIas, Display, TEXT("Uploading ondemand container '%s'"), *Path);
 
 		const uint32 BlockSize = ContainerFileReader.GetCompressionBlockSize();
 		if (OnDemandToc.Header.BlockSize == 0)
@@ -1122,6 +1251,8 @@ TIoStatusOr<FIoStoreUploadResult> UploadContainerFiles(
 		}
 		FContainerStats& ContainerStats = ContainerSummary.FindOrAdd(ContainerEntry.ContainerName);
 		
+		ContainerEntry.Entries.Reserve(ChunkInfos.Num());
+
 		for (const FIoStoreTocChunkInfo& ChunkInfo : ChunkInfos)
 		{
 			const bool bDecrypt = false;
@@ -1185,7 +1316,7 @@ TIoStatusOr<FIoStoreUploadResult> UploadContainerFiles(
 					<< TEXT("/") << HashString
 					<< TEXT(".iochunk");
 
-				if (UploadQueue.Enqueue(Key.ToString(), ReadResult.IoBuffer) == false)
+				if (UploadQueue.Enqueue(Key, ReadResult.IoBuffer) == false)
 				{
 					return FIoStatus(EIoErrorCode::WriteError, TEXT("Failed to upload chunk"));
 				}
@@ -1216,22 +1347,23 @@ TIoStatusOr<FIoStoreUploadResult> UploadContainerFiles(
 			const FS3PutObjectResponse Response = Client.TryPutObject(
 				FS3PutObjectRequest{UploadParams.Bucket, Key.ToString(), MakeMemoryView(Buffer.GetData(), Buffer.Num())});
 
-			if (Response.IsOk())
-			{
-				UE_LOG(LogIas, Display, TEXT("Uploaded '%s'"), *UTocFilePath);
-			}
-			else
+			if (!Response.IsOk())
 			{
 				return FIoStatus(EIoErrorCode::WriteError, FString::Printf(TEXT("Failed to upload '%s', StatusCode: %u"), *UTocFilePath, Response.StatusCode));
 			}
+			
+			UTocPaths.Add(ContainerEntry.UTocHash, UTocFilePath);
+
+			UE_LOG(LogIas, Display, TEXT("Uploaded '%s'"), *UTocFilePath);	
 		}
-		
+
 		if (UploadParams.bDeleteContainerFiles)
 		{
 			FilesToDelete.Add(Path);
 			ContainerFileReader.GetContainerFilePaths(FilesToDelete);
 
-			if (UploadParams.bDeletePakFiles)
+			// We need the pak files in order to mount OnDemand toc files!
+			if (UploadParams.bDeletePakFiles && !GIasWriteOnDemandTocs)
 			{
 				FilesToDelete.Add(FPaths::ChangeExtension(Path, TEXT(".pak")));
 				FilesToDelete.Add(FPaths::ChangeExtension(Path, TEXT(".sig")));
@@ -1249,24 +1381,28 @@ TIoStatusOr<FIoStoreUploadResult> UploadContainerFiles(
 		return FIoStatus(EIoErrorCode::WriteError, TEXT("Failed to upload chunk(s)"));
 	}
 
+	OnDemandToc.Meta.EpochTimestamp = FDateTime::Now().ToUnixTimestamp();
+	OnDemandToc.Meta.BuildVersion = UploadParams.BuildVersion;
+	OnDemandToc.Meta.TargetPlatform = UploadParams.TargetPlatform;
+
 	FIoStoreUploadResult UploadResult;
 	{
-		OnDemandToc.Meta.EpochTimestamp = FDateTime::Now().ToUnixTimestamp();
-		OnDemandToc.Meta.BuildVersion = UploadParams.BuildVersion;
-		OnDemandToc.Meta.TargetPlatform = UploadParams.TargetPlatform;
-
 		FLargeMemoryWriter Ar;
 		Ar << OnDemandToc;
 
 		UploadResult.TocHash = FIoHash::HashBuffer(Ar.GetView());
 		TStringBuilder<256> Key;
-		if (UploadParams.BucketPrefix.IsEmpty() == false)
+		if (!UploadParams.BucketPrefix.IsEmpty())
 		{
 			Key << UploadParams.BucketPrefix.ToLower() << TEXT("/");
 		}
 		Key << LexToString(UploadResult.TocHash) << TEXT(".iochunktoc");
+		
+		FResolvedPaths ResolvedPaths = ResolvePaths(UploadParams, Key);
+
+		UploadResult.ServiceUrl = MoveTemp(ResolvedPaths.ServiceUrl);
+		UploadResult.TocPath = MoveTemp(ResolvedPaths.TocPath);
 
-		UploadResult.TocPath = Key.ToString();
 		UploadResult.TocSize = Ar.TotalSize();
 
 		const FS3PutObjectResponse Response = Client.TryPutObject(FS3PutObjectRequest{UploadParams.Bucket, Key.ToString(), Ar.GetView()});
@@ -1282,6 +1418,7 @@ TIoStatusOr<FIoStoreUploadResult> UploadContainerFiles(
 
 		if (UploadParams.bWriteTocToDisk)
 		{
+			// Write a single .iochunktoc containing all on demand data for the current build
 			TStringBuilder<512> OnDemandTocFilePath;
 			FPathViews::Append(OnDemandTocFilePath, UploadParams.TocOutputDir, UploadResult.TocHash);
 			OnDemandTocFilePath << TEXT(".iochunktoc");
@@ -1293,6 +1430,16 @@ TIoStatusOr<FIoStoreUploadResult> UploadContainerFiles(
 		}
 	}
 
+	if (GIasWriteOnDemandTocs)
+	{
+		// Write out separate .uondemandtoc files, one per.utoc containing ondemand data.
+		FIoStatus Result = WriteContainerFiles(OnDemandToc, UTocPaths);
+		if (!Result.IsOk())
+		{
+			return Result;
+		}
+	}
+
 	for (const FString& Path : FilesToDelete)
 	{
 		if (IFileManager::Get().FileExists(*Path))
@@ -1664,7 +1811,7 @@ TIoStatusOr<FIoStoreListTocsParams> FIoStoreListTocsParams::Parse(const TCHAR* C
 {
 	FIoStoreListTocsParams Params;
 
-	// Convenience argument to specifiy both bucket and bucket prefix
+	// Convenience argument to specify both bucket and bucket prefix
 	// -BucketPath="mybucket/some/data/path" is equal to -Bucket="bucket" -BucketPrefix="some/data/path
 	FString BucketPath;
 	if (FParse::Value(CommandLine, TEXT("-BucketPath="), BucketPath))
diff --git a/Engine/Source/Runtime/Experimental/IoStoreOnDemand/Private/OnDemandIoDispatcherBackend.cpp b/Engine/Source/Runtime/Experimental/IoStoreOnDemand/Private/OnDemandIoDispatcherBackend.cpp
--- a/Engine/Source/Runtime/Experimental/IoStoreOnDemand/Private/OnDemandIoDispatcherBackend.cpp
+++ b/Engine/Source/Runtime/Experimental/IoStoreOnDemand/Private/OnDemandIoDispatcherBackend.cpp
@@ -2,6 +2,7 @@
 
 #include "OnDemandIoDispatcherBackend.h"
 
+#include "Algo/Transform.h"
 #include "AnalyticsEventAttribute.h"
 #include "Containers/BitArray.h"
 #include "Containers/StringView.h"
@@ -26,9 +27,11 @@
 #include "IO/IoStore.h"
 #include "IO/IoStoreOnDemand.h"
 #include "IasCache.h"
+#include "Logging/StructuredLog.h"
 #include "Math/NumericLimits.h"
 #include "Misc/CommandLine.h"
 #include "Misc/ConfigCacheIni.h"
+#include "Misc/CoreDelegatesInternal.h"
 #include "Misc/EncryptionKeyManager.h"
 #include "Misc/EnumClassFlags.h"
 #include "Misc/PathViews.h"
@@ -45,26 +48,20 @@
 
 #include <atomic>
 
-#if UE_IAS_LINKPAKFILE
-#include "IPlatformFilePak.h"
-#endif //UE_IAS_LINKPAKFILE
-
 #if !UE_BUILD_SHIPPING
 #include "Modules/ModuleManager.h"
 #endif 
 
-/** 
- * When enabledOnDemandToc files generated from disk will be compared to the downloaded version
- * and checked for compatibility issues. These checks will assert if a problem is found.
- */
-#define UE_VALIDATE_GENERATED_TOC (0 && !UE_BUILD_SHIPPING)
-
 /** When enabled the IAS system can add additional debug console commands for development use */
 #define UE_IAS_DEBUG_CONSOLE_CMDS (1 && !NO_CVARS && !UE_BUILD_SHIPPING)
 
 namespace UE::IO::IAS
 {
 
+///////////////////////////////////////////////////////////////////////////////
+
+extern FString GIasOnDemandTocExt;
+
 ///////////////////////////////////////////////////////////////////////////////
 int32 GIasHttpPrimaryEndpoint = 0;
 static FAutoConsoleVariableRef CVar_IasHttpPrimaryEndpoint(
@@ -187,31 +184,17 @@ bool GIasReportAnalyticsEnabled = true;
 static FAutoConsoleVariableRef CVar_IoReportAnalytics(
 	TEXT("ias.ReportAnalytics"),
 	GIasReportAnalyticsEnabled,
-	TEXT("Enables reporting statics to the analytics system"));
-
-bool GIasGenerateOnDemandToc = true;
-static FAutoConsoleVariableRef CVar_IasGenerateOnDemandToc(
-	TEXT("s.IasGenerateOnDemandToc"),
-	GIasGenerateOnDemandToc,
-	TEXT("Enables generating the FOnDemandToc from utoc files on disk rather than downloading them"),
-	ECVF_ReadOnly
-);
-
-// A temp fallback path allowing us to attempt to load the OnDemand toc from disk rather than
-// trying to generate it.
-bool GIasLoadOnDemandToc = true;
-static FAutoConsoleVariableRef CVar_IasLoadOnDemandToc(
-	TEXT("ias.LoadOnDemandToc"),
-	GIasLoadOnDemandToc,
-	TEXT("Attempts to load the OnDemand toc from disk rather than generating it (TEMP)"),
-	ECVF_ReadOnly
+	TEXT("Enables reporting statics to the analytics system")
 );
 
-bool GIasAsyncTocGenerationEnabled = true;
-static FAutoConsoleVariableRef CVar_IasAsyncTocGeneration(
-	TEXT("s.IasEnableThreadedTocGeneration"),
-	GIasAsyncTocGenerationEnabled,
-	TEXT("Enables pushing the work FOnDemandToc generation work to the task system"),
+int32 GIasTocMode = 0;
+static FAutoConsoleVariableRef CVar_IasTocMode(
+	TEXT("ias.TocMode"),
+	GIasTocMode,
+	TEXT("How should the IAS system load it's toc (see ETocMode).\n")
+	TEXT("0 = Load a single .iochunktoc\n")
+	TEXT("1 = Try to find a .uondemandtoc each time a pak file is mounted\n")
+	TEXT("2 = Download the toc from the target CDN"),
 	ECVF_ReadOnly
 );
 
@@ -259,6 +242,30 @@ static void LatencyTest(FStringView Url, FStringView Path)
 		Url.GetData(), Results[0], Results[1], Results[2], Results[3]);
 }
 #endif // !UE_BUILD_SHIPPING
+
+///////////////////////////////////////////////////////////////////////////////
+enum class ETocMode : int32
+{
+	LoadTocFromDisk = 0,	// <- Current Default
+	LoadTocFromMountedPaks,
+	LoadTocFromNetwork
+};
+
+static void ForceTocMode(ETocMode Mode)
+{
+	GIasTocMode = static_cast<int32>(Mode);
+}
+
+static ETocMode GetTocMode()
+{
+	if (GIasTocMode < 0 || GIasTocMode > static_cast<int32>(ETocMode::LoadTocFromNetwork))
+	{
+		UE_LOG(LogIas, Log, TEXT("ias.TocMode set to invalid value, defaulting to ETocMode::LoadTocFromDisk"));
+		return ETocMode::LoadTocFromDisk;
+	}
+
+	return static_cast<ETocMode>(GIasTocMode);
+}
 ///////////////////////////////////////////////////////////////////////////////
 static int32 LatencyTest(TConstArrayView<FString> Urls, FStringView Path, std::atomic_bool& bCancel)
 {
@@ -377,7 +384,9 @@ public:
 	FOnDemandIoStore();
 	~FOnDemandIoStore();
 
-	void AddToc(const FString& TocPath, FOnDemandToc&& Toc);
+	void AddToc(FStringView TocPath, FOnDemandToc&& Toc);
+	void RemoveToc(FStringView TocPath);
+
 	TIoStatusOr<uint64> GetChunkSize(const FIoChunkId& ChunkId);
 	FChunkInfo GetChunkInfo(const FIoChunkId& ChunkId);
 	FString GetFirstTocPath() const;
@@ -404,26 +413,37 @@ FOnDemandIoStore::~FOnDemandIoStore()
 	FEncryptionKeyManager::Get().OnKeyAdded().RemoveAll(this);
 }
 
-void FOnDemandIoStore::AddToc(const FString& TocPath, FOnDemandToc&& Toc)
+void FOnDemandIoStore::AddToc(FStringView TocPath, FOnDemandToc&& Toc)
 {
-	UE_LOG(LogIas, Log, TEXT("Adding TOC '%s'"), *TocPath);
+	UE_LOGFMT(LogIas, Log, "Adding TOC '{FileName}'", TocPath);
 
 	TRACE_CPUPROFILER_EVENT_SCOPE(IasBackend::AddToc);
 
-	FString Prefix;
+	TStringBuilder<128> ChunksDirectory;
+
 	{
+		// First attempt to parse the prefix from the given path (legacy)
 		int32 Idx = INDEX_NONE;
-		if (TocPath.FindLastChar(TCHAR('/'), Idx))
+		if (TocPath.FindLastChar(TEXT('/'), Idx))
+		{
+			Algo::Transform(TocPath.Left(Idx), AppendChars(ChunksDirectory), FChar::ToLower);
+		}
+
+		// If there was no prefix in the path we should just use the ChunksDirectory from the toc itself
+		if (ChunksDirectory.Len() == 0)
 		{
-			Prefix = TocPath.Left(Idx);
+			Algo::Transform(Toc.Header.ChunksDirectory, AppendChars(ChunksDirectory), FChar::ToLower);
+			
 		}
+
+		FPathViews::Append(ChunksDirectory, TEXT("chunks"));
 	}
 
 	{
 		FWriteScopeLock _(Lock);
 
 		const FOnDemandTocHeader& Header = Toc.Header;
-		FToc* NewToc = new(Tocs) FToc{TocPath};
+		FToc* NewToc = new(Tocs) FToc{FString(TocPath)};
 		NewToc->Containers.SetNum(Toc.Containers.Num()); // List of containers can never change
 
 		const FName CompressionFormat(Header.CompressionFormat);
@@ -434,7 +454,7 @@ void FOnDemandIoStore::AddToc(const FString& TocPath, FOnDemandToc&& Toc)
 			FContainer* NewContainer = &NewToc->Containers[ContainerIndex++];
 			NewContainer->Toc = NewToc;
 			NewContainer->Name = MoveTemp(Container.ContainerName);
-			NewContainer->ChunksDirectory = (Prefix.IsEmpty() ? Header.ChunksDirectory : Prefix / Header.ChunksDirectory).ToLower();
+			NewContainer->ChunksDirectory = ChunksDirectory;
 			NewContainer->CompressionFormat = CompressionFormat;
 			NewContainer->BlockSize = Header.BlockSize;
 			NewContainer->EncryptionKeyGuid = Container.EncryptionKeyGuid;
@@ -464,6 +484,12 @@ void FOnDemandIoStore::AddToc(const FString& TocPath, FOnDemandToc&& Toc)
 	AddDeferredContainers();
 }
 
+void FOnDemandIoStore::RemoveToc(FStringView TocPath)
+{
+	// TODO: Need to test this with content bundles!
+	checkNoEntry();
+}
+
 TIoStatusOr<uint64> FOnDemandIoStore::GetChunkSize(const FIoChunkId& ChunkId)
 {
 	if (FChunkInfo Info = GetChunkInfo(ChunkId))
@@ -542,7 +568,7 @@ void FOnDemandIoStore::AddDeferredContainers()
 			}
 			else
 			{
-				UE_LOG(LogIas, Log, TEXT("Defeering container '%s', encryption key '%s' not available"), *Container->Name, *Container->EncryptionKeyGuid);
+				UE_LOG(LogIas, Log, TEXT("Deferring container '%s', encryption key '%s' not available"), *Container->Name, *Container->EncryptionKeyGuid);
 			}
 		}
 	}
@@ -925,57 +951,6 @@ static void LogIoResult(
 		PrioToString(Priority));
 };
 
-/** Utility for finding all available on demand utoc files currently in valid pak directories */
-static TArray<FString> FindOnDemandUtocFilesOnDisk()
-{
-#if UE_IAS_LINKPAKFILE
-	// TODO: This line is all over the engine, should make it look a bit nicer
-	FPakPlatformFile* PakPlatformFile = static_cast<FPakPlatformFile*>(FPlatformFileManager::Get().FindPlatformFile(TEXT("PakFile")));
-	if (!PakPlatformFile)
-	{
-		return TArray<FString>();
-	}
-
-	TArray<FString> PakFolders;
-	PakPlatformFile->GetPakFolders(FCommandLine::Get(), PakFolders);
-
-	IPlatformFile* SearchFile = PakPlatformFile->GetLowerLevel();
-#else
-	TStringBuilder<260> PakPath;
-	FPathViews::Append(PakPath, FPaths::ProjectContentDir(), TEXT("Paks"));
-
-	TArray<FString> PakFolders;
-	PakFolders.Add(PakPath.ToString());
-
-	IPlatformFile* SearchFile = &FPlatformFileManager::Get().GetPlatformFile();
-#endif // UE_IAS_LINKPAKFILE
-
-	TArray<FString> FoundFiles;
-	for (const FString& Directory : PakFolders)
-	{
-		//PakPlatformFile
-		SearchFile->IterateDirectoryRecursively(*Directory, [&FoundFiles/*, PakPlatformFile*/](const TCHAR* Path, bool bIsDirectory) -> bool
-			{
-				if (!bIsDirectory)
-				{
-					FString Filename(Path);
-					if (Filename.EndsWith(TEXT(".utoc")) && Filename.Contains(TEXT("ondemand")))
-					{
-						// TODO: Cannot call IsPakFileInstalled at this point
-						//if (PakPlatformFile->IsPakFileInstalled(Filename))
-						{
-							FoundFiles.Emplace(MoveTemp(Filename));
-						}
-					}
-				}
-
-				return true;
-			});
-	}
-
-	return FoundFiles;
-}
-
 /**
  * Utility to create a FArchive capable of reading from disk using the exact same pathing
  * rules as FPlatformMisc::LoadTextFileFromPlatformPackage but without forcing the entire
@@ -1004,19 +979,16 @@ static TUniquePtr<FArchive> CreateReaderFromPlatformPackage(const FString& RelPa
 	}	
 }
 
-/** Generate a FOnDemandToc based on utoc files on disk which support the OnDemand feature */
-static TIoStatusOr<FOnDemandToc> GenerateOnDemandTocFromDisk(FStringView TocHash)
+/** Loads a single .iochunktoc (containing all of our container tocs) from disk */
+static TIoStatusOr<FOnDemandToc> LoadOnDemandTocFromDisk(FStringView TocHash)
 {
-	TRACE_CPUPROFILER_EVENT_SCOPE(IasBackend::GenerateOnDemandTocFromDisk);
+	TRACE_CPUPROFILER_EVENT_SCOPE(IasBackend::LoadOnDemandTocFromDisk);
 
 	FOnDemandToc OutToc;
 
-	if (GIasLoadOnDemandToc)
+	const ETocMode TocMode = GetTocMode();
+	if (TocMode == ETocMode::LoadTocFromDisk)
 	{
-		TRACE_CPUPROFILER_EVENT_SCOPE(IasBackend::LoadTocFromDisk);
-
-		UE_LOG(LogIas, Log, TEXT("Serializing .iochunktoc from disk"));
-
 		FString TocFileName = FString(TocHash);
 		TocFileName .Append(TEXT(".iochunktoc"));
 
@@ -1049,146 +1021,14 @@ static TIoStatusOr<FOnDemandToc> GenerateOnDemandTocFromDisk(FStringView TocHash
 			return FIoStatus(EIoErrorCode::NotFound, WriteToString<256>(TEXT("Unable to find '"), TocPath, TEXT("' on disk")));
 		}
 	}
-
-	UE_LOG(LogIas, Log, TEXT("Generating .iochunktoc via utoc files from disk"));
-
-	TArray<FString> UtocFilePaths = FindOnDemandUtocFilesOnDisk();
-	const TMap<FGuid, FAES::FAESKey> EncryptionKeys = FEncryptionKeyManager::Get().GetAllKeys();
-
-	for (const FString& UtocFilePath : UtocFilePaths)
-	{
-		FIoStoreReader Reader;
-		FIoStatus Status = Reader.Initialize(FPathViews::GetBaseFilenameWithPath(UtocFilePath), EncryptionKeys);
-		if (Status.IsOk() && EnumHasAnyFlags(Reader.GetContainerFlags(), EIoContainerFlags::OnDemand))
-		{
-			FOnDemandTocContainerEntry Container;
-
-			Container.ContainerName = FPathViews::GetBaseFilename(UtocFilePath);
-
-			const uint32 BlockSize = Reader.GetCompressionBlockSize();
-			if (OutToc.Header.BlockSize == 0)
-			{
-				OutToc.Header.BlockSize = Reader.GetCompressionBlockSize();
-			}
-			check(OutToc.Header.BlockSize == Reader.GetCompressionBlockSize());
-
-			TArray<FIoStoreTocChunkInfo> ChunkInfos;
-			Reader.EnumerateChunks([&ChunkInfos](FIoStoreTocChunkInfo&& Info)
-				{
-					ChunkInfos.Emplace(MoveTemp(Info));
-					return true;
-				});
-
-			// We can't actually hit this until we solve the FASEKey Initialize issue above
-			if (EnumHasAnyFlags(Reader.GetContainerFlags(), EIoContainerFlags::Encrypted))
-			{
-				Container.EncryptionKeyGuid = LexToString(Reader.GetEncryptionKeyGuid());
-			}
-
-			for (const FIoStoreTocChunkInfo& ChunkInfo : ChunkInfos)
-			{
-				TIoStatusOr<FIoStoreCompressedChunkInfo> InfoStatus = Reader.GetChunkCompressedInfo(ChunkInfo.Id);
-				if (!InfoStatus.IsOk())
-				{
-					return InfoStatus.Status();
-				}
-				FIoStoreCompressedChunkInfo CompressedChunkInfo = InfoStatus.ConsumeValueOrDie();
-
-				const uint32 BlockOffset = Container.BlockSizes.Num();
-				const uint32 BlockCount = CompressedChunkInfo.Blocks.Num();
-
-				uint64 RawChunkSize = 0;
-				uint64 EncodedChunkSize = 0;
-				for (const FIoStoreCompressedBlockInfo& BlockInfo : CompressedChunkInfo.Blocks)
-				{
-					check(Align(BlockInfo.CompressedSize, FAES::AESBlockSize) == BlockInfo.AlignedSize);
-					const uint64 EncodedBlockSize = BlockInfo.AlignedSize;
-					Container.BlockSizes.Add(uint32(BlockInfo.CompressedSize));
-
-					FIoBlockHash BlockHash;
-					FMemory::Memcpy(&BlockHash, &BlockInfo.DiskHash, sizeof(FIoBlockHash));
-					Container.BlockHashes.Add(BlockHash);
-
-					EncodedChunkSize += EncodedBlockSize;
-					RawChunkSize += BlockInfo.UncompressedSize;
-
-					if (OutToc.Header.CompressionFormat.IsEmpty() && BlockInfo.CompressionMethod != NAME_None)
-					{
-						OutToc.Header.CompressionFormat = BlockInfo.CompressionMethod.ToString();
-					}
-				}
-
-				FOnDemandTocEntry& TocEntry = Container.Entries.AddDefaulted_GetRef();
-				TocEntry.ChunkId = ChunkInfo.Id;
-				TocEntry.Hash = CompressedChunkInfo.DiskHash;
-				TocEntry.RawSize = RawChunkSize;
-				TocEntry.EncodedSize = EncodedChunkSize;
-				TocEntry.BlockOffset = BlockOffset;
-				TocEntry.BlockCount = BlockCount;
-			}
-
-			OutToc.Containers.Emplace(MoveTemp(Container));
-		}
-	}
-
-	if (!OutToc.Containers.IsEmpty())
+	else
 	{
-		OutToc.Header.ChunksDirectory = TEXT("chunks");
+		checkNoEntry();
 	}
 
 	return OutToc;
 }
 
-/** Validation code used during development to ensure that the results are correct */
-#if UE_VALIDATE_GENERATED_TOC
-
-static TArray<const FOnDemandTocContainerEntry*> SortContainers(const TArray<FOnDemandTocContainerEntry>& Containers)
-{
-	TArray<const FOnDemandTocContainerEntry*> SortedContainers;
-	for (const FOnDemandTocContainerEntry& Container : Containers)
-	{
-		SortedContainers.Add(&Container);
-	}
-
-	Algo::Sort(SortedContainers, [](const FOnDemandTocContainerEntry* LHS, const FOnDemandTocContainerEntry* RHS)->bool
-		{
-			return LHS->ContainerName < RHS->ContainerName;
-		});
-
-	return SortedContainers;
-}
-
-static bool operator == (const FOnDemandTocEntry& LHS, const FOnDemandTocEntry& RHS)
-{
-	return FMemory::Memcmp(&LHS, &RHS, sizeof(FOnDemandTocEntry)) == 0;
-}
-
-static void ValidateToc(const FOnDemandToc& RefToc, const FOnDemandToc& NewToc)
-{
-	check(RefToc.Header.BlockSize == NewToc.Header.BlockSize);
-	check(RefToc.Header.CompressionFormat == NewToc.Header.CompressionFormat);
-	check(RefToc.Header.ChunksDirectory == NewToc.Header.ChunksDirectory);
-
-	TArray<const FOnDemandTocContainerEntry*> RefContainers = SortContainers(RefToc.Containers);
-	TArray<const FOnDemandTocContainerEntry*> NewContainers = SortContainers(NewToc.Containers);
-
-	check(RefContainers.Num() == NewContainers.Num());
-
-	for (int32 Index = 0; Index < RefContainers.Num(); ++Index)
-	{
-		const FOnDemandTocContainerEntry* RefContainer = RefContainers[Index];
-		const FOnDemandTocContainerEntry* NewContainer = NewContainers[Index];
-
-		check(RefContainer->ContainerName == NewContainer->ContainerName);
-
-		check(RefContainer->Entries == NewContainer->Entries);
-		check(RefContainer->BlockSizes == NewContainer->BlockSizes);
-		check(RefContainer->BlockHashes == NewContainer->BlockHashes);
-	}
-}
-
-#endif //UE_VALIDATE_GENERATED_TOC
-
 ///////////////////////////////////////////////////////////////////////////////
 struct FBackendStatus
 {
@@ -1554,12 +1394,15 @@ public:
 
 private:
 
+	void MountContainer(FStringView ContainerPath);
+	void UnmountContainer(FStringView ContainerPath);
+
 	FString GetEndpointTestPath() const;
 	void ConditionallyStartBackendThread();
 	void CompleteRequest(FChunkRequest* ChunkRequest);
 	void CompleteMaterialize(FChunkRequest* ChunkRequest);
 
-	FIoStatus ApplyGeneratedOnDemandToc(const FString& CdnUrl, const FString& TocPath);
+	FIoStatus ApplyLoadedOnDemandToc(const FString& TocPath);
 	FIoStatus DownloadoadOnDemandToc(const FString& CdnUrl, const FString& TocPath);
 
 	bool ResolveDistributedEndpoint(const FDistributedEndpointUrl& Url);
@@ -1606,12 +1449,16 @@ private:
 	FDistributedEndpointUrl DistributionUrl;
 	FEventRef DistributedEndpointEvent;
 
+	FString EndpointTestPath;
+
+	FDelegateHandle PakDelegateHandle;
+
 	mutable FRWLock Lock;
 	std::atomic_uint32_t InflightCacheRequestCount{0};
 	std::atomic_bool bStopRequested{false};
 
 	bool bGeneratedOnDemandToc = false;
-	UE::Tasks::TTask<TIoStatusOr<FOnDemandToc>> OnDemandTocTask;
+	UE::Tasks::TTask<TIoStatusOr<FOnDemandToc>> LoadingOnDemandTocTask;
 
 #if UE_IAS_DEBUG_CONSOLE_CMDS
 	TArray<IConsoleCommand*> DynamicConsoleCommands;
@@ -1645,6 +1492,11 @@ FOnDemandIoBackend::FOnDemandIoBackend(TUniquePtr<IIasCache>&& InCache)
 
 FOnDemandIoBackend::~FOnDemandIoBackend()
 {
+	if (PakDelegateHandle.IsValid())
+	{
+		FCoreInternalDelegates::GetOnPakMountOperation().Remove(PakDelegateHandle);
+	}
+
 #if UE_IAS_DEBUG_CONSOLE_CMDS
 	for (IConsoleCommand* Cmd : DynamicConsoleCommands)
 	{
@@ -1690,8 +1542,56 @@ void FOnDemandIoBackend::Shutdown()
 	BackendContext.Reset();
 }
 
+void FOnDemandIoBackend::MountContainer(FStringView ContainerPath)
+{
+	// TODO: For now we just log an error if a container cannot be mounted, but for production
+	// we will need to make sure that the user is informed so that they can verify their installation.
+	IPlatformFile& PlatformFile = IPlatformFile::GetPlatformPhysical();
+
+	const FString TocPath = FPathViews::ChangeExtension(ContainerPath, GIasOnDemandTocExt);
+
+	if (PlatformFile.FileExists(*TocPath))
+	{
+		UE_LOG(LogIas, Log, TEXT("Mounting '%s'"), *TocPath);
+
+		TUniquePtr<FArchive> Ar(IFileManager::Get().CreateFileReader(*TocPath));
+
+		if (!Ar.IsValid())
+		{
+			UE_LOG(LogIas, Error, TEXT("Failed to open '%s' for reading"), *TocPath);
+			return;
+		}
+
+		FOnDemandToc ContainerToc;
+		(*Ar) << ContainerToc;
+
+		if (Ar->IsError() || Ar->IsCriticalError())
+		{
+			UE_LOG(LogIas, Error, TEXT("Failed to serialize '%s'"), *TocPath);
+			return;
+		}
+
+		const FStringView FileName = FPathViews::GetCleanFilename(TocPath);
+		IoStore->AddToc(FileName, MoveTemp(ContainerToc));
+	}
+}
+
+void FOnDemandIoBackend::UnmountContainer(FStringView ContainerPath)
+{
+	const FString TocPath = FPathViews::ChangeExtension(ContainerPath, GIasOnDemandTocExt);
+	const FStringView FileName = FPathViews::GetCleanFilename(TocPath);
+
+	IoStore->RemoveToc(FileName);
+}
+
 FString FOnDemandIoBackend::GetEndpointTestPath() const
 {
+	if (!EndpointTestPath.IsEmpty())
+	{
+		return EndpointTestPath;
+	}
+
+	// Older fallback path, it probably shouldn't be possible to get this far and can be removed later.
 	FString TestPath = IoStore.IsValid() ? IoStore->GetFirstTocPath() : FString();
 	if (TestPath.IsEmpty())
 	{
@@ -1999,12 +1899,10 @@ TIoStatusOr<FIoMappedRegion> FOnDemandIoBackend::OpenMapped(const FIoChunkId& Ch
 	return FIoStatus::Unknown;
 }
 
-FIoStatus FOnDemandIoBackend::ApplyGeneratedOnDemandToc(const FString& CdnUrl, const FString& TocPath)
+FIoStatus FOnDemandIoBackend::ApplyLoadedOnDemandToc(const FString& TocPath)
 {
-	TRACE_CPUPROFILER_EVENT_SCOPE(IasBackend::ApplyGeneratedOnDemandToc);
+	TRACE_CPUPROFILER_EVENT_SCOPE(IasBackend::ApplyLoadedOnDemandToc);
 
-	// TODO: Assumes we only ever want to generate/load a .iochunktoc once, a safe assumption
-	// now but probably not in the future.
 	if (bGeneratedOnDemandToc)
 	{
 		return FIoStatus::Ok;
@@ -2012,13 +1910,13 @@ FIoStatus FOnDemandIoBackend::ApplyGeneratedOnDemandToc(const FString& CdnUrl, c
 
 	TIoStatusOr<FOnDemandToc> GeneratedTocResult;
 
-	if (OnDemandTocTask.IsValid())
+	if (LoadingOnDemandTocTask.IsValid())
 	{
-		GeneratedTocResult = MoveTemp(OnDemandTocTask.GetResult());
+		GeneratedTocResult = MoveTemp(LoadingOnDemandTocTask.GetResult());
 	}
 	else
 	{
-		GeneratedTocResult = GenerateOnDemandTocFromDisk(FPathViews::GetBaseFilename(TocPath));
+		GeneratedTocResult = LoadOnDemandTocFromDisk(FPathViews::GetBaseFilename(TocPath));
 	}
 	
 	if (!GeneratedTocResult.IsOk())
@@ -2026,16 +1924,7 @@ FIoStatus FOnDemandIoBackend::ApplyGeneratedOnDemandToc(const FString& CdnUrl, c
 		return GeneratedTocResult.Status();
 	}
 
-#if UE_VALIDATE_GENERATED_TOC == 0
 	IoStore->AddToc(TocPath, GeneratedTocResult.ConsumeValueOrDie());
-#else
-	FOnDemandToc UrlToc = LoadTocFromUrl(CdnUrl, TocPath, 1).ConsumeValueOrDie();
-	FOnDemandToc GeneratedToc = GeneratedTocResult.ConsumeValueOrDie();
-
-	UE::IO::IAS::ValidateToc(UrlToc, GeneratedToc);
-
-	IoStore->AddToc(TocPath, MoveTemp(GeneratedToc));
-#endif // UE_VALIDATE_GENERATED_TOC
 
 	bGeneratedOnDemandToc = true;
 
@@ -2156,16 +2045,15 @@ void FOnDemandIoBackend::FlushDeferredTocs(EFlushMode FlushMode)
 		Tocs = MoveTemp(DeferredTocs);
 	}
 
+	const ETocMode TocMode = GetTocMode();
+
 	for (FTocParams& TocParams : Tocs)
 	{
-		if (GIasGenerateOnDemandToc && !TocParams.bForceDownload)
+		if (TocMode == ETocMode::LoadTocFromDisk && !TocParams.bForceDownload)
 		{
 			if (EnumHasAnyFlags(FlushMode, EFlushMode::Disk))
 			{
-				// TODO: This value if only used if UE_VALIDATE_GENERATED_TOC is enabled. The validation
-				// should be moved elsewhere so we don't need to do this work when it won't be used.
-				const FString CdnUrl = AvailableEps.HasCurrent() ? AvailableEps.GetCurrent() : FString();
-				FIoStatus Result = ApplyGeneratedOnDemandToc(CdnUrl, TocParams.Path);
+				FIoStatus Result = ApplyLoadedOnDemandToc(TocParams.Path);
 				if (!Result.IsOk())
 				{
 					UE_LOG(LogIas, Error, TEXT("Failed to add generated toc', reason '%s'"), *Result.ToString());
@@ -2207,20 +2095,53 @@ void FOnDemandIoBackend::Mount(const FOnDemandEndpoint& Endpoint)
 		return;
 	}
 
+	EndpointTestPath = Endpoint.TocPath;
+
 	// TODO: Should pass this info around via DeferredTocs
 	if (Endpoint.bForceTocDownload)
 	{
-		GIasGenerateOnDemandToc = false;
+		ForceTocMode(ETocMode::LoadTocFromNetwork);
 	}
 
-	if (GIasGenerateOnDemandToc && GIasAsyncTocGenerationEnabled)
+	if (GetTocMode() == ETocMode::LoadTocFromDisk)
 	{
 		FString TocHash = FPaths::GetBaseFilename(Endpoint.TocPath);
-		OnDemandTocTask = UE::Tasks::Launch(UE_SOURCE_LOCATION, [TocHash = MoveTemp(TocHash)]() -> TIoStatusOr<FOnDemandToc>
+		LoadingOnDemandTocTask = UE::Tasks::Launch(UE_SOURCE_LOCATION, [TocHash = MoveTemp(TocHash)]() -> TIoStatusOr<FOnDemandToc>
 			{
-				return GenerateOnDemandTocFromDisk(TocHash);
+				return LoadOnDemandTocFromDisk(TocHash);
 			});
 	}
+	else if (GetTocMode() == ETocMode::LoadTocFromMountedPaks)
+	{
+		// First make sure that we receive notifications of any new pakfiles being mounted so that we can react to them
+		PakDelegateHandle = FCoreInternalDelegates::GetOnPakMountOperation().AddLambda(
+			[this](EMountOperation Operation, const TCHAR* ContainerPath, int32 Order) -> void
+			{
+				switch (Operation)
+				{
+					case EMountOperation::Mount:
+						this->MountContainer(ContainerPath);
+						break;
+					case EMountOperation::Unmount:
+						this->UnmountContainer(ContainerPath);
+						break;
+					default:
+						checkNoEntry();
+				}
+			}
+		);
+
+		FCurrentlyMountedPaksDelegate& Delegate = FCoreInternalDelegates::GetCurrentlyMountedPaksDelegate();
+		if (Delegate.IsBound())
+		{
+			TArray<FMountedPakInfo> PreExistingMountedPaks = Delegate.Execute();
+			for (const FMountedPakInfo& Info : PreExistingMountedPaks)
+			{
+				check(Info.PakFile != nullptr);
+				MountContainer(Info.PakFile->PakGetPakFilename());
+			}
+		}
+	}
 
 	{
 		FWriteScopeLock _(Lock);
@@ -2232,7 +2153,12 @@ void FOnDemandIoBackend::Mount(const FOnDemandEndpoint& Endpoint)
 				{
 					DistributionUrl = { Endpoint.DistributionUrl, Endpoint.FallbackUrl };
 				}
-				DeferredTocs.Add(FTocParams{ Endpoint.TocPath, Endpoint.bForceTocDownload });
+
+				if (GetTocMode() != ETocMode::LoadTocFromMountedPaks)
+				{
+					DeferredTocs.Add(FTocParams{ Endpoint.TocPath, Endpoint.bForceTocDownload });
+				}
+
 				return;
 			}
 		}
@@ -2248,15 +2174,7 @@ void FOnDemandIoBackend::Mount(const FOnDemandEndpoint& Endpoint)
 
 	check(AvailableEps.HasCurrent());
 
-	if (GIasGenerateOnDemandToc && !Endpoint.bForceTocDownload)
-	{
-		FIoStatus GeneratedResult = ApplyGeneratedOnDemandToc(AvailableEps.GetCurrent(), Endpoint.TocPath);
-		if (!GeneratedResult.IsOk())
-		{
-			UE_LOG(LogIas, Error, TEXT("Failed to add generated toc', reason '%s'"), *GeneratedResult.ToString());
-		}
-	}
-	else
+	if (GetTocMode() == ETocMode::LoadTocFromNetwork)
 	{
 		FIoStatus Result = DownloadoadOnDemandToc(AvailableEps.GetCurrent(), Endpoint.TocPath);
 		if (!Result.IsOk())
@@ -2267,16 +2185,36 @@ void FOnDemandIoBackend::Mount(const FOnDemandEndpoint& Endpoint)
 			DeferredTocs.Add(FTocParams{ Endpoint.TocPath, Endpoint.bForceTocDownload });
 		}
 	}
+	else if (GetTocMode() == ETocMode::LoadTocFromDisk)
+	{
+		FIoStatus GeneratedResult = ApplyLoadedOnDemandToc(Endpoint.TocPath);
+		if (!GeneratedResult.IsOk())
+		{
+			UE_LOG(LogIas, Error, TEXT("Failed to add generated toc', reason '%s'"), *GeneratedResult.ToString());
+		}	
+	}
 }
 
 void FOnDemandIoBackend::SetBulkOptionalEnabled(bool bEnabled)
 {
-	BackendStatus.SetHttpOptionalBulkEnabled(bEnabled);
+	// Ignore enable/disable messages if we are mounted from the pak file system
+	// TODO: Remove SetBulkOptionalEnabled entirely once LoadTocFromMountedPaks is
+	// the only valid mode
+	if (GetTocMode() != ETocMode::LoadTocFromMountedPaks)
+	{
+		BackendStatus.SetHttpOptionalBulkEnabled(bEnabled);
+	}
 }
 
 void FOnDemandIoBackend::SetEnabled(bool bEnabled)
 {
-	BackendStatus.SetHttpEnabled(bEnabled);
+	// Ignore enable/disable messages if we are mounted from the pak file system
+	// TODO: Remove SetEnabled entirely once LoadTocFromMountedPaks is the only
+	// valid mode
+	if (GetTocMode() != ETocMode::LoadTocFromMountedPaks)
+	{
+		BackendStatus.SetHttpEnabled(bEnabled);
+	}
 }
 
 bool FOnDemandIoBackend::IsEnabled() const
@@ -2597,4 +2535,3 @@ TSharedPtr<IOnDemandIoDispatcherBackend> MakeOnDemandIoDispatcherBackend(TUnique
 } // namespace UE::IO::IAS
 
 #undef UE_IAS_DEBUG_CONSOLE_CMDS
-#undef UE_VALIDATE_GENERATED_TOC
diff --git a/Engine/Source/Runtime/Experimental/IoStoreOnDemand/Public/IO/IoStoreOnDemand.h b/Engine/Source/Runtime/Experimental/IoStoreOnDemand/Public/IO/IoStoreOnDemand.h
--- a/Engine/Source/Runtime/Experimental/IoStoreOnDemand/Public/IO/IoStoreOnDemand.h
+++ b/Engine/Source/Runtime/Experimental/IoStoreOnDemand/Public/IO/IoStoreOnDemand.h
@@ -117,6 +117,8 @@ struct FOnDemandTocContainerEntry
 	TArray<FOnDemandTocEntry> Entries;
 	TArray<uint32> BlockSizes;
 	TArray<FIoBlockHash> BlockHashes;
+
+	/** Hash of the .utoc file (on disk) used to generate this data */
 	FIoHash UTocHash;
 
 	UE_API friend FArchive& operator<<(FArchive& Ar, FOnDemandTocContainerEntry& ContainerEntry);
@@ -176,14 +178,22 @@ struct FIoStoreUploadParams
 	bool bWriteTocToDisk = false;
 	/** Where the .iochunktoc file should be written out. */
 	FString TocOutputDir;
-	
+
 	UE_API static TIoStatusOr<FIoStoreUploadParams> Parse(const TCHAR* CommandLine);
 };
 
+/** Results from uploading a FOnDemandToc */
 struct FIoStoreUploadResult
 {
+	/** Hash of the toc when written as a binary blob */
 	FIoHash TocHash;
+	
+	/** Url of the service that the toc was uploaded too */
+	FString ServiceUrl;
+	/** Path of the toc on the service */
 	FString TocPath;
+
+	/** Size (in bytes) of the toc when written as a binary blob */
 	uint64 TocSize = 0;
 };
 
