Chaos - fix island wake when moving/removing a kinematic body in the island

This was a couple problems
- collisions were getting reset when switching a body to kinematic (via a call to SetParticleTransform)
- islands were not being woken when a body was moved

diff --git a/Engine/Source/Runtime/Experimental/Chaos/Private/Chaos/Island/IslandManager.cpp b/Engine/Source/Runtime/Experimental/Chaos/Private/Chaos/Island/IslandManager.cpp
--- a/Engine/Source/Runtime/Experimental/Chaos/Private/Chaos/Island/IslandManager.cpp
+++ b/Engine/Source/Runtime/Experimental/Chaos/Private/Chaos/Island/IslandManager.cpp
@@ -736,17 +736,30 @@ namespace Chaos::Private
 		return IslandConstraints;
 	}
 
-	void FPBDIslandManager::WakeParticleIsland(FGeometryParticleHandle* Particle)
+	void FPBDIslandManager::WakeParticleIslands(FGeometryParticleHandle* Particle)
 	{
+		// We want to wake for at least one frame (i.e., prevent re-sleeping this frame)
+		const bool bIsSleepAllowed = false;
+
 		// When we explicitly wake we reset sleep counters etc
 		if (FPBDIslandParticle* Node = GetGraphNode(Particle))
 		{
-			// If we are in an island flag for sleep reset
+			// NOTE: We could check Flags.bIsDynamic here, but checking the Island pointer means we 
+			// are automatically handling the case where a kinematic with constraints on it 
+			// was made dynamic right before calling WakeParticleIsland
 			if (Node->Island != nullptr)
 			{
-				const bool bIsSleepAllowed = false;
+				// This is a dynamic particle and we know its island
 				EnqueueIslandCheckSleep(Node->Island, bIsSleepAllowed);
 			}
+			else
+			{
+				// This is a kinematic particle, visit our edges to find the islands we are in
+				for (const FPBDIslandConstraint* Edge : Node->Edges)
+				{
+					EnqueueIslandCheckSleep(Edge->Island, bIsSleepAllowed);
+				}
+			}
 		}
 
 		// If we are an isolated particle we keep our own sleep counter
@@ -844,6 +857,10 @@ namespace Chaos::Private
 			FPBDIslandParticle* Node0 = Edge->Nodes[0];
 			FPBDIslandParticle* Node1 = Edge->Nodes[1];
 
+			// Wake the island when we remove a constraint
+			const bool bIsSleepAllowed = true;
+			EnqueueIslandCheckSleep(Edge->Island, bIsSleepAllowed);
+
 			// Remove edge from the island
 			RemoveEdgeFromIsland(Edge);
 
@@ -907,6 +924,15 @@ namespace Chaos::Private
 		}
 	}
 
+	void FPBDIslandManager::WakeConstraintIsland(FConstraintHandle* Constraint)
+	{
+		if (FPBDIslandConstraint* Edge = GetGraphEdge(Constraint))
+		{
+			const bool bIsSleepAllowed = false;
+			EnqueueIslandCheckSleep(Edge->Island, bIsSleepAllowed);
+		}
+	}
+
 	void FPBDIslandManager::SetParticleIslandIsSleeping(FGeometryParticleHandle* Particle, const bool bInIsSleeping)
 	{
 		if (FPBDIslandParticle* Node = GetGraphNode(Particle))
diff --git a/Engine/Source/Runtime/Experimental/Chaos/Private/Chaos/PBDRigidsEvolution.cpp b/Engine/Source/Runtime/Experimental/Chaos/Private/Chaos/PBDRigidsEvolution.cpp
--- a/Engine/Source/Runtime/Experimental/Chaos/Private/Chaos/PBDRigidsEvolution.cpp
+++ b/Engine/Source/Runtime/Experimental/Chaos/Private/Chaos/PBDRigidsEvolution.cpp
@@ -1107,7 +1107,7 @@ namespace Chaos
 
 			// Reset the sleep counters etc in the sleep management system
 			// @todo(chaos): ideally we would just set reset the sleepiness of the particle
-			IslandManager.WakeParticleIsland(Particle);
+			IslandManager.WakeParticleIslands(Particle);
 		}
 	}
 
diff --git a/Engine/Source/Runtime/Experimental/Chaos/Private/Chaos/PBDRigidsEvolutionGBF.cpp b/Engine/Source/Runtime/Experimental/Chaos/Private/Chaos/PBDRigidsEvolutionGBF.cpp
--- a/Engine/Source/Runtime/Experimental/Chaos/Private/Chaos/PBDRigidsEvolutionGBF.cpp
+++ b/Engine/Source/Runtime/Experimental/Chaos/Private/Chaos/PBDRigidsEvolutionGBF.cpp
@@ -1104,13 +1104,23 @@ void FPBDRigidsEvolutionGBF::SetParticleKinematicTarget(FGeometryParticleHandle*
 
 void FPBDRigidsEvolutionGBF::OnParticleMoved(FGeometryParticleHandle* InParticle, const FVec3& PrevX, const FRotation3& PrevR, const bool bIsTeleport)
 {
-	// When a particle is moved, we need to tell the collisions because they cache friction state and 
-	// would attempt to undo small translations within the friction cone
-	InParticle->ParticleCollisions().VisitCollisions([this, InParticle](FPBDCollisionConstraint& Collision)
-		{
-			Collision.UpdateParticleTransform(InParticle);
-			return ECollisionVisitorResult::Continue;
-		});
+	// When a particle is moved, we need to 
+	// - tell the collisions because they cache friction state and would attempt to undo small translations within the friction cone
+	// - wake the island(s) that the particle is in
+	// NOTE: we have a tolerance on the transform change because SetParticleTransform may be called with the "same" transform that
+	// is different by very small amounts around 1e-7 in both position and rotation when switching from dynamic to kinematic.
+	const FReal CollisionPositionTolerance = FReal(1.e-4);
+	const FReal CollisionRotationTolerance = FReal(1.e-6);
+	if (!FVec3::IsNearlyEqual(PrevX, InParticle->X(), CollisionPositionTolerance) || !FRotation3::IsNearlyEqual(PrevR, InParticle->R(), CollisionRotationTolerance))
+	{
+		GetIslandManager().WakeParticleIslands(InParticle);
+
+		InParticle->ParticleCollisions().VisitCollisions([this, InParticle](FPBDCollisionConstraint& Collision)
+			{
+				Collision.UpdateParticleTransform(InParticle);
+				return ECollisionVisitorResult::Continue;
+			});
+	}
 
 	// If this is a teleport, we assume other state has been set separately (V, W, etc) if required, which 
 	// would include resetting any sleep-related state. If this is not a teleport, we want to prevent 
diff --git a/Engine/Source/Runtime/Experimental/Chaos/Public/Chaos/Island/IslandManager.h b/Engine/Source/Runtime/Experimental/Chaos/Public/Chaos/Island/IslandManager.h
--- a/Engine/Source/Runtime/Experimental/Chaos/Public/Chaos/Island/IslandManager.h
+++ b/Engine/Source/Runtime/Experimental/Chaos/Public/Chaos/Island/IslandManager.h
@@ -401,7 +401,7 @@ namespace Chaos::Private
 		CHAOS_API int32 ReserveParticles(const int32 InNumParticles);
 		CHAOS_API void UpdateParticleMaterial(FGeometryParticleHandle* Particle);
 		CHAOS_API int32 GetParticleLevel(FGeometryParticleHandle* Particle) const;
-		CHAOS_API void WakeParticleIsland(FGeometryParticleHandle* Particle);
+		CHAOS_API void WakeParticleIslands(FGeometryParticleHandle* Particle);
 
 		// Constraint management
 		int32 GetNumConstraints() const { return Edges.Num(); }
@@ -411,6 +411,7 @@ namespace Chaos::Private
 		CHAOS_API void RemoveParticleConstraints(FGeometryParticleHandle* Particle);
 		CHAOS_API void RemoveParticleContainerConstraints(FGeometryParticleHandle* Particle, const int32 ContainerId);
 		CHAOS_API void RemoveContainerConstraints(const int32 ContainerId);
+		CHAOS_API void WakeConstraintIsland(FConstraintHandle* Constraint);
 
 		// Access to Islands
 		int32 GetNumIslands() const { return Islands.Num(); }
