[ChaosDebugDraw] Prevent actor creation if already present in the world. Could happen on a non game Editor world on ReInitWorld (e.g. called from RebuildMap)

diff --git a/Engine/Source/Runtime/Experimental/ChaosSolverEngine/Private/Chaos/ChaosDebugDrawComponent.cpp b/Engine/Source/Runtime/Experimental/ChaosSolverEngine/Private/Chaos/ChaosDebugDrawComponent.cpp
--- a/Engine/Source/Runtime/Experimental/ChaosSolverEngine/Private/Chaos/ChaosDebugDrawComponent.cpp
+++ b/Engine/Source/Runtime/Experimental/ChaosSolverEngine/Private/Chaos/ChaosDebugDrawComponent.cpp
@@ -5,6 +5,9 @@
 #include "ChaosLog.h"
 #include "Debug/DebugDrawService.h"
 #include "DrawDebugHelpers.h"
+#if WITH_EDITOR
+#include "EngineUtils.h"
+#endif
 #include "Engine/Canvas.h"
 #include "Engine/Engine.h"
 #include "Engine/World.h"
@@ -365,7 +368,21 @@ void UChaosDebugDrawComponent::CreateDebugDrawActor(UWorld* World)
 
 #if WITH_EDITOR
 	Params.bHideFromSceneOutliner = true;
-#endif
+
+	// Make sure to not create more than one actor for non game worlds.
+	// Those can get reinitialized and OnPostWorldInitialization is called more than once
+	if (!World->IsGameWorld())
+	{
+		for (TActorIterator<AActor> It(World); It; ++It)
+		{
+			const AActor* Actor = *It;
+			if (Actor != nullptr && Actor->GetFName()== NAME_ChaosDebugDrawActor)
+			{
+				return;
+			}
+		}
+	}	
+#endif // WITH_EDITOR
 	
 	AActor* Actor = World->SpawnActor<AActor>(FVector::ZeroVector, FRotator::ZeroRotator, Params);
 	
