GeometryCollection : Override for default vertex color

diff --git a/Engine/Source/Runtime/Experimental/Chaos/Private/GeometryCollection/GeometryCollection.cpp b/Engine/Source/Runtime/Experimental/Chaos/Private/GeometryCollection/GeometryCollection.cpp
--- a/Engine/Source/Runtime/Experimental/Chaos/Private/GeometryCollection/GeometryCollection.cpp
+++ b/Engine/Source/Runtime/Experimental/Chaos/Private/GeometryCollection/GeometryCollection.cpp
@@ -49,10 +49,11 @@ bool FGeometryCollection::AreCollisionParticlesEnabled()
 }
 
 
-FGeometryCollection::FGeometryCollection()
+FGeometryCollection::FGeometryCollection(FGeometryCollectionDefaults InDefaults)
 	: FTransformCollection()
 	, FGeometryCollectionConvexPropertiesInterface(this)
 	, FGeometryCollectionProximityPropertiesInterface(this)
+	, Defaults(InDefaults)
 {
 	Construct();
 }
@@ -170,7 +171,7 @@ void FGeometryCollection::SetDefaults(FName Group, uint32 StartSize, uint32 NumE
 	{
 		for (uint32 Idx = StartSize; Idx < StartSize + NumElements; ++Idx)
 		{
-			Color[Idx] = FLinearColor::White;
+			Color[Idx] = Defaults.DefaultVertexColor;
 		}
 	}
 }
@@ -1572,9 +1573,9 @@ bool FGeometryCollection::IsVisible(int32 Element) const
 	return false;;
 }
 
-FGeometryCollection* FGeometryCollection::NewGeometryCollection(const TArray<float>& RawVertexArray, const TArray<int32>& RawIndicesArray, bool ReverseVertexOrder)
+FGeometryCollection* FGeometryCollection::NewGeometryCollection(const TArray<float>& RawVertexArray, const TArray<int32>& RawIndicesArray, bool ReverseVertexOrder, const FGeometryCollectionDefaults RawDefaults)
 {
-	FGeometryCollection* Collection = new FGeometryCollection();
+	FGeometryCollection* Collection = new FGeometryCollection(RawDefaults);
 	FGeometryCollection::Init(Collection, RawVertexArray, RawIndicesArray, ReverseVertexOrder);
 	return Collection;
 }
@@ -1616,7 +1617,7 @@ void FGeometryCollection::Init(FGeometryCollection* Collection, const TArray<flo
 			TempVertices += Vertices[Idx];
 			(*UV0)[Idx] = FVector2f::ZeroVector;
 
-			Colors[Idx] = FLinearColor::White;
+			Colors[Idx] = Collection->Defaults.DefaultVertexColor;
 			BoneMap[Idx] = 0;
 		}
 
@@ -1903,9 +1904,10 @@ FGeometryCollection* FGeometryCollection::NewGeometryCollection(const TArray<flo
 																const TManagedArray<int32>& RawParentArray,
 																const TManagedArray<TSet<int32>>& RawChildrenArray,
 																const TManagedArray<int32>& RawSimulationTypeArray,
-															    const TManagedArray<int32>& RawStatusFlagsArray)
+															    const TManagedArray<int32>& RawStatusFlagsArray,
+																const FGeometryCollectionDefaults RawDefaults)
 {
-	FGeometryCollection* RestCollection = new FGeometryCollection();
+	FGeometryCollection* RestCollection = new FGeometryCollection(RawDefaults);
 
 	int NumNewVertices = RawVertexArray.Num() / 3;
 	int VerticesIndex = RestCollection->AddElements(NumNewVertices, FGeometryCollection::VerticesGroup);
@@ -1941,7 +1943,7 @@ FGeometryCollection* FGeometryCollection::NewGeometryCollection(const TArray<flo
 		BoneMap[Idx] = RawBoneMapArray[Idx];
 		(*UV0)[Idx] = FVector2f::ZeroVector;
 
-		Colors[Idx] = FLinearColor::White;
+		Colors[Idx] = RestCollection->Defaults.DefaultVertexColor;
 	}
 
 	// Transforms
diff --git a/Engine/Source/Runtime/Experimental/Chaos/Public/GeometryCollection/GeometryCollection.h b/Engine/Source/Runtime/Experimental/Chaos/Public/GeometryCollection/GeometryCollection.h
--- a/Engine/Source/Runtime/Experimental/Chaos/Public/GeometryCollection/GeometryCollection.h
+++ b/Engine/Source/Runtime/Experimental/Chaos/Public/GeometryCollection/GeometryCollection.h
@@ -18,6 +18,10 @@ namespace Chaos
 class FGeometryCollectionConvexPropertiesInterface;
 class FGeometryCollectionProximityPropertiesInterface;
 
+struct FGeometryCollectionDefaults
+{
+	FLinearColor DefaultVertexColor = FLinearColor::White;
+};
 
 /**
 * FGeometryCollection (FTransformCollection)
@@ -30,7 +34,7 @@ class FGeometryCollection : public FTransformCollection,
 public:
 	typedef FTransformCollection Super;
 
-	CHAOS_API FGeometryCollection();
+	CHAOS_API FGeometryCollection(FGeometryCollectionDefaults InDefaults = FGeometryCollectionDefaults() );
 	FGeometryCollection(FGeometryCollection &) = delete;
 	FGeometryCollection& operator=(const FGeometryCollection &) = delete;
 	FGeometryCollection(FGeometryCollection &&) = default;
@@ -128,7 +132,7 @@ public:
 	/**
 	 * Create a GeometryCollection from Vertex and Indices arrays
 	 */
-	static CHAOS_API FGeometryCollection* NewGeometryCollection(const TArray<float>& RawVertexArray, const TArray<int32>& RawIndicesArray, bool ReverseVertexOrder = true);
+	static CHAOS_API FGeometryCollection* NewGeometryCollection(const TArray<float>& RawVertexArray, const TArray<int32>& RawIndicesArray, bool ReverseVertexOrder = true, FGeometryCollectionDefaults InDefaults = FGeometryCollectionDefaults());
 	static CHAOS_API void Init(FGeometryCollection* Collection, const TArray<float>& RawVertexArray, const TArray<int32>& RawIndicesArray, bool ReverseVertexOrder = true);
 	static CHAOS_API void DefineGeometrySchema(FManagedArrayCollection&);
 
@@ -143,7 +147,8 @@ public:
 		const TManagedArray<int32>& RawParentArray,
 		const TManagedArray<TSet<int32>>& RawChildrenArray,
 		const TManagedArray<int32>& RawSimulationTypeArray,
-		const TManagedArray<int32>& RawStatusFlagsArray);
+		const TManagedArray<int32>& RawStatusFlagsArray,
+		FGeometryCollectionDefaults InDefaults = FGeometryCollectionDefaults());
 
 	//
 	//
@@ -273,7 +278,7 @@ public:
 	//
 	//
 	//
-
+	FGeometryCollectionDefaults Defaults;
 	CHAOS_API virtual void SetDefaults(FName Group, uint32 StartSize, uint32 NumElements) override;
 
 	// Transform Group
