[ChaosDebugDraw] replaced TActorIterator by TObjectIterator

diff --git a/Engine/Source/Runtime/Experimental/ChaosSolverEngine/Private/Chaos/ChaosDebugDrawComponent.cpp b/Engine/Source/Runtime/Experimental/ChaosSolverEngine/Private/Chaos/ChaosDebugDrawComponent.cpp
--- a/Engine/Source/Runtime/Experimental/ChaosSolverEngine/Private/Chaos/ChaosDebugDrawComponent.cpp
+++ b/Engine/Source/Runtime/Experimental/ChaosSolverEngine/Private/Chaos/ChaosDebugDrawComponent.cpp
@@ -6,7 +6,7 @@
 #include "Debug/DebugDrawService.h"
 #include "DrawDebugHelpers.h"
 #if WITH_EDITOR
-#include "EngineUtils.h"
+#include "UObject/UObjectIterator.h"
 #endif
 #include "Engine/Canvas.h"
 #include "Engine/Engine.h"
@@ -373,10 +373,10 @@ void UChaosDebugDrawComponent::CreateDebugDrawActor(UWorld* World)
 	// Those can get reinitialized and OnPostWorldInitialization is called more than once
 	if (!World->IsGameWorld())
 	{
-		for (TActorIterator<AActor> It(World); It; ++It)
+		for (TObjectIterator<UChaosDebugDrawComponent> It; It; ++It)
 		{
-			const AActor* Actor = *It;
-			if (Actor != nullptr && Actor->GetFName()== NAME_ChaosDebugDrawActor)
+			const AActor* Actor = It->GetOwner();
+			if (Actor != nullptr && Actor->GetFName()== NAME_ChaosDebugDrawActor && Actor->GetWorld() == World)
 			{
 				return;
 			}
