Missing return keyword

#rnx

diff --git a/Engine/Source/Runtime/Experimental/IoStoreOnDemand/Private/IoStoreOnDemand.cpp b/Engine/Source/Runtime/Experimental/IoStoreOnDemand/Private/IoStoreOnDemand.cpp
--- a/Engine/Source/Runtime/Experimental/IoStoreOnDemand/Private/IoStoreOnDemand.cpp
+++ b/Engine/Source/Runtime/Experimental/IoStoreOnDemand/Private/IoStoreOnDemand.cpp
@@ -1838,7 +1838,7 @@ TIoStatusOr<FIoStoreListTocsParams> FIoStoreListTocsParams::Parse(const TCHAR* C
 	{
 		if (!FParse::Value(CommandLine, TEXT("Bucket="), Params.Bucket))
 		{
-			FIoStatus(EIoErrorCode::InvalidParameter, TEXT("Invalid bucket name"));
+			return FIoStatus(EIoErrorCode::InvalidParameter, TEXT("Invalid bucket name"));
 		}
 	}
 
@@ -1865,7 +1865,6 @@ TIoStatusOr<FIoStoreListTocsParams> FIoStoreListTocsParams::Parse(const TCHAR* C
 
 	FParse::Value(CommandLine, TEXT("Json="), Params.OutFile);
 
-
 	if (FParse::Value(CommandLine, TEXT("TocUrl="), Params.TocUrl))
 	{
 		FStringView UrlView(Params.TocUrl);
