Split parameter parsing and validation so they can be done independently.

#rnx

diff --git a/Engine/Source/Runtime/Experimental/IoStoreOnDemand/Private/IoStoreOnDemand.cpp b/Engine/Source/Runtime/Experimental/IoStoreOnDemand/Private/IoStoreOnDemand.cpp
--- a/Engine/Source/Runtime/Experimental/IoStoreOnDemand/Private/IoStoreOnDemand.cpp
+++ b/Engine/Source/Runtime/Experimental/IoStoreOnDemand/Private/IoStoreOnDemand.cpp
@@ -1051,33 +1051,43 @@ TIoStatusOr<FIoStoreUploadParams> FIoStoreUploadParams::Parse(const TCHAR* Comma
 		FParse::Value(FCommandLine::Get(), TEXT("ConfigFilePath="), ConfigFilePath);
 
 		Params.TocOutputDir = FPaths::GetPath(ConfigFilePath);
+	}
 
-		if (Params.TocOutputDir.IsEmpty())
-		{
-			return FIoStatus(EIoErrorCode::InvalidParameter, TEXT("Cmdline param 'WriteToDisk' requires a valid 'ConfigFilePath' param aswell"));
-		}
+	if (FIoStatus Validation = Params.Validate(); !Validation.IsOk())
+	{
+		return Validation;
 	}
 
-	if (!Params.AccessKey.IsEmpty() && Params.SecretKey.IsEmpty())
+	return Params;
+}
+
+FIoStatus FIoStoreUploadParams::Validate() const
+{
+	if (bWriteTocToDisk && TocOutputDir.IsEmpty())
+	{
+		return FIoStatus(EIoErrorCode::InvalidParameter, TEXT("Cmdline param 'WriteToDisk' requires a valid 'ConfigFilePath' param as well"));
+	}
+
+	if (!AccessKey.IsEmpty() && SecretKey.IsEmpty())
 	{
 		return FIoStatus(EIoErrorCode::InvalidParameter, TEXT("Invalid secret key"));
 	}
-	else if (Params.AccessKey.IsEmpty() && !Params.SecretKey.IsEmpty())
+	else if (AccessKey.IsEmpty() && !SecretKey.IsEmpty())
 	{
 		return FIoStatus(EIoErrorCode::InvalidParameter, TEXT("Invalid access key"));
 	}
 
-	if (!Params.CredentialsFile.IsEmpty() && Params.CredentialsFileKeyName.IsEmpty())
+	if (!CredentialsFile.IsEmpty() && CredentialsFileKeyName.IsEmpty())
 	{
 		return FIoStatus(EIoErrorCode::InvalidParameter, TEXT("Invalid credential file key name"));
 	}
 
-	if (Params.ServiceUrl.IsEmpty() && Params.Region.IsEmpty())
+	if (ServiceUrl.IsEmpty() && Region.IsEmpty())
 	{
 		return FIoStatus(EIoErrorCode::InvalidParameter, TEXT("Service URL or AWS region needs to be specified"));
 	}
 
-	return Params;
+	return FIoStatus::Ok;
 }
 
 TIoStatusOr<FIoStoreUploadResult> UploadContainerFiles(
@@ -1519,26 +1529,36 @@ TIoStatusOr<FIoStoreDownloadParams> FIoStoreDownloadParams::Parse(const TCHAR* C
 	FParse::Value(CommandLine, TEXT("CredentialsFile="), Params.CredentialsFile);
 	FParse::Value(CommandLine, TEXT("CredentialsFileKeyName="), Params.CredentialsFileKeyName);
 
-	if (!Params.AccessKey.IsEmpty() && Params.SecretKey.IsEmpty())
+	if (FIoStatus Validation = Params.Validate(); !Validation.IsOk())
+	{
+		return Validation;
+	}
+
+	return Params;
+}
+
+FIoStatus FIoStoreDownloadParams::Validate() const
+{
+	if (!AccessKey.IsEmpty() && SecretKey.IsEmpty())
 	{
 		return FIoStatus(EIoErrorCode::InvalidParameter, TEXT("Invalid secret key"));
 	}
-	else if (Params.AccessKey.IsEmpty() && !Params.SecretKey.IsEmpty())
+	else if (AccessKey.IsEmpty() && !SecretKey.IsEmpty())
 	{
 		return FIoStatus(EIoErrorCode::InvalidParameter, TEXT("Invalid access key"));
 	}
 
-	if (!Params.CredentialsFile.IsEmpty() && Params.CredentialsFileKeyName.IsEmpty())
+	if (!CredentialsFile.IsEmpty() && CredentialsFileKeyName.IsEmpty())
 	{
 		return FIoStatus(EIoErrorCode::InvalidParameter, TEXT("Invalid credential file key name"));
 	}
 
-	if (Params.ServiceUrl.IsEmpty() && Params.Region.IsEmpty())
+	if (ServiceUrl.IsEmpty() && Region.IsEmpty())
 	{
 		return FIoStatus(EIoErrorCode::InvalidParameter, TEXT("Service URL or AWS region needs to be specified"));
 	}
 
-	return Params;
+	return FIoStatus::Ok;
 }
 
 FIoStatus DownloadContainerFiles(const FIoStoreDownloadParams& DownloadParams, const FString& TocPath)
@@ -1812,7 +1832,7 @@ TIoStatusOr<FIoStoreListTocsParams> FIoStoreListTocsParams::Parse(const TCHAR* C
 	FIoStoreListTocsParams Params;
 
 	// Convenience argument to specify both bucket and bucket prefix
-	// -BucketPath="mybucket/some/data/path" is equal to -Bucket="bucket" -BucketPrefix="some/data/path
+	// -BucketPath="mybucket/some/data/path" is equal to -Bucket="mybucket" -BucketPrefix="some/data/path
 	FString BucketPath;
 	if (FParse::Value(CommandLine, TEXT("-BucketPath="), BucketPath))
 	{
@@ -1879,26 +1899,40 @@ TIoStatusOr<FIoStoreListTocsParams> FIoStoreListTocsParams::Parse(const TCHAR* C
 		}
 	}
 
-	if (!Params.AccessKey.IsEmpty() && Params.SecretKey.IsEmpty())
+	if (FIoStatus Status = Params.Validate(); !Status.IsOk())
+	{
+		return Status;
+	}
+
+	return Params;
+}
+
+FIoStatus FIoStoreListTocsParams::Validate() const
+{
+	if (Bucket.IsEmpty())
+	{
+		return FIoStatus(EIoErrorCode::InvalidParameter, TEXT("Invalid bucket name"));
+	}
+	if (!AccessKey.IsEmpty() && SecretKey.IsEmpty())
 	{
 		return FIoStatus(EIoErrorCode::InvalidParameter, TEXT("Invalid secret key"));
 	}
-	else if (Params.AccessKey.IsEmpty() && !Params.SecretKey.IsEmpty())
+	else if (AccessKey.IsEmpty() && !SecretKey.IsEmpty())
 	{
 		return FIoStatus(EIoErrorCode::InvalidParameter, TEXT("Invalid access key"));
 	}
 
-	if (!Params.CredentialsFile.IsEmpty() && Params.CredentialsFileKeyName.IsEmpty())
+	if (!CredentialsFile.IsEmpty() && CredentialsFileKeyName.IsEmpty())
 	{
 		return FIoStatus(EIoErrorCode::InvalidParameter, TEXT("Invalid credential file key name"));
 	}
 
-	if (Params.TocUrl.IsEmpty() && Params.ServiceUrl.IsEmpty() && Params.Region.IsEmpty())
+	if (TocUrl.IsEmpty() && ServiceUrl.IsEmpty() && Region.IsEmpty())
 	{
 		return FIoStatus(EIoErrorCode::InvalidParameter, TEXT("Service URL or AWS region needs to be specified"));
 	}
 
-	return Params;
+	return FIoStatus::Ok;
 }
 
 FIoStatus ListTocs(const FIoStoreListTocsParams& Params)
diff --git a/Engine/Source/Runtime/Experimental/IoStoreOnDemand/Public/IO/IoStoreOnDemand.h b/Engine/Source/Runtime/Experimental/IoStoreOnDemand/Public/IO/IoStoreOnDemand.h
--- a/Engine/Source/Runtime/Experimental/IoStoreOnDemand/Public/IO/IoStoreOnDemand.h
+++ b/Engine/Source/Runtime/Experimental/IoStoreOnDemand/Public/IO/IoStoreOnDemand.h
@@ -180,6 +180,7 @@ struct FIoStoreUploadParams
 	FString TocOutputDir;
 
 	UE_API static TIoStatusOr<FIoStoreUploadParams> Parse(const TCHAR* CommandLine);
+	FIoStatus Validate() const;
 };
 
 /** Results from uploading a FOnDemandToc */
@@ -217,6 +218,7 @@ struct FIoStoreDownloadParams
 	int32 MaxConcurrentDownloads = 16;
 	
 	UE_API static TIoStatusOr<FIoStoreDownloadParams> Parse(const TCHAR* CommandLine);
+	FIoStatus Validate() const;
 };
 
 UE_API FIoStatus DownloadContainerFiles(const FIoStoreDownloadParams& DownloadParams, const FString& TocPath);
@@ -274,6 +276,7 @@ struct FIoStoreListTocsParams
 	bool bBlockHashes = false;
 
 	UE_API static TIoStatusOr<FIoStoreListTocsParams> Parse(const TCHAR* CommandLine);
+	FIoStatus Validate() const;
 };
 
 UE_API FIoStatus ListTocs(const FIoStoreListTocsParams& Params);
