Fix partial destruction, when geometry are not created on PT on initialization

diff --git a/Engine/Source/Runtime/Experimental/Chaos/Private/Chaos/PBDRigidClustering.cpp b/Engine/Source/Runtime/Experimental/Chaos/Private/Chaos/PBDRigidClustering.cpp
--- a/Engine/Source/Runtime/Experimental/Chaos/Private/Chaos/PBDRigidClustering.cpp
+++ b/Engine/Source/Runtime/Experimental/Chaos/Private/Chaos/PBDRigidClustering.cpp
@@ -26,6 +26,8 @@
 #include "PhysicsProxy/ClusterUnionPhysicsProxy.h"
 #include "CoreMinimal.h"
 
+extern CHAOS_API bool bBuildGeometryForChildrenOnPT;
+
 namespace Chaos
 {
 	//
@@ -1042,6 +1044,31 @@ namespace Chaos
 		return ReleaseClusterParticlesImpl(ClusteredParticle, bForceRelease, false /*bCreateNewClusters*/);
 	}
 	
+	void GenerateEdges(FGeometryCollectionPhysicsProxy& ConcreteGCProxy, FPBDRigidClusteredParticleHandle& ClusteredParticle, FClusterUnionManager& ClusterUnionManager)
+	{
+		ConcreteGCProxy.CreateChildrenGeometry_Internal();
+		if (Chaos::FClusterUnion* ClusterUnion = ClusterUnionManager.FindClusterUnionFromParticle(&ClusteredParticle))
+		{
+			ClusterUnion->PendingConnectivityOperations.Add({ &ClusteredParticle, Chaos::EClusterUnionConnectivityOperation::Add });
+
+			const TArray<Chaos::TConnectivityEdge<Chaos::FReal>> Edges = ClusteredParticle.ConnectivityEdges();
+			for (const Chaos::TConnectivityEdge<Chaos::FReal>& Edge : Edges)
+			{
+				if (Edge.Sibling != nullptr && Edge.Sibling->GetParticleType() == Chaos::EParticleType::Clustered)
+				{
+					Chaos::FPBDRigidClusteredParticleHandle* Sibling = Edge.Sibling->CastToClustered();
+					if (Sibling->PhysicsProxy()->GetType() == FGeometryCollectionPhysicsProxy::ConcreteType())
+					{
+						FGeometryCollectionPhysicsProxy* GCProxy = GetConcreteProxy<FGeometryCollectionPhysicsProxy>(Sibling);
+						GCProxy->CreateChildrenGeometry_Internal();
+					}
+				}
+			}
+			ClusterUnionManager.RequestDeferredClusterPropertiesUpdate(ClusterUnion->InternalIndex, Chaos::EUpdateClusterUnionPropertiesFlags::IncrementalGenerateConnectionGraph);
+			ClusterUnionManager.HandleDeferredClusterUnionUpdateProperties();
+		}
+	}
+
 	TSet<FPBDRigidParticleHandle*> FRigidClustering::ReleaseClusterParticlesImpl(
 		FPBDRigidClusteredParticleHandle* ClusteredParticle,
 		bool bForceRelease,
@@ -1070,6 +1097,11 @@ namespace Chaos
 			bUseDamagePropagation = SimParams.bUseDamagePropagation;
 			BreakDamagePropagationFactor = SimParams.BreakDamagePropagationFactor;
 			ShockDamagePropagationFactor = SimParams.ShockDamagePropagationFactor;
+			
+			if (bBuildGeometryForChildrenOnPT == false)
+			{
+				GenerateEdges(*ConcreteGCProxy, *ClusteredParticle, ClusterUnionManager);
+			}
 		}
 
 		TArray<FPBDRigidParticleHandle*>& Children = MChildren[ClusteredParticle];
@@ -1216,10 +1248,6 @@ namespace Chaos
 
 		if (ActivatedChildren.Num() > 0)
 		{
-			if (ConcreteGCProxy)
-			{
-				ConcreteGCProxy->CreateChildrenGeometry_Internal();
-			}
 			const bool bIsClusterUnion = ClusterUnionManager.IsClusterUnionParticle(ClusteredParticle);
 			if (Children.Num() == 0)
 			{
diff --git a/Engine/Source/Runtime/Experimental/Chaos/Private/PhysicsProxy/GeometryCollectionPhysicsProxy.cpp b/Engine/Source/Runtime/Experimental/Chaos/Private/PhysicsProxy/GeometryCollectionPhysicsProxy.cpp
--- a/Engine/Source/Runtime/Experimental/Chaos/Private/PhysicsProxy/GeometryCollectionPhysicsProxy.cpp
+++ b/Engine/Source/Runtime/Experimental/Chaos/Private/PhysicsProxy/GeometryCollectionPhysicsProxy.cpp
@@ -543,7 +543,7 @@ FGeometryCollectionPhysicsProxy::~FGeometryCollectionPhysicsProxy()
 float ReportHighParticleFraction = -1.f;
 FAutoConsoleVariableRef CVarReportHighParticleFraction(TEXT("p.gc.ReportHighParticleFraction"), ReportHighParticleFraction, TEXT("Report any objects with particle fraction above this threshold"));
 
-bool bBuildGeometryForChildrenOnPT = true;
+CHAOS_API bool bBuildGeometryForChildrenOnPT = true;
 FAutoConsoleVariableRef CVarbBuildGeometryForChildrenOnPT(TEXT("p.gc.BuildGeometryForChildrenOnPT"), bBuildGeometryForChildrenOnPT, TEXT("If true build all children geometry on Physics Thread at initilaization time, otherwise wait until destruction occurs."));
 
 bool bBuildGeometryForChildrenOnGT = true;
