Client CSV metrics for tracking object creation and state updates received from the server when using Iris.

diff --git a/Engine/Source/Runtime/Experimental/Iris/Core/Private/Iris/ReplicationSystem/ReplicationReader.cpp b/Engine/Source/Runtime/Experimental/Iris/Core/Private/Iris/ReplicationSystem/ReplicationReader.cpp
--- a/Engine/Source/Runtime/Experimental/Iris/Core/Private/Iris/ReplicationSystem/ReplicationReader.cpp
+++ b/Engine/Source/Runtime/Experimental/Iris/Core/Private/Iris/ReplicationSystem/ReplicationReader.cpp
@@ -25,6 +25,7 @@
 #include "HAL/IConsoleManager.h"
 #include "UObject/Class.h"
 #include "Algo/RemoveIf.h"
+#include "ProfilingDebugging/CsvProfiler.h"
 
 #if UE_NET_ENABLE_REPLICATIONREADER_LOG
 #	define UE_LOG_REPLICATIONREADER(Format, ...)  UE_LOG(LogIris, Log, Format, ##__VA_ARGS__)
@@ -36,6 +37,8 @@
 
 #define UE_LOG_REPLICATIONREADER_CONN_WARNING(Format, ...) UE_LOG(LogIris, Warning, TEXT("Conn: %u ") Format, Parameters.ConnectionId, ##__VA_ARGS__)
 
+CSV_DEFINE_CATEGORY(IrisClient, true);
+
 namespace UE::Net::Private
 {
 
@@ -828,6 +831,12 @@ void FReplicationReader::ReadObjectInBatch(FNetSerializationContext& Context, FN
 		// Get Bridge
 		FReplicationBridgeSerializationContext BridgeContext(Context, Parameters.ConnectionId);
 
+		CSV_CUSTOM_STAT(IrisClient, ClientObjectCreate, 1, ECsvCustomStatOp::Accumulate);
+		if (!bIsSubObject)
+		{
+			CSV_CUSTOM_STAT(IrisClient, ClientObjectCreateRoot, 1, ECsvCustomStatOp::Accumulate);
+		}
+
 		const FReplicationBridgeCreateNetRefHandleResult CreateResult = ReplicationBridge->CallCreateNetRefHandleFromRemote(RootObjectOfSubObject, IncompleteHandle, BridgeContext);
 		FNetRefHandle NetRefHandle = CreateResult.NetRefHandle;
 		if (!NetRefHandle.IsValid())
