GeometryCollection : TransformFacade
- Guard against empty collections.

diff --git a/Engine/Source/Runtime/Experimental/Chaos/Private/GeometryCollection/Facades/CollectionTransformFacade.cpp b/Engine/Source/Runtime/Experimental/Chaos/Private/GeometryCollection/Facades/CollectionTransformFacade.cpp
--- a/Engine/Source/Runtime/Experimental/Chaos/Private/GeometryCollection/Facades/CollectionTransformFacade.cpp
+++ b/Engine/Source/Runtime/Experimental/Chaos/Private/GeometryCollection/Facades/CollectionTransformFacade.cpp
@@ -58,12 +58,14 @@ namespace GeometryCollection::Facades
 	{
 		// Update only root transforms
 		const TArray<int32>& RootIndices = GetRootIndices();
-
-		TManagedArray<FTransform3f>& Transforms = TransformAttribute.Modify();
-
-		for (int32 Idx : RootIndices)
+		if (TransformAttribute.IsValid())
 		{
-			Transforms[Idx] = Transforms[Idx] * FTransform3f(InTransform);
+			TManagedArray<FTransform3f>& Transforms = TransformAttribute.Modify();
+
+			for (int32 Idx : RootIndices)
+			{
+				Transforms[Idx] = Transforms[Idx] * FTransform3f(InTransform);
+			}
 		}
 	}
 
@@ -210,9 +212,11 @@ namespace GeometryCollection::Facades
 
 	void FCollectionTransformFacade::SetBoneTransformToIdentity(int32 BoneIdx)
 	{
-		TManagedArray<FTransform3f>& BoneTransforms = TransformAttribute.Modify();
-
-		BoneTransforms[BoneIdx] = FTransform3f::Identity;
+		if (TransformAttribute.IsValid())
+		{
+			TManagedArray<FTransform3f>& BoneTransforms = TransformAttribute.Modify();
+			BoneTransforms[BoneIdx] = FTransform3f::Identity;
+		}
 	}
 
 
