[ChaosVD] Fix CIS issue in FChaosVDSceneQueryDataComponentVisualizer and two of the CVD Data wrapper files

diff --git a/Engine/Plugins/Experimental/ChaosVD/Source/ChaosVD/Private/Visualizers/ChaosVDSceneQueryDataComponentVisualizer.cpp b/Engine/Plugins/Experimental/ChaosVD/Source/ChaosVD/Private/Visualizers/ChaosVDSceneQueryDataComponentVisualizer.cpp
--- a/Engine/Plugins/Experimental/ChaosVD/Source/ChaosVD/Private/Visualizers/ChaosVDSceneQueryDataComponentVisualizer.cpp
+++ b/Engine/Plugins/Experimental/ChaosVD/Source/ChaosVD/Private/Visualizers/ChaosVDSceneQueryDataComponentVisualizer.cpp
@@ -2,8 +2,8 @@
 
 #include "Visualizers/ChaosVDSceneQueryDataComponentVisualizer.h"
 
-#include <Actors/ChaosVDSceneQueryDataContainer.h>
-#include <Components/ChaosVDSceneQueryDataComponent.h>
+#include "Actors/ChaosVDSceneQueryDataContainer.h"
+#include "Components/ChaosVDSceneQueryDataComponent.h"
 
 #include "ChaosVDEditorSettings.h"
 #include "ChaosVDGeometryBuilder.h"
diff --git a/Engine/Plugins/Experimental/ChaosVD/Source/ChaosVD/Public/Visualizers/ChaosVDSceneQueryDataComponentVisualizer.h b/Engine/Plugins/Experimental/ChaosVD/Source/ChaosVD/Public/Visualizers/ChaosVDSceneQueryDataComponentVisualizer.h
--- a/Engine/Plugins/Experimental/ChaosVD/Source/ChaosVD/Public/Visualizers/ChaosVDSceneQueryDataComponentVisualizer.h
+++ b/Engine/Plugins/Experimental/ChaosVD/Source/ChaosVD/Public/Visualizers/ChaosVDSceneQueryDataComponentVisualizer.h
@@ -5,6 +5,7 @@
 #include "ComponentVisualizer.h"
 #include "IChaosVDParticleVisualizationDataProvider.h"
 #include "Chaos/ImplicitFwd.h"
+#include "Chaos/ImplicitObject.h"
 #include "Components/ChaosVDSceneQueryDataComponent.h"
 
 struct FChaosVDRecording;
diff --git a/Engine/Source/Runtime/Experimental/ChaosVisualDebugger/Public/DataWrappers/ChaosVDCollisionDataWrappers.h b/Engine/Source/Runtime/Experimental/ChaosVisualDebugger/Public/DataWrappers/ChaosVDCollisionDataWrappers.h
--- a/Engine/Source/Runtime/Experimental/ChaosVisualDebugger/Public/DataWrappers/ChaosVDCollisionDataWrappers.h
+++ b/Engine/Source/Runtime/Experimental/ChaosVisualDebugger/Public/DataWrappers/ChaosVDCollisionDataWrappers.h
@@ -401,13 +401,13 @@ struct CHAOSVDRUNTIME_API FChaosVDCollisionFilterData
 	GENERATED_BODY()
 
 	UPROPERTY(VisibleAnywhere, Category=CollisionData)
-	uint32 Word0;
+	uint32 Word0 = 0;
 	UPROPERTY(VisibleAnywhere, Category=CollisionData)
-	uint32 Word1;
+	uint32 Word1 = 0;
 	UPROPERTY(VisibleAnywhere, Category=CollisionData)
-	uint32 Word2;
+	uint32 Word2 = 0;
 	UPROPERTY(VisibleAnywhere, Category=CollisionData)
-	uint32 Word3;
+	uint32 Word3 = 0;
 
 	bool Serialize(FArchive& Ar);
 };
diff --git a/Engine/Source/Runtime/Experimental/ChaosVisualDebugger/Public/DataWrappers/ChaosVDQueryDataWrappers.h b/Engine/Source/Runtime/Experimental/ChaosVisualDebugger/Public/DataWrappers/ChaosVDQueryDataWrappers.h
--- a/Engine/Source/Runtime/Experimental/ChaosVisualDebugger/Public/DataWrappers/ChaosVDQueryDataWrappers.h
+++ b/Engine/Source/Runtime/Experimental/ChaosVisualDebugger/Public/DataWrappers/ChaosVDQueryDataWrappers.h
@@ -211,10 +211,10 @@ struct FChaosVDQueryFastData : public FChaosVDWrapperDataBase
 	FVector InvDir = FVector::ZeroVector;
 	
 	UPROPERTY(VisibleAnywhere, Category=QueryFastData)
-	float CurrentLength = 0.0f;
+	double CurrentLength = 0.0;
 
 	UPROPERTY(VisibleAnywhere, Category=QueryFastData)
-	float InvCurrentLength = 0.0f;
+	double InvCurrentLength = 0.0;
 
 	UPROPERTY(VisibleAnywhere, Category=QueryFastData)
 	bool bParallel0 = false;
@@ -374,7 +374,7 @@ struct FChaosVDQueryDataWrapper
 	uint32 InputGeometryKey = 0;
 
 	UPROPERTY(VisibleAnywhere, Category=RecordedQueryData)
-	FQuat GeometryOrientation;
+	FQuat GeometryOrientation = FQuat::Identity;
 
 	UPROPERTY(VisibleAnywhere, Category=RecordedQueryData)
 	EChaosVDSceneQueryType Type = EChaosVDSceneQueryType::Invalid;
