Add error checking when mounting a .uondemandtoc so that we can identify file corruption.

#rnx

- The sentinel is not part of FOnDemandToc but is written after to it on file as part of the .uondemandtoc format.
- At the moment if corruption is detected we will log an error and not mount that toc which somewhat replicates the current behavior of our other containers.
- This will be most useful when checking if a file has been truncated or not written to disk correctly, it will not detect problems in the middle of the file but that is much rarer.
- The sentinel format is an easy to recognise string pattern which should make it easier to debug and manually check at the end of the file.

diff --git a/Engine/Source/Runtime/Experimental/IoStoreOnDemand/Private/IoStoreOnDemand.cpp b/Engine/Source/Runtime/Experimental/IoStoreOnDemand/Private/IoStoreOnDemand.cpp
--- a/Engine/Source/Runtime/Experimental/IoStoreOnDemand/Private/IoStoreOnDemand.cpp
+++ b/Engine/Source/Runtime/Experimental/IoStoreOnDemand/Private/IoStoreOnDemand.cpp
@@ -532,6 +532,31 @@ bool LoadFromCompactBinary(FCbFieldView Field, FOnDemandTocContainerEntry& OutCo
 	return false;
 }
 
+bool FOnDemandTocSentinel::IsValid()
+{
+	return FMemory::Memcmp(&Data, FOnDemandTocSentinel::SentinelImg, FOnDemandTocSentinel::SentinelSize) == 0;
+}
+
+FArchive& operator<<(FArchive& Ar, FOnDemandTocSentinel& Sentinel)
+{
+	if (Ar.IsSaving())
+	{	
+		// We could just cast FOnDemandTocSentinel::SentinelImg to a non-const pointer but we can't be 
+		// 100% sure that the FArchive won't change the data, even if it is in Saving mode. Since this 
+		// isn't performance critical we will play it safe.
+		uint8 Output[FOnDemandTocSentinel::SentinelSize];
+		FMemory::Memcpy(Output, FOnDemandTocSentinel::SentinelImg, FOnDemandTocSentinel::SentinelSize);
+
+		Ar.Serialize(&Output, FOnDemandTocSentinel::SentinelSize);
+	}
+	else
+	{
+		Ar.Serialize(&Sentinel.Data, FOnDemandTocSentinel::SentinelSize);
+	}
+
+	return Ar;
+}
+
 FArchive& operator<<(FArchive& Ar, FOnDemandToc& Toc)
 {
 	Ar << Toc.Header;
@@ -1004,6 +1029,9 @@ static FIoStatus WriteContainerFiles(FOnDemandToc& OnDemandToc, const TMap<FIoHa
 
 		(*Ar) << ContainerToc;
 
+		FOnDemandTocSentinel Sentinel;
+		(*Ar) << Sentinel;
+
 		if (Ar->IsError() || Ar->IsCriticalError())
 		{
 			return FIoStatus(EIoErrorCode::WriteError, FString::Printf(TEXT("Failed to write to '%s'"), *OutputPath));
diff --git a/Engine/Source/Runtime/Experimental/IoStoreOnDemand/Private/OnDemandIoDispatcherBackend.cpp b/Engine/Source/Runtime/Experimental/IoStoreOnDemand/Private/OnDemandIoDispatcherBackend.cpp
--- a/Engine/Source/Runtime/Experimental/IoStoreOnDemand/Private/OnDemandIoDispatcherBackend.cpp
+++ b/Engine/Source/Runtime/Experimental/IoStoreOnDemand/Private/OnDemandIoDispatcherBackend.cpp
@@ -1562,6 +1562,31 @@ void FOnDemandIoBackend::MountContainer(FStringView ContainerPath)
 			return;
 		}
 
+		/// First we read and validate the sentinel at the end of the file to check for potential
+		// file corruption that could lead to crashing during serialization.
+		{
+			const int64 SentinelPos = Ar->TotalSize() - FOnDemandTocSentinel::SentinelSize;
+
+			if (SentinelPos < 0)
+			{
+				UE_LOG(LogIas, Error, TEXT("The file '%s' is smaller than expected and quite possible corrupted"), *TocPath);
+				return;
+			}
+
+			Ar->Seek(SentinelPos);
+
+			FOnDemandTocSentinel Sentinel;
+			(*Ar) << Sentinel;
+
+			if (!Sentinel.IsValid())
+			{
+				UE_LOG(LogIas, Error, TEXT("File corruption detected when serializing '%s'"), *TocPath);
+				return;
+			}
+
+			Ar->Seek(0);
+		}
+
 		FOnDemandToc ContainerToc;
 		(*Ar) << ContainerToc;
 
diff --git a/Engine/Source/Runtime/Experimental/IoStoreOnDemand/Public/IO/IoStoreOnDemand.h b/Engine/Source/Runtime/Experimental/IoStoreOnDemand/Public/IO/IoStoreOnDemand.h
--- a/Engine/Source/Runtime/Experimental/IoStoreOnDemand/Public/IO/IoStoreOnDemand.h
+++ b/Engine/Source/Runtime/Experimental/IoStoreOnDemand/Public/IO/IoStoreOnDemand.h
@@ -127,6 +127,20 @@ struct FOnDemandTocContainerEntry
 
 UE_API bool LoadFromCompactBinary(FCbFieldView Field, FOnDemandTocContainerEntry& OutContainer);
 
+struct FOnDemandTocSentinel
+{
+public:
+	static constexpr inline char SentinelImg[] = "-[]--[]--[]--[]-";
+	static constexpr uint32 SentinelSize = 16;
+
+	bool IsValid();
+
+	UE_API friend FArchive& operator<<(FArchive& Ar, FOnDemandTocSentinel& Sentinel);
+
+private:
+	uint8 Data[SentinelSize] = { 0 };
+};
+
 struct FOnDemandToc
 {
 	FOnDemandToc() = default;
