Creating a .uondemandtoc file per container can now be enabled via the config file system. Continues to default to disabled.

#rnx

- This continues to be a temp measure while this feature is disabled for testing, the only reason for the change is so that the feature can be enabled for testing via data rather than by changing code.
- Enabled by setting engine.ini:[Ias]:ForceTocFromMountedPaks to true, note that setting this for a specific project doesn't work as UnrealPak does not load per project config settings.

diff --git a/Engine/Config/BaseEngine.ini b/Engine/Config/BaseEngine.ini
--- a/Engine/Config/BaseEngine.ini
+++ b/Engine/Config/BaseEngine.ini
@@ -3468,6 +3468,7 @@ bUpdateBuildGraphPropertiesFile=false
 
 [Ias]
 EnableInEditor=False
+ForceTocFromMountedPaks=False
 
 [/Script/InterchangeEngine.InterchangeProjectSettings]
 ContentImportSettings=(DefaultPipelineStackOverride=((Materials, "Materials"),(Textures, "Textures")),ShowImportDialogOverride=((Textures, (bShowImportDialog=False))),PipelineStacks=(("Assets", (Pipelines=("/Interchange/Pipelines/DefaultAssetsPipeline.DefaultAssetsPipeline"),PerTranslatorPipelines=((Translator="/Script/InterchangeImport.InterchangeMaterialXTranslator",Pipelines=("/Interchange/Pipelines/DefaultAssetsPipeline.DefaultAssetsPipeline","/Interchange/Pipelines/DefaultMaterialXPipeline.DefaultMaterialXPipeline")),(Translator="/Script/InterchangeImport.InterchangeGLTFTranslator",Pipelines=("/Interchange/Pipelines/DefaultGLTFAssetsPipeline.DefaultGLTFAssetsPipeline","/Interchange/Pipelines/DefaultGLTFPipeline.DefaultGLTFPipeline"))))),("Materials", (Pipelines=("/Interchange/Pipelines/DefaultMaterialPipeline.DefaultMaterialPipeline"),PerTranslatorPipelines=((Translator="/Script/InterchangeImport.InterchangeMaterialXTranslator",Pipelines=("/Interchange/Pipelines/DefaultMaterialPipeline.DefaultMaterialPipeline","/Interchange/Pipelines/DefaultMaterialXPipeline.DefaultMaterialXPipeline"))))),("Textures", (Pipelines=("/Interchange/Pipelines/DefaultTexturePipeline.DefaultTexturePipeline")))),DefaultPipelineStack="Assets",ImportDialogClass="/Script/InterchangeEditorPipelines.InterchangePipelineConfigurationGeneric",bShowImportDialog=True)
@@ -3555,4 +3556,4 @@ UsageType=Editor
 
 [/Script/Engine.WorldPartitionSettings]
 EditorHashDefaultClass=/Script/CoreUObject.Class'/Script/Engine.WorldPartitionEditorSpatialHash'
-RuntimeHashDefaultClass=/Script/CoreUObject.Class'/Script/Engine.WorldPartitionRuntimeHashSet'
\ No newline at end of file
+RuntimeHashDefaultClass=/Script/CoreUObject.Class'/Script/Engine.WorldPartitionRuntimeHashSet'
diff --git a/Engine/Source/Runtime/Experimental/IoStoreOnDemand/Private/IoStoreOnDemand.cpp b/Engine/Source/Runtime/Experimental/IoStoreOnDemand/Private/IoStoreOnDemand.cpp
--- a/Engine/Source/Runtime/Experimental/IoStoreOnDemand/Private/IoStoreOnDemand.cpp
+++ b/Engine/Source/Runtime/Experimental/IoStoreOnDemand/Private/IoStoreOnDemand.cpp
@@ -66,9 +66,6 @@ static FAutoConsoleVariableRef CVar_DistributedEndpointFallbackUrl(
 	TEXT("CDN url to be used if a distributed endpoint cannot be reached (overrides IoStoreOnDemand.ini)")
 );
 
-// When eanbled we will write out a toc (using extension GIasOnDemandTocExt) per container, currently disabled for testing
-static bool GIasWriteOnDemandTocs = false;
-
 ////////////////////////////////////////////////////////////////////////////////
 static int64 ParseSizeParam(FStringView Value)
 {
@@ -1238,6 +1235,9 @@ TIoStatusOr<FIoStoreUploadResult> UploadContainerFiles(
 	FString ChunksRelativePath = UploadParams.BucketPrefix.IsEmpty() ? TEXT("Chunks") : FString::Printf(TEXT("%s/Chunks"), *UploadParams.BucketPrefix);
 	ChunksRelativePath.ToLowerInline();
 
+	bool bWritePerContainerToc = false;
+	GConfig->GetBool(TEXT("Ias"), TEXT("ForceTocFromMountedPaks"), bWritePerContainerToc, GEngineIni);
+
 	uint64 TotalUploadedChunks = 0;
 	uint64 TotalUploadedBytes = 0;
 
@@ -1404,7 +1404,7 @@ TIoStatusOr<FIoStoreUploadResult> UploadContainerFiles(
 			ContainerFileReader.GetContainerFilePaths(FilesToDelete);
 
 			// We need the pak files in order to mount OnDemand toc files!
-			if (UploadParams.bDeletePakFiles && !GIasWriteOnDemandTocs)
+			if (UploadParams.bDeletePakFiles && !bWritePerContainerToc)
 			{
 				FilesToDelete.Add(FPaths::ChangeExtension(Path, TEXT(".pak")));
 				FilesToDelete.Add(FPaths::ChangeExtension(Path, TEXT(".sig")));
@@ -1471,7 +1471,7 @@ TIoStatusOr<FIoStoreUploadResult> UploadContainerFiles(
 		}
 	}
 
-	if (GIasWriteOnDemandTocs)
+	if (bWritePerContainerToc)
 	{
 		// Write out separate .uondemandtoc files, one per.utoc containing ondemand data.
 		FIoStatus Result = WriteContainerFiles(OnDemandToc, UTocPaths);
@@ -2370,7 +2370,7 @@ void FIoStoreOnDemandModule::InitializeInternal()
 		return;
 	}
 #endif //WITH_EDITOR
-	
+
 	const TCHAR* CommandLine = FCommandLine::Get();
 	
 #if !UE_BUILD_SHIPPING
diff --git a/Engine/Source/Runtime/Experimental/IoStoreOnDemand/Private/OnDemandIoDispatcherBackend.cpp b/Engine/Source/Runtime/Experimental/IoStoreOnDemand/Private/OnDemandIoDispatcherBackend.cpp
--- a/Engine/Source/Runtime/Experimental/IoStoreOnDemand/Private/OnDemandIoDispatcherBackend.cpp
+++ b/Engine/Source/Runtime/Experimental/IoStoreOnDemand/Private/OnDemandIoDispatcherBackend.cpp
@@ -1488,6 +1488,16 @@ FOnDemandIoBackend::FOnDemandIoBackend(TUniquePtr<IIasCache>&& InCache)
 		ECVF_Cheat)
 	);
 #endif // UE_IAS_DEBUG_CONSOLE_CMDS
+
+	{
+		bool bForceTocFromMountedPaks = false;
+		GConfig->GetBool(TEXT("Ias"), TEXT("ForceTocFromMountedPaks"), bForceTocFromMountedPaks, GEngineIni);
+
+		if (bForceTocFromMountedPaks)
+		{
+			ForceTocMode(ETocMode::LoadTocFromMountedPaks);
+		}
+	}
 }
 
 FOnDemandIoBackend::~FOnDemandIoBackend()
