Include IAS cache tests in program-type targets

#rnx

diff --git a/Engine/Source/Runtime/Experimental/IoStoreOnDemand/Private/JournaledCache.cpp b/Engine/Source/Runtime/Experimental/IoStoreOnDemand/Private/JournaledCache.cpp
--- a/Engine/Source/Runtime/Experimental/IoStoreOnDemand/Private/JournaledCache.cpp
+++ b/Engine/Source/Runtime/Experimental/IoStoreOnDemand/Private/JournaledCache.cpp
@@ -2324,7 +2324,7 @@ TUniquePtr<IIasCache> MakeIasCache(const TCHAR* RootPath, const FIasCacheConfig&
 
 // {{{1 test ...................................................................
 
-#if defined(IAS_JOURNALED_CACHE_TEST)
+#if IS_PROGRAM
 
 namespace IasJournaledFileCacheTest
 {
@@ -2962,7 +2962,7 @@ IOSTOREONDEMAND_API void Tests(const TCHAR* CacheDir=nullptr)
 
 } // namespace IasJournaledFileCacheTest
 
-#endif // IAS_JOURNALED_CACHE_TEST
+#endif // IS_PROGRAM
 
 // }}}
 
