ManagedArrayCollection add FindOrAddAttributeTyped.

This does the same thing as AddAttribute except it will check that the type matches the existing type if the attribute already exists rather than just static casting without a check.

diff --git a/Engine/Source/Runtime/Experimental/Chaos/Public/GeometryCollection/ManagedArrayCollection.h b/Engine/Source/Runtime/Experimental/Chaos/Public/GeometryCollection/ManagedArrayCollection.h
--- a/Engine/Source/Runtime/Experimental/Chaos/Public/GeometryCollection/ManagedArrayCollection.h
+++ b/Engine/Source/Runtime/Experimental/Chaos/Public/GeometryCollection/ManagedArrayCollection.h
@@ -138,27 +138,28 @@ public:
 	{
 		if (!HasAttribute(Name, Group))
 		{
-			if (!HasGroup(Group))
-			{
-				AddGroup(Group);
-			}
-
-			FValueType Value(ManagedArrayType<T>(), *(new TManagedArray<T>()));
-			Value.Value->Resize(NumElements(Group));
-			Value.Saved = Parameters.Saved;
-			if (ensure(Parameters.bAllowCircularDependency || !IsConnected(Parameters.GroupIndexDependency, Group)))
-			{
-				Value.GroupIndexDependency = Parameters.GroupIndexDependency;
-			}
-			else
-			{
-				Value.GroupIndexDependency = "";
-			}
-			Map.Add(FManagedArrayCollection::MakeMapKey(Name, Group), MoveTemp(Value));
+			AddNewAttributeImpl<T>(Name, Group, Parameters);
 		}
 		return ModifyAttribute<T>(Name, Group);
 	}
 
+	/**
+	* Add an attribute of Type(T) to the group or find the existing attribute with the same name and type.
+	* @param Name - The name of the attribute
+	* @param Group - The group that manages the attribute
+	* @return pointer to the added/found managed array. If there is a type mismatch with an existing attribute, returns nullptr.
+	*/
+	template<typename T>
+	TManagedArray<T>* FindOrAddAttributeTyped(FName Name, FName Group, FConstructionParameters Parameters = FConstructionParameters())
+	{
+		if (!HasAttribute(Name, Group))
+		{
+			AddNewAttributeImpl<T>(Name, Group, Parameters);
+		}
+		return ModifyAttributeTyped<T>(Name, Group);
+	}
+
+
 	/**
 	* Duplicate the ManagedArrayCollection as the specified templated type
 	* @return unowned pointer to new collection.
@@ -357,6 +358,28 @@ public:
 		ManagedArray->MarkDirty();
 		return *(static_cast<TManagedArray<T>*>(ManagedArray));
 	}
+
+	/**
+	* Returns attribute access of Type(T) from the group for modification if and only if the types of T and the array match
+	* this will mark the collection dirty
+	* @param Name - The name of the attribute
+	* @param Group - The group that manages the attribute
+	* @return ManagedArray<T> *
+	*/
+	template<typename T>
+	TManagedArray<T>* ModifyAttributeTyped(FName Name, FName Group)
+	{
+		const FKeyType Key = FManagedArrayCollection::MakeMapKey(Name, Group);
+		if (FValueType* FoundValue = Map.Find(Key))
+		{
+			if (FoundValue->ArrayType == ManagedArrayType<T>())
+			{
+				FoundValue->Value->MarkDirty();
+				return static_cast<TManagedArray<T>*>(FoundValue->Value);
+			}
+		}
+		return nullptr;
+	}
 	
 	/**
 	* Returns attribute access of Type(T) from the group
@@ -589,6 +612,29 @@ public:
 
 private:
 
+	template<typename T>
+	void AddNewAttributeImpl(FName Name, FName Group, const FConstructionParameters& Parameters)
+	{
+		checkSlow(!HasAttribute(Name, Group));
+		if (!HasGroup(Group))
+		{
+			AddGroup(Group);
+		}
+
+		FValueType Value(ManagedArrayType<T>(), *(new TManagedArray<T>()));
+		Value.Value->Resize(NumElements(Group));
+		Value.Saved = Parameters.Saved;
+		if (ensure(Parameters.bAllowCircularDependency || !IsConnected(Parameters.GroupIndexDependency, Group)))
+		{
+			Value.GroupIndexDependency = Parameters.GroupIndexDependency;
+		}
+		else
+		{
+			Value.GroupIndexDependency = "";
+		}
+		Map.Add(FManagedArrayCollection::MakeMapKey(Name, Group), MoveTemp(Value));
+	}
+
 	/****
 	*  Mapping Key/Value
 	*
