Update FDataflowNode::ValidateConnections to handle nested pass throughs.

diff --git a/Engine/Source/Runtime/Experimental/Dataflow/Core/Private/Dataflow/DataflowNode.cpp b/Engine/Source/Runtime/Experimental/Dataflow/Core/Private/Dataflow/DataflowNode.cpp
--- a/Engine/Source/Runtime/Experimental/Dataflow/Core/Private/Dataflow/DataflowNode.cpp
+++ b/Engine/Source/Runtime/Experimental/Dataflow/Core/Private/Dataflow/DataflowNode.cpp
@@ -340,7 +340,7 @@ uint32 FDataflowNode::GetPropertyOffset(const FName& PropertyFullName) const
 	return Offset;
 }
 
-FName FDataflowNode::GetPropertyFullName(const TArray<const FProperty*>& PropertyChain)
+FString FDataflowNode::GetPropertyFullNameString(const TConstArrayView<const FProperty*>& PropertyChain)
 {
 	FString PropertyFullName;
 	for (const FProperty* const Property : PropertyChain)
@@ -350,6 +350,12 @@ FName FDataflowNode::GetPropertyFullName(const TArray<const FProperty*>& Propert
 			PropertyName :
 			FString::Format(TEXT("{0}.{1}"), { PropertyName, PropertyFullName });
 	}
+	return PropertyFullName;
+}
+
+FName FDataflowNode::GetPropertyFullName(const TArray<const FProperty*>& PropertyChain)
+{
+	const FString PropertyFullName = GetPropertyFullNameString(TConstArrayView<const FProperty*>(PropertyChain));
 	return FName(*PropertyFullName);
 }
 
@@ -483,13 +489,24 @@ bool FDataflowNode::ValidateConnections()
 					else if (const FString* PassthroughName = Property->FindMetaData(FDataflowNode::DataflowPassthrough))
 					{
 						void* PassthroughConnectionAddress = OutputConnection->GetPassthroughRealAddress();
-						if(PassthroughConnectionAddress == nullptr)
+						if (PassthroughConnectionAddress == nullptr)
 						{
-							UE_LOG(LogChaos, Warning, TEXT("Missing DataflowPassthrough registration for (%s:%s)"), *GetName().ToString(),*PropName.ToString());
+							UE_LOG(LogChaos, Warning, TEXT("Missing DataflowPassthrough registration for (%s:%s)"), *GetName().ToString(), *PropName.ToString());
 							bHasValidConnections = false;
 						}
 
-						const FDataflowInput* PassthroughConnectionInput = FindInput(FName(*PassthroughName));
+						// Assume passthrough name is relative to current property name.
+						FString FullPassthroughName;
+						if (PropertyChain.Num() <= 1)
+						{
+							FullPassthroughName = *PassthroughName;
+						}
+						else
+						{
+							FullPassthroughName = FString::Format(TEXT("{0}.{1}"), { GetPropertyFullNameString(TConstArrayView<const FProperty*>(&PropertyChain[1], PropertyChain.Num() - 1)), *PassthroughName});
+						}
+
+						const FDataflowInput* PassthroughConnectionInput = FindInput(FName(FullPassthroughName));
 						const FDataflowInput* PassthroughConnectionInputFromArg = FindInput(PassthroughConnectionAddress);
 
 						if(PassthroughConnectionInputFromArg != PassthroughConnectionInput)
diff --git a/Engine/Source/Runtime/Experimental/Dataflow/Core/Public/Dataflow/DataflowNode.h b/Engine/Source/Runtime/Experimental/Dataflow/Core/Public/Dataflow/DataflowNode.h
--- a/Engine/Source/Runtime/Experimental/Dataflow/Core/Public/Dataflow/DataflowNode.h
+++ b/Engine/Source/Runtime/Experimental/Dataflow/Core/Public/Dataflow/DataflowNode.h
@@ -282,6 +282,7 @@ struct FDataflowNode
 	FOnNodeInvalidated& GetOnNodeInvalidatedDelegate() { return OnNodeInvalidatedDelegate; }
 
 private:
+	static FString GetPropertyFullNameString(const TConstArrayView<const FProperty*>& PropertyChain);
 	static FName GetPropertyFullName(const TArray<const FProperty*>& PropertyChain);
 	static FText GetPropertyDisplayNameText(const TArray<const FProperty*>& PropertyChain);
 	static uint32 GetPropertyOffset(const TArray<const FProperty*>& PropertyChain);
