Moved IAS config file writing code into IAS module

#rnx

diff --git a/Engine/Source/Developer/IoStoreUtilities/Private/IoStoreUtilities.cpp b/Engine/Source/Developer/IoStoreUtilities/Private/IoStoreUtilities.cpp
--- a/Engine/Source/Developer/IoStoreUtilities/Private/IoStoreUtilities.cpp
+++ b/Engine/Source/Developer/IoStoreUtilities/Private/IoStoreUtilities.cpp
@@ -9663,62 +9663,6 @@ bool UploadIoStoreContainerFiles(const UE::IO::IAS::FIoStoreUploadParams& Upload
 
 	UE::IO::IAS::FIoStoreUploadResult UploadResult = Result.ConsumeValueOrDie();
 
-	FString ConfigFilePath;
-	if (FParse::Value(FCommandLine::Get(), TEXT("ConfigFilePath="), ConfigFilePath))
-	{
-		FStringBuilderBase Sb;
-		Sb << TEXT("[Endpoint]") << TEXT("\r\n");
-
-		if (!UploadParams.DistributionUrl.IsEmpty())
-		{
-			Sb << TEXT("DistributionUrl=\"") << UploadParams.DistributionUrl << TEXT("\"\r\n");
-
-			if (!UploadParams.FallbackUrl.IsEmpty())
-			{
-				Sb << TEXT("FallbackUrl=\"") << UploadParams.FallbackUrl << TEXT("\"\r\n");
-			}
-		}
-		else
-		{
-			Sb << TEXT("ServiceUrl=\"") << UploadResult.ServiceUrl << TEXT("\"\r\n");
-		}
-
-		Sb << TEXT("TocPath=\"") << UploadResult.TocPath << TEXT("\"\r\n");
-
-		// Temporary solution to get replays working with encrypted on demand content
-		{
-			if (!UploadParams.EncryptionKeyName.IsEmpty())
-			{
-				const TCHAR* EncryptionKeyName = *(UploadParams.EncryptionKeyName);
-
-				TOptional<FNamedAESKey> EncryptionKey;
-				for (const TPair<FGuid, FNamedAESKey>& KeyPair: KeyChain.GetEncryptionKeys())
-				{
-					if (KeyPair.Value.Name.Compare(EncryptionKeyName, ESearchCase::IgnoreCase) == 0)
-					{
-						EncryptionKey.Emplace(KeyPair.Value);
-					}
-				}
-
-				if (EncryptionKey)
-				{
-					FString KeyString = FBase64::Encode(EncryptionKey.GetValue().Key.Key, FAES::FAESKey::KeySize);
-					Sb << TEXT("ContentKey=\"") << EncryptionKey.GetValue().Guid.ToString() << TEXT(":") << KeyString << TEXT("\"\r\n");
-				}
-				else
-				{
-					UE_LOG(LogIoStore, Warning, TEXT("Failed to encryption key '%s' in key chain"), EncryptionKeyName);
-				}
-			}
-		}
-
-		UE_LOG(LogIoStore, Display, TEXT("Saving on demand config file '%s'"), *ConfigFilePath);
-		if (FFileHelper::SaveStringToFile(Sb.ToString(), *ConfigFilePath) == false)
-		{
-			UE_LOG(LogIoStore, Error, TEXT("Failed to save on demand config file '%s'"), *ConfigFilePath);
-		}
-	}
-
 	return Result.IsOk();
 }
 
diff --git a/Engine/Source/Runtime/Experimental/IoStoreOnDemand/Private/IoStoreOnDemand.cpp b/Engine/Source/Runtime/Experimental/IoStoreOnDemand/Private/IoStoreOnDemand.cpp
--- a/Engine/Source/Runtime/Experimental/IoStoreOnDemand/Private/IoStoreOnDemand.cpp
+++ b/Engine/Source/Runtime/Experimental/IoStoreOnDemand/Private/IoStoreOnDemand.cpp
@@ -1073,14 +1073,11 @@ TIoStatusOr<FIoStoreUploadParams> FIoStoreUploadParams::Parse(const TCHAR* Comma
 	Params.bDeletePakFiles = !FParse::Param(CommandLine, TEXT("KeepPakFiles"));
 	Params.bWriteTocToDisk = FParse::Param(CommandLine, TEXT("WriteTocToDisk"));
 
-	if (Params.bWriteTocToDisk)
-	{
-		// If we keep this feature we should allow the caller to set this path themselves rather than rely on the config file location
-		FString ConfigFilePath;
-		FParse::Value(FCommandLine::Get(), TEXT("ConfigFilePath="), ConfigFilePath);
-
-		Params.TocOutputDir = FPaths::GetPath(ConfigFilePath);
-	}
+	// If we keep this feature we should allow the caller to set this path
+	// themselves rather than rely on the config file location
+	FString ConfigFilePath;
+	FParse::Value(FCommandLine::Get(), TEXT("ConfigFilePath="), ConfigFilePath);
+	Params.TocOutputDir = FPaths::GetPath(ConfigFilePath);
 
 	if (FIoStatus Validation = Params.Validate(); !Validation.IsOk())
 	{
@@ -1119,6 +1116,65 @@ FIoStatus FIoStoreUploadParams::Validate() const
 	return FIoStatus::Ok;
 }
 
+static void WriteConfigFile(
+	const FIoStoreUploadParams& UploadParams,
+	const FIoStoreUploadResult& UploadResult,
+	const FKeyChain& KeyChain)
+{
+	FStringBuilderBase Sb;
+	Sb << TEXT("[Endpoint]") << TEXT("\r\n");
+
+	if (!UploadParams.DistributionUrl.IsEmpty())
+	{
+		Sb << TEXT("DistributionUrl=\"") << UploadParams.DistributionUrl << TEXT("\"\r\n");
+
+		if (!UploadParams.FallbackUrl.IsEmpty())
+		{
+			Sb << TEXT("FallbackUrl=\"") << UploadParams.FallbackUrl << TEXT("\"\r\n");
+		}
+	}
+	else
+	{
+		Sb << TEXT("ServiceUrl=\"") << UploadResult.ServiceUrl << TEXT("\"\r\n");
+	}
+
+	Sb << TEXT("TocPath=\"") << UploadResult.TocPath << TEXT("\"\r\n");
+
+	// Temporary solution to get replays working with encrypted on demand content
+	{
+		if (!UploadParams.EncryptionKeyName.IsEmpty())
+		{
+			const TCHAR* EncryptionKeyName = *(UploadParams.EncryptionKeyName);
+
+			TOptional<FNamedAESKey> EncryptionKey;
+			for (const TPair<FGuid, FNamedAESKey>& KeyPair: KeyChain.GetEncryptionKeys())
+			{
+				if (KeyPair.Value.Name.Compare(EncryptionKeyName, ESearchCase::IgnoreCase) == 0)
+				{
+					EncryptionKey.Emplace(KeyPair.Value);
+				}
+			}
+
+			if (EncryptionKey)
+			{
+				FString KeyString = FBase64::Encode(EncryptionKey.GetValue().Key.Key, FAES::FAESKey::KeySize);
+				Sb << TEXT("ContentKey=\"") << EncryptionKey.GetValue().Guid.ToString() << TEXT(":") << KeyString << TEXT("\"\r\n");
+			}
+			else
+			{
+				UE_LOG(LogIoStore, Warning, TEXT("Failed to encryption key '%s' in key chain"), EncryptionKeyName);
+			}
+		}
+	}
+
+	FString ConfigFilePath = UploadParams.TocOutputDir / TEXT("IoStoreOnDemand.ini");
+	UE_LOG(LogIoStore, Display, TEXT("Saving on demand config file '%s'"), *ConfigFilePath);
+	if (FFileHelper::SaveStringToFile(Sb.ToString(), *ConfigFilePath) == false)
+	{
+		UE_LOG(LogIoStore, Error, TEXT("Failed to save on demand config file '%s'"), *ConfigFilePath);
+	}
+}
+
 TIoStatusOr<FIoStoreUploadResult> UploadContainerFiles(
 	const FIoStoreUploadParams& UploadParams,
 	TConstArrayView<FString> ContainerFiles,
@@ -1497,6 +1553,11 @@ TIoStatusOr<FIoStoreUploadResult> UploadContainerFiles(
 		}
 	}
 
+	if (!UploadParams.TocOutputDir.IsEmpty())
+	{
+		WriteConfigFile(UploadParams, UploadResult, KeyChain);
+	}
+
 	{
 		const double Duration = FPlatformTime::Seconds() - StartTime;
 		
