Constraints: remove prerex when removing the constraint

diff --git a/Engine/Source/Runtime/Experimental/Animation/Constraints/Private/TransformConstraint.cpp b/Engine/Source/Runtime/Experimental/Animation/Constraints/Private/TransformConstraint.cpp
--- a/Engine/Source/Runtime/Experimental/Animation/Constraints/Private/TransformConstraint.cpp
+++ b/Engine/Source/Runtime/Experimental/Animation/Constraints/Private/TransformConstraint.cpp
@@ -672,13 +672,17 @@ void UTickableTransformConstraint::InitConstraint(UWorld *InWorld)
 	SetupDependencies(InWorld);
 	RegisterDelegates();
 	bValid = true;
-
 }
 
 void UTickableTransformConstraint::TeardownConstraint(UWorld* InWorld)
 {
 	FConstraintTickFunction& ConstraintTick = ConstraintTicks.FindOrAdd(InWorld->GetCurrentLevel());
-	ConstraintTick.UnRegisterTickFunction(); 
+	ConstraintTick.SetTickFunctionEnable(false);
+	ConstraintTick.UnRegisterTickFunction();
+	if (FTickFunction* ChildTickFunction = GetChildHandleTickFunction())
+	{
+		ChildTickFunction->RemovePrerequisite(this, ConstraintTick);
+	}
 	ConstraintTicks.Remove(InWorld->GetCurrentLevel());
 }
 
