Iris - Gracefully handle RPC payload deserialization errors.

diff --git a/Engine/Source/Runtime/Experimental/Iris/Core/Private/Iris/ReplicationSystem/NetBlob/NetRPC.cpp b/Engine/Source/Runtime/Experimental/Iris/Core/Private/Iris/ReplicationSystem/NetBlob/NetRPC.cpp
--- a/Engine/Source/Runtime/Experimental/Iris/Core/Private/Iris/ReplicationSystem/NetBlob/NetRPC.cpp
+++ b/Engine/Source/Runtime/Experimental/Iris/Core/Private/Iris/ReplicationSystem/NetBlob/NetRPC.cpp
@@ -135,7 +135,7 @@ void FNetRPC::DeserializeWithObject(FNetSerializationContext& Context, FNetRefHa
 	{
 		UE_LOG(LogIrisRpc, Error, TEXT("DeserializeWithObject::Skipping RPC due missing object or function."));
 
-		// Stop deserializing and seek past the entire payload if the Resolve failed
+		// Stop deserializing and seek past the entire payload if the resolve failed
 		Context.GetBitStreamReader()->Seek(PostNetRPCPos);
 		return;
 	}
@@ -145,7 +145,15 @@ void FNetRPC::DeserializeWithObject(FNetSerializationContext& Context, FNetRefHa
 	
 	if (!Context.HasErrorOrOverflow())
 	{
-		check(PostNetRPCPos == Context.GetBitStreamReader()->GetPosBits());
+		// Just because the serialization didn't detect an error doesn't mean everything is ok. Validate stream position.
+		if (PostNetRPCPos != Context.GetBitStreamReader()->GetPosBits())
+		{
+			ensureMsgf(PostNetRPCPos == Context.GetBitStreamReader()->GetPosBits(), TEXT("Bitstream mismatch while deserializing function %s. Actual stream position: %u Expected stream position: %u"), ToCStr(BlobDescriptor->DebugName), Context.GetBitStreamReader()->GetPosBits(), PostNetRPCPos);
+			Context.GetBitStreamReader()->Seek(PostNetRPCPos);
+			Context.SetError(GNetError_BitStreamError);
+			// Make sure the RPC won't be exeuted regardless of how errors are handled.
+			Function = nullptr;
+		}
 	}
 }
 
