Iris
- Fixed issue where subobjects of objects pending dormant all would not always be polled as they should

diff --git a/Engine/Source/Runtime/Experimental/Iris/Core/Private/Iris/ReplicationSystem/ObjectReplicationBridge.cpp b/Engine/Source/Runtime/Experimental/Iris/Core/Private/Iris/ReplicationSystem/ObjectReplicationBridge.cpp
--- a/Engine/Source/Runtime/Experimental/Iris/Core/Private/Iris/ReplicationSystem/ObjectReplicationBridge.cpp
+++ b/Engine/Source/Runtime/Experimental/Iris/Core/Private/Iris/ReplicationSystem/ObjectReplicationBridge.cpp
@@ -929,7 +929,9 @@ void UObjectReplicationBridge::BuildPollList(UE::Net::FNetBitArrayView ObjectsCo
 		{
 			if (const uint32 InternalObjectIndex = LocalNetRefHandleManager.GetInternalIndex(HandlePendingFlush))
 			{
-				if (!ObjectsConsideredForPolling.IsBitSet(InternalObjectIndex))
+				// If HandlePendingFlush is relevant, poll it this frame, and treat it as a forcenetupdate in order to also schedule subobjects correctly
+				// Out of scope dormant objects will be Marked Dirty to ensure that they are scheduled for polling once relevant
+				if (RelevantObjects.IsBitSet(InternalObjectIndex))
 				{
 					ObjectsConsideredForPolling.SetBit(InternalObjectIndex);
 					ForceNetUpdateObjects.SetBit(InternalObjectIndex);
@@ -947,15 +949,13 @@ void UObjectReplicationBridge::BuildPollList(UE::Net::FNetBitArrayView ObjectsCo
 	* are replicated atomically this polling propagation is required.
 	*/
 	{
-		IRIS_PROFILER_SCOPE(BuildPollList_PropagatePolling);
-
-		auto PropagateSubObjectNetForceUpdateToOwner = [&LocalNetRefHandleManager, &ObjectsConsideredForPolling](uint32 InternalObjectIndex)
+		auto PropagateSubObjectDirtinessToOwner = [&LocalNetRefHandleManager, &ObjectsConsideredForPolling](uint32 InternalObjectIndex)
 		{
 			const FNetRefHandleManager::FReplicatedObjectData& ObjectData = LocalNetRefHandleManager.GetReplicatedObjectDataNoCheck(InternalObjectIndex);
 			ObjectsConsideredForPolling.SetBit(ObjectData.SubObjectRootIndex);
 		};
 
-		auto PropagateOwnerNetForceUpdateToSubObject = [&LocalNetRefHandleManager, &ObjectsConsideredForPolling](uint32 InternalObjectIndex)
+		auto PropagateOwnerDirtinessToSubObjects = [&LocalNetRefHandleManager, &ObjectsConsideredForPolling](uint32 InternalObjectIndex)
 		{
 			const FNetRefHandleManager::FReplicatedObjectData& ObjectData = LocalNetRefHandleManager.GetReplicatedObjectDataNoCheck(InternalObjectIndex);
 			for (const FInternalNetRefIndex SubObjectInternalIndex : LocalNetRefHandleManager.GetSubObjects(InternalObjectIndex))
@@ -964,6 +964,8 @@ void UObjectReplicationBridge::BuildPollList(UE::Net::FNetBitArrayView ObjectsCo
 			}
 		};
 
+		IRIS_PROFILER_SCOPE(BuildPollList_PropagatePolling);
+
 		// Update subobjects' owner first and owners' subobjects second. It's the only way to properly mark all groups of objects in two passes.
 		const FNetBitArrayView SubObjects = MakeNetBitArrayView(LocalNetRefHandleManager.GetSubObjectInternalIndices());
 		const FNetBitArrayView ForceNetUpdateObjects = DirtyNetObjectTracker.GetForceNetUpdateObjects();
@@ -976,8 +978,7 @@ void UObjectReplicationBridge::BuildPollList(UE::Net::FNetBitArrayView ObjectsCo
 			
 			ForceNetUpdateAndRelevantObjectsView.Set(RelevantObjects, FNetBitArray::AndOp, ForceNetUpdateObjects);
 
-			FNetBitArrayView::ForAllSetBits(ForceNetUpdateAndRelevantObjectsView, SubObjects, FNetBitArray::AndOp, PropagateSubObjectNetForceUpdateToOwner);
-			FNetBitArrayView::ForAllSetBits(ForceNetUpdateAndRelevantObjectsView, SubObjects, FNetBitArray::AndNotOp, PropagateOwnerNetForceUpdateToSubObject);
+			FNetBitArrayView::ForAllSetBits(ForceNetUpdateAndRelevantObjectsView, SubObjects, FNetBitArray::AndOp, PropagateSubObjectDirtinessToOwner);
 		}
 		else
 		{
@@ -993,8 +994,7 @@ void UObjectReplicationBridge::BuildPollList(UE::Net::FNetBitArrayView ObjectsCo
 				DirtyAndRelevantObjectsView.Combine(RelevantObjects, FNetBitArray::AndOp);
 			}
 
-			FNetBitArrayView::ForAllSetBits(DirtyAndRelevantObjectsView, SubObjects, FNetBitArray::AndOp, PropagateSubObjectNetForceUpdateToOwner);
-			FNetBitArrayView::ForAllSetBits(DirtyAndRelevantObjectsView, SubObjects, FNetBitArray::AndNotOp, PropagateOwnerNetForceUpdateToSubObject);
+			FNetBitArrayView::ForAllSetBits(DirtyAndRelevantObjectsView, SubObjects, FNetBitArray::AndOp, PropagateSubObjectDirtinessToOwner);
 		}
 			
 		// If an object with dependents is about to be polled, force it's dependents to poll at the same time.
@@ -1013,6 +1013,9 @@ void UObjectReplicationBridge::BuildPollList(UE::Net::FNetBitArrayView ObjectsCo
 				}
 			);
 		}
+
+		// Currently we must enforce polling SubObjects with owner
+		FNetBitArrayView::ForAllSetBits(ObjectsConsideredForPolling, SubObjects, FNetBitArray::AndNotOp, PropagateOwnerDirtinessToSubObjects);
 	}
 }
 	
@@ -1191,6 +1194,8 @@ void UObjectReplicationBridge::SetObjectWantsToBeDormant(FNetRefHandle Handle, b
 		// Update pending dormancy status
 		WantToBeDormantObjects.SetBitValue(InternalObjectIndex, bWantsToBeDormant);
 
+		UE_LOG_OBJECTREPLICATIONBRIDGE(VeryVerbose, TEXT("SetObjectWantsToBeDormant for %s ( InternalIndex: %u ) %d "), *Handle.ToString(), InternalObjectIndex, bWantsToBeDormant ? 1 : 0);
+
 		// If we want to be dormant we want to make sure we poll the object one more time
 		if (bWantsToBeDormant)
 		{
@@ -1213,6 +1218,9 @@ void UObjectReplicationBridge::SetObjectWantsToBeDormant(FNetRefHandle Handle, b
 
 void UObjectReplicationBridge::ForceUpdateWantsToBeDormantObject(FNetRefHandle Handle)
 {
+	UE_LOG_OBJECTREPLICATIONBRIDGE(VeryVerbose, TEXT("ForceUpdateWantsToBeDormantObject for %s"), *Handle.ToString());
+
+	ReplicationSystem->MarkDirty(Handle);
 	DormantHandlesPendingFlush.Add(Handle);
 }
 
diff --git a/Engine/Source/Runtime/Experimental/Iris/Core/Private/Iris/ReplicationSystem/ReplicationWriter.cpp b/Engine/Source/Runtime/Experimental/Iris/Core/Private/Iris/ReplicationSystem/ReplicationWriter.cpp
--- a/Engine/Source/Runtime/Experimental/Iris/Core/Private/Iris/ReplicationSystem/ReplicationWriter.cpp
+++ b/Engine/Source/Runtime/Experimental/Iris/Core/Private/Iris/ReplicationSystem/ReplicationWriter.cpp
@@ -800,7 +800,7 @@ void FReplicationWriter::InternalUpdateDirtyChangeMasks(const FChangeMaskCache&
 			continue;
 		}
 
-		// We want to accumulate dirty changes even if we are going out of scope in case we get readded to scope before repliation has ended.
+		// We want to accumulate dirty changes even if we are going out of scope in case we get re-added to scope before replication has ended.
 		const bool bMarkScopedObjectDirty = ObjectsInScope.GetBit(Entry.InternalIndex);
 		if (bMarkScopedObjectDirty)
 		{
diff --git a/Engine/Source/Runtime/Net/Core/Private/Net/Core/DirtyNetObjectTracker/GlobalDirtyNetObjectTracker.cpp b/Engine/Source/Runtime/Net/Core/Private/Net/Core/DirtyNetObjectTracker/GlobalDirtyNetObjectTracker.cpp
--- a/Engine/Source/Runtime/Net/Core/Private/Net/Core/DirtyNetObjectTracker/GlobalDirtyNetObjectTracker.cpp
+++ b/Engine/Source/Runtime/Net/Core/Private/Net/Core/DirtyNetObjectTracker/GlobalDirtyNetObjectTracker.cpp
@@ -224,7 +224,7 @@ FGlobalDirtyNetObjectTracker::FPimpl::~FPimpl()
 
 void FGlobalDirtyNetObjectTracker::FPimpl::ResetDirtyNetObjects()
 {
-	if (!ensureMsgf(Pollers.IsNoBitSet(), TEXT("FGlobalDirtyNetObjectTracker poller %u forgot to call ResetDirtNetObjects."), Pollers.FindFirstOne()))
+	if (!ensureMsgf(Pollers.IsNoBitSet(), TEXT("FGlobalDirtyNetObjectTracker poller %u forgot to call ResetDirtyNetObjects."), Pollers.FindFirstOne()))
 	{
 		Pollers.Reset();
 
