Constraints: pre-evaluate child before computing the constraints
    - controls transforms may depend on a complex rig (space projection using the current bones hierarchy transforms is a simple example) that needs to be evaluated before retrieving the child's transformation.

diff --git a/Engine/Plugins/Animation/ControlRig/Source/ControlRig/Private/Constraints/ControlRigTransformableHandle.cpp b/Engine/Plugins/Animation/ControlRig/Source/ControlRig/Private/Constraints/ControlRigTransformableHandle.cpp
--- a/Engine/Plugins/Animation/ControlRig/Source/ControlRig/Private/Constraints/ControlRigTransformableHandle.cpp
+++ b/Engine/Plugins/Animation/ControlRig/Source/ControlRig/Private/Constraints/ControlRigTransformableHandle.cpp
@@ -57,6 +57,15 @@ bool UTransformableControlHandle::IsValid(const bool bDeepCheck) const
 	return true;
 }
 
+void UTransformableControlHandle::PreEvaluate(const bool bTick) const
+{
+	if (!ControlRig.IsValid())
+	{
+		return;
+	}
+	return bTick ? TickTarget() : ControlRig->Evaluate_AnyThread();
+}
+
 void UTransformableControlHandle::TickTarget() const
 {
 	if (const USkeletalMeshComponent* SkeletalMeshComponent = GetSkeletalMesh())
diff --git a/Engine/Plugins/Animation/ControlRig/Source/ControlRig/Public/Constraints/ControlRigTransformableHandle.h b/Engine/Plugins/Animation/ControlRig/Source/ControlRig/Public/Constraints/ControlRigTransformableHandle.h
--- a/Engine/Plugins/Animation/ControlRig/Source/ControlRig/Public/Constraints/ControlRigTransformableHandle.h
+++ b/Engine/Plugins/Animation/ControlRig/Source/ControlRig/Public/Constraints/ControlRigTransformableHandle.h
@@ -74,6 +74,13 @@ public:
 	/** Tick any skeletal mesh related to the bound component. */ 
 	virtual void TickTarget() const override;
 
+	/**
+	 * Perform any pre-evaluation of the handle to ensure that the transform data is up to date.
+	 * @param bTick to force any pre-evaluation ticking. The rig will still be pre-evaluated even
+	 * if bTick is false (it just won't tick the bound skeletal meshes) Default is false.
+	*/
+	virtual void PreEvaluate(const bool bTick = false) const override;
+
 	/** Registers/Unregisters useful delegates to track changes in the control's transform. */
 	void UnregisterDelegates() const;
 	void RegisterDelegates();
diff --git a/Engine/Source/Runtime/Experimental/Animation/Constraints/Private/TransformConstraint.cpp b/Engine/Source/Runtime/Experimental/Animation/Constraints/Private/TransformConstraint.cpp
--- a/Engine/Source/Runtime/Experimental/Animation/Constraints/Private/TransformConstraint.cpp
+++ b/Engine/Source/Runtime/Experimental/Animation/Constraints/Private/TransformConstraint.cpp
@@ -197,6 +197,33 @@ private:
 	
 }
 
+namespace ConstraintLocals
+{
+
+static bool	bPreEvaluateChild = true;
+static FAutoConsoleVariableRef CVarPreEvaluateChild(
+	TEXT("Constraints.PreEvaluateChild"),
+	bPreEvaluateChild,
+	TEXT("Force child evaluation before constraint computation.")
+	);
+	
+static bool	bPreTickChild = false;
+static FAutoConsoleVariableRef CVarPreTickChild(
+	TEXT("Constraints.PreTickChild"),
+	bPreTickChild,
+	TEXT("Force child ticking before constraint computation.")
+	);
+
+void PreEvaluateHandle(const TObjectPtr<UTransformableHandle>& InHandle) 
+{
+	if ((bPreEvaluateChild || bPreTickChild) && ::IsValid(InHandle))
+	{
+		InHandle->PreEvaluate(bPreTickChild);
+	}
+}
+
+}
+
 /** 
  * UTickableTransformConstraint
  **/
@@ -747,6 +774,8 @@ FConstraintTickFunction::ConstraintFunction UTickableTranslationConstraint::GetF
 			return;
 		}
 
+		ConstraintLocals::PreEvaluateHandle(ChildTRSHandle);
+		
 		const FVector ParentTranslation = GetParentGlobalTransform().GetLocation();
 		FTransform Transform = GetChildGlobalTransform();
 		const FVector ChildTranslation = Transform.GetLocation();
@@ -888,6 +917,8 @@ FConstraintTickFunction::ConstraintFunction UTickableRotationConstraint::GetFunc
 			return;
 		}
 
+		ConstraintLocals::PreEvaluateHandle(ChildTRSHandle);
+		
 		const FQuat ParentRotation = GetParentGlobalTransform().GetRotation();
 		FTransform Transform = GetChildGlobalTransform();
 		const FQuat ChildRotation = Transform.GetRotation();
@@ -1034,6 +1065,8 @@ FConstraintTickFunction::ConstraintFunction UTickableScaleConstraint::GetFunctio
 		{
 			return;
 		}
+
+		ConstraintLocals::PreEvaluateHandle(ChildTRSHandle);
 		
 		const FVector ParentScale = GetParentGlobalTransform().GetScale3D();
 		FTransform Transform = GetChildGlobalTransform();
@@ -1199,7 +1232,9 @@ FConstraintTickFunction::ConstraintFunction UTickableParentConstraint::GetFuncti
 			OutTransform.SetRotation(NewRotation);
 			OutTransform.SetScale3D(NewScale);
 		};
-	
+
+		ConstraintLocals::PreEvaluateHandle(ChildTRSHandle);
+		
 		const FTransform ParentTransform = GetParentGlobalTransform();
 		
 		FTransform TargetTransform = (!bMaintainOffset) ? ParentTransform : OffsetTransform * ParentTransform;
@@ -1347,6 +1382,8 @@ FConstraintTickFunction::ConstraintFunction UTickableLookAtConstraint::GetFuncti
 		{
 			return;
 		}
+
+		ConstraintLocals::PreEvaluateHandle(ChildTRSHandle);
 		
 		const FTransform ParentTransform = GetParentGlobalTransform();
 		const FTransform ChildTransform = GetChildGlobalTransform();
diff --git a/Engine/Source/Runtime/Experimental/Animation/Constraints/Private/TransformableHandle.cpp b/Engine/Source/Runtime/Experimental/Animation/Constraints/Private/TransformableHandle.cpp
--- a/Engine/Source/Runtime/Experimental/Animation/Constraints/Private/TransformableHandle.cpp
+++ b/Engine/Source/Runtime/Experimental/Animation/Constraints/Private/TransformableHandle.cpp
@@ -61,6 +61,14 @@ void UTransformableHandle::OnBindingIDsUpdated(const TMap<UE::MovieScene::FFixed
 	}
 }
 
+void UTransformableHandle::PreEvaluate(const bool bTick) const
+{
+	if (bTick)
+	{
+		TickTarget();
+	}
+}
+
 /**
  * UTransformableComponentHandle
  */
diff --git a/Engine/Source/Runtime/Experimental/Animation/Constraints/Public/TransformableHandle.h b/Engine/Source/Runtime/Experimental/Animation/Constraints/Public/TransformableHandle.h
--- a/Engine/Source/Runtime/Experimental/Animation/Constraints/Public/TransformableHandle.h
+++ b/Engine/Source/Runtime/Experimental/Animation/Constraints/Public/TransformableHandle.h
@@ -79,6 +79,12 @@ public:
 	
 	/** Perform any special ticking needed for this handle, by default it does nothing, todo need to see if we need to tick control rig also*/
 	virtual void TickTarget() const {};
+
+	/**
+	 * Perform any pre-evaluation of the handle to ensure that the transform data are up to date.
+	 * @param bTick to force any pre-evaluation ticking. Default is false.
+	*/
+	virtual void PreEvaluate(const bool bTick = false) const;
 	
 	/** Get the array of float channels for the specified section*/
 	CONSTRAINTS_API virtual TArrayView<FMovieSceneFloatChannel*>  GetFloatChannels(const UMovieSceneSection* InSection) const PURE_VIRTUAL(GetFloatChannels, return TArrayView<FMovieSceneFloatChannel*>(); );
