Fix for incorrect trace geometry being generated

diff --git a/Engine/Source/Runtime/Experimental/Chaos/Public/PhysicsProxy/GeometryCollectionPhysicsProxy.h b/Engine/Source/Runtime/Experimental/Chaos/Public/PhysicsProxy/GeometryCollectionPhysicsProxy.h
--- a/Engine/Source/Runtime/Experimental/Chaos/Public/PhysicsProxy/GeometryCollectionPhysicsProxy.h
+++ b/Engine/Source/Runtime/Experimental/Chaos/Public/PhysicsProxy/GeometryCollectionPhysicsProxy.h
@@ -285,6 +285,7 @@ public:
 	// set the world transform ( this needs to be called on the game thread ) 
 	CHAOS_API void SetWorldTransform_External(const FTransform& WorldTransform);
 	CHAOS_API const FTransform& GetPreviousWorldTransform_External() const { return PreviousWorldTransform_External; }
+	CHAOS_API const FTransform& GetWorldTransform_External() { return WorldTransform_External; }
 
 	// todo(chaos): Remove this and move to a cook time approach of the SM data based on the GC property
 	// Set whether the GC should be using collision from the Static Mesh or the GC itself for game thread traces ( this needs to be called on the game thread )
diff --git a/Engine/Source/Runtime/Experimental/GeometryCollectionEngine/Private/GeometryCollection/GeometryCollectionComponent.cpp b/Engine/Source/Runtime/Experimental/GeometryCollectionEngine/Private/GeometryCollection/GeometryCollectionComponent.cpp
--- a/Engine/Source/Runtime/Experimental/GeometryCollectionEngine/Private/GeometryCollection/GeometryCollectionComponent.cpp
+++ b/Engine/Source/Runtime/Experimental/GeometryCollectionEngine/Private/GeometryCollection/GeometryCollectionComponent.cpp
@@ -3896,10 +3896,11 @@ void UGeometryCollectionComponent::RegisterAndInitializePhysicsProxy()
 			AddParams.bDoubleSided = BodySetup->bDoubleSidedGeometry;
 			AddParams.CollisionData = BodyCollisionData;
 			AddParams.CollisionTraceType = BodySetup->GetCollisionTraceFlag();
-			AddParams.Scale = BodyInstance.Scale3D;
-			AddParams.SimpleMaterial = BodyInstance.GetSimplePhysicalMaterial();
+			AddParams.WorldTransform = PhysicsProxy->GetWorldTransform_External();
+			AddParams.Scale = AddParams.WorldTransform.GetScale3D();
+			AddParams.SimpleMaterial = GetPhysicalMaterial();
 			AddParams.LocalTransform = InToLocal;
-			AddParams.WorldTransform = PhysicsProxy->GetSimParameters().WorldTransform;
+			
 			AddParams.Geometry = &BodySetup->AggGeom;
 			ChaosInterface::CreateGeometry(AddParams, OutGeoms, OutShapes);
 		}
