Chaos - fix capsule-mesh manifold generation

Fix issue with one-sided mesh collision when a long capsule's cylinder is just touching one of the edges from underneath the triangle. In this case we want to ignore the nearest (downward facing) edge contact and instead generate a collision with one of the opposing edges.

This fix correctly eliminates the downward-facing edge contact, but does not eliminate the opposing edge contact.

diff --git a/Engine/Source/Runtime/Experimental/Chaos/Private/Chaos/Collision/CapsuleTriangleContactPoint.cpp b/Engine/Source/Runtime/Experimental/Chaos/Private/Chaos/Collision/CapsuleTriangleContactPoint.cpp
--- a/Engine/Source/Runtime/Experimental/Chaos/Private/Chaos/Collision/CapsuleTriangleContactPoint.cpp
+++ b/Engine/Source/Runtime/Experimental/Chaos/Private/Chaos/Collision/CapsuleTriangleContactPoint.cpp
@@ -4,16 +4,10 @@
 #include "Chaos/Collision/ContactPoint.h"
 #include "Chaos/Triangle.h"
 
-//PRAGMA_DISABLE_OPTIMIZATION
+//UE_DISABLE_OPTIMIZATION
 
 namespace Chaos
 {
-	namespace CVars
-	{
-		bool bChaosCollisionCapsuleTriMeshSATCull = false;
-		FAutoConsoleVariableRef CVarUseGenericSweptConvexConstraints(TEXT("p.Chaos.Collision.CapsuleTriMeshSATCull"), bChaosCollisionCapsuleTriMeshSATCull, TEXT("Enable the SAT cull check in capsule-triangle [default: false]."));
-	}
-
 	void AddCapsuleTriangleParallelEdgeManifoldContacts(const FVec3& P0, const FVec3& P1, const FVec3& EdgeP0, const FVec3& EdgeP1, const FReal R, const FReal RejectDistanceSq, const FReal NormalToleranceSq, FContactPointManifold& OutContactPoints);
 
 	// Return true if the value V is in the closed/inclusive range [RangeMin, RangeMax]
@@ -74,8 +68,8 @@ namespace Chaos
 		// Tolerances
 		// @todo(chaos): Everything below needs to be made size agnostic (i.e., all checks with a distance tolerance
 		// need to cope with very larger or small segments and triangles). This will probably not be good enough...(add some tests)
-		const FReal DistanceTolerance = FReal(1.e-10) * Capsule.GetHeight();
-		const FReal NormalTolerance = FReal(1.e-8);
+		const FReal DistanceTolerance = FReal(1.e-5) * Capsule.GetHeight();
+		const FReal NormalTolerance = FReal(1.e-5);
 		const FReal NormalToleranceSq = NormalTolerance * NormalTolerance;
 		const FReal FaceContactSinAngleThreshold = FReal(0.34);	// ~Sin(20deg)
 
@@ -108,55 +102,67 @@ namespace Chaos
 			return;
 		}
 
-		// Early out if any edge saparating axis distance is above the reject distance
-		if (CVars::bChaosCollisionCapsuleTriMeshSATCull)
+		// Cull based on signed distance to capsule segment
+		// The edge-segment data is saved for used later on to generate contacts if necessary
+		FVec3 EdgeSegmentDeltas[3];
+		FVec3 EdgeEdgePs[3];
+		FVec3 EdgeSegmentPs[3];
+		FReal EdgeEdgeTs[3];
+		FReal EdgeSegmentTs[3];
+		FReal EdgeDistSqs[3];
+		FReal EdgeDistSigns[3];
+		FReal EdgeDotFace[3];
+		int32 EdgeVertexIndex0 = 2;
+		for (int32 EdgeIndex = 0; EdgeIndex < 3; ++EdgeIndex)
 		{
-			// Separating axis test on the edges
-			const auto& SATEdgeCheck = [&FaceN, &R, &Centroid, &RejectDistance, &NormalTolerance](const FVec3& P0, const FVec3& P1, const FVec3& EdgeP0, const FVec3& EdgeP1) -> bool
+			const int32 EdgeVertexIndex1 = EdgeIndex;
+			const FVec3& EdgeP0 = Triangle.GetVertex(EdgeVertexIndex0);
+			const FVec3& EdgeP1 = Triangle.GetVertex(EdgeVertexIndex1);
+			EdgeVertexIndex0 = EdgeVertexIndex1;
+
+			// Find the nearest point on the capsule segment to the edge segment
+			FReal SegmentT, EdgeT;
+			FVec3 SegmentP, EdgeP;
+			Utilities::NearestPointsOnLineSegments(P0, P1, EdgeP0, EdgeP1, SegmentT, EdgeT, SegmentP, EdgeP);
+
+			// Calculate the separation vector, correct for sign
+			FVec3 SegmentEdgeN = SegmentP - EdgeP;
+			FReal SegmentEdgeDistSign = FReal(1);
+			const FReal SegmentEdgeDistSq = SegmentEdgeN.SizeSquared();
+
+			// Separating axis always points away from the triangle
+			if (FVec3::DotProduct(EdgeP1 - Centroid, SegmentEdgeN) < FReal(0))
 			{
-				FVec3 EdgeSeparationAxis = FVec3::CrossProduct(P1 - P0, EdgeP1 - EdgeP0);
-				if (!EdgeSeparationAxis.Normalize())
-				{
-					// Parallel lines are handled later
-					return true;
-				}
+				SegmentEdgeN = -SegmentEdgeN;
+				SegmentEdgeDistSign = FReal(-1);
+			}
 
-				// Separating axis always points away from the triangle
-				if (FVec3::DotProduct(EdgeP1 - Centroid, EdgeSeparationAxis) < FReal(0))
-				{
-					EdgeSeparationAxis = -EdgeSeparationAxis;
-				}
+			const FReal DotFace = FVec3::DotProduct(SegmentEdgeN, FaceN);
 
-				// We want to generate contacts with positive separation when within cull distance,
-				// but only if we're on the outside of the face. If we're below the face here
-				// we will not make a contact unless we actually touch.
+			if (SegmentEdgeDistSign > FReal(0))
+			{
+				// We generate contacts when separation is within cull distance
+				// Treat CullDistance as zero when colliding with the underneath of a triangle
 				FReal SeparationAxisCullDistance = RejectDistance;
-				if (FVec3::DotProduct(EdgeSeparationAxis, FaceN) < -NormalTolerance)
+				if (DotFace < -NormalTolerance)
 				{
 					SeparationAxisCullDistance = R;
 				}
 
-				// Separating distance
-				const FReal EdgeSeparationDist = FVec3::DotProduct(P0 - EdgeP0, EdgeSeparationAxis);
-
-				return EdgeSeparationDist <= SeparationAxisCullDistance;
-			};
-
-			if (!SATEdgeCheck(P0, P1, Triangle.GetVertex(2), Triangle.GetVertex(0)))
-			{
-				// Separated from triangle
-				return;
-			}
-			if (!SATEdgeCheck(P0, P1, Triangle.GetVertex(0), Triangle.GetVertex(1)))
-			{
-				// Separated from triangle
-				return;
-			}
-			if (!SATEdgeCheck(P0, P1, Triangle.GetVertex(1), Triangle.GetVertex(2)))
-			{
-				// Separated from triangle
-				return;
+				if (SegmentEdgeDistSq > FMath::Square(SeparationAxisCullDistance))
+				{
+					return;
+				}
 			}
+
+			EdgeSegmentDeltas[EdgeIndex] = SegmentEdgeN;
+			EdgeEdgePs[EdgeIndex] = EdgeP;
+			EdgeSegmentPs[EdgeIndex] = SegmentP;
+			EdgeEdgeTs[EdgeIndex] = EdgeT;
+			EdgeSegmentTs[EdgeIndex] = SegmentT;
+			EdgeDistSqs[EdgeIndex] = SegmentEdgeDistSq;
+			EdgeDistSigns[EdgeIndex] = SegmentEdgeDistSign;
+			EdgeDotFace[EdgeIndex] = DotFace;
 		}
 
 		// Edge plane normals and signed distances to each segment point
@@ -241,7 +247,7 @@ namespace Chaos
 		bIsParallelEdge[0] = ((EdgeD0s[0] >= FReal(0)) & (EdgeD1s[0] >= FReal(0)) & bEqualEdgeDist0) != 0;
 		bIsParallelEdge[1] = ((EdgeD0s[1] >= FReal(0)) & (EdgeD1s[1] >= FReal(0)) & bEqualEdgeDist1) != 0;
 		bIsParallelEdge[2] = ((EdgeD0s[2] >= FReal(0)) & (EdgeD1s[2] >= FReal(0)) & bEqualEdgeDist2) != 0;
-		int32 EdgeVertexIndex0 = 2;
+		EdgeVertexIndex0 = 2;
 		if ((bIsParallelFace & (bIsParallelEdge[0] | bIsParallelEdge[1] | bIsParallelEdge[2])) != 0)
 		{
 			for (int32 EdgeIndex = 0; EdgeIndex < 3; ++EdgeIndex)
@@ -330,37 +336,28 @@ namespace Chaos
 			const FVec3& EdgeP1 = Triangle.GetVertex(EdgeVertexIndex1);
 			EdgeVertexIndex0 = EdgeVertexIndex1;
 
+			// Reuse edge-segment data calculated in cull check above
+			FVec3 SegmentEdgeN = EdgeSegmentDeltas[EdgeIndex];
+			const FVec3& EdgeP = EdgeEdgePs[EdgeIndex];
+			const FVec3& SegmentP = EdgeSegmentPs[EdgeIndex];
+			const FReal EdgeT = EdgeEdgeTs[EdgeIndex];
+			const FReal SegmentT = EdgeSegmentTs[EdgeIndex];
+			const FReal SegmentEdgeDistSq = EdgeDistSqs[EdgeIndex];
+			const FReal SegmentEdgeDistSign = EdgeDistSigns[EdgeIndex];
+			const FReal DotFace = EdgeDotFace[EdgeIndex];
+
 			// We only care about edges if at least one capsule segment point is outside the edge plane
 			// (internal points were handled already)
 			if (((EdgeD0s[EdgeIndex] > -DistanceTolerance) | (EdgeD1s[EdgeIndex] > -DistanceTolerance)) != 0)
 			{
-				// Find the nearest point on the capsule segment to the edge segment
-				FReal SegmentT, EdgeT;
-				FVec3 SegmentP, EdgeP;
-				Utilities::NearestPointsOnLineSegments(P0, P1, EdgeP0, EdgeP1, SegmentT, EdgeT, SegmentP, EdgeP);
-
-				// Calculate the separation vector, correct for sign so that we are pushing out of the face
-				FVec3 SegmentEdgeN = SegmentP - EdgeP;
-				FReal SegmentEdgeDistSign = FReal(1);
-				const FReal SegmentEdgeDistSq = SegmentEdgeN.SizeSquared();
-				const FReal DotFace = FVec3::DotProduct(SegmentEdgeN, FaceN);
-
-				// Check for cull distance (only if the segment point is not inside the face)
-				if (((DotFace > FReal(0)) & (SegmentEdgeDistSq > RejectDistanceSq)) != 0)
+				// Don't collide with inside face
+				if (DotFace < FReal(0))
 				{
 					continue;
 				}
 
-				// Make sure the normal point out of the face (we ned to correct the signed distance later too)
-				// Tolerance so we don't throw away "exact" edge collisions
-				if (DotFace < -NormalTolerance)
-				{
-					SegmentEdgeN = -SegmentEdgeN;
-					SegmentEdgeDistSign = FReal(-1);
-				}
-
-				// If the angle between the axis and the face is below a threshold, we will create a face contact
-				// rather than an edge contact where possible.
+				// We will create a face contact rather than an edge contact where possible.
+				// When the angle between the axis and the face is below a threshold.
 				const bool bInEdgeRange = InRangeOpen(EdgeT, FReal(0), FReal(1));
 				const bool bInSegmentRange = InRangeOpen(SegmentT, FReal(0), FReal(1));
 				const bool bCrossedEdgeSegment = ((bInEdgeRange & bInSegmentRange) != 0);
