Iris - Don't destroy dependent objects that haven't yet been created.

diff --git a/Engine/Plugins/Runtime/ReplicationSystemTestPlugin/Source/Private/Tests/ReplicationSystem/TestObjectSplitting.cpp b/Engine/Plugins/Runtime/ReplicationSystemTestPlugin/Source/Private/Tests/ReplicationSystem/TestObjectSplitting.cpp
--- a/Engine/Plugins/Runtime/ReplicationSystemTestPlugin/Source/Private/Tests/ReplicationSystem/TestObjectSplitting.cpp
+++ b/Engine/Plugins/Runtime/ReplicationSystemTestPlugin/Source/Private/Tests/ReplicationSystem/TestObjectSplitting.cpp
@@ -1105,6 +1105,48 @@ UE_NET_TEST_FIXTURE(FSplitObjectTestFixture, SplitManyHugeObjectsAfterCreation)
 	UE_NET_ASSERT_EQ(ClientObjectsWithHugeArraysCount, HugeObjectCount);
 }
 
+UE_NET_TEST_FIXTURE(FSplitObjectTestFixture, TestDependentObjectCannotBeDestroyedWhileWaitingForCreation)
+{
+	UReplicatedTestObjectBridge* ServerBridge = Server->GetReplicationBridge();
+	FReplicationSystemTestClient* Client = CreateClient();
+
+	UTestReplicatedIrisObject* ServerObject = CreateHugeObject(Server);
+	UTestReplicatedIrisObject* ServerDependentObject = Server->CreateObject(UObjectReplicationBridge::FCreateNetRefHandleParams{});
+
+	ServerBridge->AddDependentObject(ServerObject->NetRefHandle, ServerDependentObject->NetRefHandle);
+
+	// Introduce latency by not immediately delivering packets.
+	Server->PreSendUpdate();
+	Server->SendTo(Client, TEXT("Create HugeObject + Dependent"));
+	Server->PostSendUpdate();
+
+	// Filter out dependent object to cause it to end up being destroyed
+	Server->ReplicationSystem->AddToGroup(Server->GetReplicationSystem()->GetNotReplicatedNetObjectGroup(), ServerDependentObject->NetRefHandle);
+
+	Server->PreSendUpdate();
+	Server->SendTo(Client, TEXT("Try destroy Dependent"));
+	Server->PostSendUpdate();
+
+	// Make sure at least one of the packets required for object creation is lost.
+	Server->DeliverTo(Client, DoNotDeliverPacket);
+
+	// Deliver all pending packets
+	for (uint32 RetryIt = 0; RetryIt != HugeObjectMaxNetTickCountToArrive; ++RetryIt)
+	{
+		Server->DeliverTo(Client, DeliverPacket);
+	}
+
+	// Make sure we replicate the full state of all objects
+	for (uint32 RetryIt = 0; RetryIt != HugeObjectMaxNetTickCountToArrive; ++RetryIt)
+	{
+		Server->UpdateAndSend({ Client }, DeliverPacket);
+	}
+
+	// Make sure the dependent object was destroyed.
+	const UTestReplicatedIrisObject* ClientDependentObject = Cast<UTestReplicatedIrisObject>(Client->GetReplicationBridge()->GetReplicatedObject(ServerDependentObject->NetRefHandle));
+	UE_NET_ASSERT_EQ(ClientDependentObject, nullptr);
+}
+
 // Below test will fail as we only have a special path for reliable attachments for objects that stopped replicating, not for being filtered out.
 #if 0
 UE_NET_TEST_FIXTURE(FSplitObjectTestFixture, TestReliableAttachmentAddedAfterSplittingHugeObjectIsDeliveredBeforeObjectIsFilteredOut)
diff --git a/Engine/Source/Runtime/Experimental/Iris/Core/Private/Iris/ReplicationSystem/ReplicationWriter.cpp b/Engine/Source/Runtime/Experimental/Iris/Core/Private/Iris/ReplicationSystem/ReplicationWriter.cpp
--- a/Engine/Source/Runtime/Experimental/Iris/Core/Private/Iris/ReplicationSystem/ReplicationWriter.cpp
+++ b/Engine/Source/Runtime/Experimental/Iris/Core/Private/Iris/ReplicationSystem/ReplicationWriter.cpp
@@ -551,7 +551,7 @@ uint32 FReplicationWriter::GetFlushStatus(uint32 InternalIndex, const FReplicati
 		return FlushFlags;
 	}
 
-	if (!!(FlushFlagsToTest & EFlushFlags::FlushFlags_FlushState) && (Info.HasDirtyChangeMask || HasInFlightStateChanges(InternalIndex, Info) || IsObjectPartOfActiveHugeObject(InternalIndex, Info)))
+	if (!!(FlushFlagsToTest & EFlushFlags::FlushFlags_FlushState) && (Info.HasDirtyChangeMask || HasInFlightStateChanges(InternalIndex, Info) || IsObjectPartOfActiveHugeObject(InternalIndex)))
 	{
 		FlushFlags |= EFlushFlags::FlushFlags_FlushState;
 	}
@@ -1093,7 +1093,7 @@ void FReplicationWriter::HandleDeliveredRecord(const FReplicationRecord::FRecord
 		bool bStillPendingFlush = false;
 		if (RecordInfo.HasChangeMask && !!(Info.FlushFlags & EFlushFlags::FlushFlags_FlushState))
 		{
-			bStillPendingFlush |= (Info.HasDirtyChangeMask || HasInFlightStateChanges(ReplicationRecord.GetInfoForIndex(RecordInfo.NextIndex)) || IsObjectPartOfActiveHugeObject(InternalIndex, Info));
+			bStillPendingFlush |= (Info.HasDirtyChangeMask || HasInFlightStateChanges(ReplicationRecord.GetInfoForIndex(RecordInfo.NextIndex)) || IsObjectPartOfActiveHugeObject(InternalIndex));
 		}
 
 		if (RecordInfo.HasAttachments && !!(Info.FlushFlags & FlushFlags_FlushReliable))
@@ -1692,7 +1692,7 @@ uint32 FReplicationWriter::WriteObjectsPendingDestroy(FNetSerializationContext&
 		check(Info.GetState() == EReplicatedObjectState::PendingDestroy);
 
 		// We do not support destroying an object that is currently being sent as a huge object.
-		if (IsObjectPartOfActiveHugeObject(InternalIndex, Info))
+		if (IsObjectPartOfActiveHugeObject(InternalIndex))
 		{
 			UE_LOG(LogIris, Verbose, TEXT("Skipping writing destroy for object ( InternalIndex: %u ) which is part of active huge object."), InternalIndex);
 			bWroteAllDestroyedObjects = false;
@@ -3401,9 +3401,10 @@ bool FReplicationWriter::IsActiveHugeObject(uint32 InternalIndex) const
 	return HugeObjectSendQueue.IsObjectInQueue(InternalIndex, bIncludeSubObjects);
 }
 
-bool FReplicationWriter::IsObjectPartOfActiveHugeObject(uint32 InternalIndex, const FReplicationInfo& Info) const
+bool FReplicationWriter::IsObjectPartOfActiveHugeObject(uint32 InternalIndex) const
 {
-	return HugeObjectSendQueue.IsObjectInQueue(InternalIndex, Info.IsSubObject);
+	constexpr bool bFullSearch = true;
+	return HugeObjectSendQueue.IsObjectInQueue(InternalIndex, bFullSearch);
 }
 
 bool FReplicationWriter::CanQueueHugeObject() const
@@ -3580,7 +3581,7 @@ bool FReplicationWriter::FHugeObjectSendQueue::EnqueueHugeObject(const FHugeObje
 }
 
 // Returns true if the object is a huge object root object or part of any huge object's payload. The latter is an expensive operation.
-bool FReplicationWriter::FHugeObjectSendQueue::IsObjectInQueue(FInternalNetRefIndex ObjectIndex, bool bIncludeSubObjects) const
+bool FReplicationWriter::FHugeObjectSendQueue::IsObjectInQueue(FInternalNetRefIndex ObjectIndex, bool bFullSearch) const
 {
 	if (IsEmpty())
 	{
@@ -3592,7 +3593,7 @@ bool FReplicationWriter::FHugeObjectSendQueue::IsObjectInQueue(FInternalNetRefIn
 		return true;
 	}
 
-	if (!bIncludeSubObjects)
+	if (!bFullSearch)
 	{
 		return false;
 	}
diff --git a/Engine/Source/Runtime/Experimental/Iris/Core/Private/Iris/ReplicationSystem/ReplicationWriter.h b/Engine/Source/Runtime/Experimental/Iris/Core/Private/Iris/ReplicationSystem/ReplicationWriter.h
--- a/Engine/Source/Runtime/Experimental/Iris/Core/Private/Iris/ReplicationSystem/ReplicationWriter.h
+++ b/Engine/Source/Runtime/Experimental/Iris/Core/Private/Iris/ReplicationSystem/ReplicationWriter.h
@@ -330,7 +330,7 @@ private:
 		bool EnqueueHugeObject(const FHugeObjectContext& Context);
 
 		// Returns true if the object is a huge object root object or part of any huge object's payload. The latter is an expensive operation.
-		bool IsObjectInQueue(FInternalNetRefIndex ObjectIndex, bool bIncludeSubObjects) const;
+		bool IsObjectInQueue(FInternalNetRefIndex ObjectIndex, bool bFullSearch) const;
 
 		// Best effort implementation of getting a valid index for trace.
 		FInternalNetRefIndex GetRootObjectInternalIndexForTrace() const;
@@ -508,7 +508,7 @@ private:
 	inline bool IsInitialState(const EReplicatedObjectState State) const { return State == EReplicatedObjectState::PendingCreate || State == EReplicatedObjectState::WaitOnCreateConfirmation; }
 
 	bool IsActiveHugeObject(uint32 InternalIndex) const;
-	bool IsObjectPartOfActiveHugeObject(uint32 InternalIndex, const FReplicationInfo& Info) const;
+	bool IsObjectPartOfActiveHugeObject(uint32 InternalIndex) const;
 
 	bool CanQueueHugeObject() const;
 
