Log which IAS toc loading mode we are using on startup

#rnx

diff --git a/Engine/Source/Runtime/Experimental/IoStoreOnDemand/Private/OnDemandIoDispatcherBackend.cpp b/Engine/Source/Runtime/Experimental/IoStoreOnDemand/Private/OnDemandIoDispatcherBackend.cpp
--- a/Engine/Source/Runtime/Experimental/IoStoreOnDemand/Private/OnDemandIoDispatcherBackend.cpp
+++ b/Engine/Source/Runtime/Experimental/IoStoreOnDemand/Private/OnDemandIoDispatcherBackend.cpp
@@ -2140,6 +2140,8 @@ void FOnDemandIoBackend::Mount(const FOnDemandEndpoint& Endpoint)
 
 	if (GetTocMode() == ETocMode::LoadTocFromDisk)
 	{
+		UE_LOG(LogIas, Display, TEXT("Loading .iochunktoc from disk"));
+
 		FString TocHash = FPaths::GetBaseFilename(Endpoint.TocPath);
 		LoadingOnDemandTocTask = UE::Tasks::Launch(UE_SOURCE_LOCATION, [TocHash = MoveTemp(TocHash)]() -> TIoStatusOr<FOnDemandToc>
 			{
@@ -2148,6 +2150,8 @@ void FOnDemandIoBackend::Mount(const FOnDemandEndpoint& Endpoint)
 	}
 	else if (GetTocMode() == ETocMode::LoadTocFromMountedPaks)
 	{
+		UE_LOG(LogIas, Display, TEXT("Loading .uondemandtoc from mounted pakfiles"));
+
 		// First make sure that we receive notifications of any new pakfiles being mounted so that we can react to them
 		PakDelegateHandle = FCoreInternalDelegates::GetOnPakMountOperation().AddLambda(
 			[this](EMountOperation Operation, const TCHAR* ContainerPath, int32 Order) -> void
