Iris - Use actual bitstream reader position for net traces, regardless of overflow. This will normally give more trace information when bitstream errors occur.

#rnx

diff --git a/Engine/Source/Runtime/Experimental/Iris/Core/Public/Iris/Serialization/NetBitStreamReader.h b/Engine/Source/Runtime/Experimental/Iris/Core/Public/Iris/Serialization/NetBitStreamReader.h
--- a/Engine/Source/Runtime/Experimental/Iris/Core/Public/Iris/Serialization/NetBitStreamReader.h
+++ b/Engine/Source/Runtime/Experimental/Iris/Core/Public/Iris/Serialization/NetBitStreamReader.h
@@ -70,4 +70,5 @@ private:
 
 }
 
-inline uint32 GetBitStreamPositionForNetTrace(const UE::Net::FNetBitStreamReader& Stream) { return (uint32(Stream.IsOverflown()) - 1U) & Stream.GetPosBits(); }
+// Always report the actual bitstream position, even on overflow. This normally allows for better comparisons between sending and receiving side when bitstream errors occur.
+inline uint32 GetBitStreamPositionForNetTrace(const UE::Net::FNetBitStreamReader& Stream) { return Stream.GetPosBits(); }
