IasTool skeleton with Test command

#rnx

diff --git a/Engine/Source/Programs/IoStoreOnDemand/IasTool.Build.cs b/Engine/Source/Programs/IoStoreOnDemand/IasTool.Build.cs
new file mode 100644
--- /dev/null
+++ b/Engine/Source/Programs/IoStoreOnDemand/IasTool.Build.cs
@@ -0,0 +1,16 @@
+// Copyright Epic Games, Inc. All Rights Reserved.
+
+using UnrealBuildTool;
+
+public class IasTool : ModuleRules
+{
+	public IasTool(ReadOnlyTargetRules Target) : base(Target)
+	{
+		PublicIncludePathModuleNames.Add("Launch");
+
+		PrivateDependencyModuleNames.Add("Core");
+		PrivateDependencyModuleNames.Add("IoStoreOnDemand");
+		PrivateDependencyModuleNames.Add("Projects");
+	}
+}
+
diff --git a/Engine/Source/Programs/IoStoreOnDemand/IasTool.Target.cs b/Engine/Source/Programs/IoStoreOnDemand/IasTool.Target.cs
new file mode 100644
--- /dev/null
+++ b/Engine/Source/Programs/IoStoreOnDemand/IasTool.Target.cs
@@ -0,0 +1,34 @@
+// Copyright Epic Games, Inc. All Rights Reserved.
+
+using UnrealBuildTool;
+
+[SupportedPlatforms("Win64")]
+public class IasToolTarget : TargetRules
+{
+	public IasToolTarget(TargetInfo Target) : base(Target)
+	{
+		Type = TargetType.Program;
+		LinkType = TargetLinkType.Monolithic;
+		LaunchModuleName = "IasTool";
+		IncludeOrderVersion = EngineIncludeOrderVersion.Latest;
+
+		GlobalDefinitions.Add("UE_WITH_IAS_TOOL=1");
+
+		bBuildDeveloperTools = false;
+		bBuildWithEditorOnlyData = false;
+		bCompileAgainstApplicationCore = false;
+		bCompileAgainstCoreUObject = false;
+		bCompileAgainstEngine = false;
+		bCompileICU = false;
+		bCompileICU = false;
+		bCompileNavmeshClusterLinks = false;
+		bCompileNavmeshSegmentLinks = false;
+		bCompileRecast = false;
+		bCompileWithPluginSupport = true;
+		bForceEnableExceptions = true;
+		bIsBuildingConsoleApplication = true;
+		bUsesSlate = false;
+		bWithLiveCoding = false;
+		bWithServerCode = false;
+	}
+}
diff --git a/Engine/Source/Programs/IoStoreOnDemand/Private/IasTool.cpp b/Engine/Source/Programs/IoStoreOnDemand/Private/IasTool.cpp
new file mode 100644
--- /dev/null
+++ b/Engine/Source/Programs/IoStoreOnDemand/Private/IasTool.cpp
@@ -0,0 +1,40 @@
+// Copyright Epic Games, Inc. All Rights Reserved.
+
+#include <CoreGlobals.h>
+#include <HAL/LowLevelMemTracker.h>
+#include <HAL/Platform.h>
+#include <LaunchEngineLoop.h>
+#include <Modules/ModuleManager.h>
+#include <Misc/ScopeExit.h>
+
+#include <RequiredProgramMainCPPInclude.h>
+
+IMPLEMENT_APPLICATION(IasTool, "IasTool");
+
+////////////////////////////////////////////////////////////////////////////////
+namespace UE::IO::IAS::Tool {
+
+int32 Main(int32, TCHAR*[]);
+
+} // namespace UE::IO::IAS
+
+////////////////////////////////////////////////////////////////////////////////
+INT32_MAIN_INT32_ARGC_TCHAR_ARGV()
+{
+	FTaskTagScope Scope(ETaskTag::EGameThread);
+
+	if (int Ret = GEngineLoop.PreInit(ArgC, ArgV, TEXT("-stdout")); Ret != 0)
+	{
+		return Ret;
+	}
+
+	ON_SCOPE_EXIT {
+		LLM(FLowLevelMemTracker::Get().UpdateStatsPerFrame());
+		RequestEngineExit(TEXT("Exiting"));
+		FEngineLoop::AppPreExit();
+		FModuleManager::Get().UnloadModulesAtShutdown();
+		FEngineLoop::AppExit();
+	};
+
+	return UE::IO::IAS::Tool::Main(ArgC, ArgV);
+}
diff --git a/Engine/Source/Runtime/Experimental/IoStoreOnDemand/Private/Tool/Command.cpp b/Engine/Source/Runtime/Experimental/IoStoreOnDemand/Private/Tool/Command.cpp
new file mode 100644
--- /dev/null
+++ b/Engine/Source/Runtime/Experimental/IoStoreOnDemand/Private/Tool/Command.cpp
@@ -0,0 +1,447 @@
+// Copyright Epic Games, Inc. All Rights Reserved.
+
+#if defined(UE_WITH_IAS_TOOL)
+
+#include "Command.h"
+
+#include <Containers/UnrealString.h>
+#include <CoreGlobals.h>
+#include <Misc/OutputDeviceRedirector.h>
+
+namespace UE::IO::IAS::Tool {
+
+////////////////////////////////////////////////////////////////////////////////
+namespace Error
+{
+
+struct FBase
+{
+					FBase(const TCHAR* Ever, FStringView Not) : What(Ever), Why(Not) {}
+	const TCHAR*	What;
+	FString			Why;
+};
+
+#define MAKE_IASTOOL_ERROR(Name) \
+	struct Name	: public FBase { \
+		Name(FStringView Not) : FBase(TEXT(#Name), Not) {} \
+	};
+	MAKE_IASTOOL_ERROR(FMissingValue);
+	MAKE_IASTOOL_ERROR(FCommandNotFound);
+	MAKE_IASTOOL_ERROR(FBoolPositional);
+	MAKE_IASTOOL_ERROR(FUnexpectedValue);
+	MAKE_IASTOOL_ERROR(FUnknownKey);
+	MAKE_IASTOOL_ERROR(FCommandAbort);
+#undef MAKE_IASTOOL_ERROR
+
+}
+
+////////////////////////////////////////////////////////////////////////////////
+template <typename ErrorType, typename... ArgTypes>
+requires std::is_base_of<Error::FBase, ErrorType>::value
+[[noreturn]] void FatalError(ArgTypes&&... Args)
+{
+	throw ErrorType(Forward<ArgTypes>(Args)...);
+}
+
+
+
+////////////////////////////////////////////////////////////////////////////////
+static void WriteLine(const TCHAR* Format, ...)
+{
+	va_list VaList;
+	va_start(VaList, Format);
+	ON_SCOPE_EXIT { va_end(VaList); };
+
+	TCHAR Buffer[256];
+	FCString::GetVarArgs(Buffer, UE_ARRAY_COUNT(Buffer), Format, VaList);
+
+	static FName Category = FName("IasTool");
+	GLog->Serialize(Buffer, ELogVerbosity::Type::Display, Category);
+}
+
+////////////////////////////////////////////////////////////////////////////////
+static bool IsAnAstleySingleFrom91(const TCHAR* Cassette)
+{
+	if (Cassette[0] != '-' && Cassette[0] != '/')
+	{
+		return false;
+	}
+
+	FStringView Cd = Cassette + 1;
+	for (const TCHAR* TwelveInch : { TEXT("help"), TEXT("h"), TEXT("?"), TEXT("-help") }) // '/-help' :)
+	{
+		if (Cd.Equals(TwelveInch, ESearchCase::IgnoreCase))
+		{
+			return true;
+		}
+	}
+
+	return false;
+}
+
+
+
+////////////////////////////////////////////////////////////////////////////////
+FContext::FContext(const FArguments& InArguments, int32 ArgC, const TCHAR* const* ArgV)
+: Arguments(InArguments)
+{
+	auto NextArg = [&ArgC, &ArgV, i=int32(1), Pending=FStringView()] () mutable -> FStringView
+	{
+		ArgC--;
+		if (auto Ret = Pending; !Ret.IsEmpty())
+		{
+			Pending.Reset();
+			return Ret.Mid(Ret[0] == '=');
+		}
+		const TCHAR* Arg = ArgV[i++];
+		if (Arg[0] != '-')
+		{
+			return FStringView(Arg);
+		}
+		const TCHAR* End = Arg;
+		for (uint32 c = *End; c == '-' || uint32('z' - (c | 0x20)) <= 26; c = *++End);
+		Pending = FStringView(End);
+		ArgC += (Pending.IsEmpty() == false);
+		return FStringView(Arg, int32(ptrdiff_t(End - Arg)));
+	};
+
+	int32 PendingOptional = -1;
+	auto HandleOptional = [this, &PendingOptional] (FStringView Arg)
+	{
+		if (int32 i = PendingOptional; i >= 0)
+		{
+			if (Arg[0] == '-')
+			{
+				FArgument Argument = Arguments[i];
+				FatalError<Error::FMissingValue>(Argument.Name);
+			}
+			Arguments[i].Parser(Values[i].Buffer, Arg);
+			SetMask |= 1ull << uint64(i);
+			PendingOptional = -1;
+			return;
+		}
+
+		for (int32 i = 0, n = Arguments.Num(); i < n; ++i)
+		{
+			FArgument& Argument = Arguments[i];
+			if (Argument.Name != Arg)
+			{
+				continue;
+			}
+
+			if (Argument.Parser != nullptr)
+			{
+				PendingOptional = i;
+				return;
+			}
+
+			Values[i].Buffer[0] = 1;
+			SetMask |= 1ull << uint64(i);
+			return;
+		}
+
+		FatalError<Error::FUnknownKey>(Arg);
+	};
+
+	int32 PositionalIndex = -1;
+	auto HandlePositional = [this, &PositionalIndex] (FStringView Arg)
+	{
+		for (;;)
+		{
+			++PositionalIndex;
+			if (PositionalIndex >= Arguments.Num())
+			{
+				FatalError<Error::FUnexpectedValue>(Arg);
+			}
+
+			if (Arguments[PositionalIndex].Name[0] != '-')
+			{
+				break;
+			}
+		}
+
+		FArgument& Argument = Arguments[PositionalIndex];
+		if (Argument.Parser == nullptr)
+		{
+			FatalError<Error::FBoolPositional>(Arg);
+		}
+
+		Argument.Parser(Values[PositionalIndex].Buffer, Arg);
+		SetMask |= 1ull << uint64(PositionalIndex);
+	};
+
+	Values.SetNum(InArguments.Num());
+
+	while (ArgC > 1)
+	{
+		FStringView Arg = NextArg();
+		if (Arg == TEXT("--"))
+		{
+			break;
+		}
+
+		(Arg[0] == '-' || PendingOptional >= 0)
+			? HandleOptional(Arg)
+			: HandlePositional(Arg);
+	}
+
+	if (PendingOptional >= 0)
+	{
+		FArgument Argument = Arguments[PendingOptional];
+		FatalError<Error::FMissingValue>(Argument.Name);
+	}
+}
+
+////////////////////////////////////////////////////////////////////////////////
+uint32 FContext::FindArgument(FStringView Name) const
+{
+	for (const FArgument& Argument : Arguments)
+	{
+		if (Argument.Name == Name)
+		{
+			return uint32(ptrdiff_t(&Argument - Arguments.GetData()));
+		}
+	}
+
+	FatalError<Error::FUnknownKey>(Name);
+}
+
+////////////////////////////////////////////////////////////////////////////////
+[[noreturn]] void FContext::OnMissingValue(FStringView Name) const
+{
+	FatalError<Error::FMissingValue>(Name);
+}
+
+////////////////////////////////////////////////////////////////////////////////
+[[noreturn]] void FContext::Abort(const TCHAR* Reason) const
+{
+	FatalError<Error::FCommandAbort>(Reason);
+}
+
+
+
+////////////////////////////////////////////////////////////////////////////////
+FCommand::FCommand(
+	EntryFunc* InEntry,
+	FStringView InName,
+	FStringView InDesc,
+	FArguments&& InArguments)
+{
+	Initialize(InEntry, InName, InDesc, MoveTemp(InArguments));
+	FCommand::Head() = this;
+}
+
+////////////////////////////////////////////////////////////////////////////////
+void FCommand::Initialize(
+	EntryFunc* InEntry,
+	FStringView InName,
+	FStringView InDesc,
+	FArguments&& InArguments)
+{
+	Entry = InEntry;
+	Name = InName;
+	Desc = InDesc;
+	Arguments = InArguments;
+}
+
+////////////////////////////////////////////////////////////////////////////////
+FCommand*& FCommand::Head()
+{
+	static FCommand* Instance;
+	return Instance;
+}
+
+////////////////////////////////////////////////////////////////////////////////
+FContext FCommand::GetContext(int32 ArgC, const TCHAR* const* ArgV) const
+{
+	return FContext(Arguments, ArgC, ArgV);
+}
+
+////////////////////////////////////////////////////////////////////////////////
+void FCommand::Usage() const
+{
+	if (!Desc.IsEmpty())
+	{
+		WriteLine(Desc.GetData());
+		WriteLine(TEXT(""));
+	}
+
+	bool bHasOpts = false;
+
+	TStringBuilder<128> UsageLine;
+	UsageLine << TEXT("  ");
+	UsageLine << Name;
+	for (const FArgument& Argument : Arguments)
+	{
+		if (Argument.Name.StartsWith('-'))
+		{
+			bHasOpts = true;
+			continue;
+		}
+
+		UsageLine << TEXT(" <");
+		UsageLine << Argument.Name;
+		UsageLine << TEXT(">");
+	}
+
+	if (bHasOpts)
+	{
+		UsageLine << TEXT(" [Options]");
+	}
+
+	WriteLine(TEXT("Usage:"));
+	WriteLine(TEXT("  %s"), UsageLine.ToString());
+
+	if (!bHasOpts)
+	{
+		return;
+	}
+
+	WriteLine(TEXT(""));
+	WriteLine(TEXT("Options:"));
+	for (const FArgument& Argument : Arguments)
+	{
+		if (!Argument.Name.StartsWith('-'))
+		{
+			continue;
+		}
+
+		const TCHAR* Suffix = (Argument.Parser != nullptr) ? TEXT("=<value>") : TEXT("");
+		WriteLine(TEXT( "  %s%s"), Argument.Name.GetData(), Suffix);
+		if (!Argument.Desc.IsEmpty())
+		{
+			WriteLine(TEXT("    %s"), Argument.Desc.GetData());
+			WriteLine(TEXT(""));
+		}
+	}
+}
+
+////////////////////////////////////////////////////////////////////////////////
+int32 FCommand::Call(int32 ArgC, const TCHAR* const* ArgV)
+{
+	for (int32 i = 1; i < ArgC; ++i)
+	{
+		if (IsAnAstleySingleFrom91(ArgV[i]))
+		{
+			Usage();
+			return 126;
+		}
+	}
+
+	FContext Context = GetContext(ArgC, ArgV);
+	return Entry(Context);
+}
+
+////////////////////////////////////////////////////////////////////////////////
+int32 FCommand::Main(int32 ArgC, const TCHAR* const* ArgV)
+{
+	try
+	{
+		return MainInner(ArgC, ArgV);
+	}
+ 	catch (const Error::FBase& Err)
+	{
+		WriteLine(TEXT("ERROR: %s : %s"), Err.What, *(Err.Why));
+	}
+
+	return 125;
+}
+
+////////////////////////////////////////////////////////////////////////////////
+int32 FCommand::MainInner(int32 ArgC, const TCHAR* const* ArgV)
+{
+	if ((ArgC < 2) || IsAnAstleySingleFrom91(ArgV[1]))
+	{
+		WriteLine(TEXT("Available commands:"));
+		for (FCommand* Cmd = Head(); Cmd != nullptr; Cmd = Cmd->Next)
+		{
+			WriteLine(TEXT("  %-12s %s"), Cmd->Name.GetData(), Cmd->Desc.GetData());
+		}
+		return 127;
+	}
+
+	FStringView Action = ArgV[1];
+	for (FCommand* Cmd = Head(); Cmd != nullptr; Cmd = Cmd->Next)
+	{
+		if (Cmd->Name == Action)
+		{
+			return Cmd->Call(ArgC - 1, ArgV + 1);
+		}
+	}
+
+	FatalError<Error::FCommandNotFound>(Action);
+}
+
+////////////////////////////////////////////////////////////////////////////////
+void CommandTest()
+{
+	auto Entry = [] (const FContext&) { return 0; };
+	FCommand Command;
+	Command.Initialize(Entry, TEXT("name"), TEXT("desc"), {
+		TArgument<FStringView>(TEXT("pos0"), TEXT("description")),
+		TArgument<FStringView>(TEXT("-string"), TEXT("description")),
+		TArgument<FStringView>(TEXT("pos1"), TEXT("description")),
+		TArgument<FStringView>(TEXT("pos2"), TEXT("description")),
+		TArgument<bool>(TEXT("-bool"), TEXT("description")),
+	});
+
+	using TestArgs = TArray<const TCHAR*>;
+
+	auto TestContext = [&] (TestArgs Args)
+	{
+		Args.Insert(TEXT("$"), 0);
+		return Command.GetContext(Args.Num(), Args.GetData());
+	};
+
+	// optional
+	struct {
+		TestArgs	Args;
+		FStringView	ExpectString;
+		bool		ExpectBool;
+	} TestCases0[] = {
+		{ {TEXT("-bool")}, TEXT(""), true },
+		{ {TEXT("-string=abc")}, TEXT("abc"), false },
+		{ {TEXT("-bool"), TEXT("-string=a c")}, TEXT("a c"), true },
+		{ {TEXT("-string"), TEXT("abc")}, TEXT("abc"), false },
+		{ {TEXT("-string"), TEXT("a=c"), TEXT("-bool")}, TEXT("a=c"), true },
+		{ {TEXT("-string=a=c")}, TEXT("a=c"), false },
+	};
+	for (auto& TestCase : TestCases0)
+	{
+		FContext Context = TestContext(MoveTemp(TestCase.Args));
+		check(Context.Get<bool>(TEXT("-bool")) == TestCase.ExpectBool);
+		check(Context.Get<FStringView>(TEXT("-string")) == TestCase.ExpectString);
+	}
+
+	for (TestArgs TestCase : {
+			TestArgs{ TEXT("-string") },
+			TestArgs{ TEXT("-string"), TEXT("-a9e") },
+		})
+	{
+		try { TestContext(MoveTemp(TestCase)); }
+		catch (const Error::FMissingValue&) { continue; }
+		check(false);
+	}
+
+	// positional
+	for (TestArgs TestCase : {
+			TestArgs{ TEXT("pos0") },
+			TestArgs{ TEXT("pos0"), TEXT("pos1") },
+			TestArgs{ TEXT("pos0"), TEXT("pos1"), TEXT("pos2") },
+		})
+	{
+		try { TestContext(MoveTemp(TestCase)); }
+		catch (...) { check(false); }
+	}
+
+	if (bool bTestPass = false; true)
+	{
+		try { TestContext(TestArgs{TEXT("a"), TEXT("b"), TEXT("b"), TEXT("POP!")}); }
+		catch (const Error::FUnexpectedValue&) { bTestPass = true; }
+		check(bTestPass);
+	}
+}
+
+} // namespace UE::IO::IAS::Tool
+
+#endif // UE_WITH_IAS_TOOL
diff --git a/Engine/Source/Runtime/Experimental/IoStoreOnDemand/Private/Tool/Command.h b/Engine/Source/Runtime/Experimental/IoStoreOnDemand/Private/Tool/Command.h
new file mode 100644
--- /dev/null
+++ b/Engine/Source/Runtime/Experimental/IoStoreOnDemand/Private/Tool/Command.h
@@ -0,0 +1,151 @@
+// Copyright Epic Games, Inc. All Rights Reserved.
+
+#pragma once
+
+#if defined(UE_WITH_IAS_TOOL)
+
+#include <Containers/Array.h>
+#include <Containers/StringConv.h>
+#include <Containers/StringView.h>
+#include <String/LexFromString.h>
+
+namespace UE::IO::IAS::Tool {
+
+////////////////////////////////////////////////////////////////////////////////
+using ParseFunc = void (void*, FStringView);
+
+////////////////////////////////////////////////////////////////////////////////
+struct FArgument
+{
+	static const uint32 ValueAlign = 16;
+	static const uint32 ValueSize = 16;
+	FStringView			Name;
+	FStringView			Desc;
+	ParseFunc*			Parser;
+};
+
+using FArguments = TArray<FArgument>;
+
+////////////////////////////////////////////////////////////////////////////////
+template <typename Type>
+struct TArgumentParse
+{
+	static ParseFunc* GetParser()
+	{
+		return [] (void* Out, FStringView Input) {
+			auto* Inner = (Type*)Out;
+			LexFromString(*Inner, Input);
+		};
+	}
+};
+
+////////////////////////////////////////////////////////////////////////////////
+template <>
+struct TArgumentParse<bool>
+{
+	static ParseFunc* GetParser() { return nullptr; }
+};
+
+////////////////////////////////////////////////////////////////////////////////
+template <>
+struct TArgumentParse<FStringView>
+{
+	static ParseFunc* GetParser()
+	{
+		return [] (void* Out, FStringView Input) {
+			auto* Inner = (FStringView*)Out;
+			new (Inner) FStringView(Input);
+		};
+	}
+};
+
+////////////////////////////////////////////////////////////////////////////////
+template <typename Type>
+FArgument TArgument(FStringView Name, FStringView Desc)
+{
+	static_assert(alignof(Type) <= FArgument::ValueAlign);
+	static_assert(sizeof(Type) <= FArgument::ValueSize);
+	return { Name, Desc, TArgumentParse<Type>::GetParser() };
+}
+
+
+
+////////////////////////////////////////////////////////////////////////////////
+class FContext
+{
+public:
+	template <typename Type> Type	Get(FStringView Name) const;
+	template <typename Type> Type	Get(FStringView Name, Type Default) const;
+	[[noreturn]] void				Abort(const TCHAR* Reason) const;
+
+private:
+	struct alignas(FArgument::ValueAlign) FValue
+	{
+		uint8			Buffer[FArgument::ValueSize] = {};
+	};
+
+	friend				class FCommand;
+						FContext(const FArguments& InArguments, int32 ArgC, const TCHAR* const* ArgV);
+	uint32				FindArgument(FStringView Name) const;
+	[[noreturn]] void	OnMissingValue(FStringView Name) const;
+	bool				IsSet(uint32 Index) const { return !!((1ull << Index) & SetMask); }
+	FArguments			Arguments;
+	TArray<FValue>		Values;
+	uint64				SetMask = 0;
+};
+
+////////////////////////////////////////////////////////////////////////////////
+template <typename Type>
+Type FContext::Get(FStringView Name) const
+{
+	if (Name.GetData()[0] == '-')
+	{
+		return Get(Name, Type{});
+	}
+
+	uint32 Index = FindArgument(Name);
+	if (IsSet(Index))
+	{
+		return *(Type*)(Values[Index].Buffer);
+	}
+
+	OnMissingValue(Name);
+}
+
+////////////////////////////////////////////////////////////////////////////////
+template <typename Type>
+Type FContext::Get(FStringView Name, Type Default) const
+{
+	uint32 Index = FindArgument(Name);
+	return IsSet(Index) ? *(Type*)(Values[Index].Buffer) : Default;
+}
+
+
+
+////////////////////////////////////////////////////////////////////////////////
+class FCommand
+{
+public:
+	using				EntryFunc = int32 (const FContext&);
+						FCommand(EntryFunc* InEntry, FStringView InName, FStringView InDesc, FArguments&& InArguments);
+	static int32		Main(int32 ArgC, const TCHAR* const* ArgV);
+
+private:
+	friend void			CommandTest();
+						FCommand() = default;
+	void				Initialize(EntryFunc* InEntry, FStringView InName, FStringView InDesc, FArguments&& InArguments);
+	FContext			GetContext(int32 ArgC, const TCHAR* const* ArgV) const;
+	static FCommand*&	Head();
+	static int32		MainInner(int32 ArgC, const TCHAR* const* ArgV);
+	void				Usage() const;
+	int32				Call(int32 ArgC, const TCHAR* const* ArgV);
+	FCommand*			Next = FCommand::Head();
+	FStringView			Name;
+	FStringView			Desc;
+	FArguments			Arguments;
+	EntryFunc*			Entry;
+};
+
+} // namespace UE::IO::IAS::Tool
+
+#endif // UE_WITH_IAS_TOOL
diff --git a/Engine/Source/Runtime/Experimental/IoStoreOnDemand/Private/Tool/Main.cpp b/Engine/Source/Runtime/Experimental/IoStoreOnDemand/Private/Tool/Main.cpp
new file mode 100644
--- /dev/null
+++ b/Engine/Source/Runtime/Experimental/IoStoreOnDemand/Private/Tool/Main.cpp
@@ -0,0 +1,17 @@
+// Copyright Epic Games, Inc. All Rights Reserved.
+
+#if defined(UE_WITH_IAS_TOOL)
+
+#include "Command.h"
+
+namespace UE::IO::IAS::Tool {
+
+////////////////////////////////////////////////////////////////////////////////
+int32 Main(int32 ArgC, TCHAR* ArgV[])
+{
+	return FCommand::Main(ArgC, ArgV);
+}
+
+} // namespace UE::IO::IAS::Tool
+
+#endif // UE_WITH_IAS_TOOL
diff --git a/Engine/Source/Runtime/Experimental/IoStoreOnDemand/Private/Tool/Test.cpp b/Engine/Source/Runtime/Experimental/IoStoreOnDemand/Private/Tool/Test.cpp
new file mode 100644
--- /dev/null
+++ b/Engine/Source/Runtime/Experimental/IoStoreOnDemand/Private/Tool/Test.cpp
@@ -0,0 +1,64 @@
+// Copyright Epic Games, Inc. All Rights Reserved.
+
+#if defined(UE_WITH_IAS_TOOL)
+
+#include "Command.h"
+
+#include <Async/TaskGraphInterfaces.h>
+#include <Containers/StringConv.h>
+#include <Containers/UnrealString.h>
+#include <Misc/CommandLine.h>
+
+namespace UE::IO::IAS {
+
+////////////////////////////////////////////////////////////////////////////////
+namespace IasJournaledFileCacheTest { void Tests(const TCHAR*);				}
+namespace HTTP						{ void IasHttpTest(const ANSICHAR*);	}
+
+namespace Tool {
+
+////////////////////////////////////////////////////////////////////////////////
+void CommandTest();
+
+////////////////////////////////////////////////////////////////////////////////
+static void HttpTests(const FContext& Context)
+{
+	auto TestHost = Context.Get<FStringView>(TEXT("-Host"), TEXT("localhost"));
+
+	auto TestHostAnsi = StringCast<ANSICHAR>(TestHost.GetData());
+	HTTP::IasHttpTest(TestHostAnsi.Get());
+}
+
+////////////////////////////////////////////////////////////////////////////////
+static void CacheTests(const FContext& Context)
+{
+	FStringView CacheDirStr = Context.Get<FStringView>(TEXT("-Dir"));
+
+	const TCHAR* CacheDir = CacheDirStr.IsEmpty() ? nullptr : CacheDirStr.GetData();
+	IasJournaledFileCacheTest::Tests(CacheDir);
+}
+
+////////////////////////////////////////////////////////////////////////////////
+static int32 TestCommandEntry(const FContext& Context)
+{
+	CommandTest();
+	CacheTests(Context);
+	HttpTests(Context);
+	return 0;
+}
+
+////////////////////////////////////////////////////////////////////////////////
+static FCommand TestCommand(
+	TestCommandEntry,
+	TEXT("test"),
+	TEXT("run IAS tests"),
+	{
+		TArgument<FStringView>(TEXT("-Host"), TEXT("Host of the HTTP test server")),
+		TArgument<FStringView>(TEXT("-Dir"), TEXT("Primary directory to use for cache tests")),
+	}
+);
+
+} // namespace Tool
+} // namespace UE::IO::IAS
+
+#endif // UE_WITH_IAS_TOOL
