Iris
- Improved error message if encountering too large object including NetRefHandle + name.
- Proper disconnect from server when encountering critical errors.

diff --git a/Engine/Source/Runtime/Engine/Private/Net/Experimental/Iris/DataStreamChannel.cpp b/Engine/Source/Runtime/Engine/Private/Net/Experimental/Iris/DataStreamChannel.cpp
--- a/Engine/Source/Runtime/Engine/Private/Net/Experimental/Iris/DataStreamChannel.cpp
+++ b/Engine/Source/Runtime/Engine/Private/Net/Experimental/Iris/DataStreamChannel.cpp
@@ -313,6 +313,15 @@ void UDataStreamChannel::WriteData(UE::Net::EDataStreamWriteMode WriteMode)
 			IRIS_PROFILER_SCOPE(UDataStreamChannel_NoDataSent);
 			// Do not report the bunch
 			UE_NET_TRACE_DISCARD_BUNCH(Collector);
+
+			if (SerializationContext.HasError())
+			{
+				FString ErrorMsg = NSLOCTEXT("NetworkErrors", "DataStreamChannelWriteData", "DataStreamChannel failed to write data.").ToString();
+				Connection->SendCloseReason(ENetCloseResult::HostClosedConnection);
+				FNetControlMessage<NMT_Failure>::Send(Connection, ErrorMsg);
+				Connection->FlushNet(true);
+				Connection->Close(ENetCloseResult::HostClosedConnection);
+			}
 			
 			return UDataStream::EWriteResult::NoData;
 		}
diff --git a/Engine/Source/Runtime/Experimental/Iris/Core/Private/Iris/DataStream/DataStreamManager.cpp b/Engine/Source/Runtime/Experimental/Iris/Core/Private/Iris/DataStream/DataStreamManager.cpp
--- a/Engine/Source/Runtime/Experimental/Iris/Core/Private/Iris/DataStream/DataStreamManager.cpp
+++ b/Engine/Source/Runtime/Experimental/Iris/Core/Private/Iris/DataStream/DataStreamManager.cpp
@@ -302,7 +302,16 @@ UDataStreamManager::EWriteResult UDataStreamManager::FImpl::WriteData(UE::Net::F
 		{
 			checkf(SubRecord == nullptr, TEXT("DataStream '%s' provided a record despite errors or returning NoData."), ToCStr(Stream->GetFName().GetPlainNameString()));
 			ManagerStream.DiscardSubstream(SubBitStream);
-			continue;
+
+			if (SubContext.HasError())
+			{
+				Context.SetError(SubContext.GetError(), false);
+				break;
+			}
+			else
+			{
+				continue;
+			}
 		}
 
 		DataStreamMask |= CurrentStreamMask;
diff --git a/Engine/Source/Runtime/Experimental/Iris/Core/Private/Iris/ReplicationSystem/ReplicationWriter.cpp b/Engine/Source/Runtime/Experimental/Iris/Core/Private/Iris/ReplicationSystem/ReplicationWriter.cpp
--- a/Engine/Source/Runtime/Experimental/Iris/Core/Private/Iris/ReplicationSystem/ReplicationWriter.cpp
+++ b/Engine/Source/Runtime/Experimental/Iris/Core/Private/Iris/ReplicationSystem/ReplicationWriter.cpp
@@ -2444,7 +2444,10 @@ int FReplicationWriter::PrepareAndSendHugeObjectPayload(FNetSerializationContext
 			// Cleanup data from batch
 			HandleObjectBatchFailure(WriteHugeObjectStatus, BatchInfo, WriteBitStreamInfo);
 
-			UE_LOG(LogIris, Error, TEXT("Unable to fit object ( InternalIndex: %u ) in maximum combined payload of %u bytes. Connection %u will be disconnected."), InternalIndex, MaxHugeObjectPayLoadBytes, Context.GetLocalConnectionId());
+			const FNetRefHandleManager::FReplicatedObjectData& ObjectData = NetRefHandleManager->GetReplicatedObjectDataNoCheck(InternalIndex);
+			UE_LOG(LogIris, Error, TEXT("Unable to fit object %s ( InternalIndex: %u ) %s in maximum combined payload of %u bytes. Connection %u will be disconnected."), *ObjectData.RefHandle.ToString(), InternalIndex, ToCStr(ObjectData.Protocol ? ObjectData.Protocol->DebugName : nullptr), MaxHugeObjectPayLoadBytes, Context.GetLocalConnectionId());
+			ensure(false);
+
 			Context.SetError(NetError_ObjectStateTooLarge);
 			return -1;
 		}
