[ChaosVD] Eliminating dynamic array allocations when building the Collision Data wrappers during a CVD recording

diff --git a/Engine/Plugins/Experimental/ChaosVD/Source/ChaosVD/Private/Visualizers/ChaosVDSolverCollisionDataComponentVisualizer.cpp b/Engine/Plugins/Experimental/ChaosVD/Source/ChaosVD/Private/Visualizers/ChaosVDSolverCollisionDataComponentVisualizer.cpp
--- a/Engine/Plugins/Experimental/ChaosVD/Source/ChaosVD/Private/Visualizers/ChaosVDSolverCollisionDataComponentVisualizer.cpp
+++ b/Engine/Plugins/Experimental/ChaosVD/Source/ChaosVD/Private/Visualizers/ChaosVDSolverCollisionDataComponentVisualizer.cpp
@@ -185,12 +185,6 @@ void FChaosVDSolverCollisionDataComponentVisualizer::DrawMidPhaseData(const UAct
 				HitPRoxyDataFinder.OwningConstraint = &Constraint;
 				HitPRoxyDataFinder.ContactIndex = ContactIndex;
 
-				constexpr int32 MinShapeContactAmount = 2;
-				if (Constraint.ShapeWorldTransforms.Num() < MinShapeContactAmount || ManifoldPoint.ContactPoint.ShapeContactPoints.Num() < MinShapeContactAmount)
-				{
-					continue;
-				}
-
 				const bool bIsProbe = Constraint.bIsProbe;
 				const bool bIsActive = ManifoldPoint.bIsValid && (!ManifoldPoint.NetPushOut.IsNearlyZero() || !ManifoldPoint.NetImpulse.IsNearlyZero() || (!Constraint.bUseManifold && !Constraint.AccumulatedImpulse.IsNearlyZero()));
 				if (!bIsActive && !EnumHasAnyFlags(VisualizationFlags, EChaosVDCollisionVisualizationFlags::DrawInactiveContacts))
diff --git a/Engine/Source/Runtime/Experimental/Chaos/Private/ChaosVisualDebugger/ChaosVDDataWrapperUtils.cpp b/Engine/Source/Runtime/Experimental/Chaos/Private/ChaosVisualDebugger/ChaosVDDataWrapperUtils.cpp
--- a/Engine/Source/Runtime/Experimental/Chaos/Private/ChaosVisualDebugger/ChaosVDDataWrapperUtils.cpp
+++ b/Engine/Source/Runtime/Experimental/Chaos/Private/ChaosVisualDebugger/ChaosVDDataWrapperUtils.cpp
@@ -11,6 +11,23 @@
 #include "DataWrappers/ChaosVDParticleDataWrapper.h"
 #include "Chaos/Collision/CollisionConstraintAllocator.h"
 
+namespace Chaos::VisualDebugger::Utils
+{
+	template<typename InType,typename OutType, int32 Size, typename TransformT>
+	void TransformStaticArray(const InType (&In)[Size], OutType (&Out)[Size], TransformT Trans)
+	{
+		for (int32 Index = 0; Index < Size; Index++)
+		{
+			Out[Index] = Invoke(Trans, In[Index]);
+		}
+	}
+
+	inline FTransform ConvertToFTransform(const FRigidTransform3& InChaosTransform)
+	{
+		return InChaosTransform;
+	}
+}
+
 void FChaosVDDataWrapperUtils::CopyManifoldPointsToDataWrapper(const Chaos::FManifoldPoint& InCopyFrom, FChaosVDManifoldPoint& OutCopyTo)
 {
 	OutCopyTo.bDisabled = InCopyFrom.Flags.bDisabled;
@@ -20,9 +37,9 @@ void FChaosVDDataWrapperUtils::CopyManifoldPointsToDataWrapper(const Chaos::FMan
 	OutCopyTo.TargetPhi = InCopyFrom.TargetPhi;
 	OutCopyTo.InitialPhi = InCopyFrom.InitialPhi;
 
-	Algo::Transform(InCopyFrom.ShapeAnchorPoints, OutCopyTo.ShapeAnchorPoints, &FChaosVDDataWrapperUtils::ConvertToFVector);
-	Algo::Transform(InCopyFrom.InitialShapeContactPoints, OutCopyTo.InitialShapeContactPoints, &FChaosVDDataWrapperUtils::ConvertToFVector);
-	Algo::Transform(InCopyFrom.ContactPoint.ShapeContactPoints, OutCopyTo.ContactPoint.ShapeContactPoints, &FChaosVDDataWrapperUtils::ConvertToFVector);
+	Chaos::VisualDebugger::Utils::TransformStaticArray(InCopyFrom.ShapeAnchorPoints, OutCopyTo.ShapeAnchorPoints, &FChaosVDDataWrapperUtils::ConvertToFVector);
+	Chaos::VisualDebugger::Utils::TransformStaticArray(InCopyFrom.InitialShapeContactPoints, OutCopyTo.InitialShapeContactPoints, &FChaosVDDataWrapperUtils::ConvertToFVector);
+	Chaos::VisualDebugger::Utils::TransformStaticArray(InCopyFrom.ContactPoint.ShapeContactPoints, OutCopyTo.ContactPoint.ShapeContactPoints, &FChaosVDDataWrapperUtils::ConvertToFVector);
 
 	OutCopyTo.ContactPoint.ShapeContactNormal = FVector(InCopyFrom.ContactPoint.ShapeContactNormal);
 	OutCopyTo.ContactPoint.Phi = InCopyFrom.ContactPoint.Phi;
@@ -125,9 +142,9 @@ FChaosVDConstraint FChaosVDDataWrapperUtils::BuildConstraintDataWrapperFromConst
 	WrappedConstraintData.Particle0Index = InConstraint.GetParticle0()->UniqueIdx().Idx;
 	WrappedConstraintData.Particle1Index = InConstraint.GetParticle1()->UniqueIdx().Idx;
 
-	Algo::Copy(InConstraint.ShapeWorldTransforms, WrappedConstraintData.ShapeWorldTransforms);
 
-	Algo::Copy(InConstraint.ImplicitTransform, WrappedConstraintData.ImplicitTransforms);
+	Chaos::VisualDebugger::Utils::TransformStaticArray(InConstraint.ShapeWorldTransforms, WrappedConstraintData.ShapeWorldTransforms, &Chaos::VisualDebugger::Utils::ConvertToFTransform);
+	Chaos::VisualDebugger::Utils::TransformStaticArray(InConstraint.ImplicitTransform, WrappedConstraintData.ImplicitTransforms, &Chaos::VisualDebugger::Utils::ConvertToFTransform);
 
 	WrappedConstraintData.CollisionMargins = TArray(InConstraint.CollisionMargins, std::size(InConstraint.CollisionMargins));
 	WrappedConstraintData.LastShapeWorldPositionDelta = FVector(InConstraint.LastShapeWorldPositionDelta);
@@ -145,7 +162,7 @@ FChaosVDConstraint FChaosVDDataWrapperUtils::BuildConstraintDataWrapperFromConst
 		{
 			const Chaos::FSavedManifoldPoint& CurrentChaosSavedManifoldPoint = InConstraint.SavedManifoldPoints[PointIndex];
 
-			Algo::Transform(CurrentChaosSavedManifoldPoint.ShapeContactPoints, CurrentCVDMainFoldPoint.ShapeContactPoints, &FChaosVDDataWrapperUtils::ConvertToFVector);
+			Chaos::VisualDebugger::Utils::TransformStaticArray(CurrentChaosSavedManifoldPoint.ShapeContactPoints, CurrentCVDMainFoldPoint.ShapeContactPoints, &FChaosVDDataWrapperUtils::ConvertToFVector);
 		}
 
 		if (PointIndex < InConstraint.ManifoldPoints.Num())
diff --git a/Engine/Source/Runtime/Experimental/ChaosVisualDebugger/Private/DataWrappers/ChaosVDCollisionDataWrappers.cpp b/Engine/Source/Runtime/Experimental/ChaosVisualDebugger/Private/DataWrappers/ChaosVDCollisionDataWrappers.cpp
--- a/Engine/Source/Runtime/Experimental/ChaosVisualDebugger/Private/DataWrappers/ChaosVDCollisionDataWrappers.cpp
+++ b/Engine/Source/Runtime/Experimental/ChaosVisualDebugger/Private/DataWrappers/ChaosVDCollisionDataWrappers.cpp
@@ -2,9 +2,20 @@
 
 #include "DataWrappers/ChaosVDCollisionDataWrappers.h"
 
+#ifndef CVD_SERIALIZE_STATIC_ARRAY
+	#define CVD_SERIALIZE_STATIC_ARRAY(Archive, Array) \
+		{ \
+			constexpr int32 Size = UE_ARRAY_COUNT(Array) ; \
+			for (int32 Index = 0; Index < Size; Index++)\
+			{\
+				Archive << Array[Index]; \
+			}\
+		}
+#endif
+
 bool FChaosVDContactPoint::Serialize(FArchive& Ar)
 {
-	Ar << ShapeContactPoints;
+	CVD_SERIALIZE_STATIC_ARRAY(Ar, ShapeContactPoints);
 	Ar << ShapeContactNormal;
 	Ar << Phi;
 	Ar << FaceIndex;
@@ -28,10 +39,12 @@ bool FChaosVDManifoldPoint::Serialize(FArchive& Ar)
 
 	Ar << TargetPhi;
 	Ar << InitialPhi;
-	Ar << ShapeAnchorPoints;
-	Ar << InitialShapeContactPoints;
+
+	CVD_SERIALIZE_STATIC_ARRAY(Ar, ShapeAnchorPoints);
+	CVD_SERIALIZE_STATIC_ARRAY(Ar, InitialShapeContactPoints);
+
 	Ar << ContactPoint;
-	Ar << ShapeContactPoints;
+	CVD_SERIALIZE_STATIC_ARRAY(Ar, ShapeContactPoints);
 
 	return true;
 }
@@ -73,8 +86,10 @@ bool FChaosVDConstraint::Serialize(FArchive& Ar)
 	Ar << Material;
 	Ar << AccumulatedImpulse;
 	Ar << ShapesType;
-	Ar << ShapeWorldTransforms;
-	Ar << ImplicitTransforms;
+
+	CVD_SERIALIZE_STATIC_ARRAY(Ar, ShapeWorldTransforms);
+	CVD_SERIALIZE_STATIC_ARRAY(Ar, ImplicitTransforms);
+
 	Ar << CullDistance;
 	Ar << CollisionMargins;	
 	Ar << CollisionTolerance;	
diff --git a/Engine/Source/Runtime/Experimental/ChaosVisualDebugger/Public/DataWrappers/ChaosVDCollisionDataWrappers.h b/Engine/Source/Runtime/Experimental/ChaosVisualDebugger/Public/DataWrappers/ChaosVDCollisionDataWrappers.h
--- a/Engine/Source/Runtime/Experimental/ChaosVisualDebugger/Public/DataWrappers/ChaosVDCollisionDataWrappers.h
+++ b/Engine/Source/Runtime/Experimental/ChaosVisualDebugger/Public/DataWrappers/ChaosVDCollisionDataWrappers.h
@@ -58,7 +58,7 @@ struct CHAOSVDRUNTIME_API FChaosVDContactPoint
 
 	// Shape-space contact points on the two bodies
 	UPROPERTY(VisibleAnywhere, Category=Contact)
-	TArray<FVector> ShapeContactPoints;
+	FVector ShapeContactPoints[2] = { FVector(ForceInit), FVector(ForceInit) };
 
 	// Shape-space contact normal on the second shape with direction that points away from shape 1
 	UPROPERTY(VisibleAnywhere, Category=Contact)
@@ -123,14 +123,14 @@ struct CHAOSVDRUNTIME_API FChaosVDManifoldPoint
 	UPROPERTY(VisibleAnywhere, Category=ContactData)
 	float InitialPhi = 0.f;
 	UPROPERTY()
-	TArray<FVector> ShapeAnchorPoints;
+	FVector ShapeAnchorPoints[2] = { FVector(ForceInit), FVector(ForceInit) };
 	UPROPERTY()
-	TArray<FVector> InitialShapeContactPoints;
+	FVector InitialShapeContactPoints[2] = { FVector(ForceInit), FVector(ForceInit) };
 	UPROPERTY(VisibleAnywhere, Category=ContactData)
 	FChaosVDContactPoint ContactPoint;
 
 	UPROPERTY()
-	TArray<FVector> ShapeContactPoints;
+	FVector ShapeContactPoints[2] =  { FVector(ForceInit), FVector(ForceInit) };
 
 	bool bIsSelectedInEditor = false;
 
@@ -256,10 +256,10 @@ struct CHAOSVDRUNTIME_API FChaosVDConstraint
 	EChaosVDContactShapesType ShapesType = EChaosVDContactShapesType::Unknown;
 	
 	UPROPERTY(VisibleAnywhere, Category=ConstraintData)
-	TArray<FTransform> ShapeWorldTransforms;
+	FTransform ShapeWorldTransforms[2] = { FTransform::Identity, FTransform::Identity };
 
 	UPROPERTY(VisibleAnywhere, Category=ConstraintData)
-	TArray<FTransform> ImplicitTransforms;
+	FTransform ImplicitTransforms[2] = { FTransform::Identity, FTransform::Identity };
 	
 	UPROPERTY(VisibleAnywhere, Category=ConstraintData)
 	float CullDistance = 0.f;
