Chaos Cloth: Fix FictitiousAngular forces to use the correct scaling with multiple subframes.

SKM cloth assets will be upgraded to retain behavior by multiplying FictitiousAngularScale by Subdivisions (subframes), but it will still respect the clamp of [0,2], so it's possible some assets will still change behavior.

ChaosClothAssets will not be upgraded.

diff --git a/Engine/Plugins/Experimental/ChaosCloth/Source/Chaos/Private/ChaosCloth/ChaosClothConfig.cpp b/Engine/Plugins/Experimental/ChaosCloth/Source/Chaos/Private/ChaosCloth/ChaosClothConfig.cpp
--- a/Engine/Plugins/Experimental/ChaosCloth/Source/Chaos/Private/ChaosCloth/ChaosClothConfig.cpp
+++ b/Engine/Plugins/Experimental/ChaosCloth/Source/Chaos/Private/ChaosCloth/ChaosClothConfig.cpp
@@ -116,7 +116,13 @@ void UChaosClothConfig::MigrateFrom(const UClothSharedConfigCommon* ClothSharedC
 			GravityScale = ChaosClothSharedSimConfig->GravityScale_DEPRECATED;
 			Gravity = ChaosClothSharedSimConfig->Gravity_DEPRECATED;
 		}
+		const int32 UE5ReleaseStreamObjectVersion = GetLinkerCustomVersion(FUE5ReleaseStreamObjectVersion::GUID);
+		if (UE5ReleaseStreamObjectVersion < FUE5ReleaseStreamObjectVersion::ChaosClothFictitiousAngularVelocitySubframeFix)
+		{
+			FictitiousAngularScale = FMath::Clamp(FictitiousAngularScale * (float)ChaosClothSharedSimConfig->SubdivisionCount, 0.f, 2.f);
+		}
 	}
+
 #endif  // #if WITH_EDITORONLY_DATA
 }
 
diff --git a/Engine/Plugins/Experimental/ChaosCloth/Source/Chaos/Private/ChaosCloth/ChaosClothConstraints.cpp b/Engine/Plugins/Experimental/ChaosCloth/Source/Chaos/Private/ChaosCloth/ChaosClothConstraints.cpp
--- a/Engine/Plugins/Experimental/ChaosCloth/Source/Chaos/Private/ChaosCloth/ChaosClothConstraints.cpp
+++ b/Engine/Plugins/Experimental/ChaosCloth/Source/Chaos/Private/ChaosCloth/ChaosClothConstraints.cpp
@@ -2079,14 +2079,14 @@ void FClothConstraints::CreatePBDRules()
 }
 
 void FClothConstraints::UpdateFromSolver(const FSolverVec3& SolverGravity, bool bPerClothGravityOverrideEnabled,
-	const FSolverVec3& FictitiousAngularDisplacement, const FSolverVec3& ReferenceSpaceLocation,
+	const FSolverVec3& FictitiousAngularVelocity, const FSolverVec3& ReferenceSpaceLocation,
 	const FSolverVec3& InSolverWindVelocity, const FSolverReal LegacyWindAdaptation)
 {
 	if (ExternalForces)
 	{
 		ExternalForces->SetWorldGravityMultiplier((FSolverReal)ClothingSimulationClothConsoleVariables::CVarGravityMultiplier.GetValueOnAnyThread());
 		ExternalForces->SetSolverGravityProperties(SolverGravity, bPerClothGravityOverrideEnabled);
-		ExternalForces->SetFictitiousForcesData(FictitiousAngularDisplacement, ReferenceSpaceLocation);
+		ExternalForces->SetFictitiousForcesData(FictitiousAngularVelocity, ReferenceSpaceLocation);
 		ExternalForces->SetSolverWind(InSolverWindVelocity, LegacyWindAdaptation);
 	}
 	SolverWindVelocity = InSolverWindVelocity;
diff --git a/Engine/Plugins/Experimental/ChaosCloth/Source/Chaos/Private/ChaosCloth/ChaosClothVisualization.cpp b/Engine/Plugins/Experimental/ChaosCloth/Source/Chaos/Private/ChaosCloth/ChaosClothVisualization.cpp
--- a/Engine/Plugins/Experimental/ChaosCloth/Source/Chaos/Private/ChaosCloth/ChaosClothVisualization.cpp
+++ b/Engine/Plugins/Experimental/ChaosCloth/Source/Chaos/Private/ChaosCloth/ChaosClothVisualization.cpp
@@ -753,6 +753,44 @@ static FAutoConsoleVariableRef CVarClothVizWeightMapName(TEXT("p.ChaosClothVisua
 		}
 	}
 
+
+	void FClothVisualization::DrawFictitiousAngularForces(FPrimitiveDrawInterface* PDI) const
+	{
+		if (!Solver)
+		{
+			return;
+		}
+
+		for (const FClothingSimulationCloth* const Cloth : Solver->GetCloths())
+		{
+			const int32 ParticleRangeId = Cloth->GetParticleRangeId(Solver);
+			if (ParticleRangeId == INDEX_NONE)
+			{
+				continue;
+			}
+			const FClothConstraints& ClothConstraints = Solver->GetClothConstraints(ParticleRangeId);
+			if (const Softs::FExternalForces* const ExternalForces = ClothConstraints.GetExternalForces().Get())
+			{
+				check(Solver->IsForceBasedSolver());
+				{
+					const FVec3& LocalSpaceLocation = Solver->GetLocalSpaceLocation();
+					const TConstArrayView<Softs::FSolverVec3> Positions = Solver->GetParticleXsView(ParticleRangeId);
+					const Softs::FSolverVec3& FictitousAngularVelocity = ExternalForces->GetFictitiousAngularVelocity();
+					const Softs::FSolverVec3& ReferenceSpaceLocation = ExternalForces->GetReferenceSpaceLocation();
+
+					for (int32 ParticleIndex = 0; ParticleIndex < Positions.Num(); ++ParticleIndex)
+					{
+						const FVector Pos0 = LocalSpaceLocation + FVector(Positions[ParticleIndex]);
+						const Softs::FSolverVec3 CentrifugalAccel = -Softs::FSolverVec3::CrossProduct(FictitousAngularVelocity, Softs::FSolverVec3::CrossProduct(FictitousAngularVelocity, Positions[ParticleIndex] - ReferenceSpaceLocation));
+
+						const FVector Pos1 = Pos0 + FVector(CentrifugalAccel);
+						DrawLine(PDI, Pos0, Pos1, FLinearColor::Red);
+					}
+				}
+			}
+		}
+	}
+
 	void FClothVisualization::DrawPhysMeshWired(FPrimitiveDrawInterface* PDI) const
 	{
 		if (!Solver)
diff --git a/Engine/Plugins/Experimental/ChaosCloth/Source/Chaos/Private/ChaosCloth/ChaosClothingSimulationSolver.cpp b/Engine/Plugins/Experimental/ChaosCloth/Source/Chaos/Private/ChaosCloth/ChaosClothingSimulationSolver.cpp
--- a/Engine/Plugins/Experimental/ChaosCloth/Source/Chaos/Private/ChaosCloth/ChaosClothingSimulationSolver.cpp
+++ b/Engine/Plugins/Experimental/ChaosCloth/Source/Chaos/Private/ChaosCloth/ChaosClothingSimulationSolver.cpp
@@ -187,7 +187,7 @@ FClothingSimulationSolver::FClothingSimulationSolver(bool bForceBasedSolver, FCl
 
 		// Add simulation groups arrays
 		Evolution->AddGroupArray(&PreSimulationTransforms);
-		Evolution->AddGroupArray(&FictitiousAngularDisplacements);
+		Evolution->AddGroupArray(&FictitiousAngularVelocities);
 		Evolution->AddGroupArray(&ReferenceSpaceLocations);
 
 		Evolution->AddParticleArray(&Normals);
@@ -284,7 +284,7 @@ FClothingSimulationSolver::FClothingSimulationSolver(bool bForceBasedSolver, FCl
 
 		// Add simulation groups arrays
 		PBDEvolution->AddArray(&PreSimulationTransforms);
-		PBDEvolution->AddArray(&FictitiousAngularDisplacements);
+		PBDEvolution->AddArray(&FictitiousAngularVelocities);
 		PBDEvolution->AddArray(&ReferenceSpaceLocations);
 
 		PBDEvolution->Particles().AddArray(&Normals);
@@ -1114,7 +1114,7 @@ void FClothingSimulationSolver::SetReferenceVelocityScale(
 	// Save the reference bone relative angular velocity for calculating the fictitious forces
 	const FVec3 FictitiousAngularDisplacement = ReferenceSpaceTransform.TransformVector(Axis * PartialDeltaAngle)
 		* AppliedFictitiousAngularScale;
-	FictitiousAngularDisplacements[GroupId] = Softs::FSolverVec3(FictitiousAngularDisplacement);
+	FictitiousAngularVelocities[GroupId] = DeltaTime > (Softs::FSolverReal)0.f ? Softs::FSolverVec3(FictitiousAngularDisplacement) / DeltaTime : Softs::FSolverVec3(0.f);
 	ReferenceSpaceLocations[GroupId] = ReferenceSpaceTransform.GetLocation() - LocalSpaceLocation;
 }
 
@@ -1238,7 +1238,7 @@ void FClothingSimulationSolver::SetProperties(int32 ParticleRangeId, const Softs
 		// Set properties to constraints that come from the solver (e.g., solver-level gravity, wind)
 		GetClothConstraints(ParticleRangeId).UpdateFromSolver(
 			Gravity, bIsClothGravityOverrideEnabled,
-			FictitiousAngularDisplacements[GroupId], ReferenceSpaceLocations[GroupId],
+			FictitiousAngularVelocities[GroupId], ReferenceSpaceLocations[GroupId],
 			WindVelocity, LegacyWindAdaption);
 	}
 }
@@ -1414,15 +1414,15 @@ void FClothingSimulationSolver::AddExternalForces(uint32 GroupId, bool bUseLegac
 {
 	if (PBDEvolution)
 	{
-		const FVec3& AngularDisplacement = FictitiousAngularDisplacements[GroupId];
+		const FVec3& AngularVelocity = FictitiousAngularVelocities[GroupId];
 		const FVec3& ReferenceSpaceLocation = ReferenceSpaceLocations[GroupId];
-		const bool bHasFictitiousForces = !AngularDisplacement.IsNearlyZero();
+		const bool bHasFictitiousForces = !AngularVelocity.IsNearlyZero();
 
 		static const FReal LegacyWindMultiplier = (FReal)25.;
 		const FVec3 LegacyWindVelocity = WindVelocity * LegacyWindMultiplier;
 
 		PBDEvolution->GetForceFunction(GroupId) =
-			[this, bHasFictitiousForces, bUseLegacyWind, LegacyWindVelocity, AngularDisplacement, ReferenceSpaceLocation](Softs::FSolverParticles& Particles, const FReal Dt, const int32 Index)
+			[this, bHasFictitiousForces, bUseLegacyWind, LegacyWindVelocity, AngularVelocity, ReferenceSpaceLocation](Softs::FSolverParticles& Particles, const FReal Dt, const int32 Index)
 			{
 				FVec3 Forces((FReal)0.);
 
@@ -1441,15 +1441,14 @@ void FClothingSimulationSolver::AddExternalForces(uint32 GroupId, bool bUseLegac
 				if (bHasFictitiousForces)
 				{
 					const FVec3 X = Particles.X(Index) - ReferenceSpaceLocation;
-					const FVec3 W = AngularDisplacement / Dt;
 					const FReal& M = Particles.M(Index);
 #if 0
 					// Coriolis + Centrifugal seems a bit overkilled, but let's keep the code around in case it's ever required
 					const FVec3& V = Particles.V(Index);
-					Forces -= (FVec3::CrossProduct(W, V) * 2.f + FVec3::CrossProduct(W, FVec3::CrossProduct(W, X))) * M;
+					Forces -= (FVec3::CrossProduct(AngularVelocity, V) * 2.f + FVec3::CrossProduct(AngularVelocity, FVec3::CrossProduct(AngularVelocity, X))) * M;
 #else
 					// Centrifugal force
-					Forces -= FVec3::CrossProduct(W, FVec3::CrossProduct(W, X)) * M;
+					Forces -= FVec3::CrossProduct(AngularVelocity, FVec3::CrossProduct(AngularVelocity, X)) * M;
 #endif
 				}
 				
diff --git a/Engine/Plugins/Experimental/ChaosCloth/Source/Chaos/Public/ChaosCloth/ChaosClothConstraints.h b/Engine/Plugins/Experimental/ChaosCloth/Source/Chaos/Public/ChaosCloth/ChaosClothConstraints.h
--- a/Engine/Plugins/Experimental/ChaosCloth/Source/Chaos/Public/ChaosCloth/ChaosClothConstraints.h
+++ b/Engine/Plugins/Experimental/ChaosCloth/Source/Chaos/Public/ChaosCloth/ChaosClothConstraints.h
@@ -44,7 +44,7 @@ PRAGMA_ENABLE_DEPRECATION_WARNINGS
 
 		// Force-based solver
 		void UpdateFromSolver(const Softs::FSolverVec3& SolverGravity, bool bPerClothGravityOverrideEnabled,
-			const Softs::FSolverVec3& FictitiousAngularDisplacement, const Softs::FSolverVec3& ReferenceSpaceLocation,
+			const Softs::FSolverVec3& FictitiousAngularVelocity, const Softs::FSolverVec3& ReferenceSpaceLocation,
 			const Softs::FSolverVec3& SolverWindVelocity, const Softs::FSolverReal LegacyWindAdaptation);
 
 		// PBD solver
diff --git a/Engine/Plugins/Experimental/ChaosCloth/Source/Chaos/Public/ChaosCloth/ChaosClothVisualization.h b/Engine/Plugins/Experimental/ChaosCloth/Source/Chaos/Public/ChaosCloth/ChaosClothVisualization.h
--- a/Engine/Plugins/Experimental/ChaosCloth/Source/Chaos/Public/ChaosCloth/ChaosClothVisualization.h
+++ b/Engine/Plugins/Experimental/ChaosCloth/Source/Chaos/Public/ChaosCloth/ChaosClothVisualization.h
@@ -51,6 +51,7 @@ namespace Chaos
 		CHAOSCLOTH_API void DrawSelfIntersection(FPrimitiveDrawInterface* PDI = nullptr) const;
 		CHAOSCLOTH_API void DrawBounds(FPrimitiveDrawInterface* PDI = nullptr) const;
 		CHAOSCLOTH_API void DrawGravity(FPrimitiveDrawInterface* PDI = nullptr) const;
+		CHAOSCLOTH_API void DrawFictitiousAngularForces(FPrimitiveDrawInterface* PDI = nullptr) const;
 
 #else  // #if CHAOS_DEBUG_DRAW
 		void SetSolver(const ::Chaos::FClothingSimulationSolver* /*InSolver*/) {}
@@ -75,6 +76,7 @@ namespace Chaos
 		void DrawSelfIntersection(FPrimitiveDrawInterface* /*PDI*/ = nullptr) const {}
 		void DrawBounds(FPrimitiveDrawInterface* /*PDI*/ = nullptr) const {}
 		void DrawGravity(FPrimitiveDrawInterface* /*PDI*/ = nullptr) const {}
+		void DrawFictitiousAngularForces(FPrimitiveDrawInterface* /*PDI*/ = nullptr) const {}
 
 #endif  // #if CHAOS_DEBUG_DRAW
 
diff --git a/Engine/Plugins/Experimental/ChaosCloth/Source/Chaos/Public/ChaosCloth/ChaosClothingSimulationSolver.h b/Engine/Plugins/Experimental/ChaosCloth/Source/Chaos/Public/ChaosCloth/ChaosClothingSimulationSolver.h
--- a/Engine/Plugins/Experimental/ChaosCloth/Source/Chaos/Public/ChaosCloth/ChaosClothingSimulationSolver.h
+++ b/Engine/Plugins/Experimental/ChaosCloth/Source/Chaos/Public/ChaosCloth/ChaosClothingSimulationSolver.h
@@ -400,7 +400,7 @@ namespace Chaos
 
 		// Simulation group attributes
 		TArrayCollectionArray<Softs::FSolverRigidTransform3> PreSimulationTransforms;  // Allow a different frame of reference for each cloth groups
-		TArrayCollectionArray<Softs::FSolverVec3> FictitiousAngularDisplacements;  // Relative angular displacement of the reference bone that depends on the fictitious angular scale factor
+		TArrayCollectionArray<Softs::FSolverVec3> FictitiousAngularVelocities;  // Relative angular velocity of the reference bone. Depends on the fictitious angular scale factor when using PBDEvolution (is applied later for Evolution)
 		TArrayCollectionArray<Softs::FSolverVec3> ReferenceSpaceLocations;  // Center of rotations for fictitious forces in local coordinate to the simulation space location
 
 		// Particle attributes
diff --git a/Engine/Plugins/Experimental/ChaosClothAssetEditor/Source/ChaosClothAssetEditor/Private/ChaosClothAsset/ClothEditorSimulationVisualization.cpp b/Engine/Plugins/Experimental/ChaosClothAssetEditor/Source/ChaosClothAssetEditor/Private/ChaosClothAsset/ClothEditorSimulationVisualization.cpp
--- a/Engine/Plugins/Experimental/ChaosClothAssetEditor/Source/ChaosClothAssetEditor/Private/ChaosClothAsset/ClothEditorSimulationVisualization.cpp
+++ b/Engine/Plugins/Experimental/ChaosClothAssetEditor/Source/ChaosClothAssetEditor/Private/ChaosClothAsset/ClothEditorSimulationVisualization.cpp
@@ -161,6 +161,7 @@ const FVisualizationOption FVisualizationOption::OptionData[] =
 	FVisualizationOption(&::Chaos::FClothVisualization::DrawWindAndPressureForces, TEXT("p.ChaosClothAssetEditor.DebugDrawWindAndPressureForces"), LOCTEXT("ChaosVisName_WindAndPressureForces", "Wind Aerodynamic And Pressure Forces"), LOCTEXT("ChaosVisName_WindAndPressure_Tooltip", "Draws the Wind drag and lift and pressure forces")),
 	FVisualizationOption(&::Chaos::FClothVisualization::DrawLocalSpace          , TEXT("p.ChaosClothAssetEditor.DebugDrawLocalSpace"), LOCTEXT("ChaosVisName_LocalSpace"          , "Local Space Reference Bone"), LOCTEXT("ChaosVisName_LocalSpace_Tooltip"          , "Draws the local space reference bone")),
 	FVisualizationOption(&::Chaos::FClothVisualization::DrawGravity             , TEXT("p.ChaosClothAssetEditor.DebugDrawGravity"), LOCTEXT("ChaosVisName_Gravity", "Gravity"), LOCTEXT("ChaosVisName_Gravity_Tooltip", "Draws gravity")),
+	FVisualizationOption(&::Chaos::FClothVisualization::DrawFictitiousAngularForces             , TEXT("p.ChaosClothAssetEditor.DebugDrawFictitiousAngularForces"), LOCTEXT("ChaosVisName_FictitiousAngularForces", "Fictitious Angular Forces"), LOCTEXT("ChaosVisName_Gravity_FictitiousAngularForces", "Draws fictitious angular forces (force based solver only)")),
 	FVisualizationOption(&::Chaos::FClothVisualization::DrawSelfCollision       , TEXT("p.ChaosClothAssetEditor.DebugDrawSelfCollision"), LOCTEXT("ChaosVisName_SelfCollision"       , "Self Collision"), LOCTEXT("ChaosVisName_SelfCollision_Tooltip"       , "Draws the self collision thickness/debugging information")),
 	FVisualizationOption(&::Chaos::FClothVisualization::DrawSelfIntersection    , TEXT("p.ChaosClothAssetEditor.DebugDrawSelfIntersection"), LOCTEXT("ChaosVisName_SelfIntersection"    , "Self Intersection"), LOCTEXT("ChaosVisName_SelfIntersection_Tooltip"    , "Draws the self intersection contour/region information")),
 	FVisualizationOption(&GetSimulationStatisticsString    , TEXT("p.ChaosClothAssetEditor.DebugDrawSimulationStatistics"), LOCTEXT("ChaosVisName_SimulationStatistics"    , "Simulation Statistics"), LOCTEXT("ChaosVisName_SimulationStatistics_Tooltip"    , "Displays simulation statistics")),
diff --git a/Engine/Source/Runtime/Core/Public/UObject/UE5ReleaseStreamObjectVersion.h b/Engine/Source/Runtime/Core/Public/UObject/UE5ReleaseStreamObjectVersion.h
--- a/Engine/Source/Runtime/Core/Public/UObject/UE5ReleaseStreamObjectVersion.h
+++ b/Engine/Source/Runtime/Core/Public/UObject/UE5ReleaseStreamObjectVersion.h
@@ -156,6 +156,9 @@ struct FUE5ReleaseStreamObjectVersion
 
 		// Serialize ExternalDataLayerAsset in WorldPartitionActorDesc
 		WorldPartitionExternalDataLayers,
+
+		// Fix Chaos Cloth fictitious angular scale bug that requires existing parameter rescaling.
+		ChaosClothFictitiousAngularVelocitySubframeFix,
 		
 		// -----<new versions can be added above this line>-------------------------------------------------
 		VersionPlusOne,
diff --git a/Engine/Source/Runtime/Experimental/Chaos/Public/Chaos/SoftsExternalForces.h b/Engine/Source/Runtime/Experimental/Chaos/Public/Chaos/SoftsExternalForces.h
--- a/Engine/Source/Runtime/Experimental/Chaos/Public/Chaos/SoftsExternalForces.h
+++ b/Engine/Source/Runtime/Experimental/Chaos/Public/Chaos/SoftsExternalForces.h
@@ -23,7 +23,7 @@ public:
 		, GravityScale(InGravityScale,
 			InGravityScaleMultipliers,
 			Particles.GetRangeSize())
-		, FictitiousAngularDisplacement(0.f)
+		, FictitiousAngularVelocity(0.f)
 		, ReferenceSpaceLocation(0.f)
 		, bUsePointBasedWindModel(false)
 		, PointBasedWind(0.f)
@@ -41,8 +41,7 @@ public:
 		const FSolverVec3* const V = Particles.GetV().GetData();
 		const FSolverVec3* const N = Particles.GetConstArrayView(Normals).GetData();
 
-		const bool bHasFictitiousForces = !FictitiousAngularDisplacement.IsNearlyZero();
-		const FSolverVec3 W = FictitiousAngularDisplacement / Dt;
+		const bool bHasFictitiousForces = !FictitiousAngularVelocity.IsNearlyZero();
 
 		const bool bHasPerParticleGravity = HasPerParticleGravity();
 		const FSolverReal GravityScaleConstant = bApplyGravityScale ? (FSolverReal)GravityScale : (FSolverReal)1.f;
@@ -61,7 +60,7 @@ public:
 				if (bHasFictitiousForces)
 				{
 					// Centrifugal force (*InvM to get acceleration)
-					Acceleration[Index] -= FSolverVec3::CrossProduct(W, FSolverVec3::CrossProduct(W, X[Index] - ReferenceSpaceLocation));
+					Acceleration[Index] -= FSolverVec3::CrossProduct(FictitiousAngularVelocity, FSolverVec3::CrossProduct(FictitiousAngularVelocity, X[Index] - ReferenceSpaceLocation));
 				}
 				if (bUsePointBasedWindModel)
 				{
@@ -87,8 +86,7 @@ public:
 		const FSolverVec3* const V = Particles.GetV().GetData();
 		const FSolverVec3* const N = Particles.GetConstArrayView(Normals).GetData();
 
-		const bool bHasFictitiousForces = !FictitiousAngularDisplacement.IsNearlyZero();
-		const FSolverVec3 W = FictitiousAngularDisplacement / Dt;
+		const bool bHasFictitiousForces = !FictitiousAngularVelocity.IsNearlyZero();
 
 		const bool bHasPerParticleGravity = HasPerParticleGravity();
 		const FSolverReal GravityScaleConstant = bApplyGravityScale ? (FSolverReal)GravityScale : (FSolverReal)1.f;
@@ -107,7 +105,7 @@ public:
 				if (bHasFictitiousForces)
 				{
 					// Centrifugal force (*InvM to get acceleration)
-					Force -= FSolverVec3::CrossProduct(W, FSolverVec3::CrossProduct(W, X[Index] - ReferenceSpaceLocation)) * M[Index];
+					Force -= FSolverVec3::CrossProduct(FictitiousAngularVelocity, FSolverVec3::CrossProduct(FictitiousAngularVelocity, X[Index] - ReferenceSpaceLocation)) * M[Index];
 				}
 				if (bUsePointBasedWindModel)
 				{
@@ -134,12 +132,13 @@ public:
 		const FSolverReal ParticleGravityScale = bApplyGravityScale ? (GravityScale.HasWeightMap() ? GravityScale[ParticleIndex] : GravityScale.GetLow()) : (FSolverReal)1.f;
 		return Gravity * ParticleGravityScale;
 	}
-
+	const FSolverVec3& GetFictitiousAngularVelocity() const { return FictitiousAngularVelocity; }
+	const FSolverVec3& GetReferenceSpaceLocation() const { return ReferenceSpaceLocation; }
 protected:
 	FSolverVec3 Gravity; 
 	bool bApplyGravityScale;
 	FPBDFlatWeightMap GravityScale;
-	FSolverVec3 FictitiousAngularDisplacement;
+	FSolverVec3 FictitiousAngularVelocity;
 	FSolverVec3 ReferenceSpaceLocation;
 	bool bUsePointBasedWindModel;
 	FSolverVec3 PointBasedWind;
@@ -173,7 +172,7 @@ public:
 		, WorldGravityMultiplier((FSolverReal)1.f)
 		, SolverGravity((FSolverReal)0.f, (FSolverReal)0.f, DefaultGravityZOverride)
 		, bPerSoftBodyGravityOverrideEnabled(true)
-		, FictitiousAngularDisplacementNoScale(0.f)
+		, FictitiousAngularVelocityNoScale(0.f)
 		, UseGravityOverrideIndex(PropertyCollection)
 		, GravityOverrideIndex(PropertyCollection)
 		, GravityScaleIndex(PropertyCollection)
@@ -187,9 +186,9 @@ public:
 		SolverGravity = InSolverGravity;
 		bPerSoftBodyGravityOverrideEnabled = bInPerSoftBodyGravityOverrideEnabled;
 	}
-	void SetFictitiousForcesData(const FSolverVec3& InFictitiousAngularDisplacementNoScale, const FSolverVec3& InReferenceSpaceLocation)
+	void SetFictitiousForcesData(const FSolverVec3& InFictitiousAngularVelocityNoScale, const FSolverVec3& InReferenceSpaceLocation)
 	{
-		FictitiousAngularDisplacementNoScale = InFictitiousAngularDisplacementNoScale;
+		FictitiousAngularVelocityNoScale = InFictitiousAngularVelocityNoScale;
 		ReferenceSpaceLocation = InReferenceSpaceLocation;
 	}
 	void SetSolverWind(const FSolverVec3& SolverWind, const FSolverReal InLegacyWindAdaptation)
@@ -223,7 +222,7 @@ public:
 
 		const FSolverReal FictitiousAngularScale = FMath::Min((FSolverReal)2., FictitiousAngularScaleIndex != INDEX_NONE ? GetFictitiousAngularScale(PropertyCollection) : DefaultFictitiousAngularScale);
 
-		FictitiousAngularDisplacement = FictitiousAngularDisplacementNoScale * FictitiousAngularScale;
+		FictitiousAngularVelocity = FictitiousAngularVelocityNoScale * FictitiousAngularScale;
 
 		bUsePointBasedWindModel = UsePointBasedWindModelIndex != INDEX_NONE ? GetUsePointBasedWindModel(PropertyCollection) : bDefaultUsePointBasedWindModel;
 	}
@@ -244,7 +243,7 @@ private:
 	bool bPerSoftBodyGravityOverrideEnabled;
 
 	/** Fictitious Forces */
-	FSolverVec3 FictitiousAngularDisplacementNoScale; // Without Scale parameter applied.
+	FSolverVec3 FictitiousAngularVelocityNoScale; // Without Scale parameter applied.
 
 	UE_CHAOS_DECLARE_PROPERTYCOLLECTION_NAME(UseGravityOverride, bool);
 	UE_CHAOS_DECLARE_PROPERTYCOLLECTION_NAME(GravityOverride, FVector3f);
