Correct case of command name and description

#rnx

diff --git a/Engine/Source/Runtime/Experimental/IoStoreOnDemand/Private/Tool/Test.cpp b/Engine/Source/Runtime/Experimental/IoStoreOnDemand/Private/Tool/Test.cpp
--- a/Engine/Source/Runtime/Experimental/IoStoreOnDemand/Private/Tool/Test.cpp
+++ b/Engine/Source/Runtime/Experimental/IoStoreOnDemand/Private/Tool/Test.cpp
@@ -50,8 +50,8 @@ static int32 TestCommandEntry(const FContext& Context)
 ////////////////////////////////////////////////////////////////////////////////
 static FCommand TestCommand(
 	TestCommandEntry,
-	TEXT("test"),
-	TEXT("run IAS tests"),
+	TEXT("Test"),
+	TEXT("Run IAS tests"),
 	{
 		TArgument<FStringView>(TEXT("-Host"), TEXT("Host of the HTTP test server")),
 		TArgument<FStringView>(TEXT("-Dir"), TEXT("Primary directory to use for cache tests")),
