Don't suffix a / if there's no bucket path

#rnx

diff --git a/Engine/Source/Runtime/Experimental/IoStoreOnDemand/Private/IoStoreOnDemand.cpp b/Engine/Source/Runtime/Experimental/IoStoreOnDemand/Private/IoStoreOnDemand.cpp
--- a/Engine/Source/Runtime/Experimental/IoStoreOnDemand/Private/IoStoreOnDemand.cpp
+++ b/Engine/Source/Runtime/Experimental/IoStoreOnDemand/Private/IoStoreOnDemand.cpp
@@ -1234,7 +1234,10 @@ TIoStatusOr<FIoStoreUploadResult> UploadContainerFiles(
 	uint64 TotalExistingBytes = 0;
 	{
 		TStringBuilder<256> TocsKey;
-		TocsKey << UploadParams.BucketPrefix << "/";
+		if (!UploadParams.BucketPrefix.IsEmpty())
+		{
+			TocsKey << UploadParams.BucketPrefix << "/";
+		}
 
 		UE_LOG(LogIas, Display, TEXT("	 '%s/%s/%s'"), *Client.GetConfig().ServiceUrl, *UploadParams.Bucket, TocsKey.ToString());
 		FS3ListObjectResponse Response = Client.ListObjects(FS3ListObjectsRequest
@@ -2101,7 +2104,10 @@ FIoStatus ListTocs(const FIoStoreListTocsParams& Params)
 		}
 
 		TStringBuilder<256> Path;
-		Path << PrefixView << TEXT("/");
+		if (!PrefixView.IsEmpty())
+		{
+			Path << PrefixView << TEXT("/");
+		}
 
 		UE_LOG(LogIas, Display, TEXT("Fetching TOC's from '%s/%s/%s'"), *Client.GetConfig().ServiceUrl, *Params.Bucket, Path.ToString());
 		FS3ListObjectResponse Response = Client.ListObjects(FS3ListObjectsRequest
