Moved ListTocs IAS command from UnrealPak to IasTool

#rnx

diff --git a/Engine/Source/Developer/IoStoreUtilities/Private/IoStoreUtilities.cpp b/Engine/Source/Developer/IoStoreUtilities/Private/IoStoreUtilities.cpp
--- a/Engine/Source/Developer/IoStoreUtilities/Private/IoStoreUtilities.cpp
+++ b/Engine/Source/Developer/IoStoreUtilities/Private/IoStoreUtilities.cpp
@@ -9743,23 +9743,3 @@ bool DownloadIoStoreContainerFiles(const TCHAR* TocPath)
 
 	return Status.IsOk();
 }
-
-bool ListOnDemandTocs()
-{
-	using namespace UE::IO::IAS;
-
-	TIoStatusOr<FIoStoreListTocsParams> Params = FIoStoreListTocsParams::Parse(FCommandLine::Get());
-	if (Params.IsOk() == false)
-	{
-		UE_LOG(LogIoStore, Error, TEXT("Failed to list TOC file(s), reason '%s'"), *Params.Status().ToString());
-		return false;
-	}
-
-	FIoStatus Status = ListTocs(Params.ConsumeValueOrDie());
-	if (Status.IsOk() == false)
-	{
-		UE_LOG(LogIoStore, Error, TEXT("Failed to list TOC file(s), reason '%s'"), *Status.ToString());
-	}
-
-	return Status.IsOk();
-}
diff --git a/Engine/Source/Developer/IoStoreUtilities/Public/IoStoreUtilities.h b/Engine/Source/Developer/IoStoreUtilities/Public/IoStoreUtilities.h
--- a/Engine/Source/Developer/IoStoreUtilities/Public/IoStoreUtilities.h
+++ b/Engine/Source/Developer/IoStoreUtilities/Public/IoStoreUtilities.h
@@ -62,6 +62,4 @@ UE_API bool UploadIoStoreContainerFiles(const TCHAR* ContainerPathOrWildcard);
 
 UE_API bool DownloadIoStoreContainerFiles(const TCHAR* TocPath);
 
-UE_API bool ListOnDemandTocs();
-
 #undef UE_API
diff --git a/Engine/Source/Developer/PakFileUtilities/Private/PakFileUtilities.cpp b/Engine/Source/Developer/PakFileUtilities/Private/PakFileUtilities.cpp
--- a/Engine/Source/Developer/PakFileUtilities/Private/PakFileUtilities.cpp
+++ b/Engine/Source/Developer/PakFileUtilities/Private/PakFileUtilities.cpp
@@ -5478,11 +5478,6 @@ bool ExecuteUnrealPak(const TCHAR* CmdLine)
 			{
 				return DownloadIoStoreContainerFiles(*IoStoreArg) == 0;
 			}
-
-			if (FParse::Param(CmdLine, TEXT("ListTocs")))
-			{
-				return ListOnDemandTocs();
-			}
 		}
 	}
 
diff --git a/Engine/Source/Runtime/Experimental/IoStoreOnDemand/Private/Tool/ListTocs.cpp b/Engine/Source/Runtime/Experimental/IoStoreOnDemand/Private/Tool/ListTocs.cpp
new file mode 100644
--- /dev/null
+++ b/Engine/Source/Runtime/Experimental/IoStoreOnDemand/Private/Tool/ListTocs.cpp
@@ -0,0 +1,124 @@
+// Copyright Epic Games, Inc. All Rights Reserved.
+
+#if defined(UE_WITH_IAS_TOOL)
+
+#include "Command.h"
+
+#include <IO/IoStoreOnDemand.h>
+
+namespace UE::IO::IAS::Tool {
+
+////////////////////////////////////////////////////////////////////////////////
+static FIoStoreListTocsParams BuildListTocsParams(const FContext& Context)
+{
+	FIoStoreListTocsParams Ret;
+
+	if (auto Value = Context.Get<FStringView>(TEXT("-BucketPath")); !Value.IsEmpty())
+	{
+		int32 Idx = INDEX_NONE;
+		if (Value.FindChar(TCHAR('/'), Idx))
+		{
+			Ret.Bucket = Value.Left(Idx);
+			FStringView BucketPrefix = Value.RightChop(Idx + 1);
+			if (BucketPrefix.EndsWith(TEXT("/")))
+			{
+				BucketPrefix.LeftChopInline(1);
+			}
+			Ret.BucketPrefix = BucketPrefix;
+		}
+		else
+		{
+			Ret.Bucket = Value;
+		}
+	}
+
+	Ret.OutFile					= Context.Get<FStringView>(TEXT("-Json"),					Ret.OutFile);
+	Ret.Bucket					= Context.Get<FStringView>(TEXT("-Bucket"),					Ret.Bucket);
+	Ret.BucketPrefix			= Context.Get<FStringView>(TEXT("-BucketPrefix"),			Ret.BucketPrefix);
+	Ret.ServiceUrl				= Context.Get<FStringView>(TEXT("-ServiceUrl"),				Ret.ServiceUrl);
+	Ret.Region					= Context.Get<FStringView>(TEXT("-Region"),					Ret.Region);
+	Ret.AccessKey				= Context.Get<FStringView>(TEXT("-AccessKey"),				Ret.AccessKey);
+	Ret.SecretKey				= Context.Get<FStringView>(TEXT("-SecretKey"),				Ret.SecretKey);
+	Ret.SessionToken			= Context.Get<FStringView>(TEXT("-SessionToken"),			Ret.SessionToken);
+	Ret.CredentialsFile			= Context.Get<FStringView>(TEXT("-CredentialsFile"),		Ret.CredentialsFile);
+	Ret.CredentialsFileKeyName	= Context.Get<FStringView>(TEXT("-CredentialsFileKeyName"),	Ret.CredentialsFileKeyName);
+	Ret.TocUrl					= Context.Get<FStringView>(TEXT("-TocUrl"),					Ret.TocUrl);
+	Ret.TocKey					= Context.Get<FStringView>(TEXT("-TocKey"),					Ret.TocKey);
+	Ret.BuildVersion			= Context.Get<FStringView>(TEXT("-BuildVersion"),			Ret.BuildVersion);
+	Ret.TargetPlatform			= Context.Get<FStringView>(TEXT("-TargetPlatform"),			Ret.TargetPlatform);
+	Ret.ChunkKeys				= Context.Get<FStringView>(TEXT("-ChunkKeys"),				Ret.ChunkKeys);
+	Ret.bTocEntries				= Context.Get<bool>(TEXT("-TocEntries"),					Ret.bTocEntries);
+	Ret.bBlockSizes				= Context.Get<bool>(TEXT("-BlockSizes"),					Ret.bBlockSizes);
+	Ret.bBlockHashes			= Context.Get<bool>(TEXT("-BlockHashes"),					Ret.bBlockHashes);
+
+	if (!Ret.TocUrl.IsEmpty())
+	{
+		FStringView View = Ret.TocUrl;
+		if (View.StartsWith(TEXTVIEW("http://")) && View.EndsWith(TEXTVIEW(".iochunktoc")))
+		{
+			int32 Delim = INDEX_NONE;
+			if (View.RightChop(7).FindChar(TEXT('/'), Delim))
+			{
+				Ret.ServiceUrl = FString(View.Left(7 +  Delim));
+				Ret.TocKey = View.RightChop(Ret.ServiceUrl.Len() + 1);
+			}
+		}
+	}
+
+	if (FIoStatus Validation = Ret.Validate(); !Validation.IsOk())
+	{
+		FString Reason = Validation.ToString();
+		Context.Abort(*Reason);
+	}
+
+	return Ret;
+}
+
+////////////////////////////////////////////////////////////////////////////////
+static int32 ListTocsCommandEntry(const Tool::FContext& Context)
+{
+	FIoStoreListTocsParams Params = BuildListTocsParams(Context);
+
+	FIoStatus Status = ListTocs(Params);
+	if (!Status.IsOk())
+	{
+		FString Reason = Status.ToString();
+		Context.Abort(*Reason);
+	}
+
+	return 0;
+}
+
+////////////////////////////////////////////////////////////////////////////////
+static FCommand ListTocsCommand(
+	ListTocsCommandEntry,
+	TEXT("ListTocs"),
+	TEXT("Enumerates and output information for available TOCs"),
+	{
+		TArgument<FStringView>(TEXT("-Json"),					TEXT("Optionally write output to the given JSON file")),
+		TArgument<FStringView>(TEXT("-TocUrl"),					TEXT("URL of a TOC file to fetch info of (overrides Bucket/ServiceUrl)")),
+		TArgument<FStringView>(TEXT("-TocKey"),					TEXT("Filter output by TOC key")),
+		TArgument<FStringView>(TEXT("-BuildVersion"),			TEXT("Filter results a particular build version")),
+		TArgument<FStringView>(TEXT("-TargetPlatform"),			TEXT("Only show results for a given platform")),
+		TArgument<FStringView>(TEXT("-ChunkKeys"),				TEXT("Path to write JSON file of chunk keys to")),
+		TArgument<bool>(TEXT("-TocEntries"),					TEXT("Include TOC entries in JSON output")),
+		TArgument<bool>(TEXT("-BlockSizes"),					TEXT("Write block sizes to JSON file")),
+		TArgument<bool>(TEXT("-BlockHashes"),					TEXT("Add block hash values in JSON")),
+
+		TArgument<FStringView>(TEXT("-Bucket"),					TEXT("AWS bucket name")),
+		TArgument<FStringView>(TEXT("-BucketPrefix"),			TEXT("Bucket path to enumerate")),
+		TArgument<FStringView>(TEXT("-BucketPath"),				TEXT("Alternative way to set Bucket/BucketPrefix")),
+		TArgument<FStringView>(TEXT("-ServiceUrl"),				TEXT("AWS service URL")),
+
+		TArgument<FStringView>(TEXT("-Region"),					TEXT("AWS region code")),
+		TArgument<FStringView>(TEXT("-AccessKey"),				TEXT("AWS access key")),
+		TArgument<FStringView>(TEXT("-SecretKey"),				TEXT("AWS secret key")),
+		TArgument<FStringView>(TEXT("-SessionToken"),			TEXT("AWS session token")),
+		TArgument<FStringView>(TEXT("-CredentialsFile"),		TEXT("AWS credentials file")),
+		TArgument<FStringView>(TEXT("-CredentialsFileKeyName"),	TEXT("AWS credentials to use from file")),
+	}
+);
+
+} // namespace UE::IO::IAS::Tool
+
+#endif // UE_WITH_IAS_TOOL
diff --git a/Engine/Source/Runtime/Experimental/IoStoreOnDemand/Public/IO/IoStoreOnDemand.h b/Engine/Source/Runtime/Experimental/IoStoreOnDemand/Public/IO/IoStoreOnDemand.h
--- a/Engine/Source/Runtime/Experimental/IoStoreOnDemand/Public/IO/IoStoreOnDemand.h
+++ b/Engine/Source/Runtime/Experimental/IoStoreOnDemand/Public/IO/IoStoreOnDemand.h
@@ -293,11 +293,11 @@ struct FIoStoreListTocsParams
 	bool bBlockSizes = false;
 	bool bBlockHashes = false;
 
-	UE_API static TIoStatusOr<FIoStoreListTocsParams> Parse(const TCHAR* CommandLine);
+	static TIoStatusOr<FIoStoreListTocsParams> Parse(const TCHAR* CommandLine);
 	FIoStatus Validate() const;
 };
 
-UE_API FIoStatus ListTocs(const FIoStoreListTocsParams& Params);
+FIoStatus ListTocs(const FIoStoreListTocsParams& Params);
 
 #endif // (IS_PROGRAM || WITH_EDITOR)
 
