Filter by target platform separate from build version

#rnx

diff --git a/Engine/Source/Runtime/Experimental/IoStoreOnDemand/Private/IoStoreOnDemand.cpp b/Engine/Source/Runtime/Experimental/IoStoreOnDemand/Private/IoStoreOnDemand.cpp
--- a/Engine/Source/Runtime/Experimental/IoStoreOnDemand/Private/IoStoreOnDemand.cpp
+++ b/Engine/Source/Runtime/Experimental/IoStoreOnDemand/Private/IoStoreOnDemand.cpp
@@ -2066,7 +2066,7 @@ FIoStatus ListTocs(const FIoStoreListTocsParams& Params)
 		uint64 TotalChunkSize = 0;
 	};
 
-	const bool bFilteredQuery = !Params.TocKey.IsEmpty() || !Params.BuildVersion.IsEmpty();
+	const bool bFilteredQuery = !Params.TocKey.IsEmpty() || !Params.BuildVersion.IsEmpty() || !Params.TargetPlatform.IsEmpty();
 
 	FS3ClientConfig Config;
 	Config.ServiceUrl = Params.ServiceUrl;
@@ -2159,16 +2159,14 @@ FIoStatus ListTocs(const FIoStoreListTocsParams& Params)
 				return;
 			}
 
-			if (!Params.BuildVersion.IsEmpty())
+			if (!Params.BuildVersion.IsEmpty() && !Params.BuildVersion.Equals(Toc.Meta.BuildVersion, ESearchCase::IgnoreCase))
 			{
-				if (!Params.BuildVersion.Equals(Toc.Meta.BuildVersion, ESearchCase::IgnoreCase))
-				{
-					return;
-				}
-				if (!Params.TargetPlatform.IsEmpty() && !Params.TargetPlatform.Equals(Toc.Meta.TargetPlatform, ESearchCase::IgnoreCase))
-				{
-					return;
-				}
+				return;
+			}
+
+			if (!Params.TargetPlatform.IsEmpty() && !Params.TargetPlatform.Equals(Toc.Meta.TargetPlatform, ESearchCase::IgnoreCase))
+			{
+				return;
 			}
 
 			FDateTime DateTime = FDateTime::FromUnixTimestamp(Toc.Meta.EpochTimestamp);
