Missing include

#rnx

diff --git a/Engine/Source/Runtime/Experimental/IoStoreOnDemand/Private/Tool/Command.cpp b/Engine/Source/Runtime/Experimental/IoStoreOnDemand/Private/Tool/Command.cpp
--- a/Engine/Source/Runtime/Experimental/IoStoreOnDemand/Private/Tool/Command.cpp
+++ b/Engine/Source/Runtime/Experimental/IoStoreOnDemand/Private/Tool/Command.cpp
@@ -7,6 +7,7 @@
 #include <Containers/UnrealString.h>
 #include <CoreGlobals.h>
 #include <Misc/OutputDeviceRedirector.h>
+#include <Misc/ScopeExit.h>
 
 namespace UE::IO::IAS::Tool {
 
