Iris
* Fiinish renaming Copy phase to Quantize phase in csv stats

#rnx

diff --git a/Engine/Source/Runtime/Experimental/Iris/Core/Private/Iris/ReplicationSystem/ReplicationOperationsInternal.cpp b/Engine/Source/Runtime/Experimental/Iris/Core/Private/Iris/ReplicationSystem/ReplicationOperationsInternal.cpp
--- a/Engine/Source/Runtime/Experimental/Iris/Core/Private/Iris/ReplicationSystem/ReplicationOperationsInternal.cpp
+++ b/Engine/Source/Runtime/Experimental/Iris/Core/Private/Iris/ReplicationSystem/ReplicationOperationsInternal.cpp
@@ -134,7 +134,7 @@ uint32 FReplicationInstanceOperationsInternal::QuantizeObjectStateData(FNetBitSt
 				Cache.PopLastEntry();
 			}
 
-			UE_NET_IRIS_STATS_ADD_TIME_AND_COUNT_FOR_OBJECT(Timer, Copy, InternalIndex);
+			UE_NET_IRIS_STATS_ADD_TIME_AND_COUNT_FOR_OBJECT(Timer, Quantize, InternalIndex);
 
 			return 1U;
 		}
diff --git a/Engine/Source/Runtime/Experimental/Iris/Core/Private/Iris/Stats/NetStats.cpp b/Engine/Source/Runtime/Experimental/Iris/Core/Private/Iris/Stats/NetStats.cpp
--- a/Engine/Source/Runtime/Experimental/Iris/Core/Private/Iris/Stats/NetStats.cpp
+++ b/Engine/Source/Runtime/Experimental/Iris/Core/Private/Iris/Stats/NetStats.cpp
@@ -27,8 +27,8 @@ CSV_DEFINE_CATEGORY(IrisPollMS, WITH_SERVER_CODE);
 CSV_DEFINE_CATEGORY(IrisPollCount, WITH_SERVER_CODE);
 CSV_DEFINE_CATEGORY(IrisPollWasteMS, WITH_SERVER_CODE);
 CSV_DEFINE_CATEGORY(IrisPollWasteCount, WITH_SERVER_CODE);
-CSV_DEFINE_CATEGORY(IrisCopyMS, WITH_SERVER_CODE);
-CSV_DEFINE_CATEGORY(IrisCopyCount, WITH_SERVER_CODE);
+CSV_DEFINE_CATEGORY(IrisQuantizeMS, WITH_SERVER_CODE);
+CSV_DEFINE_CATEGORY(IrisQuantizeCount, WITH_SERVER_CODE);
 CSV_DEFINE_CATEGORY(IrisWriteMS, WITH_SERVER_CODE);
 CSV_DEFINE_CATEGORY(IrisWriteCount, WITH_SERVER_CODE);
 CSV_DEFINE_CATEGORY(IrisWriteKBytes, WITH_SERVER_CODE);
@@ -214,15 +214,15 @@ void FNetTypeStats::ReportCSVStats()
 			UE_NET_STATS_RECORD_TYPESTATS_TIME(StatsName, PollWaste, TypeStatsData);
 			UE_NET_STATS_RECORD_TYPESTATS_COUNT(StatsName, PollWaste, TypeStatsData);
 
-			UE_NET_STATS_RECORD_TYPESTATS_TIME(StatsName, Copy, TypeStatsData);
-			UE_NET_STATS_RECORD_TYPESTATS_COUNT(StatsName, Copy, TypeStatsData);
+			UE_NET_STATS_RECORD_TYPESTATS_TIME(StatsName, Quantize, TypeStatsData);
+			UE_NET_STATS_RECORD_TYPESTATS_COUNT(StatsName, Quantize, TypeStatsData);
 
 			UE_NET_STATS_RECORD_TYPESTATS_TIME(StatsName, Write, TypeStatsData);
 			UE_NET_STATS_RECORD_TYPESTATS_COUNT(StatsName, Write, TypeStatsData);
 			UE_NET_STATS_RECORD_TYPESTATS_BITS(StatsName, Write, TypeStatsData);
 
 			UE_NET_STATS_RECORD_TYPESTATS_TIME(StatsName, WriteWaste, TypeStatsData);
-			UE_NET_STATS_RECORD_TYPESTATS_COUNT(StatsName, Write, TypeStatsData);
+			UE_NET_STATS_RECORD_TYPESTATS_COUNT(StatsName, WriteWaste, TypeStatsData);
 			UE_NET_STATS_RECORD_TYPESTATS_BITS(StatsName, WriteWaste, TypeStatsData);
 
 			UE_NET_STATS_RECORD_TYPESTATS_BITS(StatsName, WriteCreationInfo, TypeStatsData);
diff --git a/Engine/Source/Runtime/Experimental/Iris/Core/Private/Iris/Stats/NetStatsContext.h b/Engine/Source/Runtime/Experimental/Iris/Core/Private/Iris/Stats/NetStatsContext.h
--- a/Engine/Source/Runtime/Experimental/Iris/Core/Private/Iris/Stats/NetStatsContext.h
+++ b/Engine/Source/Runtime/Experimental/Iris/Core/Private/Iris/Stats/NetStatsContext.h
@@ -18,7 +18,7 @@ struct FNetTypeStatsData
 		PreUpdate = 0U,
 		Poll,
 		PollWaste,
-		Copy,
+		Quantize,
 		Write,
 		WriteWaste,
 		WriteCreationInfo,
