Iris
- Improved error handling when we fail to write a CreationHeader, it will now abort writing instead of retry using huge object path.

diff --git a/Engine/Source/Runtime/Experimental/Iris/Core/Private/Iris/ReplicationSystem/ReplicationWriter.cpp b/Engine/Source/Runtime/Experimental/Iris/Core/Private/Iris/ReplicationSystem/ReplicationWriter.cpp
--- a/Engine/Source/Runtime/Experimental/Iris/Core/Private/Iris/ReplicationSystem/ReplicationWriter.cpp
+++ b/Engine/Source/Runtime/Experimental/Iris/Core/Private/Iris/ReplicationSystem/ReplicationWriter.cpp
@@ -954,7 +954,7 @@ uint32 FReplicationWriter::SortScheduledObjects(FScheduleObjectInfo* ScheduledOb
 		IRIS_PROFILER_SCOPE(FReplicationWriter_SortScheduledObjects);
 
 		// We only need a partial sort of the highest priority objects as we wont be able to fit that much data in a packet anyway
-		// $IRIS TODO: Implement and evalute partial sort algorithm, currently we simply use std::partial_sort https://jira.it.epicgames.com/browse/UE-123444
+		// $IRIS TODO: Implement and evaluate partial sort algorithm, currently we simply use std::partial_sort https://jira.it.epicgames.com/browse/UE-123444
 		FScheduleObjectInfo* StartIt = ScheduledObjectIndices + StartIndex;
 		FScheduleObjectInfo* EndIt = ScheduledObjectIndices + ScheduledObjectCount;
 		FScheduleObjectInfo* SortIt = FMath::Min(StartIt + PartialSortObjectCount, EndIt);
@@ -2086,7 +2086,15 @@ FReplicationWriter::EWriteObjectStatus FReplicationWriter::WriteObjectAndSubObje
 					// We need to send creation info, so if we fail we skip this object for now
 					if (!bWriteSuccess)
 					{
-						return BridgeContext.SerializationContext.HasError() ? EWriteObjectStatus::Error : EWriteObjectStatus::BitStreamOverflow;
+						if (!Context.HasErrorOrOverflow())
+						{
+							// Unforced error, treat it as we have no instance and cannot create this object but we can continue with other objects
+							return EWriteObjectStatus::NoInstanceProtocol;
+						}
+						else
+						{
+							return Context.HasError() ? EWriteObjectStatus::Error : EWriteObjectStatus::BitStreamOverflow;
+						}
 					}
 				}
 				// Serialize initial state data for this object using delta compression against default state
