Hide some implementation details.

#rnx

diff --git a/Engine/Source/Runtime/Experimental/IoStoreOnDemand/Private/Tool/Command.cpp b/Engine/Source/Runtime/Experimental/IoStoreOnDemand/Private/Tool/Command.cpp
--- a/Engine/Source/Runtime/Experimental/IoStoreOnDemand/Private/Tool/Command.cpp
+++ b/Engine/Source/Runtime/Experimental/IoStoreOnDemand/Private/Tool/Command.cpp
@@ -82,6 +82,74 @@ static bool IsAnAstleySingleFrom91(const TCHAR* Cassette)
 
 
 
+////////////////////////////////////////////////////////////////////////////////
+using ParseFunc = void (void*, FStringView);
+
+////////////////////////////////////////////////////////////////////////////////
+template <typename Type>
+struct TArgumentParse
+{
+	static ParseFunc* GetParser()
+	{
+		return [] (void* Out, FStringView Input) {
+			auto* Inner = (Type*)Out;
+			LexFromString(*Inner, Input);
+		};
+	}
+};
+
+////////////////////////////////////////////////////////////////////////////////
+template <>
+struct TArgumentParse<bool>
+{
+	static ParseFunc* GetParser() { return nullptr; }
+};
+
+////////////////////////////////////////////////////////////////////////////////
+template <>
+struct TArgumentParse<FStringView>
+{
+	static ParseFunc* GetParser()
+	{
+		return [] (void* Out, FStringView Input) {
+			auto* Inner = (FStringView*)Out;
+			new (Inner) FStringView(Input);
+		};
+	}
+};
+
+////////////////////////////////////////////////////////////////////////////////
+static void Parse(const FArgument& Argument, void* Out, FStringView Input)
+{
+	static_assert(sizeof(PTRINT) == sizeof(ParseFunc*));
+	auto* Parser = (ParseFunc*)(Argument.Inner);
+	return Parser(Out, Input);
+}
+
+////////////////////////////////////////////////////////////////////////////////
+template <typename Type>
+static FArgument MakeArgument(FStringView Name, FStringView Desc)
+{
+	static_assert(alignof(Type) <= FArgument::ValueAlign);
+	static_assert(sizeof(Type) <= FArgument::ValueSize);
+	return { Name, Desc, PTRINT(TArgumentParse<Type>::GetParser()) };
+}
+
+#define TARGUMENT_IMPL(t) \
+	template <> FArgument TArgument<t>(FStringView Name, FStringView Desc) { \
+		return MakeArgument<t>(Name, Desc); \
+	}
+	TARGUMENT_IMPL(FStringView)
+	TARGUMENT_IMPL(bool)
+	TARGUMENT_IMPL(uint8)	TARGUMENT_IMPL(int8)
+	TARGUMENT_IMPL(uint16)	TARGUMENT_IMPL(int16)
+	TARGUMENT_IMPL(uint32)	TARGUMENT_IMPL(int32)
+	TARGUMENT_IMPL(uint64)	TARGUMENT_IMPL(int64)
+	TARGUMENT_IMPL(float)	TARGUMENT_IMPL(double)
+#undef TARGUMENT_IMPL
+
+
+
 ////////////////////////////////////////////////////////////////////////////////
 FContext::FContext(const FArguments& InArguments, int32 ArgC, const TCHAR* const* ArgV)
 : Arguments(InArguments)
@@ -116,7 +184,7 @@ FContext::FContext(const FArguments& InArguments, int32 ArgC, const TCHAR* const
 				FArgument Argument = Arguments[i];
 				FatalError<Error::FMissingValue>(Argument.Name);
 			}
-			Arguments[i].Parser(Values[i].Buffer, Arg);
+			Parse(Arguments[i], Values[i].Buffer, Arg);
 			SetMask |= 1ull << uint64(i);
 			PendingOptional = -1;
 			return;
@@ -130,7 +198,7 @@ FContext::FContext(const FArguments& InArguments, int32 ArgC, const TCHAR* const
 				continue;
 			}
 
-			if (Argument.Parser != nullptr)
+			if (Argument.Inner != 0)
 			{
 				PendingOptional = i;
 				return;
@@ -162,12 +230,12 @@ FContext::FContext(const FArguments& InArguments, int32 ArgC, const TCHAR* const
 		}
 
 		FArgument& Argument = Arguments[PositionalIndex];
-		if (Argument.Parser == nullptr)
+		if (Argument.Inner == 0)
 		{
 			FatalError<Error::FBoolPositional>(Arg);
 		}
 
-		Argument.Parser(Values[PositionalIndex].Buffer, Arg);
+		Parse(Argument, Values[PositionalIndex].Buffer, Arg);
 		SetMask |= 1ull << uint64(PositionalIndex);
 	};
 
@@ -307,7 +375,7 @@ void FCommand::Usage() const
 			continue;
 		}
 
-		const TCHAR* Suffix = (Argument.Parser != nullptr) ? TEXT("=<value>") : TEXT("");
+		const TCHAR* Suffix = (Argument.Inner != 0) ? TEXT("=<value>") : TEXT("");
 		WriteLine(TEXT( "  %s%s"), Argument.Name.GetData(), Suffix);
 		if (!Argument.Desc.IsEmpty())
 		{
@@ -373,6 +441,8 @@ int32 FCommand::MainInner(int32 ArgC, const TCHAR* const* ArgV)
 	FatalError<Error::FCommandNotFound>(Action);
 }
 
+
+
 ////////////////////////////////////////////////////////////////////////////////
 void CommandTest()
 {
diff --git a/Engine/Source/Runtime/Experimental/IoStoreOnDemand/Private/Tool/Command.h b/Engine/Source/Runtime/Experimental/IoStoreOnDemand/Private/Tool/Command.h
--- a/Engine/Source/Runtime/Experimental/IoStoreOnDemand/Private/Tool/Command.h
+++ b/Engine/Source/Runtime/Experimental/IoStoreOnDemand/Private/Tool/Command.h
@@ -11,9 +11,6 @@
 
 namespace UE::IO::IAS::Tool {
 
-////////////////////////////////////////////////////////////////////////////////
-using ParseFunc = void (void*, FStringView);
-
 ////////////////////////////////////////////////////////////////////////////////
 struct FArgument
 {
@@ -21,52 +18,13 @@ struct FArgument
 	static const uint32 ValueSize = 16;
 	FStringView			Name;
 	FStringView			Desc;
-	ParseFunc*			Parser;
+	PTRINT				Inner;
 };
 
 using FArguments = TArray<FArgument>;
 
 ////////////////////////////////////////////////////////////////////////////////
-template <typename Type>
-struct TArgumentParse
-{
-	static ParseFunc* GetParser()
-	{
-		return [] (void* Out, FStringView Input) {
-			auto* Inner = (Type*)Out;
-			LexFromString(*Inner, Input);
-		};
-	}
-};
-
-////////////////////////////////////////////////////////////////////////////////
-template <>
-struct TArgumentParse<bool>
-{
-	static ParseFunc* GetParser() { return nullptr; }
-};
-
-////////////////////////////////////////////////////////////////////////////////
-template <>
-struct TArgumentParse<FStringView>
-{
-	static ParseFunc* GetParser()
-	{
-		return [] (void* Out, FStringView Input) {
-			auto* Inner = (FStringView*)Out;
-			new (Inner) FStringView(Input);
-		};
-	}
-};
-
-////////////////////////////////////////////////////////////////////////////////
-template <typename Type>
-FArgument TArgument(FStringView Name, FStringView Desc)
-{
-	static_assert(alignof(Type) <= FArgument::ValueAlign);
-	static_assert(sizeof(Type) <= FArgument::ValueSize);
-	return { Name, Desc, TArgumentParse<Type>::GetParser() };
-}
+template <typename Type> FArgument TArgument(FStringView Name, FStringView Desc);
 
 
 
