Arg sets - a simple answer to sharing common arguments between commands

#rnx

diff --git a/Engine/Source/Runtime/Experimental/IoStoreOnDemand/Private/Tool/Command.cpp b/Engine/Source/Runtime/Experimental/IoStoreOnDemand/Private/Tool/Command.cpp
--- a/Engine/Source/Runtime/Experimental/IoStoreOnDemand/Private/Tool/Command.cpp
+++ b/Engine/Source/Runtime/Experimental/IoStoreOnDemand/Private/Tool/Command.cpp
@@ -151,9 +151,24 @@ static FArgument MakeArgument(FStringView Name, FStringView Desc)
 
 
 ////////////////////////////////////////////////////////////////////////////////
-FContext::FContext(const FArguments& InArguments, int32 ArgC, const TCHAR* const* ArgV)
+FArgumentSet::FArgumentSet(std::initializer_list<FArgument> InArguments)
 : Arguments(InArguments)
 {
+}
+
+////////////////////////////////////////////////////////////////////////////////
+FArgumentSet::operator FArgument () const
+{
+	return { FStringView(), FStringView(), 0 - PTRINT(this) };
+}
+
+
+
+////////////////////////////////////////////////////////////////////////////////
+FContext::FContext(const FArguments& InArguments, int32 ArgC, const TCHAR* const* ArgV)
+{
+	SetArguments(InArguments);
+
 	auto NextArg = [&ArgC, &ArgV, i=int32(1), Pending=FStringView()] () mutable -> FStringView
 	{
 		ArgC--;
@@ -239,7 +254,7 @@ FContext::FContext(const FArguments& InArguments, int32 ArgC, const TCHAR* const
 		SetMask |= 1ull << uint64(PositionalIndex);
 	};
 
-	Values.SetNum(InArguments.Num());
+	Values.SetNum(Arguments.Num());
 
 	while (ArgC > 1)
 	{
@@ -261,6 +276,22 @@ FContext::FContext(const FArguments& InArguments, int32 ArgC, const TCHAR* const
 	}
 }
 
+////////////////////////////////////////////////////////////////////////////////
+void FContext::SetArguments(const FArguments& Input)
+{
+	for (const FArgument& Item : Input)
+	{
+		if (Item.Inner >= 0)
+		{
+			Arguments.Add(Item);
+			continue;
+		}
+
+		const auto* Set = (FArgumentSet*)(0 - Item.Inner);
+		SetArguments(Set->Arguments);
+	}
+}
+
 ////////////////////////////////////////////////////////////////////////////////
 uint32 FContext::FindArgument(FStringView Name) const
 {
@@ -287,6 +318,23 @@ uint32 FContext::FindArgument(FStringView Name) const
 	FatalError<Error::FCommandAbort>(Reason);
 }
 
+////////////////////////////////////////////////////////////////////////////////
+template <typename Lambda>
+static void IterArguments(const FArguments& Arguments, const Lambda& Callback)
+{
+	for (const FArgument& Argument : Arguments)
+	{
+		if (Argument.Inner >= 0)
+		{
+			Callback(Argument);
+			continue;
+		}
+
+		const auto* Set = (FArgumentSet*)(0 - Argument.Inner);
+		IterArguments(Set->Arguments, Callback);
+	}
+}
+
 
 
 ////////////////////////////////////////////////////////////////////////////////
@@ -340,18 +388,18 @@ void FCommand::Usage() const
 	TStringBuilder<128> UsageLine;
 	UsageLine << TEXT("  ");
 	UsageLine << Name;
-	for (const FArgument& Argument : Arguments)
+	IterArguments(Arguments, [&bHasOpts, &UsageLine] (const FArgument& Argument)
 	{
 		if (Argument.Name.StartsWith('-'))
 		{
 			bHasOpts = true;
-			continue;
+			return;
 		}
 
 		UsageLine << TEXT(" <");
 		UsageLine << Argument.Name;
 		UsageLine << TEXT(">");
-	}
+	});
 
 	if (bHasOpts)
 	{
@@ -368,11 +416,11 @@ void FCommand::Usage() const
 
 	WriteLine(TEXT(""));
 	WriteLine(TEXT("Options:"));
-	for (const FArgument& Argument : Arguments)
+	IterArguments(Arguments, [] (const FArgument& Argument)
 	{
 		if (!Argument.Name.StartsWith('-'))
 		{
-			continue;
+			return;
 		}
 
 		const TCHAR* Suffix = (Argument.Inner != 0) ? TEXT("=<value>") : TEXT("");
@@ -382,7 +430,7 @@ void FCommand::Usage() const
 			WriteLine(TEXT("    %s"), Argument.Desc.GetData());
 			WriteLine(TEXT(""));
 		}
-	}
+	});
 }
 
 ////////////////////////////////////////////////////////////////////////////////
diff --git a/Engine/Source/Runtime/Experimental/IoStoreOnDemand/Private/Tool/Command.h b/Engine/Source/Runtime/Experimental/IoStoreOnDemand/Private/Tool/Command.h
--- a/Engine/Source/Runtime/Experimental/IoStoreOnDemand/Private/Tool/Command.h
+++ b/Engine/Source/Runtime/Experimental/IoStoreOnDemand/Private/Tool/Command.h
@@ -9,6 +9,8 @@
 #include <Containers/StringView.h>
 #include <String/LexFromString.h>
 
+#include <initializer_list>
+
 namespace UE::IO::IAS::Tool {
 
 ////////////////////////////////////////////////////////////////////////////////
@@ -23,6 +25,14 @@ struct FArgument
 
 using FArguments = TArray<FArgument>;
 
+////////////////////////////////////////////////////////////////////////////////
+struct FArgumentSet
+{
+				FArgumentSet(std::initializer_list<FArgument>);
+				operator FArgument () const;
+	FArguments	Arguments;
+};
+
 ////////////////////////////////////////////////////////////////////////////////
 template <typename Type> FArgument TArgument(FStringView Name, FStringView Desc);
 
@@ -44,6 +54,7 @@ private:
 
 	friend				class FCommand;
 						FContext(const FArguments& InArguments, int32 ArgC, const TCHAR* const* ArgV);
+	void				SetArguments(const FArguments& Input);
 	uint32				FindArgument(FStringView Name) const;
 	[[noreturn]] void	OnMissingValue(FStringView Name) const;
 	bool				IsSet(uint32 Index) const { return !!((1ull << Index) & SetMask); }
