Cheap and cheerful path to make sure lower level exceptions are surface

#rnx

diff --git a/Engine/Source/Runtime/Experimental/IoStoreOnDemand/Private/Tool/Command.cpp b/Engine/Source/Runtime/Experimental/IoStoreOnDemand/Private/Tool/Command.cpp
--- a/Engine/Source/Runtime/Experimental/IoStoreOnDemand/Private/Tool/Command.cpp
+++ b/Engine/Source/Runtime/Experimental/IoStoreOnDemand/Private/Tool/Command.cpp
@@ -9,6 +9,12 @@
 #include <Misc/OutputDeviceRedirector.h>
 #include <Misc/ScopeExit.h>
 
+#if PLATFORM_WINDOWS
+#	include "Windows/AllowWindowsPlatformTypes.h"
+#		include <errhandlingapi.h>
+#	include "Windows/HideWindowsPlatformTypes.h"
+#endif
+
 namespace UE::IO::IAS::Tool {
 
 ////////////////////////////////////////////////////////////////////////////////
@@ -452,6 +458,17 @@ int32 FCommand::Call(int32 ArgC, const TCHAR* const* ArgV)
 ////////////////////////////////////////////////////////////////////////////////
 int32 FCommand::Main(int32 ArgC, const TCHAR* const* ArgV)
 {
+#if PLATFORM_WINDOWS
+	::SetUnhandledExceptionFilter([] (EXCEPTION_POINTERS* Pointers) -> LONG
+	{
+		uint32 Code = Pointers->ExceptionRecord->ExceptionCode;
+		void* Where = Pointers->ExceptionRecord->ExceptionAddress;
+		WriteLine(TEXT("ERROR: An exceptional circumstance occurred; %08x @ %p"), Code, Where);
+		check(false);
+		return EXCEPTION_CONTINUE_SEARCH;
+	});
+#endif
+
 	try
 	{
 		return MainInner(ArgC, ArgV);
