Add an argument set for S3 credentials that can be shared between commands

#rnx

diff --git a/Engine/Source/Runtime/Experimental/IoStoreOnDemand/Private/Tool/ListTocs.cpp b/Engine/Source/Runtime/Experimental/IoStoreOnDemand/Private/Tool/ListTocs.cpp
--- a/Engine/Source/Runtime/Experimental/IoStoreOnDemand/Private/Tool/ListTocs.cpp
+++ b/Engine/Source/Runtime/Experimental/IoStoreOnDemand/Private/Tool/ListTocs.cpp
@@ -90,6 +90,8 @@ static int32 ListTocsCommandEntry(const Tool::FContext& Context)
 }
 
 ////////////////////////////////////////////////////////////////////////////////
+extern FArgumentSet S3Arguments;
+
 static FCommand ListTocsCommand(
 	ListTocsCommandEntry,
 	TEXT("ListTocs"),
@@ -105,17 +107,9 @@ static FCommand ListTocsCommand(
 		TArgument<bool>(TEXT("-BlockSizes"),					TEXT("Write block sizes to JSON file")),
 		TArgument<bool>(TEXT("-BlockHashes"),					TEXT("Add block hash values in JSON")),
 
-		TArgument<FStringView>(TEXT("-Bucket"),					TEXT("AWS bucket name")),
 		TArgument<FStringView>(TEXT("-BucketPrefix"),			TEXT("Bucket path to enumerate")),
 		TArgument<FStringView>(TEXT("-BucketPath"),				TEXT("Alternative way to set Bucket/BucketPrefix")),
-		TArgument<FStringView>(TEXT("-ServiceUrl"),				TEXT("AWS service URL")),
-
-		TArgument<FStringView>(TEXT("-Region"),					TEXT("AWS region code")),
-		TArgument<FStringView>(TEXT("-AccessKey"),				TEXT("AWS access key")),
-		TArgument<FStringView>(TEXT("-SecretKey"),				TEXT("AWS secret key")),
-		TArgument<FStringView>(TEXT("-SessionToken"),			TEXT("AWS session token")),
-		TArgument<FStringView>(TEXT("-CredentialsFile"),		TEXT("AWS credentials file")),
-		TArgument<FStringView>(TEXT("-CredentialsFileKeyName"),	TEXT("AWS credentials to use from file")),
+		S3Arguments,
 	}
 );
 
diff --git a/Engine/Source/Runtime/Experimental/IoStoreOnDemand/Private/Tool/Main.cpp b/Engine/Source/Runtime/Experimental/IoStoreOnDemand/Private/Tool/Main.cpp
--- a/Engine/Source/Runtime/Experimental/IoStoreOnDemand/Private/Tool/Main.cpp
+++ b/Engine/Source/Runtime/Experimental/IoStoreOnDemand/Private/Tool/Main.cpp
@@ -6,6 +6,18 @@
 
 namespace UE::IO::IAS::Tool {
 
+////////////////////////////////////////////////////////////////////////////////
+FArgumentSet S3Arguments = {
+	TArgument<FStringView>(TEXT("-ServiceUrl"),				TEXT("S3 service URL")),
+	TArgument<FStringView>(TEXT("-Bucket"),					TEXT("S3 bucket name")),
+	TArgument<FStringView>(TEXT("-Region"),					TEXT("S3 region code")),
+	TArgument<FStringView>(TEXT("-AccessKey"),				TEXT("S3 access key")),
+	TArgument<FStringView>(TEXT("-SecretKey"),				TEXT("S3 secret key")),
+	TArgument<FStringView>(TEXT("-SessionToken"),			TEXT("S3 session token")),
+	TArgument<FStringView>(TEXT("-CredentialsFile"),		TEXT("S3 credentials file")),
+	TArgument<FStringView>(TEXT("-CredentialsFileKeyName"),	TEXT("S3 credentials to use from file")),
+};
+
 ////////////////////////////////////////////////////////////////////////////////
 int32 Main(int32 ArgC, TCHAR* ArgV[])
 {
