Constraints: fix dependency problems during creation
    - ensure a global uniqueness of the new name, as constraints can have different outers since CL28214854

+ added the ability to retrieve all constraints acting on the same target (UObject) when creating these dependencies. Indeed, in the context of constraining controls, comparing the hash of the handle is not sufficient, since they ultimately act on the same ControlRig.

#rnx

diff --git a/Engine/Source/Runtime/Experimental/Animation/Constraints/Private/ConstraintSubsystem.cpp b/Engine/Source/Runtime/Experimental/Animation/Constraints/Private/ConstraintSubsystem.cpp
--- a/Engine/Source/Runtime/Experimental/Animation/Constraints/Private/ConstraintSubsystem.cpp
+++ b/Engine/Source/Runtime/Experimental/Animation/Constraints/Private/ConstraintSubsystem.cpp
@@ -170,6 +170,11 @@ void UConstraintSubsystem::SetConstraintDependencies(
 	FConstraintTickFunction* InFunctionToTickBefore,
 	FConstraintTickFunction* InFunctionToTickAfter)
 {
+	if (!InFunctionToTickBefore || !InFunctionToTickAfter || InFunctionToTickBefore == InFunctionToTickAfter)
+	{
+		return;
+	}
+	
 	// look for child tick function in in parent's prerequisites. 
 	const TArray<FTickPrerequisite>& ParentPrerequisites = InFunctionToTickAfter->GetPrerequisites();
 	const bool bIsChildAPrerexOfParent = ParentPrerequisites.ContainsByPredicate([InFunctionToTickBefore](const FTickPrerequisite& Prerex)
diff --git a/Engine/Source/Runtime/Experimental/Animation/Constraints/Private/ConstraintsManager.cpp b/Engine/Source/Runtime/Experimental/Animation/Constraints/Private/ConstraintsManager.cpp
--- a/Engine/Source/Runtime/Experimental/Animation/Constraints/Private/ConstraintsManager.cpp
+++ b/Engine/Source/Runtime/Experimental/Animation/Constraints/Private/ConstraintsManager.cpp
@@ -500,12 +500,13 @@ UTickableConstraint* FConstraintsManagerController::AddConstraintFromCopy(UTicka
 
 int32 FConstraintsManagerController::GetConstraintIndex(const FGuid& InGuid) const
 {
-	UConstraintSubsystem* Subsystem = UConstraintSubsystem::Get();
+	const UConstraintSubsystem* Subsystem = UConstraintSubsystem::Get();
 	if (!Subsystem)
 	{
 		return INDEX_NONE;
 	}
-	TArray<TWeakObjectPtr<UTickableConstraint>> Constraints = Subsystem->GetConstraints(World);
+
+	const TArray<TWeakObjectPtr<UTickableConstraint>> Constraints = Subsystem->GetConstraints(World);
 	return Constraints.IndexOfByPredicate([InGuid](const TWeakObjectPtr<UTickableConstraint>& Constraint)
 		{
 			return 	(Constraint.IsValid() && Constraint->ConstraintID == InGuid);
@@ -514,12 +515,13 @@ int32 FConstraintsManagerController::GetConstraintIndex(const FGuid& InGuid) con
 
 int32 FConstraintsManagerController::GetConstraintIndex(const FName& InConstraintName) const
 {
-	UConstraintSubsystem* Subsystem = UConstraintSubsystem::Get();
+	const UConstraintSubsystem* Subsystem = UConstraintSubsystem::Get();
 	if (!Subsystem)
 	{
 		return INDEX_NONE;
 	}
-	TArray<TWeakObjectPtr<UTickableConstraint>> Constraints = Subsystem->GetConstraints(World);
+
+	const TArray<TWeakObjectPtr<UTickableConstraint>> Constraints = Subsystem->GetConstraints(World);
 	return Constraints.IndexOfByPredicate([InConstraintName](const TWeakObjectPtr<UTickableConstraint>& Constraint)
 	{
 		return 	(Constraint.IsValid() && Constraint->GetFName() == InConstraintName);
@@ -681,7 +683,7 @@ void FConstraintsManagerController::SetConstraintsDependencies(
 	{
 		return;
 	}
-	TArray<TWeakObjectPtr<UTickableConstraint>> Constraints = Subsystem->GetConstraints(World);
+	const TArray<TWeakObjectPtr<UTickableConstraint>> Constraints = Subsystem->GetConstraints(World);
 
 	const int32 IndexBefore = GetConstraintIndex(InNameToTickBefore);
 	const int32 IndexAfter = GetConstraintIndex(InNameToTickAfter);
@@ -693,7 +695,42 @@ void FConstraintsManagerController::SetConstraintsDependencies(
 	FConstraintTickFunction& FunctionToTickBefore = Constraints[IndexBefore]->GetTickFunction(World);
 	FConstraintTickFunction& FunctionToTickAfter = Constraints[IndexAfter]->GetTickFunction(World);
 
-	Subsystem->SetConstraintDependencies( &FunctionToTickBefore, &FunctionToTickAfter);
+	Subsystem->SetConstraintDependencies(&FunctionToTickBefore, &FunctionToTickAfter);
+}
+
+void FConstraintsManagerController::SetConstraintsDependencies(const struct FGuid& InGuidToTickBefore, const struct FGuid& InGuidToTickAfter) const
+{
+	if (!InGuidToTickBefore.IsValid() || !InGuidToTickAfter.IsValid() || InGuidToTickBefore == InGuidToTickAfter)
+	{
+		return;
+	}
+
+	UConstraintSubsystem* Subsystem = UConstraintSubsystem::Get();
+	if (!Subsystem)
+	{
+		return;
+	}
+
+	const TArray<TWeakObjectPtr<UTickableConstraint>>& Constraints = Subsystem->GetConstraintsArray(World);
+	const int32 IndexBefore = Constraints.IndexOfByPredicate([InGuidToTickBefore](const TWeakObjectPtr<UTickableConstraint>& Constraint)
+	{
+		return Constraint.IsValid() && Constraint->ConstraintID == InGuidToTickBefore;
+	});
+	
+	const int32 IndexAfter = Constraints.IndexOfByPredicate([InGuidToTickAfter](const TWeakObjectPtr<UTickableConstraint>& Constraint)
+	{
+		return Constraint.IsValid() && Constraint->ConstraintID == InGuidToTickAfter;
+	});
+	
+	if (IndexBefore == INDEX_NONE || IndexAfter == INDEX_NONE || IndexAfter == IndexBefore)
+	{
+		return;
+	}
+	
+	FConstraintTickFunction& FunctionToTickBefore = Constraints[IndexBefore]->GetTickFunction(World);
+	FConstraintTickFunction& FunctionToTickAfter = Constraints[IndexAfter]->GetTickFunction(World);
+
+	Subsystem->SetConstraintDependencies(&FunctionToTickBefore, &FunctionToTickAfter);
 }
 
 const TArray< TWeakObjectPtr<UTickableConstraint> >& FConstraintsManagerController::GetConstraintsArray() const
diff --git a/Engine/Source/Runtime/Experimental/Animation/Constraints/Private/TransformConstraint.cpp b/Engine/Source/Runtime/Experimental/Animation/Constraints/Private/TransformConstraint.cpp
--- a/Engine/Source/Runtime/Experimental/Animation/Constraints/Private/TransformConstraint.cpp
+++ b/Engine/Source/Runtime/Experimental/Animation/Constraints/Private/TransformConstraint.cpp
@@ -222,6 +222,40 @@ void PreEvaluateHandle(const TObjectPtr<UTransformableHandle>& InHandle)
 	}
 }
 
+static bool	bIncludeTarget = true;
+static FAutoConsoleVariableRef CVarIncludeTarget(
+	TEXT("Constraints.IncludeTarget"),
+	bIncludeTarget,
+	TEXT("Include target when getting child's existing constraints.")
+	);
+
+void PrintTickFunctions(UWorld* InWorld)
+{
+	const FConstraintsManagerController& Controller = FConstraintsManagerController::Get(InWorld);
+	const TArray<TWeakObjectPtr<UTickableConstraint>>& Constraint = Controller.GetConstraintsArray();
+	for (const TWeakObjectPtr<UTickableConstraint>& ConstraintPtr: Constraint )
+	{
+		if (ConstraintPtr.IsValid())
+		{
+			FConstraintTickFunction& Function = ConstraintPtr->GetTickFunction(InWorld);
+			UE_LOG(LogTemp, Warning, TEXT("constraint (%p) - function %s (%p)"), ConstraintPtr.Get(), *Function.DiagnosticMessage(), &Function);
+	
+			const TArray<FTickPrerequisite>& Prerequisites = Function.GetPrerequisites();
+			for (const FTickPrerequisite& Prerequisite: Prerequisites)
+			{
+				if (const FTickFunction* TickFunction = Prerequisite.Get())
+				{
+					UE_LOG(LogTemp, Warning, TEXT("\t%s (%p)"), *Prerequisite.PrerequisiteObject->GetFullName(), TickFunction);
+				}
+				else
+				{
+					UE_LOG(LogTemp, Warning, TEXT("\tInvalid Prerequisite"));
+				}
+			}
+		}
+	}
+}
+
 }
 
 /** 
@@ -1796,11 +1830,12 @@ bool FTransformConstraintUtils::AddConstraint(
 	Constraint->EnsurePrimaryDependency(InWorld);
 
 	// if child handle is the parent of some other constraints, ensure they will tick after that new one
+	const bool bSelf = InParentHandle->GetTarget().Get() && InParentHandle->GetTarget().Get() == InChildHandle->GetTarget().Get();
 	TArray<TWeakObjectPtr<UTickableConstraint>> ChildChildConstraints;
-	GetChildrenConstraints(InWorld, InChildHandle, ChildChildConstraints);
+	GetChildrenConstraints(InWorld, InChildHandle, ChildChildConstraints, ConstraintLocals::bIncludeTarget && !bSelf);
 	for (const TWeakObjectPtr<UTickableConstraint>& ChildConstraint: ChildChildConstraints)
 	{
-		Controller.SetConstraintsDependencies( NewConstraintName, ChildConstraint->GetFName());
+		Controller.SetConstraintsDependencies(Constraint->ConstraintID, ChildConstraint->ConstraintID);
 	}
 
 	// warn for possible cycles
@@ -2030,30 +2065,45 @@ int32 FTransformConstraintUtils::GetLastActiveConstraintIndex(const TArray< TWea
 
 void FTransformConstraintUtils::GetChildrenConstraints(
 	UWorld* World,
-	const UTransformableHandle* InParentHandle,
-	TArray< TWeakObjectPtr<UTickableConstraint> >& OutConstraints)
+	const UTransformableHandle* InHandle,
+	TArray< TWeakObjectPtr<UTickableConstraint> >& OutConstraints,
+	const bool bIncludeTarget)
 {
 	using ConstraintPtr = TWeakObjectPtr<UTickableConstraint>;
 	
-	// filter for transform constraints where the InParentHandle is the parent (based on its hash value)
-	const uint32 ParentHash = InParentHandle->GetHash();
-	auto Predicate = [ParentHash](const ConstraintPtr& Constraint)
+	// filter for transform constraints where the InHandle is the parent (based on its hash value)
+	// and also has the same target if bIncludeTarget is true
+	const uint32 ParentHash = InHandle->GetHash();
+	const UObject* ParentTarget = InHandle->GetTarget().Get();
+	auto Predicate = [ParentHash, bIncludeTarget, ParentTarget](const ConstraintPtr& Constraint)
 	{
 		const UTickableTransformConstraint* TransformConstraint = Cast<UTickableTransformConstraint>(Constraint.Get());
 		if (!TransformConstraint)
 		{
 			return false;
 		}
-		
-		if (TransformConstraint->ParentTRSHandle && (TransformConstraint->ParentTRSHandle->GetHash() == ParentHash))
+
+		if (TransformConstraint->ParentTRSHandle)
 		{
-			return true;
+			if (TransformConstraint->ParentTRSHandle->GetHash() == ParentHash)
+			{
+				return true;
+			}
+
+			if (bIncludeTarget && ParentTarget)
+			{
+				const UObject* Target = TransformConstraint->ParentTRSHandle->GetTarget().Get();
+				if (Target == ParentTarget)
+				{
+					return true;
+				}
+			}
 		}
 
 		return false;
 	};
 
-	const FConstraintsManagerController& Controller = FConstraintsManagerController::Get(World);	
+	const FConstraintsManagerController& Controller = FConstraintsManagerController::Get(World);
 	const TArray<ConstraintPtr> FilteredConstraints = Controller.GetConstraintsByPredicate(Predicate);
 	OutConstraints.Append(FilteredConstraints);
 }
diff --git a/Engine/Source/Runtime/Experimental/Animation/Constraints/Public/ConstraintsManager.h b/Engine/Source/Runtime/Experimental/Animation/Constraints/Public/ConstraintsManager.h
--- a/Engine/Source/Runtime/Experimental/Animation/Constraints/Public/ConstraintsManager.h
+++ b/Engine/Source/Runtime/Experimental/Animation/Constraints/Public/ConstraintsManager.h
@@ -252,9 +252,6 @@ public:
 	/** Get the index of the given constraint's ID. */
 	CONSTRAINTS_API int32 GetConstraintIndex(const FGuid& InGuid) const;
 
-	/** Remove the constraint by name. */
-	CONSTRAINTS_API bool RemoveConstraint(const FName& InConstraintName, bool bDoNotCompensate = false);
-
 	/** Remove the constraint at the given index. */
 	CONSTRAINTS_API bool RemoveConstraint(const int32 InConstraintIndex, bool bDoNotCompensate = false);
 
@@ -293,6 +290,10 @@ public:
 		const FName& InNameToTickBefore,
 		const FName& InNameToTickAfter) const;
 
+	CONSTRAINTS_API void SetConstraintsDependencies(
+		const struct FGuid& InGuidToTickBefore,
+		const struct FGuid& InGuidToTickAfter) const;
+
 	/** Go through each constraint in order and evaluate and tick them*/
 	CONSTRAINTS_API void EvaluateAllConstraints() const;
 
@@ -302,9 +303,6 @@ public:
 	/** If static contraint created, we call this, which stores it to the manager/actor*/
 	CONSTRAINTS_API void StaticConstraintCreated(UWorld* InWorld, UTickableConstraint* InConstraint);
 
-	/** If an animated constraint is ready we initialize it (ticks) and add it to the world*/
-	CONSTRAINTS_API void DynamicConstraintCreated(UWorld* InWorld, UTickableConstraint* InConstraint);
-
 private:
 	/** Get the index of the given constraint's name. */
 	int32 GetConstraintIndex(const FName& InConstraintName) const;
diff --git a/Engine/Source/Runtime/Experimental/Animation/Constraints/Public/ConstraintsManager.inl b/Engine/Source/Runtime/Experimental/Animation/Constraints/Public/ConstraintsManager.inl
--- a/Engine/Source/Runtime/Experimental/Animation/Constraints/Public/ConstraintsManager.inl
+++ b/Engine/Source/Runtime/Experimental/Animation/Constraints/Public/ConstraintsManager.inl
@@ -15,7 +15,18 @@ TConstraint* FConstraintsManagerController::AllocateConstraintT(const FName& InB
 		return nullptr;
 	}
 	// unique name (we may want to use another approach here to manage uniqueness)
-	const FName Name = MakeUniqueObjectName(Subsystem, TConstraint::StaticClass(), InBaseName);
+	FName Name = MakeUniqueObjectName(Subsystem, TConstraint::StaticClass(), InBaseName);
+
+	// ensure that the constraint isn't already registered in the ConstraintManager
+	const TArray<TWeakObjectPtr<UTickableConstraint>>& Constraints = Subsystem->GetConstraints(World);
+	const bool bNameFound = Constraints.ContainsByPredicate([Name](const TWeakObjectPtr<UTickableConstraint>& Constraint)
+	{
+		return Constraint.IsValid() && Constraint->GetFName() == Name;
+	});
+	if (bNameFound)
+	{
+		Name = MakeUniqueObjectName(Subsystem, TConstraint::StaticClass(), InBaseName, EUniqueObjectNameOptions::GloballyUnique);
+	}
 
 	TConstraint* NewConstraint = NewObject<TConstraint>(Subsystem, Name, RF_Transactional, bUseDefault ? GetMutableDefault<TConstraint>() : nullptr, bUseDefault);
 	NewConstraint->Modify();
diff --git a/Engine/Source/Runtime/Experimental/Animation/Constraints/Public/TransformConstraint.h b/Engine/Source/Runtime/Experimental/Animation/Constraints/Public/TransformConstraint.h
--- a/Engine/Source/Runtime/Experimental/Animation/Constraints/Public/TransformConstraint.h
+++ b/Engine/Source/Runtime/Experimental/Animation/Constraints/Public/TransformConstraint.h
@@ -489,9 +489,13 @@ struct FTransformConstraintUtils
 	/** Get the last active constraint that has dynamic offset. */
 	static CONSTRAINTS_API int32 GetLastActiveConstraintIndex(const TArray< TWeakObjectPtr<UTickableConstraint> >& InConstraints);
 
-	/** Fills a constraint array that InParentHandle is the parent of. */
+	/**
+	 * Fills a constraint array that InHandle is the parent of.
+	 * If bIncludeTarget is true, we also get the other constraints that act on the same target.
+	 */
 	static CONSTRAINTS_API void GetChildrenConstraints(
 		UWorld* World,
-		const UTransformableHandle* InParentHandle,
-		TArray< TWeakObjectPtr<UTickableConstraint> >& OutConstraints);
+		const UTransformableHandle* InHandle,
+		TArray< TWeakObjectPtr<UTickableConstraint> >& OutConstraints,
+		const bool bIncludeTarget = false);
 };
