prevent GC proxy sync from doing work if its parent GC component removed from the scene

diff --git a/Engine/Source/Runtime/Experimental/Chaos/Private/PBDRigidsSolver.cpp b/Engine/Source/Runtime/Experimental/Chaos/Private/PBDRigidsSolver.cpp
--- a/Engine/Source/Runtime/Experimental/Chaos/Private/PBDRigidsSolver.cpp
+++ b/Engine/Source/Runtime/Experimental/Chaos/Private/PBDRigidsSolver.cpp
@@ -849,6 +849,7 @@ namespace Chaos
 		InProxy->MarkDeleted();
 
 		RemoveDirtyProxy(InProxy);
+		InProxy->OnUnregisteredFromSolver();
 
 		// Particles are removed from acceleration structure in FPhysScene_Chaos::RemoveObject.
 
diff --git a/Engine/Source/Runtime/Experimental/Chaos/Private/PhysicsProxy/GeometryCollectionPhysicsProxy.cpp b/Engine/Source/Runtime/Experimental/Chaos/Private/PhysicsProxy/GeometryCollectionPhysicsProxy.cpp
--- a/Engine/Source/Runtime/Experimental/Chaos/Private/PhysicsProxy/GeometryCollectionPhysicsProxy.cpp
+++ b/Engine/Source/Runtime/Experimental/Chaos/Private/PhysicsProxy/GeometryCollectionPhysicsProxy.cpp
@@ -3086,6 +3086,11 @@ int32 FGeometryCollectionPhysicsProxy::CalculateAndSetLevel(int32 TransformGroup
 	return Levels[TransformGroupIndex];
 }
 
+void FGeometryCollectionPhysicsProxy::OnUnregisteredFromSolver()
+{
+	IsObjectDeleting = true;
+}
+
 void FGeometryCollectionPhysicsProxy::OnRemoveFromSolver(Chaos::FPBDRigidsSolver *RBDSolver)
 {
 	Chaos::FPBDRigidsEvolutionGBF* Evolution = RBDSolver->GetEvolution();
@@ -3143,8 +3148,6 @@ void FGeometryCollectionPhysicsProxy::OnRemoveFromSolver(Chaos::FPBDRigidsSolver
 			}
 		}		
 	}
-
-	IsObjectDeleting = true;
 }
 
 void FGeometryCollectionPhysicsProxy::OnRemoveFromScene()
diff --git a/Engine/Source/Runtime/Experimental/Chaos/Public/PhysicsProxy/GeometryCollectionPhysicsProxy.h b/Engine/Source/Runtime/Experimental/Chaos/Public/PhysicsProxy/GeometryCollectionPhysicsProxy.h
--- a/Engine/Source/Runtime/Experimental/Chaos/Public/PhysicsProxy/GeometryCollectionPhysicsProxy.h
+++ b/Engine/Source/Runtime/Experimental/Chaos/Public/PhysicsProxy/GeometryCollectionPhysicsProxy.h
@@ -233,6 +233,7 @@ public:
 	CHAOS_API void SyncBeforeDestroy();
 	CHAOS_API void OnRemoveFromSolver(Chaos::FPBDRigidsSolver *RBDSolver);
 	CHAOS_API void OnRemoveFromScene();
+	CHAOS_API void OnUnregisteredFromSolver();
 
 	void SetCollisionParticlesPerObjectFraction(float CollisionParticlesPerObjectFractionIn) 
 	{CollisionParticlesPerObjectFraction = CollisionParticlesPerObjectFractionIn;}
