if UGeometryCollectionComponent::SetSimulatePhysics needs to initialize the physics proxy it should go through the route of recreating the physics state to ensure that the physics state gets destroyed when the component is destroyed

diff --git a/Engine/Source/Runtime/Experimental/GeometryCollectionEngine/Private/GeometryCollection/GeometryCollectionComponent.cpp b/Engine/Source/Runtime/Experimental/GeometryCollectionEngine/Private/GeometryCollection/GeometryCollectionComponent.cpp
--- a/Engine/Source/Runtime/Experimental/GeometryCollectionEngine/Private/GeometryCollection/GeometryCollectionComponent.cpp
+++ b/Engine/Source/Runtime/Experimental/GeometryCollectionEngine/Private/GeometryCollection/GeometryCollectionComponent.cpp
@@ -1244,7 +1244,7 @@ void UGeometryCollectionComponent::SetSimulatePhysics(bool bEnabled)
 
 	if (bEnabled && !PhysicsProxy && RestCollection)
 	{
-		RegisterAndInitializePhysicsProxy();
+		RecreatePhysicsState();
 	}
 }
 
