Iris
* Added CSV stats counting the exported references in the write phase
PackageMap
* Added CSV stat counting exported NetGuids every frame.

#rnx

diff --git a/Engine/Source/Runtime/Engine/Private/PackageMapClient.cpp b/Engine/Source/Runtime/Engine/Private/PackageMapClient.cpp
--- a/Engine/Source/Runtime/Engine/Private/PackageMapClient.cpp
+++ b/Engine/Source/Runtime/Engine/Private/PackageMapClient.cpp
@@ -1460,6 +1460,8 @@ void UPackageMapClient::ExportNetGUIDHeader()
 	{
 		UE_LOG(LogNetPackageMap, Warning, TEXT("Attempted to export a NetGUID Bunch with no NetGUIDs!"));
 	}
+
+	CSV_CUSTOM_STAT(PackageMap, NetGuidExports, ExportNetGUIDCount, ECsvCustomStatOp::Accumulate);
 	
 	CurrentExportBunch = NULL;
 	ExportNetGUIDCount = 0;
diff --git a/Engine/Source/Runtime/Experimental/Iris/Core/Private/Iris/ReplicationSystem/ObjectReferenceCache.cpp b/Engine/Source/Runtime/Experimental/Iris/Core/Private/Iris/ReplicationSystem/ObjectReferenceCache.cpp
--- a/Engine/Source/Runtime/Experimental/Iris/Core/Private/Iris/ReplicationSystem/ObjectReferenceCache.cpp
+++ b/Engine/Source/Runtime/Experimental/Iris/Core/Private/Iris/ReplicationSystem/ObjectReferenceCache.cpp
@@ -16,6 +16,7 @@
 #include "Iris/Serialization/NetExportContext.h"
 #include "Iris/Serialization/NetSerializationContext.h"
 #include "Iris/Serialization/ObjectNetSerializer.h"
+#include "Iris/Stats/NetStatsContext.h"
 #include "Misc/CoreMiscDefines.h"
 #include "Misc/PackageName.h"
 #include "Misc/StringBuilder.h"
@@ -1590,7 +1591,7 @@ void FObjectReferenceCache::AddPendingExports(FNetSerializationContext& Context,
 	}
 }
 
-FObjectReferenceCache::EWriteExportsResult FObjectReferenceCache::WritePendingExports(FNetSerializationContext& Context)
+FObjectReferenceCache::EWriteExportsResult FObjectReferenceCache::WritePendingExports(FNetSerializationContext& Context, FInternalNetRefIndex ObjectIndex)
 {	
 	FNetBitStreamWriter& Writer = *Context.GetBitStreamWriter();
 
@@ -1610,6 +1611,10 @@ FObjectReferenceCache::EWriteExportsResult FObjectReferenceCache::WritePendingEx
 	{
 		FForceInlineExportScope ForceInlineExportScope(Context.GetInternalContext());
 
+#if UE_NET_IRIS_CSV_STATS
+		uint32 CountExports = 0;
+#endif
+
 		for (const FNetObjectReference& Reference : ExportsView)
 		{
 			const bool bIsClientAssigned = Reference.PathToken.IsValid();
@@ -1618,15 +1623,21 @@ FObjectReferenceCache::EWriteExportsResult FObjectReferenceCache::WritePendingEx
 			{
 				Writer.WriteBool(true);
 				WriteFullReference(Context, Reference);
+
+#if UE_NET_IRIS_CSV_STATS
+				++CountExports;
+#endif
 			}
 		}
 
+		UE_NET_IRIS_STATS_ADD_COUNT_FOR_OBJECT(Context.GetNetStatsContext(), WriteExports, ObjectIndex, CountExports);
+
 		// Write stop bit
 		Writer.WriteBool(false);
 	}
 
 	// We also write any must be mapped exports
-	WriteMustBeMappedExports(Context, ExportsView);
+	WriteMustBeMappedExports(Context, ObjectIndex, ExportsView);
 
 	// Reset state of pending exports
 	ExportContext->ClearPendingExports();
@@ -1671,7 +1682,7 @@ bool FObjectReferenceCache::ReadExports(FNetSerializationContext& Context, TArra
 	return !Context.HasErrorOrOverflow();
 }
 
-bool FObjectReferenceCache::WriteMustBeMappedExports(FNetSerializationContext& Context, TArrayView<const FNetObjectReference> ExportsView) const
+bool FObjectReferenceCache::WriteMustBeMappedExports(FNetSerializationContext& Context, FInternalNetRefIndex ObjectIndex,TArrayView<const FNetObjectReference> ExportsView) const
 {
 	FNetBitStreamWriter& Writer = *Context.GetBitStreamWriter();
 	FNetExportContext* ExportContext = Context.GetExportContext();
@@ -1680,6 +1691,10 @@ bool FObjectReferenceCache::WriteMustBeMappedExports(FNetSerializationContext& C
 	{
 		UE_NET_TRACE_SCOPE(MustBeMappedExports, Writer, Context.GetTraceCollector(), ENetTraceVerbosity::Verbose);
 
+#if UE_NET_IRIS_CSV_STATS
+		uint32 CountExports = 0;
+#endif
+
 		for (const FNetObjectReference& Reference : ExportsView)
 		{
 			const FNetRefHandle Handle = Reference.GetRefHandle();
@@ -1689,8 +1704,14 @@ bool FObjectReferenceCache::WriteMustBeMappedExports(FNetSerializationContext& C
 				UE_NET_TRACE_OBJECT_SCOPE(Handle, *Context.GetBitStreamWriter(), Context.GetTraceCollector(), ENetTraceVerbosity::Verbose);
 				Writer.WriteBool(true);
 				WriteNetRefHandle(Context, Handle);
+
+#if UE_NET_IRIS_CSV_STATS
+				++CountExports;
+#endif
 			}
 		}
+
+		UE_NET_IRIS_STATS_ADD_COUNT_FOR_OBJECT(Context.GetNetStatsContext(), WriteExports, ObjectIndex, CountExports);
 	}
 
 	// Write stop bit
diff --git a/Engine/Source/Runtime/Experimental/Iris/Core/Private/Iris/ReplicationSystem/ObjectReferenceCache.h b/Engine/Source/Runtime/Experimental/Iris/Core/Private/Iris/ReplicationSystem/ObjectReferenceCache.h
--- a/Engine/Source/Runtime/Experimental/Iris/Core/Private/Iris/ReplicationSystem/ObjectReferenceCache.h
+++ b/Engine/Source/Runtime/Experimental/Iris/Core/Private/Iris/ReplicationSystem/ObjectReferenceCache.h
@@ -20,6 +20,7 @@ namespace UE::Net
 		class FNetRefHandleManager;
 		class FNetExportContext;
 		struct FPendingBatches;
+		typedef uint32 FInternalNetRefIndex;
 	}
 }
 
@@ -111,7 +112,7 @@ public:
 
 	// Exports are expected to be part of the written state, so if the result is a BitStreamOverflow
 	// it is up to the caller to roll back written data and pending exports
-	EWriteExportsResult WritePendingExports(FNetSerializationContext& Context);
+	EWriteExportsResult WritePendingExports(FNetSerializationContext& Context, FInternalNetRefIndex ObjectIndex);
 
 	bool ReadExports(FNetSerializationContext& Context, TArray<FNetRefHandle>* MustBeMappedExports);
 
@@ -205,7 +206,7 @@ private:
 
 	// Must be mapped exports are written for each batch that serializes object references, if async loading is enabled the client
 	// will defer application of data contained in the batch until the must be mapped exports are resolvable.
-	bool WriteMustBeMappedExports(FNetSerializationContext& Context, TArrayView<const FNetObjectReference> ExportsView) const;
+	bool WriteMustBeMappedExports(FNetSerializationContext& Context, FInternalNetRefIndex ObjectIndex, TArrayView<const FNetObjectReference> ExportsView) const;
 	void ReadMustBeMappedExports(FNetSerializationContext& Context, TArray<FNetRefHandle>* MustBeMappedExports);
 
 	void StartAsyncLoadingPackage(FCachedNetObjectReference& Object, FName PackagePath, const FNetRefHandle RefHandle, const bool bWasAlreadyAsyncLoading);
diff --git a/Engine/Source/Runtime/Experimental/Iris/Core/Private/Iris/ReplicationSystem/ReplicationSystem.cpp b/Engine/Source/Runtime/Experimental/Iris/Core/Private/Iris/ReplicationSystem/ReplicationSystem.cpp
--- a/Engine/Source/Runtime/Experimental/Iris/Core/Private/Iris/ReplicationSystem/ReplicationSystem.cpp
+++ b/Engine/Source/Runtime/Experimental/Iris/Core/Private/Iris/ReplicationSystem/ReplicationSystem.cpp
@@ -645,7 +645,7 @@ void UReplicationSystem::PreSendUpdate(const FSendUpdateParams& Params)
 		// $IRIS TODO. There may be some throttling of connections to tick that we should take into account.
 		const FNetBitArrayView& ReplicatingConnections = MakeNetBitArrayView(Impl->ReplicationSystemInternal.GetConnections().GetValidConnections());
 
-#if UE_NET_IRIS_CSV_STATS && CSV_PROFILER
+#if UE_NET_IRIS_CSV_STATS
 		{
 			FNetSendStats& SendStats = InternalSys.GetSendStats();
 			SendStats.Reset();
@@ -795,7 +795,7 @@ void UReplicationSystem::PostSendUpdate()
 			Impl->ReplicationSystemInternal.GetDeltaCompressionBaselineManager().PostSendUpdate(UpdateParams);
 		}
 
-#if UE_NET_IRIS_CSV_STATS && CSV_PROFILER
+#if UE_NET_IRIS_CSV_STATS
 		{
 			FNetSendStats& SendStats = Impl->ReplicationSystemInternal.GetSendStats();
 			SendStats.ReportCsvStats();
diff --git a/Engine/Source/Runtime/Experimental/Iris/Core/Private/Iris/ReplicationSystem/ReplicationWriter.cpp b/Engine/Source/Runtime/Experimental/Iris/Core/Private/Iris/ReplicationSystem/ReplicationWriter.cpp
--- a/Engine/Source/Runtime/Experimental/Iris/Core/Private/Iris/ReplicationSystem/ReplicationWriter.cpp
+++ b/Engine/Source/Runtime/Experimental/Iris/Core/Private/Iris/ReplicationSystem/ReplicationWriter.cpp
@@ -1952,7 +1952,7 @@ FReplicationWriter::EWriteObjectStatus FReplicationWriter::WriteObjectAndSubObje
 	if (bNeedToFilterChangeMask)
 	{
 		ApplyFilterToChangeMask(OutBatchInfo.ParentInternalIndex, InternalIndex, Info, ObjectData.Protocol, ReplicatedObjectStateBuffer, bIsInitialState);
-#if UE_NET_IRIS_CSV_STATS && CSV_PROFILER
+#if UE_NET_IRIS_CSV_STATS
 		if (!bIsInitialState && Info.HasDirtyChangeMask)
 		{
 			WriteContext.Stats.AddNumberOfReplicatedObjectStatesMaskedOut(1U);
@@ -2304,7 +2304,7 @@ FReplicationWriter::EWriteObjectStatus FReplicationWriter::WriteObjectAndSubObje
 		const bool bWroteData = (ParentBatchEntry.bSentState || ParentBatchEntry.bSentAttachments || bSentTearOff || Info.SubObjectPendingDestroy);
 		if (bWroteData || (SubObjectsWrittenBits != 0U))
 		{
-			const FObjectReferenceCache::EWriteExportsResult WriteExportResult = ObjectReferenceCache->WritePendingExports(Context);
+			const FObjectReferenceCache::EWriteExportsResult WriteExportResult = ObjectReferenceCache->WritePendingExports(Context, InternalIndex);
 
 			if (WriteExportResult == FObjectReferenceCache::EWriteExportsResult::BitStreamOverflow)
 			{
@@ -2993,7 +2993,7 @@ int FReplicationWriter::HandleObjectBatchSuccess(const FBatchInfo& BatchInfo, FR
 		}
 	}
 
-#if UE_NET_IRIS_CSV_STATS && CSV_PROFILER
+#if UE_NET_IRIS_CSV_STATS
 	if (bTrackObjectStats)
 	{
 		FNetSendStats& NetStats = WriteContext.Stats;
@@ -3159,7 +3159,7 @@ void FReplicationWriter::EndWrite()
 
 	if (WriteContext.bIsValid)
 	{
-#if UE_NET_IRIS_CSV_STATS && CSV_PROFILER
+#if UE_NET_IRIS_CSV_STATS
 		// Update stats
 		{
 			FNetSendStats& Stats = WriteContext.Stats;
diff --git a/Engine/Source/Runtime/Experimental/Iris/Core/Private/Iris/Stats/NetStats.cpp b/Engine/Source/Runtime/Experimental/Iris/Core/Private/Iris/Stats/NetStats.cpp
--- a/Engine/Source/Runtime/Experimental/Iris/Core/Private/Iris/Stats/NetStats.cpp
+++ b/Engine/Source/Runtime/Experimental/Iris/Core/Private/Iris/Stats/NetStats.cpp
@@ -37,6 +37,7 @@ CSV_DEFINE_CATEGORY(IrisWriteWasteCount, WITH_SERVER_CODE);
 CSV_DEFINE_CATEGORY(IrisWriteWasteKBytes, WITH_SERVER_CODE);
 CSV_DEFINE_CATEGORY(IrisWriteCreationInfoCount, WITH_SERVER_CODE);
 CSV_DEFINE_CATEGORY(IrisWriteCreationInfoKBytes, WITH_SERVER_CODE);
+CSV_DEFINE_CATEGORY(IrisWriteExportsCount, WITH_SERVER_CODE);
 
 void FNetSendStats::Accumulate(const FNetSendStats& Other)
 {
@@ -188,7 +189,7 @@ void FNetTypeStats::Accumulate(FNetStatsContext& Context)
 
 void FNetTypeStats::ReportCSVStats()
 {
-#if UE_NET_IRIS_CSV_STATS && CSV_PROFILER
+#if UE_NET_IRIS_CSV_STATS
 
 	FCsvProfiler* Profiler = FCsvProfiler::Get();
 	bIsEnabled = Profiler->IsCapturing();
@@ -228,6 +229,8 @@ void FNetTypeStats::ReportCSVStats()
 			UE_NET_STATS_RECORD_TYPESTATS_BITS(StatsName, WriteCreationInfo, TypeStatsData);
 			UE_NET_STATS_RECORD_TYPESTATS_COUNT(StatsName, WriteCreationInfo, TypeStatsData);
 
+			UE_NET_STATS_RECORD_TYPESTATS_COUNT(StatsName, WriteExports, TypeStatsData);
+
 			// Reset
 			TypeStatsData.Reset();
 		}
diff --git a/Engine/Source/Runtime/Experimental/Iris/Core/Private/Iris/Stats/NetStatsContext.h b/Engine/Source/Runtime/Experimental/Iris/Core/Private/Iris/Stats/NetStatsContext.h
--- a/Engine/Source/Runtime/Experimental/Iris/Core/Private/Iris/Stats/NetStatsContext.h
+++ b/Engine/Source/Runtime/Experimental/Iris/Core/Private/Iris/Stats/NetStatsContext.h
@@ -22,6 +22,7 @@ struct FNetTypeStatsData
 		Write,
 		WriteWaste,
 		WriteCreationInfo,
+		WriteExports,
 		Count
 	};
 
@@ -70,6 +71,8 @@ public:
 	// Depending on config subobjects report their stats with the root, for those cases we do not bump the count of root objects.
 	// returns the FNetTypeStatsData associated with the protocol used by the object references by the InternalIndex
 	static FNetTypeStatsData& GetTypeStatsDataForObject(FNetStatsContext& Context, FInternalNetRefIndex InternalIndex, uint32& OutUpdateCount);
+	static FNetTypeStatsData& GetTypeStatsDataForObject(FNetStatsContext& Context, FInternalNetRefIndex InternalIndex);
+	static FNetTypeStatsData& GetTypeStatsData(FNetStatsContext& Context, const FNetRefHandleManager::FReplicatedObjectData& ObjectData, bool bTreatAsRoot);
 	
 private:
 	friend class FNetTypeStats;
@@ -98,10 +101,8 @@ private:
 	uint64 StartCycle;
 };
 
-inline FNetTypeStatsData& FNetStatsContext::GetTypeStatsDataForObject(FNetStatsContext& Context, FInternalNetRefIndex InternalIndex, uint32& OutUpdateCount)
+inline FNetTypeStatsData& FNetStatsContext::GetTypeStatsData(FNetStatsContext& Context, const FNetRefHandleManager::FReplicatedObjectData& ObjectData, bool bTreatAsRoot)
 {
-	const FNetRefHandleManager::FReplicatedObjectData& ObjectData = Context.NetRefHandleManager->GetReplicatedObjectDataNoCheck(InternalIndex);		
-	const bool bTreatAsRoot = !Context.bShouldIncludeSubObjectWithRoot || !ObjectData.IsSubObject();
 	int32 TypeStatsIndex = FNetTypeStats::DefaultTypeStatsIndex;
 	if (bTreatAsRoot)
 	{
@@ -112,12 +113,27 @@ inline FNetTypeStatsData& FNetStatsContext::GetTypeStatsDataForObject(FNetStatsC
 		const FNetRefHandleManager::FReplicatedObjectData& RootObjectData = Context.NetRefHandleManager->GetReplicatedObjectDataNoCheck(ObjectData.SubObjectRootIndex);
 		TypeStatsIndex = RootObjectData.Protocol ? RootObjectData.Protocol->TypeStatsIndex : FNetTypeStats::OOBChannelTypeStatsIndex;
 	}
-	
-	OutUpdateCount = bTreatAsRoot ? 1U : 0U;
 
 	return Context.TypeStatsData[TypeStatsIndex];
 }
 
+inline FNetTypeStatsData& FNetStatsContext::GetTypeStatsDataForObject(FNetStatsContext& Context, FInternalNetRefIndex InternalIndex)
+{
+	const FNetRefHandleManager::FReplicatedObjectData& ObjectData = Context.NetRefHandleManager->GetReplicatedObjectDataNoCheck(InternalIndex);
+	const bool bTreatAsRoot = !Context.bShouldIncludeSubObjectWithRoot || !ObjectData.IsSubObject();
+	return GetTypeStatsData(Context, ObjectData, bTreatAsRoot);
+}
+
+
+inline FNetTypeStatsData& FNetStatsContext::GetTypeStatsDataForObject(FNetStatsContext& Context, FInternalNetRefIndex InternalIndex, uint32& OutUpdateCount)
+{
+	const FNetRefHandleManager::FReplicatedObjectData& ObjectData = Context.NetRefHandleManager->GetReplicatedObjectDataNoCheck(InternalIndex);		
+	const bool bTreatAsRoot = !Context.bShouldIncludeSubObjectWithRoot || !ObjectData.IsSubObject();
+	OutUpdateCount = bTreatAsRoot ? 1U : 0U;
+	return GetTypeStatsData(Context, ObjectData, bTreatAsRoot);
+}
+
+
 inline void FNetStatsContext::ResetStats(int32 NumTypeStats)
 { 
 	TypeStatsData.SetNum(NumTypeStats);
@@ -130,7 +146,7 @@ inline void FNetStatsContext::ResetStats(int32 NumTypeStats)
 }
 
 // Wrap usage of stats in macros so we can compile it out
-#if UE_NET_IRIS_CSV_STATS && CSV_PROFILER
+#if UE_NET_IRIS_CSV_STATS
 
 #define UE_NET_IRIS_STATS_TIMER(TimerName, NetStatsContext) UE::Net::Private::FNetStatsTimer TimerName(NetStatsContext);
 
@@ -174,8 +190,7 @@ inline void FNetStatsContext::ResetStats(int32 NumTypeStats)
 	do { \
 		if (NetStatsContext) \
 		{ \
-			uint32 CountIncrement = 0U; \
-			UE::Net::Private::FNetTypeStatsData& StatsData = UE::Net::Private::FNetStatsContext::GetTypeStatsDataForObject(*NetStatsContext, ObjectIndex, CountIncrement); \
+			UE::Net::Private::FNetTypeStatsData& StatsData = UE::Net::Private::FNetStatsContext::GetTypeStatsDataForObject(*NetStatsContext, ObjectIndex); \
 			StatsData.Values[UE::Net::Private::FNetTypeStatsData::EStatsIndex::StatName].Bits += BitCount; \
 		} \
 	} while (0)
@@ -184,13 +199,32 @@ inline void FNetStatsContext::ResetStats(int32 NumTypeStats)
 	do { \
 		if (NetStatsContext) \
 		{ \
-			uint32 CountIncrement = 0U; \
-			UE::Net::Private::FNetTypeStatsData& StatsData = UE::Net::Private::FNetStatsContext::GetTypeStatsDataForObject(*NetStatsContext, ObjectIndex, CountIncrement); \
+			UE::Net::Private::FNetTypeStatsData& StatsData = UE::Net::Private::FNetStatsContext::GetTypeStatsDataForObject(*NetStatsContext, ObjectIndex); \
 			StatsData.Values[UE::Net::Private::FNetTypeStatsData::EStatsIndex::StatName].Bits += BitCount; \
 			StatsData.Values[UE::Net::Private::FNetTypeStatsData::EStatsIndex::StatName##Waste].Bits += BitCount; \
 		} \
 	} while (0)
 
+// Increment stat count by 1
+#define UE_NET_IRIS_STATS_INCREMENT_FOR_OBJECT(NetStatsContext, StatName, ObjectIndex) \
+	do { \
+		if (NetStatsContext) \
+		{ \
+			UE::Net::Private::FNetTypeStatsData& StatsData = UE::Net::Private::FNetStatsContext::GetTypeStatsDataForObject(*NetStatsContext, ObjectIndex); \
+			StatsData.Values[UE::Net::Private::FNetTypeStatsData::EStatsIndex::StatName].Count++; \
+		} \
+	} while (0)
+
+// Increment stat count by any amount
+#define UE_NET_IRIS_STATS_ADD_COUNT_FOR_OBJECT(NetStatsContext, StatName, ObjectIndex, CountToAdd) \
+	do { \
+		if (NetStatsContext) \
+		{ \
+			UE::Net::Private::FNetTypeStatsData& StatsData = UE::Net::Private::FNetStatsContext::GetTypeStatsDataForObject(*NetStatsContext, ObjectIndex); \
+			StatsData.Values[UE::Net::Private::FNetTypeStatsData::EStatsIndex::StatName].Count += CountToAdd; \
+		} \
+	} while (0)
+
 #else
 
 #define UE_NET_IRIS_STATS_TIMER(...)
@@ -199,5 +233,7 @@ inline void FNetStatsContext::ResetStats(int32 NumTypeStats)
 #define UE_NET_IRIS_STATS_ADD_BITS_WRITTEN_AND_COUNT_FOR_OBJECT(...)
 #define UE_NET_IRIS_STATS_ADD_BITS_WRITTEN_FOR_OBJECT(...)
 #define UE_NET_IRIS_STATS_ADD_BITS_WRITTEN_FOR_OBJECT_AS_WASTE(...)
+#define UE_NET_IRIS_STATS_INCREMENT_FOR_OBJECT(...)
+#define UE_NET_IRIS_STATS_ADD_COUNT_FOR_OBJECT(...)
 
 #endif
diff --git a/Engine/Source/Runtime/Experimental/Iris/Core/Public/Iris/IrisConfig.h b/Engine/Source/Runtime/Experimental/Iris/Core/Public/Iris/IrisConfig.h
--- a/Engine/Source/Runtime/Experimental/Iris/Core/Public/Iris/IrisConfig.h
+++ b/Engine/Source/Runtime/Experimental/Iris/Core/Public/Iris/IrisConfig.h
@@ -5,6 +5,7 @@
 #include "HAL/Platform.h"
 
 #include "Net/Core/Connection/NetEnums.h"
+#include "ProfilingDebugging/CsvProfilerConfig.h"
 
 namespace UE::Net
 {
@@ -29,9 +30,9 @@ IRISCORE_API EReplicationSystem GetUseIrisReplicationCmdlineValue();
 #endif
 #endif
 
-/** CSV stats. Please check if CSV_PROFILER is enabled too if executing non-trivial code for stats. */
+/** CSV stats. */
 #ifndef UE_NET_IRIS_CSV_STATS
-#	define UE_NET_IRIS_CSV_STATS 1
+#	define UE_NET_IRIS_CSV_STATS CSV_PROFILER
 #endif
 
 /** Verbose CSV stats is not recommended in shipping builds due to the expected CPU overhead. */
