[Iris] Added client-side stat for active replicated object count

#changelist validated

#tests PIE

diff --git a/Engine/Source/Runtime/Experimental/Iris/Core/Private/Iris/ReplicationSystem/NetRefHandleManager.cpp b/Engine/Source/Runtime/Experimental/Iris/Core/Private/Iris/ReplicationSystem/NetRefHandleManager.cpp
--- a/Engine/Source/Runtime/Experimental/Iris/Core/Private/Iris/ReplicationSystem/NetRefHandleManager.cpp
+++ b/Engine/Source/Runtime/Experimental/Iris/Core/Private/Iris/ReplicationSystem/NetRefHandleManager.cpp
@@ -14,6 +14,9 @@
 #include "UObject/Object.h"
 #include "UObject/UObjectGlobals.h"
 #include "UObject/CoreNetTypes.h"
+#include "ProfilingDebugging/CsvProfiler.h"
+
+CSV_DEFINE_CATEGORY(IrisCommon, true);
 
 namespace UE::Net::Private
 {
@@ -823,6 +826,8 @@ void FNetRefHandleManager::OnPostSendUpdate()
 
 	// From here no-one should access the ScopeFrameData
 	ScopeFrameData.bIsValid = false;
+
+	CSV_CUSTOM_STAT(IrisCommon, ActiveReplicatedObjectCount, (float)ActiveObjectCount, ECsvCustomStatOp::Set);
 }
 
 FString FNetRefHandleManager::PrintObjectFromIndex(FInternalNetRefIndex ObjectIndex) const
