Clamp IAS reads from going past the end of chunks.

#rnx

- Matches the same logic as FFileIoStore::Resolve

diff --git a/Engine/Source/Runtime/Experimental/IoStoreOnDemand/Private/OnDemandIoDispatcherBackend.cpp b/Engine/Source/Runtime/Experimental/IoStoreOnDemand/Private/OnDemandIoDispatcherBackend.cpp
--- a/Engine/Source/Runtime/Experimental/IoStoreOnDemand/Private/OnDemandIoDispatcherBackend.cpp
+++ b/Engine/Source/Runtime/Experimental/IoStoreOnDemand/Private/OnDemandIoDispatcherBackend.cpp
@@ -734,7 +734,7 @@ struct FChunkRequestParams
 		}
 		else
 		{
-			const uint64 RawSize = FMath::Min<uint64>(OffsetLength.GetLength(), ChunkInfo.Entry->RawSize);
+			const uint64 RawSize = FMath::Min<uint64>(OffsetLength.GetLength(), ChunkInfo.Entry->RawSize - OffsetLength.GetOffset());
 
 			ChunkRange = FIoChunkEncoding::GetChunkRange(
 				ChunkInfo.Entry->RawSize,
