Lifted IAS' TOC downloading out from UnrealPak

#rnx

diff --git a/Engine/Source/Developer/IoStoreUtilities/Private/IoStoreUtilities.cpp b/Engine/Source/Developer/IoStoreUtilities/Private/IoStoreUtilities.cpp
--- a/Engine/Source/Developer/IoStoreUtilities/Private/IoStoreUtilities.cpp
+++ b/Engine/Source/Developer/IoStoreUtilities/Private/IoStoreUtilities.cpp
@@ -9722,24 +9722,3 @@ bool UploadIoStoreContainerFiles(const TCHAR* ContainerPathOrWildcard)
 
 	return UploadIoStoreContainerFiles(UploadParams.ConsumeValueOrDie(), ContainerFiles, KeyChain);
 }
-
-bool DownloadIoStoreContainerFiles(const TCHAR* TocPath)
-{
-	using namespace UE::IO::IAS;
-	check(TocPath);
-
-	TIoStatusOr<FIoStoreDownloadParams> Params = FIoStoreDownloadParams::Parse(FCommandLine::Get());
-	if (Params.IsOk() == false)
-	{
-		UE_LOG(LogIoStore, Error, TEXT("Failed to download container file(s), reason '%s'"), *Params.Status().ToString());
-		return false;
-	}
-
-	FIoStatus Status = DownloadContainerFiles(Params.ConsumeValueOrDie(), TocPath);
-	if (Status.IsOk() == false)
-	{
-		UE_LOG(LogIoStore, Error, TEXT("Failed to download container file(s), reason '%s'"), *Status.ToString());
-	}
-
-	return Status.IsOk();
-}
diff --git a/Engine/Source/Developer/IoStoreUtilities/Public/IoStoreUtilities.h b/Engine/Source/Developer/IoStoreUtilities/Public/IoStoreUtilities.h
--- a/Engine/Source/Developer/IoStoreUtilities/Public/IoStoreUtilities.h
+++ b/Engine/Source/Developer/IoStoreUtilities/Public/IoStoreUtilities.h
@@ -60,6 +60,4 @@ UE_API bool SignIoStoreContainer(const TCHAR* InContainerFilename, const FRSAKey
 
 UE_API bool UploadIoStoreContainerFiles(const TCHAR* ContainerPathOrWildcard); 
 
-UE_API bool DownloadIoStoreContainerFiles(const TCHAR* TocPath);
-
 #undef UE_API
diff --git a/Engine/Source/Developer/PakFileUtilities/Private/PakFileUtilities.cpp b/Engine/Source/Developer/PakFileUtilities/Private/PakFileUtilities.cpp
--- a/Engine/Source/Developer/PakFileUtilities/Private/PakFileUtilities.cpp
+++ b/Engine/Source/Developer/PakFileUtilities/Private/PakFileUtilities.cpp
@@ -5473,11 +5473,6 @@ bool ExecuteUnrealPak(const TCHAR* CmdLine)
 			{
 				return UploadIoStoreContainerFiles(*IoStoreArg) == 0;
 			}
-
-			if (FParse::Value(CmdLine, TEXT("-Download="), IoStoreArg))
-			{
-				return DownloadIoStoreContainerFiles(*IoStoreArg) == 0;
-			}
 		}
 	}
 
diff --git a/Engine/Source/Runtime/Experimental/IoStoreOnDemand/Private/Tool/Download.cpp b/Engine/Source/Runtime/Experimental/IoStoreOnDemand/Private/Tool/Download.cpp
new file mode 100644
--- /dev/null
+++ b/Engine/Source/Runtime/Experimental/IoStoreOnDemand/Private/Tool/Download.cpp
@@ -0,0 +1,70 @@
+// Copyright Epic Games, Inc. All Rights Reserved.
+
+#if defined(UE_WITH_IAS_TOOL)
+
+#include "Command.h"
+
+#include <IO/IoStoreOnDemand.h>
+
+namespace UE::IO::IAS::Tool {
+
+////////////////////////////////////////////////////////////////////////////////
+static FIoStoreDownloadParams BuildDownloadParams(const FContext& Context)
+{
+	FIoStoreDownloadParams Ret;
+
+	Ret.Directory				= Context.Get(TEXT("-Directory"),				Ret.Directory);
+	Ret.ServiceUrl				= Context.Get(TEXT("-ServiceUrl"),				Ret.ServiceUrl);
+	Ret.Bucket					= Context.Get(TEXT("-Bucket"),					Ret.Bucket);
+	Ret.Region					= Context.Get(TEXT("-Region"),					Ret.Region); 
+	Ret.AccessKey				= Context.Get(TEXT("-AccessKey"),				Ret.AccessKey);
+	Ret.SecretKey				= Context.Get(TEXT("-SecretKey"),				Ret.SecretKey);
+	Ret.SessionToken			= Context.Get(TEXT("-SessionToken"),			Ret.SessionToken);
+	Ret.CredentialsFile			= Context.Get(TEXT("-CredentialsFile"),			Ret.CredentialsFile);
+	Ret.CredentialsFileKeyName	= Context.Get(TEXT("-CredentialsFileKeyName"),	Ret.CredentialsFileKeyName);
+	Ret.MaxConcurrentDownloads	= Context.Get(TEXT("-MaxConcurrentDownloads"),	Ret.MaxConcurrentDownloads);
+
+	if (FIoStatus Validation = Ret.Validate(); !Validation.IsOk())
+	{
+		FString Reason = Validation.ToString();
+		Context.Abort(*Reason);
+	}
+
+	return Ret;
+}
+
+////////////////////////////////////////////////////////////////////////////////
+static int32 DownloadCommandEntry(const Tool::FContext& Context)
+{
+	FStringView TocPath = Context.Get<FStringView>(TEXT("TocPath"));
+
+	FIoStoreDownloadParams Params = BuildDownloadParams(Context);
+
+	FIoStatus Status = DownloadContainerFiles(Params, TocPath.GetData());
+	if (!Status.IsOk())
+	{
+		FString Reason = Status.ToString();
+		Context.Abort(*Reason);
+	}
+
+	return 0;
+}
+
+////////////////////////////////////////////////////////////////////////////////
+extern FArgumentSet S3Arguments;
+
+static FCommand DownloadCommand(
+	DownloadCommandEntry,
+	TEXT("Download"),
+	TEXT(""),
+	{
+		TArgument<FStringView>(TEXT("TocPath"),				TEXT("Bucket-relative path of the TOC to download")),
+		TArgument<FStringView>(TEXT("-Directory"),			TEXT("Output directory")),
+		TArgument<int32>(TEXT("-MaxConcurrentDownloads"),	TEXT("Number of downloads that happen all at once")),
+		S3Arguments,
+	}
+);
+
+} // namespace UE::IO::IAS::Tool
+
+#endif // UE_WITH_IAS_TOOL
diff --git a/Engine/Source/Runtime/Experimental/IoStoreOnDemand/Public/IO/IoStoreOnDemand.h b/Engine/Source/Runtime/Experimental/IoStoreOnDemand/Public/IO/IoStoreOnDemand.h
--- a/Engine/Source/Runtime/Experimental/IoStoreOnDemand/Public/IO/IoStoreOnDemand.h
+++ b/Engine/Source/Runtime/Experimental/IoStoreOnDemand/Public/IO/IoStoreOnDemand.h
@@ -235,11 +235,11 @@ struct FIoStoreDownloadParams
 	FString CredentialsFileKeyName;
 	int32 MaxConcurrentDownloads = 16;
 	
-	UE_API static TIoStatusOr<FIoStoreDownloadParams> Parse(const TCHAR* CommandLine);
+	static TIoStatusOr<FIoStoreDownloadParams> Parse(const TCHAR* CommandLine);
 	FIoStatus Validate() const;
 };
 
-UE_API FIoStatus DownloadContainerFiles(const FIoStoreDownloadParams& DownloadParams, const FString& TocPath);
+FIoStatus DownloadContainerFiles(const FIoStoreDownloadParams& DownloadParams, const FString& TocPath);
 
 /**
  * Parameters for listing uploaded TOC file(s) from an S3 compatible endpoint.
