Upload command for IasTool

#rnx

diff --git a/Engine/Source/Runtime/Experimental/IoStoreOnDemand/Private/Tool/Upload.cpp b/Engine/Source/Runtime/Experimental/IoStoreOnDemand/Private/Tool/Upload.cpp
new file mode 100644
--- /dev/null
+++ b/Engine/Source/Runtime/Experimental/IoStoreOnDemand/Private/Tool/Upload.cpp
@@ -0,0 +1,155 @@
+// Copyright Epic Games, Inc. All Rights Reserved.
+
+#if defined(UE_WITH_IAS_TOOL)
+
+#include "Command.h"
+
+#include <HAL/FileManager.h>
+#include <IO/IoStoreOnDemand.h>
+#include <Misc/KeyChainUtilities.h>
+#include <Misc/Paths.h>
+#include <Misc/PathViews.h>
+
+namespace UE::IO::IAS::Tool {
+
+////////////////////////////////////////////////////////////////////////////////
+static FKeyChain LoadKeys(const FContext& Context)
+{
+	FKeyChain Ret;
+
+	auto Path = Context.Get<FStringView>(TEXT("-CryptoKeys"));
+	if (Path.IsEmpty())
+	{
+		return Ret;
+	}
+
+	// UE_LOG(LogIoStore, Display, TEXT("Parsing crypto keys from a crypto key cache file '%s'"), *CryptoKeysCacheFilename);
+	KeyChainUtilities::LoadKeyChainFromFile(Path.GetData(), Ret);
+	return Ret;
+}
+
+////////////////////////////////////////////////////////////////////////////////
+static FIoStoreUploadParams BuildUploadParams(const FContext& Context)
+{
+	FIoStoreUploadParams Ret;
+
+	Ret.ServiceUrl				= Context.Get<FStringView>(TEXT("-ServiceUrl"),				Ret.ServiceUrl);
+	Ret.Bucket					= Context.Get<FStringView>(TEXT("-Bucket"),					Ret.Bucket);
+	Ret.BucketPrefix			= Context.Get<FStringView>(TEXT("-BucketPrefix"),			Ret.BucketPrefix);
+	Ret.Region					= Context.Get<FStringView>(TEXT("-Region"),					Ret.Region);
+	Ret.AccessKey				= Context.Get<FStringView>(TEXT("-AccessKey"),				Ret.AccessKey);
+	Ret.SecretKey				= Context.Get<FStringView>(TEXT("-SecretKey"),				Ret.SecretKey);
+	Ret.SessionToken			= Context.Get<FStringView>(TEXT("-SessionToken"),			Ret.SessionToken);
+	Ret.CredentialsFile			= Context.Get<FStringView>(TEXT("-CredentialsFile"),		Ret.CredentialsFile);
+	Ret.CredentialsFileKeyName	= Context.Get<FStringView>(TEXT("-CredentialsFileKeyName"),	Ret.CredentialsFileKeyName);
+	Ret.BuildVersion			= Context.Get<FStringView>(TEXT("-BuildVersion"),			Ret.BuildVersion);
+	Ret.TargetPlatform			= Context.Get<FStringView>(TEXT("-TargetPlatform"),			Ret.TargetPlatform);
+	Ret.bWriteTocToDisk			= Context.Get<bool>(TEXT("-WriteTocToDisk"),				Ret.bWriteTocToDisk);
+	Ret.MaxConcurrentUploads	= Context.Get<int32>(TEXT("-MaxConcurrentUploads"),			Ret.MaxConcurrentUploads);
+
+	Ret.bDeleteContainerFiles	= !Context.Get<bool>(TEXT("-KeepContainerFiles"),	!Ret.bDeleteContainerFiles);
+	Ret.bDeletePakFiles			= !Context.Get<bool>(TEXT("-KeepPakFiles"),			!Ret.bDeletePakFiles);
+
+	if (Ret.bWriteTocToDisk)
+	{
+		FStringView Path = Context.Get<FStringView>(TEXT("-ConfigFilePath"));
+		Ret.TocOutputDir = FPathViews::GetPath(Path);
+
+		Ret.DistributionUrl = Context.Get<FStringView>(TEXT("-DistributionUrl"),	Ret.DistributionUrl);
+		Ret.FallbackUrl = Context.Get<FStringView>(TEXT("-FallbackUrl"),			Ret.FallbackUrl);
+	}
+
+	if (FIoStatus Validation = Ret.Validate(); !Validation.IsOk())
+	{
+		FString Reason = Validation.ToString();
+		Context.Abort(*Reason);
+	}
+
+	return Ret;
+}
+
+////////////////////////////////////////////////////////////////////////////////
+static TArray<FString> GlobContainers(const FContext& Context)
+{
+	const TCHAR* GlobPattern = Context.Get<FStringView>(TEXT("ContainerGlob")).GetData();
+
+	TArray<FString> Ret;
+
+	if (IFileManager::Get().FileExists(GlobPattern))
+	{
+		Ret.Add(GlobPattern);
+	}
+	else if (IFileManager::Get().DirectoryExists(GlobPattern))
+	{
+		FString Directory = GlobPattern;
+		FPaths::NormalizeDirectoryName(Directory);
+
+		TArray<FString> FoundContainerFiles;
+		IFileManager::Get().FindFiles(FoundContainerFiles, *(Directory / TEXT("*.utoc")), true, false);
+
+		for (const FString& Filename : FoundContainerFiles)
+		{
+			Ret.Emplace(Directory / Filename);
+		}
+	}
+	else
+	{
+		FString Directory = FPaths::GetPath(GlobPattern);
+		FPaths::NormalizeDirectoryName(Directory);
+
+		TArray<FString> FoundContainerFiles;
+		IFileManager::Get().FindFiles(FoundContainerFiles, GlobPattern, true, false);
+
+		for (const FString& Filename : FoundContainerFiles)
+		{
+			Ret.Emplace(Directory / Filename);
+		}
+	}
+
+	return Ret;
+}
+
+////////////////////////////////////////////////////////////////////////////////
+static int32 UploadCommandEntry(const FContext& Context)
+{
+	TArray<FString> Containers = GlobContainers(Context);
+	FKeyChain KeyChain = LoadKeys(Context);
+	FIoStoreUploadParams Params = BuildUploadParams(Context);
+
+	TIoStatusOr<FIoStoreUploadResult> Result = UploadContainerFiles(Params, Containers, KeyChain);
+	if (!Result.IsOk())
+	{
+		FString Reason = Result.Status().ToString();
+		Context.Abort(*Reason);
+	}
+
+	return 0;
+}
+
+////////////////////////////////////////////////////////////////////////////////
+extern FArgumentSet S3Arguments;
+
+static FCommand UploadCommand(
+	UploadCommandEntry,
+	TEXT("upload"),
+	TEXT(""),
+	{
+		TArgument<FStringView>(TEXT("ContainerGlob"),	TEXT("Path globbed to discover input containers")),
+		TArgument<FStringView>(TEXT("-CryptoKeys"),		TEXT("JSON-format keyring for input containers")),
+		TArgument<FStringView>(TEXT("-BuildVersion"),	TEXT("Optional build version to embed it TOC")),
+		TArgument<FStringView>(TEXT("-TargetPlatform"),	TEXT("If given, embedded in the output TOC")),
+		TArgument<FStringView>(TEXT("-ConfigFilePath"),	TEXT("Path to the config file to write runtime parameters to")),
+		TArgument<FStringView>(TEXT("-DistributionUrl"),TEXT("URL for IAS to use at runtime. Output to config file")),
+		TArgument<FStringView>(TEXT("-FallbackUrl"),	TEXT("Alternative fallack for DistributionUrl")),
+		TArgument<FStringView>(TEXT("-BucketPrefix"),	TEXT("Path to prefix to bucket objects")),
+		TArgument<bool>(TEXT("-KeepContainerFiles"),	TEXT("Do not delete container files after upload")),
+		TArgument<bool>(TEXT("-KeepPakFiles"),			TEXT("Do not delete the springboard pak files")),
+		TArgument<bool>(TEXT("-WriteTocToDisk"),		TEXT("Output the TOC to disk as well as uploading")),
+		TArgument<int32>(TEXT("-MaxConcurrentUploads"),	TEXT("Number of simultaneous uploads")),
+		S3Arguments,
+	}
+);
+
+} // namespace UE::IO::IAS::Tool
+
+#endif // UE_WITH_IAS_TOOL
