Iris
* Added net.Iris.PrintReplicatedObjects and net.Iris.PrintRelevantObjects

#rnx

diff --git a/Engine/Source/Runtime/Experimental/Iris/Core/Private/Iris/ReplicationSystem/ObjectReplicationBridgeDebugging.cpp b/Engine/Source/Runtime/Experimental/Iris/Core/Private/Iris/ReplicationSystem/ObjectReplicationBridgeDebugging.cpp
--- a/Engine/Source/Runtime/Experimental/Iris/Core/Private/Iris/ReplicationSystem/ObjectReplicationBridgeDebugging.cpp
+++ b/Engine/Source/Runtime/Experimental/Iris/Core/Private/Iris/ReplicationSystem/ObjectReplicationBridgeDebugging.cpp
@@ -13,6 +13,7 @@
 #include "Iris/ReplicationSystem/ReplicationOperations.h"
 
 #include "Net/Core/NetBitArrayPrinter.h"
+#include "Net/Core/Trace/NetDebugName.h"
 
 namespace UE::Net::Private::ObjectBridgeDebugging
 {
@@ -73,12 +74,54 @@ UReplicationSystem* FindReplicationSystem(const TArray<FString>& Args)
 }
 
 
+FString PrintNetObject(FNetRefHandleManager* NetRefHandleManager, FInternalNetRefIndex ObjectIndex)
+{
+	using namespace UE::Net;
+	using namespace UE::Net::Private;
+
+	FNetRefHandle NetRefHandle = NetRefHandleManager->GetNetRefHandleFromInternalIndex(ObjectIndex);
+	const FNetRefHandleManager::FReplicatedObjectData& NetObjectData = NetRefHandleManager->GetReplicatedObjectDataNoCheck(ObjectIndex);
+	UObject* ObjectPtr = NetRefHandleManager->GetReplicatedObjectInstance(ObjectIndex);
+
+	return FString::Printf(TEXT("%s %s (InternalIndex: %u) (%s)"), 
+		(NetObjectData.SubObjectRootIndex == FNetRefHandleManager::InvalidInternalIndex) ? TEXT("RootObject"):TEXT("SubObject"),
+		*GetNameSafe(ObjectPtr), ObjectIndex, *NetRefHandle.ToString()
+	);
+}
+
+void LogRootObjectListWithSubObjects(FNetRefHandleManager* NetRefHandleManager, const FNetBitArrayView RootObjectList, TTuple<uint32, uint32>& OutStats, TFunction<FString(FInternalNetRefIndex ObjectIndex)> OptionalObjectPrint)
+{
+	using namespace UE::Net;
+	using namespace UE::Net::Private;
+
+	auto PrintRootObjectInfo = [&](FInternalNetRefIndex RootObjectIndex)
+	{
+		UE_LOG(LogIrisBridge, Display, TEXT("%s %s"), *PrintNetObject(NetRefHandleManager, RootObjectIndex), OptionalObjectPrint ? *OptionalObjectPrint(RootObjectIndex) : TEXT(""));
+
+		// TotalRootObjects
+		OutStats.Get<0>()++;
+
+		TArrayView<const FInternalNetRefIndex> SubObjects = NetRefHandleManager->GetSubObjects(RootObjectIndex);
+		for (FInternalNetRefIndex SubObjectIndex : SubObjects)
+		{
+			UE_LOG(LogIrisBridge, Display, TEXT("\t%s %s"), *PrintNetObject(NetRefHandleManager, SubObjectIndex), OptionalObjectPrint ? *OptionalObjectPrint(RootObjectIndex) : TEXT(""));
+
+			// TotalSubObjects
+			OutStats.Get<1>()++;
+		}
+	};
+
+	RootObjectList.ForAllSetBits(PrintRootObjectInfo);
+}
+
+
 } // end namespace UE::Net::Private::ObjectBridgeDebugging
 
 // --------------------------------------------------------------------------------------------------------------------------------------------
 // Debug commands
 // --------------------------------------------------------------------------------------------------------------------------------------------
 
+//-----------------------------------------------
 FAutoConsoleCommand ObjectBridgePrintDynamicFilter(TEXT("Net.Iris.PrintDynamicFilterClassConfig"), TEXT("Prints the dynamic filter configured to be assigned to specific classes."), FConsoleCommandWithArgsDelegate::CreateLambda([](const TArray< FString >& Args)
 {
 	using namespace UE::Net::Private::ObjectBridgeDebugging;
@@ -118,4 +161,97 @@ void UObjectReplicationBridge::PrintDynamicFilterClassConfig() const
 			UE_LOG(LogIrisFilterConfig, Display, TEXT("\t%s -> %s"), *ClassName.ToString(), *RepSystem->GetFilterName(FilterHandle).ToString());
 		}
 	}
+}
+
+//-----------------------------------------------
+FAutoConsoleCommand ObjectBridgePrintReplicatedObjects(
+	TEXT("Net.Iris.PrintReplicatedObjects"), 
+	TEXT("Prints the list of replicated objects registered for replication in Iris"), 
+	FConsoleCommandWithArgsDelegate::CreateLambda([](const TArray< FString >& Args)
+{
+	using namespace UE::Net::Private::ObjectBridgeDebugging;
+
+	UReplicationSystem* RepSystem = FindReplicationSystem(Args);
+	if (RepSystem)
+	{
+		if (UObjectReplicationBridge* ObjectBridge = CastChecked<UObjectReplicationBridge>(RepSystem->GetReplicationBridge()))
+		{
+			ObjectBridge->PrintReplicatedObjects();
+		}
+	}
+}));
+
+void UObjectReplicationBridge::PrintReplicatedObjects() const
+{
+	using namespace UE::Net;
+	using namespace UE::Net::Private;
+	using namespace UE::Net::Private::ObjectBridgeDebugging;
+
+	UE_LOG(LogIrisBridge, Display, TEXT("------------  Start Printing ALL Replicated Objects ------------"));
+	UE_LOG(LogIrisBridge, Display, TEXT(""));
+
+	uint32 TotalRootObjects = 0;
+	uint32 TotalSubObjects = 0;
+
+	FNetBitArray RootObjects;
+	RootObjects.Init(NetRefHandleManager->GetMaxActiveObjectCount());
+	MakeNetBitArrayView(RootObjects).Set(NetRefHandleManager->GetGlobalScopableInternalIndices(), FNetBitArrayView::AndNotOp, NetRefHandleManager->GetSubObjectInternalIndicesView());
+
+	auto PrintClassOrProtocol = [&](FInternalNetRefIndex ObjectIndex) -> FString
+	{
+		const FNetRefHandleManager::FReplicatedObjectData& NetObjectData = NetRefHandleManager->GetReplicatedObjectDataNoCheck(ObjectIndex);
+		UObject* ObjectPtr = NetRefHandleManager->GetReplicatedObjectInstance(ObjectIndex);
+
+		return FString::Printf(TEXT("Class %s"), ObjectPtr ? *(ObjectPtr->GetClass()->GetName()) : NetObjectData.Protocol->DebugName->Name);
+	};
+	
+	TTuple<uint32, uint32> Stats(0,0);
+	LogRootObjectListWithSubObjects(NetRefHandleManager, MakeNetBitArrayView(RootObjects), Stats, PrintClassOrProtocol);
+
+	UE_LOG(LogIrisBridge, Display, TEXT(""));
+	UE_LOG(LogIrisBridge, Display, TEXT("Printed %u root objects and %u sub objects"), Stats.Get<0>(), Stats.Get<1>());
+	UE_LOG(LogIrisBridge, Display, TEXT("------------  Stop Printing ALL Replicated Objects ------------"));
+}
+
+//-----------------------------------------------
+FAutoConsoleCommand ObjectBridgePrintRelevantObjects(
+	TEXT("Net.Iris.PrintRelevantObjects"), 
+	TEXT("Prints the list of netobjects currently relevant to any connection"), 
+	FConsoleCommandWithArgsDelegate::CreateLambda([](const TArray< FString >& Args)
+{
+	using namespace UE::Net::Private::ObjectBridgeDebugging;
+
+	UReplicationSystem* RepSystem = FindReplicationSystem(Args);
+	if (RepSystem)
+	{
+		if (UObjectReplicationBridge* ObjectBridge = CastChecked<UObjectReplicationBridge>(RepSystem->GetReplicationBridge()))
+		{
+			ObjectBridge->PrintRelevantObjects();
+		}
+	}
+}));
+
+void UObjectReplicationBridge::PrintRelevantObjects() const
+{
+	using namespace UE::Net;
+	using namespace UE::Net::Private;
+	using namespace UE::Net::Private::ObjectBridgeDebugging;
+
+	UE_LOG(LogIrisBridge, Display, TEXT("------------  Start Printing Relevant Objects ------------"));
+	UE_LOG(LogIrisBridge, Display, TEXT(""));
+
+	uint32 TotalRootObjects = 0;
+	uint32 TotalSubObjects = 0;
+
+	FNetBitArray RootObjects;
+	RootObjects.Init(NetRefHandleManager->GetMaxActiveObjectCount());
+	MakeNetBitArrayView(RootObjects).Set(NetRefHandleManager->GetRelevantObjectsInternalIndices(), FNetBitArrayView::AndNotOp, NetRefHandleManager->GetSubObjectInternalIndicesView());
+
+	TTuple<uint32, uint32> Stats(0, 0);
+	LogRootObjectListWithSubObjects(NetRefHandleManager, MakeNetBitArrayView(RootObjects), Stats, nullptr);
+	
+
+	UE_LOG(LogIrisBridge, Display, TEXT(""));
+	UE_LOG(LogIrisBridge, Display, TEXT("Printed %u root objects and %u sub objects"), Stats.Get<0>(), Stats.Get<1>());
+	UE_LOG(LogIrisBridge, Display, TEXT("------------  Stop Printing Relevant Objects ------------"));
 }
\ No newline at end of file
diff --git a/Engine/Source/Runtime/Experimental/Iris/Core/Public/Iris/ReplicationSystem/ObjectReplicationBridge.h b/Engine/Source/Runtime/Experimental/Iris/Core/Public/Iris/ReplicationSystem/ObjectReplicationBridge.h
--- a/Engine/Source/Runtime/Experimental/Iris/Core/Public/Iris/ReplicationSystem/ObjectReplicationBridge.h
+++ b/Engine/Source/Runtime/Experimental/Iris/Core/Public/Iris/ReplicationSystem/ObjectReplicationBridge.h
@@ -183,6 +183,8 @@ public:
 
 	// Debug functions exposed via console commands
 	void PrintDynamicFilterClassConfig() const;
+	void PrintReplicatedObjects() const;
+	void PrintRelevantObjects() const;
 
 protected:
 	IRISCORE_API virtual ~UObjectReplicationBridge();
