Dataflow : Rendering colleciton issue.

diff --git a/Engine/Source/Runtime/Experimental/Dataflow/Engine/Private/Dataflow/DataflowObject.cpp b/Engine/Source/Runtime/Experimental/Dataflow/Engine/Private/Dataflow/DataflowObject.cpp
--- a/Engine/Source/Runtime/Experimental/Dataflow/Engine/Private/Dataflow/DataflowObject.cpp
+++ b/Engine/Source/Runtime/Experimental/Dataflow/Engine/Private/Dataflow/DataflowObject.cpp
@@ -47,9 +47,15 @@ void UDataflow::PostEditCallback()
 
 void UDataflow::AddReferencedObjects(UObject* InThis, FReferenceCollector& Collector)
 {
-	Super::AddReferencedObjects(InThis, Collector);
 	UDataflow* const This = CastChecked<UDataflow>(InThis);
+
+	for(TObjectPtr<const UDataflowEdNode> Target : This->GetRenderTargets())
+	{
+		Collector.AddReferencedObject(Target);
+	}
+
 	This->Dataflow->AddReferencedObjects(Collector);
+	Super::AddReferencedObjects(InThis, Collector);
 }
 
 #if WITH_EDITOR
@@ -98,14 +104,14 @@ void UDataflow::PostLoad()
 	UObject::PostLoad();
 }
 
-void UDataflow::AddRenderTarget(UDataflowEdNode* InNode)
+void UDataflow::AddRenderTarget(TObjectPtr<UDataflowEdNode> InNode)
 {
 	LastModifiedRenderTarget = Dataflow::FTimestamp::Current();
 	InNode->bRenderInAssetEditor = true;
 	RenderTargets.AddUnique(InNode);
 }
 
-void UDataflow::RemoveRenderTarget(UDataflowEdNode* InNode)
+void UDataflow::RemoveRenderTarget(TObjectPtr<UDataflowEdNode> InNode)
 {
 	LastModifiedRenderTarget = Dataflow::FTimestamp::Current();
 	InNode->bRenderInAssetEditor = false;
diff --git a/Engine/Source/Runtime/Experimental/Dataflow/Engine/Public/Dataflow/DataflowObject.h b/Engine/Source/Runtime/Experimental/Dataflow/Engine/Public/Dataflow/DataflowObject.h
--- a/Engine/Source/Runtime/Experimental/Dataflow/Engine/Public/Dataflow/DataflowObject.h
+++ b/Engine/Source/Runtime/Experimental/Dataflow/Engine/Public/Dataflow/DataflowObject.h
@@ -60,7 +60,7 @@ class UDataflow : public UEdGraph
 	GENERATED_UCLASS_BODY()
 
 	Dataflow::FTimestamp LastModifiedRenderTarget = Dataflow::FTimestamp::Invalid; 
-	TArray<const UDataflowEdNode*> RenderTargets; // Not Serialized
+	TArray< TObjectPtr<const UDataflowEdNode> > RenderTargets; // Not Serialized
 	TSharedPtr<Dataflow::FGraph, ESPMode::ThreadSafe> Dataflow;
 	DATAFLOWENGINE_API void PostEditCallback();
 
@@ -113,9 +113,9 @@ public:
 	//
 	// Render Targets
 	//
-	DATAFLOWENGINE_API void AddRenderTarget(UDataflowEdNode*);
-	DATAFLOWENGINE_API void RemoveRenderTarget(UDataflowEdNode*);
-	const TArray<const UDataflowEdNode*>& GetRenderTargets() const { return RenderTargets; }
+	DATAFLOWENGINE_API void AddRenderTarget(TObjectPtr<UDataflowEdNode>);
+	DATAFLOWENGINE_API void RemoveRenderTarget(TObjectPtr<UDataflowEdNode>);
+	const TArray< TObjectPtr<const UDataflowEdNode> >& GetRenderTargets() const { return RenderTargets; }
 	const Dataflow::FTimestamp& GetRenderingTimestamp() const { return LastModifiedRenderTarget; }
 
 };
