Store velocities and rotations in single precision in Physics particles
This is saving 40 bytes per particle

diff --git a/Engine/Plugins/Experimental/Buoyancy/Source/Runtime/Private/BuoyancyAlgorithms.cpp b/Engine/Plugins/Experimental/Buoyancy/Source/Runtime/Private/BuoyancyAlgorithms.cpp
--- a/Engine/Plugins/Experimental/Buoyancy/Source/Runtime/Private/BuoyancyAlgorithms.cpp
+++ b/Engine/Plugins/Experimental/Buoyancy/Source/Runtime/Private/BuoyancyAlgorithms.cpp
@@ -571,7 +571,7 @@ namespace BuoyancyAlgorithms
 
 		// Get the velocities of the submerged portion relative to the water - We want the
 		// drag force to bring these values to zero
-		const FVec3 SubmergedV = RigidParticle->V() + FVec3::CrossProduct(RigidParticle->W(), CoMDiff);
+		const FVec3 SubmergedV = RigidParticle->GetV() + FVec3::CrossProduct(RigidParticle->GetW(), CoMDiff);
 		const FVec3 RelativeV = SubmergedV - WaterVel;
 
 		// Compute water drag force
@@ -584,7 +584,7 @@ namespace BuoyancyAlgorithms
 
 		// Account for water drag in deltas
 		OutDeltaV = (DragFactor * OutDeltaV) + (DragFactor - 1.f) * RelativeV;
-		OutDeltaW = (DragFactor * OutDeltaW) + (DragFactor - 1.f) * RigidParticle->W();
+		OutDeltaW = (DragFactor * OutDeltaW) + (DragFactor - 1.f) * RigidParticle->GetW();
 
 		//
 		return true;
diff --git a/Engine/Plugins/Experimental/Buoyancy/Source/Runtime/Private/BuoyancySubsystem.cpp b/Engine/Plugins/Experimental/Buoyancy/Source/Runtime/Private/BuoyancySubsystem.cpp
--- a/Engine/Plugins/Experimental/Buoyancy/Source/Runtime/Private/BuoyancySubsystem.cpp
+++ b/Engine/Plugins/Experimental/Buoyancy/Source/Runtime/Private/BuoyancySubsystem.cpp
@@ -907,7 +907,7 @@ void FBuoyancySubsystemSimCallback::ProcessInteraction(Chaos::FPBDRigidsEvolutio
 				// Proceed only if this CoM is moving fast enough to generate events
 				const FBuoyancySubmersion& Submersion = Submersions[RigidParticleIndex];
 				const Chaos::FVec3 CoMDiff = Submersion.CoM - Interaction.RigidParticle->XCom();
-				const Chaos::FVec3 CoMVel = Interaction.RigidParticle->V() + Chaos::FVec3::CrossProduct(Interaction.RigidParticle->W(), CoMDiff);
+				const Chaos::FVec3 CoMVel = Interaction.RigidParticle->GetV() + Chaos::FVec3::CrossProduct(Interaction.RigidParticle->GetW(), CoMDiff);
 				const float CoMVelSq = Chaos::FVec3::DotProduct(CoMVel, CoMVel);
 				const float MinVel = BuoyancySettings->MinVelocityForSurfaceTouchCallback;
 				const float MinVelSq = MinVel * MinVel;
@@ -962,8 +962,8 @@ void FBuoyancySubsystemSimCallback::ApplyBuoyantForces(Chaos::FPBDRigidsEvolutio
 			DeltaW = DeltaW.GetClampedToSize(0.f, BuoyancySettings->MaxDeltaW);
 
 			// Apply the deltas
-			Submersion.Particle->SetV(Submersion.Particle->V() + DeltaV);
-			Submersion.Particle->SetW(Submersion.Particle->W() + DeltaW);
+			Submersion.Particle->SetV(Submersion.Particle->GetV() + DeltaV);
+			Submersion.Particle->SetW(Submersion.Particle->GetW() + DeltaW);
 
 			// Wake up the body??
 			if (BuoyancySettings->bKeepAwake)
diff --git a/Engine/Plugins/Experimental/ChaosCaching/Source/ChaosCaching/Private/Chaos/Adapters/GeometryCollectionComponentCacheAdapter.cpp b/Engine/Plugins/Experimental/ChaosCaching/Source/ChaosCaching/Private/Chaos/Adapters/GeometryCollectionComponentCacheAdapter.cpp
--- a/Engine/Plugins/Experimental/ChaosCaching/Source/ChaosCaching/Private/Chaos/Adapters/GeometryCollectionComponentCacheAdapter.cpp
+++ b/Engine/Plugins/Experimental/ChaosCaching/Source/ChaosCaching/Private/Chaos/Adapters/GeometryCollectionComponentCacheAdapter.cpp
@@ -153,7 +153,7 @@ namespace Chaos
 					Pending.ParticleIndex = TransformIndex;
 					
 					// All recorded transforms are in actor space, ie relative to the Cache Manager making the recording.
-					FTransform LocalTransform = MassToLocal[TransformIndex].Inverse() * FTransform(Handle->R(), Handle->X());
+					FTransform LocalTransform = MassToLocal[TransformIndex].Inverse() * FTransform(Handle->GetR(), Handle->X());
 					FTransform ActorSpaceTransform = LocalTransform * WorldToActor;
 					Pending.PendingTransform = ActorSpaceTransform;
 				}
@@ -176,7 +176,7 @@ namespace Chaos
 					OutFrame.PendingParticleData.AddDefaulted();
 					FPendingParticleWrite& Pending = OutFrame.PendingParticleData.Last();
 					Pending.ParticleIndex = TransformIndex;
-					FTransform LocalTransform = MassToLocal[TransformIndex].Inverse() * FTransform(Handle->R(), Handle->X());
+					FTransform LocalTransform = MassToLocal[TransformIndex].Inverse() * FTransform(Handle->GetR(), Handle->X());
 					FTransform ActorSpaceTransform = LocalTransform * WorldToActor;
 					Pending.PendingTransform = ActorSpaceTransform;
 					Pending.bPendingDeactivate = true;
@@ -491,7 +491,7 @@ namespace Chaos
 				Handle->SetP(WorldTransform.GetTranslation());
 				Handle->SetQ(WorldTransform.GetRotation());
 				Handle->SetX(Handle->P());
-				Handle->SetR(Handle->Q());
+				Handle->SetRf(Handle->GetQf());
 				
 				Handle->UpdateWorldSpaceState(WorldTransform, FVec3(0));
 
diff --git a/Engine/Plugins/Experimental/ChaosCaching/Source/ChaosCaching/Private/Chaos/Adapters/StaticMeshComponentCacheAdapter.cpp b/Engine/Plugins/Experimental/ChaosCaching/Source/ChaosCaching/Private/Chaos/Adapters/StaticMeshComponentCacheAdapter.cpp
--- a/Engine/Plugins/Experimental/ChaosCaching/Source/ChaosCaching/Private/Chaos/Adapters/StaticMeshComponentCacheAdapter.cpp
+++ b/Engine/Plugins/Experimental/ChaosCaching/Source/ChaosCaching/Private/Chaos/Adapters/StaticMeshComponentCacheAdapter.cpp
@@ -53,7 +53,7 @@ namespace Chaos
 				FPendingParticleWrite NewData;
 
 				NewData.ParticleIndex = 0; // Only one particle for static caches
-				NewData.PendingTransform = FTransform(AsRigid->R(), AsRigid->X()).GetRelativeTransform(InRootTransform);
+				NewData.PendingTransform = FTransform(AsRigid->GetR(), AsRigid->X()).GetRelativeTransform(InRootTransform);
 
 				OutFrame.PendingParticleData.Add(MoveTemp(NewData));
 			}
diff --git a/Engine/Plugins/Experimental/ChaosCloth/Source/Chaos/Private/ChaosCloth/ChaosClothingSimulationSolver.cpp b/Engine/Plugins/Experimental/ChaosCloth/Source/Chaos/Private/ChaosCloth/ChaosClothingSimulationSolver.cpp
--- a/Engine/Plugins/Experimental/ChaosCloth/Source/Chaos/Private/ChaosCloth/ChaosClothingSimulationSolver.cpp
+++ b/Engine/Plugins/Experimental/ChaosCloth/Source/Chaos/Private/ChaosCloth/ChaosClothingSimulationSolver.cpp
@@ -312,7 +312,7 @@ FClothingSimulationSolver::FClothingSimulationSolver(bool bForceBasedSolver, FCl
 		PBDEvolution->SetCollisionKinematicUpdateFunction(
 			[this](Softs::FSolverCollisionParticles& ParticlesInput, const Softs::FSolverReal Dt, const Softs::FSolverReal LocalTime, const int32 Index)
 		{
-			LastSubframeCollisionTransformsCCD[Index] = Softs::FSolverRigidTransform3(ParticlesInput.X(Index), ParticlesInput.R(Index));
+			LastSubframeCollisionTransformsCCD[Index] = Softs::FSolverRigidTransform3(ParticlesInput.X(Index), ParticlesInput.GetR(Index));
 
 			checkSlow(Dt > SMALL_NUMBER && DeltaTime > SMALL_NUMBER);
 			const Softs::FSolverReal Alpha = (LocalTime - Time) / DeltaTime;
@@ -321,23 +321,23 @@ FClothingSimulationSolver::FClothingSimulationSolver(bool bForceBasedSolver, FCl
 			ParticlesInput.V(Index) = (NewX - ParticlesInput.X(Index)) / Dt;
 			ParticlesInput.X(Index) = NewX;
 			const Softs::FSolverRotation3 NewR = Softs::FSolverRotation3::Slerp(OldCollisionTransforms[Index].GetRotation(), CollisionTransforms[Index].GetRotation(), Alpha);
-			const Softs::FSolverRotation3 Delta = NewR * ParticlesInput.R(Index).Inverse();
+			const Softs::FSolverRotation3 Delta = NewR * ParticlesInput.GetR(Index).Inverse();
 			const Softs::FSolverReal Angle = Delta.GetAngle();
 			const Softs::FSolverVec3 Axis = Delta.GetRotationAxis();
 			ParticlesInput.W(Index) = (Softs::FSolverVec3)Axis * Angle / Dt;
-			ParticlesInput.R(Index) = NewR;
+			ParticlesInput.SetR(Index, NewR);
 
 			if (TWeightedLatticeImplicitObject<FLevelSet>* SkinnedLevelSet =
 				const_cast<FImplicitObject*>(ParticlesInput.GetGeometry(Index).GetReference())->GetObject<TWeightedLatticeImplicitObject<FLevelSet>>())
 			{
 				const TArray<int32>& SubBoneIndices = SkinnedLevelSet->GetSolverBoneIndices();
-				const FTransform RootTransformInv = TRigidTransform<FReal, 3>(ParticlesInput.X(Index), ParticlesInput.R(Index)).Inverse();
+				const FTransform RootTransformInv = TRigidTransform<FReal, 3>(ParticlesInput.X(Index), ParticlesInput.GetR(Index)).Inverse();
 				TArray<FTransform> SubBoneTransforms;
 				SubBoneTransforms.SetNum(SubBoneIndices.Num());
 				for (int32 SubBoneIdx = 0; SubBoneIdx < SubBoneIndices.Num(); ++SubBoneIdx)
 				{
 					checkSlow(SubBoneIndices[SubBoneIdx] < Index);
-					SubBoneTransforms[SubBoneIdx] = TRigidTransform<FReal, 3>(ParticlesInput.X(SubBoneIndices[SubBoneIdx]), ParticlesInput.R(SubBoneIndices[SubBoneIdx])) * RootTransformInv;
+					SubBoneTransforms[SubBoneIdx] = TRigidTransform<FReal, 3>(ParticlesInput.X(SubBoneIndices[SubBoneIdx]), ParticlesInput.GetR(SubBoneIndices[SubBoneIdx])) * RootTransformInv;
 				}
 				SkinnedLevelSet->DeformPoints(SubBoneTransforms);
 				SkinnedLevelSet->UpdateSpatialHierarchy();
@@ -1574,7 +1574,7 @@ void FClothingSimulationSolver::ApplyPreSimulationTransforms()
 					OldCollisionTransformsView[Index] = OldCollisionTransformsView[Index] * GroupSpaceTransform;
 					OldCollisionTransformsView[Index].AddToTranslation(-DeltaLocalSpaceLocation);
 					CollisionParticles.X(Index) = OldCollisionTransformsView[Index].GetTranslation();
-					CollisionParticles.R(Index) = OldCollisionTransformsView[Index].GetRotation();
+					CollisionParticles.SetR(Index, OldCollisionTransformsView[Index].GetRotation());
 				}
 			}
 		}, /*bForceSingleThreaded =*/ !bClothSolverParallelClothPreUpdate);
@@ -1685,7 +1685,7 @@ void FClothingSimulationSolver::ApplyPreSimulationTransforms()
 				OldCollisionTransforms[Index] = OldCollisionTransforms[Index] * GroupSpaceTransform;
 				OldCollisionTransforms[Index].AddToTranslation(-DeltaLocalSpaceLocation);
 				CollisionParticles.X(Index) = OldCollisionTransforms[Index].GetTranslation();
-				CollisionParticles.R(Index) = OldCollisionTransforms[Index].GetRotation();
+				CollisionParticles.SetR(Index, OldCollisionTransforms[Index].GetRotation());
 			});
 		}
 	}
diff --git a/Engine/Plugins/Experimental/ChaosModularVehicle/Source/ChaosModularVehicleEngine/Private/ChaosModularVehicle/ModularVehicleSimulationCU.cpp b/Engine/Plugins/Experimental/ChaosModularVehicle/Source/ChaosModularVehicleEngine/Private/ChaosModularVehicle/ModularVehicleSimulationCU.cpp
--- a/Engine/Plugins/Experimental/ChaosModularVehicle/Source/ChaosModularVehicleEngine/Private/ChaosModularVehicle/ModularVehicleSimulationCU.cpp
+++ b/Engine/Plugins/Experimental/ChaosModularVehicle/Source/ChaosModularVehicleEngine/Private/ChaosModularVehicle/ModularVehicleSimulationCU.cpp
@@ -157,7 +157,7 @@ void FModularVehicleSimulationCU::PerformAdditionalSimWork(UWorld* InWorld, cons
 			if (FPBDRigidClusteredParticleHandle* ClusterHandle = ClusterUnion->InternalCluster)
 			{
 				TArray<FPBDRigidParticleHandle*> Particles = ClusterUnion->ChildParticles;
-				FRigidTransform3 ClusterWorldTM = FRigidTransform3(ClusterHandle->X(), ClusterHandle->R());
+				FRigidTransform3 ClusterWorldTM = FRigidTransform3(ClusterHandle->X(), ClusterHandle->GetR());
 
 				//Chaos::FDebugDrawQueue::GetInstance().DrawDebugCoordinateSystem(ClusterHandle->X(), FRotator(ClusterHandle->R()), 200, false, -1.f, 1.f, 3.0f);
 
@@ -183,7 +183,7 @@ void FModularVehicleSimulationCU::PerformAdditionalSimWork(UWorld* InWorld, cons
 						{
 							Frame = ClusterChild->ChildToParent();
 
-							const FRigidTransform3 ChildWorldTM(Child->X(), Child->R());
+							const FRigidTransform3 ChildWorldTM(Child->X(), Child->GetR());
 							Frame = ChildWorldTM.GetRelativeTransform(ClusterWorldTM);
 						}
 
diff --git a/Engine/Plugins/Experimental/Mover/Source/Mover/Private/Backends/MoverNetworkPhysicsLiaison.cpp b/Engine/Plugins/Experimental/Mover/Source/Mover/Private/Backends/MoverNetworkPhysicsLiaison.cpp
--- a/Engine/Plugins/Experimental/Mover/Source/Mover/Private/Backends/MoverNetworkPhysicsLiaison.cpp
+++ b/Engine/Plugins/Experimental/Mover/Source/Mover/Private/Backends/MoverNetworkPhysicsLiaison.cpp
@@ -766,7 +766,7 @@ void UMoverNetworkPhysicsLiaisonComponent::AsyncPhysics_OnPreSimulate(const FPhy
 			}
 		}
 	}
-	SyncState.SetTransforms_WorldSpace(CharacterParticle->X(), FRotator(CharacterParticle->R()), CharacterParticle->V() - LocalGroundVelocity);
+	SyncState.SetTransforms_WorldSpace(CharacterParticle->X(), FRotator(CharacterParticle->GetR()), CharacterParticle->GetV() - LocalGroundVelocity);
 
 	///////////////////////////////////////////////////////////////////////////////////////////////////////////////////////////
 	// Update the simulation
@@ -853,7 +853,7 @@ void UMoverNetworkPhysicsLiaisonComponent::AsyncPhysics_OnPreSimulate(const FPhy
 
 	// Note: Output sync state does not have a target angular velocity so
 	// use the target orientation
-	FRotator DeltaRotation = OutputSyncState->GetOrientation_WorldSpace() - FRotator(CharacterParticle->R());
+	FRotator DeltaRotation = OutputSyncState->GetOrientation_WorldSpace() - FRotator(CharacterParticle->GetR());
 	FRotator Winding, Remainder;
 	DeltaRotation.GetWindingAndRemainder(Winding, Remainder);
 	float TargetDeltaFacing = FMath::DegreesToRadians(Remainder.Yaw);
diff --git a/Engine/Plugins/Experimental/Mover/Source/Mover/Private/PhysicsMover/PhysicsMovementUtils.cpp b/Engine/Plugins/Experimental/Mover/Source/Mover/Private/PhysicsMover/PhysicsMovementUtils.cpp
--- a/Engine/Plugins/Experimental/Mover/Source/Mover/Private/PhysicsMover/PhysicsMovementUtils.cpp
+++ b/Engine/Plugins/Experimental/Mover/Source/Mover/Private/PhysicsMover/PhysicsMovementUtils.cpp
@@ -252,13 +252,13 @@ FVector UPhysicsMovementUtils::ComputeGroundVelocityFromHitResult(const FVector&
 		if (Rigid->KinematicTarget().IsSet())
 		{
 			const FVector LinearDisplacement = Rigid->KinematicTarget().GetTargetPosition() - Rigid->X();
-			const FQuat RelativeQuat = Rigid->R().Inverse() * Rigid->KinematicTarget().GetTargetRotation();
+			const FQuat RelativeQuat = Rigid->GetR().Inverse() * Rigid->KinematicTarget().GetTargetRotation();
 			const FVector AngularDisplacement = RelativeQuat.ToRotationVector();
 			GroundVelocity = (LinearDisplacement + AngularDisplacement.Cross(Offset)) / DeltaSeconds;
 		}
 		else
 		{
-			GroundVelocity = Rigid->V() + Rigid->W().Cross(Offset);
+			GroundVelocity = Rigid->GetV() + Rigid->GetW().Cross(Offset);
 		}
 	}
 	return GroundVelocity;
diff --git a/Engine/Plugins/Runtime/NetworkPrediction/Source/NetworkPrediction/Private/NetworkPredictionPhysics.cpp b/Engine/Plugins/Runtime/NetworkPrediction/Source/NetworkPrediction/Private/NetworkPredictionPhysics.cpp
--- a/Engine/Plugins/Runtime/NetworkPrediction/Source/NetworkPrediction/Private/NetworkPredictionPhysics.cpp
+++ b/Engine/Plugins/Runtime/NetworkPrediction/Source/NetworkPrediction/Private/NetworkPredictionPhysics.cpp
@@ -107,17 +107,17 @@ bool FNetworkPredictionPhysicsState::ShouldReconcile(int32 PhysicsFrame, Chaos::
 		return true;
 	}
 
-	if (CompareVector(LocalState.V(), RecvState->LinearVelocity, NetworkPredictionPhysicsCvars::ToleranceV(), "V:"))
+	if (CompareVector(LocalState.GetV(), RecvState->LinearVelocity, NetworkPredictionPhysicsCvars::ToleranceV(), "V:"))
 	{
 		return true;
 	}
 
-	if (CompareVector(LocalState.W(), RecvState->AngularVelocity, NetworkPredictionPhysicsCvars::ToleranceW(), "W:"))
+	if (CompareVector(LocalState.GetW(), RecvState->AngularVelocity, NetworkPredictionPhysicsCvars::ToleranceW(), "W:"))
 	{
 		return true;
 	}
 
-	if (CompareQuat(LocalState.R(), RecvState->Rotation, NetworkPredictionPhysicsCvars::ToleranceR(), "R:"))
+	if (CompareQuat(LocalState.GetR(), RecvState->Rotation, NetworkPredictionPhysicsCvars::ToleranceR(), "R:"))
 	{
 		return true;
 	}
@@ -150,7 +150,7 @@ void FNetworkPredictionPhysicsState::ToString(int32 PhysicsFrame, Chaos::FRewind
 	FPhysicsActorHandle& Handle = BodyInstance->GetPhysicsActorHandle();
 
 	const Chaos::FGeometryParticleState LocalState = RewindData->GetPastStateAtFrame(*Handle->GetHandle_LowLevel(), PhysicsFrame);
-	ToStringInternal(LocalState.X(), LocalState.R(), LocalState.V(), LocalState.W(), Builder);
+	ToStringInternal(LocalState.X(), LocalState.GetR(), LocalState.GetV(), LocalState.GetW(), Builder);
 }
 
 void FNetworkPredictionPhysicsState::ToString(FBodyInstance* BodyInstance, FAnsiStringBuilderBase& Builder)
diff --git a/Engine/Source/Programs/HeadlessChaos/Private/GeometryCollection/GeometryCollectionTestClustering.cpp b/Engine/Source/Programs/HeadlessChaos/Private/GeometryCollection/GeometryCollectionTestClustering.cpp
--- a/Engine/Source/Programs/HeadlessChaos/Private/GeometryCollection/GeometryCollectionTestClustering.cpp
+++ b/Engine/Source/Programs/HeadlessChaos/Private/GeometryCollection/GeometryCollectionTestClustering.cpp
@@ -1796,7 +1796,7 @@ namespace GeometryCollectionTest
 			EXPECT_TRUE(ClusterMapContains(ClusterMap,CollectionParticles2[2],{CollectionParticles2[1],CollectionParticles2[0]}));
 		});
 
-		for (int Frame = 0; Frame < 100; Frame++)
+		for (int Frame = 0; Frame < 50; Frame++)
 		{
 			UnitTest.Advance();
 
@@ -1835,7 +1835,7 @@ namespace GeometryCollectionTest
 			}
 			else
 			{
-				FTransform RootTransform(Root->R(), Root->X());
+				FTransform RootTransform(Root->GetR(), Root->X());
 
 				TArray<FTransform> GlobalTransform1;
 				GeometryCollectionAlgo::Private::GlobalMatrices(*DynamicCollection, GlobalTransform1);
diff --git a/Engine/Source/Programs/HeadlessChaos/Private/GeometryCollection/GeometryCollectionTestInitilization.cpp b/Engine/Source/Programs/HeadlessChaos/Private/GeometryCollection/GeometryCollectionTestInitilization.cpp
--- a/Engine/Source/Programs/HeadlessChaos/Private/GeometryCollection/GeometryCollectionTestInitilization.cpp
+++ b/Engine/Source/Programs/HeadlessChaos/Private/GeometryCollection/GeometryCollectionTestInitilization.cpp
@@ -26,7 +26,7 @@ GTEST_TEST(AllTraits,GeometryCollection_Initilization_TransformedGeometryCollect
 	{ // test results
 		EXPECT_EQ(UnitTest.Solver->GetParticles().GetGeometryCollectionParticles().Size(),1);
 		FVector X = UnitTest.Solver->GetParticles().GetGeometryCollectionParticles().X(0);
-		FQuat R = UnitTest.Solver->GetParticles().GetGeometryCollectionParticles().R(0);
+		FQuat R = UnitTest.Solver->GetParticles().GetGeometryCollectionParticles().GetR(0);
 		EXPECT_TRUE((R * GlobalRotation.Inverse()).IsIdentity(KINDA_SMALL_NUMBER));
 		EXPECT_NEAR(X.X - GlobalTranslation[0],0.0f,KINDA_SMALL_NUMBER);
 		EXPECT_NEAR(X.Y - GlobalTranslation[1],0.0f,KINDA_SMALL_NUMBER);
@@ -57,7 +57,7 @@ GTEST_TEST(AllTraits,GeometryCollection_Initilization_TransformedGeometryCollect
 	{ // test results
 		EXPECT_EQ(UnitTest.Solver->GetParticles().GetGeometryCollectionParticles().Size(),1);
 		FVector X = UnitTest.Solver->GetParticles().GetGeometryCollectionParticles().X(0);
-		FQuat R = UnitTest.Solver->GetParticles().GetGeometryCollectionParticles().R(0);
+		FQuat R = UnitTest.Solver->GetParticles().GetGeometryCollectionParticles().GetR(0);
 		EXPECT_TRUE((R * GlobalRotation.Inverse()).IsIdentity(KINDA_SMALL_NUMBER));
 		EXPECT_NEAR(X.X - GlobalTranslation[0],0.0f,KINDA_SMALL_NUMBER);
 		EXPECT_NEAR(X.Y - GlobalTranslation[1],0.0f,KINDA_SMALL_NUMBER);
@@ -88,7 +88,7 @@ GTEST_TEST(AllTraits,GeometryCollection_Initilization_TransformedGeometryCollect
 	{ // test results
 		EXPECT_EQ(UnitTest.Solver->GetParticles().GetGeometryCollectionParticles().Size(),1);
 		FVector X = UnitTest.Solver->GetParticles().GetGeometryCollectionParticles().X(0);
-		FQuat R = UnitTest.Solver->GetParticles().GetGeometryCollectionParticles().R(0);
+		FQuat R = UnitTest.Solver->GetParticles().GetGeometryCollectionParticles().GetR(0);
 		EXPECT_TRUE((R * GlobalRotation.Inverse()).IsIdentity(KINDA_SMALL_NUMBER));
 		EXPECT_NEAR(X.X - GlobalTranslation[0],0.0f,KINDA_SMALL_NUMBER);
 		EXPECT_NEAR(X.Y - GlobalTranslation[1],0.0f,KINDA_SMALL_NUMBER);
@@ -143,7 +143,7 @@ GTEST_TEST(AllTraits,GeometryCollection_Initilization_TransformedGeometryCollect
 
 		EXPECT_EQ(UnitTest.Solver->GetParticles().GetGeometryCollectionParticles().Size(), 1);
 		FVector ParticlePos = UnitTest.Solver->GetParticles().GetGeometryCollectionParticles().X(0);
-		FQuat ParticleR = UnitTest.Solver->GetParticles().GetGeometryCollectionParticles().R(0);
+		FQuat ParticleR = UnitTest.Solver->GetParticles().GetGeometryCollectionParticles().GetR(0);
 
 		TArray<FTransform> Transform;
 		GeometryCollectionAlgo::Private::GlobalMatrices(*Collection->DynamicCollection, Transform);
diff --git a/Engine/Source/Programs/HeadlessChaos/Private/GeometryCollection/GeometryCollectionTestSimulationField.cpp b/Engine/Source/Programs/HeadlessChaos/Private/GeometryCollection/GeometryCollectionTestSimulationField.cpp
--- a/Engine/Source/Programs/HeadlessChaos/Private/GeometryCollection/GeometryCollectionTestSimulationField.cpp
+++ b/Engine/Source/Programs/HeadlessChaos/Private/GeometryCollection/GeometryCollectionTestSimulationField.cpp
@@ -350,10 +350,10 @@ namespace GeometryCollectionTest
 			else
 			{
 				EXPECT_NE(FMath::Abs(Collection->DynamicCollection->GetTransform(0).GetRotation().Euler().Y), SMALL_THRESHOLD);
-				EXPECT_GT(Particles.W(0).Y, PreviousY);
+				EXPECT_GT(Particles.GetW(0).Y, PreviousY);
 			}
 
-			PreviousY = Particles.W(0).Y;
+			PreviousY = Particles.GetW(0).Y;
 		}
 
 	}
@@ -1073,7 +1073,7 @@ namespace GeometryCollectionTest
 		TArray<Chaos::TPBDRigidClusteredParticleHandle<FReal, 3>*> ParticleHandles = {
 			Collection->PhysObject->GetParticle_Internal(0),
 		};
-		Chaos::TVector<float, 3> CurrV = ParticleHandles[0]->V();
+		Chaos::TVector<float, 3> CurrV = ParticleHandles[0]->GetV();
 		
 
 		for (int Frame = 1; Frame < 10; Frame++)
@@ -1083,7 +1083,7 @@ namespace GeometryCollectionTest
 
 			UnitTest.Advance();
 
-			CurrV = ParticleHandles[0]->V();
+			CurrV = ParticleHandles[0]->GetV();
 			EXPECT_NEAR(CurrV.X, Params.InitialLinearVelocity.X, KINDA_SMALL_NUMBER); // Velocity in +x
 			EXPECT_GT(Collection->DynamicCollection->GetTransform(0).GetTranslation().X, LastLocation.X); // Pos in +x
 
@@ -1132,11 +1132,11 @@ namespace GeometryCollectionTest
 
 			Chaos::TPBDGeometryCollectionParticles<Chaos::FReal, 3>& Particles = UnitTest.Solver->GetParticles().GetGeometryCollectionParticles();
 			EXPECT_NE(FMath::Abs(Collection->DynamicCollection->GetTransform(0).GetRotation().Euler().Y), SMALL_THRESHOLD); // not rotating in Y?
-			EXPECT_GT(Particles.W(0).X, PreviousX); // rotating in X?
+			EXPECT_GT(Particles.GetW(0).X, PreviousX); // rotating in X?
 			EXPECT_LT(Particles.X(0).Z, PreviousHeight); // still falling?
 
 			PreviousHeight = Particles.X(0).Z;
-			PreviousX = Particles.W(0).X;
+			PreviousX = Particles.GetW(0).X;
 		}
 
 
diff --git a/Engine/Source/Programs/HeadlessChaos/Private/HeadlessChaosTestBroadphase.cpp b/Engine/Source/Programs/HeadlessChaos/Private/HeadlessChaosTestBroadphase.cpp
--- a/Engine/Source/Programs/HeadlessChaos/Private/HeadlessChaosTestBroadphase.cpp
+++ b/Engine/Source/Programs/HeadlessChaos/Private/HeadlessChaosTestBroadphase.cpp
@@ -58,7 +58,7 @@ namespace ChaosTest
 		template <SQType>
 		bool Visit(int32 Idx, FQueryFastData& CurData)
 		{
-			const FRigidTransform3 BoxTM(Boxes.X(Idx), Boxes.R(Idx));
+			const FRigidTransform3 BoxTM(Boxes.X(Idx), Boxes.GetR(Idx));
 			FAABB3 Box = static_cast<const TBox<FReal, 3>*>(Boxes.GetGeometry(Idx).GetReference())->BoundingBox().TransformedAABB(BoxTM);
 			FAABB3 ThicknedBox(Box.Min() - HalfExtents, Box.Max() + HalfExtents);
 
@@ -133,7 +133,7 @@ namespace ChaosTest
 		bool VisitOverlap(TSpatialVisitorData<int32> Instance)
 		{
 			const int32 Idx = Instance.Payload;
-			const FRigidTransform3 BoxTM(Boxes.X(Idx), Boxes.R(Idx));
+			const FRigidTransform3 BoxTM(Boxes.X(Idx), Boxes.GetR(Idx));
 			FAABB3 Box = static_cast<const TBox<FReal, 3>*>(Boxes.GetGeometry(Idx).GetReference())->BoundingBox().TransformedAABB(BoxTM);
 			
 			if (Box.Intersects(Bounds))
@@ -241,7 +241,7 @@ namespace ChaosTest
 				for (int32 Col = 0; Col < NumCols; ++Col)
 				{
 					Boxes->X(Idx) = FVec3(Col * 100, Row * 100, Height * 100) + Offset;
-					Boxes->R(Idx) = FRotation3::Identity;
+					Boxes->SetR(Idx, FRotation3::Identity);
 					Boxes->SetGeometry(Idx, Box);
 					++Idx;
 				}
@@ -287,7 +287,7 @@ namespace ChaosTest
 			{
 				const int32 MoveIdx = Visitor2.Instances[0];
 				Boxes->X(MoveIdx) += FVec3(1000, 0, 0);
-				FAABB3 NewBounds = Boxes->GetGeometry(MoveIdx)->template GetObject<TBox<FReal, 3>>()->BoundingBox().TransformedAABB(FRigidTransform3(Boxes->X(MoveIdx), Boxes->R(MoveIdx)));
+				FAABB3 NewBounds = Boxes->GetGeometry(MoveIdx)->template GetObject<TBox<FReal, 3>>()->BoundingBox().TransformedAABB(FRigidTransform3(Boxes->X(MoveIdx), Boxes->GetR(MoveIdx)));
 				Spatial2->UpdateElementIn(MoveIdx, NewBounds, true, SpatialIdx);
 
 				FVisitor Visitor3(FVec3(10, 0, 0), FVec3(0, 1, 0), 0, *Boxes);
@@ -296,7 +296,7 @@ namespace ChaosTest
 
 				//move instance back
 				Boxes->X(MoveIdx) -= FVec3(1000, 0, 0);
-				NewBounds = Boxes->GetGeometry(MoveIdx)->template GetObject<TBox<FReal, 3>>()->BoundingBox().TransformedAABB(FRigidTransform3(Boxes->X(MoveIdx), Boxes->R(MoveIdx)));
+				NewBounds = Boxes->GetGeometry(MoveIdx)->template GetObject<TBox<FReal, 3>>()->BoundingBox().TransformedAABB(FRigidTransform3(Boxes->X(MoveIdx), Boxes->GetR(MoveIdx)));
 				Spatial2->UpdateElementIn(MoveIdx, NewBounds, true, SpatialIdx);
 
 				FVisitor Visitor4(FVec3(10, 0, 0), FVec3(0, 1, 0), 0, *Boxes);
@@ -309,7 +309,7 @@ namespace ChaosTest
 				const int32 MoveIdx = 5 * 5 * 5;
 				const FVec3 OldPos = Boxes->X(MoveIdx);
 				Boxes->X(MoveIdx) = FVec3(0, 0, 0);
-				FAABB3 NewBounds = Boxes->GetGeometry(MoveIdx)->template GetObject<TBox<FReal, 3>>()->BoundingBox().TransformedAABB(FRigidTransform3(Boxes->X(MoveIdx), Boxes->R(MoveIdx)));
+				FAABB3 NewBounds = Boxes->GetGeometry(MoveIdx)->template GetObject<TBox<FReal, 3>>()->BoundingBox().TransformedAABB(FRigidTransform3(Boxes->X(MoveIdx), Boxes->GetR(MoveIdx)));
 				Spatial2->UpdateElementIn(MoveIdx, NewBounds, true, SpatialIdx);
 
 				FVisitor Visitor3(FVec3(10, 0, 0), FVec3(0, 1, 0), 0, *Boxes);
@@ -318,7 +318,7 @@ namespace ChaosTest
 
 				//move instance back
 				Boxes->X(MoveIdx) = OldPos;
-				NewBounds = Boxes->GetGeometry(MoveIdx)->template GetObject<TBox<FReal, 3>>()->BoundingBox().TransformedAABB(FRigidTransform3(Boxes->X(MoveIdx), Boxes->R(MoveIdx)));
+				NewBounds = Boxes->GetGeometry(MoveIdx)->template GetObject<TBox<FReal, 3>>()->BoundingBox().TransformedAABB(FRigidTransform3(Boxes->X(MoveIdx), Boxes->GetR(MoveIdx)));
 				Spatial2->UpdateElementIn(MoveIdx, NewBounds, true, SpatialIdx);
 			}
 
@@ -327,7 +327,7 @@ namespace ChaosTest
 				const int32 MoveIdx = 5 * 5 * 5;
 				const FVec3 OldPos = Boxes->X(MoveIdx);
 				Boxes->X(MoveIdx) = FVec3(-50, 0, 0);
-				FAABB3 NewBounds = Boxes->GetGeometry(MoveIdx)->template GetObject<TBox<FReal, 3>>()->BoundingBox().TransformedAABB(FRigidTransform3(Boxes->X(MoveIdx), Boxes->R(MoveIdx)));
+				FAABB3 NewBounds = Boxes->GetGeometry(MoveIdx)->template GetObject<TBox<FReal, 3>>()->BoundingBox().TransformedAABB(FRigidTransform3(Boxes->X(MoveIdx), Boxes->GetR(MoveIdx)));
 				Spatial2->UpdateElementIn(MoveIdx, NewBounds, true, SpatialIdx);
 
 				FVisitor Visitor3(FVec3(10, 0, 0), FVec3(0, 1, 0), 0, *Boxes);
@@ -352,9 +352,9 @@ namespace ChaosTest
 				const int32 NewIdx = Boxes->Size();
 				Boxes->AddParticles(1);
 				Boxes->X(NewIdx) = FVec3(-20, 0, 0);
-				Boxes->R(NewIdx) = FRotation3::Identity;
+				Boxes->SetR(NewIdx, FRotation3::Identity);
 				Boxes->SetGeometry(NewIdx, Box);
-				NewBounds = Boxes->GetGeometry(NewIdx)->template GetObject<TBox<FReal, 3>>()->BoundingBox().TransformedAABB(FRigidTransform3(Boxes->X(NewIdx), Boxes->R(NewIdx)));
+				NewBounds = Boxes->GetGeometry(NewIdx)->template GetObject<TBox<FReal, 3>>()->BoundingBox().TransformedAABB(FRigidTransform3(Boxes->X(NewIdx), Boxes->GetR(NewIdx)));
 				Spatial2->UpdateElementIn(NewIdx, NewBounds, true, SpatialIdx);
 				FVisitor Visitor6(FVec3(-20, 0, 0), FVec3(0, 1, 0), 0, *Boxes);
 				Spatial2->Raycast(Visitor6.Start, Visitor6.Dir, 1000, Visitor6);
@@ -496,7 +496,7 @@ namespace ChaosTest
 				for (int32 Col = 0; Col < NumCols; ++Col)
 				{
 					Boxes.X(Idx) = FVec3(Col * 100, Row * 100, Height * 100);
-					Boxes.R(Idx) = FRotation3::Identity;
+					Boxes.SetR(Idx, FRotation3::Identity);
 					Boxes.SetGeometry(Idx, Box);
 					++Idx;
 				}
@@ -1058,7 +1058,7 @@ namespace ChaosTest
 		// Construct a particle and set HasBounds to false.
 		int32 Idx = 0;
 		Boxes->X(Idx) = FVec3(0);
-		Boxes->R(Idx) = FRotation3::Identity;
+		Boxes->SetR(Idx, FRotation3::Identity);
 		Boxes->SetGeometry(Idx, Box);
 
 		// Tell BV we have no bounds, this used to cause issues.
diff --git a/Engine/Source/Programs/HeadlessChaos/Private/HeadlessChaosTestCCD.cpp b/Engine/Source/Programs/HeadlessChaos/Private/HeadlessChaosTestCCD.cpp
--- a/Engine/Source/Programs/HeadlessChaos/Private/HeadlessChaosTestCCD.cpp
+++ b/Engine/Source/Programs/HeadlessChaos/Private/HeadlessChaosTestCCD.cpp
@@ -56,10 +56,10 @@ namespace ChaosTest
 		// Positions and velocities
 		Static->X() = FVec3(0, 0, 0);
 		Dynamic->X() = FVec3(0, 0, InitialPosition); // Start 30cm above the static box
-		Dynamic->V() = FVec3(0, 0, -InitialSpeed);
+		Dynamic->SetV(FVec3(0, 0, -InitialSpeed));
 
 		// The position of the static has changed and statics don't automatically update bounds, so update explicitly
-		Static->UpdateWorldSpaceState(TRigidTransform<FReal, 3>(Static->X(), Static->R()), FVec3(0));
+		Static->UpdateWorldSpaceState(TRigidTransform<FReal, 3>(Static->X(), Static->GetR()), FVec3(0));
 
 		::ChaosTest::SetParticleSimDataToCollide({ Static,Dynamic });
 
@@ -70,7 +70,7 @@ namespace ChaosTest
 		Evolution.EnableParticle(Dynamic);
 
 
-		Dynamic->V() = FVec3(0, 0, -InitialSpeed);
+		Dynamic->SetV(FVec3(0, 0, -InitialSpeed));
 
 		for (int i = 0; i < 1; ++i)
 		{
@@ -120,10 +120,10 @@ namespace ChaosTest
 		// Positions and velocities
 		Static->X() = FVec3(0, 0, 0);
 		Dynamic->X() = FVec3(0, 0, BoxHalfSize * 2 + 30); // Start 30cm above the static box
-		Dynamic->V() = FVec3(0, 0, -InitialSpeed);
+		Dynamic->SetV(FVec3(0, 0, -InitialSpeed));
 
 		// The position of the static has changed and statics don't automatically update bounds, so update explicitly
-		Static->UpdateWorldSpaceState(TRigidTransform<FReal, 3>(Static->X(), Static->R()), FVec3(0));
+		Static->UpdateWorldSpaceState(TRigidTransform<FReal, 3>(Static->X(), Static->GetR()), FVec3(0));
 
 		::ChaosTest::SetParticleSimDataToCollide({ Static,Dynamic });
 
@@ -131,7 +131,7 @@ namespace ChaosTest
 		Dynamic->SetGravityEnabled(false);
 
 
-		Dynamic->V() = FVec3(0, 0, -InitialSpeed);
+		Dynamic->SetV(FVec3(0, 0, -InitialSpeed));
 
 		Evolution.EnableParticle(Static);
 		Evolution.EnableParticle(Dynamic);
@@ -187,10 +187,10 @@ namespace ChaosTest
 		Static->X() = FVec3(0, 0, 0);
 
 		Dynamic->X() = FVec3(0, 0, SphereRadius * 2 + 10);
-		Dynamic->V() = FVec3(0, 0, -InitialSpeed);
+		Dynamic->SetV(FVec3(0, 0, -InitialSpeed));
 		
 		// The position of the static has changed and statics don't automatically update bounds, so update explicitly
-		Static->UpdateWorldSpaceState(TRigidTransform<FReal, 3>(Static->X(), Static->R()), FVec3(0));
+		Static->UpdateWorldSpaceState(TRigidTransform<FReal, 3>(Static->X(), Static->GetR()), FVec3(0));
 
 		::ChaosTest::SetParticleSimDataToCollide({ Static,Dynamic });
 
@@ -198,7 +198,7 @@ namespace ChaosTest
 		Dynamic->SetGravityEnabled(false);
 		
 
-		Dynamic->V() = FVec3(0, 0, -InitialSpeed);
+		Dynamic->SetV(FVec3(0, 0, -InitialSpeed));
 
 		Evolution.EnableParticle(Static);
 		Evolution.EnableParticle(Dynamic);
@@ -279,12 +279,12 @@ namespace ChaosTest
 		Dynamic->X() = FVec3(0, 0, 0);
 
 		// The position of the static has changed and statics don't automatically update bounds, so update explicitly
-		ContainerFaces[0]->UpdateWorldSpaceState(TRigidTransform<FReal, 3>(ContainerFaces[0]->X(), ContainerFaces[0]->R()), FVec3(0));
-		ContainerFaces[1]->UpdateWorldSpaceState(TRigidTransform<FReal, 3>(ContainerFaces[1]->X(), ContainerFaces[1]->R()), FVec3(0));
-		ContainerFaces[2]->UpdateWorldSpaceState(TRigidTransform<FReal, 3>(ContainerFaces[2]->X(), ContainerFaces[2]->R()), FVec3(0));
-		ContainerFaces[3]->UpdateWorldSpaceState(TRigidTransform<FReal, 3>(ContainerFaces[3]->X(), ContainerFaces[3]->R()), FVec3(0));
-		ContainerFaces[4]->UpdateWorldSpaceState(TRigidTransform<FReal, 3>(ContainerFaces[4]->X(), ContainerFaces[4]->R()), FVec3(0));
-		ContainerFaces[5]->UpdateWorldSpaceState(TRigidTransform<FReal, 3>(ContainerFaces[5]->X(), ContainerFaces[5]->R()), FVec3(0));
+		ContainerFaces[0]->UpdateWorldSpaceState(TRigidTransform<FReal, 3>(ContainerFaces[0]->X(), ContainerFaces[0]->GetR()), FVec3(0));
+		ContainerFaces[1]->UpdateWorldSpaceState(TRigidTransform<FReal, 3>(ContainerFaces[1]->X(), ContainerFaces[1]->GetR()), FVec3(0));
+		ContainerFaces[2]->UpdateWorldSpaceState(TRigidTransform<FReal, 3>(ContainerFaces[2]->X(), ContainerFaces[2]->GetR()), FVec3(0));
+		ContainerFaces[3]->UpdateWorldSpaceState(TRigidTransform<FReal, 3>(ContainerFaces[3]->X(), ContainerFaces[3]->GetR()), FVec3(0));
+		ContainerFaces[4]->UpdateWorldSpaceState(TRigidTransform<FReal, 3>(ContainerFaces[4]->X(), ContainerFaces[4]->GetR()), FVec3(0));
+		ContainerFaces[5]->UpdateWorldSpaceState(TRigidTransform<FReal, 3>(ContainerFaces[5]->X(), ContainerFaces[5]->GetR()), FVec3(0));
 
 		::ChaosTest::SetParticleSimDataToCollide({ Dynamic });
 		::ChaosTest::SetParticleSimDataToCollide({ ContainerFaces });
@@ -301,7 +301,7 @@ namespace ChaosTest
 		Evolution.EnableParticle(ContainerFaces[5]);
 		Evolution.EnableParticle(Dynamic);
 
-		Dynamic->V() = InitialVelocity;
+		Dynamic->SetV(InitialVelocity);
 		///////////////////////////////////
 		// Test 1: bouncing from two opposite walls
 		for (int i = 0; i < 10; ++i)
@@ -321,12 +321,12 @@ namespace ChaosTest
 		}
 		/////////////////////////////////////////////
 		// Test2: Now launch to cube to a corner
-		Dynamic->V() = FVec3(-ContainerBoxHalfSize * Fps * 10);
-		Dynamic->W() = FVec3(0);
+		Dynamic->SetVf(FVec3f(-ContainerBoxHalfSize * Fps * 10));
+		Dynamic->SetWf(FVec3f(0));
 		Dynamic->X() = FVec3(0);
 		Dynamic->P() = FVec3(0);
-		Dynamic->R() = TRotation<FReal, 3>::FromIdentity();
-		Dynamic->Q() = TRotation<FReal, 3>::FromIdentity();
+		Dynamic->SetRf(TRotation<FRealSingle, 3>::FromIdentity());
+		Dynamic->SetQf(TRotation<FRealSingle, 3>::FromIdentity());
 
 		for (int i = 0; i < 10; ++i)
 		{
@@ -345,15 +345,15 @@ namespace ChaosTest
 		/////////////////////////////////////////////////////////////////////
 		// Test 3: Now we give it something impossible to solve with PBD solver (restitution of 1, high velocities causes final position to be outside of box). 
 		// Make sure it still stays inside the box (albeit with a very reduced velocity) 
-		Dynamic->V() = InitialVelocity;
-		Dynamic->W() = FVec3(0);
+		Dynamic->SetV(InitialVelocity);
+		Dynamic->SetW(FVec3(0));
 		Dynamic->X() = FVec3(0);
 		Dynamic->P() = FVec3(0);
-		Dynamic->R() = TRotation<FReal, 3>::FromIdentity();
-		Dynamic->Q() = TRotation<FReal, 3>::FromIdentity();
-		PhysicsMaterial->Restitution = 1.0f;
+		Dynamic->SetR(TRotation<FReal, 3>::FromIdentity());
+		Dynamic->SetQ(TRotation<FReal, 3>::FromIdentity());
+		PhysicsMaterial->Restitution = 0.9f;
 
-		for (int i = 0; i < 10; ++i)
+		for (int i = 0; i < 10; ++i) // To fix this unit test
 		{
 			Evolution.AdvanceOneTimeStep(Dt);
 			Evolution.EndFrame(Dt);
diff --git a/Engine/Source/Programs/HeadlessChaos/Private/HeadlessChaosTestClustering.cpp b/Engine/Source/Programs/HeadlessChaos/Private/HeadlessChaosTestClustering.cpp
--- a/Engine/Source/Programs/HeadlessChaos/Private/HeadlessChaosTestClustering.cpp
+++ b/Engine/Source/Programs/HeadlessChaos/Private/HeadlessChaosTestClustering.cpp
@@ -52,12 +52,12 @@ namespace ChaosTest {
 		EXPECT_EQ(ClusteredParticles.Size(), 3);
 
 		FVec3 ClusterX = ClusteredParticles.X(2);
-		FRotation3 ClusterRot = ClusteredParticles.R(2);
+		FRotation3 ClusterRot = ClusteredParticles.GetR(2);
 
 		EXPECT_TRUE(ClusterX.Equals(FVec3 {(FReal)50, 0, 0}));
 		EXPECT_TRUE(ClusterRot.Equals(FRotation3::Identity));
 		EXPECT_TRUE(ClusterX.Equals(ClusteredParticles.P(2)));
-		EXPECT_TRUE(ClusterRot.Equals(ClusteredParticles.Q(2)));
+		EXPECT_TRUE(ClusterRot.Equals(ClusteredParticles.GetQ(2)));
 
 		FRigidTransform3 ClusterTM(ClusterX, ClusterRot);
 		FVec3 LocalPos = ClusterTM.InverseTransformPositionNoScale(FVec3 {(FReal)200, (FReal)0, (FReal)0});
@@ -114,7 +114,7 @@ namespace ChaosTest {
 		Chaos::FPBDRigidParticleHandle* RootClusterHandle = Evolution.GetRigidClustering().CreateClusterParticle(0, MoveTemp(ClusterHandlesCopy), ClusterParams, FImplicitObjectPtr(nullptr));
 		FVec3 InitialVelocity((FReal)50, (FReal)20, (FReal)100);
 
-		RootClusterHandle->V() = InitialVelocity;		
+		RootClusterHandle->SetV(InitialVelocity);
 		
 		constexpr int NumParticles = NumBoxes + NumBoxes / 4 + 1;
 		EXPECT_EQ(ClusteredParticles.Size(), NumParticles);
@@ -160,7 +160,7 @@ namespace ChaosTest {
 			}
 
 			EXPECT_TRUE(bFoundInNonDisabled);	//clusters are enabled and in non disabled array
-			EXPECT_TRUE(ClusterHandle->V().Equals(InitialVelocity));
+			EXPECT_TRUE(ClusterHandle->GetV().Equals(InitialVelocity));
 		}
 
 		Evolution.AdvanceOneTimeStep(Dt);
@@ -187,7 +187,7 @@ namespace ChaosTest {
 				bFoundInNonDisabled |= Particle.Handle() == BoxHandle;
 			}
 			EXPECT_TRUE(bFoundInNonDisabled);
-			EXPECT_TRUE(BoxHandle->V().Equals(InitialVelocity));
+			EXPECT_TRUE(BoxHandle->GetV().Equals(InitialVelocity));
 		}
 	}
 
@@ -237,7 +237,7 @@ namespace ChaosTest {
 		Chaos::FPBDRigidParticleHandle* RootClusterHandle = Evolution.GetRigidClustering().CreateClusterParticle(0, MoveTemp(ClusterHandles), ClusterParams, FImplicitObjectPtr(nullptr));
 		FVec3 InitialVelocity((FReal)50, (FReal)20, (FReal)100);
 
-		RootClusterHandle->V() = InitialVelocity;
+		RootClusterHandle->SetV(InitialVelocity);
 
 		TUniquePtr<FChaosPhysicsMaterial> PhysicalMaterial = MakeUnique<FChaosPhysicsMaterial>();
 		PhysicalMaterial->Friction = 0;
diff --git a/Engine/Source/Programs/HeadlessChaos/Private/HeadlessChaosTestCollisions.cpp b/Engine/Source/Programs/HeadlessChaos/Private/HeadlessChaosTestCollisions.cpp
--- a/Engine/Source/Programs/HeadlessChaos/Private/HeadlessChaosTestCollisions.cpp
+++ b/Engine/Source/Programs/HeadlessChaos/Private/HeadlessChaosTestCollisions.cpp
@@ -15,11 +15,11 @@
 
 #define SMALL_THRESHOLD 1e-4
 
-#define RESET_PQ(Particle) Particle->P() = Particle->X(); Particle->Q() = Particle->R()
-#define INVARIANT_XR_START(Particle) FVec3 InvariantPreX_##Particle = Particle->X(); FRotation3 InvariantPreR_##Particle = Particle->R()
-#define INVARIANT_XR_END(Particle) EXPECT_TRUE(InvariantPreX_##Particle.Equals(Particle->X())); EXPECT_TRUE(InvariantPreR_##Particle.Equals(Particle->R()))
-#define INVARIANT_VW_START(Particle) FVec3 InvariantPreV_##Particle = Particle->V(); FVec3 InvariantPreW_##Particle = Particle->W()
-#define INVARIANT_VW_END(Particle) EXPECT_TRUE(InvariantPreV_##Particle.Equals(Particle->V())); EXPECT_TRUE(InvariantPreW_##Particle.Equals(Particle->W()))
+#define RESET_PQ(Particle) Particle->P() = Particle->X(); Particle->SetQ(Particle->GetR());
+#define INVARIANT_XR_START(Particle) FVec3 InvariantPreX_##Particle = Particle->X(); FRotation3 InvariantPreR_##Particle = Particle->GetR()
+#define INVARIANT_XR_END(Particle) EXPECT_TRUE(InvariantPreX_##Particle.Equals(Particle->X())); EXPECT_TRUE(InvariantPreR_##Particle.Equals(Particle->GetR()))
+#define INVARIANT_VW_START(Particle) FVec3 InvariantPreV_##Particle = Particle->GetV(); FVec3 InvariantPreW_##Particle = Particle->GetW()
+#define INVARIANT_VW_END(Particle) EXPECT_TRUE(InvariantPreV_##Particle.Equals(Particle->GetV())); EXPECT_TRUE(InvariantPreW_##Particle.Equals(Particle->GetW()))
 
 namespace ChaosTest {
 
@@ -44,19 +44,19 @@ namespace ChaosTest {
 
 		auto Box1 = AppendDynamicParticleBox(Particles);
 		Box1->X() = FVec3(1.f);
-		Box1->R() = FRotation3(FQuat::Identity);
+		Box1->SetR(FRotation3(FQuat::Identity));
 		Box1->P() = Box1->X();
-		Box1->Q() = Box1->R();
+		Box1->SetQ(Box1->GetR());
 		Box1->AuxilaryValue(PhysicsMaterials) = MakeSerializable(PhysicsMaterial);
-		Box1->UpdateWorldSpaceState(FRigidTransform3(Box1->X(), Box1->R()), FVec3(0));
+		Box1->UpdateWorldSpaceState(FRigidTransform3(Box1->X(), Box1->GetR()), FVec3(0));
 
 		auto Box2 = AppendDynamicParticleBox(Particles);
 		Box2->X() = FVec3(0.5f, 0.5f, 1.9f);
-		Box2->R() = FRotation3(FQuat::Identity);
+		Box2->SetR(FRotation3(FQuat::Identity));
 		Box2->P() = Box2->X();
-		Box2->Q() = Box2->R();
+		Box2->SetQ(Box2->GetR());
 		Box2->AuxilaryValue(PhysicsMaterials) = MakeSerializable(PhysicsMaterial);
-		Box2->UpdateWorldSpaceState(FRigidTransform3(Box2->X(), Box2->R()), FVec3(0));
+		Box2->UpdateWorldSpaceState(FRigidTransform3(Box2->X(), Box2->GetR()), FVec3(0));
 
 		FPBDCollisionConstraintAccessor Collisions(Particles, Collided, PhysicsMaterials, PerParticlePhysicsMaterials, 1, 1);
 		Collisions.ComputeConstraints(0.f);
@@ -96,19 +96,19 @@ namespace ChaosTest {
 
 		auto Box1 = AppendDynamicParticleConvexBox(Particles, FVec3(1.f) );
 		Box1->X() = FVec3(0.f);
-		Box1->R() = FRotation3(FQuat::Identity);
+		Box1->SetR(FRotation3(FQuat::Identity));
 		Box1->P() = Box1->X();
-		Box1->Q() = Box1->R();
+		Box1->SetQ(Box1->GetR());
 		Box1->AuxilaryValue(PhysicsMaterials) = MakeSerializable(PhysicsMaterial);
-		Box1->UpdateWorldSpaceState(FRigidTransform3(Box1->X(), Box1->R()), FVec3(0));
+		Box1->UpdateWorldSpaceState(FRigidTransform3(Box1->X(), Box1->GetR()), FVec3(0));
 
 		auto Box2 = AppendDynamicParticleBox(Particles, FVec3(1.f) );
 		Box2->X() = FVec3(1.25f, 0.f, 0.f);
-		Box2->R() = FRotation3(FQuat::Identity);
+		Box2->SetR(FRotation3(FQuat::Identity));
 		Box2->P() = Box2->X();
-		Box2->Q() = Box2->R();
+		Box2->SetQ(Box2->GetR());
 		Box2->AuxilaryValue(PhysicsMaterials) = MakeSerializable(PhysicsMaterial);
-		Box2->UpdateWorldSpaceState(FRigidTransform3(Box2->X(), Box2->R()), FVec3(0));
+		Box2->UpdateWorldSpaceState(FRigidTransform3(Box2->X(), Box2->GetR()), FVec3(0));
 
 		FPBDCollisionConstraintAccessor Collisions(Particles, Collided, PhysicsMaterials, PerParticlePhysicsMaterials, 1, 1);
 		Collisions.ComputeConstraints(0.f);
@@ -147,13 +147,13 @@ namespace ChaosTest {
 		auto Floor = AppendStaticAnalyticFloor(Particles);
 		auto Box = AppendDynamicParticleBox(Particles);
 		Box->P() = FVec3(0, 1, 0);
-		Box->Q() = FRotation3(FQuat::Identity);
-		Box->V() = FVec3(0, 0, -1);
-		Box->PreV() = Box->V();
-		Box->X() = Box->P() - Box->V() * Dt;
-		Box->R() = Box->Q();
+		Box->SetQ(FRotation3(FQuat::Identity));
+		Box->SetV(FVec3(0, 0, -1));
+		Box->SetPreV(Box->GetV());
+		Box->X() = Box->P() - Box->GetV() * Dt;
+		Box->SetR(Box->GetQ());
 		Box->AuxilaryValue(PhysicsMaterials) = MakeSerializable(PhysicsMaterial);
-		Box->UpdateWorldSpaceState(FRigidTransform3(Box->P(), Box->Q()), FVec3(0));
+		Box->UpdateWorldSpaceState(FRigidTransform3(Box->P(), Box->GetQ()), FVec3(0));
 
 		const FReal InitialBoxZ = Box->X().Z;
 
@@ -206,8 +206,8 @@ namespace ChaosTest {
 		EXPECT_NEAR(Box->P().Z, ExpectedBoxZ, 1.e-2f);
 
 		// Velocity is below the restitution threshold, so expecting 0 velocity despite the fact that restitution is 1
-		EXPECT_TRUE(Box->V().Equals(FVec3(0)));
-		EXPECT_TRUE(Box->W().Equals(FVec3(0)));
+		EXPECT_TRUE(Box->GetV().Equals(FVec3(0)));
+		EXPECT_TRUE(Box->GetW().Equals(FVec3(0)));
 	}
 
 	void CollisionConvexConvex()
@@ -230,13 +230,13 @@ namespace ChaosTest {
 		auto Floor = AppendStaticConvexFloor(Particles);
 		auto Box = AppendDynamicParticleConvexBox( Particles, FVec3(50) );
 		Box->X() = FVec3(0, 0, 49);
-		Box->R() = FRotation3(FQuat::Identity);
-		Box->V() = FVec3(0, 0, -1);
-		Box->PreV() = Box->V();
+		Box->SetR(FRotation3(FQuat::Identity));
+		Box->SetV(FVec3(0, 0, -1));
+		Box->SetPreV(Box->GetV());
 		Box->P() = Box->X();
-		Box->Q() = Box->R();
+		Box->SetQ(Box->GetR());
 		Box->AuxilaryValue(PhysicsMaterials) = MakeSerializable(PhysicsMaterial);
-		Box->UpdateWorldSpaceState(FRigidTransform3(Box->P(), Box->Q()), FVec3(0));
+		Box->UpdateWorldSpaceState(FRigidTransform3(Box->P(), Box->GetQ()), FVec3(0));
 
 		const FReal Dt = FReal(1) / FReal(24.);
 
@@ -303,13 +303,13 @@ namespace ChaosTest {
 		auto Floor = AppendStaticAnalyticFloor(Particles);
 		auto Box = AppendDynamicParticleBox(Particles);
 		Box->P() = FVec3(0, 1, 0);
-		Box->Q() = FRotation3(FQuat::Identity);
-		Box->V() = FVec3(0, 0, -1);
-		Box->PreV() = Box->V();
-		Box->X() = Box->P() - Box->V() * Dt;
-		Box->R() = Box->R();
+		Box->SetQ(FRotation3(FQuat::Identity));
+		Box->SetV(FVec3(0, 0, -1));
+		Box->SetPreV(Box->GetV());
+		Box->X() = Box->P() - Box->GetV() * Dt;
+		Box->SetR(Box->GetR());
 		Box->AuxilaryValue(PhysicsMaterials) = MakeSerializable(PhysicsMaterial);
-		Box->UpdateWorldSpaceState(FRigidTransform3(Box->P(), Box->Q()), FVec3(0));
+		Box->UpdateWorldSpaceState(FRigidTransform3(Box->P(), Box->GetQ()), FVec3(0));
 
 		FReal InitialBoxZ = Box->X().Z;
 
@@ -347,8 +347,8 @@ namespace ChaosTest {
 		Collisions.ScatterOutput(Dt);
 
 		// 0 restitution so expecting 0 velocity
-		EXPECT_TRUE(Box->V().Equals(FVec3(0)));
-		EXPECT_TRUE(Box->W().Equals(FVec3(0)));
+		EXPECT_TRUE(Box->GetV().Equals(FVec3(0)));
+		EXPECT_TRUE(Box->GetW().Equals(FVec3(0)));
 
 		// Box will not move because the default depentration velocity is zero
 		const FReal ExpectedBoxZ = InitialBoxZ + Collisions.CollisionConstraints.GetSolverSettings().DepenetrationVelocity * Dt;
@@ -374,13 +374,13 @@ namespace ChaosTest {
 		auto Floor = AppendStaticAnalyticFloor(Particles);
 		auto Box = AppendDynamicParticleBox(Particles);
 		Box->P() = FVec3(0, 0, 0);
-		Box->Q() = FRotation3(FQuat::Identity);
-		Box->V() = FVec3(0, 0, -100);
-		Box->PreV() = Box->V();
-		Box->X() = Box->P() - Box->V() * Dt;
-		Box->R() = Box->Q();
+		Box->SetQ(FRotation3(FQuat::Identity));
+		Box->SetV(FVec3(0, 0, -100));
+		Box->SetPreV(Box->GetV());
+		Box->X() = Box->P() - Box->GetV() * Dt;
+		Box->SetR(Box->GetQ());
 		Box->AuxilaryValue(PhysicsMaterials) = MakeSerializable(PhysicsMaterial);
-		Box->UpdateWorldSpaceState(FRigidTransform3(Box->P(), Box->Q()), FVec3(0));
+		Box->UpdateWorldSpaceState(FRigidTransform3(Box->P(), Box->GetQ()), FVec3(0));
 
 		FPBDCollisionConstraintAccessor Collisions(Particles, Collided, PhysicsMaterials, PerParticlePhysicsMaterials, 2, 5);
 
@@ -420,12 +420,12 @@ namespace ChaosTest {
 		Collisions.ScatterOutput(Dt);
 
 		// full restitution, so expecting negative velocity
-		EXPECT_TRUE(Box->V().Equals(FVec3(0.f, 0.f, 100.f)));
-		EXPECT_TRUE(Box->W().Equals(FVec3(0)));
+		EXPECT_TRUE(Box->GetV().Equals(FVec3(0.f, 0.f, 100.f)));
+		EXPECT_TRUE(Box->GetW().Equals(FVec3(0)));
 
 		// should end up outside the plane
 		EXPECT_GE(Box->P().Z, -Box->GetGeometry()->BoundingBox().Min().Z);
-		EXPECT_TRUE(Box->Q().Equals(FQuat::Identity));
+		EXPECT_TRUE(Box->GetQ().Equals(FQuat::Identity));
 	}
 
 	// This test will make sure that a dynamic cube colliding with a static floor will have the correct bounce velocity
@@ -448,17 +448,17 @@ namespace ChaosTest {
 
 		FGeometryParticleHandle* StaticCube = AppendStaticParticleBox(Particles, FVec3(100.0f));
 		StaticCube->X() = FVec3(0, 0, -50.0f);
-		StaticCube->UpdateWorldSpaceState(FRigidTransform3(StaticCube->X(), StaticCube->R()), FVec3(0));
+		StaticCube->UpdateWorldSpaceState(FRigidTransform3(StaticCube->X(), StaticCube->GetR()), FVec3(0));
 
 		FPBDRigidParticleHandle* DynamicCube = AppendDynamicParticleBox(Particles, FVec3(100.0f));
 		DynamicCube->X() = FVec3(0, 0, 50);
-		DynamicCube->R() = FRotation3::FromIdentity();
-		DynamicCube->V() = FVec3(0, 0, -100);
-		DynamicCube->PreV() = DynamicCube->V();
-		DynamicCube->P() = DynamicCube->X() + DynamicCube->V() * Dt;
-		DynamicCube->Q() = DynamicCube->R();
+		DynamicCube->SetR(FRotation3::FromIdentity());
+		DynamicCube->SetV(FVec3(0, 0, -100));
+		DynamicCube->SetPreV(DynamicCube->GetV());
+		DynamicCube->P() = DynamicCube->X() + DynamicCube->GetV() * Dt;
+		DynamicCube->SetQ(DynamicCube->GetR());
 		DynamicCube->AuxilaryValue(PhysicsMaterials) = MakeSerializable(PhysicsMaterial);
-		DynamicCube->UpdateWorldSpaceState(FRigidTransform3(DynamicCube->P(), DynamicCube->Q()), FVec3(0));
+		DynamicCube->UpdateWorldSpaceState(FRigidTransform3(DynamicCube->P(), DynamicCube->GetQ()), FVec3(0));
 
 		FPBDCollisionConstraintAccessor Collisions(Particles, Collided, PhysicsMaterials, PerParticlePhysicsMaterials, 2, 5);
 
@@ -493,9 +493,9 @@ namespace ChaosTest {
 		Collisions.ScatterOutput(Dt);
 
 		// This test's tolerances are set to be very crude as to not be over sensitive (for now)
-		EXPECT_NEAR(DynamicCube->V().Z, 50.0f, 5.0f);  // restitution not too low
-		EXPECT_TRUE(FMath::Abs(DynamicCube->V().X) < 1.0f);
-		EXPECT_TRUE(FMath::Abs(DynamicCube->V().Y) < 1.0f);
+		EXPECT_NEAR(DynamicCube->GetV().Z, 50.0f, 5.0f);  // restitution not too low
+		EXPECT_TRUE(FMath::Abs(DynamicCube->GetV().X) < 1.0f);
+		EXPECT_TRUE(FMath::Abs(DynamicCube->GetV().Y) < 1.0f);
 	}
 
 	void CollisionBoxToStaticBox()
@@ -515,19 +515,19 @@ namespace ChaosTest {
 		auto StaticBox = AppendStaticParticleBox(Particles);
 		StaticBox->X() = FVec3(-0.05f, -0.05f, -0.1f);
 		StaticBox->AuxilaryValue(PhysicsMaterials) = MakeSerializable(PhysicsMaterial);
-		StaticBox->UpdateWorldSpaceState(FRigidTransform3(StaticBox->X(), StaticBox->R()), FVec3(0));
+		StaticBox->UpdateWorldSpaceState(FRigidTransform3(StaticBox->X(), StaticBox->GetR()), FVec3(0));
 
 		FReal Dt = FReal(1) / FReal(24.);
 
 		auto Box2 = AppendDynamicParticleBox(Particles);
 		FVec3 StartingPoint(0.5f);
 		Box2->P() = StartingPoint;
-		Box2->Q() = Box2->R();
-		Box2->V() = FVec3(0, 0, -1);
-		Box2->PreV() = Box2->V();
-		Box2->X() = Box2->P() - Box2->V() * Dt;
+		Box2->SetQ(Box2->GetR());
+		Box2->SetV(FVec3(0, 0, -1));
+		Box2->SetPreV(Box2->GetV());
+		Box2->X() = Box2->P() - Box2->GetV() * Dt;
 		Box2->AuxilaryValue(PhysicsMaterials) = MakeSerializable(PhysicsMaterial);
-		Box2->UpdateWorldSpaceState(FRigidTransform3(Box2->P(), Box2->Q()), FVec3(0));
+		Box2->UpdateWorldSpaceState(FRigidTransform3(Box2->P(), Box2->GetQ()), FVec3(0));
 
 		FBox Region(FVector(FReal(.2)), FVector(FReal(.5)));
 
@@ -550,7 +550,7 @@ namespace ChaosTest {
 		EXPECT_TRUE(FMath::Abs(Constraint.GetPhi() - FReal(-0.4)) < SMALL_THRESHOLD);
 
 
-		EXPECT_TRUE(FMath::Abs(Box2->V().Size() - 1.f)<SMALL_THRESHOLD ); // no velocity change yet  
+		EXPECT_TRUE(FMath::Abs(Box2->GetV().Size() - 1.f)<SMALL_THRESHOLD ); // no velocity change yet  
 
 		{
 			INVARIANT_XR_START(Box2);
@@ -574,11 +574,11 @@ namespace ChaosTest {
 
 		Collisions.ScatterOutput(Dt);
 
-		EXPECT_TRUE(Box2->V().Size() < FVector(0, -1, 0).Size()); // slowed down  
-		EXPECT_TRUE(Box2->W().Size() > 0); // now has rotation 
+		EXPECT_TRUE(Box2->GetV().Size() < FVector(0, -1, 0).Size()); // slowed down  
+		EXPECT_TRUE(Box2->GetW().Size() > 0); // now has rotation 
 
 		EXPECT_FALSE(Box2->P().Equals(StartingPoint)); // moved
-		EXPECT_FALSE(Box2->Q().Equals(FQuat::Identity)); // and rotated
+		EXPECT_FALSE(Box2->GetQ().Equals(FQuat::Identity)); // and rotated
 	}
 }
 
diff --git a/Engine/Source/Programs/HeadlessChaos/Private/HeadlessChaosTestConstraints.cpp b/Engine/Source/Programs/HeadlessChaos/Private/HeadlessChaosTestConstraints.cpp
--- a/Engine/Source/Programs/HeadlessChaos/Private/HeadlessChaosTestConstraints.cpp
+++ b/Engine/Source/Programs/HeadlessChaos/Private/HeadlessChaosTestConstraints.cpp
@@ -179,7 +179,7 @@ namespace ChaosTest {
 			Evolution.EndFrame(0.1);
 
 			const FVec3& Pos = Evolution.GetParticleHandles().Handle(0)->X();
-			const FRotation3& Rot = Evolution.GetParticleHandles().Handle(0)->R();
+			const FRotation3& Rot = Evolution.GetParticleHandles().Handle(0)->GetR();
 
 			//UE_LOG(LogChaos, Warning, TEXT("Pos %s"), *Pos.ToString());
 			//UE_LOG(LogChaos, Warning, TEXT("Rot %s"), *Rot.ToString());
@@ -227,7 +227,7 @@ namespace ChaosTest {
 			Evolution.EnableParticle(DynamicParticle);
 
 			const FVec3& Pos = Evolution.GetParticleHandles().Handle(0)->X();
-			const FRotation3& Rot = Evolution.GetParticleHandles().Handle(0)->R();
+			const FRotation3& Rot = Evolution.GetParticleHandles().Handle(0)->GetR();
 
 			Evolution.AdvanceOneTimeStep(0.1);
 			Evolution.EndFrame(0.1);
@@ -304,7 +304,7 @@ namespace ChaosTest {
 		Evolution.EnableParticle(DynamicParticle);
 
 		const FVec3& Pos = Evolution.GetParticleHandles().Handle(0)->X();
-		const FRotation3& Rot = Evolution.GetParticleHandles().Handle(0)->R();
+		const FRotation3& Rot = Evolution.GetParticleHandles().Handle(0)->GetR();
 		const float DeltaTime = 1.0f / 30.0f;
 
 		const float PositionTolerance = KINDA_SMALL_NUMBER;
diff --git a/Engine/Source/Programs/HeadlessChaos/Private/HeadlessChaosTestJointForces.cpp b/Engine/Source/Programs/HeadlessChaos/Private/HeadlessChaosTestJointForces.cpp
--- a/Engine/Source/Programs/HeadlessChaos/Private/HeadlessChaosTestJointForces.cpp
+++ b/Engine/Source/Programs/HeadlessChaos/Private/HeadlessChaosTestJointForces.cpp
@@ -203,7 +203,7 @@ namespace ChaosTest {
 		FGenericParticleHandle P1 = Test.GetParticle(1);
 
 		// Reposition the particle to have some extension in the spring
-		P1->InitTransform(P1->X() + FVec3(0,0,-Extension), P1->R());
+		P1->InitTransform(P1->X() + FVec3(0,0,-Extension), P1->GetR());
 
 		// Run the sim
 		Test.Evolution.AdvanceOneTimeStep(Dt);
diff --git a/Engine/Source/Programs/HeadlessChaos/Private/HeadlessChaosTestJointRestitution.cpp b/Engine/Source/Programs/HeadlessChaos/Private/HeadlessChaosTestJointRestitution.cpp
--- a/Engine/Source/Programs/HeadlessChaos/Private/HeadlessChaosTestJointRestitution.cpp
+++ b/Engine/Source/Programs/HeadlessChaos/Private/HeadlessChaosTestJointRestitution.cpp
@@ -50,8 +50,8 @@ namespace ChaosTest {
 		// Check that the velocity and angular velocity was affected by restitution
 		if (auto KinParticle = Test.GetParticle(1)->CastToKinematicParticle())
 		{
-			const FVec3 ResultVel = KinParticle->V();
-			const FVec3 ResultAngVel = KinParticle->W();
+			const FVec3 ResultVel = KinParticle->GetV();
+			const FVec3 ResultAngVel = KinParticle->GetW();
 
 			const FVec3 ExpectedVel = -Test.JointSettings[0].SwingRestitution * InitialVel1;
 			const FVec3 ExpectedAngVel = -Test.JointSettings[0].SwingRestitution * InitialAngVel1;
diff --git a/Engine/Source/Programs/HeadlessChaos/Private/HeadlessChaosTestJoints.cpp b/Engine/Source/Programs/HeadlessChaos/Private/HeadlessChaosTestJoints.cpp
--- a/Engine/Source/Programs/HeadlessChaos/Private/HeadlessChaosTestJoints.cpp
+++ b/Engine/Source/Programs/HeadlessChaos/Private/HeadlessChaosTestJoints.cpp
@@ -136,7 +136,7 @@ namespace ChaosTest {
 
 
 			// Joint position calculated from pose and local-space joint pos
-			const FVec3 Box2WorldSpaceJointPosition = Test.GetParticle(Box2Id)->R().RotateVector(Box2LocalSpaceJointPosition) + Test.GetParticle(Box2Id)->X();
+			const FVec3 Box2WorldSpaceJointPosition = Test.GetParticle(Box2Id)->GetR().RotateVector(Box2LocalSpaceJointPosition) + Test.GetParticle(Box2Id)->X();
 			EXPECT_LT((Box2WorldSpaceJointPosition - Test.JointPositions[0]).Size(), (FReal)0.1);
 
 			// Kinematic particle should not have moved
@@ -203,7 +203,7 @@ namespace ChaosTest {
 
 		// Move the kinematic body
 		const FVec3 RootPosition = Test.ParticlePositions[0] + RootDelta;
-		Test.Evolution.SetParticleKinematicTarget(Test.GetParticle(Box1Id)->CastToKinematicParticle(), FKinematicTarget::MakePositionTarget(RootPosition, Test.GetParticle(Box1Id)->R()));
+		Test.Evolution.SetParticleKinematicTarget(Test.GetParticle(Box1Id)->CastToKinematicParticle(), FKinematicTarget::MakePositionTarget(RootPosition, Test.GetParticle(Box1Id)->GetR()));
 
 		for (int32 i = 0; i < 1000; ++i)
 		{
@@ -220,7 +220,7 @@ namespace ChaosTest {
 			EXPECT_NEAR(Distance, ExpectedDistance, (FReal)0.15 * BoxSize) << "Post-move instability on frame " << i;
 
 			// Joint position calculted from pose and local-space joint pos
-			const FVec3 Box2WorldSpaceJointPosition = Test.GetParticle(Box2Id)->R().RotateVector(Box2LocalSpaceJointPosition) + Test.GetParticle(Box2Id)->X();
+			const FVec3 Box2WorldSpaceJointPosition = Test.GetParticle(Box2Id)->GetR().RotateVector(Box2LocalSpaceJointPosition) + Test.GetParticle(Box2Id)->X();
 			EXPECT_LT((Box2WorldSpaceJointPosition - RootPosition).Size(), (FReal)0.15 * BoxSize) << "Post-move instability on frame " << i;
 		}
 	}
@@ -280,7 +280,7 @@ namespace ChaosTest {
 			const FVec3 RootOffset = FMath::Sin((FReal)2 * PI * Time / AnimPeriod) * AnimDelta;
 			const FVec3 RootPosition = Test.ParticlePositions[0] + RootOffset;
 
-			Test.Evolution.SetParticleKinematicTarget(Test.GetParticle(Box1Id)->CastToKinematicParticle(), FKinematicTarget::MakePositionTarget(RootPosition, Test.GetParticle(Box1Id)->R()));
+			Test.Evolution.SetParticleKinematicTarget(Test.GetParticle(Box1Id)->CastToKinematicParticle(), FKinematicTarget::MakePositionTarget(RootPosition, Test.GetParticle(Box1Id)->GetR()));
 
 			Test.Evolution.AdvanceOneTimeStep(Dt);
 			Test.Evolution.EndFrame(Dt);
@@ -294,7 +294,7 @@ namespace ChaosTest {
 			EXPECT_NEAR(Distance, ExpectedDistance, (FReal)1) << "Failed on frame " << i;
 
 			// Joint position calculated from pose and local-space joint pos
-			const FVec3 Box2WorldSpaceJointPosition = Test.GetParticle(Box2Id)->R().RotateVector(Box2LocalSpaceJointPosition) + Test.GetParticle(Box2Id)->X();
+			const FVec3 Box2WorldSpaceJointPosition = Test.GetParticle(Box2Id)->GetR().RotateVector(Box2LocalSpaceJointPosition) + Test.GetParticle(Box2Id)->X();
 			EXPECT_LT((Box2WorldSpaceJointPosition - RootPosition).Size(), (FReal)1) << "Failed on frame " << i;
 		}
 	}
@@ -358,7 +358,7 @@ namespace ChaosTest {
 			const FVec3 RootOffset = FMath::Sin((FReal)2 * PI * Time / AnimPeriod) * AnimDelta;
 			const FVec3 RootPosition = Test.ParticlePositions[0] + RootOffset;
 
-			Test.Evolution.SetParticleKinematicTarget(Test.GetParticle(0)->CastToKinematicParticle(), FKinematicTarget::MakePositionTarget(RootPosition, Test.GetParticle(0)->R()));
+			Test.Evolution.SetParticleKinematicTarget(Test.GetParticle(0)->CastToKinematicParticle(), FKinematicTarget::MakePositionTarget(RootPosition, Test.GetParticle(0)->GetR()));
 
 			Test.Evolution.AdvanceOneTimeStep(Dt);
 			Test.Evolution.EndFrame(Dt);
@@ -444,7 +444,7 @@ namespace ChaosTest {
 			const FVec3 RootOffset = FMath::Sin((FReal)2 * PI * Time / AnimPeriod) * AnimDelta;
 			const FVec3 RootPosition = Test.ParticlePositions[0] + RootOffset;
 
-			Test.Evolution.SetParticleKinematicTarget(Test.GetParticle(0)->CastToKinematicParticle(), FKinematicTarget::MakePositionTarget(RootPosition, Test.GetParticle(0)->R()));
+			Test.Evolution.SetParticleKinematicTarget(Test.GetParticle(0)->CastToKinematicParticle(), FKinematicTarget::MakePositionTarget(RootPosition, Test.GetParticle(0)->GetR()));
 
 			Test.Evolution.AdvanceOneTimeStep(Dt);
 			Test.Evolution.EndFrame(Dt);
@@ -513,7 +513,7 @@ namespace ChaosTest {
 		{
 			Evolution.AdvanceOneTimeStep(Dt);
 			Evolution.EndFrame(Dt);
-			EXPECT_LT(FMath::Abs((Box2->R().RotateVector(FVec3((FReal)-100, (FReal)0, (FReal)0)) + Box2->X() - Points[0]).Size() - 300.f), 0.1);
+			EXPECT_LT(FMath::Abs((Box2->GetR().RotateVector(FVec3((FReal)-100, (FReal)0, (FReal)0)) + Box2->X() - Points[0]).Size() - 300.f), 0.1);
 		}
 	}
 
diff --git a/Engine/Source/Programs/HeadlessChaos/Private/HeadlessChaosTestMinEvolution.cpp b/Engine/Source/Programs/HeadlessChaos/Private/HeadlessChaosTestMinEvolution.cpp
--- a/Engine/Source/Programs/HeadlessChaos/Private/HeadlessChaosTestMinEvolution.cpp
+++ b/Engine/Source/Programs/HeadlessChaos/Private/HeadlessChaosTestMinEvolution.cpp
@@ -70,7 +70,7 @@ namespace ChaosTest
 		Particles[0]->InvM() = 1.0f;
 		Particles[0]->InvI() = TVec3<FRealSingle>(1.0f / 100.0f, 1.0f / 100.0f, 1.0f / 100.0f);
 		Particles[0]->AuxilaryValue(ParticlePrevXs) = Particles[0]->X();
-		Particles[0]->AuxilaryValue(ParticlePrevRs) = Particles[0]->R();
+		Particles[0]->AuxilaryValue(ParticlePrevRs) = Particles[0]->GetR();
 
 		Particles[1]->X() = FVec3(50, 0, 0);
 		Particles[1]->M() = 1.0f;
@@ -78,7 +78,7 @@ namespace ChaosTest
 		Particles[1]->InvM() = 1.0f;
 		Particles[1]->InvI() = TVec3<FRealSingle>(1.0f / 100.0f, 1.0f / 100.0f, 1.0f / 100.0f);
 		Particles[1]->AuxilaryValue(ParticlePrevXs) = Particles[1]->X();
-		Particles[1]->AuxilaryValue(ParticlePrevRs) = Particles[1]->R();
+		Particles[1]->AuxilaryValue(ParticlePrevRs) = Particles[1]->GetR();
 
 		// Spring connectors at particle centres
 		TArray<FVec3> Locations =
diff --git a/Engine/Source/Programs/HeadlessChaos/Private/HeadlessChaosTestParticleHandle.cpp b/Engine/Source/Programs/HeadlessChaos/Private/HeadlessChaosTestParticleHandle.cpp
--- a/Engine/Source/Programs/HeadlessChaos/Private/HeadlessChaosTestParticleHandle.cpp
+++ b/Engine/Source/Programs/HeadlessChaos/Private/HeadlessChaosTestParticleHandle.cpp
@@ -277,22 +277,22 @@ namespace ChaosTest {
 		EXPECT_EQ(Geometry->X()[1], 0);
 		EXPECT_EQ(Geometry->X()[2], 0);
 
-		EXPECT_EQ(KinematicGeometry->V()[0], 0);	//default constructor
-		EXPECT_EQ(KinematicGeometry->V()[1], 0);
-		EXPECT_EQ(KinematicGeometry->V()[2], 0);
+		EXPECT_EQ(KinematicGeometry->GetV()[0], 0);	//default constructor
+		EXPECT_EQ(KinematicGeometry->GetV()[1], 0);
+		EXPECT_EQ(KinematicGeometry->GetV()[2], 0);
 
 		EXPECT_EQ(PBDRigid->X()[0], 0);	//default constructor of base
 		EXPECT_EQ(PBDRigid->X()[1], 0);
 		EXPECT_EQ(PBDRigid->X()[2], 0);
-		EXPECT_EQ(PBDRigid->V()[0], 0);
-		EXPECT_EQ(PBDRigid->V()[1], 0);
-		EXPECT_EQ(PBDRigid->V()[2], 0);
+		EXPECT_EQ(PBDRigid->GetV()[0], 0);
+		EXPECT_EQ(PBDRigid->GetV()[1], 0);
+		EXPECT_EQ(PBDRigid->GetV()[2], 0);
 		EXPECT_EQ(PBDRigid->M(), 1);
 
 		PBDRigid->SetX(FVec3(1, 2, 3));
 		EXPECT_EQ(PBDRigid->X()[0], 1);
 		KinematicGeometry->SetV(FVec3(3, 3, 3));
-		EXPECT_EQ(KinematicGeometry->V()[0], 3);
+		EXPECT_EQ(KinematicGeometry->GetV()[0], 3);
 
 		EXPECT_EQ(Geometry->ObjectState(), EObjectStateType::Static);
 		EXPECT_EQ(KinematicGeometry->ObjectState(), EObjectStateType::Kinematic);
diff --git a/Engine/Source/Programs/HeadlessChaos/Private/HeadlessChaosTestPhysicsScene.cpp b/Engine/Source/Programs/HeadlessChaos/Private/HeadlessChaosTestPhysicsScene.cpp
--- a/Engine/Source/Programs/HeadlessChaos/Private/HeadlessChaosTestPhysicsScene.cpp
+++ b/Engine/Source/Programs/HeadlessChaos/Private/HeadlessChaosTestPhysicsScene.cpp
@@ -89,7 +89,7 @@ namespace ChaosTest {
 				}
 				else
 				{
-					FTransform ParticleTransform(Data.Payload.GetGeometryParticleHandle_PhysicsThread()->R(), Data.Payload.GetGeometryParticleHandle_PhysicsThread()->X());
+					FTransform ParticleTransform(Data.Payload.GetGeometryParticleHandle_PhysicsThread()->GetR(), Data.Payload.GetGeometryParticleHandle_PhysicsThread()->X());
 					const FVec3 DirLocal = ParticleTransform.InverseTransformVectorNoScale(CurData.Dir);
 					const FVec3 StartLocal = ParticleTransform.InverseTransformPositionNoScale(Start);
 					bHit = Data.Payload.GetGeometryParticleHandle_PhysicsThread()->GetGeometry()->Raycast(StartLocal, DirLocal, CurData.CurrentLength, 0, OutTime, OutPos, OutNorm, FaceIdx);
diff --git a/Engine/Source/Programs/HeadlessChaos/Private/HeadlessChaosTestProbes.cpp b/Engine/Source/Programs/HeadlessChaos/Private/HeadlessChaosTestProbes.cpp
--- a/Engine/Source/Programs/HeadlessChaos/Private/HeadlessChaosTestProbes.cpp
+++ b/Engine/Source/Programs/HeadlessChaos/Private/HeadlessChaosTestProbes.cpp
@@ -45,11 +45,11 @@ namespace ChaosTest
 		// Positions
 		Static->X() = FVec3(0, 0, 0);
 		Dynamic->X() = FVec3(-(2 * BoxHalfSize) - Separation, 0, 0);
-		Dynamic->V() = FVec3(InitialSpeed, 0, 0);
+		Dynamic->SetV(FVec3(InitialSpeed, 0, 0));
 		Dynamic->SetCCDEnabled(false);
 
 		// The position of the static has changed and statics don't automatically update bounds, so update explicitly
-		Static->UpdateWorldSpaceState(TRigidTransform<FReal, 3>(Static->X(), Static->R()), FVec3(0));
+		Static->UpdateWorldSpaceState(TRigidTransform<FReal, 3>(Static->X(), Static->GetR()), FVec3(0));
 
 		// Make sure the particles would collide if Dynamic wasn't a probe
 		::ChaosTest::SetParticleSimDataToCollide({ Static,Dynamic });
@@ -65,7 +65,7 @@ namespace ChaosTest
 		EXPECT_GT(Evolution.GetCollisionConstraints().NumConstraints(), 0);
 
 		// Make sure that the velocity of the dynamic body wasn't affected
-		EXPECT_EQ(Dynamic->V(), FVec3(InitialSpeed, 0, 0));
+		EXPECT_EQ(Dynamic->GetV(), FVec3(InitialSpeed, 0, 0));
 	}
 
 	GTEST_TEST(ProbeTests, ProbeBodyConstraintWithCCD)
@@ -96,11 +96,11 @@ namespace ChaosTest
 		// Positions
 		Static->X() = FVec3(0, 0, 0);
 		Dynamic->X() = FVec3(-(2 * BoxHalfSize) - Separation, 0, 0);
-		Dynamic->V() = FVec3(InitialSpeed, 0, 0);
+		Dynamic->SetV(FVec3(InitialSpeed, 0, 0));
 		Dynamic->SetCCDEnabled(true);
 
 		// The position of the static has changed and statics don't automatically update bounds, so update explicitly
-		Static->UpdateWorldSpaceState(TRigidTransform<FReal, 3>(Static->X(), Static->R()), FVec3(0));
+		Static->UpdateWorldSpaceState(TRigidTransform<FReal, 3>(Static->X(), Static->GetR()), FVec3(0));
 
 		// Make sure the particles would collide if Dynamic wasn't a probe
 		::ChaosTest::SetParticleSimDataToCollide({ Static,Dynamic });
@@ -123,6 +123,6 @@ namespace ChaosTest
 		}
 
 		// Make sure that the velocity of the dynamic body wasn't affected
-		EXPECT_EQ(Dynamic->V(), FVec3(InitialSpeed, 0, 0));
+		EXPECT_EQ(Dynamic->GetV(), FVec3(InitialSpeed, 0, 0));
 	}
 }
diff --git a/Engine/Source/Programs/HeadlessChaos/Private/HeadlessChaosTestRewind.cpp b/Engine/Source/Programs/HeadlessChaos/Private/HeadlessChaosTestRewind.cpp
--- a/Engine/Source/Programs/HeadlessChaos/Private/HeadlessChaosTestRewind.cpp
+++ b/Engine/Source/Programs/HeadlessChaos/Private/HeadlessChaosTestRewind.cpp
@@ -466,7 +466,7 @@ namespace ChaosTest {
 
 						for (int32 Step = 0; Step < PhysicsStep; ++Step)
 						{
-							const FReal OldV = RewindData->GetPastStateAtFrame(*Proxy->GetHandle_LowLevel(), Step).V()[2];
+							const FReal OldV = RewindData->GetPastStateAtFrame(*Proxy->GetHandle_LowLevel(), Step).GetV()[2];
 							const FReal Time = Step * SimDt;
 							if (Time == 2)	//velocity was reset by gt
 							{
@@ -968,8 +968,7 @@ namespace ChaosTest {
 			});
 	}
 
-	// DISABLED: Due to GT side changes in precision of particle data this will fail on minor differences. Disabling until fix available (PT side changes to precision)
-	GTEST_TEST(AllTraits, DISABLED_RewindTest_ResimInSync3)
+	GTEST_TEST(AllTraits, RewindTest_ResimInSync3)
 	{
 		//different position during resim step 5 which should cause a desync
 		//want a completely clean property to make sure we properly update resim buffer when property is dirtied only on second run
@@ -1722,7 +1721,7 @@ namespace ChaosTest {
 				}
 
 				EXPECT_NEAR(ParticleState.X()[2], ExpectedXZ, 1e-4);
-				EXPECT_NEAR(ParticleState.V()[2], ExpectedVZ, 1e-4);
+				EXPECT_NEAR(ParticleState.GetV()[2], ExpectedVZ, 1e-4);
 
 				ExpectedVZ -= SimDt;
 				ExpectedXZ += ExpectedVZ * SimDt;
@@ -2315,7 +2314,7 @@ namespace ChaosTest {
 					//we'll see the teleport automatically because ResimAsFollower
 					//but it's done by solver so before tick teleport is not known
 					EXPECT_NEAR(Particle.X()[2], ExpectedXZ, 1e-4);
-					EXPECT_NEAR(Particle.V()[2], ExpectedVZ, 1e-4);
+					EXPECT_NEAR(Particle.GetV()[2], ExpectedVZ, 1e-4);
 #endif
 				}
 
@@ -2326,7 +2325,7 @@ namespace ChaosTest {
 
 #if REWIND_DESYNC
 				EXPECT_NEAR(Particle.X()[2], ExpectedXZ, 1e-4);
-				EXPECT_NEAR(Particle.V()[2], ExpectedVZ, 1e-4);
+				EXPECT_NEAR(Particle.GetV()[2], ExpectedVZ, 1e-4);
 #endif
 			}
 
@@ -3873,7 +3872,7 @@ namespace ChaosTest {
 			for (const auto& Dynamic : NonDisabledDyanmic)
 			{
 				Frame.X.Add(Dynamic.X());
-				Frame.R.Add(Dynamic.R());
+				Frame.R.Add(Dynamic.GetR());
 			}
 			History.Add(MoveTemp(Frame));
 		}
diff --git a/Engine/Source/Programs/HeadlessChaos/Private/HeadlessChaosTestRotation.cpp b/Engine/Source/Programs/HeadlessChaos/Private/HeadlessChaosTestRotation.cpp
--- a/Engine/Source/Programs/HeadlessChaos/Private/HeadlessChaosTestRotation.cpp
+++ b/Engine/Source/Programs/HeadlessChaos/Private/HeadlessChaosTestRotation.cpp
@@ -55,28 +55,29 @@ namespace ChaosTest {
 		{
 			AppendAnalyticBox(Particles, Dims);
 
-			Particles.R(ParticleIndex) = FRotation3::FromAxisAngle(RandAxis(), FMath::RandRange(-PI, PI)).GetNormalized();
-			Particles.Q(ParticleIndex) = Particles.R(ParticleIndex);
-			Particles.V(ParticleIndex) = FVec3::ZeroVector;
-			Particles.W(ParticleIndex) = AngularVelocities[ParticleIndex];
-			Particles.PreV(ParticleIndex) = Particles.V(ParticleIndex);
-			Particles.PreW(ParticleIndex) = Particles.W(ParticleIndex);
+			// Particles.SetR(ParticleIndex, FRotation3::FromAxisAngle(RandAxis(), FMath::RandRange(-PI, PI)).GetNormalized());
+			Particles.SetR(ParticleIndex, FRotation3::FromAxisAngle(RandAxis(), FMath::RandRange(-PI, PI)).GetNormalized());
+			Particles.SetQ(ParticleIndex, Particles.GetR(ParticleIndex));
+			Particles.SetV(ParticleIndex, FVec3::ZeroVector);
+			Particles.SetW(ParticleIndex, AngularVelocities[ParticleIndex]);
+			Particles.SetPreV(ParticleIndex, Particles.GetV(ParticleIndex));
+			Particles.SetPreW(ParticleIndex, Particles.GetW(ParticleIndex));
 			Particles.CenterOfMass(ParticleIndex) = FVec3(0);
 			Particles.RotationOfMass(ParticleIndex) = FRotation3::FromIdentity();
 		}
 
 		for (int32 ParticleIndex = 0; ParticleIndex < (int32)Particles.Size(); ++ParticleIndex)
 		{
-			FRotation3 R0 = Particles.Q(ParticleIndex);
+			FRotation3 R0 = Particles.GetQ(ParticleIndex);
 
 			StepRule.Apply(Particles, Dt, ParticleIndex);
 
-			FRotation3 R1 = Particles.Q(ParticleIndex);
-			FVec3 AngularVelocity = Particles.W(ParticleIndex);
+			FRotation3 R1 = Particles.GetQ(ParticleIndex);
+			FVec3 AngularVelocity = Particles.GetW(ParticleIndex);
 
 			FVec3 CalculatedAngularVelocity1 = FRotation3::CalculateAngularVelocity1(R0, R1, Dt);
 
-			// Veryify that we calculated the same angular velocity that was used to integrate the rotation
+			// Verify that we calculated the same angular velocity that was used to integrate the rotation
 			const FReal ExpectedAccuracy = AngularVelocity.Size() * (FReal)0.01;
 			FVec3 Error = CalculatedAngularVelocity1 - AngularVelocity;
 			EXPECT_NEAR(Error[0], (FReal)0, ExpectedAccuracy);
@@ -157,24 +158,24 @@ namespace ChaosTest {
 
 			Particles.X(ParticleIndex) = FVec3::ZeroVector;
 			Particles.P(ParticleIndex) = FVec3::ZeroVector;
-			Particles.R(ParticleIndex) = InitialRotations[ParticleIndex];
-			Particles.Q(ParticleIndex) = FinalRotations[ParticleIndex];
-			Particles.V(ParticleIndex) = FVec3::ZeroVector;
-			Particles.W(ParticleIndex) = FVec3::ZeroVector;
-			Particles.PreV(ParticleIndex) = Particles.V(ParticleIndex);
-			Particles.PreW(ParticleIndex) = Particles.W(ParticleIndex);
+			Particles.SetR(ParticleIndex, InitialRotations[ParticleIndex]);
+			Particles.SetQ(ParticleIndex, FinalRotations[ParticleIndex]);
+			Particles.SetV(ParticleIndex, FVec3::ZeroVector);
+			Particles.SetW(ParticleIndex, FVec3::ZeroVector);
+			Particles.SetPreV(ParticleIndex, Particles.GetV(ParticleIndex));
+			Particles.SetPreW(ParticleIndex, Particles.GetW(ParticleIndex));
 			Particles.CenterOfMass(ParticleIndex) = FVec3(0);
 			Particles.RotationOfMass(ParticleIndex) = FRotation3::FromIdentity();
 		}
 
 		for (int32 ParticleIndex = 0; ParticleIndex < (int32)Particles.Size(); ++ParticleIndex)
 		{
-			FRotation3 R0 = Particles.R(ParticleIndex);
-			FRotation3 R1 = Particles.Q(ParticleIndex);
+			FRotation3 R0 = Particles.GetR(ParticleIndex);
+			FRotation3 R1 = Particles.GetQ(ParticleIndex);
 
 			UpdateRule.Apply(Particles, Dt, ParticleIndex);
 
-			FVec3 ExpectedAngVel = Particles.W(ParticleIndex);
+			FVec3 ExpectedAngVel = Particles.GetW(ParticleIndex);
 			FVec3 CalculatedAngVel1 = FRotation3::CalculateAngularVelocity1(R0, R1, Dt);
 
 			EXPECT_NEAR(ExpectedAngVel.X, CalculatedAngVel1.X, KINDA_SMALL_NUMBER);
diff --git a/Engine/Source/Programs/HeadlessChaos/Private/HeadlessChaosTestSerialization.cpp b/Engine/Source/Programs/HeadlessChaos/Private/HeadlessChaosTestSerialization.cpp
--- a/Engine/Source/Programs/HeadlessChaos/Private/HeadlessChaosTestSerialization.cpp
+++ b/Engine/Source/Programs/HeadlessChaos/Private/HeadlessChaosTestSerialization.cpp
@@ -255,8 +255,8 @@ namespace ChaosTest
 		{
 			FGeometryParticles OriginalParticles;
 			OriginalParticles.AddParticles(2);
-			OriginalParticles.R(0) = FRotation3::Identity;
-			OriginalParticles.R(1) = FRotation3::Identity;
+			OriginalParticles.SetR(0, FRotation3::Identity);
+			OriginalParticles.SetR(1, FRotation3::Identity);
 			OriginalParticles.SetGeometry(0, OriginalSpheres[0]);
 			OriginalParticles.SetGeometry(1, OriginalSpheres[1]);
 
@@ -291,8 +291,8 @@ namespace ChaosTest
 		{
 			auto OriginalParticles = MakeUnique<FGeometryParticles>();
 			OriginalParticles->AddParticles(2);
-			OriginalParticles->R(0) = FRotation3::Identity;
-			OriginalParticles->R(1) = FRotation3::Identity;
+			OriginalParticles->SetR(0, FRotation3::Identity);
+			OriginalParticles->SetR(1, FRotation3::Identity);
 			OriginalParticles->SetGeometry(0, OriginalSpheres[0]);
 			OriginalParticles->SetGeometry(1, OriginalSpheres[1]);
 
@@ -334,14 +334,14 @@ namespace ChaosTest
 
 			FGeometryParticles OriginalParticles;
 			OriginalParticles.AddParticles(2);
-			OriginalParticles.R(0) = FRotation3::Identity;
-			OriginalParticles.R(1) = FRotation3::Identity;
+			OriginalParticles.SetR(0, FRotation3::Identity);
+			OriginalParticles.SetR(1, FRotation3::Identity);
 			OriginalParticles.SetGeometry(0, OriginalSpheres[0]);
 			OriginalParticles.SetGeometry(1, OriginalSpheres[1]);
 			OriginalParticles.X(0) = FVec3(100, 1, 2);
 			OriginalParticles.X(1) = FVec3(0, 1, 2);
-			OriginalParticles.R(0) = FRotation3::Identity;
-			OriginalParticles.R(1) = FRotation3::Identity;
+			OriginalParticles.SetR(0, FRotation3::Identity);
+			OriginalParticles.SetR(1, FRotation3::Identity);
 
 			TBoundingVolumeHierarchy<FGeometryParticles, TArray<int32>> OriginalBVH(OriginalParticles);
 
@@ -369,7 +369,7 @@ namespace ChaosTest
 			TArray<int32> FinalIntersections;
 			for (int32 Potential : PotentialIntersections)
 			{
-				FRigidTransform3 TM(SerializedParticles.X(Potential), SerializedParticles.R(Potential));
+				FRigidTransform3 TM(SerializedParticles.X(Potential), SerializedParticles.GetR(Potential));
 				const FAABB3 Bounds = SerializedParticles.GetGeometry(Potential)->BoundingBox().TransformedAABB(TM);
 				if (Bounds.Intersects(QueryBox))
 				{
@@ -394,8 +394,8 @@ namespace ChaosTest
 
 		TRigidParticles<FReal, 3> Particles;
 		Particles.AddParticles(2);
-		Particles.R(0) = FRotation3::Identity;
-		Particles.R(1) = FRotation3::Identity;
+		Particles.SetR(0, FRotation3::Identity);
+		Particles.SetR(1, FRotation3::Identity);
 
 		Particles.Acceleration(0) = F[0];
 		Particles.Acceleration(1) = F[1];
@@ -430,14 +430,14 @@ namespace ChaosTest
 
 		FGeometryParticles Particles;
 		Particles.AddParticles(3);
-		Particles.R(0) = FRotation3::Identity;
-		Particles.R(1) = FRotation3::Identity;
+		Particles.SetR(0, FRotation3::Identity);
+		Particles.SetR(1, FRotation3::Identity);
 		Particles.X(0) = FVec3(15, 1, 2);
 		Particles.X(1) = FVec3(0, 2, 2);
 		Particles.X(2) = FVec3(0, 2, 2);
-		Particles.R(0) = FRotation3::Identity;
-		Particles.R(1) = FRotation3::Identity;
-		Particles.R(2) = FRotation3::Identity;
+		Particles.SetR(0, FRotation3::Identity);
+		Particles.SetR(1, FRotation3::Identity);
+		Particles.SetR(2, FRotation3::Identity);
 		Particles.SetGeometry(0, Spheres[0]);
 		Particles.SetGeometry(1, Spheres[1]);
 		Particles.SetGeometry(2, Spheres[2]);
diff --git a/Engine/Source/Programs/HeadlessChaos/Private/HeadlessChaosTestShapePadding.cpp b/Engine/Source/Programs/HeadlessChaos/Private/HeadlessChaosTestShapePadding.cpp
--- a/Engine/Source/Programs/HeadlessChaos/Private/HeadlessChaosTestShapePadding.cpp
+++ b/Engine/Source/Programs/HeadlessChaos/Private/HeadlessChaosTestShapePadding.cpp
@@ -48,20 +48,20 @@ namespace ChaosTest {
 
 		auto Box0 = AppendDynamicParticleBoxMargin(Particles, Size, Margin0);
 		Box0->X() = FVec3(0, 0, 0);
-		Box0->R() = FRotation3(FQuat::Identity);
-		Box0->V() = FVec3(0);
-		Box0->PreV() = Box0->V();
+		Box0->SetR(FRotation3(FQuat::Identity));
+		Box0->SetV(FVec3(0));
+		Box0->SetPreV(Box0->GetV());
 		Box0->P() = Box0->X();
-		Box0->Q() = Box0->R();
+		Box0->SetQ(Box0->GetR());
 		Box0->AuxilaryValue(PhysicsMaterials) = MakeSerializable(PhysicsMaterial);
 
 		auto Box1 = AppendDynamicParticleBoxMargin(Particles, Size, Margin1);
 		Box1->X() = Delta;
-		Box1->R() = FRotation3(FQuat::Identity);
-		Box1->V() = FVec3(0);
-		Box1->PreV() = Box1->V();
+		Box1->SetR(FRotation3(FQuat::Identity));
+		Box1->SetV(FVec3(0));
+		Box1->SetPreV(Box1->GetV());
 		Box1->P() = Box1->X();
-		Box1->Q() = Box1->R();
+		Box1->SetQ(Box1->GetR());
 		Box1->AuxilaryValue(PhysicsMaterials) = MakeSerializable(PhysicsMaterial);
 
 		const FImplicitBox3* BoxImplicit0 = Box0->GetGeometry()->template GetObject<FImplicitBox3>();
@@ -167,20 +167,20 @@ namespace ChaosTest {
 
 		auto Box0 = AppendDynamicParticleConvexBoxMargin(Particles, 0.5f * Size, Margin0);
 		Box0->X() = FVec3(0, 0, 0);
-		Box0->R() = FRotation3(FQuat::Identity);
-		Box0->V() = FVec3(0);
-		Box0->PreV() = Box0->V();
+		Box0->SetR(FRotation3(FQuat::Identity));
+		Box0->SetV(FVec3(0));
+		Box0->SetPreV(Box0->GetV());
 		Box0->P() = Box0->X();
-		Box0->Q() = Box0->R();
+		Box0->SetQ(Box0->GetR());
 		Box0->AuxilaryValue(PhysicsMaterials) = MakeSerializable(PhysicsMaterial);
 
 		auto Box1 = AppendDynamicParticleConvexBoxMargin(Particles, 0.5f * Size, Margin1);
 		Box1->X() = Delta;
-		Box1->R() = FRotation3(FQuat::Identity);
-		Box1->V() = FVec3(0);
-		Box1->PreV() = Box1->V();
+		Box1->SetR(FRotation3(FQuat::Identity));
+		Box1->SetV(FVec3(0));
+		Box1->SetPreV(Box1->GetV());
 		Box1->P() = Box1->X();
-		Box1->Q() = Box1->R();
+		Box1->SetQ(Box1->GetR());
 		Box1->AuxilaryValue(PhysicsMaterials) = MakeSerializable(PhysicsMaterial);
 
 		const FImplicitConvex3* ConvexImplicit0 = Box0->GetGeometry()->template GetObject<FImplicitConvex3>();
@@ -291,11 +291,11 @@ namespace ChaosTest {
 
 		auto Box0 = AppendDynamicParticleBoxMargin(Particles, Size, Margin0);
 		Box0->X() = FVec3(0, 0, 0);
-		Box0->R() = FRotation3(FQuat::Identity);
-		Box0->V() = FVec3(0);
-		Box0->PreV() = Box0->V();
+		Box0->SetR(FRotation3(FQuat::Identity));
+		Box0->SetV(FVec3(0));
+		Box0->SetPreV(Box0->GetV());
 		Box0->P() = Box0->X();
-		Box0->Q() = Box0->R();
+		Box0->SetQ(Box0->GetR());
 		Box0->AuxilaryValue(PhysicsMaterials) = MakeSerializable(PhysicsMaterial);
 
 		const FImplicitBox3* BoxImplicit0 = Box0->GetGeometry()->template GetObject<FImplicitBox3>();
diff --git a/Engine/Source/Programs/HeadlessChaos/Private/HeadlessChaosTestSim.cpp b/Engine/Source/Programs/HeadlessChaos/Private/HeadlessChaosTestSim.cpp
--- a/Engine/Source/Programs/HeadlessChaos/Private/HeadlessChaosTestSim.cpp
+++ b/Engine/Source/Programs/HeadlessChaos/Private/HeadlessChaosTestSim.cpp
@@ -85,7 +85,7 @@ namespace ChaosTest {
 		Dynamic->InvI() = TVec3<FRealSingle>(1.0f / 100000.0f);
 
 		// The position of the static has changed and statics don't automatically update bounds, so update explicitly
-		Static->UpdateWorldSpaceState(FRigidTransform3(Static->X(), Static->R()), FVec3(0));
+		Static->UpdateWorldSpaceState(FRigidTransform3(Static->X(), Static->GetR()), FVec3(0));
 
 		::ChaosTest::SetParticleSimDataToCollide({ Static,Dynamic });
 
@@ -120,7 +120,7 @@ namespace ChaosTest {
 		Dynamic->SetGeometry(DynamicBox);
 
 		Static->X() = FVec3(10, 10, 10);
-		Static->UpdateWorldSpaceState(FRigidTransform3(Static->X(), Static->R()), FVec3(0));
+		Static->UpdateWorldSpaceState(FRigidTransform3(Static->X(), Static->GetR()), FVec3(0));
 		Dynamic->X() = FVec3(10, 10, 120);
 		Dynamic->I() = TVec3<FRealSingle>(100000.0f);
 		Dynamic->InvI() = TVec3<FRealSingle>(1.0f / 100000.0f);
@@ -166,7 +166,7 @@ namespace ChaosTest {
 		Dynamic->SetGeometry(DynamicBox);
 
 		Static->X() = FVec3(10, 10, 10);
-		Static->UpdateWorldSpaceState(FRigidTransform3(Static->X(), Static->R()), FVec3(0));
+		Static->UpdateWorldSpaceState(FRigidTransform3(Static->X(), Static->GetR()), FVec3(0));
 		Dynamic->X() = FVec3(10, 10, 300);
 		Dynamic->I() = TVec3<FRealSingle>(1);
 		Dynamic->InvI() = TVec3<FRealSingle>(1);
@@ -213,7 +213,7 @@ namespace ChaosTest {
 		Dynamic2->SetGeometry(DynamicBox);
 
 		Static->X() = FVec3(10, 10, 10);
-		Static->UpdateWorldSpaceState(FRigidTransform3(Static->X(), Static->R()), FVec3(0));
+		Static->UpdateWorldSpaceState(FRigidTransform3(Static->X(), Static->GetR()), FVec3(0));
 		Dynamic1->X() = FVec3(10, 10, 120);
 		Dynamic2->X() = FVec3(10, 10, 400);
 
diff --git a/Engine/Source/Programs/HeadlessChaos/Private/HeadlessChaosTestUtility.cpp b/Engine/Source/Programs/HeadlessChaos/Private/HeadlessChaosTestUtility.cpp
--- a/Engine/Source/Programs/HeadlessChaos/Private/HeadlessChaosTestUtility.cpp
+++ b/Engine/Source/Programs/HeadlessChaos/Private/HeadlessChaosTestUtility.cpp
@@ -26,11 +26,11 @@ namespace ChaosTest {
 		int32 RigidBodyIndex = InParticles.Size() - 1;
 
 		InParticles.X(RigidBodyIndex) = FVec3(0, 0, 0);
-		InParticles.V(RigidBodyIndex) = FVec3(0, 0, 0);
-		InParticles.R(RigidBodyIndex) = FRotation3::MakeFromEuler(FVec3(0, 0, 0)).GetNormalized();
-		InParticles.W(RigidBodyIndex) = FVec3(0, 0, 0);
+		InParticles.SetV(RigidBodyIndex, FVec3(0, 0, 0));
+		InParticles.SetR(RigidBodyIndex, FRotation3::MakeFromEuler(FVec3(0, 0, 0)).GetNormalized());
+		InParticles.SetW(RigidBodyIndex, FVec3(0, 0, 0));
 		InParticles.P(RigidBodyIndex) = InParticles.X(RigidBodyIndex);
-		InParticles.Q(RigidBodyIndex) = InParticles.R(RigidBodyIndex);
+		InParticles.SetQ(RigidBodyIndex, InParticles.GetR(RigidBodyIndex));
 
 		InParticles.M(RigidBodyIndex) = 1.0;
 		InParticles.InvM(RigidBodyIndex) = 1.0;
@@ -48,11 +48,11 @@ namespace ChaosTest {
 		int32 RigidBodyIndex = InParticles.Size() - 1;
 
 		InParticles.X(RigidBodyIndex) = FVec3(0, 0, 0);
-		InParticles.V(RigidBodyIndex) = FVec3(0, 0, 0);
-		InParticles.R(RigidBodyIndex) = FRotation3::MakeFromEuler(FVec3(0, 0, 0)).GetNormalized();
-		InParticles.W(RigidBodyIndex) = FVec3(0, 0, 0);
+		InParticles.SetV(RigidBodyIndex, FVec3(0, 0, 0));
+		InParticles.SetR(RigidBodyIndex, FRotation3::MakeFromEuler(FVec3(0, 0, 0)).GetNormalized());
+		InParticles.SetW(RigidBodyIndex, FVec3(0, 0, 0));
 		InParticles.P(RigidBodyIndex) = InParticles.X(RigidBodyIndex);
-		InParticles.Q(RigidBodyIndex) = InParticles.R(RigidBodyIndex);
+		InParticles.SetQ(RigidBodyIndex, InParticles.GetR(RigidBodyIndex));
 
 		InParticles.M(RigidBodyIndex) = 1.0;
 		InParticles.InvM(RigidBodyIndex) = 1.0;
@@ -67,16 +67,16 @@ namespace ChaosTest {
 	void InitAnalyticBox2(FKinematicGeometryParticleHandle* Particle, FVec3 Scale)
 	{
 		Particle->X() = FVec3(0, 0, 0);
-		Particle->V() = FVec3(0, 0, 0);
-		Particle->R() = FRotation3::MakeFromEuler(FVec3(0, 0, 0)).GetNormalized();
-		Particle->W() = FVec3(0, 0, 0);
+		Particle->SetV(FVec3(0, 0, 0));
+		Particle->SetR(FRotation3::MakeFromEuler(FVec3(0, 0, 0)).GetNormalized());
+		Particle->SetW(FVec3(0, 0, 0));
 		Particle->SetGeometry(MakeImplicitObjectPtr<TBox<FReal, 3>>(-Scale / 2.0, Scale / 2.0));
 
 		FPBDRigidParticleHandle* DynamicParticle = Particle->CastToRigidParticle();
 		if(DynamicParticle && DynamicParticle->ObjectState() == EObjectStateType::Dynamic)
 		{
 			DynamicParticle->P() = Particle->X();
-			DynamicParticle->Q() = Particle->R();
+			DynamicParticle->SetQ(Particle->GetR());
 
 			DynamicParticle->M() = 1.0;
 			DynamicParticle->InvM() = 1.0;
@@ -91,11 +91,11 @@ namespace ChaosTest {
 		int32 RigidBodyIndex = InParticles.Size() - 1;
 
 		InParticles.X(RigidBodyIndex) = FVec3(0, 0, 0);
-		InParticles.V(RigidBodyIndex) = FVec3(0, 0, 0);
-		InParticles.R(RigidBodyIndex) = FRotation3::MakeFromEuler(FVec3(0, 0, 0)).GetNormalized();
-		InParticles.W(RigidBodyIndex) = FVec3(0, 0, 0);
+		InParticles.SetV(RigidBodyIndex, FVec3(0, 0, 0));
+		InParticles.SetR(RigidBodyIndex, FRotation3::MakeFromEuler(FVec3(0, 0, 0)).GetNormalized());
+		InParticles.SetW(RigidBodyIndex, FVec3(0, 0, 0));
 		InParticles.P(RigidBodyIndex) = InParticles.X(RigidBodyIndex);
-		InParticles.Q(RigidBodyIndex) = InParticles.R(RigidBodyIndex);
+		InParticles.SetQ(RigidBodyIndex, InParticles.GetR(RigidBodyIndex));
 
 		check(Scale.X == Scale.Y && Scale.X == Scale.Z);
 		FReal ScaleSq = Scale.X * Scale.X;
@@ -157,12 +157,12 @@ namespace ChaosTest {
 	void InitDynamicParticleBox2(FPBDRigidParticleHandle* Particle, const FVec3& Scale, FReal Margin, TArray<TVector<int32, 3>>* OutElements)
 	{
 		Particle->X() = FVec3(0, 0, 0);
-		Particle->V() = FVec3(0, 0, 0);
-		Particle->R() = FRotation3::MakeFromEuler(FVec3(0, 0, 0)).GetNormalized();
-		Particle->W() = FVec3(0, 0, 0);
+		Particle->SetV(FVec3(0, 0, 0));
+		Particle->SetR(FRotation3::MakeFromEuler(FVec3(0, 0, 0)).GetNormalized());
+		Particle->SetW(FVec3(0, 0, 0));
 
 		Particle->P() = Particle->X();
-		Particle->Q() = Particle->R();
+		Particle->SetQ(Particle->GetR());
 
 		// Assume unit mass - this gets scaled externally by the actual mass
 		const FVec3 ScaleSq = Scale * Scale;
@@ -228,12 +228,12 @@ namespace ChaosTest {
 	void InitDynamicParticleSphere2(FPBDRigidParticleHandle* Particle, const FVec3& Scale, TArray<TVec3<int32>>* OutElements) 
 	{
 		Particle->X() = FVec3(0, 0, 0);
-		Particle->V() = FVec3(0, 0, 0);
-		Particle->R() = FRotation3::MakeFromEuler(FVec3(0, 0, 0)).GetNormalized();
-		Particle->W() = FVec3(0, 0, 0);
+		Particle->SetV(FVec3(0, 0, 0));
+		Particle->SetR(FRotation3::MakeFromEuler(FVec3(0, 0, 0)).GetNormalized());
+		Particle->SetW(FVec3(0, 0, 0));
 
 		Particle->P() = Particle->X();
-		Particle->Q() = Particle->R();
+		Particle->SetQ(Particle->GetR());
 
 		check(Scale.X == Scale.Y && Scale.X == Scale.Z);
 		FReal ScaleSq = Scale.X * Scale.X;
@@ -291,12 +291,12 @@ namespace ChaosTest {
 	void InitDynamicParticleCylinder2(FPBDRigidParticleHandle* Particle, const FVec3& Scale, TArray<TVec3<int32>>* OutElements, bool Tapered) 
 	{
 		Particle->X() = FVec3(0, 0, 0);
-		Particle->V() = FVec3(0, 0, 0);
-		Particle->R() = FRotation3::MakeFromEuler(FVec3(0, 0, 0)).GetNormalized();
-		Particle->W() = FVec3(0, 0, 0);
+		Particle->SetV(FVec3(0, 0, 0));
+		Particle->SetR(FRotation3::MakeFromEuler(FVec3(0, 0, 0)).GetNormalized());
+		Particle->SetW(FVec3(0, 0, 0));
 
 		Particle->P() = Particle->X();
-		Particle->Q() = Particle->R();
+		Particle->SetQ(Particle->GetR());
 
 		check(Scale.X == Scale.Y && Scale.X == Scale.Z);
 		FReal ScaleSq = Scale.X * Scale.X;
@@ -406,7 +406,7 @@ namespace ChaosTest {
 	void InitStaticParticleBox(FGeometryParticleHandle* Particle, const FVec3& Scale, TArray<TVec3<int32>>* OutElements)
 	{
 		Particle->X() = FVec3(0, 0, 0);
-		Particle->R() = FRotation3::FromIdentity();
+		Particle->SetR(FRotation3::FromIdentity());
 
 		FReal ScaleSq = Scale.X * Scale.X;
 
@@ -466,9 +466,9 @@ namespace ChaosTest {
 		int32 RigidBodyIndex = InParticles.Size() - 1;
 
 		InParticles.X(RigidBodyIndex) = FVec3(0, 0, 0);
-		InParticles.V(RigidBodyIndex) = FVec3(0, 0, 0);
-		InParticles.R(RigidBodyIndex) = FRotation3::MakeFromEuler(FVec3(0, 0, 0)).GetNormalized();
-		InParticles.W(RigidBodyIndex) = FVec3(0, 0, 0);
+		InParticles.SetV(RigidBodyIndex, FVec3(0, 0, 0));
+		InParticles.SetR(RigidBodyIndex, FRotation3::MakeFromEuler(FVec3(0, 0, 0)).GetNormalized());
+		InParticles.SetW(RigidBodyIndex, FVec3(0, 0, 0));
 		InParticles.M(RigidBodyIndex) = 1.0;
 		InParticles.InvM(RigidBodyIndex) = 0.0;
 		InParticles.I(RigidBodyIndex) = TVec3<FRealSingle>(1);
@@ -477,7 +477,7 @@ namespace ChaosTest {
 		InParticles.SetObjectState(RigidBodyIndex, EObjectStateType::Kinematic);
 
 		InParticles.P(RigidBodyIndex) = InParticles.X(RigidBodyIndex);
-		InParticles.Q(RigidBodyIndex) = InParticles.R(RigidBodyIndex);
+		InParticles.SetQ(RigidBodyIndex, InParticles.GetR(RigidBodyIndex));
 
 		return RigidBodyIndex;
 	}
@@ -488,9 +488,9 @@ namespace ChaosTest {
 		FKinematicGeometryParticleHandle* Particle = Particles[0];
 
 		Particle->X() = FVec3(0, 0, 0);
-		Particle->V() = FVec3(0, 0, 0);
-		Particle->R() = FRotation3::MakeFromEuler(FVec3(0, 0, 0)).GetNormalized();
-		Particle->W() = FVec3(0, 0, 0);
+		Particle->SetV(FVec3(0, 0, 0));
+		Particle->SetR(FRotation3::MakeFromEuler(FVec3(0, 0, 0)).GetNormalized());
+		Particle->SetW(FVec3(0, 0, 0));
 		Particle->SetGeometry(MakeImplicitObjectPtr<TPlane<FReal, 3>>(FVec3(0, 0, 0), FVec3(0, 0, 1)));
 
 		::ChaosTest::SetParticleSimDataToCollide({ Particle });
@@ -504,10 +504,10 @@ namespace ChaosTest {
 		TArray<FKinematicGeometryParticleHandle*> Particles = SOAs.CreateKinematicParticles(1);
 		FKinematicGeometryParticleHandle* Particle = Particles[0];
 
-		Particle->X() = FVec3(0, 0, 0);
-		Particle->V() = FVec3(0, 0, 0);
-		Particle->R() = FRotation3::MakeFromEuler(FVec3(0, 0, 0)).GetNormalized();
-		Particle->W() = FVec3(0, 0, 0);
+		Particle->SetX(FVec3(0, 0, 0));
+		Particle->SetV(FVec3(0, 0, 0));
+		Particle->SetR(FRotation3::MakeFromEuler(FVec3(0, 0, 0)).GetNormalized());
+		Particle->SetW(FVec3(0, 0, 0));
 
 		TArray<Chaos::FConvex::FVec3Type> Cube;
 		Cube.SetNum(9);
@@ -571,11 +571,11 @@ namespace ChaosTest {
 		Cube[8] = FConvex::FVec3Type{ 0, 0, 0 };
 
 		InParticles.X() = FVec3(0, 0, 0);
-		InParticles.V() = FVec3(0, 0, 0);
-		InParticles.R() = FRotation3::MakeFromEuler(FVec3(0, 0, 0)).GetNormalized();
-		InParticles.W() = FVec3(0, 0, 0);
+		InParticles.SetV(FVec3(0, 0, 0));
+		InParticles.SetR(FRotation3::MakeFromEuler(FVec3(0, 0, 0)).GetNormalized());
+		InParticles.SetW(FVec3(0, 0, 0));
 		InParticles.P() = InParticles.X();
-		InParticles.Q() = InParticles.R();
+		InParticles.SetQ(InParticles.GetR());
 
 		// TODO: Change this error prone API to set bounds more automatically. This is easy to forget
 		InParticles.SetLocalBounds(FAABB3(Cube[0], Cube[7]));
@@ -614,23 +614,23 @@ namespace ChaosTest {
 	/**/
 	FVec3 ObjectSpacePoint(FPBDRigidParticles& InParticles, const int32 Index, const FVec3& WorldSpacePoint)
 	{
-		FRigidTransform3 LocalToWorld(InParticles.X(Index), InParticles.R(Index));
+		FRigidTransform3 LocalToWorld(InParticles.X(Index), InParticles.GetR(Index));
 		return LocalToWorld.InverseTransformPosition(WorldSpacePoint);
 	}
 
 	FVec3 ObjectSpacePoint(FGeometryParticleHandle& Particle, const FVec3& WorldSpacePoint)
 	{
-		FRigidTransform3 LocalToWorld(Particle.X(), Particle.R());
+		FRigidTransform3 LocalToWorld(Particle.X(), Particle.GetR());
 		return LocalToWorld.InverseTransformPosition(WorldSpacePoint);
 	}
 
 	/**/
 	FReal PhiWithNormal(FPBDRigidParticles& InParticles, const int32 Index, const FVec3& WorldSpacePoint, FVec3& Normal)
 	{
-		FRigidTransform3(InParticles.X(Index), InParticles.R(Index));
+		FRigidTransform3(InParticles.X(Index), InParticles.GetR(Index));
 		FVec3 BodySpacePoint = ObjectSpacePoint(InParticles, Index, WorldSpacePoint);
 		FReal LocalPhi = InParticles.GetGeometry(Index)->PhiWithNormal(BodySpacePoint, Normal);
-		Normal = FRigidTransform3(InParticles.X(Index), InParticles.R(Index)).TransformVector(Normal);
+		Normal = FRigidTransform3(InParticles.X(Index), InParticles.GetR(Index)).TransformVector(Normal);
 		return LocalPhi;
 	}
 
@@ -644,10 +644,10 @@ namespace ChaosTest {
 	/**/
 	FReal PhiWithNormal(FGeometryParticleHandle& Particle, const FVec3& WorldSpacePoint, FVec3& Normal)
 	{
-		FRigidTransform3(Particle.X(), Particle.R());
+		FRigidTransform3(Particle.X(), Particle.GetR());
 		FVec3 BodySpacePoint = ObjectSpacePoint(Particle, WorldSpacePoint);
 		FReal LocalPhi = Particle.GetGeometry()->PhiWithNormal(BodySpacePoint, Normal);
-		Normal = FRigidTransform3(Particle.X(), Particle.R()).TransformVector(Normal);
+		Normal = FRigidTransform3(Particle.X(), Particle.GetR()).TransformVector(Normal);
 		return LocalPhi;
 	}
 
diff --git a/Engine/Source/Programs/HeadlessChaos/Private/HeadlessChaosTestVehicles.cpp b/Engine/Source/Programs/HeadlessChaos/Private/HeadlessChaosTestVehicles.cpp
--- a/Engine/Source/Programs/HeadlessChaos/Private/HeadlessChaosTestVehicles.cpp
+++ b/Engine/Source/Programs/HeadlessChaos/Private/HeadlessChaosTestVehicles.cpp
@@ -1444,7 +1444,7 @@ namespace ChaosTest
 	{
 		const Chaos::FVec3 COM = Rigid ? Chaos::FParticleUtilitiesGT::GetCoMWorldPosition(Rigid) : Chaos::FParticleUtilitiesGT::GetActorWorldTransform(Rigid).GetTranslation();
 		const Chaos::FVec3 Diff = InPoint - COM;
-		return Rigid->V() - Chaos::FVec3::CrossProduct(Diff, Rigid->W());
+		return Rigid->GetV() - Chaos::FVec3::CrossProduct(Diff, Rigid->GetW());
 
 	}
 
@@ -1532,7 +1532,7 @@ namespace ChaosTest
 		for (int i = 0; i < 500; ++i)
 		{
 			// latest body transform
-			const FTransform BodyTM(Dynamic->R(), Dynamic->X());
+			const FTransform BodyTM(Dynamic->GetR(), Dynamic->X());
 
 			for (int SpringIdx = 0; SpringIdx < 4; SpringIdx++)
 			{
diff --git a/Engine/Source/Programs/HeadlessChaos/Public/HeadlessChaosTestConstraints.h b/Engine/Source/Programs/HeadlessChaos/Public/HeadlessChaosTestConstraints.h
--- a/Engine/Source/Programs/HeadlessChaos/Public/HeadlessChaosTestConstraints.h
+++ b/Engine/Source/Programs/HeadlessChaos/Public/HeadlessChaosTestConstraints.h
@@ -81,16 +81,16 @@ namespace ChaosTest
 		void ResetParticle(FGeometryParticleHandle* Particle, const FVec3& Position, const FRotation3& Rotation, const FVec3& Velocity, const FVec3& AngularVelocity)
 		{
 			Particle->X() = Position;
-			Particle->R() = Rotation;
+			Particle->SetR(Rotation);
 			if (FKinematicGeometryParticleHandle* KinParticle = Particle->CastToKinematicParticle())
 			{
-				KinParticle->V() = Velocity;
-				KinParticle->W() = AngularVelocity;
+				KinParticle->SetV(Velocity);
+				KinParticle->SetW(AngularVelocity);
 			}
 			if (FPBDRigidParticleHandle* PBDParticle = Particle->CastToRigidParticle())
 			{
 				PBDParticle->P() = Position;
-				PBDParticle->Q() = Rotation;
+				PBDParticle->SetQ(Rotation);
 			}
 		}
 
diff --git a/Engine/Source/Programs/HeadlessChaos/Public/HeadlessChaosTestEvolution.h b/Engine/Source/Programs/HeadlessChaos/Public/HeadlessChaosTestEvolution.h
--- a/Engine/Source/Programs/HeadlessChaos/Public/HeadlessChaosTestEvolution.h
+++ b/Engine/Source/Programs/HeadlessChaos/Public/HeadlessChaosTestEvolution.h
@@ -159,16 +159,16 @@ namespace ChaosTest
 		void ResetParticle(FGeometryParticleHandle* Particle, const FVec3& Position, const FRotation3& Rotation, const FVec3& Velocity, const FVec3& AngularVelocity)
 		{
 			Particle->X() = Position;
-			Particle->R() = Rotation;
+			Particle->SetR(Rotation);
 			if (FKinematicGeometryParticleHandle* KinParticle = Particle->CastToKinematicParticle())
 			{
-				KinParticle->V() = Velocity;
-				KinParticle->W() = AngularVelocity;
+				KinParticle->SetV(Velocity);
+				KinParticle->SetW(AngularVelocity);
 			}
 			if (FPBDRigidParticleHandle* PBDParticle = Particle->CastToRigidParticle())
 			{
 				PBDParticle->P() = Position;
-				PBDParticle->Q() = Rotation;
+				PBDParticle->SetQ(Rotation);
 			}
 		}
 
diff --git a/Engine/Source/Runtime/Core/Public/UObject/UE5ReleaseStreamObjectVersion.h b/Engine/Source/Runtime/Core/Public/UObject/UE5ReleaseStreamObjectVersion.h
--- a/Engine/Source/Runtime/Core/Public/UObject/UE5ReleaseStreamObjectVersion.h
+++ b/Engine/Source/Runtime/Core/Public/UObject/UE5ReleaseStreamObjectVersion.h
@@ -159,6 +159,9 @@ struct FUE5ReleaseStreamObjectVersion
 
 		// Fix Chaos Cloth fictitious angular scale bug that requires existing parameter rescaling.
 		ChaosClothFictitiousAngularVelocitySubframeFix,
+
+		// Store physics thread particles data in single precision
+		SinglePrecisonParticleDataPT,
 		
 		// -----<new versions can be added above this line>-------------------------------------------------
 		VersionPlusOne,
diff --git a/Engine/Source/Runtime/Engine/Private/Collision/CollisionConversions.cpp b/Engine/Source/Runtime/Engine/Private/Collision/CollisionConversions.cpp
--- a/Engine/Source/Runtime/Engine/Private/Collision/CollisionConversions.cpp
+++ b/Engine/Source/Runtime/Engine/Private/Collision/CollisionConversions.cpp
@@ -110,7 +110,7 @@ static FVector FindGeomOpposingNormal(ECollisionShapeType QueryGeomType, const T
 			}
 			else
 			{
-				const FTransform ActorTM(Hit.Actor->R(), Hit.Actor->X());
+				const FTransform ActorTM(Hit.Actor->GetR(), Hit.Actor->X());
 				const FVector LocalInNormal = ActorTM.InverseTransformVectorNoScale(InNormal);
 				const FVector LocalTraceDirectionDenorm = ActorTM.InverseTransformVectorNoScale(TraceDirectionDenorm);
 				const FVector LocalNormal = Shape->GetGeometry()->FindGeometryOpposingNormal(LocalTraceDirectionDenorm, Hit.FaceIndex, LocalInNormal);
diff --git a/Engine/Source/Runtime/Engine/Private/PhysicsEngine/Experimental/PhysScene_Chaos.cpp b/Engine/Source/Runtime/Engine/Private/PhysicsEngine/Experimental/PhysScene_Chaos.cpp
--- a/Engine/Source/Runtime/Engine/Private/PhysicsEngine/Experimental/PhysScene_Chaos.cpp
+++ b/Engine/Source/Runtime/Engine/Private/PhysicsEngine/Experimental/PhysScene_Chaos.cpp
@@ -469,9 +469,9 @@ static TUniquePtr<FPhysScene_ChaosPauseHandler> PhysScene_ChaosPauseHandler;
 static void CopyParticleData(Chaos::FPBDRigidParticles& ToParticles, const int32 ToIndex, Chaos::FPBDRigidParticles& FromParticles, const int32 FromIndex)
 {
 	ToParticles.X(ToIndex) = FromParticles.X(FromIndex);
-	ToParticles.R(ToIndex) = FromParticles.R(FromIndex);
-	ToParticles.V(ToIndex) = FromParticles.V(FromIndex);
-	ToParticles.W(ToIndex) = FromParticles.W(FromIndex);
+	ToParticles.SetR(ToIndex, FromParticles.GetR(FromIndex));
+	ToParticles.SetV(ToIndex, FromParticles.GetV(FromIndex));
+	ToParticles.SetW(ToIndex, FromParticles.GetW(FromIndex));
 	ToParticles.M(ToIndex) = FromParticles.M(FromIndex);
 	ToParticles.InvM(ToIndex) = FromParticles.InvM(FromIndex);
 	ToParticles.I(ToIndex) = FromParticles.I(FromIndex);
@@ -1360,9 +1360,9 @@ void FPhysScene_Chaos::PopulateReplicationCache(const int32 PhysicsStep)
 		{
 			FRigidBodyState& ReplicationState = ReplicationData.GetState();
 			ReplicationState.Position = Handle->X();
-			ReplicationState.Quaternion = Handle->R();
-			ReplicationState.LinVel = Handle->V();
-			ReplicationState.AngVel = Handle->W();
+			ReplicationState.Quaternion = Handle->GetR();
+			ReplicationState.LinVel = Handle->GetV();
+			ReplicationState.AngVel = Handle->GetW();
 			ReplicationState.Flags = Handle->ObjectState() == Chaos::EObjectStateType::Sleeping ? ERigidBodyFlags::Sleeping : 0;
 			StateWasCached = true;
 		}
diff --git a/Engine/Source/Runtime/Engine/Private/PhysicsEngine/ImmediatePhysics/ImmediatePhysicsChaos/ImmediatePhysicsActorHandle_Chaos.cpp b/Engine/Source/Runtime/Engine/Private/PhysicsEngine/ImmediatePhysics/ImmediatePhysicsChaos/ImmediatePhysicsActorHandle_Chaos.cpp
--- a/Engine/Source/Runtime/Engine/Private/PhysicsEngine/ImmediatePhysics/ImmediatePhysicsChaos/ImmediatePhysicsActorHandle_Chaos.cpp
+++ b/Engine/Source/Runtime/Engine/Private/PhysicsEngine/ImmediatePhysics/ImmediatePhysicsChaos/ImmediatePhysicsActorHandle_Chaos.cpp
@@ -233,13 +233,13 @@ namespace ImmediatePhysics_Chaos
 			{
 				ParticleHandle->SetHasBounds(true);
 				ParticleHandle->SetLocalBounds(Geometry->BoundingBox());
-				ParticleHandle->UpdateWorldSpaceState(FRigidTransform3(ParticleHandle->X(), ParticleHandle->R()), FVec3(0));
+				ParticleHandle->UpdateWorldSpaceState(FRigidTransform3(ParticleHandle->X(), ParticleHandle->GetR()), FVec3(0));
 			}
 
 			if (FKinematicGeometryParticleHandle* Kinematic = ParticleHandle->CastToKinematicParticle())
 			{
-				Kinematic->SetV(FVector3f::ZeroVector);
-				Kinematic->SetW(FVector3f::ZeroVector);
+				Kinematic->SetVf(FVector3f::ZeroVector);
+				Kinematic->SetWf(FVector3f::ZeroVector);
 			}
 
 			FPBDRigidParticleHandle* Dynamic = ParticleHandle->CastToRigidParticle();
@@ -354,14 +354,14 @@ namespace ImmediatePhysics_Chaos
 
 		if (FKinematicGeometryParticleHandle* Kinematic = ParticleHandle->CastToKinematicParticle())
 		{
-			Kinematic->V() = FVec3(0);
-			Kinematic->W() = FVec3(0);
+			Kinematic->SetVf(FVec3f(0));
+			Kinematic->SetWf(FVec3f(0));
 			Kinematic->KinematicTarget().Clear();
 		}
 
 		// Initialize the bounds. Important because if the particle never moves its 
 		// bounds will never get updated (see FPBDMinEvolution::ApplyKinematicTargets) 
-		ParticleHandle->UpdateWorldSpaceState(FRigidTransform3(ParticleHandle->X(), ParticleHandle->R()), FVec3(0));
+		ParticleHandle->UpdateWorldSpaceState(FRigidTransform3(ParticleHandle->X(), ParticleHandle->GetR()), FVec3(0));
 	}
 
 	void FActorHandle::SetWorldTransform(const FTransform& WorldTM)
@@ -369,15 +369,15 @@ namespace ImmediatePhysics_Chaos
 		using namespace Chaos;
 
 		ParticleHandle->X() = WorldTM.GetTranslation();
-		ParticleHandle->R() = WorldTM.GetRotation();
+		ParticleHandle->SetR(WorldTM.GetRotation());
 
 		FPBDRigidParticleHandle* Dynamic = ParticleHandle->CastToRigidParticle();
 		if(Dynamic && Dynamic->ObjectState() == Chaos::EObjectStateType::Dynamic)
 		{
 			Dynamic->P() = Dynamic->X();
-			Dynamic->Q() = Dynamic->R();
+			Dynamic->SetQf(Dynamic->GetRf());
 			Dynamic->AuxilaryValue(ParticlePrevXs) = Dynamic->P();
-			Dynamic->AuxilaryValue(ParticlePrevRs) = Dynamic->Q();
+			Dynamic->AuxilaryValue(ParticlePrevRs) = Dynamic->GetQ();
 		}
 	}
 
@@ -570,7 +570,7 @@ namespace ImmediatePhysics_Chaos
 
 			if (ForceType == EForceType::AddImpulse || ForceType == EForceType::AddVelocity)
 			{
-				Rigid->V() += ApplyDelta;
+				Rigid->SetV(Rigid->GetV() + ApplyDelta);
 			}
 			else
 			{
diff --git a/Engine/Source/Runtime/Engine/Private/PhysicsEngine/ImmediatePhysics/ImmediatePhysicsChaos/ImmediatePhysicsSimulation_Chaos.cpp b/Engine/Source/Runtime/Engine/Private/PhysicsEngine/ImmediatePhysics/ImmediatePhysicsChaos/ImmediatePhysicsSimulation_Chaos.cpp
--- a/Engine/Source/Runtime/Engine/Private/PhysicsEngine/ImmediatePhysics/ImmediatePhysicsChaos/ImmediatePhysicsSimulation_Chaos.cpp
+++ b/Engine/Source/Runtime/Engine/Private/PhysicsEngine/ImmediatePhysics/ImmediatePhysicsChaos/ImmediatePhysicsSimulation_Chaos.cpp
@@ -458,7 +458,7 @@ namespace ImmediatePhysics_Chaos
 		ActorHandle->GetParticle()->AuxilaryValue(Implementation->PerParticleMaterials) = MoveTemp(Material);
 		ActorHandle->GetParticle()->AuxilaryValue(Implementation->CollidedParticles) = false;
 		ActorHandle->GetParticle()->AuxilaryValue(Implementation->ParticlePrevXs) = ActorHandle->GetParticle()->X();
-		ActorHandle->GetParticle()->AuxilaryValue(Implementation->ParticlePrevRs) = ActorHandle->GetParticle()->R();
+		ActorHandle->GetParticle()->AuxilaryValue(Implementation->ParticlePrevRs) = ActorHandle->GetParticle()->GetR();
 
 		Implementation->bActorsDirty = true;
 
diff --git a/Engine/Source/Runtime/Engine/Private/PhysicsEngine/PhysicsReplication.cpp b/Engine/Source/Runtime/Engine/Private/PhysicsEngine/PhysicsReplication.cpp
--- a/Engine/Source/Runtime/Engine/Private/PhysicsEngine/PhysicsReplication.cpp
+++ b/Engine/Source/Runtime/Engine/Private/PhysicsEngine/PhysicsReplication.cpp
@@ -1262,9 +1262,9 @@ bool FPhysicsReplicationAsync::DefaultReplication(Chaos::FPBDRigidParticleHandle
 	// Get Current state
 	FRigidBodyState CurrentState;
 	CurrentState.Position = Handle->X();
-	CurrentState.Quaternion = Handle->R();
-	CurrentState.AngVel = Handle->W();
-	CurrentState.LinVel = Handle->V();
+	CurrentState.Quaternion = Handle->GetR();
+	CurrentState.AngVel = Handle->GetW();
+	CurrentState.LinVel = Handle->GetV();
 
 
 	// Starting from the last known authoritative position, and
@@ -1497,7 +1497,7 @@ bool FPhysicsReplicationAsync::PredictiveInterpolation(Chaos::FPBDRigidParticleH
 	if (bXCanEarlyOut)
 	{
 		// Get the rotational offset between the blended rotation target and the current rotation
-		const FQuat TargetRotDelta = Target.TargetState.Quaternion * Handle->R().Inverse();
+		const FQuat TargetRotDelta = Target.TargetState.Quaternion * Handle->GetR().Inverse();
 
 		// Convert to angle axis
 		float Angle;
@@ -1523,9 +1523,9 @@ bool FPhysicsReplicationAsync::PredictiveInterpolation(Chaos::FPBDRigidParticleH
 	// CurrentState
 	FRigidBodyState CurrentState;
 	CurrentState.Position = Handle->X();
-	CurrentState.Quaternion = Handle->R();
-	CurrentState.LinVel = Handle->V();
-	CurrentState.AngVel = Handle->W(); // Note: Current angular velocity is in Radians
+	CurrentState.Quaternion = Handle->GetR();
+	CurrentState.LinVel = Handle->GetV();
+	CurrentState.AngVel = Handle->GetW(); // Note: Current angular velocity is in Radians
 
 	// NewState
 	const FVector TargetPos = FVector(Target.TargetState.Position);
@@ -1769,9 +1769,9 @@ bool FPhysicsReplicationAsync::ResimulationReplication(Chaos::FPBDRigidParticleH
 		UE_LOG(LogTemp, Log, TEXT("Apply Rigid body state at local frame %d with offset = %d"), LocalFrame, Target.FrameOffset);
 		UE_LOG(LogTemp, Log, TEXT("Particle Position Error = %f | Should Trigger Resim = %s | Server Frame = %d | Client Frame = %d"), ErrorDistance, (ShouldTriggerResim ? TEXT("True") : TEXT("False")), Target.ServerFrame, LocalFrame);
 		UE_LOG(LogTemp, Log, TEXT("Particle Target Position = %s | Current Position = %s"), *Target.TargetState.Position.ToString(), *PastState.X().ToString());
-		UE_LOG(LogTemp, Log, TEXT("Particle Target Velocity = %s | Current Velocity = %s"), *Target.TargetState.LinVel.ToString(), *PastState.V().ToString());
-		UE_LOG(LogTemp, Log, TEXT("Particle Target Quaternion = %s | Current Quaternion = %s"), *Target.TargetState.Quaternion.ToString(), *PastState.R().ToString());
-		UE_LOG(LogTemp, Log, TEXT("Particle Target Omega = %s | Current Omega= %s"), *Target.TargetState.AngVel.ToString(), *PastState.W().ToString());
+		UE_LOG(LogTemp, Log, TEXT("Particle Target Velocity = %s | Current Velocity = %s"), *Target.TargetState.LinVel.ToString(), *PastState.GetV().ToString());
+		UE_LOG(LogTemp, Log, TEXT("Particle Target Quaternion = %s | Current Quaternion = %s"), *Target.TargetState.Quaternion.ToString(), *PastState.GetR().ToString());
+		UE_LOG(LogTemp, Log, TEXT("Particle Target Omega = %s | Current Omega= %s"), *Target.TargetState.AngVel.ToString(), *PastState.GetW().ToString());
 
 		{ // DrawDebug
 			static constexpr float BoxSize = 5.0f;
@@ -1779,7 +1779,7 @@ bool FPhysicsReplicationAsync::ResimulationReplication(Chaos::FPBDRigidParticleH
 			const FColor DebugColor = FLinearColor::LerpUsingHSV(FLinearColor::Green, FLinearColor::Red, ColorLerp).ToFColor(false);
 
 			Chaos::FDebugDrawQueue::GetInstance().DrawDebugBox(Target.TargetState.Position, FVector(BoxSize, BoxSize, BoxSize), Target.TargetState.Quaternion, FColor::Orange, true, CharacterMovementCVars::NetCorrectionLifetime, 0, 1.0f);
-			Chaos::FDebugDrawQueue::GetInstance().DrawDebugBox(PastState.X(), FVector(6, 6, 6), PastState.R(), DebugColor, true, CharacterMovementCVars::NetCorrectionLifetime, 0, 1.0f);
+			Chaos::FDebugDrawQueue::GetInstance().DrawDebugBox(PastState.X(), FVector(6, 6, 6), PastState.GetR(), DebugColor, true, CharacterMovementCVars::NetCorrectionLifetime, 0, 1.0f);
 
 			Chaos::FDebugDrawQueue::GetInstance().DrawDebugDirectionalArrow(PastState.X(), Target.TargetState.Position, 5.0f, FColor::MakeRandomSeededColor(LocalFrame), true, CharacterMovementCVars::NetCorrectionLifetime, 0, 0.5f);
 		}
@@ -1807,10 +1807,10 @@ bool FPhysicsReplicationAsync::ResimulationReplication(Chaos::FPBDRigidParticleH
 
 			// Calculate correction to rotation
 			const float CorrectionAmountR = (1.f / NumPredictedFrames) * SettingsCurrent.ResimulationSettings.GetRotStabilityMultiplier();
-			const FQuat InvCurrentQuat = PastState.R().Inverse();
+			const FQuat InvCurrentQuat = PastState.GetR().Inverse();
 			const FQuat DeltaQuat = Target.TargetState.Quaternion * InvCurrentQuat;
-			const FQuat TargetCorrectionR = Handle->R() * DeltaQuat;
-			const FQuat CorrectedR = FQuat::Slerp(Handle->R(), TargetCorrectionR, CorrectionAmountR);
+			const FQuat TargetCorrectionR = Handle->GetR() * DeltaQuat;
+			const FQuat CorrectedR = FQuat::Slerp(Handle->GetR(), TargetCorrectionR, CorrectionAmountR);
 
 #if !(UE_BUILD_SHIPPING || UE_BUILD_TEST)
 			if (Chaos::FPhysicsSolverBase::CanDebugNetworkPhysicsPrediction())
diff --git a/Engine/Source/Runtime/Experimental/Chaos/Private/Chaos/CCDUtilities.cpp b/Engine/Source/Runtime/Experimental/Chaos/Private/Chaos/CCDUtilities.cpp
--- a/Engine/Source/Runtime/Experimental/Chaos/Private/Chaos/CCDUtilities.cpp
+++ b/Engine/Source/Runtime/Experimental/Chaos/Private/Chaos/CCDUtilities.cpp
@@ -85,7 +85,7 @@ namespace Chaos
 			const TPBDRigidParticleHandle<FReal, 3>* Rigid = Constraint->GetParticle(i)->CastToRigidParticle();
 			if (Rigid && Rigid->ObjectState() == EObjectStateType::Kinematic)
 			{
-				if (FCCDHelpers::DeltaExceedsThreshold(Rigid->CCDAxisThreshold(), Displacements[i], Rigid->Q()))
+				if (FCCDHelpers::DeltaExceedsThreshold(Rigid->CCDAxisThreshold(), Displacements[i], Rigid->GetQ()))
 				{
 					return i;
 				}
@@ -194,7 +194,7 @@ namespace Chaos
 				if (RigidParticle)
 				{
 					// One can also use P - X for dynamic particles. But notice that for kinematic particles, both P and X are end-frame positions and P - X won't work for kinematic particles.
-					Displacements[i] = RigidParticle->V() * Dt;
+					Displacements[i] = RigidParticle->GetV() * Dt;
 				}
 			}
 
@@ -350,7 +350,7 @@ namespace Chaos
 		// @todo(chaos): we should store the sweep positions in CCDParticle and onlymodify the original particle at the end
 		for (int32 i = ParticleStart; i < ParticleStart + ParticleNum; i++)
 		{
-			GroupedCCDParticles[i]->Particle->X() = GroupedCCDParticles[i]->Particle->P() - GroupedCCDParticles[i]->Particle->V() * Dt;
+			GroupedCCDParticles[i]->Particle->X() = GroupedCCDParticles[i]->Particle->P() - GroupedCCDParticles[i]->Particle->GetV() * Dt;
 		}
 
 		// Sort constraints based on TOI
@@ -495,8 +495,8 @@ namespace Chaos
 		// NOTE: This is not the same as its original X if we have been rewound to a TOI.
 		for (int32 i = ParticleStart; i < ParticleStart + ParticleNum; i++)
 		{
-			const FVec3 CoMPrev = GroupedCCDParticles[i]->Particle->PCom() - GroupedCCDParticles[i]->Particle->V() * Dt;
-			const FVec3 CoMOffsetPrev = GroupedCCDParticles[i]->Particle->R() * GroupedCCDParticles[i]->Particle->CenterOfMass();
+			const FVec3 CoMPrev = GroupedCCDParticles[i]->Particle->PCom() - GroupedCCDParticles[i]->Particle->GetV() * Dt;
+			const FVec3 CoMOffsetPrev = GroupedCCDParticles[i]->Particle->GetR() * GroupedCCDParticles[i]->Particle->CenterOfMass();
 			GroupedCCDParticles[i]->Particle->X() = CoMPrev - CoMOffsetPrev;
 		}
 	}
@@ -699,7 +699,7 @@ namespace Chaos
 						{
 							AdvanceParticleXToTOI(AffectedCCDParticle, IslandTOI, Dt);
 						}
-						ParticleStartWorldTransforms[j] = FRigidTransform3(AffectedParticle->X(), AffectedParticle->R());
+						ParticleStartWorldTransforms[j] = FRigidTransform3(AffectedParticle->X(), AffectedParticle->GetR());
 					}
 					else
 					{
@@ -711,7 +711,7 @@ namespace Chaos
 						}
 						else // Static case
 						{
-							ParticleStartWorldTransforms[j] = FRigidTransform3(AffectedParticle->X(), AffectedParticle->R());
+							ParticleStartWorldTransforms[j] = FRigidTransform3(AffectedParticle->X(), AffectedParticle->GetR());
 						}
 					}
 				}
@@ -794,7 +794,7 @@ namespace Chaos
 		{
 			TPBDRigidParticleHandle<FReal, 3>* Particle = CCDParticle->Particle;
 			const FReal RestDt = (TOI - CCDParticle->TOI) * Dt;
-			Particle->X() = Particle->X() + Particle->V() * RestDt;
+			Particle->X() = Particle->X() + Particle->GetV() * RestDt;
 			CCDParticle->TOI = TOI;
 		}
 	}
@@ -803,7 +803,7 @@ namespace Chaos
 	{
 		TPBDRigidParticleHandle<FReal, 3>* Particle = CCDParticle->Particle;
 		const FReal RestDt = (1.f - CCDParticle->TOI) * Dt;
-		Particle->P() = Particle->X() + Particle->V() * RestDt;
+		Particle->P() = Particle->X() + Particle->GetV() * RestDt;
 	}
 
 	void FCCDManager::ClipParticleP(FCCDParticle *CCDParticle) const
@@ -836,8 +836,8 @@ namespace Chaos
 			}
 			
 			const FVec3 Normal = ShapeWorldTransform1.TransformVectorNoScale(FVec3(ManifoldPoint.ContactPoint.ShapeContactNormal));
-			const FVec3 V0 = Rigid0 != nullptr ? Rigid0->V() : FVec3(0.f);
-			const FVec3 V1 = Rigid1 != nullptr ? Rigid1->V() : FVec3(0.f);
+			const FVec3 V0 = Rigid0 != nullptr ? Rigid0->GetV() : FVec3(0.f);
+			const FVec3 V1 = Rigid1 != nullptr ? Rigid1->GetV() : FVec3(0.f);
 			const FReal NormalV = FVec3::DotProduct(V0 - V1, Normal);
 			if (NormalV < 0.f)
 			{
@@ -850,11 +850,11 @@ namespace Chaos
 				const FVec3 Impulse = (TargetNormalV - NormalV) * Normal / (InvM0 + InvM1);
 				if (InvM0 > 0.f)
 				{
-					Rigid0->V() += Impulse * InvM0;
+					Rigid0->SetV(Rigid0->GetV() + Impulse * InvM0);
 				}
 				if (InvM1 > 0.f)
 				{
-					Rigid1->V() -= Impulse * InvM1;
+					Rigid1->SetV(Rigid1->GetV() - Impulse * InvM1);
 				}
 
 				CCDConstraint->NetImpulse += Impulse;
@@ -1057,7 +1057,7 @@ namespace Chaos
 		for (FCCDParticle& CCDParticle : CCDParticles)
 		{
 			TPBDRigidParticleHandle<FReal, 3>* Particle = CCDParticle.Particle;
-			Particle->X() = Particle->P() - Particle->V() * Dt;
+			Particle->X() = Particle->P() - Particle->GetV() * Dt;
 		}
 	}
 
@@ -1148,8 +1148,8 @@ namespace Chaos
 		const auto Rigid1 = Particle1.CastToRigidParticle();
 		const FVec3 DeltaX0 = Rigid0 ? Rigid0->P() - Rigid0->X() : FVec3::ZeroVector;
 		const FVec3 DeltaX1 = Rigid1 ? Rigid1->P() - Rigid1->X() : FVec3::ZeroVector;
-		const FQuat& R0 = Rigid0 ? Rigid0->Q() : Particle0.R();
-		const FQuat& R1 = Rigid1 ? Rigid1->Q() : Particle1.R();
+		const FQuat& R0 = Rigid0 ? Rigid0->GetQ() : Particle0.GetR();
+		const FQuat& R1 = Rigid1 ? Rigid1->GetQ() : Particle1.GetR();
 		return DeltaExceedsThreshold(
 			Particle0.CCDAxisThreshold(), DeltaX0, R0,
 			Particle1.CCDAxisThreshold(), DeltaX1, R1);
@@ -1161,10 +1161,10 @@ namespace Chaos
 		// For non-rigids, DeltaX is zero and use R for rotation.
 		const auto Rigid0 = Particle0.CastToRigidParticle();
 		const auto Rigid1 = Particle1.CastToRigidParticle();
-		const FVec3 DeltaX0 = Rigid0 ? Rigid0->V() * Dt : FVec3::ZeroVector;
-		const FVec3 DeltaX1 = Rigid1 ? Rigid1->V() * Dt : FVec3::ZeroVector;
-		const FQuat& R0 = Rigid0 ? Rigid0->Q() : Particle0.R();
-		const FQuat& R1 = Rigid1 ? Rigid1->Q() : Particle1.R();
+		const FVec3 DeltaX0 = Rigid0 ? Rigid0->GetV() * Dt : FVec3::ZeroVector;
+		const FVec3 DeltaX1 = Rigid1 ? Rigid1->GetV() * Dt : FVec3::ZeroVector;
+		const FQuat& R0 = Rigid0 ? Rigid0->GetQ() : Particle0.GetR();
+		const FQuat& R1 = Rigid1 ? Rigid1->GetQ() : Particle1.GetR();
 		return DeltaExceedsThreshold(
 			Particle0.CCDAxisThreshold(), DeltaX0, R0,
 			Particle1.CCDAxisThreshold(), DeltaX1, R1);
diff --git a/Engine/Source/Runtime/Experimental/Chaos/Private/Chaos/ChaosDebugDraw.cpp b/Engine/Source/Runtime/Experimental/Chaos/Private/Chaos/ChaosDebugDraw.cpp
--- a/Engine/Source/Runtime/Experimental/Chaos/Private/Chaos/ChaosDebugDraw.cpp
+++ b/Engine/Source/Runtime/Experimental/Chaos/Private/Chaos/ChaosDebugDraw.cpp
@@ -930,7 +930,7 @@ namespace Chaos
 		void DrawParticleShapesImpl(const FRigidTransform3& SpaceTransform, const FGeometryParticleHandle* Particle, const FColor& InColor, const FChaosDebugDrawSettings& Settings)
 		{
 			const FVec3 P = SpaceTransform.TransformPosition(Particle->ObjectState() == EObjectStateType::Dynamic ? Particle->CastToRigidParticle()->P() : Particle->X());
-			const FRotation3 Q = SpaceTransform.GetRotation() * (Particle->ObjectState() == EObjectStateType::Dynamic ? Particle->CastToRigidParticle()->Q() : Particle->R());
+			const FRotation3 Q = SpaceTransform.GetRotation() * (Particle->ObjectState() == EObjectStateType::Dynamic ? Particle->CastToRigidParticle()->GetQ() : Particle->GetR());
 			const FRigidTransform3 ParticleSpaceTransform = FRigidTransform3(P, Q);
 
 			if(!DrawConvexOptimizerShapes(P, Q, Particle, InColor, Settings))
@@ -991,7 +991,7 @@ namespace Chaos
 				if (Union->GetBVH() != nullptr)
 				{
 					FVec3 P = SpaceTransform.TransformPosition(Particle->X());
-					FRotation3 Q = SpaceTransform.GetRotation() * (Particle->R());
+					FRotation3 Q = SpaceTransform.GetRotation() * (Particle->GetR());
 
 					DrawBVHImpl(Particle->Handle(), FRigidTransform3(P, Q), Union->GetBVH(), InColor, 0.0f, Settings);
 				}
@@ -1804,7 +1804,7 @@ namespace Chaos
 						}
 					}
 
-					return FTransform{ Particle->R(), Particle->X() };
+					return FTransform{ Particle->GetR(), Particle->X() };
 				};
 
 				return GetClusterTransformImpl(Particle, GetClusterTransformImpl);
diff --git a/Engine/Source/Runtime/Experimental/Chaos/Private/Chaos/ClusterUnionManager.cpp b/Engine/Source/Runtime/Experimental/Chaos/Private/Chaos/ClusterUnionManager.cpp
--- a/Engine/Source/Runtime/Experimental/Chaos/Private/Chaos/ClusterUnionManager.cpp
+++ b/Engine/Source/Runtime/Experimental/Chaos/Private/Chaos/ClusterUnionManager.cpp
@@ -42,7 +42,7 @@ namespace Chaos
 			}
 			else
 			{
-				const FRigidTransform3 ChildWorldTM(Child->X(), Child->R());
+				const FRigidTransform3 ChildWorldTM(Child->X(), Child->GetR());
 				Frame = ChildWorldTM.GetRelativeTransform(ClusterWorldTM);
 			}
 
@@ -54,7 +54,7 @@ namespace Chaos
 			if (!ClusterUnion.ChildParticles.IsEmpty() && ClusterUnion.Geometry->GetType() == ImplicitObjectType::Union)
 			{
 				constexpr FReal kAddThickness = 5.0;
-				const FRigidTransform3 ClusterWorldTM(ClusterUnion.InternalCluster->X(), ClusterUnion.InternalCluster->R());
+				const FRigidTransform3 ClusterWorldTM(ClusterUnion.InternalCluster->X(), ClusterUnion.InternalCluster->GetR());
 
 				// Use the acceleration structure of the cluster union itself to make finding overlaps easy.
 				const FImplicitObjectUnion& ShapeUnion = ClusterUnion.Geometry->GetObjectChecked<FImplicitObjectUnion>();
@@ -378,7 +378,7 @@ namespace Chaos
 		}
 
 		// TODO: Can we do something better than a union?
-		const FRigidTransform3 ClusterWorldTM(Union.InternalCluster->X(), Union.InternalCluster->R());
+		const FRigidTransform3 ClusterWorldTM(Union.InternalCluster->X(), Union.InternalCluster->GetR());
 		TArray<Chaos::FImplicitObjectPtr> Objects;
 		Objects.Reserve(Union.ChildParticles.Num());
 
@@ -680,7 +680,7 @@ namespace Chaos
 			// the deferred cluster union update properties is called, the position of the particle will get reset to the position of the cluster union.
 			if (FPBDRigidClusteredParticleHandle* ClusterParticle = Particle->CastToClustered())
 			{
-				const FRigidTransform3 ClusterWorldTM(Cluster->InternalCluster->X(), Cluster->InternalCluster->R());
+				const FRigidTransform3 ClusterWorldTM(Cluster->InternalCluster->X(), Cluster->InternalCluster->GetR());
 				const FRigidTransform3 Frame = GetParticleRigidFrameInClusterUnion(Particle, ClusterWorldTM);
 				ClusterParticle->SetChildToParent(Frame);
 			}
@@ -1092,7 +1092,7 @@ namespace Chaos
 		MEvolution.InvalidateParticle(ClusterUnion.InternalCluster);
 		ImplicitUnion->SetAllowBVH(false);
 
-		const FRigidTransform3 ClusterWorldTM(ClusterUnion.InternalCluster->X(), ClusterUnion.InternalCluster->R());
+		const FRigidTransform3 ClusterWorldTM(ClusterUnion.InternalCluster->X(), ClusterUnion.InternalCluster->GetR());
 		if (!PendingGeometryAdditions.IsEmpty())
 		{
 			ModifyAdditionOfChildrenToClusterUnionGeometry(
@@ -1343,7 +1343,7 @@ namespace Chaos
 
 						// Update the child's world transform to be consistent with its ChildToParent transform
 						const FPBDRigidClusteredParticleHandle* ParentHandle = ClusterUnion->InternalCluster;
-						const FRigidTransform3 ParticleToWorld = ChildToParent * FRigidTransform3(ParentHandle->X(), ParentHandle->R());
+						const FRigidTransform3 ParticleToWorld = ChildToParent * FRigidTransform3(ParentHandle->X(), ParentHandle->GetR());
 						ClusteredParticle->SetX(ParticleToWorld.GetTranslation());
 						ClusteredParticle->SetP(ParticleToWorld.GetTranslation());
 						ClusteredParticle->SetR(ParticleToWorld.GetRotation());
diff --git a/Engine/Source/Runtime/Experimental/Chaos/Private/Chaos/CollisionResolutionUtil.cpp b/Engine/Source/Runtime/Experimental/Chaos/Private/Chaos/CollisionResolutionUtil.cpp
--- a/Engine/Source/Runtime/Experimental/Chaos/Private/Chaos/CollisionResolutionUtil.cpp
+++ b/Engine/Source/Runtime/Experimental/Chaos/Private/Chaos/CollisionResolutionUtil.cpp
@@ -50,15 +50,15 @@ namespace Chaos
 			if (bIsRigidDynamic0)
 			{
 				Jr0 = FVec3::CrossProduct(VectorToPoint1, Impulse);
-				IInvJr0 = PBDRigid0->Q().RotateVector(FVec3(PBDRigid0->InvI()) * PBDRigid0->Q().UnrotateVector(Jr0));
-				ImpulseRatioNumerator0 = FVec3::DotProduct(Impulse, PBDRigid0->V() - KinematicVelocity) + FVec3::DotProduct(Jr0, PBDRigid0->W());
+				IInvJr0 = PBDRigid0->GetQ().RotateVector(FVec3(PBDRigid0->InvI()) * PBDRigid0->GetQ().UnrotateVector(Jr0));
+				ImpulseRatioNumerator0 = FVec3::DotProduct(Impulse, PBDRigid0->GetV() - KinematicVelocity) + FVec3::DotProduct(Jr0, PBDRigid0->GetW());
 				ImpulseRatioDenom0 = ImpulseSizeSQ / PBDRigid0->M() + FVec3::DotProduct(Jr0, IInvJr0);
 			}
 			if (bIsRigidDynamic1)
 			{
 				Jr1 = FVec3::CrossProduct(VectorToPoint2, Impulse);
-				IInvJr1 = PBDRigid1->Q().RotateVector(FVec3(PBDRigid1->InvI()) * PBDRigid1->Q().UnrotateVector(Jr1));
-				ImpulseRatioNumerator1 = FVec3::DotProduct(Impulse, PBDRigid1->V() - KinematicVelocity) + FVec3::DotProduct(Jr1, PBDRigid1->W());
+				IInvJr1 = PBDRigid1->GetQ().RotateVector(FVec3(PBDRigid1->InvI()) * PBDRigid1->GetQ().UnrotateVector(Jr1));
+				ImpulseRatioNumerator1 = FVec3::DotProduct(Impulse, PBDRigid1->GetV() - KinematicVelocity) + FVec3::DotProduct(Jr1, PBDRigid1->GetW());
 				ImpulseRatioDenom1 = ImpulseSizeSQ / PBDRigid1->M() + FVec3::DotProduct(Jr1, IInvJr1);
 			}
 			FReal Numerator = -2 * (ImpulseRatioNumerator0 - ImpulseRatioNumerator1);
diff --git a/Engine/Source/Runtime/Experimental/Chaos/Private/Chaos/ContactModification.cpp b/Engine/Source/Runtime/Experimental/Chaos/Private/Chaos/ContactModification.cpp
--- a/Engine/Source/Runtime/Experimental/Chaos/Private/Chaos/ContactModification.cpp
+++ b/Engine/Source/Runtime/Experimental/Chaos/Private/Chaos/ContactModification.cpp
@@ -228,7 +228,7 @@ namespace Chaos
 			return FVec3(0);
 		}
 
-		return KinematicHandle->V();
+		return KinematicHandle->GetV();
 	}
 
 	void FContactPairModifier::ModifyParticleVelocity(FVec3 Velocity, int32 ParticleIdx)
@@ -267,7 +267,7 @@ namespace Chaos
 			return FVec3(0);
 		}
 
-		return KinematicHandle->W();
+		return KinematicHandle->GetW();
 	}
 
 	void FContactPairModifier::ModifyParticleAngularVelocity(FVec3 AngularVelocity, int32 ParticleIdx)
@@ -289,7 +289,7 @@ namespace Chaos
 			EObjectStateType State(RigidHandle->ObjectState());
 			if (State == EObjectStateType::Dynamic || State == EObjectStateType::Sleeping)
 			{
-				RigidHandle->SetR(FRotation3::IntegrateRotationWithAngularVelocity(RigidHandle->Q(), -RigidHandle->W(), Modifier->Dt));
+				RigidHandle->SetR(FRotation3::IntegrateRotationWithAngularVelocity(RigidHandle->GetQ(), -RigidHandle->GetW(), Modifier->Dt));
 			}
 		}
 	}
@@ -331,7 +331,7 @@ namespace Chaos
 
 				if (bMaintainVelocity)
 				{
-					RigidHandle->SetX(RigidHandle->P() - RigidHandle->V() * Modifier->Dt);
+					RigidHandle->SetX(RigidHandle->P() - RigidHandle->GetV() * Modifier->Dt);
 				}
 				else if(Modifier->Dt > 0.0f)
 				{
@@ -373,10 +373,10 @@ namespace Chaos
 		// We give predicted position for simulated objects
 		if (RigidHandle)
 		{
-			return RigidHandle->Q();
+			return RigidHandle->GetQ();
 		}
 
-		return Particle->R();
+		return Particle->GetR();
 	}
 
 	void FContactPairModifier::ModifyParticleRotation(FRotation3 Rotation, bool bMaintainVelocity, int32 ParticleIdx)
@@ -396,12 +396,12 @@ namespace Chaos
 
 				if (bMaintainVelocity)
 				{
-					RigidHandle->SetR(FRotation3::IntegrateRotationWithAngularVelocity(RigidHandle->Q(), -RigidHandle->W(), Modifier->Dt));
+					RigidHandle->SetR(FRotation3::IntegrateRotationWithAngularVelocity(RigidHandle->GetQ(), -RigidHandle->GetW(), Modifier->Dt));
 				}
 				else if (Modifier->Dt > 0.0f)
 				{
 					// Update W to new implicit velocity
-					RigidHandle->SetW(FRotation3::CalculateAngularVelocity(RigidHandle->R(), RigidHandle->Q(), Modifier->Dt));
+					RigidHandle->SetW(FRotation3::CalculateAngularVelocity(RigidHandle->GetR(), RigidHandle->GetQ(), Modifier->Dt));
 				}
 				UpdateConstraintShapeTransforms();
 				return;
diff --git a/Engine/Source/Runtime/Experimental/Chaos/Private/Chaos/Deformable/ChaosDeformableSolver.cpp b/Engine/Source/Runtime/Experimental/Chaos/Private/Chaos/Deformable/ChaosDeformableSolver.cpp
--- a/Engine/Source/Runtime/Experimental/Chaos/Private/Chaos/Deformable/ChaosDeformableSolver.cpp
+++ b/Engine/Source/Runtime/Experimental/Chaos/Private/Chaos/Deformable/ChaosDeformableSolver.cpp
@@ -170,7 +170,7 @@ namespace Chaos::Softs
 			Chaos::FVec3 EulerRot(0.f);
 			int32 CollisionParticleOffset = Evolution->AddCollisionParticleRange(1, INDEX_NONE, true);
 			Evolution->CollisionParticles().X(0) = Position;
-			Evolution->CollisionParticles().R(0) = Chaos::TRotation<Chaos::FReal, 3>::MakeFromEuler(EulerRot);
+			Evolution->CollisionParticles().SetR(0, Chaos::TRotation<Chaos::FReal, 3>::MakeFromEuler(EulerRot));
 			Evolution->CollisionParticles().SetGeometry(0, MakeImplicitObjectPtr<Chaos::TPlane<Chaos::FReal, 3>>(Chaos::FVec3(0.f, 0.f, 0.f), Chaos::FVec3(0.f, 0.f, 1.f)));
 		}
 	}
@@ -642,7 +642,7 @@ namespace Chaos::Softs
 								int32 Index = Evolution->AddCollisionParticle(INDEX_NONE, true);
 								int32 ViewIndex = Evolution->CollisionParticlesActiveView().GetRanges().Num() - 1;
 								Evolution->CollisionParticles().X(Index) = AddBody.Transform.GetTranslation();
-								Evolution->CollisionParticles().R(Index) = AddBody.Transform.GetRotation();
+								Evolution->CollisionParticles().SetR(Index, AddBody.Transform.GetRotation());
 								Chaos::FImplicitObjectPtr UniquePtr(AddBody.Shapes); AddBody.Shapes = nullptr;
 								Evolution->CollisionParticles().SetGeometry(Index, MoveTemp(UniquePtr));
 								Proxy.CollisionBodies.Add(AddBody.Key, FCollisionObjectParticleHandel(Index, ViewIndex, AddBody.Transform));
@@ -705,7 +705,7 @@ namespace Chaos::Softs
 						{
 							FCollisionObjectParticleHandel* ParticleHandle = Proxy.CollisionBodies.Find(UpdateBody.Key);
 							Evolution->CollisionParticles().X(ParticleHandle->ParticleIndex) = UpdateBody.Transform.GetTranslation();
-							Evolution->CollisionParticles().R(ParticleHandle->ParticleIndex) = UpdateBody.Transform.GetRotation();
+							Evolution->CollisionParticles().SetR(ParticleHandle->ParticleIndex, UpdateBody.Transform.GetRotation());
 						}
 					}
 				}
@@ -1698,13 +1698,13 @@ namespace Chaos::Softs
 									const TBox<FReal, 3>& BoxGeometry = Geometry->GetObjectChecked<TBox<FReal, 3>>();
 									FVector Extent = 0.5 * (BoxGeometry.Max() - BoxGeometry.Min());
 									FVector Center = ToFVector(CollisionParticles.X(Index)) + BoxGeometry.GetCenter();
-									const FQuat& Rotation = ToFQuat(CollisionParticles.R(Index));
+									const FQuat& Rotation = ToFQuat(CollisionParticles.GetR(Index));
 									Chaos::FDebugDrawQueue::GetInstance().DrawDebugBox(Center, Extent, Rotation, FColor::Red, false, -1.0f, 0, 1.f);
 								}
 								else if (GeomType == ImplicitObjectType::Convex)
 								{
 									const FConvex& ConvexGeometry = Geometry->GetObjectChecked<FConvex>();
-									FTransform M = FTransform(ToFQuat(CollisionParticles.R(Index)), ToFVector(CollisionParticles.X(Index)));
+									FTransform M = FTransform(ToFQuat(CollisionParticles.GetR(Index)), ToFVector(CollisionParticles.X(Index)));
 									for (int32 EdgeIndex = 0; EdgeIndex < ConvexGeometry.NumEdges(); ++EdgeIndex)
 									{
 										int32 Index0 = ConvexGeometry.GetEdgeVertex(EdgeIndex, 0);
diff --git a/Engine/Source/Runtime/Experimental/Chaos/Private/Chaos/Evolution/PBDMinEvolution.cpp b/Engine/Source/Runtime/Experimental/Chaos/Private/Chaos/Evolution/PBDMinEvolution.cpp
--- a/Engine/Source/Runtime/Experimental/Chaos/Private/Chaos/Evolution/PBDMinEvolution.cpp
+++ b/Engine/Source/Runtime/Experimental/Chaos/Private/Chaos/Evolution/PBDMinEvolution.cpp
@@ -181,20 +181,20 @@ namespace Chaos
 			if (Particle.ObjectState() == EObjectStateType::Dynamic)
 			{
 				Particle.X() = FVec3::Lerp(Particle.Handle()->AuxilaryValue(ParticlePrevXs), Particle.X(), T);
-				Particle.R() = FRotation3::Slerp(Particle.Handle()->AuxilaryValue(ParticlePrevRs), Particle.R(), (decltype(FQuat::X))T);	// LWC_TODO: Remove decltype cast once FQuat supports variants
+				Particle.SetR(FRotation3::Slerp(Particle.Handle()->AuxilaryValue(ParticlePrevRs), Particle.GetR(), (decltype(FQuat::X))T));	// LWC_TODO: Remove decltype cast once FQuat supports variants
 
 				if (bRewindVelocities)
 				{
-					Particle.V() = FVec3::Lerp(Particle.PreV(), Particle.V(), T);
-					Particle.W() = FVec3::Lerp(Particle.PreW(), Particle.W(), T);
+					Particle.SetV(FVec3::Lerp(Particle.GetPreV(), Particle.GetV(), T));
+					Particle.SetW(FVec3::Lerp(Particle.GetPreW(), Particle.GetW(), T));
 				}
 			}
 		}
 
 		for (auto& Particle : Particles.GetActiveKinematicParticlesView())
 		{
-			Particle.X() = Particle.X() - Particle.V() * RewindDt;
-			Particle.R() = FRotation3::IntegrateRotationWithAngularVelocity(Particle.R(), -Particle.W(), RewindDt);
+			Particle.X() = Particle.X() - Particle.GetV() * RewindDt;
+			Particle.SetR(FRotation3::IntegrateRotationWithAngularVelocity(Particle.GetR(), -Particle.GetW(), RewindDt));
 		}
 	}
 
@@ -221,8 +221,8 @@ namespace Chaos
 		{
 			if (Particle.ObjectState() == EObjectStateType::Dynamic)
 			{
-				Particle.PreV() = Particle.V();
-				Particle.PreW() = Particle.W();
+				Particle.SetPreVf(Particle.GetVf());
+				Particle.SetPreWf(Particle.GetWf());
 
 				const FVec3 XCoM = Particle.XCom();
 				const FRotation3 RCoM = Particle.RCom();
@@ -244,7 +244,7 @@ namespace Chaos
 				// https://en.wikipedia.org/wiki/Rotating_reference_frame
 				if (SimulationSpaceSettings.Alpha > 0.0f)
 				{
-					const FVec3 CoriolisAcc = SimulationSpaceSettings.CoriolisAlpha * 2.0f * FVec3::CrossProduct(SpaceW, Particle.V());
+					const FVec3 CoriolisAcc = SimulationSpaceSettings.CoriolisAlpha * 2.0f * FVec3::CrossProduct(SpaceW, Particle.GetV());
 					const FVec3 CentrifugalAcc = SimulationSpaceSettings.CentrifugalAlpha * FVec3::CrossProduct(SpaceW, FVec3::CrossProduct(SpaceW, XCoM));
 					const FVec3 EulerAcc = SimulationSpaceSettings.EulerAlpha * FVec3::CrossProduct(SpaceB, XCoM);
 					const FVec3 LinearAcc = SimulationSpaceSettings.LinearAccelerationAlpha * SpaceA;
@@ -259,8 +259,8 @@ namespace Chaos
 				// New velocity
 				const FReal LinearDrag = FMath::Min(FReal(1), Particle.LinearEtherDrag() * Dt);
 				const FReal AngularDrag = FMath::Min(FReal(1), Particle.AngularEtherDrag() * Dt);
-				const FVec3 V = FMath::Lerp(Particle.V() + DV, TargetV, LinearDrag);
-				const FVec3 W = FMath::Lerp(Particle.W() + DW, TargetW, AngularDrag);
+				const FVec3 V = FMath::Lerp(Particle.GetV() + DV, TargetV, LinearDrag);
+				const FVec3 W = FMath::Lerp(Particle.GetW() + DW, TargetW, AngularDrag);
 
 				// New position
 				const FVec3 PCoM = XCoM + V * Dt;
@@ -268,15 +268,15 @@ namespace Chaos
 
 				// Update particle state (forces are not zeroed until the end of the frame)
 				Particle.SetTransformPQCom(PCoM, QCoM);
-				Particle.V() = V;
-				Particle.W() = W;
+				Particle.SetV(V);
+				Particle.SetW(W);
 				Particle.LinearImpulseVelocity() = FVec3(0);
 				Particle.AngularImpulseVelocity() = FVec3(0);
 
 				// Update cached world space state, including bounds. We use the Swept bounds update so that the bounds includes P,Q and X,Q.
 				// This is because when we have joints, they often pull bodies back to their original positions, so we need to know if there
 				// are contacts at that location.
-				Particle.UpdateWorldSpaceStateSwept(FRigidTransform3(Particle.P(), Particle.Q()), BoundsExpansion, -V * Dt);
+				Particle.UpdateWorldSpaceStateSwept(FRigidTransform3(Particle.P(), Particle.GetQ()), BoundsExpansion, -V * Dt);
 			}
 		}
 	}
@@ -303,7 +303,7 @@ namespace Chaos
 
 			TKinematicTarget<FReal, 3>& KinematicTarget = Particle.KinematicTarget();
 			const FVec3 CurrentX = Particle.X();
-			const FRotation3 CurrentR = Particle.R();
+			const FRotation3 CurrentR = Particle.GetR();
 			constexpr FReal MinDt = 1e-6f;
 
 			bool bMoved = false;
@@ -316,8 +316,8 @@ namespace Chaos
 			case EKinematicTargetMode::Reset:
 			{
 				// Reset velocity and then switch to do-nothing mode
-				Particle.V() = FVec3(0, 0, 0);
-				Particle.W() = FVec3(0, 0, 0);
+				Particle.SetVf(FVec3f(0, 0, 0));
+				Particle.SetWf(FVec3f(0, 0, 0));
 				KinematicTarget.SetMode(EKinematicTargetMode::None);
 				break;
 			}
@@ -358,10 +358,10 @@ namespace Chaos
 						NewW = FRotation3::CalculateAngularVelocity(CurrentR, NewR, Dt);
 					}
 				}
-				Particle.X() = NewX;
-				Particle.R() = NewR;
-				Particle.V() = NewV;
-				Particle.W() = NewW;
+				Particle.SetX(NewX);
+				Particle.SetR(NewR);
+				Particle.SetV(NewV);
+				Particle.SetW(NewW);
 
 				break;
 			}
@@ -370,8 +370,8 @@ namespace Chaos
 			{
 				// Move based on velocity
 				bMoved = true;
-				Particle.X() = Particle.X() + Particle.V() * Dt;
-				Particle.R() = FRotation3::IntegrateRotationWithAngularVelocity(Particle.R(), Particle.W(), Dt);
+				Particle.X() = Particle.X() + Particle.GetV() * Dt;
+				Particle.SetR(FRotation3::IntegrateRotationWithAngularVelocity(Particle.GetR(), Particle.GetW(), Dt));
 				break;
 			}
 			}
@@ -379,19 +379,19 @@ namespace Chaos
 			// Set positions and previous velocities if we can
 			// Note: At present kinematics are in fact rigid bodies
 			Particle.P() = Particle.X();
-			Particle.Q() = Particle.R();
-			Particle.PreV() = Particle.V();
-			Particle.PreW() = Particle.W();
+			Particle.SetQ(Particle.GetR());
+			Particle.SetPreV(Particle.GetV());
+			Particle.SetPreW(Particle.GetW());
 
 			if (bMoved)
 			{
 				if (!Particle.CCDEnabled())
 				{
-					Particle.UpdateWorldSpaceState(FRigidTransform3(Particle.P(), Particle.Q()), FVec3(0));
+					Particle.UpdateWorldSpaceState(FRigidTransform3(Particle.P(), Particle.GetQ()), FVec3(0));
 				}
 				else
 				{
-					Particle.UpdateWorldSpaceStateSwept(FRigidTransform3(Particle.P(), Particle.Q()), FVec3(0), -Particle.V() * Dt);
+					Particle.UpdateWorldSpaceStateSwept(FRigidTransform3(Particle.P(), Particle.GetQ()), FVec3(0), -Particle.GetV() * Dt);
 				}
 			}
 		};
@@ -403,7 +403,7 @@ namespace Chaos
 		{
 			TKinematicTarget<FReal, 3>& KinematicTarget = Particle.KinematicTarget();
 			const FVec3 CurrentX = Particle.X();
-			const FRotation3 CurrentR = Particle.R();
+			const FRotation3 CurrentR = Particle.GetR();
 			constexpr FReal MinDt = 1e-6f;
 
 			bool bMoved = false;
@@ -416,8 +416,8 @@ namespace Chaos
 			case EKinematicTargetMode::Reset:
 			{
 				// Reset velocity and then switch to do-nothing mode
-				Particle.V() = FVec3(0, 0, 0);
-				Particle.W() = FVec3(0, 0, 0);
+				Particle.SetVf(FVec3f(0, 0, 0));
+				Particle.SetWf(FVec3f(0, 0, 0));
 				KinematicTarget.SetMode(EKinematicTargetMode::None);
 				break;
 			}
@@ -459,9 +459,9 @@ namespace Chaos
 					}
 				}
 				Particle.X() = NewX;
-				Particle.R() = NewR;
-				Particle.V() = NewV;
-				Particle.W() = NewW;
+				Particle.SetR(NewR);
+				Particle.SetV(NewV);
+				Particle.SetW(NewW);
 
 				break;
 			}
@@ -470,15 +470,15 @@ namespace Chaos
 			{
 				// Move based on velocity
 				bMoved = true;
-				Particle.X() = Particle.X() + Particle.V() * Dt;
-				Particle.R() = FRotation3::IntegrateRotationWithAngularVelocity(Particle.R(), Particle.W(), Dt);
+				Particle.X() = Particle.X() + Particle.GetV() * Dt;
+				Particle.SetR(FRotation3::IntegrateRotationWithAngularVelocity(Particle.GetR(), Particle.GetW(), Dt));
 				break;
 			}
 			}
 
 			if (bMoved)
 			{
-				Particle.UpdateWorldSpaceState(FRigidTransform3(Particle.X(), Particle.R()), FVec3(0));
+				Particle.UpdateWorldSpaceState(FRigidTransform3(Particle.X(), Particle.GetR()), FVec3(0));
 			}
 		};
 
@@ -537,9 +537,9 @@ namespace Chaos
 		for (auto& Particle : Particles.GetActiveParticlesView())
 		{
 			Particle.Handle()->AuxilaryValue(ParticlePrevXs) = Particle.X();
-			Particle.Handle()->AuxilaryValue(ParticlePrevRs) = Particle.R();
+			Particle.Handle()->AuxilaryValue(ParticlePrevRs) = Particle.GetR();
 			Particle.X() = Particle.P();
-			Particle.R() = Particle.Q();
+			Particle.SetR(Particle.GetQ());
 		}
 	}
 
diff --git a/Engine/Source/Runtime/Experimental/Chaos/Private/Chaos/Evolution/SolverBodyContainer.cpp b/Engine/Source/Runtime/Experimental/Chaos/Private/Chaos/Evolution/SolverBodyContainer.cpp
--- a/Engine/Source/Runtime/Experimental/Chaos/Private/Chaos/Evolution/SolverBodyContainer.cpp
+++ b/Engine/Source/Runtime/Experimental/Chaos/Private/Chaos/Evolution/SolverBodyContainer.cpp
@@ -22,8 +22,8 @@ namespace Chaos
 			SolverBody.SetQ(RigidParticle->QCom());
 			SolverBody.SetCoM(RigidParticle->CenterOfMass());
 			SolverBody.SetRoM(RigidParticle->RotationOfMass());
-			SolverBody.SetV(RigidParticle->V());
-			SolverBody.SetW(RigidParticle->W());
+			SolverBody.SetV(RigidParticle->GetV());
+			SolverBody.SetW(RigidParticle->GetW());
 			SolverBody.SetInvM(RigidParticle->InvM());
 			SolverBody.SetInvILocal(RigidParticle->ConditionedInvI());
 			// Note: SetInvILocal also calculates InvI for dynamic
@@ -32,25 +32,25 @@ namespace Chaos
 		{
 			// Static or kinematic particle
 			SolverBody.SetP(Particle->X());
-			SolverBody.SetQ(Particle->R());
+			SolverBody.SetQ(Particle->GetR());
 			SolverBody.SetCoM(FVec3(0));
 			SolverBody.SetRoM(FRotation3::FromIdentity());
 			if (KinematicParticle != nullptr)
 			{
 				// Kinematic particle
-				SolverBody.SetX(KinematicParticle->X() - KinematicParticle->V() * Dt);
-				SolverBody.SetR(Particle->R());
-				if (!KinematicParticle->W().IsNearlyZero())
+				SolverBody.SetX(KinematicParticle->X() - KinematicParticle->GetV() * Dt);
+				SolverBody.SetR(Particle->GetR());
+				if (!KinematicParticle->GetW().IsNearlyZero())
 				{
-					SolverBody.SetR(FRotation3::IntegrateRotationWithAngularVelocity(KinematicParticle->R(), -KinematicParticle->W(), Dt));
+					SolverBody.SetR(FRotation3::IntegrateRotationWithAngularVelocity(KinematicParticle->GetR(), -KinematicParticle->GetW(), Dt));
 				}
-				SolverBody.SetV(KinematicParticle->V());
-				SolverBody.SetW(KinematicParticle->W());
+				SolverBody.SetV(KinematicParticle->GetV());
+				SolverBody.SetW(KinematicParticle->GetW());
 			}
 			else
 			{
 				SolverBody.SetX(Particle->X());
-				SolverBody.SetR(Particle->R());
+				SolverBody.SetR(Particle->GetR());
 				SolverBody.SetV(FVec3(0));
 				SolverBody.SetW(FVec3(0));
 			}
diff --git a/Engine/Source/Runtime/Experimental/Chaos/Private/Chaos/GeometryParticles.cpp b/Engine/Source/Runtime/Experimental/Chaos/Private/Chaos/GeometryParticles.cpp
--- a/Engine/Source/Runtime/Experimental/Chaos/Private/Chaos/GeometryParticles.cpp
+++ b/Engine/Source/Runtime/Experimental/Chaos/Private/Chaos/GeometryParticles.cpp
@@ -55,7 +55,7 @@ namespace Chaos
 	template <typename T, int d, EGeometryParticlesSimType SimType>
 	void TGeometryParticlesImp<T, d, SimType>::UpdateShapesArray(const int32 Index)
 	{
-		UpdateShapesArrayFromGeometry(MShapesArray[Index], GetGeometry(Index), FRigidTransform3(X(Index), R(Index)));
+		UpdateShapesArrayFromGeometry(MShapesArray[Index], GetGeometry(Index), FRigidTransform3(X(Index), GetR(Index)));
 	}
 
 	template <typename T, int d, EGeometryParticlesSimType SimType>
diff --git a/Engine/Source/Runtime/Experimental/Chaos/Private/Chaos/Island/IslandManager.cpp b/Engine/Source/Runtime/Experimental/Chaos/Private/Chaos/Island/IslandManager.cpp
--- a/Engine/Source/Runtime/Experimental/Chaos/Private/Chaos/Island/IslandManager.cpp
+++ b/Engine/Source/Runtime/Experimental/Chaos/Private/Chaos/Island/IslandManager.cpp
@@ -125,11 +125,11 @@ namespace Chaos::Private
 				const FKinematicTarget& KinematicTarget = Kinematic->KinematicTarget();
 				if (KinematicTarget.GetMode() == EKinematicTargetMode::Position)
 				{
-					bIsStationary = (Kinematic->X() - KinematicTarget.GetTargetPosition()).IsZero() && (Kinematic->R() * KinematicTarget.GetTargetRotation().Inverse()).IsIdentity();
+					bIsStationary = (Kinematic->X() - KinematicTarget.GetTargetPosition()).IsZero() && (Kinematic->GetR() * KinematicTarget.GetTargetRotation().Inverse()).IsIdentity();
 				}
 				else
 				{
-					bIsStationary = Kinematic->V().IsZero() && Kinematic->W().IsZero();
+					bIsStationary = Kinematic->GetV().IsZero() && Kinematic->GetW().IsZero();
 				}
 			}
 			else
@@ -251,8 +251,8 @@ namespace Chaos::Private
 	{
 		if (Dt > UE_SMALL_NUMBER)
 		{
-			Rigid.SetVSmooth(Rigid.V());
-			Rigid.SetWSmooth(Rigid.W());
+			Rigid.SetVSmooth(Rigid.GetV());
+			Rigid.SetWSmooth(Rigid.GetW());
 		}
 	}
 
@@ -263,7 +263,7 @@ namespace Chaos::Private
 		{
 			const FReal SmoothRate = FMath::Clamp(CVars::SmoothedPositionLerpRate, 0.0f, 1.0f);
 			const FVec3 VImp = FVec3::CalculateVelocity(Rigid.X(), Rigid.P(), Dt);
-			const FVec3 WImp = FRotation3::CalculateAngularVelocity(Rigid.R(), Rigid.Q(), Dt);
+			const FVec3 WImp = FRotation3::CalculateAngularVelocity(Rigid.GetR(), Rigid.GetQ(), Dt);
 			Rigid.SetVSmooth(FMath::Lerp(Rigid.VSmooth(), VImp, SmoothRate));
 			Rigid.SetWSmooth(FMath::Lerp(Rigid.WSmooth(), WImp, SmoothRate));
 		}
@@ -2411,8 +2411,8 @@ namespace Chaos::Private
 				int32 DisableCounter = 0;
 
 				// Did we exceed the velocity thresholds?
-				const FReal VSq = Particle->V().SizeSquared();
-				const FReal WSq = Particle->W().SizeSquared();
+				const FReal VSq = Particle->GetV().SizeSquared();
+				const FReal WSq = Particle->GetW().SizeSquared();
 				if ((VSq < Node->DisableLinearThresholdSq) && (WSq < Node->DisableAngularThresholdSq))
 				{
 					// We are within the velocity thresholds, so see if we should disable
diff --git a/Engine/Source/Runtime/Experimental/Chaos/Private/Chaos/NewtonEvolution.cpp b/Engine/Source/Runtime/Experimental/Chaos/Private/Chaos/NewtonEvolution.cpp
--- a/Engine/Source/Runtime/Experimental/Chaos/Private/Chaos/NewtonEvolution.cpp
+++ b/Engine/Source/Runtime/Experimental/Chaos/Private/Chaos/NewtonEvolution.cpp
@@ -779,7 +779,7 @@ void FNewtonEvolution::WriteOutputLog(const int32 Frame)
 					[this, Dt](FSolverCollisionParticles& CollisionParticles, int32 Index)
 					{
 						// Store active collision particle frames prior to the kinematic update for CCD collisions
-						MCollisionTransforms[Index] = FSolverRigidTransform3(CollisionParticles.X(Index), CollisionParticles.R(Index));
+						MCollisionTransforms[Index] = FSolverRigidTransform3(CollisionParticles.X(Index), CollisionParticles.GetR(Index));
 
 						// Update collision transform and velocity
 						MCollisionKinematicUpdate(CollisionParticles, Dt, MTime, Index);
diff --git a/Engine/Source/Runtime/Experimental/Chaos/Private/Chaos/PBDCollisionConstraints.cpp b/Engine/Source/Runtime/Experimental/Chaos/Private/Chaos/PBDCollisionConstraints.cpp
--- a/Engine/Source/Runtime/Experimental/Chaos/Private/Chaos/PBDCollisionConstraints.cpp
+++ b/Engine/Source/Runtime/Experimental/Chaos/Private/Chaos/PBDCollisionConstraints.cpp
@@ -719,7 +719,7 @@ namespace Chaos
 
 					if (bCollisionsEnableSubSurfaceCollisionPruning)
 					{
-						const FVec3 UpVector = ParticleHandle.R().GetAxisZ();
+						const FVec3 UpVector = ParticleHandle.GetR().GetAxisZ();
 						FParticleSubSurfaceCollisionPruner SubSurfacePruner(ParticleHandle.Handle());
 						SubSurfacePruner.Prune(UpVector);
 					}
diff --git a/Engine/Source/Runtime/Experimental/Chaos/Private/Chaos/PBDEvolution.cpp b/Engine/Source/Runtime/Experimental/Chaos/Private/Chaos/PBDEvolution.cpp
--- a/Engine/Source/Runtime/Experimental/Chaos/Private/Chaos/PBDEvolution.cpp
+++ b/Engine/Source/Runtime/Experimental/Chaos/Private/Chaos/PBDEvolution.cpp
@@ -553,7 +553,7 @@ void FPBDEvolution::AdvanceOneTimeStep(const FSolverReal Dt)
 				[this, Dt](FSolverCollisionParticles& CollisionParticles, int32 Index)
 				{
 					// Store active collision particle frames prior to the kinematic update for CCD collisions
-					MCollisionTransforms[Index] = FSolverRigidTransform3(CollisionParticles.X(Index), CollisionParticles.R(Index));
+					MCollisionTransforms[Index] = FSolverRigidTransform3(CollisionParticles.X(Index), CollisionParticles.GetR(Index));
 
 					// Update collision transform and velocity
 					MCollisionKinematicUpdate(CollisionParticles, Dt, MTime, Index); // This expects Sequential update.
diff --git a/Engine/Source/Runtime/Experimental/Chaos/Private/Chaos/PBDJointConstraints.cpp b/Engine/Source/Runtime/Experimental/Chaos/Private/Chaos/PBDJointConstraints.cpp
--- a/Engine/Source/Runtime/Experimental/Chaos/Private/Chaos/PBDJointConstraints.cpp
+++ b/Engine/Source/Runtime/Experimental/Chaos/Private/Chaos/PBDJointConstraints.cpp
@@ -558,11 +558,11 @@ namespace Chaos
 		FPBDJointSettings JointSettings;
 		JointSettings.ConnectorTransforms[0] = FRigidTransform3(
 			WorldConstraintFrame.GetTranslation() - InConstrainedParticles[0]->X(),
-			WorldConstraintFrame.GetRotation() * InConstrainedParticles[0]->R().Inverse()
+			WorldConstraintFrame.GetRotation() * InConstrainedParticles[0]->GetR().Inverse()
 			);
 		JointSettings.ConnectorTransforms[1] = FRigidTransform3(
 			WorldConstraintFrame.GetTranslation() - InConstrainedParticles[1]->X(),
-			WorldConstraintFrame.GetRotation() * InConstrainedParticles[1]->R().Inverse()
+			WorldConstraintFrame.GetRotation() * InConstrainedParticles[1]->GetR().Inverse()
 			);
 		return AddConstraint(InConstrainedParticles, JointSettings);
 	}
diff --git a/Engine/Source/Runtime/Experimental/Chaos/Private/Chaos/PBDRigidClustering.cpp b/Engine/Source/Runtime/Experimental/Chaos/Private/Chaos/PBDRigidClustering.cpp
--- a/Engine/Source/Runtime/Experimental/Chaos/Private/Chaos/PBDRigidClustering.cpp
+++ b/Engine/Source/Runtime/Experimental/Chaos/Private/Chaos/PBDRigidClustering.cpp
@@ -711,8 +711,8 @@ namespace Chaos
 
 			// Update linear and angular impulses for the body, to be integrated next solve
 			const float RestorationPercent = RestoreBreakingMomentumPercent;
-			Rigid.V() += ImpulseVelocity * RestorationPercent;
-			Rigid.W() += AngularImpulseVelocity * RestorationPercent;
+			Rigid.SetV(Rigid.GetV() + ImpulseVelocity * RestorationPercent);
+			Rigid.SetW(Rigid.GetW() + AngularImpulseVelocity * RestorationPercent);
 		}
 	}
 
@@ -727,8 +727,8 @@ namespace Chaos
 				FBreakingData& ClusterBreak = MAllClusterBreakings.AddDefaulted_GetRef();
 				ClusterBreak.Proxy = ClusteredParticle->PhysicsProxy();
 				ClusterBreak.Location = ClusteredParticle->X();
-				ClusterBreak.Velocity = ClusteredParticle->V();
-				ClusterBreak.AngularVelocity = ClusteredParticle->W();
+				ClusterBreak.Velocity = ClusteredParticle->GetV();
+				ClusterBreak.AngularVelocity = ClusteredParticle->GetW();
 				ClusterBreak.Mass = ClusteredParticle->M();
 				if (ClusteredParticle->GetGeometry() && ClusteredParticle->GetGeometry()->HasBoundingBox())
 				{
@@ -753,9 +753,9 @@ namespace Chaos
 				FCrumblingData& ClusterCrumbling = MAllClusterCrumblings.AddDefaulted_GetRef();
 				ClusterCrumbling.Proxy = ClusteredParticle->PhysicsProxy();
 				ClusterCrumbling.Location = ClusteredParticle->X();
-				ClusterCrumbling.Orientation = ClusteredParticle->R();
-				ClusterCrumbling.LinearVelocity = ClusteredParticle->V();
-				ClusterCrumbling.AngularVelocity = ClusteredParticle->W();
+				ClusterCrumbling.Orientation = ClusteredParticle->GetR();
+				ClusterCrumbling.LinearVelocity = ClusteredParticle->GetV();
+				ClusterCrumbling.AngularVelocity = ClusteredParticle->GetW();
 				ClusterCrumbling.Mass = ClusteredParticle->M();
 				ClusterCrumbling.SetEmitterFlag(SimParams.bGenerateCrumblingData, SimParams.bGenerateGlobalCrumblingData);
 				if (ClusteredParticle->GetGeometry() && ClusteredParticle->GetGeometry()->HasBoundingBox())
@@ -860,22 +860,22 @@ namespace Chaos
 
 			ClusteredChild->SetClusterId(ClusterId(nullptr, ClusteredChild->ClusterIds().NumChildren)); // clear Id but retain number of children
 
-			const FRigidTransform3 PreSolveTM(ClusteredParent->P(), ClusteredParent->Q());
+			const FRigidTransform3 PreSolveTM(ClusteredParent->P(), ClusteredParent->GetQ());
 			const FRigidTransform3 ChildFrame = ClusteredChild->ChildToParent() * PreSolveTM;
 			Child->SetX(ChildFrame.GetTranslation());
 			Child->SetR(ChildFrame.GetRotation());
 
 			Child->SetP(Child->X());
-			Child->SetQ(Child->R());
+			Child->SetQf(Child->GetRf());
 
 			//todo(ocohen): for now just inherit velocity at new COM. This isn't quite right for rotation
 			//todo(ocohen): in the presence of collisions, this will leave all children with the post-collision
 			// velocity. This should be controlled by material properties so we can allow the broken pieces to
 			// maintain the clusters pre-collision velocity.
-			Child->SetV(Child->V() + ClusteredParent->V());
-			Child->SetW(Child->W() + ClusteredParent->W());
-			Child->SetPreV(Child->PreV() + ClusteredParent->PreV());
-			Child->SetPreW(Child->PreW() + ClusteredParent->PreW());
+			Child->SetVf(Child->GetVf() + ClusteredParent->GetVf());
+			Child->SetWf(Child->GetWf() + ClusteredParent->GetWf());
+			Child->SetPreVf(Child->GetPreVf() + ClusteredParent->GetPreVf());
+			Child->SetPreWf(Child->GetPreWf() + ClusteredParent->GetPreWf());
 
 			// We also need to do cluster book-keeping on the parent.
 			// If the parent is an internal cluster and has become empty, we need to mark this particle as ready to be destroyed.
@@ -928,7 +928,7 @@ namespace Chaos
 		}
 		NewClusters.Reserve(NumNewClusters);
 
-		const FRigidTransform3 PreSolveTM = FRigidTransform3(ClusteredParent->P(), ClusteredParent->Q());
+		const FRigidTransform3 PreSolveTM = FRigidTransform3(ClusteredParent->P(), ClusteredParent->GetQ());
 		
 		TArray<Chaos::FPBDRigidClusteredParticleHandle*> NewClusterHandles = MEvolution.CreateClusteredParticles(NumNewClusters);
 		int32 ClusterHandlesIdx = 0;
@@ -948,12 +948,12 @@ namespace Chaos
 				MEvolution.SetPhysicsMaterial(NewCluster, MEvolution.GetPhysicsMaterial(ClusteredParent));
 
 				NewCluster->SetInternalStrains(ClusteredParent->GetInternalStrains());
-				NewCluster->SetV(ClusteredParent->V());
-				NewCluster->SetW(ClusteredParent->W());
-				NewCluster->SetPreV(ClusteredParent->PreV());
-				NewCluster->SetPreW(ClusteredParent->PreW());
+				NewCluster->SetVf(ClusteredParent->GetVf());
+				NewCluster->SetWf(ClusteredParent->GetWf());
+				NewCluster->SetPreVf(ClusteredParent->GetPreVf());
+				NewCluster->SetPreWf(ClusteredParent->GetPreWf());
 				NewCluster->SetP(NewCluster->X());
-				NewCluster->SetQ(NewCluster->R());
+				NewCluster->SetQf(NewCluster->GetRf());
 
 				UpdateTopLevelParticle(NewCluster);
 
@@ -2205,8 +2205,8 @@ namespace Chaos
 			if (bUseContactSpeedForStrainThreshold)
 			{
 				// Get dV between the two particles and project onto the normal to get the approach speed (take PreV as V is the new velocity post-solve)
-				const FVec3 V0 = Rigid0 ? Rigid0->PreV() : FVec3(0);
-				const FVec3 V1 = Rigid1 ? Rigid1->PreV() : FVec3(0);
+				const FVec3 V0 = Rigid0 ? Rigid0->GetPreV() : FVec3(0);
+				const FVec3 V1 = Rigid1 ? Rigid1->GetPreV() : FVec3(0);
 				const FVec3 DeltaV = V0 - V1;
 				const FReal SpeedAlongNormal = FVec3::DotProduct(DeltaV, ContactHandle->GetContact().CalculateWorldContactNormal());
 
@@ -2296,7 +2296,7 @@ namespace Chaos
 					}
 					else
 					{
-						const FRigidTransform3 WorldToClusterTM = FRigidTransform3(Cluster->P(), Cluster->Q());
+						const FRigidTransform3 WorldToClusterTM = FRigidTransform3(Cluster->P(), Cluster->GetQ());
 						const FVec3 ContactLocationClusterLocal = WorldToClusterTM.InverseTransformPosition(ContactWorldLocation);
 						FAABB3 ContactBox(ContactLocationClusterLocal, ContactLocationClusterLocal);
 						ContactBox.Thicken(ClusterDistanceThreshold);
@@ -2636,7 +2636,7 @@ namespace Chaos
 					ParticlePairArray& ConnectionList = Connections[i];
 
 					const FVec3& Child1X = Child1->X();
-					FRigidTransform3 TM1 = FRigidTransform3(Child1X, Child1->R());
+					FRigidTransform3 TM1 = FRigidTransform3(Child1X, Child1->GetR());
 
 					const int32 Offset = i + 1;
 					const int32 NumRemainingParticles = Particles.Num() - Offset;
@@ -2649,7 +2649,7 @@ namespace Chaos
 						{
 
 							const FVec3& Child2X = Child2->X();
-							const FRigidTransform3 TM = TM1.GetRelativeTransform(FRigidTransform3(Child2X, Child2->R()));
+							const FRigidTransform3 TM = TM1.GetRelativeTransform(FRigidTransform3(Child2X, Child2->GetR()));
 							const uint32 NumCollisionParticles = Child2->CollisionParticles()->Size();
 							for (uint32 CollisionIdx = 0; CollisionIdx < NumCollisionParticles; ++CollisionIdx)
 							{
diff --git a/Engine/Source/Runtime/Experimental/Chaos/Private/Chaos/PBDRigidClusteringAlgo.cpp b/Engine/Source/Runtime/Experimental/Chaos/Private/Chaos/PBDRigidClusteringAlgo.cpp
--- a/Engine/Source/Runtime/Experimental/Chaos/Private/Chaos/PBDRigidClusteringAlgo.cpp
+++ b/Engine/Source/Runtime/Experimental/Chaos/Private/Chaos/PBDRigidClusteringAlgo.cpp
@@ -117,7 +117,7 @@ namespace Chaos
 		{
 			MassToLocal.SetRotation(Cluster->RotationOfMass());
 			Cluster->SetR(Cluster->RCom());
-			Cluster->SetQ(Cluster->R());
+			Cluster->SetQf(Cluster->GetRf());
 			Cluster->SetRotationOfMass(FQuat::Identity);
 		}
 
@@ -275,7 +275,7 @@ namespace Chaos
 		Objects.Reserve(Children.Num());
 		Objects2.Reserve(Children.Num());
 
-		const FRigidTransform3 ClusterWorldTM(Parent->X(), Parent->R());
+		const FRigidTransform3 ClusterWorldTM(Parent->X(), Parent->GetR());
 
 		TArray<FVec3> OriginalPoints;
 		TArray<FPBDRigidParticleHandle*> ChildParticleHandles;
@@ -307,7 +307,7 @@ namespace Chaos
 				}
 				else
 				{
-					const FRigidTransform3 ChildWorldTM(Child->X(), Child->R());
+					const FRigidTransform3 ChildWorldTM(Child->X(), Child->GetR());
 					Frame = ChildWorldTM.GetRelativeTransform(ClusterWorldTM);
 				}
 
@@ -458,7 +458,7 @@ namespace Chaos
 				Parent->CollisionParticles()->UpdateAccelerationStructures();
 			}
 		}
-		const Chaos::FRigidTransform3 Xf(Parent->X(), Parent->R());
+		const Chaos::FRigidTransform3 Xf(Parent->X(), Parent->GetR());
 		Parent->UpdateWorldSpaceState(Xf, FVec3(0));
 
 		UpdateCollisionFlags(Parent, bUseParticleImplicit);
diff --git a/Engine/Source/Runtime/Experimental/Chaos/Private/Chaos/PBDRigidDynamicSpringConstraints.cpp b/Engine/Source/Runtime/Experimental/Chaos/Private/Chaos/PBDRigidDynamicSpringConstraints.cpp
--- a/Engine/Source/Runtime/Experimental/Chaos/Private/Chaos/PBDRigidDynamicSpringConstraints.cpp
+++ b/Engine/Source/Runtime/Experimental/Chaos/Private/Chaos/PBDRigidDynamicSpringConstraints.cpp
@@ -33,8 +33,8 @@ namespace Chaos
 				continue;
 			}
 
-			const FRotation3 Q0 = bIsRigidDynamic0 ? PBDRigid0->Q() : Static0->R();
-			const FRotation3 Q1 = bIsRigidDynamic1 ? PBDRigid1->Q() : Static1->R();
+			const FRotation3 Q0 = bIsRigidDynamic0 ? PBDRigid0->GetQ() : Static0->GetR();
+			const FRotation3 Q1 = bIsRigidDynamic1 ? PBDRigid1->GetQ() : Static1->GetR();
 			const FVec3& P0 = bIsRigidDynamic0 ? PBDRigid0->P() : Static0->X();
 			const FVec3& P1 = bIsRigidDynamic1 ? PBDRigid1->P() : Static1->X();
 
diff --git a/Engine/Source/Runtime/Experimental/Chaos/Private/Chaos/PBDRigidSpringConstraints.cpp b/Engine/Source/Runtime/Experimental/Chaos/Private/Chaos/PBDRigidSpringConstraints.cpp
--- a/Engine/Source/Runtime/Experimental/Chaos/Private/Chaos/PBDRigidSpringConstraints.cpp
+++ b/Engine/Source/Runtime/Experimental/Chaos/Private/Chaos/PBDRigidSpringConstraints.cpp
@@ -96,8 +96,8 @@ namespace Chaos
 		const TGeometryParticleHandle<FReal, 3>* Particle0 = Constraint[0];
 		const TGeometryParticleHandle<FReal, 3>* Particle1 = Constraint[1];
 
-		Distances[ConstraintIndex][0] = Particle0->R().Inverse().RotateVector(Location0 - Particle0->X());
-		Distances[ConstraintIndex][1] = Particle1->R().Inverse().RotateVector(Location1 - Particle1->X());
+		Distances[ConstraintIndex][0] = Particle0->GetR().Inverse().RotateVector(Location0 - Particle0->X());
+		Distances[ConstraintIndex][1] = Particle1->GetR().Inverse().RotateVector(Location1 - Particle1->X());
 	}
 
 	FVec3 FPBDRigidSpringConstraints::GetDelta(int32 ConstraintIndex, const FVec3& WorldSpaceX1, const FVec3& WorldSpaceX2) const
diff --git a/Engine/Source/Runtime/Experimental/Chaos/Private/Chaos/PBDRigidsEvolution.cpp b/Engine/Source/Runtime/Experimental/Chaos/Private/Chaos/PBDRigidsEvolution.cpp
--- a/Engine/Source/Runtime/Experimental/Chaos/Private/Chaos/PBDRigidsEvolution.cpp
+++ b/Engine/Source/Runtime/Experimental/Chaos/Private/Chaos/PBDRigidsEvolution.cpp
@@ -1232,10 +1232,10 @@ namespace Chaos
 			}
 			Data->X = Rigid->X();
 			Data->P = Rigid->P();
-			Data->R = Rigid->R();
-			Data->Q = Rigid->Q();
-			Data->V = Rigid->V();
-			Data->W = Rigid->W();
+			Data->R = Rigid->GetR();
+			Data->Q = Rigid->GetQ();
+			Data->V = Rigid->GetV();
+			Data->W = Rigid->GetW();
 		}
 	}
 
@@ -1275,10 +1275,10 @@ namespace Chaos
 			{
 				Rigid->X() = Data->X;
 				Rigid->P() = Data->P;
-				Rigid->R() = Data->R;
-				Rigid->Q() = Data->Q;
-				Rigid->V() = Data->V;
-				Rigid->W() = Data->W;
+				Rigid->SetR(Data->R);
+				Rigid->SetQ(Data->Q);
+				Rigid->SetV(Data->V);
+				Rigid->SetW(Data->W);
 			}
 		}
 	}
diff --git a/Engine/Source/Runtime/Experimental/Chaos/Private/Chaos/PBDRigidsEvolutionGBF.cpp b/Engine/Source/Runtime/Experimental/Chaos/Private/Chaos/PBDRigidsEvolutionGBF.cpp
--- a/Engine/Source/Runtime/Experimental/Chaos/Private/Chaos/PBDRigidsEvolutionGBF.cpp
+++ b/Engine/Source/Runtime/Experimental/Chaos/Private/Chaos/PBDRigidsEvolutionGBF.cpp
@@ -288,10 +288,10 @@ void CheckMovingKinematicFlag(FPBDRigidsSOAs& Particles)
 		{
 			if (Rigid->IsKinematic())
 			{
-				const bool bIsMoving = !Rigid->V().IsNearlyZero() || !Rigid->W().IsNearlyZero();
+				const bool bIsMoving = !Rigid->GetV().IsNearlyZero() || !Rigid->GetW().IsNearlyZero();
 				ensureMsgf(bIsMoving == Rigid->IsMovingKinematic(),
 					TEXT("Kinematic IsMoving flag mismatch. IsMoving=%d V=(%f, %f, %f) W=(%f %f %f) %s"),
-					Rigid->IsMovingKinematic(), Rigid->V().X, Rigid->V().Y, Rigid->V().Z, Rigid->W().X, Rigid->W().Y, Rigid->W().Z, *Rigid->GetDebugName());
+					Rigid->IsMovingKinematic(), Rigid->GetV().X, Rigid->GetV().Y, Rigid->GetV().Z, Rigid->GetW().X, Rigid->GetW().Y, Rigid->GetW().Z, *Rigid->GetDebugName());
 			}
 			else
 			{
@@ -793,7 +793,7 @@ FPBDRigidsEvolutionGBF::FPBDRigidsEvolutionGBF(
 		ParticlesInput.ParallelFor([&](auto& Particle, int32 Index)
 		{
 			Particle.X() = Particle.P();
-			Particle.R() = Particle.Q();
+			Particle.SetR(Particle.GetQ());
 
 			//TODO: rename this function since it's not just updating position
 			Particle.SetPreObjectStateLowLevel(Particle.ObjectState());
@@ -1031,7 +1031,7 @@ void FPBDRigidsEvolutionGBF::DestroyTransientConstraints(FGeometryParticleHandle
 CHAOS_API void FPBDRigidsEvolutionGBF::SetParticleTransform(FGeometryParticleHandle* InParticle, const FVec3& InPos, const FRotation3& InRot, const bool bIsTeleport)
 {
 	const FVec3 PrevX = InParticle->X();
-	const FRotation3 PrevR = InParticle->R();
+	const FRotation3 PrevR = InParticle->GetR();
 
 	FGenericParticleHandle(InParticle)->SetTransform(InPos, InRot);
 
@@ -1053,7 +1053,7 @@ void FPBDRigidsEvolutionGBF::SetParticleTransformSwept(FGeometryParticleHandle*
 	// that are not centerd on their center of mass (either it's multi-shape body, or the user modified the CoM).
 	const FVec3 CoMOffset = FConstGenericParticleHandle(InParticle)->CenterOfMass();
 	const FVec3 EndCoM = InPos + InRot * CoMOffset;
-	const FVec3 StartCoM = InParticle->X() + InParticle->R() * CoMOffset;
+	const FVec3 StartCoM = InParticle->X() + InParticle->GetR() * CoMOffset;
 	FVec3 SweepDir = EndCoM - StartCoM;
 	const FReal SweepLength = SweepDir.SafeNormalize();
 	const FVec3 SweepStart = InPos - SweepDir * SweepLength;
@@ -1111,7 +1111,7 @@ void FPBDRigidsEvolutionGBF::OnParticleMoved(FGeometryParticleHandle* InParticle
 	// is different by very small amounts around 1e-7 in both position and rotation when switching from dynamic to kinematic.
 	const FReal CollisionPositionTolerance = FReal(1.e-4);
 	const FReal CollisionRotationTolerance = FReal(1.e-6);
-	if (!FVec3::IsNearlyEqual(PrevX, InParticle->X(), CollisionPositionTolerance) || !FRotation3::IsNearlyEqual(PrevR, InParticle->R(), CollisionRotationTolerance))
+	if (!FVec3::IsNearlyEqual(PrevX, InParticle->X(), CollisionPositionTolerance) || !FRotation3::IsNearlyEqual(PrevR, InParticle->GetR(), CollisionRotationTolerance))
 	{
 		GetIslandManager().WakeParticleIslands(InParticle);
 
@@ -1136,7 +1136,7 @@ void FPBDRigidsEvolutionGBF::OnParticleMoved(FGeometryParticleHandle* InParticle
 				const FReal InvDt = FReal(30.0);
 				const FVec3 DV = (PrevX - Rigid->X()) * InvDt;
 				const FReal SmoothRate = FMath::Clamp(CVars::SmoothedPositionLerpRate, 0.0f, 1.0f);
-				Rigid->VSmooth() = FMath::Lerp(Rigid->VSmooth(), Rigid->V() + DV, SmoothRate);
+				Rigid->VSmooth() = FMath::Lerp(Rigid->VSmooth(), Rigid->GetV() + DV, SmoothRate);
 			}
 
 			if (Rigid->IsSleeping())
diff --git a/Engine/Source/Runtime/Experimental/Chaos/Private/Chaos/PhysicsObjectInterface.cpp b/Engine/Source/Runtime/Experimental/Chaos/Private/Chaos/PhysicsObjectInterface.cpp
--- a/Engine/Source/Runtime/Experimental/Chaos/Private/Chaos/PhysicsObjectInterface.cpp
+++ b/Engine/Source/Runtime/Experimental/Chaos/Private/Chaos/PhysicsObjectInterface.cpp
@@ -105,7 +105,7 @@ namespace
 				{
 					if (const FPBDRigidClusteredParticleHandle* ClusteredParent = ClusteredParticle->Parent())
 					{
-						const FTransform ParentTransform(ClusteredParent->R(), ClusteredParent->X());
+						const FTransform ParentTransform(ClusteredParent->GetR(), ClusteredParent->X());
 						ParticleX = ParentTransform.TransformPosition(ClusteredParticle->ChildToParent().GetTranslation());
 					}
 				}
@@ -272,7 +272,7 @@ namespace Chaos
 
 		if (TThreadParticle<Id>* Particle = Object->GetParticle<Id>())
 		{
-			return Particle->R();
+			return Particle->GetR();
 		}
 
 		return FQuat::Identity;
@@ -290,7 +290,7 @@ namespace Chaos
 		{
 			if (Chaos::TThreadRigidParticle<Id>* Rigid = Particle->CastToRigidParticle())
 			{
-				return Rigid->V();
+				return Rigid->GetV();
 			}
 		}
 
@@ -311,7 +311,7 @@ namespace Chaos
 			{
 				const FVector CenterOfMass = GetWorldCoM(Object);
 				const FVector Diff = Point - CenterOfMass;
-				return Kinematic->V() - FVector::CrossProduct(Diff, Kinematic->W());
+				return Kinematic->GetV() - FVector::CrossProduct(Diff, Kinematic->GetW());
 			}
 		}
 
@@ -329,7 +329,7 @@ namespace Chaos
 		{
 			if (Chaos::TThreadRigidParticle<Id>* Rigid = Particle->CastToRigidParticle())
 			{
-				return Rigid->W();
+				return Rigid->GetW();
 			}
 		}
 
@@ -755,7 +755,7 @@ namespace Chaos
 			FBox ParticleBox(ForceInit);
 			if (const FImplicitObjectRef Geometry = Particle->GetGeometry(); Geometry && Geometry->HasBoundingBox())
 			{
-				const Chaos::FAABB3 WorldBox = Geometry->CalculateTransformedBounds(TRigidTransform<FReal, 3>(Particle->X(), Particle->R()));
+				const Chaos::FAABB3 WorldBox = Geometry->CalculateTransformedBounds(TRigidTransform<FReal, 3>(Particle->X(), Particle->GetR()));
 				ParticleBox = FBox{ WorldBox.Min(), WorldBox.Max() };
 			}
 
diff --git a/Engine/Source/Runtime/Experimental/Chaos/Private/EventDefaults.cpp b/Engine/Source/Runtime/Experimental/Chaos/Private/EventDefaults.cpp
--- a/Engine/Source/Runtime/Experimental/Chaos/Private/EventDefaults.cpp
+++ b/Engine/Source/Runtime/Experimental/Chaos/Private/EventDefaults.cpp
@@ -165,9 +165,9 @@ namespace Chaos
 								{
 									if (ensure(!Constraint.CalculateWorldContactLocation().ContainsNaN() &&
 										!Constraint.CalculateWorldContactNormal().ContainsNaN()) &&
-										!Primary->V().ContainsNaN() &&
-										!Primary->W().ContainsNaN() &&
-										(Secondary == nullptr || ((!Secondary->V().ContainsNaN()) && !Secondary->W().ContainsNaN())))
+										!Primary->GetV().ContainsNaN() &&
+										!Primary->GetW().ContainsNaN() &&
+										(Secondary == nullptr || ((!Secondary->GetV().ContainsNaN()) && !Secondary->GetW().ContainsNaN())))
 									{
 										ValidArray[Index] = true;
 									}
@@ -230,11 +230,11 @@ namespace Chaos
 
 								if (const FPBDRigidParticleHandle* Rigid0 = Particle0->CastToRigidParticle())
 								{
-									Data.DeltaVelocity1 = Rigid0->V() - Rigid0->PreV();
+									Data.DeltaVelocity1 = Rigid0->GetV() - Rigid0->GetPreV();
 								}
 								if (const FPBDRigidParticleHandle* Rigid1 = Particle1->CastToRigidParticle())
 								{
-									Data.DeltaVelocity2 = Rigid1->V() - Rigid1->PreV();
+									Data.DeltaVelocity2 = Rigid1->GetV() - Rigid1->GetPreV();
 								}
 
 								// todo: do we need these anymore now we are storing the particles you can access all of this stuff from there
@@ -242,16 +242,16 @@ namespace Chaos
 								const FPBDRigidParticleHandle* PBDRigid0 = Particle0->CastToRigidParticle();
 								if (PBDRigid0 && PBDRigid0->ObjectState() == EObjectStateType::Dynamic)
 								{
-									Data.Velocity1 = PBDRigid0->V();
-									Data.AngularVelocity1 = PBDRigid0->W();
+									Data.Velocity1 = PBDRigid0->GetV();
+									Data.AngularVelocity1 = PBDRigid0->GetW();
 									Data.Mass1 = PBDRigid0->M();
 								}
 
 								const FPBDRigidParticleHandle* PBDRigid1 = Particle1->CastToRigidParticle();
 								if (PBDRigid1 && PBDRigid1->ObjectState() == EObjectStateType::Dynamic)
 								{
-									Data.Velocity2 = PBDRigid1->V();
-									Data.AngularVelocity2 = PBDRigid1->W();
+									Data.Velocity2 = PBDRigid1->GetV();
+									Data.AngularVelocity2 = PBDRigid1->GetW();
 									Data.Mass2 = PBDRigid1->M();
 								}
 
@@ -394,14 +394,14 @@ namespace Chaos
 						ActiveParticle->GetGeometry()->HasBoundingBox())
 					{
 						if (ensure(!ActiveParticle->X().ContainsNaN() &&
-							!ActiveParticle->V().ContainsNaN() &&
-							!ActiveParticle->W().ContainsNaN() &&
+							!ActiveParticle->GetV().ContainsNaN() &&
+							!ActiveParticle->GetW().ContainsNaN() &&
 							FMath::IsFinite(ActiveParticle->M())))
 						{
 							FTrailingData TrailingData;
 							TrailingData.Location = ActiveParticle->X();
-							TrailingData.Velocity = ActiveParticle->V();
-							TrailingData.AngularVelocity = ActiveParticle->W();
+							TrailingData.Velocity = ActiveParticle->GetV();
+							TrailingData.AngularVelocity = ActiveParticle->GetW();
 							TrailingData.Mass = ActiveParticle->M();
 							TrailingData.Proxy = ActiveParticle->PhysicsProxy();
 							
diff --git a/Engine/Source/Runtime/Experimental/Chaos/Private/PBDRigidsSolver.cpp b/Engine/Source/Runtime/Experimental/Chaos/Private/PBDRigidsSolver.cpp
--- a/Engine/Source/Runtime/Experimental/Chaos/Private/PBDRigidsSolver.cpp
+++ b/Engine/Source/Runtime/Experimental/Chaos/Private/PBDRigidsSolver.cpp
@@ -1896,12 +1896,12 @@ namespace Chaos
 					if(!Handle.IsSleeping())
 					{
 						UE_LOG(LogChaos, Log, TEXT("Particle Dynamic At Position = %s | Velocity = %s | Quaternion = %s | Omega = %s"),
-							*Handle.X().ToString(), *Handle.V().ToString(), *Handle.R().ToString(), *Handle.W().ToString());
+							*Handle.X().ToString(), *Handle.GetV().ToString(), *Handle.GetR().ToString(), *Handle.GetW().ToString());
 					}
 					else
 					{
 						UE_LOG(LogChaos, Log, TEXT("Particle Sleeping At Position = %s | Velocity = %s | Quaternion = %s | Omega = %s"), 
-							*Handle.X().ToString(), *Handle.V().ToString(), *Handle.R().ToString(), *Handle.W().ToString());
+							*Handle.X().ToString(), *Handle.GetV().ToString(), *Handle.GetR().ToString(), *Handle.GetW().ToString());
 					}
 				}
 			}
@@ -1951,14 +1951,14 @@ namespace Chaos
 										if (!Handle.IsSleeping())
 										{
 											UE_LOG(LogChaos, Log, TEXT("Particle Dynamic At Position = %s | Velocity = %s | Quaternion = %s | Omega = %s | Resim Frame = %d | Sync State = %d | Needs Resim = %d | Unique Idx = %d"),
-												*Handle.X().ToString(), *Handle.V().ToString(), *Handle.R().ToString(), *Handle.W().ToString(),
+												*Handle.X().ToString(), *Handle.GetV().ToString(), *Handle.GetR().ToString(), *Handle.GetW().ToString(),
 												GetEvolution()->GetIslandManager().GetParticleIsland(Handle.Handle())->GetResimFrame(), (uint8)Handle.SyncState(),
 												GetEvolution()->GetIslandManager().GetParticleIsland(Handle.Handle())->NeedsResim(), Handle.UniqueIdx().Idx);
 										}
 										else
 										{
 											UE_LOG(LogChaos, Log, TEXT("Particle Sleeping At Position = %s | Velocity = %s | Quaternion = %s | Omega = %s | Resim Frame = %d | Sync State = %d | Needs Resim = %d | Unique Idx = %d"),
-												*Handle.X().ToString(), *Handle.V().ToString(), *Handle.R().ToString(), *Handle.W().ToString(),
+												*Handle.X().ToString(), *Handle.GetV().ToString(), *Handle.GetR().ToString(), *Handle.GetW().ToString(),
 												GetEvolution()->GetIslandManager().GetParticleIsland(Handle.Handle())->GetResimFrame(), (uint8)Handle.SyncState(),
 												GetEvolution()->GetIslandManager().GetParticleIsland(Handle.Handle())->NeedsResim(), Handle.UniqueIdx().Idx);
 										}
diff --git a/Engine/Source/Runtime/Experimental/Chaos/Private/PhysicsProxy/ClusterUnionPhysicsProxy.cpp b/Engine/Source/Runtime/Experimental/Chaos/Private/PhysicsProxy/ClusterUnionPhysicsProxy.cpp
--- a/Engine/Source/Runtime/Experimental/Chaos/Private/PhysicsProxy/ClusterUnionPhysicsProxy.cpp
+++ b/Engine/Source/Runtime/Experimental/Chaos/Private/PhysicsProxy/ClusterUnionPhysicsProxy.cpp
@@ -43,9 +43,9 @@ namespace Chaos
 
 			BufferData.SetProxy(*Proxy);
 			BufferData.X = Particle->X();
-			BufferData.R = Particle->R();
-			BufferData.V = Particle->V();
-			BufferData.W = Particle->W();
+			BufferData.R = Particle->GetR();
+			BufferData.V = Particle->GetV();
+			BufferData.W = Particle->GetW();
 			BufferData.ObjectState = Particle->ObjectState();
 			BufferData.Geometry = nullptr;
 
@@ -132,10 +132,10 @@ namespace Chaos
 			Particle_Internal->SetUnbreakable(InitData.bUnbreakable);
 			Particle_Internal->SetX(InitData.InitialTransform.GetTranslation());
 			Particle_Internal->SetR(InitData.InitialTransform.GetRotation());
-			Particle_Internal->SetV(Chaos::FVec3(0.f));
-			Particle_Internal->SetW(Chaos::FVec3(0.f));
+			Particle_Internal->SetVf(Chaos::FVec3f(0.f));
+			Particle_Internal->SetWf(Chaos::FVec3f(0.f));
 			Particle_Internal->SetP(Particle_Internal->X());
-			Particle_Internal->SetQ(Particle_Internal->R());
+			Particle_Internal->SetQf(Particle_Internal->GetRf());
 			Particle_Internal->SetCenterOfMass(FVector3f::ZeroVector);
 			Particle_Internal->SetRotationOfMass(FQuat::Identity);
 
@@ -380,7 +380,7 @@ namespace Chaos
 						{
 							ParticlesToUpdate.Add(Particle);
 
-							const FRigidTransform3 ChildWorldTM(Particle->X(), Particle->R());
+							const FRigidTransform3 ChildWorldTM(Particle->X(), Particle->GetR());
 							NewChildToParent.Add(ChildWorldTM.GetRelativeTransform(NewTransform));
 						}
 					}
@@ -400,7 +400,7 @@ namespace Chaos
 			// but this is primarily for proxies within the cluster union. They (i.e. GCs) need to see this change to the particle's transform.
 			Evolution.GetParticles().MarkTransientDirtyParticle(Particle_Internal);
 
-			const FRigidTransform3 WorldTransform{ Particle_Internal->X(), Particle_Internal->R() };
+			const FRigidTransform3 WorldTransform{ Particle_Internal->X(), Particle_Internal->GetR() };
 			Particle_Internal->UpdateWorldSpaceState(WorldTransform, FVec3(0));
 
 			Evolution.DirtyParticle(*Particle_Internal);
diff --git a/Engine/Source/Runtime/Experimental/Chaos/Private/PhysicsProxy/GeometryCollectionPhysicsProxy.cpp b/Engine/Source/Runtime/Experimental/Chaos/Private/PhysicsProxy/GeometryCollectionPhysicsProxy.cpp
--- a/Engine/Source/Runtime/Experimental/Chaos/Private/PhysicsProxy/GeometryCollectionPhysicsProxy.cpp
+++ b/Engine/Source/Runtime/Experimental/Chaos/Private/PhysicsProxy/GeometryCollectionPhysicsProxy.cpp
@@ -337,7 +337,7 @@ void SetImplicitToPTParticles(Chaos::TPBDRigidParticleHandle<Chaos::FReal, 3>* H
 			Handle->SetLocalBounds(SharedImplicitTS->BoundingBox());
 		}
 		Handle->SetHasBounds(true);
-		const Chaos::FRigidTransform3 Xf(Handle->X(), Handle->R());
+		const Chaos::FRigidTransform3 Xf(Handle->X(), Handle->GetR());
 		Handle->UpdateWorldSpaceState(Xf, Chaos::FVec3(0));
 	}
 
@@ -423,11 +423,11 @@ void PopulateSimulatedParticle(
 	SCOPE_CYCLE_COUNTER(STAT_PopulateSimulatedParticle);
 	Handle->SetDisabledLowLevel(false);
 	Handle->SetX(WorldTransform.GetTranslation());
-	Handle->SetV(Chaos::FVec3(0.f));
+	Handle->SetVf(Chaos::FVec3f(0.f));
 	Handle->SetR(WorldTransform.GetRotation().GetNormalized());
-	Handle->SetW(Chaos::FVec3(0.f));
+	Handle->SetWf(Chaos::FVec3f(0.f));
 	Handle->SetP(Handle->X());
-	Handle->SetQ(Handle->R());
+	Handle->SetQf(Handle->GetRf());
 	Handle->SetCenterOfMass(FVector3f::ZeroVector);
 	Handle->SetRotationOfMass(FQuat::Identity);
 
@@ -1998,7 +1998,7 @@ Chaos::TPBDGeometryCollectionParticleHandle<Chaos::FReal, 3>* FGeometryCollectio
 	{
 		Handle->SetHasBounds(true);
 		Handle->SetLocalBounds(Handle->GetGeometry()->BoundingBox());
-		const Chaos::FRigidTransform3 Xf(Handle->X(), Handle->R());
+		const Chaos::FRigidTransform3 Xf(Handle->X(), Handle->GetR());
 		Handle->UpdateWorldSpaceState(Xf, Chaos::FVec3(0));
 
 		static_cast<Chaos::FPBDRigidsSolver*>(Solver)->GetEvolution()->DirtyParticle(*Handle);
@@ -2147,7 +2147,7 @@ FGeometryCollectionPhysicsProxy::BuildClusters_Internal(
 	{
 		Parent->SetHasBounds(true);
 		Parent->SetLocalBounds(Parent->GetGeometry()->BoundingBox());
-		const Chaos::FRigidTransform3 Xf(Parent->X(), Parent->R());
+		const Chaos::FRigidTransform3 Xf(Parent->X(), Parent->GetR());
 		Parent->UpdateWorldSpaceState(Xf, Chaos::FVec3(0));
 
 		static_cast<Chaos::FPBDRigidsSolver*>(Solver)->GetEvolution()->DirtyParticle(*Parent);
@@ -2438,7 +2438,7 @@ void FGeometryCollectionPhysicsProxy::ApplyImpulseAt_External(FVector Impulse, F
 					const Chaos::FVec3 WorldCOM = Chaos::FParticleUtilitiesXR::GetCoMWorldPosition(ClosestHandle);
 					ClosestHandle->SetLinearImpulseVelocity(ClosestHandle->LinearImpulseVelocity() + (Impulse * ClosestHandle->InvM()), false);
 
-					const Chaos::FMatrix33 WorldInvI = Chaos::Utilities::ComputeWorldSpaceInertia(ClosestHandle->R() * ClosestHandle->RotationOfMass(), ClosestHandle->InvI());
+					const Chaos::FMatrix33 WorldInvI = Chaos::Utilities::ComputeWorldSpaceInertia(ClosestHandle->GetR() * ClosestHandle->RotationOfMass(), ClosestHandle->InvI());
 					const Chaos::FVec3 AngularImpulse = Chaos::FVec3::CrossProduct(WorldLocation - WorldCOM, Impulse);
 					ClosestHandle->SetAngularImpulseVelocity(ClosestHandle->AngularImpulseVelocity() + WorldInvI * AngularImpulse, false);
 				}
@@ -2774,7 +2774,7 @@ void FGeometryCollectionPhysicsProxy::ApplyBreakingLinearVelocity_External(FGeom
 				ApplyToBreakingChildren_Internal(RBDSolver->GetEvolution()->GetRigidClustering(), *ClusteredHandle,
 					[LinearVelocity](Chaos::FPBDRigidClusteredParticleHandle* ClusteredChildHandle)
 					{
-						ClusteredChildHandle->V() += LinearVelocity;
+						ClusteredChildHandle->SetV(ClusteredChildHandle->GetV() + LinearVelocity);
 					});
 			}
 		});
@@ -2794,7 +2794,7 @@ void FGeometryCollectionPhysicsProxy::ApplyBreakingAngularVelocity_External(FGeo
 				ApplyToBreakingChildren_Internal(RBDSolver->GetEvolution()->GetRigidClustering(), *ClusteredHandle,
 					[AngularVelocity](Chaos::FPBDRigidClusteredParticleHandle* ClusteredChildHandle)
 					{
-						ClusteredChildHandle->W() += AngularVelocity;
+						ClusteredChildHandle->SetW(ClusteredChildHandle->GetW() + AngularVelocity);
 					});
 			}
 		});
@@ -2811,7 +2811,7 @@ void FGeometryCollectionPhysicsProxy::ApplyLinearVelocity_External(FGeometryColl
 		{
 			if (Chaos::FPBDRigidClusteredParticleHandle* ClusteredHandle = FindClusteredParticleHandleByItemIndex_Internal(ItemIndex))
 			{
-				ClusteredHandle->V() += LinearVelocity;
+				ClusteredHandle->SetV(ClusteredHandle->GetV() + LinearVelocity);
 			}
 		});
 	}
@@ -2827,7 +2827,7 @@ void FGeometryCollectionPhysicsProxy::ApplyAngularVelocity_External(FGeometryCol
 		{
 			if (Chaos::FPBDRigidClusteredParticleHandle* ClusteredHandle = FindClusteredParticleHandleByItemIndex_Internal(ItemIndex))
 			{
-				ClusteredHandle->W() += AngularVelocity;
+				ClusteredHandle->SetW(ClusteredHandle->GetW() + AngularVelocity);
 			}
 		});
 	}
@@ -3306,7 +3306,7 @@ void FGeometryCollectionPhysicsProxy::SetWorldTransform_Internal(const FTransfor
 
 				if (KinematicRootHandle)
 				{
-					const FTransform RootWorldTransform(KinematicRootHandle->R(), KinematicRootHandle->X());
+					const FTransform RootWorldTransform(KinematicRootHandle->GetR(), KinematicRootHandle->X());
 					const FTransform RootRelativeTransform = RootWorldTransform.GetRelativeTransform(Parameters.PrevWorldTransform);
 					const FTransform WorldTransform = RootRelativeTransform * ActorToWorld;
 
@@ -3322,7 +3322,7 @@ void FGeometryCollectionPhysicsProxy::SetWorldTransform_Internal(const FTransfor
 					if (ClusterUnionIndex != INDEX_NONE)
 					{
 						const int32 TransformGroupIndex = FromParticleToTransformIndex[ParticleIndex];
-						const FTransform ParentWorldTransform{ ParentHandle->R(), ParentHandle->X() };
+						const FTransform ParentWorldTransform{ ParentHandle->GetR(), ParentHandle->X() };
 						const FTransform NewWorldTransform = MassToLocal[TransformGroupIndex] * FTransform(PhysicsThreadCollection.GetTransform(TransformGroupIndex)) * Parameters.WorldTransform;
 						const FTransform RelativeTransform = NewWorldTransform.GetRelativeTransform(ParentWorldTransform);
 
@@ -3710,7 +3710,7 @@ void FGeometryCollectionPhysicsProxy::BufferPhysicsResults_External(Chaos::FDirt
 static void UpdateParticleHandleTransform(Chaos::FPBDRigidsSolver& CurrentSolver, Chaos::FPBDRigidClusteredParticleHandle& Handle, const Chaos::FRigidTransform3& NewTransform)
 {
 	Handle.X() = NewTransform.GetTranslation();
-	Handle.R() = NewTransform.GetRotation();
+	Handle.SetR(NewTransform.GetRotation());
 	Handle.UpdateWorldSpaceState(NewTransform, Chaos::FVec3{0});
 	CurrentSolver.GetEvolution()->DirtyParticle(Handle);
 }
@@ -3721,7 +3721,7 @@ static void UpdateParticleHandleTransformIfNeeded(Chaos::FPBDRigidsSolver& Curre
 	const Chaos::FVec3 OldX = Handle.X();
 	
 	const Chaos::FRotation3 NewR  = NewTransform.GetRotation();
-	const Chaos::FRotation3 OldR = Handle.R();
+	const Chaos::FRotation3 OldR = Handle.GetR();
 	
 	if (OldX != NewX || OldR != NewR)
 	{
@@ -3802,11 +3802,11 @@ void FGeometryCollectionPhysicsProxy::BufferPhysicsResults_Internal(Chaos::FPBDR
 
 			FGeometryCollectionResults::FPositionData PositionData;
 			PositionData.ParticleX = Handle->X();
-			PositionData.ParticleR = Handle->R();
+			PositionData.ParticleR = Handle->GetR();
 
 			FGeometryCollectionResults::FVelocityData VelocityData;
-			VelocityData.ParticleV = Handle->V();
-			VelocityData.ParticleW = Handle->W();
+			VelocityData.ParticleV = Handle->GetV();
+			VelocityData.ParticleW = Handle->GetW();
 
 			FGeometryCollectionResults::FStateData StateData;
 			StateData.TransformIndex = TransformGroupIndex;
@@ -3875,7 +3875,7 @@ void FGeometryCollectionPhysicsProxy::BufferPhysicsResults_Internal(Chaos::FPBDR
 
 				// child of internal clusters may not have their position up to date 
 				// so we need to compute it from the ChildToParent transform 
-				const FTransform ParticleToWorld = Handle->ChildToParent() * FRigidTransform3(ClusterParent->X(), ClusterParent->R()); // aka ClusterChildToWorld
+				const FTransform ParticleToWorld = Handle->ChildToParent() * FRigidTransform3(ClusterParent->X(), ClusterParent->GetR()); // aka ClusterChildToWorld
 				PositionData.ParticleX = ParticleToWorld.GetTranslation();
 				PositionData.ParticleR = ParticleToWorld.GetRotation();
 				
@@ -3895,7 +3895,7 @@ void FGeometryCollectionPhysicsProxy::BufferPhysicsResults_Internal(Chaos::FPBDR
 			{
 				if (ClusterParent && !ClusterParent->Disabled())
 				{
-					const FRigidTransform3 ChildToWorld = Handle->ChildToParent() * FRigidTransform3(ClusterParent->X(), ClusterParent->R());
+					const FRigidTransform3 ChildToWorld = Handle->ChildToParent() * FRigidTransform3(ClusterParent->X(), ClusterParent->GetR());
 					UpdateParticleHandleTransformIfNeeded(*CurrentSolver, *Handle, ChildToWorld);
 					// fields may have applied velocities, we need to make sure to clear that up, so that we don't accumulate
 					VelocityData.ParticleV = FVec3::ZeroVector;
diff --git a/Engine/Source/Runtime/Experimental/Chaos/Private/PhysicsProxy/SingleParticlePhysicsProxy.cpp b/Engine/Source/Runtime/Experimental/Chaos/Private/PhysicsProxy/SingleParticlePhysicsProxy.cpp
--- a/Engine/Source/Runtime/Experimental/Chaos/Private/PhysicsProxy/SingleParticlePhysicsProxy.cpp
+++ b/Engine/Source/Runtime/Experimental/Chaos/Private/PhysicsProxy/SingleParticlePhysicsProxy.cpp
@@ -104,7 +104,7 @@ void PushToPhysicsStateImp(const Chaos::FDirtyPropertiesManager& Manager, Chaos:
 			// Update world-space cached state like the bounds
 			// @todo(chaos): do we need to do this here? It should be done in Integrate and ApplyKinematicTarget so only really Statics need this...
 			const bool bHasKinematicTarget = (NewKinematicTargetGT != nullptr) && (NewKinematicTargetGT->GetMode() == EKinematicTargetMode::Position);
-			const FRigidTransform3 WorldTransform = !bHasKinematicTarget ? FRigidTransform3(Handle->X(), Handle->R()) : NewKinematicTargetGT->GetTarget();
+			const FRigidTransform3 WorldTransform = !bHasKinematicTarget ? FRigidTransform3(Handle->X(), Handle->GetR()) : NewKinematicTargetGT->GetTarget();
 			Handle->UpdateWorldSpaceState(WorldTransform, FVec3(0));
 
 			Evolution.DirtyParticle(*Handle);
@@ -227,9 +227,9 @@ template <typename T>
 void BufferPhysicsResultsImp(Chaos::FDirtyRigidParticleData& PullData, T* Particle)
 {
 	PullData.X = Particle->X();
-	PullData.R = Particle->R();
-	PullData.V = Particle->V();
-	PullData.W = Particle->W();
+	PullData.R = Particle->GetR();
+	PullData.V = Particle->GetV();
+	PullData.W = Particle->GetW();
 	PullData.ObjectState = Particle->ObjectState();
 }
 
diff --git a/Engine/Source/Runtime/Experimental/Chaos/Private/RewindData.cpp b/Engine/Source/Runtime/Experimental/Chaos/Private/RewindData.cpp
--- a/Engine/Source/Runtime/Experimental/Chaos/Private/RewindData.cpp
+++ b/Engine/Source/Runtime/Experimental/Chaos/Private/RewindData.cpp
@@ -418,7 +418,7 @@ void FRewindData::AccumulateErrorIfNecessary(FGeometryParticleHandle& Obj, const
 
 	// Get the error offset after a correction
 	const FVec3 ErrorX = DirtyInfo->GetHistory().PreCorrectionXR.X() - Obj.X();
-	FQuat ErrorR = DirtyInfo->GetHistory().PreCorrectionXR.R() * Obj.R().Inverse();
+	FQuat ErrorR = DirtyInfo->GetHistory().PreCorrectionXR.R() * Obj.GetR().Inverse();
 	ErrorR.EnforceShortestArcWith(FQuat::Identity);
 	ErrorR.Normalize();
 
@@ -1091,8 +1091,8 @@ void FRewindData::PushPTDirtyData(TPBDRigidParticleHandle<FReal,3>& Handle,const
 	{
 		if (FParticleVelocities* PreVelocities = Latest.Velocities.WriteAccessNonDecreasing(FrameAndPhase, PropertiesPool))
 		{
-			PreVelocities->SetV(Handle.PreV());
-			PreVelocities->SetW(Handle.PreW());
+			PreVelocities->SetV(Handle.GetPreV());
+			PreVelocities->SetW(Handle.GetPreW());
 		}
 	}
 	
diff --git a/Engine/Source/Runtime/Experimental/Chaos/Public/Chaos/BoundingVolumeUtilities.h b/Engine/Source/Runtime/Experimental/Chaos/Public/Chaos/BoundingVolumeUtilities.h
--- a/Engine/Source/Runtime/Experimental/Chaos/Public/Chaos/BoundingVolumeUtilities.h
+++ b/Engine/Source/Runtime/Experimental/Chaos/Public/Chaos/BoundingVolumeUtilities.h
@@ -52,7 +52,7 @@ inline FVec3 ComputeBoundsThickness(const TPBDRigidParticles<FReal, 3>& InPartic
 		}
 	}
 
-	return ComputeBoundsThickness(InParticles.V(BodyIndex), Dt, MinBoundsThickness, MaxBoundsThickness, BoundsVelocityInflation);
+	return ComputeBoundsThickness(InParticles.GetV(BodyIndex), Dt, MinBoundsThickness, MaxBoundsThickness, BoundsVelocityInflation);
 }
 
 template <typename THandle>
@@ -174,7 +174,7 @@ template<class T, int d>
 TAABB<T, d> ComputeWorldSpaceBoundingBox(const TGeometryParticles<T, d>& Objects, const int32 i, bool bUseVelocity = false, T Dt = 0)
 {
 	ensure(!bUseVelocity);
-	TRigidTransform<T, d> LocalToWorld(Objects.X(i), Objects.R(i));
+	TRigidTransform<T, d> LocalToWorld(Objects.X(i), Objects.GetR(i));
 	const auto& LocalBoundingBox = Objects.GetGeometry(i)->BoundingBox();
 	return LocalBoundingBox.TransformedAABB(LocalToWorld);
 }
@@ -182,7 +182,7 @@ TAABB<T, d> ComputeWorldSpaceBoundingBox(const TGeometryParticles<T, d>& Objects
 template<class T, int d>
 TAABB<T, d> ComputeWorldSpaceBoundingBox(const TPBDRigidParticles<T, d>& Objects, const int32 i, bool bUseVelocity = false, T Dt = 0)
 {
-	TRigidTransform<T, d> LocalToWorld(Objects.P(i), Objects.Q(i));
+	TRigidTransform<T, d> LocalToWorld(Objects.P(i), Objects.GetQ(i));
 	TAABB<T, d> WorldSpaceBox;
 	if (Objects.GetGeometry(i))
 	{
diff --git a/Engine/Source/Runtime/Experimental/Chaos/Public/Chaos/DynamicParticles.h b/Engine/Source/Runtime/Experimental/Chaos/Public/Chaos/DynamicParticles.h
--- a/Engine/Source/Runtime/Experimental/Chaos/Public/Chaos/DynamicParticles.h
+++ b/Engine/Source/Runtime/Experimental/Chaos/Public/Chaos/DynamicParticles.h
@@ -29,7 +29,9 @@ class TDynamicParticles : public TParticles<T, d>
 	}
 
 	const TVector<T, d>& V(const int32 Index) const { return MV[Index]; }
+	const TVector<T, d> GetV(const int32 Index) const { return MV[Index]; }
 	TVector<T, d>& V(const int32 Index) { return MV[Index]; }
+	void SetV(const int32 Index, const TVector<T, d>& InV) { MV[Index] = InV; }
 	const TArrayCollectionArray<TVector<T, d>>& GetV() const { return MV; }
 	TArrayCollectionArray<TVector<T, d>>& GetV() { return MV; }
 
diff --git a/Engine/Source/Runtime/Experimental/Chaos/Public/Chaos/EvolutionResimCache.h b/Engine/Source/Runtime/Experimental/Chaos/Public/Chaos/EvolutionResimCache.h
--- a/Engine/Source/Runtime/Experimental/Chaos/Public/Chaos/EvolutionResimCache.h
+++ b/Engine/Source/Runtime/Experimental/Chaos/Public/Chaos/EvolutionResimCache.h
@@ -45,9 +45,9 @@ namespace Chaos
 		{
 			FPBDSolveCache& Cache = ParticleToCachedSolve.FindOrAdd(Particle.UniqueIdx());
 			Cache.P = Particle.P();
-			Cache.Q = Particle.Q();
-			Cache.V = Particle.V();
-			Cache.W = Particle.W();
+			Cache.Q = Particle.GetQ();
+			Cache.V = Particle.GetV();
+			Cache.W = Particle.GetW();
 		}
 
 		void ReloadParticlePostSolve(FPBDRigidParticleHandle& Particle) const
@@ -57,9 +57,9 @@ namespace Chaos
 			if (Cache)
 			{
 				Particle.P() = Cache->P;
-				Particle.Q() = Cache->Q;
-				Particle.V() = Cache->V;
-				Particle.W() = Cache->W;
+				Particle.SetQ(Cache->Q);
+				Particle.SetV(Cache->V);
+				Particle.SetW(Cache->W);
 			}
 		}
 
diff --git a/Engine/Source/Runtime/Experimental/Chaos/Public/Chaos/GeometryParticles.h b/Engine/Source/Runtime/Experimental/Chaos/Public/Chaos/GeometryParticles.h
--- a/Engine/Source/Runtime/Experimental/Chaos/Public/Chaos/GeometryParticles.h
+++ b/Engine/Source/Runtime/Experimental/Chaos/Public/Chaos/GeometryParticles.h
@@ -155,7 +155,7 @@ namespace Chaos
 
 		using TArrayCollection::Size;
 		using TParticles<T,d>::X;
-		using TSimpleGeometryParticles<T, d>::R;
+		using TSimpleGeometryParticles<T, d>::GetR;
 		using TSimpleGeometryParticles<T, d>::GetGeometry;
 		using TSimpleGeometryParticles<T, d>::SetGeometry;
 		using TSimpleGeometryParticles<T, d>::GetAllGeometry;
@@ -320,7 +320,7 @@ namespace Chaos
 				// Update the world-space stat of all the shapes - must be called after UpdateShapesArray
 				// world space inflated bounds needs to take expansion into account - this is done in integrate for dynamics anyway, so
 				// this computation is mainly for statics
-				UpdateWorldSpaceState(Index, TRigidTransform<FReal, 3>(X(Index), R(Index)), FVec3(0));
+				UpdateWorldSpaceState(Index, TRigidTransform<FReal, 3>(X(Index), GetR(Index)), FVec3(0));
 			}
 		}
 
@@ -557,7 +557,7 @@ public:
 		FString ToString(int32 index) const
 		{
 			FString BaseString = TParticles<T, d>::ToString(index);
-			return FString::Printf(TEXT("%s, MUniqueIdx:%d MR:%s, MGeometry:%s"), *BaseString, UniqueIdx(index).Idx, *R(index).ToString(), (GetGeometry(index) ? *(GetGeometry(index)->ToString()) : TEXT("none")));
+			return FString::Printf(TEXT("%s, MUniqueIdx:%d MR:%s, MGeometry:%s"), *BaseString, UniqueIdx(index).Idx, *GetR(index).ToString(), (GetGeometry(index) ? *(GetGeometry(index)->ToString()) : TEXT("none")));
 		}
 
 		virtual void Serialize(FChaosArchive& Ar) override
@@ -587,7 +587,7 @@ public:
 				{
 					for (int32 Idx = 0; Idx < MShapesArray.Num(); ++Idx)
 					{
-						UpdateWorldSpaceState(Idx, FRigidTransform3(X(Idx), R(Idx)), FVec3(0));
+						UpdateWorldSpaceState(Idx, FRigidTransform3(X(Idx), GetR(Idx)), FVec3(0));
 					}
 				}
 			}
@@ -601,7 +601,7 @@ public:
 					{
 						MLocalBounds[Idx] = TAABB<T, d>(GetGeometry(Idx)->BoundingBox());
 						//ignore velocity too, really just trying to get something reasonable)
-						UpdateWorldSpaceState(Idx, FRigidTransform3(X(Idx), R(Idx)), FVec3(0));
+						UpdateWorldSpaceState(Idx, FRigidTransform3(X(Idx), GetR(Idx)), FVec3(0));
 					}
 				}
 			}
diff --git a/Engine/Source/Runtime/Experimental/Chaos/Public/Chaos/KinematicGeometryParticles.h b/Engine/Source/Runtime/Experimental/Chaos/Public/Chaos/KinematicGeometryParticles.h
--- a/Engine/Source/Runtime/Experimental/Chaos/Public/Chaos/KinematicGeometryParticles.h
+++ b/Engine/Source/Runtime/Experimental/Chaos/Public/Chaos/KinematicGeometryParticles.h
@@ -31,11 +31,23 @@ class TKinematicGeometryParticlesImp : public TGeometryParticlesImp<T, d, SimTyp
 	}
 	virtual ~TKinematicGeometryParticlesImp() {};
 
-	const TVector<T, d>& V(const int32 Index) const { return MV[Index]; }
-	TVector<T, d>& V(const int32 Index) { return MV[Index]; }
-
-	const TVector<T, d>& W(const int32 Index) const { return MW[Index]; }
-	TVector<T, d>& W(const int32 Index) { return MW[Index]; }
+	UE_DEPRECATED(5.4, "Use GetV instead")
+	const TVector<T, d> V(const int32 Index) const { return TVector<T, d>(MV[Index]); }
+	UE_DEPRECATED(5.4, "Use GetV or SetV instead")
+	TVector<T, d> V(const int32 Index) { return TVector<T, d>(MV[Index]); }
+	const TVector<T, d> GetV(const int32 Index) const { return TVector<T, d>(MV[Index]); }
+	void SetV(const int32 Index, const TVector<T, d>& InV) { MV[Index] = InV; }
+	const TVector<FRealSingle, d> GetVf(const int32 Index) const { return MV[Index]; }
+	void SetVf(const int32 Index, const TVector<FRealSingle, d>& InV) { MV[Index] = InV; }
+
+	UE_DEPRECATED(5.4, "Use GetW instead")
+	const TVector<T, d> W(const int32 Index) const { return TVector<T, d>(MW[Index]); }
+	UE_DEPRECATED(5.4, "Use GetW or SetW instead")
+	TVector<T, d> W(const int32 Index) { return TVector<T, d>(MW[Index]); }
+	const TVector<T, d> GetW(const int32 Index) const { return TVector<T, d>(MW[Index]); }
+	void SetW(const int32 Index, const TVector<T, d>& InW) { MW[Index] = InW; }
+	const TVector<FRealSingle, d> GetWf(const int32 Index) const { return MW[Index]; }
+	void SetWf(const int32 Index, const TVector<FRealSingle, d>& InW) { MW[Index] = InW; }
 
 	const TKinematicTarget<T, d>& KinematicTarget(const int32 Index) const { return KinematicTargets[Index]; }
 	TKinematicTarget<T, d>& KinematicTarget(const int32 Index) { return KinematicTargets[Index]; }
@@ -43,7 +55,7 @@ class TKinematicGeometryParticlesImp : public TGeometryParticlesImp<T, d, SimTyp
 	FString ToString(int32 index) const
 	{
 		FString BaseString = TGeometryParticlesImp<T, d, SimType>::ToString(index);
-		return FString::Printf(TEXT("%s, MV:%s, MW:%s"), *BaseString, *V(index).ToString(), *W(index).ToString());
+		return FString::Printf(TEXT("%s, MV:%s, MW:%s"), *BaseString, *GetV(index).ToString(), *GetW(index).ToString());
 	}
 
 	typedef TKinematicGeometryParticleHandle<T, d> THandleType;
@@ -55,7 +67,40 @@ class TKinematicGeometryParticlesImp : public TGeometryParticlesImp<T, d, SimTyp
 	virtual void Serialize(FChaosArchive& Ar) override
 	{
 		TGeometryParticlesImp<T, d, SimType>::Serialize(Ar);
-		Ar << MV << MW;
+		
+		Ar.UsingCustomVersion(FUE5ReleaseStreamObjectVersion::GUID);
+		if (Ar.CustomVer(FUE5ReleaseStreamObjectVersion::GUID) >= FUE5ReleaseStreamObjectVersion::SinglePrecisonParticleDataPT)
+		{
+			Ar << MV << MW;
+		}
+		else
+		{
+			TArrayCollectionArray<TVector<FReal, d>> VDouble;
+			VDouble.Resize(MV.Num());
+			for (int32 Index = 0; Index < MV.Num(); ++Index)
+			{
+				VDouble[Index] = MV[Index];
+			}
+			TArrayCollectionArray<TVector<FReal, d>> WDouble;
+			WDouble.Resize(MW.Num());
+			for (int32 Index = 0; Index < MW.Num(); ++Index)
+			{
+				WDouble[Index] = MW[Index];
+			}
+
+			Ar << VDouble << WDouble;
+
+			MV.Resize(VDouble.Num());
+			for (int32 Index = 0; Index < VDouble.Num(); ++Index)
+			{
+				MV[Index] = VDouble[Index];
+			}
+			MW.Resize(WDouble.Num());
+			for (int32 Index = 0; Index < WDouble.Num(); ++Index)
+			{
+				MW[Index] = WDouble[Index];
+			}
+		}
 
 		Ar.UsingCustomVersion(FExternalPhysicsCustomObjectVersion::GUID);
 		if(Ar.CustomVer(FExternalPhysicsCustomObjectVersion::GUID) >= FExternalPhysicsCustomObjectVersion::KinematicTargets)
@@ -64,13 +109,13 @@ class TKinematicGeometryParticlesImp : public TGeometryParticlesImp<T, d, SimTyp
 		}
 	}
 
-	FORCEINLINE TArray<TVector<T, d>>& AllV() { return MV; }
-	FORCEINLINE TArray<TVector<T, d>>& AllW() { return MW; }
+	FORCEINLINE TArray<TVector<FRealSingle, d>>& AllV() { return MV; }
+	FORCEINLINE TArray<TVector<FRealSingle, d>>& AllW() { return MW; }
 	FORCEINLINE TArray<TKinematicTarget<T, d>>& AllKinematicTargets() { return KinematicTargets; }
 
   private:
-	TArrayCollectionArray<TVector<T, d>> MV;
-	TArrayCollectionArray<TVector<T, d>> MW;
+	TArrayCollectionArray<TVector<FRealSingle, d>> MV;
+	TArrayCollectionArray<TVector<FRealSingle, d>> MW;
 	TArrayCollectionArray<TKinematicTarget<T, d>> KinematicTargets;
 };
 
diff --git a/Engine/Source/Runtime/Experimental/Chaos/Public/Chaos/PBDRigidParticles.h b/Engine/Source/Runtime/Experimental/Chaos/Public/Chaos/PBDRigidParticles.h
--- a/Engine/Source/Runtime/Experimental/Chaos/Public/Chaos/PBDRigidParticles.h
+++ b/Engine/Source/Runtime/Experimental/Chaos/Public/Chaos/PBDRigidParticles.h
@@ -62,27 +62,45 @@ class TPBDRigidParticles : public TRigidParticles<T, d>
 	FORCEINLINE const TVector<T, d>& P(const int32 index) const { return MP[index]; }
 	FORCEINLINE TVector<T, d>& P(const int32 index) { return MP[index]; }
 
-	FORCEINLINE const TRotation<T, d>& Q(const int32 index) const { return MQ[index]; }
-	FORCEINLINE TRotation<T, d>& Q(const int32 index) { return MQ[index]; }
-
-	const TVector<T, d>& PreV(const int32 index) const { return MPreV[index]; }
-	TVector<T, d>& PreV(const int32 index) { return MPreV[index]; }
-
-	const TVector<T, d>& PreW(const int32 index) const { return MPreW[index]; }
-	TVector<T, d>& PreW(const int32 index) { return MPreW[index]; }
+	UE_DEPRECATED(5.4, "Use GetQ instead")
+	FORCEINLINE const TRotation<T, d> Q(const int32 index) const { return TRotation<T, d>(MQ[index]); }
+	UE_DEPRECATED(5.4, "Use GetQ or SetQ instead")
+	FORCEINLINE TRotation<T, d> Q(const int32 index) { return TRotation<T, d>(MQ[index]); }
+	FORCEINLINE const TRotation<T, d> GetQ(const int32 index) const { return TRotation<T, d>(MQ[index]); }
+	FORCEINLINE void SetQ(const int32 index, const TRotation<T, d>& InQ) { MQ[index] = TRotation<FRealSingle, d>(InQ); }
+	FORCEINLINE const TRotation<FRealSingle, d> GetQf(const int32 index) const { return MQ[index]; }
+	FORCEINLINE void SetQf(const int32 index, const TRotation<FRealSingle, d>& InQ) { MQ[index] = InQ; }
+
+	UE_DEPRECATED(5.4, "Use GetPreV instead")
+	const TVector<T, d> PreV(const int32 index) const { return TVector<T, d>(MPreV[index]); }
+	UE_DEPRECATED(5.4, "Use GetPreV or SetPreV instead")
+	TVector<T, d> PreV(const int32 index) { return TVector<T, d>(MPreV[index]); }
+	const TVector<T, d> GetPreV(const int32 index) const { return TVector<T, d>(MPreV[index]); }
+	void SetPreV(const int32 index, const TVector<T, d>& InPreV) { MPreV[index] = TVector<FRealSingle, d>(InPreV); }
+	const TVector<FRealSingle, d> GetPreVf(const int32 index) const { return MPreV[index]; }
+	void SetPreVf(const int32 index, const TVector<FRealSingle, d>& InPreV) { MPreV[index] = InPreV; }
+
+	UE_DEPRECATED(5.4, "Use GetPreW instead")
+	const TVector<T, d> PreW(const int32 index) const { return MPreW[index]; }
+	UE_DEPRECATED(5.4, "Use GetPreW or SetPreW instead")
+	TVector<T, d> PreW(const int32 index) { return MPreW[index]; }
+	const TVector<T, d> GetPreW(const int32 index) const { return MPreW[index]; }
+	void SetPreW(const int32 index, const TVector<T, d>& InPreW) { MPreW[index] = TVector<FRealSingle, d>(InPreW); }
+	const TVector<FRealSingle, d> GetPreWf(const int32 index) const { return MPreW[index]; }
+	void SetPreWf(const int32 index, const TVector<FRealSingle, d>& InPreW) { MPreW[index] = InPreW; }
 
 	// World-space center of mass location
-	const TVector<T, d> XCom(const int32 index) const { return this->X(index) + this->R(index).RotateVector(CenterOfMass(index)); }
-	const TVector<T, d> PCom(const int32 index) const { return this->P(index) + this->Q(index).RotateVector(CenterOfMass(index)); }
+	const TVector<T, d> XCom(const int32 index) const { return this->X(index) + this->GetR(index).RotateVector(CenterOfMass(index)); }
+	const TVector<T, d> PCom(const int32 index) const { return this->P(index) + this->GetQ(index).RotateVector(CenterOfMass(index)); }
 
 	// World-space center of mass rotation
-	const TRotation<T, d> RCom(const int32 index) const { return this->R(index) * RotationOfMass(index); }
-	const TRotation<T, d> QCom(const int32 index) const { return this->Q(index) * RotationOfMass(index); }
+	const TRotation<T, d> RCom(const int32 index) const { return this->GetR(index) * RotationOfMass(index); }
+	const TRotation<T, d> QCom(const int32 index) const { return this->GetQ(index) * RotationOfMass(index); }
 
 	void SetTransformPQCom(const int32 index, const TVector<T, d>& InPCom, const TRotation<T, d>& InQCom)
 	{
-		Q(index) = InQCom * RotationOfMass(index).Inverse();
-		P(index) = InPCom - Q(index) * CenterOfMass(index);
+		SetQ(index,InQCom * RotationOfMass(index).Inverse());
+		P(index) = InPCom - GetQ(index) * CenterOfMass(index);
 	}
 
 	// The index into an FSolverBodyContainer (for dynamic particles only), or INDEX_NONE.
@@ -101,14 +119,14 @@ class TPBDRigidParticles : public TRigidParticles<T, d>
 	{
 		if (Sleeping(Index) && bSleeping == false)
 		{
-			PreV(Index) = this->V(Index);
-			PreW(Index) = this->W(Index);
+			SetPreV(Index, this->GetV(Index));
+			SetPreW(Index, this->GetW(Index));
 		}
-		else if(bSleeping)
+		else if (bSleeping)
 		{
 			//being put to sleep, so zero out velocities
-			this->V(Index) = FVec3(0);
-			this->W(Index) = FVec3(0);
+			this->SetV(Index, FVec3(0));
+			this->SetW(Index, FVec3(0));
 		}
 
 		bool CurrentlySleeping = this->ObjectState(Index) == EObjectStateType::Sleeping;
@@ -164,7 +182,7 @@ class TPBDRigidParticles : public TRigidParticles<T, d>
 				1.f / this->I(Index)[2]);
 
 			this->P(Index) = this->X(Index);
-			this->Q(Index) = this->R(Index);
+			this->SetQ(Index, this->GetR(Index));
 		}
 		else if (InObjectState == EObjectStateType::Sleeping)
 		{
@@ -196,32 +214,66 @@ class TPBDRigidParticles : public TRigidParticles<T, d>
 	FString ToString(int32 index) const
 	{
 		FString BaseString = TRigidParticles<T, d>::ToString(index);
-		return FString::Printf(TEXT("%s, MP:%s, MQ:%s, MPreV:%s, MPreW:%s"), *BaseString, *P(index).ToString(), *Q(index).ToString(), *PreV(index).ToString(), *PreW(index).ToString());
+		return FString::Printf(TEXT("%s, MP:%s, MQ:%s, MPreV:%s, MPreW:%s"), *BaseString, *P(index).ToString(), *GetQ(index).ToString(), *GetPreV(index).ToString(), *GetPreW(index).ToString());
 	}
 
 	virtual void Serialize(FChaosArchive& Ar) override
 	{
 		TRigidParticles<T, d>::Serialize(Ar);
-		Ar << MP << MQ << MPreV << MPreW;
-	}
-
-	UE_DEPRECATED(5.2, "To be removed")
-	FORCEINLINE TArray<TVector<T, d>>& AllP() { return MP; }
+		
+		Ar << MP;
 
-	UE_DEPRECATED(5.2, "To be removed")
-	FORCEINLINE TArray<TRotation<T, d>>& AllQ() { return MQ; }
-	
-	UE_DEPRECATED(5.2, "To be removed")
-	FORCEINLINE TArray<TVector<T, d>>& AllPreV() { return MPreV; }
-	
-	UE_DEPRECATED(5.2, "To be removed")
-	FORCEINLINE TArray<TVector<T, d>>& AllPreW() { return MPreW; }
+		Ar.UsingCustomVersion(FUE5ReleaseStreamObjectVersion::GUID);
+		if (Ar.CustomVer(FUE5ReleaseStreamObjectVersion::GUID) >= FUE5ReleaseStreamObjectVersion::SinglePrecisonParticleDataPT)
+		{
+			Ar << MQ << MPreV << MPreW; 
+		}
+		else
+		{
+			TArrayCollectionArray<TRotation<FReal, d>> QDouble;
+			QDouble.Resize(MQ.Num());
+			for (int32 Index = 0; Index < MQ.Num(); ++Index)
+			{
+				QDouble[Index] = TRotation<FReal, d >(MQ[Index]);
+			}
+			TArrayCollectionArray<TVector<FReal, d>> PreVDouble;
+			PreVDouble.Resize(MPreV.Num());
+			for (int32 Index = 0; Index < MPreV.Num(); ++Index)
+			{
+				PreVDouble[Index] = TVector<FReal, d >(MPreV[Index]);
+			}
+			TArrayCollectionArray<TVector<FReal, d>> PreWDouble;
+			PreWDouble.Resize(MPreW.Num());
+			for (int32 Index = 0; Index < MPreW.Num(); ++Index)
+			{
+				PreWDouble[Index] = TVector<FReal, d >(MPreW[Index]);
+			}
+
+			Ar << QDouble << PreVDouble << PreWDouble;
+
+			MQ.Resize(QDouble.Num());
+			for (int32 Index = 0; Index < QDouble.Num(); ++Index)
+			{
+				MQ[Index] = TRotation<FRealSingle, d >(QDouble[Index]);
+			}
+			MPreV.Resize(PreVDouble.Num());
+			for (int32 Index = 0; Index < PreVDouble.Num(); ++Index)
+			{
+				MPreV[Index] = TVector<FRealSingle, d >(PreVDouble[Index]);
+			}
+			MPreW.Resize(PreWDouble.Num());
+			for (int32 Index = 0; Index < PreWDouble.Num(); ++Index)
+			{
+				MPreW[Index] = TVector<FRealSingle, d >(PreWDouble[Index]);
+			}
+		}
+	}
 
   private:
 	TArrayCollectionArray<TVector<T, d>> MP;
-	TArrayCollectionArray<TRotation<T, d>> MQ;
-	TArrayCollectionArray<TVector<T, d>> MPreV;
-	TArrayCollectionArray<TVector<T, d>> MPreW;
+	TArrayCollectionArray<TRotation<FRealSingle, d>> MQ;
+	TArrayCollectionArray<TVector<FRealSingle, d>> MPreV;
+	TArrayCollectionArray<TVector<FRealSingle, d>> MPreW;
 	TArrayCollectionArray<int32> MSolverBodyIndex;	// Transient for use in constraint solver
 };
 
diff --git a/Engine/Source/Runtime/Experimental/Chaos/Public/Chaos/PBDRigidsEvolution.h b/Engine/Source/Runtime/Experimental/Chaos/Public/Chaos/PBDRigidsEvolution.h
--- a/Engine/Source/Runtime/Experimental/Chaos/Public/Chaos/PBDRigidsEvolution.h
+++ b/Engine/Source/Runtime/Experimental/Chaos/Public/Chaos/PBDRigidsEvolution.h
@@ -752,7 +752,7 @@ public:
 
 			TKinematicTarget<FReal, 3>& KinematicTarget = Particle.KinematicTarget();
 			const FVec3 CurrentX = Particle.X();
-			const FRotation3 CurrentR = Particle.R();
+			const FRotation3 CurrentR = Particle.GetR();
 			constexpr FReal MinDt = 1e-6f;
 
 			bool bMoved = false;
@@ -765,8 +765,8 @@ public:
 			case EKinematicTargetMode::Reset:
 			{
 				// Reset velocity and then switch to do-nothing mode
-				Particle.V() = FVec3(0, 0, 0);
-				Particle.W() = FVec3(0, 0, 0);
+				Particle.SetVf(FVec3f(0, 0, 0));
+				Particle.SetWf(FVec3f(0, 0, 0));
 				Particle.ClearIsMovingKinematic();
 				KinematicTarget.SetMode(EKinematicTargetMode::None);
 				break;
@@ -809,9 +809,9 @@ public:
 					}
 				}
 				Particle.X() = NewX;
-				Particle.R() = NewR;
-				Particle.V() = NewV;
-				Particle.W() = NewW;
+				Particle.SetR(NewR);
+				Particle.SetV(NewV);
+				Particle.SetW(NewW);
 				Particle.SetIsMovingKinematic();
 
 				break;
@@ -821,8 +821,8 @@ public:
 			{
 				// Move based on velocity
 				bMoved = true;
-				Particle.X() = Particle.X() + Particle.V() * Dt;
-				Particle.R() = FRotation3::IntegrateRotationWithAngularVelocity(Particle.R(), Particle.W(), Dt);
+				Particle.X() = Particle.X() + Particle.GetV() * Dt;
+				Particle.SetR(FRotation3::IntegrateRotationWithAngularVelocity(Particle.GetR(), Particle.GetW(), Dt));
 				Particle.SetIsMovingKinematic();
 
 				break;
@@ -830,21 +830,21 @@ public:
 			}
 			
 			// Set positions and previous velocities if we can
-			// Note: At present kininematics are in fact rigid bodies
+			// Note: At present kinematics are in fact rigid bodies
 			Particle.P() = Particle.X();
-			Particle.Q() = Particle.R();
-			Particle.PreV() = Particle.V();
-			Particle.PreW() = Particle.W();
+			Particle.SetQf(Particle.GetRf());
+			Particle.SetPreVf(Particle.GetVf());
+			Particle.SetPreWf(Particle.GetWf());
 
 			if (bMoved)
 			{
 				if (!Particle.CCDEnabled())
 				{
-					Particle.UpdateWorldSpaceState(FRigidTransform3(Particle.P(), Particle.Q()), FVec3(0));
+					Particle.UpdateWorldSpaceState(FRigidTransform3(Particle.P(), Particle.GetQ()), FVec3(0));
 				}
 				else
 				{
-					Particle.UpdateWorldSpaceStateSwept(FRigidTransform3(Particle.P(), Particle.Q()), FVec3(0), -Particle.V() * Dt);
+					Particle.UpdateWorldSpaceStateSwept(FRigidTransform3(Particle.P(), Particle.GetQ()), FVec3(0), -Particle.GetV() * Dt);
 				}
 			}
 
diff --git a/Engine/Source/Runtime/Experimental/Chaos/Public/Chaos/PBDRigidsEvolutionGBF.h b/Engine/Source/Runtime/Experimental/Chaos/Public/Chaos/PBDRigidsEvolutionGBF.h
--- a/Engine/Source/Runtime/Experimental/Chaos/Public/Chaos/PBDRigidsEvolutionGBF.h
+++ b/Engine/Source/Runtime/Experimental/Chaos/Public/Chaos/PBDRigidsEvolutionGBF.h
@@ -234,8 +234,8 @@ namespace Chaos
 					auto& Particle = *PBDParticle;
 
 					//save off previous velocities
-					Particle.PreV() = Particle.V();
-					Particle.PreW() = Particle.W();
+					Particle.SetPreVf(Particle.GetVf());
+					Particle.SetPreWf(Particle.GetWf());
 
 					for (FForceRule ForceRule : ForceRules)
 					{
@@ -243,19 +243,19 @@ namespace Chaos
 					}
 
 					//EulerStepVelocityRule.Apply(Particle, Dt);
-					Particle.V() += Particle.Acceleration() * Dt;
-					Particle.W() += Particle.AngularAcceleration() * Dt;
+					Particle.SetV(Particle.GetV() + Particle.Acceleration() * Dt);
+					Particle.SetW(Particle.GetW() + Particle.AngularAcceleration() * Dt);
 
 					//AddImpulsesRule.Apply(Particle, Dt);
-					Particle.V() += Particle.LinearImpulseVelocity();
-					Particle.W() += Particle.AngularImpulseVelocity();
+					Particle.SetV(Particle.GetV() + Particle.LinearImpulseVelocity());
+					Particle.SetW(Particle.GetW() + Particle.AngularImpulseVelocity());
 					Particle.LinearImpulseVelocity() = FVec3(0);
 					Particle.AngularImpulseVelocity() = FVec3(0);
 
 					//EtherDragRule.Apply(Particle, Dt);
 					{
-						FVec3& V = Particle.V();
-						FVec3& W = Particle.W();
+						FVec3 V = Particle.GetV();
+						FVec3 W = Particle.GetW();
 
 						const FReal LinearDrag = LinearEtherDragOverride >= 0 ? LinearEtherDragOverride : Particle.LinearEtherDrag() * Dt;
 						const FReal LinearMultiplier = FMath::Max(FReal(0), FReal(1) - LinearDrag);
@@ -277,31 +277,33 @@ namespace Chaos
 						{
 							W *= FMath::Sqrt(Particle.MaxAngularSpeedSq() / AngularSpeedSq);
 						}
+						Particle.SetV(V);
+						Particle.SetW(W);
 					}
 
 					if (CVars::HackMaxAngularVelocity >= 0.f)
 					{
-						const FReal AngularSpeedSq = Particle.W().SizeSquared();
+						const FReal AngularSpeedSq = Particle.GetW().SizeSquared();
 						if (AngularSpeedSq > MaxAngularSpeedSq)
 						{
-							Particle.W() = Particle.W() * (CVars::HackMaxAngularVelocity / FMath::Sqrt(AngularSpeedSq));
+							Particle.SetW(Particle.GetW() * (CVars::HackMaxAngularVelocity / FMath::Sqrt(AngularSpeedSq)));
 						}
 					}
 
 					if (CVars::HackMaxVelocity >= 0.f)
 					{
-						const FReal SpeedSq = Particle.V().SizeSquared();
+						const FReal SpeedSq = Particle.GetV().SizeSquared();
 						if (SpeedSq > MaxSpeedSq)
 						{
-							Particle.V() = Particle.V() * (CVars::HackMaxVelocity / FMath::Sqrt(SpeedSq));
+							Particle.SetV(Particle.GetV() * (CVars::HackMaxVelocity / FMath::Sqrt(SpeedSq)));
 						}
 					}
 
 					FVec3 PCoM = Particle.XCom();
 					FRotation3 QCoM = Particle.RCom();
 
-					PCoM = PCoM + Particle.V() * Dt;
-					QCoM = FRotation3::IntegrateRotationWithAngularVelocity(QCoM, Particle.W(), Dt);
+					PCoM = PCoM + Particle.GetV() * Dt;
+					QCoM = FRotation3::IntegrateRotationWithAngularVelocity(QCoM, Particle.GetW(), Dt);
 
 					Particle.SetTransformPQCom(PCoM, QCoM);
 
@@ -311,7 +313,7 @@ namespace Chaos
 					FVec3 VelocityBoundsDelta = FVec3(0);
 					if ((VelocityBoundsMultiplier > 0) && (MaxVelocityBoundsExpansion > 0))
 					{
-						VelocityBoundsDelta = (-VelocityBoundsMultiplier * Dt) * Particle.V();
+						VelocityBoundsDelta = (-VelocityBoundsMultiplier * Dt) * Particle.GetV();
 						
 						// Box clamp to avoid sqrt
 						VelocityBoundsDelta.BoundToCube(MaxVelocityBoundsExpansion);
@@ -321,7 +323,7 @@ namespace Chaos
 					{
 						// Expand bounds about P/Q by a small amount. This can still result in missed collisions, especially
 						// when we have joints that pull the body back to X/R, if P-X is greater than the BoundsThickness
-						Particle.UpdateWorldSpaceStateSwept(FRigidTransform3(Particle.P(), Particle.Q()), FVec3(BoundsThickness), VelocityBoundsDelta);
+						Particle.UpdateWorldSpaceStateSwept(FRigidTransform3(Particle.P(), Particle.GetQ()), FVec3(BoundsThickness), VelocityBoundsDelta);
 					}
 					else
 					{
@@ -329,26 +331,26 @@ namespace Chaos
 #if CHAOS_DEBUG_DRAW
 						if (CVars::ChaosSolverDrawCCDThresholds)
 						{
-							DebugDraw::DrawCCDAxisThreshold(Particle.X(), Particle.CCDAxisThreshold(), Particle.P() - Particle.X(), Particle.Q());
+							DebugDraw::DrawCCDAxisThreshold(Particle.X(), Particle.CCDAxisThreshold(), Particle.P() - Particle.X(), Particle.GetQ());
 						}
 #endif
 
-						if (FCCDHelpers::DeltaExceedsThreshold(Particle.CCDAxisThreshold(), Particle.P() - Particle.X(), Particle.Q()))
+						if (FCCDHelpers::DeltaExceedsThreshold(Particle.CCDAxisThreshold(), Particle.P() - Particle.X(), Particle.GetQ()))
 						{
 							// We sweep the bounds from P back along the velocity and expand by a small amount.
 							// If not using tight bounds we also expand the bounds in all directions by Velocity. This is necessary only for secondary CCD collisions
 							// @todo(chaos): expanding the bounds by velocity is very expensive - revisit this
-							const FVec3 VDt = Particle.V() * Dt;
+							const FVec3 VDt = Particle.GetV() * Dt;
 							FReal CCDBoundsExpansion = BoundsThickness;
 							if (!CVars::bChaosCollisionCCDUseTightBoundingBox && (CVars::ChaosCollisionCCDConstraintMaxProcessCount > 1))
 							{
 								CCDBoundsExpansion += VDt.GetAbsMax();
 							}
-							Particle.UpdateWorldSpaceStateSwept(FRigidTransform3(Particle.P(), Particle.Q()), FVec3(CCDBoundsExpansion), -VDt);
+							Particle.UpdateWorldSpaceStateSwept(FRigidTransform3(Particle.P(), Particle.GetQ()), FVec3(CCDBoundsExpansion), -VDt);
 						}
 						else
 						{
-							Particle.UpdateWorldSpaceStateSwept(FRigidTransform3(Particle.P(), Particle.Q()), FVec3(BoundsThickness), VelocityBoundsDelta);
+							Particle.UpdateWorldSpaceStateSwept(FRigidTransform3(Particle.P(), Particle.GetQ()), FVec3(BoundsThickness), VelocityBoundsDelta);
 						}
 					}
 
@@ -382,12 +384,12 @@ namespace Chaos
 			FReal FakeDT = (FReal)1. / (FReal)30.;
 			if (Particle.LinearImpulseVelocity().IsNearlyZero() == false || Particle.Acceleration().IsNearlyZero() == false)
 			{
-				const FVec3 PredictedLinearVelocity = Particle.V() + Particle.Acceleration() * FakeDT + Particle.LinearImpulseVelocity();
+				const FVec3 PredictedLinearVelocity = Particle.GetV() + Particle.Acceleration() * FakeDT + Particle.LinearImpulseVelocity();
 				Particle.VSmooth() =FMath::Lerp(Particle.VSmooth(), PredictedLinearVelocity, SmoothRate);
 			}
 			if (Particle.AngularImpulseVelocity().IsNearlyZero() == false || Particle.AngularAcceleration().IsNearlyZero() == false)
 			{
-				const FVec3 PredictedAngularVelocity = Particle.W() + Particle.AngularAcceleration() * FakeDT + Particle.AngularImpulseVelocity();
+				const FVec3 PredictedAngularVelocity = Particle.GetW() + Particle.AngularAcceleration() * FakeDT + Particle.AngularImpulseVelocity();
 				Particle.WSmooth() = FMath::Lerp(Particle.WSmooth(), PredictedAngularVelocity, SmoothRate);
 			}
 		}
diff --git a/Engine/Source/Runtime/Experimental/Chaos/Public/Chaos/PBDRigidsSOAs.h b/Engine/Source/Runtime/Experimental/Chaos/Public/Chaos/PBDRigidsSOAs.h
--- a/Engine/Source/Runtime/Experimental/Chaos/Public/Chaos/PBDRigidsSOAs.h
+++ b/Engine/Source/Runtime/Experimental/Chaos/Public/Chaos/PBDRigidsSOAs.h
@@ -496,13 +496,13 @@ public:
 	void DisableParticle(FGeometryParticleHandle* Particle)
 	{
 		// Rigid particles express their disabled state with a boolean.
-		// Disabled kinematic and static particles get shuffled to differnt SOAs.
+		// Disabled kinematic and static particles get shuffled to different SOAs.
 
 		if (auto PBDRigid = Particle->CastToRigidParticle())
 		{
 			PBDRigid->SetDisabled(true);
-			PBDRigid->V() = FVec3(0);
-			PBDRigid->W() = FVec3(0);
+			PBDRigid->SetVf(FVec3f(0));
+			PBDRigid->SetWf(FVec3f(0));
 
 			if (TPBDGeometryCollectionParticleHandle<FReal, 3>* PBDRigidGC = Particle->CastToGeometryCollection())
 			{
diff --git a/Engine/Source/Runtime/Experimental/Chaos/Public/Chaos/Particle/ParticleUtilities.h b/Engine/Source/Runtime/Experimental/Chaos/Public/Chaos/Particle/ParticleUtilities.h
--- a/Engine/Source/Runtime/Experimental/Chaos/Public/Chaos/Particle/ParticleUtilities.h
+++ b/Engine/Source/Runtime/Experimental/Chaos/Public/Chaos/Particle/ParticleUtilities.h
@@ -25,12 +25,12 @@ namespace Chaos
 		static inline const FVec3& GetPosition(T_PARTICLEHANDLE Particle) { return Particle->X(); }
 		template<typename T_PARTICLEHANDLE>
 		UE_DEPRECATED(5.4, "Use Particle->SetR() instead.")
-		static inline FRotation3 Rotation(T_PARTICLEHANDLE Particle) { return Particle->R(); }
+		static inline FRotation3 Rotation(T_PARTICLEHANDLE Particle) { return Particle->GetR(); }
 		template<typename T_PARTICLEHANDLE>
-		static inline const FRotation3 GetRotation(T_PARTICLEHANDLE Particle) { return Particle->R(); }
+		static inline const FRotation3 GetRotation(T_PARTICLEHANDLE Particle) { return Particle->GetR(); }
 
 		static inline FVec3& Position(TPBDRigidParticles<FReal, 3>& Particles, const int32 Index) { return Particles.X(Index); }
-		static inline FRotation3& Rotation(TPBDRigidParticles<FReal, 3>& Particles, const int32 Index) { return Particles.R(Index); }
+		static inline const FRotation3 Rotation(TPBDRigidParticles<FReal, 3>& Particles, const int32 Index) { return Particles.GetR(Index); }
 	};
 
 	class FParticleSpatialAccessorPQ
@@ -41,12 +41,12 @@ namespace Chaos
 		template<typename T_PARTICLEHANDLE>
 		static inline const FVec3& GetPosition(T_PARTICLEHANDLE Particle) { return Particle->P(); }
 		template<typename T_PARTICLEHANDLE>
-		static inline FRotation3& Rotation(T_PARTICLEHANDLE Particle) { return Particle->Q(); }
+		static inline FRotation3& Rotation(T_PARTICLEHANDLE Particle) { return Particle->GetQ(); }
 		template<typename T_PARTICLEHANDLE>
-		static inline const FRotation3& GetRotation(T_PARTICLEHANDLE Particle) { return Particle->Q(); }
+		static inline const FRotation3 GetRotation(T_PARTICLEHANDLE Particle) { return Particle->GetQ(); }
 
 		static inline FVec3& Position(TPBDRigidParticles<FReal, 3>& Particles, const int32 Index) { return Particles.P(Index); }
-		static inline FRotation3& Rotation(TPBDRigidParticles<FReal, 3>& Particles, const int32 Index) { return Particles.Q(Index); }
+		static inline FRotation3 Rotation(TPBDRigidParticles<FReal, 3>& Particles, const int32 Index) { return Particles.GetQ(Index); }
 	};
 
 	/**
diff --git a/Engine/Source/Runtime/Experimental/Chaos/Public/Chaos/ParticleDirtyFlags.h b/Engine/Source/Runtime/Experimental/Chaos/Public/Chaos/ParticleDirtyFlags.h
--- a/Engine/Source/Runtime/Experimental/Chaos/Public/Chaos/ParticleDirtyFlags.h
+++ b/Engine/Source/Runtime/Experimental/Chaos/Public/Chaos/ParticleDirtyFlags.h
@@ -81,7 +81,7 @@ public:
 	void CopyFrom(const TOther& Other)
 	{
 		MX = Other.X();
-		MR = FRotation3f(Other.R());
+		MR = FRotation3f(Other.GetR());
 	}
 
 	template <typename TOther>
@@ -99,6 +99,7 @@ public:
 	void SetX(const FVec3& InX){ MX = InX; }
 
 	const FRotation3 R() const { return FRotation3(MR); }
+	const FRotation3 GetR() const { return FRotation3(MR); }
 	void SetR(const FRotation3& InR) { MR = FRotation3f(InR); }
 	
 private:
@@ -138,14 +139,14 @@ public:
 	template <typename TOther>
 	void CopyFrom(const TOther& Other)
 	{
-		MV = Other.V();
-		MW = Other.W();
+		MV = Other.GetV();
+		MW = Other.GetW();
 	}
 
 	template <typename TOther>
 	bool IsEqual(const TOther& Other) const
 	{
-		return MV == FVec3f(Other.V()) && MW == FVec3f(Other.W());
+		return MV == FVec3f(Other.GetV()) && MW == FVec3f(Other.GetW());
 	}
 
 	bool operator==(const FParticleVelocities& Other) const
@@ -154,9 +155,11 @@ public:
 	}
 
 	const FVec3 V() const { return FVec3(MV); }
+	const FVec3 GetV() const { return FVec3(MV); }
 	void SetV(const FVec3& V) { MV = FVec3f(V); }
 
 	const FVec3 W() const { return FVec3(MW); }
+	const FVec3 GetW() const { return FVec3(MW); }
 	void SetW(const FVec3& W){ MW = FVec3f(W); }
 
 private:
diff --git a/Engine/Source/Runtime/Experimental/Chaos/Public/Chaos/ParticleHandle.h b/Engine/Source/Runtime/Experimental/Chaos/Public/Chaos/ParticleHandle.h
--- a/Engine/Source/Runtime/Experimental/Chaos/Public/Chaos/ParticleHandle.h
+++ b/Engine/Source/Runtime/Experimental/Chaos/Public/Chaos/ParticleHandle.h
@@ -524,11 +524,16 @@ public:
 	FUniqueIdx UniqueIdx() const { return GeometryParticles->UniqueIdx(ParticleIdx); }
 	void SetUniqueIdx(const FUniqueIdx UniqueIdx, bool bInvalidate = false) const { GeometryParticles->UniqueIdx(ParticleIdx) = UniqueIdx; }
 
-	const TRotation<T, d>& R() const { return GeometryParticles->R(ParticleIdx); }
-	TRotation<T, d>& R() { return GeometryParticles->R(ParticleIdx); }
-	void SetR(const TRotation<T, d>& InR, bool bInvalidate = false) { GeometryParticles->R(ParticleIdx) = InR; }
+	UE_DEPRECATED(5.4, "Use GetR instead")
+	const TRotation<T, d> R() const { return GeometryParticles->GetR(ParticleIdx); }
+	UE_DEPRECATED(5.4, "Use SetR or GetR instead")
+	TRotation<T, d> R() { return GeometryParticles->GetR(ParticleIdx); }
+	const TRotation<T, d> GetR() const { return GeometryParticles->GetR(ParticleIdx); }
+	void SetR(const TRotation<T, d>& InR, bool bInvalidate = false) { GeometryParticles->SetR(ParticleIdx, InR); }
+	const TRotation<FRealSingle, d> GetRf() const { return GeometryParticles->GetRf(ParticleIdx); }
+	void SetRf(const TRotation<FRealSingle, d>& InR, bool bInvalidate = false) { GeometryParticles->SetRf(ParticleIdx, InR); }
 
-	FRigidTransform3 GetTransformXR() const { return FRigidTransform3(X(), R()); }
+	FRigidTransform3 GetTransformXR() const { return FRigidTransform3(X(), GetR()); }
 
 	// Initialize the transform
 	void InitTransform(const FVec3& InP, const FRotation3& InQ)
@@ -896,13 +901,23 @@ public:
 		return Serializable;
 	}
 
-	const TVector<T, d>& V() const { return KinematicGeometryParticles->V(ParticleIdx); }
-	TVector<T, d>& V() { return KinematicGeometryParticles->V(ParticleIdx); }
-	void SetV(const TVector<T, d>& InV, bool bInvalidate = false) { KinematicGeometryParticles->V(ParticleIdx) = InV; }
-
-	const TVector<T, d>& W() const { return KinematicGeometryParticles->W(ParticleIdx); }
-	TVector<T, d>& W() { return KinematicGeometryParticles->W(ParticleIdx); }
-	void SetW(const TVector<T, d>& InW, bool bInvalidate = false) { KinematicGeometryParticles->W(ParticleIdx) = InW; }
+	UE_DEPRECATED(5.4, "Use GetV instead")
+	const TVector<T, d> V() const { return KinematicGeometryParticles->GetV(ParticleIdx); }
+	UE_DEPRECATED(5.4, "Use SetV or GetV instead")
+	TVector<T, d> V() { return KinematicGeometryParticles->GetV(ParticleIdx); }
+	const TVector<T, d> GetV() const { return KinematicGeometryParticles->GetV(ParticleIdx); }
+	void SetV(const TVector<T, d>& InV, bool bInvalidate = false) { KinematicGeometryParticles->SetV(ParticleIdx, InV); }
+	const TVector<FRealSingle, d> GetVf() const { return KinematicGeometryParticles->GetVf(ParticleIdx); }
+	void SetVf(const TVector<FRealSingle, d>& InV, bool bInvalidate = false) { KinematicGeometryParticles->SetVf(ParticleIdx, InV); }
+
+	UE_DEPRECATED(5.4, "Use GetW instead")
+	const TVector<T, d> W() const { return KinematicGeometryParticles->GetW(ParticleIdx); }
+	UE_DEPRECATED(5.4, "Use SetW or GetW instead")
+	TVector<T, d> W() { return KinematicGeometryParticles->GetW(ParticleIdx); }
+	const TVector<T, d> GetW() const { return KinematicGeometryParticles->GetW(ParticleIdx); }
+	void SetW(const TVector<T, d>& InW, bool bInvalidate = false) { KinematicGeometryParticles->SetW(ParticleIdx, InW); }
+	const TVector<FRealSingle, d> GetWf() const { return KinematicGeometryParticles->GetWf(ParticleIdx); }
+	void SetWf(const TVector<FRealSingle, d>& InW, bool bInvalidate = false) { KinematicGeometryParticles->SetWf(ParticleIdx, InW); }
 
 	UE_DEPRECATED(5.4, "Use FPBDRigidsEvolutionGBF::SetParticleVelocitied or (SetV, SetW if that is not appropriate)")
 	void SetVelocities(const FParticleVelocities& Velocities)
@@ -938,8 +953,8 @@ public:
 	using TGeometryParticleHandleImp<T, d, bPersistent>::SetX;
 	using TGeometryParticleHandleImp<T, d, bPersistent>::SetR;
 	using TGeometryParticleHandleImp<T, d, bPersistent>::Type;
-	using TKinematicGeometryParticleHandleImp<T, d, bPersistent>::V;
-	using TKinematicGeometryParticleHandleImp<T, d, bPersistent>::W;
+	using TKinematicGeometryParticleHandleImp<T, d, bPersistent>::GetV;
+	using TKinematicGeometryParticleHandleImp<T, d, bPersistent>::GetW;
 
 	using TTransientHandle = TTransientPBDRigidParticleHandle<T, d>;
 	using TSOAType = TPBDRigidParticles<T, d>;
@@ -959,13 +974,13 @@ protected:
 		PBDRigidParticleDefaultConstruct<T, d>(*this, Params);
 		SetCollisionConstraintFlags(0);
 		SetDisabled(Params.bDisabled);
-		SetPreV(this->V());
-		SetPreW(this->W());
+		SetPreVf(this->GetVf());
+		SetPreWf(this->GetWf());
 		SetSolverBodyIndex(INDEX_NONE);
 		SetP(this->X());
-		SetQ(this->R());
-		SetVSmooth(this->V());
-		SetWSmooth(this->W());
+		SetQf(this->GetRf());
+		SetVSmooth(this->GetV());
+		SetWSmooth(this->GetW());
 		SetAcceleration(TVector<T, d>(0));
 		SetAngularAcceleration(TVector<T, d>(0));
 		SetObjectStateLowLevel(Params.bStartSleeping ? EObjectStateType::Sleeping : EObjectStateType::Dynamic);
@@ -1019,13 +1034,23 @@ public:
 	void SetDisabledLowLevel(bool disabled) { PBDRigidParticles->SetDisabledLowLevel(ParticleIdx, disabled); }
 	void SetDisabled(const bool InDisabled) { PBDRigidParticles->DisabledRef(ParticleIdx) = InDisabled; }
 
-	const TVector<T, d>& PreV() const { return PBDRigidParticles->PreV(ParticleIdx); }
-	TVector<T, d>& PreV() { return PBDRigidParticles->PreV(ParticleIdx); }
-	void SetPreV(const TVector<T, d>& InPreV) { PBDRigidParticles->PreV(ParticleIdx) = InPreV; }
-
-	const TVector<T, d>& PreW() const { return PBDRigidParticles->PreW(ParticleIdx); }
-	TVector<T, d>& PreW() { return PBDRigidParticles->PreW(ParticleIdx); }
-	void SetPreW(const TVector<T, d>& InPreW) { PBDRigidParticles->PreW(ParticleIdx) = InPreW; }
+	UE_DEPRECATED(5.4, "Use GetPreV instead")
+	const TVector<T, d> PreV() const { return PBDRigidParticles->GetPreV(ParticleIdx); }
+	UE_DEPRECATED(5.4, "Use GetPreV or SetPreV instead")
+	TVector<T, d> PreV() { return PBDRigidParticles->GetPreV(ParticleIdx); }
+	const TVector<T, d> GetPreV() const { return PBDRigidParticles->GetPreV(ParticleIdx); }
+	void SetPreV(const TVector<T, d>& InPreV) { PBDRigidParticles->SetPreV(ParticleIdx, InPreV); }
+	const TVector<FRealSingle, d> GetPreVf() const { return PBDRigidParticles->GetPreVf(ParticleIdx); }
+	void SetPreVf(const TVector<FRealSingle, d>& InPreV) { PBDRigidParticles->SetPreVf(ParticleIdx, InPreV); }
+
+	UE_DEPRECATED(5.4, "Use GetPreW instead")
+	const TVector<T, d> PreW() const { return PBDRigidParticles->GetPreW(ParticleIdx); }
+	UE_DEPRECATED(5.4, "Use GetPreW or SetPreW instead")
+	TVector<T, d> PreW() { return PBDRigidParticles->GetPreW(ParticleIdx); }
+	const TVector<T, d> GetPreW() const { return PBDRigidParticles->GetPreW(ParticleIdx); }
+	void SetPreW(const TVector<T, d>& InPreW) { PBDRigidParticles->SetPreW(ParticleIdx, InPreW); }
+	const TVector<FRealSingle, d> GetPreWf() const { return PBDRigidParticles->GetPreWf(ParticleIdx); }
+	void SetPreWf(const TVector<FRealSingle, d>& InPreW) { PBDRigidParticles->SetPreWf(ParticleIdx, InPreW); }
 
 	int32 SolverBodyIndex() const { return PBDRigidParticles->SolverBodyIndex(ParticleIdx); }
 	void SetSolverBodyIndex(const int32 InSolverBodyIndex) { PBDRigidParticles->SetSolverBodyIndex(ParticleIdx, InSolverBodyIndex); }
@@ -1034,9 +1059,14 @@ public:
 	TVector<T, d>& P() { return PBDRigidParticles->P(ParticleIdx); }
 	void SetP(const TVector<T, d>& InP) { PBDRigidParticles->P(ParticleIdx) = InP; }
 
-	const TRotation<T, d>& Q() const { return PBDRigidParticles->Q(ParticleIdx); }
-	TRotation<T, d>& Q() { return PBDRigidParticles->Q(ParticleIdx); }
-	void SetQ(const TRotation<T, d>& InQ) { PBDRigidParticles->Q(ParticleIdx) = InQ; }
+	UE_DEPRECATED(5.4, "Use GetQ instead")
+	const TRotation<T, d> Q() const { return PBDRigidParticles->GetQ(ParticleIdx); }
+	UE_DEPRECATED(5.4, "Use GetQ or SetQ instead")
+	TRotation<T, d> Q() { return PBDRigidParticles->GetQ(ParticleIdx); }
+	const TRotation<T, d> GetQ() const { return PBDRigidParticles->GetQ(ParticleIdx); }
+	void SetQ(const TRotation<T, d>& InQ) { PBDRigidParticles->SetQ(ParticleIdx, InQ); }
+	const TRotation<FRealSingle, d> GetQf() const { return PBDRigidParticles->GetQf(ParticleIdx); }
+	void SetQf(const TRotation<FRealSingle, d>& InQ) { PBDRigidParticles->SetQf(ParticleIdx, InQ); }
 
 	// World-space center of mass position
 	const TVector<T, d> XCom() const { return PBDRigidParticles->XCom(ParticleIdx); }
@@ -1058,7 +1088,7 @@ public:
 	// Set world-space center of mass transform
 	void SetTransformPQCom(const TVector<T, d>& InPCom, const TRotation<T, d>& InQCom) { PBDRigidParticles->SetTransformPQCom(ParticleIdx, InPCom, InQCom); }
 
-	FRigidTransform3 GetTransformPQ() const { return FRigidTransform3(P(), Q()); }
+	FRigidTransform3 GetTransformPQ() const { return FRigidTransform3(P(), GetQ()); }
 
 	FRigidTransform3 GetTransformXRCom() const { return FRigidTransform3(XCom(), RCom()); }
 
@@ -1121,8 +1151,8 @@ public:
 
 	void ResetSmoothedVelocities()
 	{
-		SetVSmooth(V());
-		SetWSmooth(W());
+		SetVSmooth(GetV());
+		SetWSmooth(GetW());
 	}
 
 	// Get the raw inertia. @see ConditionedInvI()
@@ -1359,7 +1389,7 @@ void TGeometryParticleHandleImp<T,d,bPersistent>::SetXR(const FParticlePositionR
 	if(auto Rigid = CastToRigidParticle())
 	{
 		Rigid->SetP(X());
-		Rigid->SetQ(R());
+		Rigid->SetQf(GetRf());
 	}
 }
 
@@ -1667,19 +1697,23 @@ public:
 	void SetTransform(const FVec3& Pos, const FRotation3& Rot)
 	{
 		MHandle->X() = Pos;
-		MHandle->R() = Rot;
+		MHandle->SetR(Rot);
 		if (FPBDRigidParticleHandle* Dynamic = CastToRigidParticle())
 		{
 			Dynamic->P() = Pos;
-			Dynamic->Q() = Rot;
+			Dynamic->SetQ(Rot);
 		}
 	}
 
 	// Static Particles
 	FVec3& X() { return MHandle->X(); }
 	const FVec3& X() const { return MHandle->X(); }
-	FRotation3& R() { return MHandle->R(); }
-	const FRotation3& R() const { return MHandle->R(); }
+	UE_DEPRECATED(5.4, "Use GetR or SetR instead")
+	FRotation3 R() { return MHandle->GetR(); }
+	UE_DEPRECATED(5.4, "Use GetR instead")
+	const FRotation3 R() const { return MHandle->GetR(); }
+	void SetR(const FRotation3& InR) { MHandle->SetR(InR); }
+	const FRotation3 GetR() const { return MHandle->GetR(); }
 	const FImplicitObjectRef GetGeometry() const { return MHandle->GetGeometry(); }
 	bool Sleeping() const { return MHandle->Sleeping(); }
 	FString ToString() const { return MHandle->ToString(); }
@@ -1698,11 +1732,11 @@ public:
 	auto& AuxilaryValue(Container& AuxContainer) { return MHandle->AuxilaryValue(AuxContainer); }
 
 	// Kinematic Particles
-	const FVec3& V() const { return (MHandle->CastToKinematicParticle()) ? MHandle->CastToKinematicParticle()->V() : ZeroVector; }
-	const FVec3& W() const { return (MHandle->CastToKinematicParticle()) ? MHandle->CastToKinematicParticle()->W() : ZeroVector; }
+	const FVec3 V() const { return (MHandle->CastToKinematicParticle()) ? MHandle->CastToKinematicParticle()->GetV() : ZeroVector; }
+	const FVec3 W() const { return (MHandle->CastToKinematicParticle()) ? MHandle->CastToKinematicParticle()->GetW() : ZeroVector; }
 
-	void SetV(const FVec3& InV) { if (MHandle->CastToKinematicParticle()) { MHandle->CastToKinematicParticle()->V() = InV; } }
-	void SetW(const FVec3& InW) { if (MHandle->CastToKinematicParticle()) { MHandle->CastToKinematicParticle()->W() = InW; } }
+	void SetV(const FVec3& InV) { if (MHandle->CastToKinematicParticle()) { MHandle->CastToKinematicParticle()->SetV(InV); } }
+	void SetW(const FVec3& InW) { if (MHandle->CastToKinematicParticle()) { MHandle->CastToKinematicParticle()->SetW(InW); } }
 
 	const FKinematicTarget& KinematicTarget() const { return (MHandle->CastToKinematicParticle())? MHandle->CastToKinematicParticle()->KinematicTarget() : EmptyKinematicTarget; }
 
@@ -1771,21 +1805,21 @@ public:
 		return false;
 	}
 
-	const FVec3& PreV() const
+	const FVec3 PreV() const
 	{
 		if (MHandle->CastToRigidParticle())
 		{
-			return MHandle->CastToRigidParticle()->PreV();
+			return MHandle->CastToRigidParticle()->GetPreV();
 		}
 
 		return ZeroVector;
 	}
 
-	const FVec3& PreW() const
+	const FVec3 PreW() const
 	{
 		if (MHandle->CastToRigidParticle())
 		{
-			return MHandle->CastToRigidParticle()->PreW();
+			return MHandle->CastToRigidParticle()->GetPreW();
 		}
 		return ZeroVector;
 	}
@@ -1827,24 +1861,29 @@ public:
 		return X();
 	}
 
-	FRotation3& Q()
+	void SetQ(const FRotation3& InQ)
 	{
 		if (IsDynamic())
 		{
-			return MHandle->CastToRigidParticle()->Q();
+			return MHandle->CastToRigidParticle()->SetQ(InQ);
 		}
 
-		return R();
+		SetR(InQ);
 	}
 
-	const FRotation3& Q() const
+	const FRotation3 Q() const
 	{
 		if (IsDynamic())
 		{
-			return MHandle->CastToRigidParticle()->Q();
+			return MHandle->CastToRigidParticle()->GetQ();
 		}
 
-		return R();
+		return GetR();
+	}
+
+	const FRotation3 GetQ() const
+	{
+		return Q();
 	}
 
 	// World-space center of mass position
@@ -1872,7 +1911,7 @@ public:
 		{
 			return MHandle->CastToRigidParticle()->RCom();
 		}
-		return R();
+		return GetR();
 	}
 	const FRotation3 QCom() const
 	{ 
@@ -1880,7 +1919,7 @@ public:
 		{
 			return MHandle->CastToRigidParticle()->QCom();
 		}
-		return R();
+		return GetR();
 	}
 
 	void InitTransform(const FVec3& InP, const FRotation3& InQ)
@@ -1974,7 +2013,7 @@ public:
 		return T;
 	}
 
-	const FVec3& VSmooth() const
+	const FVec3 VSmooth() const
 	{
 		if (MHandle->CastToRigidParticle())
 		{
@@ -1984,7 +2023,7 @@ public:
 		return V();
 	}
 
-	const FVec3& WSmooth() const
+	const FVec3 WSmooth() const
 	{
 		if (MHandle->CastToRigidParticle())
 		{
@@ -2591,6 +2630,7 @@ public:
 	}
 
 	const TRotation<T, d> R() const { return MXR.Read().R(); }
+	const TRotation<T, d> GetR() const { return MXR.Read().R(); }
 	void SetR(const TRotation<T, d>& InR, bool bInvalidate = true);
 
 	const FParticlePositionRotation& XR() const { return MXR.Read(); }
@@ -3099,9 +3139,11 @@ public:
 	}
 
 	const TVector<T, d> V() const { return MVelocities.Read().V(); }
+	const TVector<T, d> GetV() const { return MVelocities.Read().V(); }
 	void SetV(const TVector<T, d>& InV, bool bInvalidate = true);
 
 	const TVector<T, d> W() const { return MVelocities.Read().W(); }
+	const TVector<T, d> GetW() const { return MVelocities.Read().W(); }
 	void SetW(const TVector<T, d>& InW, bool bInvalidate = true);
 
 	const FKinematicTarget KinematicTarget() const {
diff --git a/Engine/Source/Runtime/Experimental/Chaos/Public/Chaos/PerCellBoundaryConditions.h b/Engine/Source/Runtime/Experimental/Chaos/Public/Chaos/PerCellBoundaryConditions.h
--- a/Engine/Source/Runtime/Experimental/Chaos/Public/Chaos/PerCellBoundaryConditions.h
+++ b/Engine/Source/Runtime/Experimental/Chaos/Public/Chaos/PerCellBoundaryConditions.h
@@ -22,22 +22,22 @@ class TPerCellBoundaryConditions
 		for (uint32 i = 0; i < MParticles.Size(); ++i)
 		{
 			const TVector<T, d>& X = Grid.Location(Index);
-			TRigidTransform<T, d> Frame(MParticles.X(i), MParticles.R(i));
+			TRigidTransform<T, d> Frame(MParticles.X(i), MParticles.GetR(i));
 			if (MParticles.GetGeometry(i)->SignedDistance(Frame.InverseTransformPosition(X)) < 0)
 			{
 				BoundaryConditions(Index) = true;
-				Velocity(Index) = MParticles.V(i)[Index.First];
+				Velocity(Index) = MParticles.GetV(i)[Index.First];
 				break;
 			}
 		}
 		for (uint32 i = 0; i < MSources.Size(); ++i)
 		{
 			const TVector<T, d>& X = Grid.Location(Index);
-			TRigidTransform<T, d> Frame(MSources.X(i), MSources.R(i));
+			TRigidTransform<T, d> Frame(MSources.X(i), MSources.GetR(i));
 			if (MSources.GetGeometry(i)->SignedDistance(Frame.InverseTransformPosition(X)) < 0)
 			{
 				BoundaryConditions(Index) = true;
-				Velocity(Index) = MSources.V(i)[Index.First];
+				Velocity(Index) = MSources.GetV(i)[Index.First];
 				break;
 			}
 		}
@@ -50,7 +50,7 @@ class TPerCellBoundaryConditions
 		for (uint32 i = 0; i < MSources.Size(); ++i)
 		{
 			const TVector<T, d>& X = Grid.Location(Index);
-			TRigidTransform<T, d> Frame(MSources.X(i), MSources.R(i));
+			TRigidTransform<T, d> Frame(MSources.X(i), MSources.GetR(i));
 			if (MSources.GetGeometry(i)->SignedDistance(Frame.InverseTransformPosition(X)) < 0)
 			{
 				Density(Index) = 1.f;
diff --git a/Engine/Source/Runtime/Experimental/Chaos/Public/Chaos/PerParticleAddImpulses.h b/Engine/Source/Runtime/Experimental/Chaos/Public/Chaos/PerParticleAddImpulses.h
--- a/Engine/Source/Runtime/Experimental/Chaos/Public/Chaos/PerParticleAddImpulses.h
+++ b/Engine/Source/Runtime/Experimental/Chaos/Public/Chaos/PerParticleAddImpulses.h
@@ -16,7 +16,7 @@ namespace Chaos
 		template<class T_PARTICLES>
 		inline void ApplyHelper(T_PARTICLES& InParticles, const FReal Dt, const int32 Index) const
 		{
-			InParticles.V(Index) += InParticles.LinearImpulseVelocity(Index);
+			InParticles.SetV(Index, InParticles.GetV(Index) + InParticles.LinearImpulseVelocity(Index));
 		}
 
 		inline void Apply(FDynamicParticles& InParticles, const FReal Dt, const int32 Index) const override //-V762
@@ -39,15 +39,15 @@ namespace Chaos
 			//		 and W are in an inertial frame.
 			//
 
-			InParticles.W(Index) += InParticles.AngularImpulseVelocity(Index);
+			InParticles.SetW(Index, InParticles.GetW(Index) + InParticles.AngularImpulseVelocity(Index));
 			InParticles.LinearImpulseVelocity(Index) = FVec3(0);
 			InParticles.AngularImpulseVelocity(Index) = FVec3(0);
 		}
 
 		inline void Apply(TTransientPBDRigidParticleHandle<FReal, 3>& Particle, const FReal Dt) const override //-V762
 		{
-			Particle.V() += Particle.LinearImpulseVelocity();
-			Particle.W() += Particle.AngularImpulseVelocity();
+			Particle.SetV(Particle.GetV() + Particle.LinearImpulseVelocity());
+			Particle.SetW(Particle.GetW() + Particle.AngularImpulseVelocity());
 			Particle.LinearImpulseVelocity() = FVec3(0);
 			Particle.AngularImpulseVelocity() = FVec3(0);
 		}
diff --git a/Engine/Source/Runtime/Experimental/Chaos/Public/Chaos/PerParticleEtherDrag.h b/Engine/Source/Runtime/Experimental/Chaos/Public/Chaos/PerParticleEtherDrag.h
--- a/Engine/Source/Runtime/Experimental/Chaos/Public/Chaos/PerParticleEtherDrag.h
+++ b/Engine/Source/Runtime/Experimental/Chaos/Public/Chaos/PerParticleEtherDrag.h
@@ -47,12 +47,20 @@ namespace Chaos
 
 		inline void Apply(TRigidParticles<FReal, 3>& InParticles, const FReal Dt, const int32 Index) const override //-V762
 		{
-			ApplyHelper(InParticles.V(Index), InParticles.W(Index), InParticles.LinearEtherDrag(Index), InParticles.AngularEtherDrag(Index), InParticles.MaxLinearSpeedSq(Index), InParticles.MaxAngularSpeedSq(Index), Dt);
+			FVec3 V = InParticles.GetV(Index);
+			FVec3 W = InParticles.GetW(Index);
+			ApplyHelper(V, W, InParticles.LinearEtherDrag(Index), InParticles.AngularEtherDrag(Index), InParticles.MaxLinearSpeedSq(Index), InParticles.MaxAngularSpeedSq(Index), Dt);
+			InParticles.SetV(Index, V);
+			InParticles.SetW(Index, W);
 		}
 
 		inline void Apply(TTransientPBDRigidParticleHandle<FReal, 3>& Particle, const FReal Dt) const override //-V762
 		{
-			ApplyHelper(Particle.V(), Particle.W(), Particle.LinearEtherDrag(), Particle.AngularEtherDrag(), Particle.MaxLinearSpeedSq(), Particle.MaxAngularSpeedSq(), Dt);
+			FVec3 V = Particle.GetV();
+			FVec3 W = Particle.GetW();
+			ApplyHelper(V, W, Particle.LinearEtherDrag(), Particle.AngularEtherDrag(), Particle.MaxLinearSpeedSq(), Particle.MaxAngularSpeedSq(), Dt);
+			Particle.SetV(V);
+			Particle.SetW(W);
 		}
 	};
 
diff --git a/Engine/Source/Runtime/Experimental/Chaos/Public/Chaos/PerParticleEulerStepVelocity.h b/Engine/Source/Runtime/Experimental/Chaos/Public/Chaos/PerParticleEulerStepVelocity.h
--- a/Engine/Source/Runtime/Experimental/Chaos/Public/Chaos/PerParticleEulerStepVelocity.h
+++ b/Engine/Source/Runtime/Experimental/Chaos/Public/Chaos/PerParticleEulerStepVelocity.h
@@ -16,7 +16,7 @@ class FPerParticleEulerStepVelocity : public FPerParticleRule
 	template<class T_PARTICLES>
 	inline void ApplyHelper(T_PARTICLES& InParticles, const FReal Dt, const int32 Index) const
 	{
-		InParticles.V(Index) += InParticles.Acceleration(Index) * Dt;
+		InParticles.SetV(Index, InParticles.GetV(Index) + InParticles.Acceleration(Index) * Dt);
 	}
 
 	inline void Apply(FDynamicParticles& InParticles, const FReal Dt, const int32 Index) const override //-V762
@@ -39,13 +39,13 @@ class FPerParticleEulerStepVelocity : public FPerParticleRule
 		//       Just using W += InvI * (Torque - W x (I * W)) * dt is not correct, since Torque
 		//		 and W are in an inertial frame.
 		//
-		InParticles.W(Index) += InParticles.AngularAcceleration(Index) * Dt;
+		InParticles.SetW(Index, InParticles.GetW(Index) + InParticles.AngularAcceleration(Index) * Dt);
 	}
 	
 	inline void Apply(TTransientPBDRigidParticleHandle<FReal, 3>& Particle, const FReal Dt) const override //-V762
 	{
-		Particle.V() += Particle.Acceleration() * Dt;
-		Particle.W() += Particle.AngularAcceleration() * Dt;
+		Particle.SetV(Particle.GetV() + Particle.Acceleration() * Dt);
+		Particle.SetW(Particle.GetW() + Particle.AngularAcceleration() * Dt);
 	}
 };
 
diff --git a/Engine/Source/Runtime/Experimental/Chaos/Public/Chaos/PerParticlePBDCCDCollisionConstraint.h b/Engine/Source/Runtime/Experimental/Chaos/Public/Chaos/PerParticlePBDCCDCollisionConstraint.h
--- a/Engine/Source/Runtime/Experimental/Chaos/Public/Chaos/PerParticlePBDCCDCollisionConstraint.h
+++ b/Engine/Source/Runtime/Experimental/Chaos/Public/Chaos/PerParticlePBDCCDCollisionConstraint.h
@@ -87,7 +87,7 @@ private:
 					return; // Bail out if the collision groups doesn't match the particle group id, or use INDEX_NONE (= global collision that affects all particle)
 				}
 
-				const FSolverRigidTransform3 Frame(CollisionParticles.X(CollisionIndex), CollisionParticles.R(CollisionIndex));
+				const FSolverRigidTransform3 Frame(CollisionParticles.X(CollisionIndex), CollisionParticles.GetR(CollisionIndex));
 
 				const Pair<FVec3, bool> PointPair = CollisionParticles.GetGeometry(CollisionIndex)->FindClosestIntersection(  // Geometry operates in FReal
 					FVec3(CollisionTransforms[CollisionIndex].InverseTransformPositionNoScale(Particles.X(Index))),        // hence the back and forth
diff --git a/Engine/Source/Runtime/Experimental/Chaos/Public/Chaos/PerParticlePBDCollisionConstraint.h b/Engine/Source/Runtime/Experimental/Chaos/Public/Chaos/PerParticlePBDCollisionConstraint.h
--- a/Engine/Source/Runtime/Experimental/Chaos/Public/Chaos/PerParticlePBDCollisionConstraint.h
+++ b/Engine/Source/Runtime/Experimental/Chaos/Public/Chaos/PerParticlePBDCollisionConstraint.h
@@ -129,7 +129,7 @@ private:
 					{
 						return; // Bail out if the collision groups doesn't match the particle group id, or use INDEX_NONE (= global collision that affects all particle)
 					}
-					const FSolverRigidTransform3 Frame(CollisionParticles.X(i), CollisionParticles.R(i));
+					const FSolverRigidTransform3 Frame(CollisionParticles.X(i), CollisionParticles.GetR(i));
 					const FVec3 RigidSpacePosition(Frame.InverseTransformPosition(Particles.P(Index)));  // PhiWithNormal requires FReal based arguments
 					FVec3 ImplicitNormal;                                                                // since implicits don't use FSolverReal
 					FSolverReal Phi;
@@ -220,7 +220,7 @@ private:
 					{
 						return; // Bail out if the collision groups doesn't match the particle group id, or use INDEX_NONE (= global collision that affects all particle)
 					}
-					const FSolverRigidTransform3 Frame(CollisionParticles.X(i), CollisionParticles.R(i));
+					const FSolverRigidTransform3 Frame(CollisionParticles.X(i), CollisionParticles.GetR(i));
 					const FVec3 RigidSpacePosition(Frame.InverseTransformPosition(Particles.P(Index)));  // PhiWithNormal requires FReal based arguments
 					FVec3 ImplicitNormal;                                                                // since implicits don't use FSolverReal
 					const FSolverReal Phi = (FSolverReal)CollisionParticles.GetGeometry(i)->PhiWithNormal(RigidSpacePosition, ImplicitNormal);
diff --git a/Engine/Source/Runtime/Experimental/Chaos/Public/Chaos/PerParticlePBDEulerStep.h b/Engine/Source/Runtime/Experimental/Chaos/Public/Chaos/PerParticlePBDEulerStep.h
--- a/Engine/Source/Runtime/Experimental/Chaos/Public/Chaos/PerParticlePBDEulerStep.h
+++ b/Engine/Source/Runtime/Experimental/Chaos/Public/Chaos/PerParticlePBDEulerStep.h
@@ -32,8 +32,8 @@ class FPerParticlePBDEulerStep : public FPerParticleRule
 		FVec3 PCoM = InParticles.XCom(Index);
 		FRotation3 QCoM = InParticles.RCom(Index);
 
-		PCoM = PCoM + InParticles.V(Index) * Dt;
-		QCoM = FRotation3::IntegrateRotationWithAngularVelocity(QCoM, InParticles.W(Index), Dt);
+		PCoM = PCoM + InParticles.GetV(Index) * Dt;
+		QCoM = FRotation3::IntegrateRotationWithAngularVelocity(QCoM, InParticles.GetW(Index), Dt);
 
 		InParticles.SetTransformPQCom(Index, PCoM, QCoM);
 	}
@@ -43,8 +43,8 @@ class FPerParticlePBDEulerStep : public FPerParticleRule
 		FVec3 PCoM = Particle.XCom();
 		FRotation3 QCoM = Particle.RCom();
 
-		PCoM = PCoM + Particle.V() * Dt;
-		QCoM = FRotation3::IntegrateRotationWithAngularVelocity(QCoM, Particle.W(), Dt);
+		PCoM = PCoM + Particle.GetV() * Dt;
+		QCoM = FRotation3::IntegrateRotationWithAngularVelocity(QCoM, Particle.GetW(), Dt);
 
 		Particle.SetTransformPQCom(PCoM, QCoM);
 	}
diff --git a/Engine/Source/Runtime/Experimental/Chaos/Public/Chaos/PerParticlePBDUpdateFromDeltaPosition.h b/Engine/Source/Runtime/Experimental/Chaos/Public/Chaos/PerParticlePBDUpdateFromDeltaPosition.h
--- a/Engine/Source/Runtime/Experimental/Chaos/Public/Chaos/PerParticlePBDUpdateFromDeltaPosition.h
+++ b/Engine/Source/Runtime/Experimental/Chaos/Public/Chaos/PerParticlePBDUpdateFromDeltaPosition.h
@@ -17,7 +17,7 @@ class FPerParticlePBDUpdateFromDeltaPosition : public FPerParticleRule
 	template<class T_PARTICLES>
 	inline void ApplyHelper(T_PARTICLES& InParticles, const FReal Dt, const int32 Index) const
 	{
-		InParticles.V(Index) = (InParticles.P(Index) - InParticles.X(Index)) / Dt;
+		InParticles.SetV(Index, (InParticles.P(Index) - InParticles.X(Index)) / Dt);
 		//InParticles.X(Index) = InParticles.P(Index);
 	}
 
@@ -30,7 +30,7 @@ class FPerParticlePBDUpdateFromDeltaPosition : public FPerParticleRule
 	inline void Apply(TPBDRigidParticles<FReal, 3>& InParticles, const FReal Dt, const int32 Index) const override //-V762
 	{
 		ApplyHelper(InParticles, Dt, Index);
-		InParticles.W(Index) = FRotation3::CalculateAngularVelocity(InParticles.R(Index), InParticles.Q(Index), Dt);
+		InParticles.SetW(Index, FRotation3::CalculateAngularVelocity(InParticles.GetR(Index), InParticles.GetQ(Index), Dt));
 	}
 
 	inline void Apply(FPBDRigidParticleHandle* Particle, const FReal Dt) const override //-V762
@@ -38,8 +38,8 @@ class FPerParticlePBDUpdateFromDeltaPosition : public FPerParticleRule
 		const FVec3& CenterOfMass = Particle->CenterOfMass();
 		const FVec3 CenteredX = Particle->XCom();
 		const FVec3 CenteredP = Particle->PCom();
-		Particle->V() = FVec3::CalculateVelocity(CenteredX, CenteredP, Dt);
-		Particle->W() = FRotation3::CalculateAngularVelocity(Particle->R(), Particle->Q(), Dt);
+		Particle->SetV(FVec3::CalculateVelocity(CenteredX, CenteredP, Dt));
+		Particle->SetW(FRotation3::CalculateAngularVelocity(Particle->GetR(), Particle->GetQ(), Dt));
 	}
 
 	inline void Apply(TTransientPBDRigidParticleHandle<FReal, 3>& Particle, const FReal Dt) const override //-V762
@@ -47,8 +47,8 @@ class FPerParticlePBDUpdateFromDeltaPosition : public FPerParticleRule
 		const FVec3& CenterOfMass = Particle.CenterOfMass();
 		const FVec3 CenteredX = Particle.XCom();
 		const FVec3 CenteredP = Particle.PCom();
-		Particle.V() = FVec3::CalculateVelocity(CenteredX, CenteredP, Dt);
-		Particle.W() = FRotation3::CalculateAngularVelocity(Particle.R(), Particle.Q(), Dt);
+		Particle.SetV(FVec3::CalculateVelocity(CenteredX, CenteredP, Dt));
+		Particle.SetW(FRotation3::CalculateAngularVelocity(Particle.GetR(), Particle.GetQ(), Dt));
 	}
 };
 
diff --git a/Engine/Source/Runtime/Experimental/Chaos/Public/Chaos/SimpleGeometryParticles.h b/Engine/Source/Runtime/Experimental/Chaos/Public/Chaos/SimpleGeometryParticles.h
--- a/Engine/Source/Runtime/Experimental/Chaos/Public/Chaos/SimpleGeometryParticles.h
+++ b/Engine/Source/Runtime/Experimental/Chaos/Public/Chaos/SimpleGeometryParticles.h
@@ -6,6 +6,7 @@
 #include "Chaos/Particles.h"
 #include "Chaos/Rotation.h"
 #include "UObject/FortniteValkyrieBranchObjectVersion.h"
+#include "UObject/UE5ReleaseStreamObjectVersion.h"
 
 
 namespace Chaos
@@ -46,10 +47,16 @@ namespace Chaos
 		virtual ~TSimpleGeometryParticles() override
 		{}
 
-		FORCEINLINE const TRotation<T, d>& R(const int32 Index) const { return MR[Index]; }
-		FORCEINLINE TRotation<T, d>& R(const int32 Index) { return MR[Index]; }
-		const TArrayCollectionArray<TRotation<T, d>>& GetR() const { return MR; }
-		TArrayCollectionArray<TRotation<T, d>>& GetR() { return MR; }
+		UE_DEPRECATED(5.4, "Use GetR instead")
+		FORCEINLINE const TRotation<T, d> R(const int32 Index) const { return TRotation<T, d>(MR[Index]); }
+		UE_DEPRECATED(5.4, "Use GetR or SetR instead.")
+		FORCEINLINE TRotation<T, d> R(const int32 Index) { return MR[Index]; }
+		FORCEINLINE const TRotation<T, d> GetR(const int32 Index) const { return TRotation<T, d>(MR[Index]); }
+		FORCEINLINE void SetR(const int32 Index, const TRotation<T, d>& InR) { MR[Index] = TRotation<FRealSingle, d>(InR); }
+		FORCEINLINE const TRotation<FRealSingle, d> GetRf(const int32 Index) const { return MR[Index]; }
+		FORCEINLINE void SetRf(const int32 Index, const TRotation<FRealSingle, d>& InR) { MR[Index] = InR; }
+		const TArrayCollectionArray<TRotation<FRealSingle, d>>& GetR() const { return MR; }
+		TArrayCollectionArray<TRotation<FRealSingle, d>>& GetR() { return MR; }
 
 		FORCEINLINE const FImplicitObjectPtr& GetGeometry(const int32 Index) const { return MGeometry[Index]; }
 		void SetGeometry(const int32 Index, const FImplicitObjectPtr& InGeometry)
@@ -58,7 +65,7 @@ namespace Chaos
 		}
 
 		FORCEINLINE const TArray<FImplicitObjectPtr>& GetAllGeometry() const { return MGeometry; }
-		FORCEINLINE TArray<TRotation<T, d>>& AllR() { return MR; }
+		FORCEINLINE TArray<TRotation<FRealSingle, d>>& AllR() { return MR; }
 
 		virtual void Serialize(FChaosArchive& Ar)
 		{
@@ -85,7 +92,29 @@ namespace Chaos
 			{
 				Ar << MGeometry;
 			}
-			Ar << MR;
+
+			Ar.UsingCustomVersion(FUE5ReleaseStreamObjectVersion::GUID);
+			if (Ar.CustomVer(FUE5ReleaseStreamObjectVersion::GUID) >= FUE5ReleaseStreamObjectVersion::SinglePrecisonParticleDataPT)
+			{
+				Ar << MR;
+			}
+			else
+			{
+				TArrayCollectionArray<TRotation<FReal, d>> RDouble;
+				RDouble.Resize(MR.Num());
+				for (int32 Index = 0; Index < MR.Num(); ++Index)
+				{
+					RDouble[Index] = TRotation<FReal, d >(MR[Index]);
+				}
+
+				Ar << RDouble;
+
+				MR.Resize(RDouble.Num());
+				for (int32 Index = 0; Index < RDouble.Num(); ++Index)
+				{
+					MR[Index] = TRotation<FRealSingle, d >(RDouble[Index]);
+				}
+			}
 		}
 
 	protected:
@@ -95,7 +124,7 @@ namespace Chaos
 		}
 
 	private:
-		TArrayCollectionArray<TRotation<T, d>> MR;
+		TArrayCollectionArray<TRotation<FRealSingle, d>> MR;
 		// MGeometry contains raw ptrs to every entry in both MSharedGeometry and MDynamicGeometry.
 		// It may also contain raw ptrs to geometry which is managed outside of Chaos.
 		TArrayCollectionArray<FImplicitObjectPtr> MGeometry;
diff --git a/Engine/Source/Runtime/Experimental/Chaos/Public/Chaos/SoftsSolverCollisionParticlesRange.h b/Engine/Source/Runtime/Experimental/Chaos/Public/Chaos/SoftsSolverCollisionParticlesRange.h
--- a/Engine/Source/Runtime/Experimental/Chaos/Public/Chaos/SoftsSolverCollisionParticlesRange.h
+++ b/Engine/Source/Runtime/Experimental/Chaos/Public/Chaos/SoftsSolverCollisionParticlesRange.h
@@ -32,8 +32,9 @@ public:
 	TArrayView<FSolverVec3> GetW() { return GetArrayView(GetParticles().GetW()); }
 
 	// SimpleGeometryParticles data
-	const FSolverRotation3& R(const int32 Index) const { return GetParticles().R(Index + Offset); }
-	FSolverRotation3& R(const int32 Index) { return GetParticles().R(Index + Offset); }
+	const FSolverRotation3 R(const int32 Index) const { return GetParticles().GetR(Index + Offset); }
+	void SetR(const int32 Index, const FSolverRotation3& InR) { GetParticles().SetR(Index + Offset, InR); }
+	// FSolverRotation3 R(const int32 Index) { return GetParticles().SetR(Index + Offset); }
 	TConstArrayView<FSolverRotation3> GetR() const { return GetConstArrayView(GetParticles().GetR()); }
 	TArrayView<FSolverRotation3> GetR() { return GetArrayView(GetParticles().GetR()); }
 
diff --git a/Engine/Source/Runtime/Experimental/Chaos/Public/PhysicsProxy/FieldSystemProxyHelper.h b/Engine/Source/Runtime/Experimental/Chaos/Public/PhysicsProxy/FieldSystemProxyHelper.h
--- a/Engine/Source/Runtime/Experimental/Chaos/Public/PhysicsProxy/FieldSystemProxyHelper.h
+++ b/Engine/Source/Runtime/Experimental/Chaos/Public/PhysicsProxy/FieldSystemProxyHelper.h
@@ -95,7 +95,7 @@ namespace Chaos
 						{
 							if (ParentHandle->Disabled() == false)
 							{
-								const FRigidTransform3 ParentWorldTM(ParentHandle->P(), ParentHandle->Q());
+								const FRigidTransform3 ParentWorldTM(ParentHandle->P(), ParentHandle->GetQ());
 								const FRigidTransform3 ChildFrame = ClusterHandle->ChildToParent() * ParentWorldTM;
 								SamplePosition = ChildFrame.GetTranslation();
 							}
@@ -237,8 +237,8 @@ namespace Chaos
 
 			if (FieldState == Chaos::EObjectStateType::Kinematic || FieldState == Chaos::EObjectStateType::Static)
 			{
-				RigidHandle->SetV(Chaos::FVec3(0));
-				RigidHandle->SetW(Chaos::FVec3(0));
+				RigidHandle->SetVf(Chaos::FVec3f(0));
+				RigidHandle->SetWf(Chaos::FVec3f(0));
 			}
 			else if (FieldState == Chaos::EObjectStateType::Dynamic)
 			{
@@ -603,7 +603,7 @@ namespace Chaos
 					Chaos::FPBDRigidParticleHandle* RigidHandle = ParticleHandles[Index.Sample]->CastToRigidParticle();
 					if (RigidHandle && RigidHandle->ObjectState() == Chaos::EObjectStateType::Dynamic)
 					{
-						RigidHandle->V() += ResultsView[Index.Result];
+						RigidHandle->SetV(RigidHandle->GetV() + ResultsView[Index.Result]);
 					}
 				}
 			}
@@ -635,7 +635,7 @@ namespace Chaos
 					Chaos::FPBDRigidParticleHandle* RigidHandle = ParticleHandles[Index.Sample]->CastToRigidParticle();
 					if (RigidHandle && RigidHandle->ObjectState() == Chaos::EObjectStateType::Dynamic)
 					{
-						RigidHandle->W() += ResultsView[Index.Result];
+						RigidHandle->SetW(RigidHandle->GetW() + ResultsView[Index.Result]);
 					}
 				}
 			}
diff --git a/Engine/Source/Runtime/Experimental/Chaos/Public/PhysicsProxy/SingleParticlePhysicsProxy.h b/Engine/Source/Runtime/Experimental/Chaos/Public/PhysicsProxy/SingleParticlePhysicsProxy.h
--- a/Engine/Source/Runtime/Experimental/Chaos/Public/PhysicsProxy/SingleParticlePhysicsProxy.h
+++ b/Engine/Source/Runtime/Experimental/Chaos/Public/PhysicsProxy/SingleParticlePhysicsProxy.h
@@ -228,7 +228,8 @@ public:
 	FUniqueIdx UniqueIdx() const { return Read([](auto* Particle) { return Particle->UniqueIdx(); }); }
 	void SetUniqueIdx(const FUniqueIdx UniqueIdx, bool bInvalidate = true) { Write([UniqueIdx, bInvalidate](auto* Particle) { Particle->SetUniqueIdx(UniqueIdx, bInvalidate); }); }
 
-	FRotation3 R() const { return Read([](auto* Particle) -> const auto { return Particle->R(); }); }
+	FRotation3 R() const { return Read([](auto* Particle) -> const auto { return Particle->GetR(); }); }
+	FRotation3 GetR() const { return Read([](auto* Particle) -> const auto { return Particle->GetR(); }); }
 
 protected:
 	void SetRBase(const FRotation3& InR, bool bInvalidate = true){ Write([&InR, bInvalidate, this](auto* Particle)
@@ -280,12 +281,16 @@ public:
 		{
 			if (auto Kinematic = Particle->CastToKinematicParticle())
 			{
-				return Kinematic->V();
+				return Kinematic->GetV();
 			}
 			
 			return FVec3(0);
 		});
 	}
+	const FVec3 GetV() const
+	{
+		return V();
+	}
 
 protected:
 	void SetVBase(const FVec3& InV, bool bInvalidate = true)
@@ -315,13 +320,18 @@ public:
 		{
 			if (auto Kinematic = Particle->CastToKinematicParticle())
 			{
-				return Kinematic->W();
+				return Kinematic->GetW();
 			}
 
 			return FVec3(0);
 		});
 	}
 
+	const FVec3 GetW() const
+	{
+		return W();
+	}
+
 protected:
 	void SetWBase(const FVec3& InW, bool bInvalidate = true)
 	{
@@ -676,7 +686,7 @@ public:
 		{
 			if (auto Rigid = Particle->CastToRigidParticle())
 			{
-				const FMatrix33 WorldI = Utilities::ComputeWorldSpaceInertia(Rigid->R() * Rigid->RotationOfMass(), Rigid->I());
+				const FMatrix33 WorldI = Utilities::ComputeWorldSpaceInertia(Rigid->GetR() * Rigid->RotationOfMass(), Rigid->I());
 				return WorldI * Rigid->AngularImpulseVelocity();
 			}
 
@@ -703,7 +713,7 @@ public:
 					}
 					else
 					{
-						const FMatrix33 WorldInvI = Utilities::ComputeWorldSpaceInertia(Rigid->R() * Rigid->RotationOfMass(), Rigid->InvI());
+						const FMatrix33 WorldInvI = Utilities::ComputeWorldSpaceInertia(Rigid->GetR() * Rigid->RotationOfMass(), Rigid->InvI());
 						Rigid->SetAngularImpulseVelocity(WorldInvI * InAngularImpulse, bInvalidate);
 					}
 				}
@@ -1341,7 +1351,7 @@ public:
 		VerifyContext();
 		if (auto Rigid = GetHandle_LowLevel()->CastToRigidParticle())
 		{
-			return Rigid->PreV();
+			return Rigid->GetPreV();
 		}
 		return FVec3(0);
 	}
@@ -1351,7 +1361,7 @@ public:
 		VerifyContext();
 		if (auto Rigid = GetHandle_LowLevel()->CastToRigidParticle())
 		{
-			return Rigid->PreW();
+			return Rigid->GetPreW();
 		}
 		return FVec3(0);
 	}
diff --git a/Engine/Source/Runtime/Experimental/Chaos/Public/RewindData.h b/Engine/Source/Runtime/Experimental/Chaos/Public/RewindData.h
--- a/Engine/Source/Runtime/Experimental/Chaos/Public/RewindData.h
+++ b/Engine/Source/Runtime/Experimental/Chaos/Public/RewindData.h
@@ -990,7 +990,7 @@ struct FGeometryParticleStateBase
 	void CachePreCorrectionState(const TParticle& Particle)
 	{
 		PreCorrectionXR.SetX(Particle.X());
-		PreCorrectionXR.SetR(Particle.R());
+		PreCorrectionXR.SetR(Particle.GetR());
 	}
 
 	TParticlePropertyBuffer<FParticlePositionRotation,EChaosProperty::XR> ParticlePositionRotation;
@@ -1031,10 +1031,10 @@ public:
 
 
 	REWIND_PARTICLE_STATIC_PROPERTY(ParticlePositionRotation, X)
-	REWIND_PARTICLE_STATIC_PROPERTY(ParticlePositionRotation, R)
+	REWIND_PARTICLE_STATIC_PROPERTY(ParticlePositionRotation, GetR)
 
-	REWIND_PARTICLE_KINEMATIC_PROPERTY(Velocities, V)
-	REWIND_PARTICLE_KINEMATIC_PROPERTY(Velocities, W)
+	REWIND_PARTICLE_KINEMATIC_PROPERTY(Velocities, GetV)
+	REWIND_PARTICLE_KINEMATIC_PROPERTY(Velocities, GetW)
 
 	REWIND_PARTICLE_RIGID_PROPERTY(DynamicsMisc, LinearEtherDrag)
 	REWIND_PARTICLE_RIGID_PROPERTY(DynamicsMisc, AngularEtherDrag)
@@ -1087,15 +1087,15 @@ public:
 		FString Out = FString::Printf(TEXT("ParticleID:[Global: %d Local: %d]\n"), Particle.ParticleID().GlobalID, Particle.ParticleID().LocalID);
 
 		REWIND_PARTICLE_TO_STR(X)
-		REWIND_PARTICLE_TO_STR(R)
+		REWIND_PARTICLE_TO_STR(GetR)
 		//REWIND_PARTICLE_TO_STR(Geometry)
 		//REWIND_PARTICLE_TO_STR(UniqueIdx)
 		//REWIND_PARTICLE_TO_STR(SpatialIdx)
 
 		if(Particle.CastToKinematicParticle())
 		{
-			REWIND_PARTICLE_TO_STR(V)
-			REWIND_PARTICLE_TO_STR(W)
+			REWIND_PARTICLE_TO_STR(GetV)
+			REWIND_PARTICLE_TO_STR(GetW)
 		}
 
 		if(Particle.CastToRigidParticle())
diff --git a/Engine/Source/Runtime/Experimental/ChaosVehicles/ChaosVehiclesCore/Private/SimModule/DeferredForcesModular.cpp b/Engine/Source/Runtime/Experimental/ChaosVehicles/ChaosVehiclesCore/Private/SimModule/DeferredForcesModular.cpp
--- a/Engine/Source/Runtime/Experimental/ChaosVehicles/ChaosVehiclesCore/Private/SimModule/DeferredForcesModular.cpp
+++ b/Engine/Source/Runtime/Experimental/ChaosVehicles/ChaosVehiclesCore/Private/SimModule/DeferredForcesModular.cpp
@@ -116,7 +116,7 @@ Chaos::FPBDRigidParticleHandle* FDeferredForcesModular::GetParticle(const FTrans
 
 	if (ClusterParticle != nullptr && Child != nullptr)
 	{
-		const FRigidTransform3 ClusterWorldTM(ClusterParticle->X(), ClusterParticle->R());
+		const FRigidTransform3 ClusterWorldTM(ClusterParticle->X(), ClusterParticle->GetR());
 		
 		Proxy->GetParticle_Internal(TransformIndex);
 
@@ -128,7 +128,7 @@ Chaos::FPBDRigidParticleHandle* FDeferredForcesModular::GetParticle(const FTrans
 		{
 			//FTransform ChildT = Child->GetTransformXRCom();
 			//const FRigidTransform3 ChildWorldTM(ChildT.GetLocation(), ChildT.GetRotation());
-			const FRigidTransform3 ChildWorldTM(Child->X(), Child->R());
+			const FRigidTransform3 ChildWorldTM(Child->X(), Child->GetR());
 			Frame = ChildWorldTM.GetRelativeTransform(ClusterWorldTM);
 			Frame.SetLocation(Frame.GetLocation() + OffsetTransform.TransformVector(PositionalOffset));
 
@@ -154,7 +154,7 @@ Chaos::FPBDRigidParticleHandle* FDeferredForcesModular::GetParticle(const FTrans
 	const int SingleChassisIndex = 0;
 	TransformOut = FTransform::Identity;
 
-	const FRigidTransform3 ClusterWorldTM(ClusterParticles[SingleChassisIndex]->X(), ClusterParticles[SingleChassisIndex]->R());
+	const FRigidTransform3 ClusterWorldTM(ClusterParticles[SingleChassisIndex]->X(), ClusterParticles[SingleChassisIndex]->GetR());
 	FRigidTransform3 Frame = FRigidTransform3::Identity;
 
 	if (Chaos::FPBDRigidParticleHandle* Child = GetParticleFromUniqueIndex(ParticleIdx, Particles))
@@ -167,7 +167,7 @@ Chaos::FPBDRigidParticleHandle* FDeferredForcesModular::GetParticle(const FTrans
 		{
 			//FTransform ChildT = Child->GetTransformXRCom();
 			//const FRigidTransform3 ChildWorldTM(ChildT.GetLocation(), ChildT.GetRotation());
-			const FRigidTransform3 ChildWorldTM(Child->X(), Child->R());
+			const FRigidTransform3 ChildWorldTM(Child->X(), Child->GetR());
 			Frame = ChildWorldTM.GetRelativeTransform(ClusterWorldTM);
 			Frame.SetLocation(Frame.GetLocation() + OffsetTransform.TransformVector(PositionalOffset));
 		}
@@ -321,7 +321,7 @@ void AddForce_Implementation(Chaos::FPBDRigidParticleHandle* RigidHandle, const
 	const Chaos::FVec3 WorldCOM = Chaos::FParticleUtilitiesGT::GetCoMWorldPosition(RigidHandle);
 
 	// local to world
-	const FTransform WorldTM(RigidHandle->R(), RigidHandle->X());
+	const FTransform WorldTM(RigidHandle->GetR(), RigidHandle->X());
 	Chaos::FVec3 Position = WorldTM.TransformPosition(OffsetTransform.GetLocation());
 	Chaos::FVec3 Force = WorldTM.TransformVector(ParticleOffsetTransform.TransformVector(OffsetTransform.TransformVector(LocalForce)));
 
@@ -371,7 +371,7 @@ void AddTorque_Implementation(Chaos::FPBDRigidParticleHandle* RigidHandle, const
 	if (RP)
 	{
 		// local to world
-		const FTransform WorldTM(RigidHandle->R(), RigidHandle->X());
+		const FTransform WorldTM(RigidHandle->GetR(), RigidHandle->X());
 		FVector T1 = OffsetTransform.TransformVector(LocalTorque);
 		FVector T2 = ParticleOffsetTransform.TransformVector(T1);
 		Chaos::FVec3 WorldTorque = WorldTM.TransformVector(T2);
diff --git a/Engine/Source/Runtime/Experimental/ChaosVehicles/ChaosVehiclesCore/Private/SimModule/SimModuleTree.cpp b/Engine/Source/Runtime/Experimental/ChaosVehicles/ChaosVehiclesCore/Private/SimModule/SimModuleTree.cpp
--- a/Engine/Source/Runtime/Experimental/ChaosVehicles/ChaosVehiclesCore/Private/SimModule/SimModuleTree.cpp
+++ b/Engine/Source/Runtime/Experimental/ChaosVehicles/ChaosVehiclesCore/Private/SimModule/SimModuleTree.cpp
@@ -301,7 +301,7 @@ void FSimModuleTree::UpdateModuleVelocites(FGeometryCollectionPhysicsProxy* Phys
 		{
 			if (Chaos::FPBDRigidClusteredParticleHandle* ParentParticle = PhysicsProxy->GetSolverClusterHandle_Internal(Module->GetTransformIndex()))
 			{
-				const FTransform BodyTransform(ParentParticle->R(), ParentParticle->X());
+				const FTransform BodyTransform(ParentParticle->GetR(), ParentParticle->X());
 
 				if (Module->IsBehaviourType(eSimModuleTypeFlags::Velocity))
 				{
@@ -320,11 +320,11 @@ void FSimModuleTree::UpdateModuleVelocites(FGeometryCollectionPhysicsProxy* Phys
 						FVector WorldLocation = BodyTransform.TransformPosition(Module->GetParentRelativeTransform().GetLocation());
 						const Chaos::FVec3 Arm = WorldLocation - Particle->X();
 
-						FVector WorldVelocity = Particle->V() - Chaos::FVec3::CrossProduct(Arm, Particle->W());
+						FVector WorldVelocity = Particle->GetV() - Chaos::FVec3::CrossProduct(Arm, Particle->GetW());
 						FVector LocalLinearVelocity = BodyTransform.InverseTransformVector(WorldVelocity);
 						LocalLinearVelocity = Module->GetClusteredTransform().InverseTransformVector(LocalLinearVelocity);
 
-						FVector LocalAngular = BodyTransform.InverseTransformVector(Particle->W());
+						FVector LocalAngular = BodyTransform.InverseTransformVector(Particle->GetW());
 						LocalAngular = Module->GetClusteredTransform().InverseTransformVector(LocalAngular);
 
 						Module->SetLocalLinearVelocity(LocalLinearVelocity);
@@ -358,7 +358,7 @@ void FSimModuleTree::UpdateModuleVelocites(FClusterUnionPhysicsProxy* PhysicsPro
 		{
 			if (ISimulationModuleBase* Module = SimulationModuleTree[i].SimModule)
 			{
-				const FTransform BodyTransform(ParentParticle->R(), ParentParticle->X());
+				const FTransform BodyTransform(ParentParticle->GetR(), ParentParticle->X());
 
 				if (Module->IsBehaviourType(eSimModuleTypeFlags::Velocity))
 				{
@@ -381,10 +381,10 @@ void FSimModuleTree::UpdateModuleVelocites(FClusterUnionPhysicsProxy* PhysicsPro
 
 						//Chaos::FDebugDrawQueue::GetInstance().DrawDebugLine(Particle->X(), Particle->X() + Arm, FColor::Yellow, false, -1.f, 0, 2.f);
 
-						FVector WorldVelocity = Particle->V() - Chaos::FVec3::CrossProduct(Arm, Particle->W());
+						FVector WorldVelocity = Particle->GetV() - Chaos::FVec3::CrossProduct(Arm, Particle->GetW());
 						FVector LocalVelocity = OffsetTransform.InverseTransformVector(BodyTransform.InverseTransformVector(WorldVelocity));
 
-						FVector LocalAngular = BodyTransform.InverseTransformVector(Particle->W());
+						FVector LocalAngular = BodyTransform.InverseTransformVector(Particle->GetW());
 						LocalAngular = Module->GetClusteredTransform().InverseTransformVector(LocalAngular);
 
 						Module->SetLocalLinearVelocity(LocalVelocity);
@@ -407,13 +407,13 @@ void FSimModuleTree::UpdateVehicleState(FClusterUnionPhysicsProxy* PhysicsProxy)
 
 	if (Chaos::FClusterUnionPhysicsProxy::FInternalParticle* ParentParticle = PhysicsProxy->GetParticle_Internal())
 	{
-		const FTransform BodyTransform(ParentParticle->R(), ParentParticle->X());
+		const FTransform BodyTransform(ParentParticle->GetR(), ParentParticle->X());
 
 		VehicleState.ForwardDir = BodyTransform.GetUnitAxis(EAxis::X);
 		VehicleState.UpDir = BodyTransform.GetUnitAxis(EAxis::Z);
 		VehicleState.RightDir = BodyTransform.GetUnitAxis(EAxis::Y);
-		VehicleState.ForwardSpeedKmh = Chaos::CmSToKmH(FVector::DotProduct(ParentParticle->V(), VehicleState.ForwardDir));
-		VehicleState.AngularVelocityRad = ParentParticle->W();
+		VehicleState.ForwardSpeedKmh = Chaos::CmSToKmH(FVector::DotProduct(ParentParticle->GetV(), VehicleState.ForwardDir));
+		VehicleState.AngularVelocityRad = ParentParticle->GetW();
 	}
 
 }
diff --git a/Engine/Source/Runtime/Experimental/ChaosVisualDebugger/Public/DataWrappers/ChaosVDParticleDataWrapper.h b/Engine/Source/Runtime/Experimental/ChaosVisualDebugger/Public/DataWrappers/ChaosVDParticleDataWrapper.h
--- a/Engine/Source/Runtime/Experimental/ChaosVisualDebugger/Public/DataWrappers/ChaosVDParticleDataWrapper.h
+++ b/Engine/Source/Runtime/Experimental/ChaosVisualDebugger/Public/DataWrappers/ChaosVDParticleDataWrapper.h
@@ -127,7 +127,7 @@ struct FChaosVDParticlePositionRotation : public FChaosVDWrapperDataBase
 	void CopyFrom(const TOther& Other)
 	{
 		MX = Other.X();
-		MR = Other.R();
+		MR = Other.GetR();
 		
 		bHasValidData = true;
 	}
@@ -173,8 +173,8 @@ struct FChaosVDParticleVelocities : public FChaosVDWrapperDataBase
 	template <typename TOther>
 	void CopyFrom(const TOther& Other)
 	{
-		MV = Other.V();
-		MW = Other.W();
+		MV = Other.GetV();
+		MW = Other.GetW();
 		bHasValidData = true;
 	}
 	
diff --git a/Engine/Source/Runtime/Experimental/GeometryCollectionEngine/Private/GeometryCollection/GeometryCollectionComponent.cpp b/Engine/Source/Runtime/Experimental/GeometryCollectionEngine/Private/GeometryCollection/GeometryCollectionComponent.cpp
--- a/Engine/Source/Runtime/Experimental/GeometryCollectionEngine/Private/GeometryCollection/GeometryCollectionComponent.cpp
+++ b/Engine/Source/Runtime/Experimental/GeometryCollectionEngine/Private/GeometryCollection/GeometryCollectionComponent.cpp
@@ -2195,7 +2195,7 @@ void UGeometryCollectionComponent::UpdateRepData()
 						// Because we cull ClustersToRep with abandoned level, we must make sure we don't add duplicates to one off activated.
 						// TODO: avoid search for entry for perf
 						// TODO: once we support deep fracture we should be able to remove one offs clusters that are now disabled, reducing the amount to be replicated
-						FGeometryCollectionActivatedCluster OneOffActivated(TransformGroupIdx, Root->V(), Root->W());
+						FGeometryCollectionActivatedCluster OneOffActivated(TransformGroupIdx, Root->GetV(), Root->GetW());
 						if(!LocalRepData.OneOffActivated.Contains(OneOffActivated))
 						{
 							LocalRepData.OneOffActivated.Add(OneOffActivated);
@@ -2238,9 +2238,9 @@ void UGeometryCollectionComponent::UpdateRepData()
 					FGeometryCollectionClusterRep& ClusterRep = LocalRepData.Clusters.AddDefaulted_GetRef();
 
 					ClusterRep.Position = Cluster->X();
-					ClusterRep.Rotation = Cluster->R();
-					ClusterRep.LinearVelocity = Cluster->V();
-					ClusterRep.AngularVelocity = Cluster->W();
+					ClusterRep.Rotation = Cluster->GetR();
+					ClusterRep.LinearVelocity = Cluster->GetV();
+					ClusterRep.AngularVelocity = Cluster->GetW();
 					ClusterRep.ClusterState.SetObjectState(Cluster->ObjectState());
 					ClusterRep.ClusterState.SetInternalCluster(Cluster->InternalCluster());
 					int32 TransformGroupIdx;
@@ -2371,8 +2371,8 @@ void UGeometryCollectionComponent::UpdateRepStateAndDynamicData()
 							{
 								FGeometryCollectionRepStateData::FReleasedData& Data = LocalRepStateData.ReleasedData.AddDefaulted_GetRef();
 								Data.TransformIndex = TransformIndex;
-								Data.LinearVelocity = ParticleHandle->V();
-								Data.AngularVelocityInDegreesPerSecond = FMath::RadiansToDegrees(ParticleHandle->W());
+								Data.LinearVelocity = ParticleHandle->GetV();
+								Data.AngularVelocityInDegreesPerSecond = FMath::RadiansToDegrees(ParticleHandle->GetW());
 							}
 						}
 
@@ -2412,9 +2412,9 @@ void UGeometryCollectionComponent::UpdateRepStateAndDynamicData()
 				Data.TransformIndex = Root.TransformIndex;
 				Data.bIsInternalCluster = Root.Handle->InternalCluster();
 				Data.Position = Root.Handle->X();
-				Data.EulerRotation = Root.Handle->R().Rotator().Euler();
-				Data.LinearVelocity = Root.Handle->V();
-				Data.AngularVelocityInDegreesPerSecond = FMath::RadiansToDegrees(Root.Handle->W());
+				Data.EulerRotation = Root.Handle->GetR().Rotator().Euler();
+				Data.LinearVelocity = Root.Handle->GetV();
+				Data.AngularVelocityInDegreesPerSecond = FMath::RadiansToDegrees(Root.Handle->GetW());
 				Data.LastUpdatedVersion = LocalRepDynamicData.Version + 1;
 			}
 
@@ -2600,7 +2600,7 @@ namespace
 			//
 			const FVector RepAngVel = AngularVelocity;
 			const Chaos::FRotation3 RepExtrapAng = Chaos::FRotation3::IntegrateRotationWithAngularVelocity(Rotation, RepAngVel, RepExtrapTime);
-			const FVector AngVel = Chaos::FRotation3::CalculateAngularVelocity(Cluster.R(), RepExtrapAng, GeometryCollectionRepAngularMatchTime);
+			const FVector AngVel = Chaos::FRotation3::CalculateAngularVelocity(Cluster.GetR(), RepExtrapAng, GeometryCollectionRepAngularMatchTime);
 			if (AngVel.SizeSquared() > SMALL_NUMBER)
 			{
 				Cluster.SetW(RepAngVel + AngVel);
@@ -2935,7 +2935,7 @@ bool UGeometryCollectionComponent::ProcessRepData(const float DeltaTime, const f
 			{
 				FVector Axis;
 				float Angle;
-				(RepCluster.Rotation.Inverse() * Cluster.R()).ToAxisAndAngle(Axis, Angle);
+				(RepCluster.Rotation.Inverse() * Cluster.GetR()).ToAxisAndAngle(Axis, Angle);
 				if (FMath::Abs(Angle) < .1f)
 				{
 					return;
@@ -2943,8 +2943,8 @@ bool UGeometryCollectionComponent::ProcessRepData(const float DeltaTime, const f
 			}
 
 			Chaos::FDebugDrawQueue& DrawQueue = Chaos::FDebugDrawQueue::GetInstance();
-			DrawQueue.DrawDebugCoordinateSystem(Cluster.X(), FRotator(Cluster.R()), 100.f, false, -1, -1, 1.f);
-			DrawQueue.DrawDebugBox(Cluster.X() + Cluster.LocalBounds().Center(), Cluster.LocalBounds().Extents(), Cluster.R(), FColor::White, false, -1, -1, 1.f);
+			DrawQueue.DrawDebugCoordinateSystem(Cluster.X(), FRotator(Cluster.GetR()), 100.f, false, -1, -1, 1.f);
+			DrawQueue.DrawDebugBox(Cluster.X() + Cluster.LocalBounds().Center(), Cluster.LocalBounds().Extents(), Cluster.GetR(), FColor::White, false, -1, -1, 1.f);
 			DrawQueue.DrawDebugBox(RepCluster.Position + Cluster.LocalBounds().Center(), Cluster.LocalBounds().Extents(), RepCluster.Rotation, FColor::Green, false, -1, -1, 1.f);
 
 			if (bGeometryCollectionRepUseClusterVelocityMatch)
diff --git a/Engine/Source/Runtime/PhysicsCore/Private/PhysicsInterfaceUtilsCore.cpp b/Engine/Source/Runtime/PhysicsCore/Private/PhysicsInterfaceUtilsCore.cpp
--- a/Engine/Source/Runtime/PhysicsCore/Private/PhysicsInterfaceUtilsCore.cpp
+++ b/Engine/Source/Runtime/PhysicsCore/Private/PhysicsInterfaceUtilsCore.cpp
@@ -29,7 +29,7 @@ FCollisionFilterData ToUnrealFilterData(const FChaosFilterData& FilterData)
 template <typename THitLocation>
 uint32 FindFaceIndexImp(const THitLocation& PHit, const FVector& UnitDir)
 {
-	const FTransform WorldTM(PHit.Actor->R(), PHit.Actor->X());
+	const FTransform WorldTM(PHit.Actor->GetR(), PHit.Actor->X());
 	const FVector LocalPosition = WorldTM.InverseTransformPositionNoScale(PHit.WorldPosition);
 	const FVector LocalUnitDir = WorldTM.InverseTransformVectorNoScale(UnitDir);
 	return PHit.Shape->GetGeometry()->FindMostOpposingFace(LocalPosition, LocalUnitDir, PHit.FaceIndex, 1);	//todo:this number matches the one above, but is it right?
diff --git a/Engine/Source/Runtime/PhysicsCore/Public/SQVisitor.h b/Engine/Source/Runtime/PhysicsCore/Public/SQVisitor.h
--- a/Engine/Source/Runtime/PhysicsCore/Public/SQVisitor.h
+++ b/Engine/Source/Runtime/PhysicsCore/Public/SQVisitor.h
@@ -201,7 +201,7 @@ private:
 		const bool bTestShapeBounds = Shapes.Num() > 1;
 		bool bContinue = true;
 
-		const FRigidTransform3 ActorTM(GeometryParticle->X(), GeometryParticle->R());
+		const FRigidTransform3 ActorTM(GeometryParticle->X(), GeometryParticle->GetR());
 		const TAABB<FReal, 3> QueryGeomWorldBounds = QueryGeom ? QueryGeom->CalculateTransformedBounds(StartTM) : TAABB<FReal, 3>(-HalfExtents, HalfExtents);
 
 #if CHAOS_DEBUG_DRAW
@@ -215,7 +215,7 @@ private:
 
 			const FImplicitObject* Geom = Shape->GetGeometry();
 
-			CVD_TRACE_SCOPED_SCENE_QUERY_VISIT_HELPER(EChaosVDSceneQueryVisitorType::NarrowPhase, FTransform(GeometryParticle->R(), GeometryParticle->X()), Payload.UniqueIdx().Idx, ShapeIndex, CurData);
+			CVD_TRACE_SCOPED_SCENE_QUERY_VISIT_HELPER(EChaosVDSceneQueryVisitorType::NarrowPhase, FTransform(GeometryParticle->GetR(), GeometryParticle->X()), Payload.UniqueIdx().Idx, ShapeIndex, CurData);
 
 			if (bTestShapeBounds)
 			{
@@ -564,7 +564,7 @@ private:
 		{
 			const TUniquePtr<FPerShapeData>& Shape = Shapes[ShapeIndex];
 
-			CVD_TRACE_SCOPED_SCENE_QUERY_VISIT_HELPER(EChaosVDSceneQueryVisitorType::BroadPhase, FTransform(GeometryParticle->R(), GeometryParticle->X()), Payload.UniqueIdx().Idx, ShapeIndex, CurData);
+			CVD_TRACE_SCOPED_SCENE_QUERY_VISIT_HELPER(EChaosVDSceneQueryVisitorType::BroadPhase, FTransform(GeometryParticle->GetR(), GeometryParticle->X()), Payload.UniqueIdx().Idx, ShapeIndex, CurData);
 
 			ECollisionQueryHitType HitType = QueryFilterData.flags & FPhysicsQueryFlag::ePREFILTER ? QueryCallback.PreFilter(QueryFilterDataConcrete, *Shape, *GeometryParticle) : ECollisionQueryHitType::Block;
 			if (HitType != ECollisionQueryHitType::None)
