Iris - Don't crash if trying to resolve object references for non-existing object.

#rnx

diff --git a/Engine/Source/Runtime/Experimental/Iris/Core/Private/Iris/ReplicationSystem/ReplicationReader.cpp b/Engine/Source/Runtime/Experimental/Iris/Core/Private/Iris/ReplicationSystem/ReplicationReader.cpp
--- a/Engine/Source/Runtime/Experimental/Iris/Core/Private/Iris/ReplicationSystem/ReplicationReader.cpp
+++ b/Engine/Source/Runtime/Experimental/Iris/Core/Private/Iris/ReplicationSystem/ReplicationReader.cpp
@@ -1314,6 +1314,15 @@ void FReplicationReader::ResolveAndDispatchUnresolvedReferencesForObject(FNetSer
 	IRIS_PROFILER_SCOPE(FReplicationReader_ResolveAndDispatchUnresolvedReferencesForObject);
 
 	FReplicatedObjectInfo* ReplicationInfo = GetReplicatedObjectInfo(InternalIndex);
+	// Unexpected. Get more info.
+	if (!ReplicationInfo)
+	{
+		static bool bHasLogged = false;
+		UE_CLOG(!bHasLogged, LogIris, Error, TEXT("Trying to resolve references for non-existing object ( InternalIndex: %u )"), InternalIndex);
+		bHasLogged = true;
+		ensure(false);
+		return;
+	}
 
 	const bool bObjectHasAttachments = ReplicationInfo->bHasAttachments;
 	const bool bObjectHasReferences = ReplicationInfo->bHasUnresolvedInitialReferences | ReplicationInfo->bHasUnresolvedReferences;
