Iris - Handle unexpected state gracefully instead of crashing.

#rnx

diff --git a/Engine/Source/Runtime/Experimental/Iris/Core/Private/Iris/ReplicationSystem/ReplicationWriter.cpp b/Engine/Source/Runtime/Experimental/Iris/Core/Private/Iris/ReplicationSystem/ReplicationWriter.cpp
--- a/Engine/Source/Runtime/Experimental/Iris/Core/Private/Iris/ReplicationSystem/ReplicationWriter.cpp
+++ b/Engine/Source/Runtime/Experimental/Iris/Core/Private/Iris/ReplicationSystem/ReplicationWriter.cpp
@@ -1695,7 +1695,12 @@ uint32 FReplicationWriter::WriteObjectsPendingDestroy(FNetSerializationContext&
 			}
 		}
 
-		check(Info.GetState() == EReplicatedObjectState::PendingDestroy);
+		// Unexpected. Get more info.
+		if (Info.GetState() != EReplicatedObjectState::PendingDestroy)
+		{
+			ensureMsgf(Info.GetState() == EReplicatedObjectState::PendingDestroy, TEXT("Skipping writing destroy for object %s which is in unexpected state %s. IsSubObject: %u IsDestructionInfo: %u"), ToCStr(NetRefHandleManager->PrintObjectFromIndex(InternalIndex)), LexToString(Info.GetState()), Info.IsSubObject, Info.IsDestructionInfo);
+			continue;
+		}
 
 		// We do not support destroying an object that is currently being sent as a huge object.
 		if (IsObjectPartOfActiveHugeObject(InternalIndex))
