Add soft collision thickness and sticky friction

#changelist validated
#virtualized

diff --git a/Engine/Programs/UnrealHeaderTool/Config/DefaultEngine.ini b/Engine/Programs/UnrealHeaderTool/Config/DefaultEngine.ini
--- a/Engine/Programs/UnrealHeaderTool/Config/DefaultEngine.ini
+++ b/Engine/Programs/UnrealHeaderTool/Config/DefaultEngine.ini
@@ -113,6 +113,7 @@ gc.MaxObjectsInProgram=262144
 +Units="kgcm/s\u00B2"
 +Units="NewtonMeters"
 +Units="Nm"
++Units="kgcm"
 +Units="KilogramCentimetersSquaredPerSecondSquared"
 +Units="kgcm\u00B2/s\u00B2"
 +Units="Hertz"
diff --git a/Engine/Source/Editor/UnrealEd/Private/Settings/EditorProjectSettings.cpp b/Engine/Source/Editor/UnrealEd/Private/Settings/EditorProjectSettings.cpp
--- a/Engine/Source/Editor/UnrealEd/Private/Settings/EditorProjectSettings.cpp
+++ b/Engine/Source/Editor/UnrealEd/Private/Settings/EditorProjectSettings.cpp
@@ -125,6 +125,14 @@ void UEditorProjectAppearanceSettings::PostEditChangeProperty( struct FPropertyC
 	{
 		Settings.SetDisplayUnits(EUnitType::Torque, TorqueUnits);
 	}
+	else if (Name == GET_MEMBER_NAME_CHECKED(UEditorProjectAppearanceSettings, ImpulseUnits))
+	{
+		Settings.SetDisplayUnits(EUnitType::Impulse, ImpulseUnits);
+	}
+	else if (Name == GET_MEMBER_NAME_CHECKED(UEditorProjectAppearanceSettings, PositionalImpulseUnits))
+	{
+		Settings.SetDisplayUnits(EUnitType::PositionalImpulse, PositionalImpulseUnits);
+	}
 	else if (Name == GET_MEMBER_NAME_CHECKED(UEditorProjectAppearanceSettings, bDisplayUnits))
 	{
 		Settings.SetShouldDisplayUnits(bDisplayUnits);
diff --git a/Engine/Source/Editor/UnrealEd/Public/Settings/EditorProjectSettings.h b/Engine/Source/Editor/UnrealEd/Public/Settings/EditorProjectSettings.h
--- a/Engine/Source/Editor/UnrealEd/Public/Settings/EditorProjectSettings.h
+++ b/Engine/Source/Editor/UnrealEd/Public/Settings/EditorProjectSettings.h
@@ -104,6 +104,12 @@ public:
 	UPROPERTY(EditAnywhere, config, Category = Units, AdvancedDisplay, meta = (DisplayName = "Torque", Tooltip = "Choose the units in which to display torques.", ValidEnumValues = "NewtonMeters, KilogramCentimetersSquaredPerSecondSquared"))
 	EUnit TorqueUnits;
 
+	UPROPERTY(EditAnywhere, config, Category = Units, AdvancedDisplay, meta = (DisplayName = "Impulse", Tooltip = "Choose the units in which to display impulses.", ValidEnumValues = "NewtonSeconds"))
+	EUnit ImpulseUnits;
+
+	UPROPERTY(EditAnywhere, config, Category = Units, AdvancedDisplay, meta = (DisplayName = "PositionalImpulse", Tooltip = "Choose the units in which to display positional impulses.", ValidEnumValues = "KilogramMeters, KilogramCentimeters"))
+	EUnit PositionalImpulseUnits;
+
 	// Should the Reference Viewer have 'Show Searchable Names' checked by default when opened in this project
 	UPROPERTY(EditAnywhere, config, Category=ReferenceViewer)
 	EReferenceViewerSettingMode ShowSearchableNames;
diff --git a/Engine/Source/Programs/HeadlessChaos/Private/HeadlessChaosTestConstraintSolverCore.cpp b/Engine/Source/Programs/HeadlessChaos/Private/HeadlessChaosTestConstraintSolverCore.cpp
--- a/Engine/Source/Programs/HeadlessChaos/Private/HeadlessChaosTestConstraintSolverCore.cpp
+++ b/Engine/Source/Programs/HeadlessChaos/Private/HeadlessChaosTestConstraintSolverCore.cpp
@@ -77,7 +77,7 @@ namespace ChaosTest
 
 
 		Collision.SetSolverBodies(&Bodies[0], &Bodies[1]);
-		Collision.SetFriction(0, 0, 0);
+		Collision.SetFriction(0, 0, 0, 0);
 		Collision.SetNumManifoldPoints(1);
 		Collision.InitContact(0, FVec3(0,0,Radius), FVec3(0,0,-Radius), FVec3(0,0,1));
 		Collision.InitMaterial(0, 0, false, 0);
diff --git a/Engine/Source/Runtime/Core/Private/Math/UnitConversion.cpp b/Engine/Source/Runtime/Core/Private/Math/UnitConversion.cpp
--- a/Engine/Source/Runtime/Core/Private/Math/UnitConversion.cpp
+++ b/Engine/Source/Runtime/Core/Private/Math/UnitConversion.cpp
@@ -80,6 +80,10 @@ constexpr FParseCandidate ParseCandidates[] = {
 	{ TEXT("NewtonMeters"),			EUnit::NewtonMeters },			{ TEXT("Nm"),		EUnit::NewtonMeters },
 	{ TEXT("KilogramsCentimetersSquaredPerSecondSquared"),	EUnit::KilogramCentimetersSquaredPerSecondSquared },		{ TEXT("kgcm2/s2"),	EUnit::KilogramCentimetersSquaredPerSecondSquared },	{ TEXT("kgcm\u00B2/s\u00B2"),	EUnit::KilogramCentimetersSquaredPerSecondSquared },
 
+	{ TEXT("NewtonSeconds"),		EUnit::NewtonSeconds },			{ TEXT("Ns"),		EUnit::NewtonSeconds },
+	{ TEXT("KilogramCentimeters"),	EUnit::KilogramCentimeters },	{ TEXT("kgcm"),		EUnit::KilogramCentimeters },
+	{ TEXT("KilogramMeters"),		EUnit::KilogramMeters },		{ TEXT("kgm"),		EUnit::KilogramMeters },
+
 	{ TEXT("Hertz"),				EUnit::Hertz },					{ TEXT("Hz"),		EUnit::Hertz },
 	{ TEXT("Kilohertz"),			EUnit::Kilohertz },				{ TEXT("KHz"),		EUnit::Kilohertz },
 	{ TEXT("Megahertz"),			EUnit::Megahertz },				{ TEXT("MHz"),		EUnit::Megahertz },
@@ -145,6 +149,8 @@ constexpr const TCHAR* const DisplayStrings[] = {
 
 	TEXT("Nm"), TEXT("kgcm\u00B2/s\u00B2"),
 
+	TEXT("Ns"), TEXT("kgcm"), TEXT("kgm"),
+
 	TEXT("Hz"), TEXT("KHz"), TEXT("MHz"), TEXT("GHz"), TEXT("rpm"),
 
 	TEXT("B"), TEXT("KiB"), TEXT("MiB"), TEXT("GiB"), TEXT("TiB"),
@@ -214,6 +220,10 @@ constexpr const TCHAR* const SupportedUnitsStrings[] = {
 	TEXT("NewtonMeters"),			
 	TEXT("KilogramsCentimetersSquaredPerSecondSquared"),
 
+	TEXT("NewtonSeconds"),
+	TEXT("KilogramCentimeters"),
+	TEXT("KilogramMeters"),
+
 	TEXT("Hertz"),				
 	TEXT("Kilohertz"),			
 	TEXT("Megahertz"),			
@@ -282,6 +292,8 @@ constexpr EUnitType UnitTypes[] = {
 
 	EUnitType::Torque,		EUnitType::Torque,
 
+	EUnitType::Impulse,		EUnitType::PositionalImpulse, EUnitType::PositionalImpulse,
+
 	EUnitType::Frequency,	EUnitType::Frequency,	EUnitType::Frequency,	EUnitType::Frequency,	EUnitType::Frequency,
 
 	EUnitType::DataSize,	EUnitType::DataSize,	EUnitType::DataSize,	EUnitType::DataSize,	EUnitType::DataSize,
@@ -566,6 +578,8 @@ FUnitSettings::FUnitSettings()
 	DisplayUnits[(uint8)EUnitType::Density].Add(EUnit::GramsPerCubicCentimeter);
 	DisplayUnits[(uint8)EUnitType::Force].Add(EUnit::Newtons);
 	DisplayUnits[(uint8)EUnitType::Torque].Add(EUnit::NewtonMeters);
+	DisplayUnits[(uint8)EUnitType::Impulse].Add(EUnit::NewtonSeconds);
+	DisplayUnits[(uint8)EUnitType::PositionalImpulse].Add(EUnit::KilogramCentimeters);
 	DisplayUnits[(uint8)EUnitType::Frequency].Add(EUnit::Hertz);
 	DisplayUnits[(uint8)EUnitType::DataSize].Add(EUnit::Megabytes);
 	DisplayUnits[(uint8)EUnitType::LuminousFlux].Add(EUnit::Lumens);
@@ -807,6 +821,16 @@ namespace UnitConversion
 		}
 	}
 
+	double PositionalImpulseUnificationFactor(EUnit From)
+	{
+		// Convert to KilogramMeters
+		switch (From)
+		{
+		case EUnit::KilogramCentimeters:	return 0.01;
+		default:							return 1;
+		}
+	}
+
 	double FrequencyUnificationFactor(EUnit From)
 	{
 		// Convert to KHz
diff --git a/Engine/Source/Runtime/Core/Public/Math/UnitConversion.h b/Engine/Source/Runtime/Core/Public/Math/UnitConversion.h
--- a/Engine/Source/Runtime/Core/Public/Math/UnitConversion.h
+++ b/Engine/Source/Runtime/Core/Public/Math/UnitConversion.h
@@ -50,6 +50,9 @@ enum class EUnit : uint8
 	/** Torque Units */
 	NewtonMeters, KilogramCentimetersSquaredPerSecondSquared,
 
+	/** Impulse Units */
+	NewtonSeconds, KilogramCentimeters, KilogramMeters,
+
 	/** Frequency units */
 	Hertz, Kilohertz, Megahertz, Gigahertz, RevolutionsPerMinute,
 
@@ -78,7 +81,7 @@ enum class EUnit : uint8
 /** Enumeration that specifies particular classes of unit */
 enum class EUnitType
 {
-	Distance, Angle, Speed, AngularSpeed, Acceleration, Temperature, Mass, Density, Force, Torque, Frequency, DataSize, LuminousFlux, LuminousIntensity, Illuminance, Luminance, Time, PixelDensity, Multipliers, ExposureValue, Stress,
+	Distance, Angle, Speed, AngularSpeed, Acceleration, Temperature, Mass, Density, Force, Torque, Impulse, PositionalImpulse, Frequency, DataSize, LuminousFlux, LuminousIntensity, Illuminance, Luminance, Time, PixelDensity, Multipliers, ExposureValue, Stress,
 
 	// Symbolic entry - do not use directly
 	NumberOf,
diff --git a/Engine/Source/Runtime/Core/Public/Math/UnitConversion.inl b/Engine/Source/Runtime/Core/Public/Math/UnitConversion.inl
--- a/Engine/Source/Runtime/Core/Public/Math/UnitConversion.inl
+++ b/Engine/Source/Runtime/Core/Public/Math/UnitConversion.inl
@@ -38,6 +38,8 @@ namespace UnitConversion
 	CORE_API double ForceUnificationFactor(EUnit From);
 	/** Find the common quantization factor for the specified torque unit. Quantizes to Newton meters. */
 	CORE_API double TorqueUnificationFactor(EUnit From);
+	/** Find the common quantization factor for the specified position impulse unit. Quantizes to kilogram centimeters. */
+	CORE_API double PositionalImpulseUnificationFactor(EUnit From);
 	/** Find the common quantization factor for the specified frequency unit. Quantizes to KHz. */
 	CORE_API double FrequencyUnificationFactor(EUnit From);
 	/** Find the common quantization factor for the specified data size unit. Quantizes to MB. */
@@ -95,6 +97,7 @@ T FUnitConversion::Convert(T InValue, EUnit From, EUnit To)
 		case EUnitType::Density:			return InValue * DensityUnificationFactor(From)				* (1.0 / DensityUnificationFactor(To));
 		case EUnitType::Force:				return InValue * ForceUnificationFactor(From) 				* (1.0 / ForceUnificationFactor(To));
 		case EUnitType::Torque:				return InValue * TorqueUnificationFactor(From)				* (1.0 / TorqueUnificationFactor(To));
+		case EUnitType::PositionalImpulse:	return InValue * PositionalImpulseUnificationFactor(From)	* (1.0 / PositionalImpulseUnificationFactor(To));
 		case EUnitType::Frequency:			return InValue * FrequencyUnificationFactor(From) 			* (1.0 / FrequencyUnificationFactor(To));
 		case EUnitType::DataSize:			return InValue * DataSizeUnificationFactor(From) 			* (1.0 / DataSizeUnificationFactor(To));
 		case EUnitType::LuminousFlux:		return InValue;	
diff --git a/Engine/Source/Runtime/CoreUObject/Public/UObject/NoExportTypes.h b/Engine/Source/Runtime/CoreUObject/Public/UObject/NoExportTypes.h
--- a/Engine/Source/Runtime/CoreUObject/Public/UObject/NoExportTypes.h
+++ b/Engine/Source/Runtime/CoreUObject/Public/UObject/NoExportTypes.h
@@ -385,6 +385,9 @@ enum class EUnit : uint8
 
 	/** Torque units */
 	NewtonMeters, KilogramCentimetersSquaredPerSecondSquared,
+
+	/** Impulse units */
+	NewtonSeconds, KilogramCentimeters, KilogramMeters,
 	
 	/** Frequency units */
 	Hertz, Kilohertz, Megahertz, Gigahertz, RevolutionsPerMinute,
diff --git a/Engine/Source/Runtime/Experimental/Chaos/Private/Chaos/Collision/PBDCollisionContainerSolver.cpp b/Engine/Source/Runtime/Experimental/Chaos/Private/Chaos/Collision/PBDCollisionContainerSolver.cpp
--- a/Engine/Source/Runtime/Experimental/Chaos/Private/Chaos/Collision/PBDCollisionContainerSolver.cpp
+++ b/Engine/Source/Runtime/Experimental/Chaos/Private/Chaos/Collision/PBDCollisionContainerSolver.cpp
@@ -59,6 +59,10 @@ namespace Chaos
 		// The maximum number of constraints we will attempt to solve (-1 for unlimited)
 		int32 Chaos_Collision_MaxSolverManifoldPoints = -1;
 		FAutoConsoleVariableRef CVarChaosCollisionMaxSolverManifoldPoints(TEXT("p.Chaos.PBDCollisionSolver.MaxManifoldPoints"), Chaos_Collision_MaxSolverManifoldPoints, TEXT(""));
+
+		// Whether to enable the experimental soft collisions
+		bool bChaos_Collision_EnableSoftCollisions = true;
+		FAutoConsoleVariableRef CVarChaosCollisionEnableSoftCollisions(TEXT("p.Chaos.PBDCollisionSolver.EnableSoftCollisions"), bChaos_Collision_EnableSoftCollisions, TEXT(""));
 	}
 
 
@@ -153,6 +157,9 @@ namespace Chaos
 			}
 		}
 
+		// Overlap remaining from the previous frame, estimated from current contact phi and velocity
+		FRealSingle WorldContactResidualPhi = FMath::Min(WorldContactDeltaNormal, FRealSingle(0)) - FMath::Min(ContactVelocityNormal * Dt, FRealSingle(0));
+
 		// Initial Phi for initial-overlap depenetration.
 		// If we have an initial contact, calculate the initial overlap. This will get saved in SetSolverResults
 		FRealSingle WorldContactInitialPhi = 0;
@@ -161,10 +168,7 @@ namespace Chaos
 			if (ManifoldPoint.Flags.bInitialContact)
 			{
 				// This is a new manifold point, capture current Phi as the initial Phi
-				WorldContactInitialPhi = FMath::Min(WorldContactDeltaNormal, FRealSingle(0));
-
-				// Reduce initial penetration depth by any movement towards each other this tick
-				WorldContactInitialPhi -= FMath::Min(ContactVelocityNormal * Dt, FRealSingle(0));
+				WorldContactInitialPhi = WorldContactResidualPhi;
 
 				// If this is a new manifold point on a pre-existing manifold, we limit the initial depth.
 				// This is so that as we are depenetrating and new points are added to the manifold, we don't suddenly pop out, which would
@@ -278,6 +282,10 @@ namespace Chaos
 		FSolverReal PositionStaticFriction = FSolverReal(0);
 		FSolverReal PositionDynamicFriction = FSolverReal(0);
 		FSolverReal VelocityDynamicFriction = FSolverReal(0);
+		FSolverReal MinFrictionPushOut
+			= CVars::bChaos_PBDCollisionSolver_Velocity_MinMaxFriction
+			? FSolverReal(Constraint->GetMinFrictionPushOut())
+			: 0;
 		if (SolverSettings.NumPositionFrictionIterations > 0)
 		{
 			PositionStaticFriction = StaticFriction;
@@ -298,11 +306,15 @@ namespace Chaos
 			VelocityDynamicFriction = DynamicFriction;
 		}
 
-		Solver.SetFriction(PositionStaticFriction, PositionDynamicFriction, VelocityDynamicFriction);
+		Solver.SetFriction(PositionStaticFriction, PositionDynamicFriction, VelocityDynamicFriction, MinFrictionPushOut);
+
 
 		const FReal SolverStiffness = Constraint->GetStiffness();
+
 		Solver.SetStiffness(FSolverReal(SolverStiffness));
 
+		Solver.SetHardContact();
+
 		Solver.SolverBody0().SetInvMScale(Constraint->GetInvMassScale0());
 		Solver.SolverBody0().SetInvIScale(Constraint->GetInvInertiaScale0());
 		Solver.SolverBody0().SetShockPropagationScale(FReal(1));
@@ -313,6 +325,18 @@ namespace Chaos
 		bOutPerIterationCollision = (!Constraint->GetUseManifold() || Constraint->GetUseIncrementalCollisionDetection());
 
 		UpdateCollisionSolverManifoldFromConstraint(Solver, Constraint, Dt, 0, Constraint->NumManifoldPoints(), SolverSettings);
+
+		// Convert to a soft collision if the constraint has a nonzero soft separation
+		// NOTE: must come after UpdateCollisionSolverManifoldFromConstraint because we scale the spring forces by the number of manifold points.
+		if (Constraint->IsSoftContact() && CVars::bChaos_Collision_EnableSoftCollisions)
+		{
+			// NOTE: convert stiffness and damping into XPBD terms
+			if (Solver.NumManifoldPoints() > 0)
+			{
+				const FSolverReal SoftScale = FSolverReal(1) / FSolverReal(Solver.NumManifoldPoints());
+				Solver.SetSoftContact(Constraint->GetSoftSeparation());
+			}
+		}
 	}
 
 	FORCEINLINE_DEBUGGABLE void UpdateCollisionConstraintFromSolver(FPBDCollisionConstraint* Constraint, const Private::FPBDCollisionSolver& Solver, const FSolverReal Dt)
diff --git a/Engine/Source/Runtime/Experimental/Chaos/Private/Chaos/Collision/PBDCollisionSolver.cpp b/Engine/Source/Runtime/Experimental/Chaos/Private/Chaos/Collision/PBDCollisionSolver.cpp
--- a/Engine/Source/Runtime/Experimental/Chaos/Private/Chaos/Collision/PBDCollisionSolver.cpp
+++ b/Engine/Source/Runtime/Experimental/Chaos/Private/Chaos/Collision/PBDCollisionSolver.cpp
@@ -47,12 +47,14 @@ namespace Chaos
 		bool bChaos_PBDCollisionSolver_Velocity_FrictionEnabled = true;
 		float Chaos_PBDCollisionSolver_Velocity_StaticFrictionStiffness = 1.0f;
 		bool bChaos_PBDCollisionSolver_Velocity_AveragePointEnabled = false;
+		bool bChaos_PBDCollisionSolver_Velocity_MinMaxFriction = false;
 
 		FAutoConsoleVariableRef CVarChaos_PBDCollisionSolver_Velocity_SolveEnabled(TEXT("p.Chaos.PBDCollisionSolver.Velocity.SolveEnabled"), bChaos_PBDCollisionSolver_Velocity_SolveEnabled, TEXT(""));
 		FAutoConsoleVariableRef CVarChaos_PBDCollisionSolver_Velocity_MinInvMassScale(TEXT("p.Chaos.PBDCollisionSolver.Velocity.MinInvMassScale"), Chaos_PBDCollisionSolver_Velocity_MinInvMassScale, TEXT(""));
 		FAutoConsoleVariableRef CVarChaos_PBDCollisionSolver_Velocity_FrictionEnabled(TEXT("p.Chaos.PBDCollisionSolver.Velocity.FrictionEnabled"), bChaos_PBDCollisionSolver_Velocity_FrictionEnabled, TEXT(""));
 		FAutoConsoleVariableRef CVarChaos_PBDCollisionSolver_Velocity_StaticFrictionStiffness(TEXT("p.Chaos.PBDCollisionSolver.Velocity.StaticFriction.Stiffness"), Chaos_PBDCollisionSolver_Velocity_StaticFrictionStiffness, TEXT(""));
 		FAutoConsoleVariableRef CVarChaos_PBDCollisionSolver_Velocity_AveragePointEnabled(TEXT("p.Chaos.PBDCollisionSolver.Velocity.AveragePointEnabled"), bChaos_PBDCollisionSolver_Velocity_AveragePointEnabled, TEXT(""));
+		FAutoConsoleVariableRef CVarChaos_PBDCollisionSolver_Velocity_MinMaxFriction(TEXT("p.Chaos.PBDCollisionSolver.Velocity.MinMaxFriction"), bChaos_PBDCollisionSolver_Velocity_MinMaxFriction, TEXT(""));
 	}
 	using namespace CVars;
 
diff --git a/Engine/Source/Runtime/Experimental/Chaos/Private/Chaos/PBDCollisionConstraints.cpp b/Engine/Source/Runtime/Experimental/Chaos/Private/Chaos/PBDCollisionConstraints.cpp
--- a/Engine/Source/Runtime/Experimental/Chaos/Private/Chaos/PBDCollisionConstraints.cpp
+++ b/Engine/Source/Runtime/Experimental/Chaos/Private/Chaos/PBDCollisionConstraints.cpp
@@ -417,6 +417,24 @@ namespace Chaos
 		return ConstraintAllocator.GetConstConstraints();
 	}
 
+	void UpdateSoftCollisionSettings(const FChaosPhysicsMaterial* PhysicsMaterial, const FGeometryParticleHandle* Particle, FReal& InOutThickess)
+	{
+		if ((PhysicsMaterial->SoftCollisionMode != EChaosPhysicsMaterialSoftCollisionMode::None) && (PhysicsMaterial->SoftCollisionThickness > 0))
+		{
+			FReal MaterialThickness = FReal(0);
+			if (PhysicsMaterial->SoftCollisionMode == EChaosPhysicsMaterialSoftCollisionMode::AbsoluteThickness)
+			{
+				MaterialThickness = PhysicsMaterial->SoftCollisionThickness;
+			}
+			else if ((Particle != nullptr) && Particle->HasBounds())
+			{
+				MaterialThickness = FMath::Clamp(PhysicsMaterial->SoftCollisionThickness, 0.0f, 0.5f) * Particle->LocalBounds().Extents().GetAbsMin();
+			}
+
+			InOutThickess = FMath::Max(InOutThickess, MaterialThickness);
+		}
+	}
+
 	void FPBDCollisionConstraints::UpdateConstraintMaterialProperties(FPBDCollisionConstraint& Constraint)
 	{
 		// We only support one material shared by all manifold points for now, even when our 
@@ -435,6 +453,9 @@ namespace Chaos
 		FReal MaterialRestitutionThreshold = 0;
 		FReal MaterialStaticFriction = 0;
 		FReal MaterialDynamicFriction = 0;
+		FReal MaterialSoftThickness = 0;
+		FReal MaterialSoftDivisor = 0;
+		FReal MaterialBaseFrictionImpulse = 0;
 
 		if (PhysicsMaterial0 && PhysicsMaterial1)
 		{
@@ -446,6 +467,13 @@ namespace Chaos
 			const FReal StaticFriction0 = FMath::Max(PhysicsMaterial0->Friction, PhysicsMaterial0->StaticFriction);
 			const FReal StaticFriction1 = FMath::Max(PhysicsMaterial1->Friction, PhysicsMaterial1->StaticFriction);
 			MaterialStaticFriction = FChaosPhysicsMaterial::CombineHelper(StaticFriction0, StaticFriction1, FrictionCombineMode);
+
+			// @todo(chaos): could do with a nicer way to deal with collisions between two soft objects with different softness settings
+			UpdateSoftCollisionSettings(PhysicsMaterial0, Constraint.GetParticle0(), MaterialSoftThickness);
+			UpdateSoftCollisionSettings(PhysicsMaterial1, Constraint.GetParticle1(), MaterialSoftThickness);
+
+			// Combine base friction impulse
+			MaterialBaseFrictionImpulse = FChaosPhysicsMaterial::CombineHelper(PhysicsMaterial0->BaseFrictionImpulse, PhysicsMaterial1->BaseFrictionImpulse, FrictionCombineMode);
 		}
 		else if (PhysicsMaterial0)
 		{
@@ -453,6 +481,8 @@ namespace Chaos
 			MaterialRestitution = PhysicsMaterial0->Restitution;
 			MaterialDynamicFriction = PhysicsMaterial0->Friction;
 			MaterialStaticFriction = StaticFriction0;
+			MaterialBaseFrictionImpulse = PhysicsMaterial0->BaseFrictionImpulse;
+			UpdateSoftCollisionSettings(PhysicsMaterial0, Constraint.GetParticle0(), MaterialSoftThickness);
 		}
 		else if (PhysicsMaterial1)
 		{
@@ -460,6 +490,8 @@ namespace Chaos
 			MaterialRestitution = PhysicsMaterial1->Restitution;
 			MaterialDynamicFriction = PhysicsMaterial1->Friction;
 			MaterialStaticFriction = StaticFriction1;
+			MaterialBaseFrictionImpulse = PhysicsMaterial1->BaseFrictionImpulse;
+			UpdateSoftCollisionSettings(PhysicsMaterial1, Constraint.GetParticle1(), MaterialSoftThickness);
 		}
 		else
 		{
@@ -493,12 +525,16 @@ namespace Chaos
 			MaterialRestitution = 0.0f;
 		}
 		
-		Constraint.Material.MaterialRestitution = FRealSingle(MaterialRestitution);
+		Constraint.Material.Restitution = FRealSingle(MaterialRestitution);
 		Constraint.Material.RestitutionThreshold = FRealSingle(MaterialRestitutionThreshold);
-		Constraint.Material.MaterialStaticFriction = FRealSingle(MaterialStaticFriction);
-		Constraint.Material.MaterialDynamicFriction = FRealSingle(MaterialDynamicFriction);
-
-		Constraint.Material.ResetMaterialModifications();
+		Constraint.Material.StaticFriction = FRealSingle(MaterialStaticFriction);
+		Constraint.Material.DynamicFriction = FRealSingle(MaterialDynamicFriction);
+		Constraint.Material.SoftSeparation = FRealSingle(-MaterialSoftThickness);	// Negate: convert from penetration to separation
+		Constraint.Material.InvMassScale0 = 1;
+		Constraint.Material.InvMassScale1 = 1;
+		Constraint.Material.InvInertiaScale0 = 1;
+		Constraint.Material.InvInertiaScale1 = 1;
+		Constraint.Material.BaseFrictionImpulse = FRealSingle(MaterialBaseFrictionImpulse);
 	}
 
 	void FPBDCollisionConstraints::BeginFrame()
diff --git a/Engine/Source/Runtime/Experimental/Chaos/Private/Chaos/PBDSuspensionConstraints.cpp b/Engine/Source/Runtime/Experimental/Chaos/Private/Chaos/PBDSuspensionConstraints.cpp
--- a/Engine/Source/Runtime/Experimental/Chaos/Private/Chaos/PBDSuspensionConstraints.cpp
+++ b/Engine/Source/Runtime/Experimental/Chaos/Private/Chaos/PBDSuspensionConstraints.cpp
@@ -217,6 +217,7 @@ namespace Chaos
 			FPBDCollisionSolver* Solver = &CollisionSolvers[ConstraintIndex];
 			Solver->Reset(&CollisionSolverManifoldPoints[ConstraintIndex], 1);
 			Solver->SetStiffness(1);
+			Solver->SetHardContact();
 
 			FSolverBody* Body0 = ConstraintSolverBodies[ConstraintIndex];	// vehicle chassis			
 			FSolverBody* Body1 = &StaticCollisionBodies[ConstraintIndex];	// Spoofed terrain
@@ -257,7 +258,7 @@ namespace Chaos
 				Solver->SetSolverBodies(*Body0, *Body1);
 				Solver->SolverBody0().Init();
 				Solver->SolverBody1().Init();
-				Solver->SetFriction(0, 0, 0);
+				Solver->SetFriction(0, 0, 0, 0);
 
 #if CHAOS_DEBUG_DRAW
 				if (bChaos_Suspension_DebugDraw_Hardstop)
diff --git a/Engine/Source/Runtime/Experimental/Chaos/Private/ChaosVisualDebugger/ChaosVDDataWrapperUtils.cpp b/Engine/Source/Runtime/Experimental/Chaos/Private/ChaosVisualDebugger/ChaosVDDataWrapperUtils.cpp
--- a/Engine/Source/Runtime/Experimental/Chaos/Private/ChaosVisualDebugger/ChaosVDDataWrapperUtils.cpp
+++ b/Engine/Source/Runtime/Experimental/Chaos/Private/ChaosVisualDebugger/ChaosVDDataWrapperUtils.cpp
@@ -58,9 +58,6 @@ void FChaosVDDataWrapperUtils::CopyManifoldPointResultsToDataWrapper(const Chaos
 void FChaosVDDataWrapperUtils::CopyCollisionMaterialToDataWrapper(const Chaos::FPBDCollisionConstraintMaterial& InCopyFrom, FChaosVDCollisionMaterial& OutCopyTo)
 {
 	OutCopyTo.FaceIndex = InCopyFrom.FaceIndex;
-	OutCopyTo.MaterialDynamicFriction = InCopyFrom.MaterialDynamicFriction;
-	OutCopyTo.MaterialStaticFriction = InCopyFrom.MaterialStaticFriction;
-	OutCopyTo.MaterialRestitution = InCopyFrom.MaterialRestitution;
 	OutCopyTo.DynamicFriction = InCopyFrom.DynamicFriction;
 	OutCopyTo.StaticFriction = InCopyFrom.StaticFriction;
 	OutCopyTo.Restitution = InCopyFrom.Restitution;
diff --git a/Engine/Source/Runtime/Experimental/Chaos/Public/Chaos/Collision/PBDCollisionConstraint.h b/Engine/Source/Runtime/Experimental/Chaos/Public/Chaos/Collision/PBDCollisionConstraint.h
--- a/Engine/Source/Runtime/Experimental/Chaos/Public/Chaos/Collision/PBDCollisionConstraint.h
+++ b/Engine/Source/Runtime/Experimental/Chaos/Public/Chaos/Collision/PBDCollisionConstraint.h
@@ -49,9 +49,6 @@ namespace Chaos
 	public:
 		FPBDCollisionConstraintMaterial()
 			: FaceIndex(INDEX_NONE)
-			, MaterialDynamicFriction(0)
-			, MaterialStaticFriction(0)
-			, MaterialRestitution(0)
 			, DynamicFriction(0)
 			, StaticFriction(0)
 			, Restitution(0)
@@ -60,40 +57,38 @@ namespace Chaos
 			, InvMassScale1(1)
 			, InvInertiaScale0(1)
 			, InvInertiaScale1(1)
+			, SoftSeparation(0)
+			, BaseFrictionImpulse(0)
 		{
 		}
 
 		// The face index that the material was extracted from
 		int32 FaceIndex;
 
-		// Material properties pulled from the materials of the two shapes involved in the contact
-		// @todo(chaos): we can remove these now and just recollect the material data when a modifier was applied
-		FRealSingle MaterialDynamicFriction;
-		FRealSingle MaterialStaticFriction;
-		FRealSingle MaterialRestitution;
-
-		// Final material properties (post modifier) used by the solver. These get reset every frame to the material values above
+		// Final material properties (post modifier) used by the solver
+		// These get reset every frame to the material values (if modified)
 		FRealSingle DynamicFriction;
 		FRealSingle StaticFriction;
 		FRealSingle Restitution;
-
 		FRealSingle RestitutionThreshold;
 		FRealSingle InvMassScale0;
 		FRealSingle InvMassScale1;
 		FRealSingle InvInertiaScale0;
 		FRealSingle InvInertiaScale1;
 
-		// Reset the material properties to those pulled from the shape's materials (i.e., back to the state before any contact modification)
-		void ResetMaterialModifications()
-		{
-			DynamicFriction = MaterialDynamicFriction;
-			StaticFriction = MaterialStaticFriction;
-			Restitution = MaterialRestitution;
-			InvMassScale0 = FRealSingle(1);
-			InvMassScale1 = FRealSingle(1);
-			InvInertiaScale0 = FRealSingle(1);
-			InvInertiaScale1 = FRealSingle(1);
-		}
+		UE_DEPRECATED(5.5, "Use DynamicFriction instead")
+		FRealSingle MaterialDynamicFriction;
+		UE_DEPRECATED(5.5, "Use StaticFriction instead")
+		FRealSingle MaterialStaticFriction;
+		UE_DEPRECATED(5.5, "Use Restitution instead")
+		FRealSingle MaterialRestitution;
+
+		// SoftSeparation is usually negative (0 for disabled), and indicates the depth at which the hard collision starts
+		FRealSingle SoftSeparation;
+
+		// BaseFrictionThickness is usually positive (0 for disabled) and is the distance away from the core shape
+		// at which the base friction impulse can be applied.
+		FRealSingle BaseFrictionImpulse;
 	};
 
 	// Renamed to FPBDCollisionConstraintMaterial
@@ -415,6 +410,7 @@ namespace Chaos
 
 		void SetModifierApplied() { Flags.bModifierApplied = true; }
 
+		UE_DEPRECATED(5.5, "Use specific material property getters instead")
 		const FPBDCollisionConstraintMaterial& GetCollisionMaterial() const { return Material; }
 
 		void SetInvMassScale0(const FReal InInvMassScale) { Material.InvMassScale0 = FRealSingle(InInvMassScale); }
@@ -444,6 +440,14 @@ namespace Chaos
 		void SetDynamicFriction(const FReal InDynamicFriction) { Material.DynamicFriction = FRealSingle(InDynamicFriction); }
 		FReal GetDynamicFriction() const { return Material.DynamicFriction; }
 
+		void SetMinFrictionPushOut(const FReal InMinFrictionPushOut) { Material.BaseFrictionImpulse = FRealSingle(InMinFrictionPushOut); }
+		FReal GetMinFrictionPushOut() const { return Material.BaseFrictionImpulse; }
+
+		bool IsSoftContact() const { return (Material.SoftSeparation != 0); }
+
+		void SetSoftSeparation(const FReal InSoftSeparation) { Material.SoftSeparation = FRealSingle(InSoftSeparation); }
+		FRealSingle GetSoftSeparation() const { return Material.SoftSeparation; }
+
 		EContactShapesType GetShapesType() const { return ShapesType; }
 
 		CHAOS_API FString ToString() const;
@@ -471,12 +475,9 @@ namespace Chaos
 		{
 			if (Flags.bModifierApplied)
 			{
-				Material.ResetMaterialModifications();
-
-				if (Flags.bIsOneWayInteraction)
-				{
-					UpdateMassScales();
-				}
+				// Re-gather the material properties
+				ClearMaterialProperties();
+				UpdateMaterialProperties();
 
 				// Reset other properties which may have changed in contact modification last frame
 				Flags.bIsProbe = BoundsTestFlags.bIsProbe;
diff --git a/Engine/Source/Runtime/Experimental/Chaos/Public/Chaos/Collision/PBDCollisionSolver.h b/Engine/Source/Runtime/Experimental/Chaos/Public/Chaos/Collision/PBDCollisionSolver.h
--- a/Engine/Source/Runtime/Experimental/Chaos/Public/Chaos/Collision/PBDCollisionSolver.h
+++ b/Engine/Source/Runtime/Experimental/Chaos/Public/Chaos/Collision/PBDCollisionSolver.h
@@ -22,6 +22,7 @@ namespace Chaos
 		extern bool bChaos_PBDCollisionSolver_Velocity_FrictionEnabled;
 		extern float Chaos_PBDCollisionSolver_Position_StaticFrictionStiffness;
 		extern float Chaos_PBDCollisionSolver_Velocity_StaticFrictionStiffness;
+		extern bool bChaos_PBDCollisionSolver_Velocity_MinMaxFriction;
 	}
 
 	namespace Private
@@ -70,6 +71,7 @@ namespace Chaos
 			FSolverReal NetImpulseNormal;
 			FSolverReal NetImpulseTangentU;
 			FSolverReal NetImpulseTangentV;
+			FSolverReal NetSoftPushOutNormal;
 
 			// A measure of how much we exceeded the static friction threshold.
 			// Equal to (NormalPushOut / TangentialPushOut) before clamping to the friction cone.
@@ -127,11 +129,12 @@ namespace Chaos
 			FSolverReal GetDynamicFriction() const { return State.DynamicFriction; }
 			FSolverReal GetVelocityFriction() const { return State.VelocityFriction; }
 
-			void SetFriction(const FSolverReal InStaticFriction, const FSolverReal InDynamicFriction, const FSolverReal InVelocityFriction)
+			void SetFriction(const FSolverReal InStaticFriction, const FSolverReal InDynamicFriction, const FSolverReal InVelocityFriction, const FSolverReal InMinMaxFrictionPushOut)
 			{
 				State.StaticFriction = InStaticFriction;
 				State.DynamicFriction = InDynamicFriction;
 				State.VelocityFriction = InVelocityFriction;
+				State.MinMaxFrictionPushout = InMinMaxFrictionPushOut;
 			}
 
 			void SetStiffness(const FSolverReal InStiffness)
@@ -139,6 +142,16 @@ namespace Chaos
 				State.Stiffness = InStiffness;
 			}
 
+			void SetHardContact()
+			{
+				State.SoftPhi = 0;
+			}
+
+			void SetSoftContact(const FSolverReal SoftPhi)
+			{
+				State.SoftPhi = SoftPhi;
+			}
+
 			void SetSolverBodies(FSolverBody& SolverBody0, FSolverBody& SolverBody1)
 			{
 				State.SolverBodies[0].SetSolverBody(SolverBody0);
@@ -218,7 +231,7 @@ namespace Chaos
 				const FSolverReal MaxPushOut);
 
 			void SolvePositionWithFriction(
-				const FSolverReal Dt, 
+				const FSolverReal Dt,
 				const FSolverReal MaxPushOut);
 
 			/**
@@ -226,7 +239,7 @@ namespace Chaos
 			 * @return true if we need to run more iterations, false if we did not apply any correction
 			*/
 			void SolveVelocity(
-				const FSolverReal Dt, 
+				const FSolverReal Dt,
 				const bool bApplyDynamicFriction);
 
 			void UpdateMass();
@@ -237,8 +250,8 @@ namespace Chaos
 			// Get the mass properties for the bodies if they are dynamic
 			FORCEINLINE_DEBUGGABLE void GetDynamicMassProperties(
 				FSolverReal& OutInvM0,
-				FSolverMatrix33& OutInvI0, 
-				FSolverReal& OutInvM1, 
+				FSolverMatrix33& OutInvI0,
+				FSolverReal& OutInvM1,
 				FSolverMatrix33& OutInvI1)
 			{
 				OutInvM0 = 0;
@@ -264,26 +277,26 @@ namespace Chaos
 				return (State.InvMs[BodyIndex] > FSolverBody::ZeroMassThreshold());
 			}
 
-			void CalculateContactPositionErrorNormal(
-				const FPBDCollisionSolverManifoldPoint& ManifoldPoint, 
-				const FSolverReal MaxPushOut, 
+			void CalculateContactPositionCorrectionNormal(
+				const FPBDCollisionSolverManifoldPoint& ManifoldPoint,
+				const FSolverReal MaxPushOut,
 				FSolverReal& OutContactDeltaNormal) const;
 
 			void CalculateContactPositionErrorTangential(
-				const FPBDCollisionSolverManifoldPoint& ManifoldPoint, 
-				FSolverReal& OutContactDeltaTangentU, 
+				const FPBDCollisionSolverManifoldPoint& ManifoldPoint,
+				FSolverReal& OutContactDeltaTangentU,
 				FSolverReal& OutContactDeltaTangentV) const;
 
 			void CalculateContactVelocityError(
-				const FPBDCollisionSolverManifoldPoint& ManifoldPoint, 
-				const FSolverReal DynamicFriction, 
-				const FSolverReal Dt, 
-				FSolverReal& OutContactVelocityDeltaNormal, 
-				FSolverReal& OutContactVelocityDeltaTangent0, 
+				const FPBDCollisionSolverManifoldPoint& ManifoldPoint,
+				const FSolverReal DynamicFriction,
+				const FSolverReal Dt,
+				FSolverReal& OutContactVelocityDeltaNormal,
+				FSolverReal& OutContactVelocityDeltaTangent0,
 				FSolverReal& OutContactVelocityDeltaTangent1) const;
 
 			void CalculateContactVelocityErrorNormal(
-				const FPBDCollisionSolverManifoldPoint& ManifoldPoint, 
+				const FPBDCollisionSolverManifoldPoint& ManifoldPoint,
 				FSolverReal& OutContactVelocityDeltaNormal) const;
 
 			bool ShouldSolveVelocity(
@@ -327,6 +340,12 @@ namespace Chaos
 				const FSolverReal ContactDeltaNormal,
 				FPBDCollisionSolverManifoldPoint& ManifoldPoint);
 
+			void ApplySoftPositionCorrectionNormal(
+				const FSolverReal Stiffness,
+				const FSolverReal ContactDeltaNormal,
+				const FSolverReal ContactVelocityNormalDt,
+				FPBDCollisionSolverManifoldPoint& ManifoldPoint);
+
 			void ApplyVelocityCorrection(
 				const FSolverReal Stiffness,
 				const FSolverReal Dt,
@@ -368,6 +387,8 @@ namespace Chaos
 					ManifoldPoints = nullptr;
 					NumManifoldPoints = 0;
 					MaxManifoldPoints = 0;
+					SoftPhi = 0;
+					MinMaxFrictionPushout = 0;
 				}
 
 				// Static Friction in the position-solve phase
@@ -379,9 +400,17 @@ namespace Chaos
 				// Dynamic Friction in the velocity-solve phase
 				FSolverReal VelocityFriction;
 
+				// A min clamp on the max friction pushout - essentially the minimum
+				// friction position impulse that is always available to this contact
+				// to counteract lateral motion
+				FSolverReal MinMaxFrictionPushout;
+
 				// Solver stiffness (scales all pushout and impulses)
 				FSolverReal Stiffness;
 
+				// Soft contact penetration
+				FSolverReal SoftPhi;
+
 				// Bodies
 				FConstraintSolverBody SolverBodies[MaxConstrainedBodies];
 				FSolverReal InvMs[2];
@@ -530,7 +559,7 @@ namespace Chaos
 			}
 			if (IsDynamic(1))
 			{
-				const FSolverVec3 DX1 = -State.InvMs[1]* PushOut;
+				const FSolverVec3 DX1 = -State.InvMs[1] * PushOut;
 				const FSolverVec3 DR1 = ManifoldPoint.ContactTangentUAngular1 * -PushOutTangentU + ManifoldPoint.ContactTangentVAngular1 * -PushOutTangentV;
 				FConstraintSolverBody& Body1 = SolverBody1();
 				Body1.ApplyPositionDelta(DX1);
@@ -573,6 +602,14 @@ namespace Chaos
 			}
 		}
 
+		FORCEINLINE_DEBUGGABLE void FPBDCollisionSolver::ApplySoftPositionCorrectionNormal(
+			const FSolverReal Stiffness,
+			const FSolverReal ContactDeltaNormal,
+			const FSolverReal ContactVelocityNormalDt,
+			FPBDCollisionSolverManifoldPoint& ManifoldPoint)
+		{
+			// TODO: Apply soft pushout correction
+		}
 
 		FORCEINLINE_DEBUGGABLE void FPBDCollisionSolver::ApplyVelocityCorrection(
 			const FSolverReal Stiffness,
@@ -608,7 +645,13 @@ namespace Chaos
 				ImpulseTangentU = -FrictionStiffness * ManifoldPoint.ContactMassTangentU * ContactVelocityDeltaTangent0;
 				ImpulseTangentV = -FrictionStiffness * ManifoldPoint.ContactMassTangentV * ContactVelocityDeltaTangent1;
 
-				const FSolverReal MaxImpulseTangent = FMath::Max(FSolverReal(0), DynamicFriction * (ManifoldPoint.NetImpulseNormal + ImpulseNormal + ManifoldPoint.NetPushOutNormal / Dt));
+				const FSolverReal TotalImpulseNormal
+					= ManifoldPoint.NetImpulseNormal
+					+ ManifoldPoint.NetSoftPushOutNormal
+					+ ManifoldPoint.NetPushOutNormal
+					+ ImpulseNormal;
+
+				const FSolverReal MaxImpulseTangent = DynamicFriction * FMath::Max(State.MinMaxFrictionPushout, TotalImpulseNormal) / Dt;
 				const FSolverReal MaxImpulseTangentSq = FMath::Square(MaxImpulseTangent);
 				const FSolverReal ImpulseTangentSq = FMath::Square(ImpulseTangentU) + FMath::Square(ImpulseTangentV);
 				if (ImpulseTangentSq > (MaxImpulseTangentSq + UE_SMALL_NUMBER))
@@ -634,7 +677,7 @@ namespace Chaos
 			}
 			if (IsDynamic(1))
 			{
-				const FSolverVec3 DV1 = -State.InvMs[1]* Impulse;
+				const FSolverVec3 DV1 = -State.InvMs[1] * Impulse;
 				const FSolverVec3 DW1 = ManifoldPoint.ContactNormalAngular1 * -ImpulseNormal + ManifoldPoint.ContactTangentUAngular1 * -ImpulseTangentU + ManifoldPoint.ContactTangentVAngular1 * -ImpulseTangentV;
 				FConstraintSolverBody& Body1 = SolverBody1();
 				Body1.ApplyVelocityDelta(DV1, DW1);
@@ -771,7 +814,7 @@ namespace Chaos
 			}
 		}
 
-		FORCEINLINE_DEBUGGABLE void FPBDCollisionSolver::CalculateContactPositionErrorNormal(
+		FORCEINLINE_DEBUGGABLE void FPBDCollisionSolver::CalculateContactPositionCorrectionNormal(
 			const FPBDCollisionSolverManifoldPoint& ManifoldPoint,
 			const FSolverReal MaxPushOut,
 			FSolverReal& OutContactDeltaNormal) const
@@ -780,7 +823,7 @@ namespace Chaos
 			const FConstraintSolverBody& Body1 = SolverBody1();
 
 			// Linear version: calculate the contact delta assuming linear motion after applying a positional impulse at the contact point. There will be an error that depends on the size of the rotation.
-			FSolverReal ContactDelta = ManifoldPoint.ContactDeltaNormal;
+			FSolverReal ContactDelta = 0;
 			if (IsDynamic(0) && IsDynamic(1))
 			{
 				ContactDelta += FSolverVec3::DotProduct(Body0.DP() - Body1.DP(), ManifoldPoint.ContactNormal);
@@ -824,7 +867,7 @@ namespace Chaos
 				ContactDeltaU += FSolverVec3::DotProduct(Body0.DP() - Body1.DP(), ManifoldPoint.ContactTangentU);
 				ContactDeltaU += FSolverVec3::DotProduct(Body0.DQ(), ManifoldPoint.ContactRxTangentU0);
 				ContactDeltaU -= FSolverVec3::DotProduct(Body1.DQ(), ManifoldPoint.ContactRxTangentU1);
-				
+
 				ContactDeltaV += FSolverVec3::DotProduct(Body0.DP() - Body1.DP(), ManifoldPoint.ContactTangentV);
 				ContactDeltaV += FSolverVec3::DotProduct(Body0.DQ(), ManifoldPoint.ContactRxTangentV0);
 				ContactDeltaV -= FSolverVec3::DotProduct(Body1.DQ(), ManifoldPoint.ContactRxTangentV1);
@@ -902,7 +945,7 @@ namespace Chaos
 			const FPBDCollisionSolverManifoldPoint& ManifoldPoint) const
 		{
 			// We ensure positive separating velocity for close contacts even if they didn't receive a pushout
-			return (ManifoldPoint.NetPushOutNormal > FSolverReal(0)) || (ManifoldPoint.ContactDeltaNormal < FSolverReal(0));
+			return (ManifoldPoint.NetPushOutNormal > FSolverReal(0)) || (ManifoldPoint.NetSoftPushOutNormal > FSolverReal(0)) || (ManifoldPoint.ContactDeltaNormal < State.SoftPhi);
 		}
 
 		FORCEINLINE_DEBUGGABLE void FPBDCollisionSolver::InitManifoldPoint(
@@ -958,48 +1001,34 @@ namespace Chaos
 			const FSolverReal Dt,
 			const FSolverReal MaxPushOut)
 		{
-			bool bApplyFriction = false;
 
 			// Apply the position correction along the normal and determine if we want to run friction on each point
-			for (int32 PointIndex = 0; PointIndex < NumManifoldPoints(); ++PointIndex)
-			{
-				FPBDCollisionSolverManifoldPoint& SolverManifoldPoint = State.ManifoldPoints[PointIndex];
-
-				FSolverReal ContactDeltaNormal;
-				CalculateContactPositionErrorNormal(SolverManifoldPoint, MaxPushOut, ContactDeltaNormal);
-
-				// Apply a normal correction if we still have penetration or if we are now separated but have previously applied a correction that we may want to undo
-				const bool bProcessManifoldPoint = (ContactDeltaNormal < FSolverReal(0)) || (SolverManifoldPoint.NetPushOutNormal > FSolverReal(UE_SMALL_NUMBER));
-				if (bProcessManifoldPoint)
-				{
-					ApplyPositionCorrectionNormal(
-						State.Stiffness,
-						ContactDeltaNormal,
-						SolverManifoldPoint);
-				}
-
-				// Friction gets updated for any point with a net normal correction or where we have previously had a normal correction and 
-				// already applied friction (in which case we may need to zero it if the normal correction gets removed)
-				const bool bHasAnyPushOut = (SolverManifoldPoint.NetPushOutNormal > 0) || (SolverManifoldPoint.NetPushOutTangentU != 0) || (SolverManifoldPoint.NetPushOutTangentV != 0);
-				SolverManifoldPoint.bApplyFriction = bHasAnyPushOut;
-				bApplyFriction = bApplyFriction || bHasAnyPushOut;
-			}
+			SolvePositionNoFriction(Dt, MaxPushOut);
 
 			// Apply the tangential position correction if required
-			if (bApplyFriction)
+			const FSolverReal FrictionStiffness = State.Stiffness * CVars::Chaos_PBDCollisionSolver_Position_StaticFrictionStiffness;
+			if (FrictionStiffness > 0)
 			{
-				const FSolverReal FrictionStiffness = State.Stiffness * CVars::Chaos_PBDCollisionSolver_Position_StaticFrictionStiffness;
-
 				for (int32 PointIndex = 0; PointIndex < NumManifoldPoints(); ++PointIndex)
 				{
 					FPBDCollisionSolverManifoldPoint& SolverManifoldPoint = State.ManifoldPoints[PointIndex];
-					if (SolverManifoldPoint.bApplyFriction)
+
+					const FSolverReal TotalPushOutNormal
+						= SolverManifoldPoint.NetPushOutNormal
+						+ SolverManifoldPoint.NetSoftPushOutNormal;
+
+					const FSolverReal FrictionMaxPushOut
+						= CVars::bChaos_PBDCollisionSolver_Velocity_MinMaxFriction
+						? FMath::Max(State.MinMaxFrictionPushout, TotalPushOutNormal)
+						: TotalPushOutNormal;
+
+					const bool bApplyFriction = (FrictionMaxPushOut > 0) || (SolverManifoldPoint.NetPushOutTangentU != 0) || (SolverManifoldPoint.NetPushOutTangentV != 0);
+
+					if (bApplyFriction)
 					{
 						FSolverReal ContactDeltaTangentU, ContactDeltaTangentV;
 						CalculateContactPositionErrorTangential(SolverManifoldPoint, ContactDeltaTangentU, ContactDeltaTangentV);
 
-						const FSolverReal FrictionMaxPushOut = SolverManifoldPoint.NetPushOutNormal;
-
 						ApplyPositionCorrectionTangential(
 							FrictionStiffness,
 							State.StaticFriction,
@@ -1022,17 +1051,24 @@ namespace Chaos
 			{
 				FPBDCollisionSolverManifoldPoint& SolverManifoldPoint = State.ManifoldPoints[PointIndex];
 
-				FSolverReal ContactDeltaNormal;
-				CalculateContactPositionErrorNormal(SolverManifoldPoint, MaxPushOut, ContactDeltaNormal);
+				FSolverReal ContactCorrectionNormal;
+				CalculateContactPositionCorrectionNormal(SolverManifoldPoint, MaxPushOut, ContactCorrectionNormal);
+
+				// Shift the contact to the hard shell under the soft layer
+				const FSolverReal HardContactErrorNormal = SolverManifoldPoint.ContactDeltaNormal + ContactCorrectionNormal - State.SoftPhi;
 
-				const bool bProcessManifoldPoint = (ContactDeltaNormal < FSolverReal(0)) || (SolverManifoldPoint.NetPushOutNormal > FSolverReal(UE_SMALL_NUMBER));
+				const bool bProcessManifoldPoint = (HardContactErrorNormal < FSolverReal(0)) || (SolverManifoldPoint.NetPushOutNormal > FSolverReal(UE_SMALL_NUMBER));
 				if (bProcessManifoldPoint)
 				{
 					ApplyPositionCorrectionNormal(
 						State.Stiffness,
-						ContactDeltaNormal,
+						HardContactErrorNormal,
 						SolverManifoldPoint);
 				}
+
+				//
+				// TODO: when State.SoftPhi < 0, apply soft position correction
+				//
 			}
 		}
 
@@ -1061,15 +1097,24 @@ namespace Chaos
 			{
 				FPBDCollisionSolverManifoldPoint& SolverManifoldPoint = State.ManifoldPoints[PointIndex];
 
-				if (ShouldSolveVelocity(SolverManifoldPoint))
+				const bool bShouldSolveNormalVelocity = (SolverManifoldPoint.NetPushOutNormal > FSolverReal(0));
+				const bool bShouldSolveTangentVelocity = (DynamicFriction > 0) && (bShouldSolveNormalVelocity || (SolverManifoldPoint.ContactDeltaNormal < State.SoftPhi) || (SolverManifoldPoint.NetSoftPushOutNormal > FSolverReal(0)));
+
+				if (bShouldSolveNormalVelocity || bShouldSolveTangentVelocity)
 				{
-					const FSolverReal MinImpulseNormal = FMath::Min(FSolverReal(0), -SolverManifoldPoint.NetPushOutNormal / Dt);
+					const FSolverReal MinImpulseNormal = FMath::Min(FSolverReal(0), -(SolverManifoldPoint.NetPushOutNormal + SolverManifoldPoint.NetSoftPushOutNormal) / Dt);
 
-					if (DynamicFriction > 0)
+					// @todo(chaos): clean this up - friction and no-friction versions should share solve-normal code
+					if (bShouldSolveTangentVelocity)
 					{
 						FSolverReal ContactVelocityDeltaNormal, ContactVelocityDeltaTangentU, ContactVelocityDeltaTangentV;
 						CalculateContactVelocityError(SolverManifoldPoint, DynamicFriction, Dt, ContactVelocityDeltaNormal, ContactVelocityDeltaTangentU, ContactVelocityDeltaTangentV);
 
+						if (!bShouldSolveNormalVelocity)
+						{
+							ContactVelocityDeltaNormal = 0;
+						}
+
 						ApplyVelocityCorrection(
 							State.Stiffness,
 							Dt,
@@ -1080,7 +1125,7 @@ namespace Chaos
 							MinImpulseNormal,
 							SolverManifoldPoint);
 					}
-					else
+					else if (bShouldSolveNormalVelocity)
 					{
 						FSolverReal ContactVelocityDeltaNormal;
 						CalculateContactVelocityErrorNormal(SolverManifoldPoint, ContactVelocityDeltaNormal);
diff --git a/Engine/Source/Runtime/Experimental/Chaos/Public/Chaos/Defines.h b/Engine/Source/Runtime/Experimental/Chaos/Public/Chaos/Defines.h
--- a/Engine/Source/Runtime/Experimental/Chaos/Public/Chaos/Defines.h
+++ b/Engine/Source/Runtime/Experimental/Chaos/Public/Chaos/Defines.h
@@ -73,6 +73,14 @@ namespace Chaos
 		float DamageThresholdMultiplier = 1.0f;
 	};
 
+	// NOTE: must match EPhysicalMaterialSoftCollisionMode
+	enum class EChaosPhysicsMaterialSoftCollisionMode : uint8
+	{
+		None,
+		RelativeThickness,
+		AbsoluteThickness,
+	};
+
 	class FChaosPhysicsMaterial
 	{
 	public:
@@ -101,6 +109,9 @@ namespace Chaos
 		int32 SleepCounterThreshold;
 		void* UserData;
 
+		/** Settings for "sticky" friction */
+		FRealSingle BaseFrictionImpulse;
+
 		/** Variable defaults are used for \c UChaosPhysicalMaterial \c UPROPERTY defaults. */
 		ECombineMode FrictionCombineMode;
 		ECombineMode RestitutionCombineMode;
@@ -108,6 +119,8 @@ namespace Chaos
 		FChaosPhysicsMaterialStrength Strength;
 		FChaosPhysicsMaterialDamageModifier DamageModifier;
 
+		EChaosPhysicsMaterialSoftCollisionMode SoftCollisionMode;
+		FRealSingle SoftCollisionThickness;
 
 		FChaosPhysicsMaterial()
 			: Friction(0.5)
@@ -122,8 +135,11 @@ namespace Chaos
 			, DisabledAngularThreshold(0)
 			, SleepCounterThreshold(0)
 			, UserData(nullptr)
+			, BaseFrictionImpulse(0.0)
 			, FrictionCombineMode(ECombineMode::Avg)
 			, RestitutionCombineMode(ECombineMode::Avg)
+			, SoftCollisionMode(EChaosPhysicsMaterialSoftCollisionMode::None)
+			, SoftCollisionThickness(0)
 		{
 		}
 
diff --git a/Engine/Source/Runtime/PhysicsCore/Private/ChaosEngineInterface.cpp b/Engine/Source/Runtime/PhysicsCore/Private/ChaosEngineInterface.cpp
--- a/Engine/Source/Runtime/PhysicsCore/Private/ChaosEngineInterface.cpp
+++ b/Engine/Source/Runtime/PhysicsCore/Private/ChaosEngineInterface.cpp
@@ -29,6 +29,7 @@ FPhysicsDelegatesCore::FOnUpdatePhysXMaterial FPhysicsDelegatesCore::OnUpdatePhy
 #include "PBDRigidsSolver.h"
 #include "PhysicsProxy/SingleParticlePhysicsProxy.h"
 #include "Chaos/CastingUtilities.h"
+#include "Math/UnitConversion.h"
 
 namespace Chaos
 {
@@ -244,6 +245,9 @@ void FChaosEngineInterface::UpdateMaterial(FPhysicsMaterialHandle& InHandle,UPhy
 		Material->Strength.CompressionStrength = Chaos::MegaPascalToKgPerCmS2(InMaterial->Strength.CompressionStrength);
 		Material->Strength.ShearStrength = Chaos::MegaPascalToKgPerCmS2(InMaterial->Strength.ShearStrength);
 		Material->DamageModifier.DamageThresholdMultiplier = InMaterial->DamageModifier.DamageThresholdMultiplier;
+		Material->BaseFrictionImpulse = InMaterial->BaseFrictionImpulse;
+		Material->SoftCollisionMode = (Chaos::EChaosPhysicsMaterialSoftCollisionMode)InMaterial->SoftCollisionMode;
+		Material->SoftCollisionThickness = InMaterial->SoftCollisionThickness;
 	}
 
 	Chaos::FPhysicalMaterialManager::Get().UpdateMaterial(InHandle);
diff --git a/Engine/Source/Runtime/PhysicsCore/Private/PhysicalMaterial.cpp b/Engine/Source/Runtime/PhysicsCore/Private/PhysicalMaterial.cpp
--- a/Engine/Source/Runtime/PhysicsCore/Private/PhysicalMaterial.cpp
+++ b/Engine/Source/Runtime/PhysicsCore/Private/PhysicalMaterial.cpp
@@ -8,9 +8,17 @@
 #include "PhysicalMaterials/PhysicalMaterialPropertyBase.h"
 #include "UObject/UObjectIterator.h"
 #include "Chaos/PhysicalMaterials.h"
+#include "HAL/IConsoleManager.h"
 
 #include UE_INLINE_GENERATED_CPP_BY_NAME(PhysicalMaterial)
 
+namespace CVars
+{
+	bool bPhysicalMaterial_ShowExperimentalProperties = false;
+
+	FAutoConsoleVariableRef CVarPhysicalMaterial_ShowExperimentalProperties(TEXT("p.PhysicalMaterial_ShowExperimentalProperties"), bPhysicalMaterial_ShowExperimentalProperties, TEXT(""));
+}
+
 UDEPRECATED_PhysicalMaterialPropertyBase::UDEPRECATED_PhysicalMaterialPropertyBase(const FObjectInitializer& ObjectInitializer)
 	: Super(ObjectInitializer)
 {
@@ -42,6 +50,11 @@ UPhysicalMaterial::UPhysicalMaterial(const FObjectInitializer& ObjectInitializer
 	SleepCounterThreshold = 4;
 	bOverrideFrictionCombineMode = false;
 	UserData = FChaosUserData(this);
+
+	SoftCollisionMode = EPhysicalMaterialSoftCollisionMode::None;
+	SoftCollisionThickness = 0;
+
+	BaseFrictionImpulse = 0;
 }
 
 UPhysicalMaterial::UPhysicalMaterial(FVTableHelper& Helper)
diff --git a/Engine/Source/Runtime/PhysicsCore/Public/PhysicalMaterials/PhysicalMaterial.h b/Engine/Source/Runtime/PhysicsCore/Public/PhysicalMaterials/PhysicalMaterial.h
--- a/Engine/Source/Runtime/PhysicsCore/Public/PhysicalMaterials/PhysicalMaterial.h
+++ b/Engine/Source/Runtime/PhysicsCore/Public/PhysicalMaterials/PhysicalMaterial.h
@@ -19,6 +19,12 @@
 
 struct FPropertyChangedEvent;
 
+namespace CVars
+{
+	extern bool bPhysicalMaterial_ShowExperimentalProperties;
+}
+
+
 /**
  * Defines the directional strengths of a physical material in term of force per surface area
  */
@@ -71,6 +77,25 @@ struct FPhysicalMaterialDamageModifier
 	float DamageThresholdMultiplier;
 };
 
+/**
+* Soft collision mode for a physical material.
+* 
+* NOTE: Must match EChaosPhysicsMaterialSoftCollisionMode
+*/
+UENUM()
+enum class EPhysicalMaterialSoftCollisionMode : uint8
+{
+	// No soft collisionss
+	None,
+
+	// SoftCollisionThickess is a fraction of the bounds (minimum axis). Should be less than 0.5
+	RelativeThickness,
+
+	// SoftCollisionThickess is an absolute value in cm
+	AbsoluteThickess,
+};
+
+
 /**
  * Physical materials are used to define the response of a physical object when interacting dynamically with the world.
  */
@@ -161,6 +186,30 @@ class UPhysicalMaterial : public UObject
 	UPROPERTY(EditAnywhere, BlueprintReadOnly, Category = PhysicalProperties)
 	FPhysicalMaterialDamageModifier DamageModifier;
 
+	//
+	// Experimental properties
+	//
+
+	/** Experimental material properties are enabled via the p.PhysicalMaterial_ShowExperimentalProperties console variable.
+		NOTE: These are _experimental_ properties which may change. Use at your own risk! */
+	UPROPERTY(VisibleAnywhere, Transient, Category = "Experimental")
+	bool bShowExperimentalProperties = CVars::bPhysicalMaterial_ShowExperimentalProperties;
+
+	/** For enable soft collision shell thickness mode */
+	UPROPERTY(EditAnywhere, Category = "Experimental|Softness", meta = (EditCondition = "bShowExperimentalProperties"))
+	EPhysicalMaterialSoftCollisionMode SoftCollisionMode;
+
+	/** Thickness of the layer just inside the collision shape in which contact is considered "soft".
+		The units depend on SoftCollisionMode */
+	UPROPERTY(EditAnywhere, Category = "Experimental|Softness", meta = (ClampMin = 0, EditCondition = "bShowExperimentalProperties"))
+	float SoftCollisionThickness;
+
+	/** A friction (positional) impulse of at least this magnitude may be applied,
+		regardless the normal force. This is analogous to adding only the lateral part of a
+		"stickiness" to a material */
+	UPROPERTY(EditAnywhere, Category = "Experimental|Stickiness", meta = (ClampMin = 0, Units = kgcm, EditCondition = "bShowExperimentalProperties"))
+	float BaseFrictionImpulse;
+
 public:
 
 	TUniquePtr<FPhysicsMaterialHandle> MaterialHandle;
