Chaos Cloth: self collision sets to disable collisions and enable collisions on kinematic cloth.

Added "SelfCollisionDisabledFaces" face set which will remove these faces from self collision calculations.

Added "SelfCollisionEnabledKinematicFaces" face set which enables collisions with kinematic cloth. Kinematic cloth only does dynamic vertex against kinematic faces, and only collides against the front faces. Different thickness, stiffness, and friction parameters are applied to these constraints. There are additional optimizations planned for this code path--this check in is for testing functionality of this feature.

#changelist validated
#virtualized

diff --git a/Engine/Plugins/Experimental/ChaosCloth/Source/Chaos/Private/ChaosCloth/ChaosClothConstraints.cpp b/Engine/Plugins/Experimental/ChaosCloth/Source/Chaos/Private/ChaosCloth/ChaosClothConstraints.cpp
--- a/Engine/Plugins/Experimental/ChaosCloth/Source/Chaos/Private/ChaosCloth/ChaosClothConstraints.cpp
+++ b/Engine/Plugins/Experimental/ChaosCloth/Source/Chaos/Private/ChaosCloth/ChaosClothConstraints.cpp
@@ -469,7 +469,7 @@ void FClothConstraints::AddRules(
 	}
 
 	// Call new AddRules function
-	AddRules(ConfigProperties, TriangleMesh, nullptr, WeightMaps, TMap<FString, const TSet<int32>*>(),
+	AddRules(ConfigProperties, TriangleMesh, nullptr, WeightMaps, TMap<FString, const TSet<int32>*>(), TMap<FString, const TSet<int32>*>(),
 		TMap<FString, TConstArrayView<int32>>(), Tethers, MeshScale, bEnabled);
 }
 
@@ -481,7 +481,7 @@ void FClothConstraints::AddRules(
 	const TArray<TConstArrayView<TTuple<int32, int32, FRealSingle>>>& Tethers,
 	Softs::FSolverReal MeshScale, bool bEnabled)
 {
-	AddRules(ConfigProperties, TriangleMesh, PatternData, WeightMaps, TMap<FString, const TSet<int32>*>(),
+	AddRules(ConfigProperties, TriangleMesh, PatternData, WeightMaps, TMap<FString, const TSet<int32>*>(), TMap<FString, const TSet<int32>*>(),
 		TMap<FString, TConstArrayView<int32>>(), Tethers, MeshScale, bEnabled);
 }
 
@@ -797,12 +797,13 @@ void FClothConstraints::AddRules(
 	const FClothingPatternData* PatternData,
 	const TMap<FString, TConstArrayView<FRealSingle>>& WeightMaps,
 	const TMap<FString, const TSet<int32>*>& VertexSets,
+	const TMap<FString, const TSet<int32>*>& FaceSets,
 	const TMap<FString, TConstArrayView<int32>>& FaceIntMaps,
 	const TArray<TConstArrayView<TTuple<int32, int32, FRealSingle>>>& Tethers,
 	Softs::FSolverReal MeshScale, bool bEnabled)
 {
 	// Self collisions
-	CreateSelfCollisionConstraints(ConfigProperties, VertexSets, FaceIntMaps, TriangleMesh);
+	CreateSelfCollisionConstraints(ConfigProperties, VertexSets, FaceSets, FaceIntMaps, TriangleMesh);
 
 	// Edge constraints
 	CreateStretchConstraints(ConfigProperties, WeightMaps, TriangleMesh, PatternData);
@@ -869,6 +870,7 @@ void FClothConstraints::AddRules(
 
 void FClothConstraints::CreateSelfCollisionConstraints(const Softs::FCollectionPropertyConstFacade& ConfigProperties, 
 	const TMap<FString, const TSet<int32>*>& VertexSets,
+	const TMap<FString, const TSet<int32>*>& FaceSets,
 	const TMap<FString, TConstArrayView<int32>>& FaceIntMaps, const FTriangleMesh& TriangleMesh)
 {
 	const bool bUseSelfCollisions = ConfigProperties.GetValue<bool>(TEXT("UseSelfCollisions"));
@@ -878,6 +880,7 @@ void FClothConstraints::CreateSelfCollisionConstraints(const Softs::FCollectionP
 		SelfCollisionInit = MakeShared<Softs::FPBDTriangleMeshCollisions>(
 			ParticleOffset,
 			NumParticles,
+			FaceSets,
 			TriangleMesh,
 			ConfigProperties);
 		++NumConstraintInits;
@@ -1592,7 +1595,7 @@ void FClothConstraints::CreateForceBasedRules()
 			}
 			// Thickness * 2 to account for collision radius for both particles
 			SelfCollisionInit->Init(Particles, SelfCollisionConstraints->GetThickness() * (Softs::FSolverReal)2.f);
-			SelfCollisionConstraints->Init(Particles, SelfCollisionInit->GetSpatialHash(), SelfCollisionInit->GetVertexGIAColors(), SelfCollisionInit->GetTriangleGIAColors());
+			SelfCollisionConstraints->Init(Particles, Dt, SelfCollisionInit->GetCollidableSubMesh(), SelfCollisionInit->GetSpatialHash(), SelfCollisionInit->GetVertexGIAColors(), SelfCollisionInit->GetTriangleGIAColors());
 		});
 
 		constexpr bool bPostCollisions = true;
@@ -1990,7 +1993,7 @@ void FClothConstraints::CreatePBDRules()
 	if (SelfCollisionInit && SelfCollisionConstraints)
 	{
 		ConstraintInits[ConstraintInitIndex++] =
-			[this](Softs::FSolverParticles& Particles, const Softs::FSolverReal /*Dt*/)
+			[this](Softs::FSolverParticles& Particles, const Softs::FSolverReal Dt)
 			{
 				if (bSkipSelfCollisionInit)
 				{
@@ -1998,7 +2001,7 @@ void FClothConstraints::CreatePBDRules()
 				}
 				// Thickness * 2 to account for collision radius for both particles
 				SelfCollisionInit->Init(Particles, SelfCollisionConstraints->GetThickness() * (Softs::FSolverReal)2.f);
-				SelfCollisionConstraints->Init(Particles, SelfCollisionInit->GetSpatialHash(), SelfCollisionInit->GetVertexGIAColors(), SelfCollisionInit->GetTriangleGIAColors());
+				SelfCollisionConstraints->Init(Particles, Dt, SelfCollisionInit->GetCollidableSubMesh(), SelfCollisionInit->GetSpatialHash(), SelfCollisionInit->GetVertexGIAColors(), SelfCollisionInit->GetTriangleGIAColors());
 			};
 
 		PostCollisionConstraintRules[PostCollisionConstraintRuleIndex++] =
@@ -2096,6 +2099,7 @@ void FClothConstraints::Update(
 	const Softs::FCollectionPropertyConstFacade& ConfigProperties,
 	const TMap<FString, TConstArrayView<FRealSingle>>& WeightMaps,
 	const TMap<FString, const TSet<int32>*>& VertexSets,
+	const TMap<FString, const TSet<int32>*>& FaceSets,
 	const TMap<FString, TConstArrayView<int32>>& FaceIntMaps,
 	Softs::FSolverReal MeshScale,
 	Softs::FSolverReal MaxDistancesScale)
@@ -2166,7 +2170,7 @@ void FClothConstraints::Update(
 	}
 	if (SelfCollisionInit)
 	{
-		SelfCollisionInit->SetProperties(ConfigProperties);
+		SelfCollisionInit->SetProperties(ConfigProperties, FaceSets);
 	}
 	if (SelfCollisionSphereConstraints)
 	{
@@ -2208,7 +2212,7 @@ void FClothConstraints::Update(
 	Softs::FSolverReal MeshScale,
 	Softs::FSolverReal MaxDistancesScale)
 {
-	Update(ConfigProperties, WeightMaps, TMap<FString, const TSet<int32>*>(), TMap<FString, TConstArrayView<int32>>(), MeshScale, MaxDistancesScale);
+	Update(ConfigProperties, WeightMaps, TMap<FString, const TSet<int32>*>(), TMap<FString, const TSet<int32>*>(), TMap<FString, TConstArrayView<int32>>(), MeshScale, MaxDistancesScale);
 }
 
 // Deprecated
@@ -2217,7 +2221,7 @@ void FClothConstraints::Update(
 	Softs::FSolverReal MeshScale,
 	Softs::FSolverReal MaxDistancesScale)
 {
-	Update(ConfigProperties, TMap<FString, TConstArrayView<FRealSingle>>(), TMap<FString, const TSet<int32>*>(), TMap<FString, TConstArrayView<int32>>(), MeshScale, MaxDistancesScale);
+	Update(ConfigProperties, TMap<FString, TConstArrayView<FRealSingle>>(), TMap<FString, const TSet<int32>*>(), TMap<FString, const TSet<int32>*>(), TMap<FString, TConstArrayView<int32>>(), MeshScale, MaxDistancesScale);
 }
 
 }  // End namespace Chaos
diff --git a/Engine/Plugins/Experimental/ChaosCloth/Source/Chaos/Private/ChaosCloth/ChaosClothingSimulationCloth.cpp b/Engine/Plugins/Experimental/ChaosCloth/Source/Chaos/Private/ChaosCloth/ChaosClothingSimulationCloth.cpp
--- a/Engine/Plugins/Experimental/ChaosCloth/Source/Chaos/Private/ChaosCloth/ChaosClothingSimulationCloth.cpp
+++ b/Engine/Plugins/Experimental/ChaosCloth/Source/Chaos/Private/ChaosCloth/ChaosClothingSimulationCloth.cpp
@@ -59,6 +59,7 @@ struct FClothingSimulationCloth::FLODData
 	const int32 NumParticles;
 	const TMap<FString, TConstArrayView<FRealSingle>> WeightMaps;
 	const TMap<FString, const TSet<int32>*> VertexSets;
+	const TMap<FString, const TSet<int32>*> FaceSets;
 	const TMap<FString, TConstArrayView<int32>> FaceIntMaps;
 	const TArray<TConstArrayView<TTuple<int32, int32, FRealSingle>>> Tethers;
 
@@ -86,6 +87,7 @@ struct FClothingSimulationCloth::FLODData
 		const TConstArrayView<uint32>& InPatternToWeldedIndices,
 		TMap<FString, TConstArrayView<FRealSingle>>&& InWeightMaps,
 		TMap<FString, const TSet<int32>*>&& InVertexSets,
+		TMap<FString, const TSet<int32>*>&& InFaceSets,
 		TMap<FString, TConstArrayView<int32>>&& InFaceIntMaps,
 		const TArray<TConstArrayView<TTuple<int32, int32, FRealSingle>>>& InTethers);
 
@@ -114,11 +116,13 @@ FClothingSimulationCloth::FLODData::FLODData(
 	const TConstArrayView<uint32>& InPatternToWeldedIndices,
 	TMap<FString, TConstArrayView<FRealSingle>>&& InWeightMaps,
 	TMap<FString, const TSet<int32>*>&& InVertexSets,
+	TMap<FString, const TSet<int32>*>&& InFaceSets,
 	TMap<FString, TConstArrayView<int32>>&& InFaceIntMaps,
 	const TArray<TConstArrayView<TTuple<int32, int32, FRealSingle>>>& InTethers)
 	: NumParticles(InNumParticles)
 	, WeightMaps(MoveTemp(InWeightMaps))
 	, VertexSets(MoveTemp(InVertexSets))
+	, FaceSets(MoveTemp(InFaceSets))
 	, FaceIntMaps(MoveTemp(InFaceIntMaps))
 	, Tethers(InTethers)
 	, PatternData(NumParticles, InIndices, InPatternPositions, InPatternIndices, InPatternToWeldedIndices)
@@ -254,7 +258,7 @@ PRAGMA_ENABLE_DEPRECATION_WARNINGS
 
 	// Create constraints
 	const bool bEnabled = false;  // Set constraint disabled by default
-	ClothConstraints.AddRules(ConfigProperties, TriangleMesh, &PatternData, WeightMaps, VertexSets, FaceIntMaps, Tethers, MeshScale, bEnabled);
+	ClothConstraints.AddRules(ConfigProperties, TriangleMesh, &PatternData, WeightMaps, VertexSets, FaceSets, FaceIntMaps, Tethers, MeshScale, bEnabled);
 
 	// Update LOD stats
 	const TConstArrayView<Softs::FSolverReal> InvMasses(Solver->GetParticleInvMasses(ParticleRangeId), NumParticles);
@@ -295,7 +299,7 @@ PRAGMA_DISABLE_DEPRECATION_WARNINGS  // TODO: CHAOS_IS_CLOTHINGSIMULATIONMESH_AB
 	const Softs::FSolverReal MeshScale = Cloth->Mesh->GetScale();
 PRAGMA_ENABLE_DEPRECATION_WARNINGS
 	const Softs::FSolverReal MaxDistancesScale = (Softs::FSolverReal)Cloth->MaxDistancesMultiplier;
-	ClothConstraints.Update(Cloth->Config->GetProperties(SolverDatum.LODIndex), WeightMaps, VertexSets, FaceIntMaps, MeshScale, MaxDistancesScale);
+	ClothConstraints.Update(Cloth->Config->GetProperties(SolverDatum.LODIndex), WeightMaps, VertexSets, FaceSets, FaceIntMaps, MeshScale, MaxDistancesScale);
 }
 
 void FClothingSimulationCloth::FLODData::Enable(FClothingSimulationSolver* Solver, bool bEnable) const
@@ -543,6 +547,7 @@ PRAGMA_ENABLE_DEPRECATION_WARNINGS
 	Properties.AddValue(TEXT("SelfCollisionThickness"), InSelfCollisionThickness);
 	Properties.AddValue(TEXT("SelfCollisionFrictionCoefficient"), InSelfCollisionFrictionCoefficient);
 	Properties.AddValue(TEXT("UseSelfIntersections"), bInUseSelfIntersections);
+	Properties.AddValue(TEXT("SelfCollideAgainstAllKinematicVertices"), true); // Match legacy behavior
 
 	// Max distance
 	{
@@ -609,6 +614,7 @@ void FClothingSimulationCloth::SetMesh(FClothingSimulationMesh* InMesh)
 		}
 
 		TMap<FString, const TSet<int32>*> VertexSets = Mesh->GetVertexSets(LODIndex);
+		TMap<FString, const TSet<int32>*> FaceSets = Mesh->GetFaceSets(LODIndex);
 		TMap<FString, TConstArrayView<int32>> FaceIntMaps = Mesh->GetFaceIntMaps(LODIndex);
 
 		const bool bUseGeodesicTethers = Config->GetProperties(LODIndex).GetValue<bool>(TEXT("UseGeodesicTethers"), ClothingSimulationClothDefault::bUseGeodesicTethers);
@@ -622,6 +628,7 @@ void FClothingSimulationCloth::SetMesh(FClothingSimulationMesh* InMesh)
 			Mesh->GetPatternToWeldedIndices(LODIndex),
 			MoveTemp(WeightMaps),
 			MoveTemp(VertexSets),
+			MoveTemp(FaceSets),
 			MoveTemp(FaceIntMaps),
 			Mesh->GetTethers(LODIndex, bUseGeodesicTethers)));
 	}
diff --git a/Engine/Plugins/Experimental/ChaosCloth/Source/Chaos/Private/ChaosCloth/ChaosClothingSimulationMesh.cpp b/Engine/Plugins/Experimental/ChaosCloth/Source/Chaos/Private/ChaosCloth/ChaosClothingSimulationMesh.cpp
--- a/Engine/Plugins/Experimental/ChaosCloth/Source/Chaos/Private/ChaosCloth/ChaosClothingSimulationMesh.cpp
+++ b/Engine/Plugins/Experimental/ChaosCloth/Source/Chaos/Private/ChaosCloth/ChaosClothingSimulationMesh.cpp
@@ -206,6 +206,12 @@ TMap<FString, const TSet<int32>*> FClothingSimulationMesh::GetVertexSets(int32 L
 	return TMap<FString, const TSet<int32>*>();
 }
 
+TMap<FString, const TSet<int32>*> FClothingSimulationMesh::GetFaceSets(int32 LODIndex) const
+{
+	// Not supported
+	return TMap<FString, const TSet<int32>*>();
+}
+
 TMap<FString, TConstArrayView<int32>> FClothingSimulationMesh::GetFaceIntMaps(int32 LODIndex) const
 {
 	// Not supported
diff --git a/Engine/Plugins/Experimental/ChaosCloth/Source/Chaos/Private/ChaosCloth/ChaosClothingSimulationSkeletalMesh.cpp b/Engine/Plugins/Experimental/ChaosCloth/Source/Chaos/Private/ChaosCloth/ChaosClothingSimulationSkeletalMesh.cpp
--- a/Engine/Plugins/Experimental/ChaosCloth/Source/Chaos/Private/ChaosCloth/ChaosClothingSimulationSkeletalMesh.cpp
+++ b/Engine/Plugins/Experimental/ChaosCloth/Source/Chaos/Private/ChaosCloth/ChaosClothingSimulationSkeletalMesh.cpp
@@ -173,6 +173,12 @@ PRAGMA_ENABLE_DEPRECATION_WARNINGS
 		return VertexSets;
 	}
 
+	TMap<FString, const TSet<int32>*> FClothingSimulationSkeletalMesh::GetFaceSets(int32 LODIndex) const
+	{
+		// Not supported
+		return TMap<FString, const TSet<int32>*>();
+	}
+
 	TMap<FString, TConstArrayView<int32>> FClothingSimulationSkeletalMesh::GetFaceIntMaps(int32 LODIndex) const
 	{
 		// Not supported
diff --git a/Engine/Plugins/Experimental/ChaosCloth/Source/Chaos/Public/ChaosCloth/ChaosClothConstraints.h b/Engine/Plugins/Experimental/ChaosCloth/Source/Chaos/Public/ChaosCloth/ChaosClothConstraints.h
--- a/Engine/Plugins/Experimental/ChaosCloth/Source/Chaos/Public/ChaosCloth/ChaosClothConstraints.h
+++ b/Engine/Plugins/Experimental/ChaosCloth/Source/Chaos/Public/ChaosCloth/ChaosClothConstraints.h
@@ -67,12 +67,13 @@ PRAGMA_ENABLE_DEPRECATION_WARNINGS
 			const FClothingPatternData* PatternData,
 			const TMap<FString, TConstArrayView<FRealSingle>>& WeightMaps,
 			const TMap<FString, const TSet<int32>*>& VertexSets,
+			const TMap<FString, const TSet<int32>*>& FaceSets,
 			const TMap<FString, TConstArrayView<int32>>& FaceIntMaps,
 			const TArray<TConstArrayView<TTuple<int32, int32, FRealSingle>>>& Tethers,
 			Softs::FSolverReal MeshScale,
 			bool bEnabled);
 
-		UE_DEPRECATED(5.4, "Use AddRules() with WeightMaps, VertexSets, FaceIntMaps, and optional PatternData instead.")
+		UE_DEPRECATED(5.4, "Use AddRules() with WeightMaps, VertexSets, FaceSets, FaceIntMaps, and optional PatternData instead.")
 		void AddRules(
 			const Softs::FCollectionPropertyConstFacade& ConfigProperties,
 			const FTriangleMesh& TriangleMesh,
@@ -82,7 +83,7 @@ PRAGMA_ENABLE_DEPRECATION_WARNINGS
 			Softs::FSolverReal MeshScale,
 			bool bEnabled);
 
-		UE_DEPRECATED(5.3, "Use AddRules() with WeightMaps, VertexSets, FaceIntMap, and optional PatternData instead.")
+		UE_DEPRECATED(5.3, "Use AddRules() with WeightMaps, VertexSets, FaceSets, FaceIntMap, and optional PatternData instead.")
 		void AddRules(
 			const Softs::FCollectionPropertyConstFacade& ConfigProperties,
 			const FTriangleMesh& TriangleMesh,
@@ -95,18 +96,19 @@ PRAGMA_ENABLE_DEPRECATION_WARNINGS
 			const Softs::FCollectionPropertyConstFacade& ConfigProperties,
 			const TMap<FString, TConstArrayView<FRealSingle>>& WeightMaps,
 			const TMap<FString, const TSet<int32>*>& VertexSets,
+			const TMap<FString, const TSet<int32>*>& FaceSets,
 			const TMap<FString, TConstArrayView<int32>>& FaceIntMaps,
 			Softs::FSolverReal MeshScale,
 			Softs::FSolverReal MaxDistancesScale = (Softs::FSolverReal)1.);
 
-		UE_DEPRECATED(5.4, "Use Update() with WeightMaps, VertexSets, and FaceIntMaps instead.")
+		UE_DEPRECATED(5.4, "Use Update() with WeightMaps, VertexSets, FaceSets, and FaceIntMaps instead.")
 		void Update(
 			const Softs::FCollectionPropertyConstFacade& ConfigProperties,
 			const TMap<FString, TConstArrayView<FRealSingle>>& WeightMaps,
 			Softs::FSolverReal MeshScale,
 			Softs::FSolverReal MaxDistancesScale = (Softs::FSolverReal)1.);
 
-		UE_DEPRECATED(5.3, "Use Update() with WeightMaps, VertexSets, and FaceIntMaps instead.")
+		UE_DEPRECATED(5.3, "Use Update() with WeightMaps, VertexSets, FaceSets, and FaceIntMaps instead.")
 		void Update(
 			const Softs::FCollectionPropertyConstFacade& ConfigProperties,
 			Softs::FSolverReal MeshScale,
@@ -150,6 +152,7 @@ PRAGMA_ENABLE_DEPRECATION_WARNINGS
 		void CreateSelfCollisionConstraints(
 			const Softs::FCollectionPropertyConstFacade& ConfigProperties,
 			const TMap<FString, const TSet<int32>*>& VertexSets,
+			const TMap<FString, const TSet<int32>*>& FaceSets,
 			const TMap<FString, TConstArrayView<int32>>& FaceIntMaps,
 			const FTriangleMesh& TriangleMesh);
 		void CreateStretchConstraints(
diff --git a/Engine/Plugins/Experimental/ChaosCloth/Source/Chaos/Public/ChaosCloth/ChaosClothingSimulationMesh.h b/Engine/Plugins/Experimental/ChaosCloth/Source/Chaos/Public/ChaosCloth/ChaosClothingSimulationMesh.h
--- a/Engine/Plugins/Experimental/ChaosCloth/Source/Chaos/Public/ChaosCloth/ChaosClothingSimulationMesh.h
+++ b/Engine/Plugins/Experimental/ChaosCloth/Source/Chaos/Public/ChaosCloth/ChaosClothingSimulationMesh.h
@@ -73,6 +73,7 @@ namespace Chaos
 		virtual TMap<FString, int32> GetWeightMapIndices() const { return GetWeightMapIndices(0); }
 		virtual TArray<TConstArrayView<FRealSingle>> GetWeightMaps(int32 LODIndex) const;
 		virtual TMap<FString, const TSet<int32>*> GetVertexSets(int32 LODIndex) const;
+		virtual TMap<FString, const TSet<int32>*> GetFaceSets(int32 LODIndex) const;
 		virtual TMap<FString, TConstArrayView<int32>> GetFaceIntMaps(int32 LODIndex) const;
 		virtual TArray<TConstArrayView<TTuple<int32, int32, float>>> GetTethers(int32 LODIndex, bool bUseGeodesicTethers) const;
 		virtual int32 GetReferenceBoneIndex() const;
@@ -137,6 +138,9 @@ namespace Chaos
 		/* Return the specified LOD's vertex sets. */
 		virtual const TMap<FString, TSet<int32>*> GetVertexSets(int32 LODIndex) const = 0;
 
+		/* Return the specified LOD's face sets. */
+		virtual const TMap<FString, TSet<int32>*> GetFaceSets(int32 LODIndex) const = 0;
+
 		/* Return the specified LOD's face int maps. */
 		virtual TMap<FString, TConstArrayView<int32>> GetFaceIntMaps(int32 LODIndex) const = 0;
 
diff --git a/Engine/Plugins/Experimental/ChaosCloth/Source/Chaos/Public/ChaosCloth/ChaosClothingSimulationSkeletalMesh.h b/Engine/Plugins/Experimental/ChaosCloth/Source/Chaos/Public/ChaosCloth/ChaosClothingSimulationSkeletalMesh.h
--- a/Engine/Plugins/Experimental/ChaosCloth/Source/Chaos/Public/ChaosCloth/ChaosClothingSimulationSkeletalMesh.h
+++ b/Engine/Plugins/Experimental/ChaosCloth/Source/Chaos/Public/ChaosCloth/ChaosClothingSimulationSkeletalMesh.h
@@ -53,6 +53,7 @@ PRAGMA_ENABLE_DEPRECATION_WARNINGS
 		virtual TMap<FString, int32> GetWeightMapIndices() const override { return GetWeightMapIndices(0); }
 		virtual TArray<TConstArrayView<FRealSingle>> GetWeightMaps(int32 LODIndex) const override;
 		virtual TMap<FString, const TSet<int32>*> GetVertexSets(int32 LODIndex) const override;
+		virtual TMap<FString, const TSet<int32>*> GetFaceSets(int32 LODIndex) const override;
 		virtual TMap<FString, TConstArrayView<int32>> GetFaceIntMaps(int32 LODIndex) const override;
 		virtual TArray<TConstArrayView<TTuple<int32, int32, float>>> GetTethers(int32 LODIndex, bool bUseGeodesicTethers) const override;
 		virtual int32 GetReferenceBoneIndex() const override;
diff --git a/Engine/Plugins/Experimental/ChaosClothAsset/Source/ChaosClothAssetEngine/Private/ChaosClothAsset/ClothSimulationMesh.cpp b/Engine/Plugins/Experimental/ChaosClothAsset/Source/ChaosClothAssetEngine/Private/ChaosClothAsset/ClothSimulationMesh.cpp
--- a/Engine/Plugins/Experimental/ChaosClothAsset/Source/ChaosClothAssetEngine/Private/ChaosClothAsset/ClothSimulationMesh.cpp
+++ b/Engine/Plugins/Experimental/ChaosClothAsset/Source/ChaosClothAssetEngine/Private/ChaosClothAsset/ClothSimulationMesh.cpp
@@ -137,6 +137,20 @@ PRAGMA_ENABLE_DEPRECATION_WARNINGS
 		return VertexSets;
 	}
 
+	TMap<FString, const TSet<int32>*> FClothSimulationMesh::GetFaceSets(int32 LODIndex) const
+	{
+		TMap<FString, const TSet<int32>*> FaceSets;
+		if (ClothSimulationModel.ClothSimulationLodModels.IsValidIndex(LODIndex))
+		{
+			FaceSets.Reserve(ClothSimulationModel.ClothSimulationLodModels[LODIndex].FaceSets.Num());
+			for (TMap<FName, TSet<int32>>::TConstIterator Iter = ClothSimulationModel.ClothSimulationLodModels[LODIndex].FaceSets.CreateConstIterator(); Iter; ++Iter)
+			{
+				FaceSets.Emplace(Iter.Key().ToString(), &Iter.Value());
+			}
+		}
+		return FaceSets;
+	}
+
 	TMap<FString, TConstArrayView<int32>> FClothSimulationMesh::GetFaceIntMaps(int32 LODIndex) const
 	{
 		TMap<FString, TConstArrayView<int32>> FaceIntMaps;
diff --git a/Engine/Plugins/Experimental/ChaosClothAsset/Source/ChaosClothAssetEngine/Private/ChaosClothAsset/ClothSimulationMesh.h b/Engine/Plugins/Experimental/ChaosClothAsset/Source/ChaosClothAssetEngine/Private/ChaosClothAsset/ClothSimulationMesh.h
--- a/Engine/Plugins/Experimental/ChaosClothAsset/Source/ChaosClothAssetEngine/Private/ChaosClothAsset/ClothSimulationMesh.h
+++ b/Engine/Plugins/Experimental/ChaosClothAsset/Source/ChaosClothAssetEngine/Private/ChaosClothAsset/ClothSimulationMesh.h
@@ -46,6 +46,7 @@ PRAGMA_ENABLE_DEPRECATION_WARNINGS
 		virtual TMap<FString, int32> GetWeightMapIndices() const override { return GetWeightMapIndices(0); }
 		virtual TArray<TConstArrayView<::Chaos::FRealSingle>> GetWeightMaps(int32 LODIndex) const override;
 		virtual TMap<FString, const TSet<int32>*> GetVertexSets(int32 LODIndex) const override;
+		virtual TMap<FString, const TSet<int32>*> GetFaceSets(int32 LODIndex) const override;
 		virtual TMap<FString, TConstArrayView<int32>> GetFaceIntMaps(int32 LODIndex) const override;
 		// Note: there is only one set of tethers stored on ClothSimulationMesh assets
 		virtual TArray<TConstArrayView<TTuple<int32, int32, float>>> GetTethers(int32 LODIndex, bool /*bUseGeodesicTethers*/) const override;
diff --git a/Engine/Plugins/Experimental/ChaosClothAsset/Source/ChaosClothAssetEngine/Private/ChaosClothAsset/ClothSimulationModel.cpp b/Engine/Plugins/Experimental/ChaosClothAsset/Source/ChaosClothAssetEngine/Private/ChaosClothAsset/ClothSimulationModel.cpp
--- a/Engine/Plugins/Experimental/ChaosClothAsset/Source/ChaosClothAssetEngine/Private/ChaosClothAsset/ClothSimulationModel.cpp
+++ b/Engine/Plugins/Experimental/ChaosClothAsset/Source/ChaosClothAssetEngine/Private/ChaosClothAsset/ClothSimulationModel.cpp
@@ -147,6 +147,8 @@ bool FChaosClothSimulationLodModel::Serialize(FArchive& Ar)
 
 	Ar << VertexSets;
 	Ar << FaceIntMaps;
+	Ar << FaceSets;
+
 
 	// Return true to confirm that serialization has already been taken care of
 	return true;
@@ -182,16 +184,22 @@ FChaosClothSimulationModel::FChaosClothSimulationModel(const TArray<TSharedRef<c
 			LodModel.WeightMaps.Add(WeightMapName) = Cloth.GetWeightMap(WeightMapName);
 		}
 
-		// Copy vertex sets
+		// Copy vertex and face sets
 		FCollectionClothSelectionConstFacade Selection(ClothCollections[LodIndex]);
 		const TArray<FName> SelectionNames = Selection.GetNames();
-		LodModel.VertexSets.Reserve(SelectionNames.Num()); // At this point, only the SimVertex3d selections should be in the collection.
+		LodModel.VertexSets.Reserve(SelectionNames.Num());
+		LodModel.FaceSets.Reserve(SelectionNames.Num());
 		for (const FName& SelectionName : SelectionNames)
 		{
-			if (Selection.GetSelectionGroup(SelectionName) == ClothCollectionGroup::SimVertices3D)
+			const FName SelectionGroup = Selection.GetSelectionGroup(SelectionName);
+			if (SelectionGroup == ClothCollectionGroup::SimVertices3D)
 			{
 				LodModel.VertexSets.Add(SelectionName) = Selection.GetSelectionSet(SelectionName);
 			}
+			else if (SelectionGroup == ClothCollectionGroup::SimFaces)
+			{
+				LodModel.FaceSets.Add(SelectionName) = Selection.GetSelectionSet(SelectionName);
+			}
 		}
 
 		// Copy face int maps
diff --git a/Engine/Plugins/Experimental/ChaosClothAsset/Source/ChaosClothAssetEngine/Private/ChaosClothAsset/ClothSimulationModel.h b/Engine/Plugins/Experimental/ChaosClothAsset/Source/ChaosClothAssetEngine/Private/ChaosClothAsset/ClothSimulationModel.h
--- a/Engine/Plugins/Experimental/ChaosClothAsset/Source/ChaosClothAssetEngine/Private/ChaosClothAsset/ClothSimulationModel.h
+++ b/Engine/Plugins/Experimental/ChaosClothAsset/Source/ChaosClothAssetEngine/Private/ChaosClothAsset/ClothSimulationModel.h
@@ -68,6 +68,9 @@ struct FChaosClothSimulationLodModel
 	/** Face int maps (currently used by cloth collision layers)*/
 	TMap<FName, TArray<int32>> FaceIntMaps;
 
+	/** Face sets */
+	TMap<FName, TSet<int32>> FaceSets;
+
 	// Custom serialize for weight maps
 	bool Serialize(FArchive& Ar);
 };
diff --git a/Engine/Plugins/Experimental/ChaosClothAssetEditor/Source/ChaosClothAssetDataflowNodes/Private/ChaosClothAsset/SimulationSelfCollisionConfigNode.cpp b/Engine/Plugins/Experimental/ChaosClothAssetEditor/Source/ChaosClothAssetDataflowNodes/Private/ChaosClothAsset/SimulationSelfCollisionConfigNode.cpp
--- a/Engine/Plugins/Experimental/ChaosClothAssetEditor/Source/ChaosClothAssetDataflowNodes/Private/ChaosClothAsset/SimulationSelfCollisionConfigNode.cpp
+++ b/Engine/Plugins/Experimental/ChaosClothAssetEditor/Source/ChaosClothAssetDataflowNodes/Private/ChaosClothAsset/SimulationSelfCollisionConfigNode.cpp
@@ -11,6 +11,8 @@ FChaosClothAssetSimulationSelfCollisionConfigNode::FChaosClothAssetSimulationSel
 {
 	RegisterCollectionConnections();
 	RegisterInputConnection(&SelfCollisionLayers.StringValue, GET_MEMBER_NAME_CHECKED(FChaosClothAssetConnectableIStringValue, StringValue));
+	RegisterInputConnection(&SelfCollisionDisabledFaces.StringValue, GET_MEMBER_NAME_CHECKED(FChaosClothAssetConnectableIStringValue, StringValue));
+	RegisterInputConnection(&SelfCollisionEnabledKinematicFaces.StringValue, GET_MEMBER_NAME_CHECKED(FChaosClothAssetConnectableIStringValue, StringValue));
 }
 
 void FChaosClothAssetSimulationSelfCollisionConfigNode::AddProperties(FPropertyHelper& PropertyHelper) const
@@ -21,6 +23,12 @@ void FChaosClothAssetSimulationSelfCollisionConfigNode::AddProperties(FPropertyH
 	PropertyHelper.SetProperty(this, &SelfCollisionFriction);
 	PropertyHelper.SetProperty(this, &SelfCollisionDisableNeighborDistance, {}, ECollectionPropertyFlags::None); // Non animatable
 	PropertyHelper.SetPropertyString(this, &SelfCollisionLayers);
+	PropertyHelper.SetPropertyString(this, &SelfCollisionDisabledFaces);
+	PropertyHelper.SetPropertyBool(this, &bSelfCollideAgainstAllKinematicVertices);
+	PropertyHelper.SetPropertyString(this, &SelfCollisionEnabledKinematicFaces);
+	PropertyHelper.SetProperty(this, &SelfCollisionKinematicColliderThickness);
+	PropertyHelper.SetProperty(this, &SelfCollisionKinematicColliderStiffness);
+	PropertyHelper.SetProperty(this, &SelfCollisionKinematicColliderFriction);
 
 	PropertyHelper.SetPropertyBool(this, &bUseSelfIntersections);
 	PropertyHelper.SetPropertyBool(this, &bUseGlobalIntersectionAnalysis);
diff --git a/Engine/Plugins/Experimental/ChaosClothAssetEditor/Source/ChaosClothAssetDataflowNodes/Private/ChaosClothAsset/TerminalNode.cpp b/Engine/Plugins/Experimental/ChaosClothAssetEditor/Source/ChaosClothAssetDataflowNodes/Private/ChaosClothAsset/TerminalNode.cpp
--- a/Engine/Plugins/Experimental/ChaosClothAssetEditor/Source/ChaosClothAssetDataflowNodes/Private/ChaosClothAsset/TerminalNode.cpp
+++ b/Engine/Plugins/Experimental/ChaosClothAssetEditor/Source/ChaosClothAssetDataflowNodes/Private/ChaosClothAsset/TerminalNode.cpp
@@ -36,11 +36,13 @@ namespace UE::Chaos::ClothAsset::Private
 			FCollectionClothSelectionConstFacade Selection(ClothCollection);
 			if (Selection.IsValid())
 			{
-				// Just checksum the sets that are SimVertex3D sets since those are the only ones we care about right now
+				// Just checksum the sets that are SimVertex3D and SimFace sets since those are the only ones we care about right now
 				const TArray<FName> SelectionNames = Selection.GetNames();
 				for (const FName& SelectionName : SelectionNames)
 				{
-					if (Selection.GetSelectionGroup(SelectionName) == ClothCollectionGroup::SimVertices3D)
+					const FName SelectionGroup = Selection.GetSelectionGroup(SelectionName);
+					if (SelectionGroup == ClothCollectionGroup::SimVertices3D ||
+						SelectionGroup == ClothCollectionGroup::SimFaces)
 					{
 						const TArray<int32> SelectionAsArray = Selection.GetSelectionSet(SelectionName).Array();
 						Checksum = HashCombineFast(Checksum, GetTypeHash(SelectionName));
@@ -189,9 +191,11 @@ void FChaosClothAssetTerminalNode::SetAssetValue(TObjectPtr<UObject> Asset, Data
 				const TArray<FName> InSelectionNames = InSelection.GetNames();
 				for (const FName& InSelectionName : InSelectionNames)
 				{
-					if (InSelection.GetSelectionGroup(InSelectionName) == ClothCollectionGroup::SimVertices3D)
+					const FName SelectionGroup = InSelection.GetSelectionGroup(InSelectionName);
+					if (SelectionGroup == ClothCollectionGroup::SimVertices3D || 
+						SelectionGroup == ClothCollectionGroup::SimFaces)
 					{
-						Selection.FindOrAddSelectionSet(InSelectionName, ClothCollectionGroup::SimVertices3D) = InSelection.GetSelectionSet(InSelectionName);
+						Selection.FindOrAddSelectionSet(InSelectionName, SelectionGroup) = InSelection.GetSelectionSet(InSelectionName);
 					}
 				}
 			}
diff --git a/Engine/Plugins/Experimental/ChaosClothAssetEditor/Source/ChaosClothAssetDataflowNodes/Public/ChaosClothAsset/SimulationSelfCollisionConfigNode.h b/Engine/Plugins/Experimental/ChaosClothAssetEditor/Source/ChaosClothAssetDataflowNodes/Public/ChaosClothAsset/SimulationSelfCollisionConfigNode.h
--- a/Engine/Plugins/Experimental/ChaosClothAssetEditor/Source/ChaosClothAssetDataflowNodes/Public/ChaosClothAsset/SimulationSelfCollisionConfigNode.h
+++ b/Engine/Plugins/Experimental/ChaosClothAssetEditor/Source/ChaosClothAssetDataflowNodes/Public/ChaosClothAsset/SimulationSelfCollisionConfigNode.h
@@ -14,7 +14,7 @@ struct FChaosClothAssetSimulationSelfCollisionConfigNode : public FChaosClothAss
 	DATAFLOW_NODE_DEFINE_INTERNAL(FChaosClothAssetSimulationSelfCollisionConfigNode, "SimulationSelfCollisionConfig", "Cloth", "Cloth Simulation Self Collision Config")
 
 public:
-	/** The radius of the spheres used in self collision (i.e., offset per side. total thickness of cloth is 2x this value). */
+	/** The self collision offset per side. Total thickness of cloth is 2x this value. */
 	UPROPERTY(EditAnywhere, Category = "Self-Collision Properties", meta = (UIMin = "0", UIMax = "100", ClampMin = "0", ClampMax = "1000", EditCondition = "bUseSelfCollisions"))
 	float SelfCollisionThickness = 0.5f;
 
@@ -37,6 +37,30 @@ public:
 	UPROPERTY(EditAnywhere, Category = "Self-Collision Properties")
 	FChaosClothAssetConnectableIStringValue SelfCollisionLayers = { TEXT("SelfCollisionLayers") };
 
+	/** Sim face selection set of faces which should not self collide */
+	UPROPERTY(EditAnywhere, Category = "Self-Collision Properties")
+	FChaosClothAssetConnectableIStringValue SelfCollisionDisabledFaces = { TEXT("SelfCollisionDisabledFaces") };
+
+	/** Self collide against all kinematic vertices. Kinematic colliders do not do Self Intersections. They always collide against the front-face.*/
+	UPROPERTY(EditAnywhere, Category = "Self-Collision Properties - Kinematic Colliders")
+	bool bSelfCollideAgainstAllKinematicVertices = false;
+
+	/** Sim face selection set of kinematic faces which should self collide. Kinematic colliders do not do Self Intersections. They always collide against the front-face. */
+	UPROPERTY(EditAnywhere, Category = "Self-Collision Properties - Kinematic Colliders", meta = (EditCondition = "!bSelfCollideAgainstAllKinematicVertices"))
+	FChaosClothAssetConnectableIStringValue SelfCollisionEnabledKinematicFaces = { TEXT("SelfCollisionEnabledKinematicFaces") };
+
+	/** Thickness of kinematic colliders. Total offset between cloth and kinematic colliders is SelfCollisionThickness + SelfCollisionKinematicColliderThickness. */
+	UPROPERTY(EditAnywhere, Category = "Self-Collision Properties - Kinematic Colliders", meta = (UIMin = "0", UIMax = "100", ClampMin = "0", ClampMax = "1000", EditCondition = "bUseSelfCollisions"))
+	float SelfCollisionKinematicColliderThickness = 0.f;
+
+	/** The stiffness of the springs used to control self collision (PBD Solver). */
+	UPROPERTY(EditAnywhere, Category = "Self-Collision Properties - Kinematic Colliders", meta = (UIMin = "0", UIMax = "1", ClampMin = "0", ClampMax = "1"))
+	float SelfCollisionKinematicColliderStiffness = 1.f;
+
+	/** Friction coefficient for cloth - kinematic cloth interaction. */
+	UPROPERTY(EditAnywhere, Category = "Self-Collision Properties - Kinematic Colliders", meta = (UIMin = "0", UIMax = "1", ClampMin = "0", ClampMax = "1"))
+	float SelfCollisionKinematicColliderFriction = 0.0f;
+
 	/** Enable self intersection resolution. This will try to fix any cloth intersections that are not handled by collision repulsions. */
 	UPROPERTY(EditAnywhere, Category = "Experimental")
 	bool bUseSelfIntersections = false;
diff --git a/Engine/Source/Runtime/Experimental/Chaos/Private/Chaos/PBDCollisionSpringConstraintsBase.cpp b/Engine/Source/Runtime/Experimental/Chaos/Private/Chaos/PBDCollisionSpringConstraintsBase.cpp
--- a/Engine/Source/Runtime/Experimental/Chaos/Private/Chaos/PBDCollisionSpringConstraintsBase.cpp
+++ b/Engine/Source/Runtime/Experimental/Chaos/Private/Chaos/PBDCollisionSpringConstraintsBase.cpp
@@ -14,6 +14,14 @@
 
 namespace Chaos::Softs {
 
+int32 DebugParticleIndex = INDEX_NONE;
+FAutoConsoleVariableRef CVarChaosClothDebugParticleIndex(TEXT("p.ChaosCloth.DebugParticleIndex"), DebugParticleIndex, TEXT("DebugParticleIndex"));
+
+FSolverReal Chaos_CollisionSpring_MaxTimer = (FSolverReal)0.1f;
+FAutoConsoleVariableRef CVarChaosCollisionSPringMaxTimer(TEXT("p.Chaos.CollisionSpring.MaxTimer"), Chaos_CollisionSpring_MaxTimer, TEXT("Max Timer"));
+
+
+
 static TConstArrayView<FSolverVec3> GetConstArrayView(const FSolverParticles& Particles, int32 Offset, int32 NumParticles, const TArray<FSolverVec3>* const Data)
 {
 	if (Data)
@@ -42,13 +50,18 @@ FPBDCollisionSpringConstraintsBase::FPBDCollisionSpringConstraintsBase(
 	const FSolverReal InThickness,
 	const FSolverReal InStiffness,
 	const FSolverReal InFrictionCoefficient,
+	const FSolverReal InKinematicColliderThickness,
+	const FSolverReal InKinematicColliderStiffness,
+	const FSolverReal InKinematicColliderFrictionCoefficient,
 	const FSolverReal InProximityStiffness)
 	: Thickness(InThickness)
 	, Stiffness(InStiffness)
 	, FrictionCoefficient(InFrictionCoefficient)
+	, KinematicColliderThickness(InKinematicColliderThickness)
+	, KinematicColliderStiffness(InKinematicColliderStiffness)
+	, KinematicColliderFrictionCoefficient(InKinematicColliderFrictionCoefficient)
 	, ProximityStiffness(InProximityStiffness)
 	, TriangleMesh(InTriangleMesh)
-	, Elements(InTriangleMesh.GetSurfaceElements())
 	, ReferencePositions(InReferencePositions)
 	, DisabledCollisionElements(InDisabledCollisionElements)
 	, Offset(InOffset)
@@ -60,7 +73,7 @@ FPBDCollisionSpringConstraintsBase::FPBDCollisionSpringConstraintsBase(
 
 void FPBDCollisionSpringConstraintsBase::UpdateCollisionLayers(const TConstArrayView<int32>& InFaceCollisionLayers)
 {
-	if (InFaceCollisionLayers.Num() != Elements.Num())
+	if (InFaceCollisionLayers.Num() != TriangleMesh.GetElements().Num())
 	{
 		// Reset collision layers
 		FaceCollisionLayers = TConstArrayView<int32>();
@@ -92,55 +105,52 @@ void FPBDCollisionSpringConstraintsBase::UpdateCollisionLayers(const TConstArray
 	}
 }
 
-void FPBDCollisionSpringConstraintsBase::Init(const FSolverParticles& Particles)
-{
-	if (!Elements.Num())
-	{
-		Constraints.Reset();
-		Barys.Reset();
-		FlipNormal.Reset();
-		return;
-	}
-
-	FTriangleMesh::TBVHType<FSolverReal> BVH;
-	{
-		TRACE_CPUPROFILER_EVENT_SCOPE(ChaosPBDCollisionSpring_BuildBVH);
-		TriangleMesh.BuildBVH(static_cast<const TArrayView<const FSolverVec3>&>(Particles.XArray()), BVH);
-	}
-	TArray<FPBDTriangleMeshCollisions::FGIAColor> EmptyGIAColors;
-	return Init(Particles, BVH, static_cast<TConstArrayView<FPBDTriangleMeshCollisions::FGIAColor>>(EmptyGIAColors), EmptyGIAColors);
-}
-
 template<typename SpatialAccelerator, typename SolverParticlesOrRange>
-void FPBDCollisionSpringConstraintsBase::Init(const SolverParticlesOrRange& Particles, const SpatialAccelerator& Spatial,
+void FPBDCollisionSpringConstraintsBase::Init(const SolverParticlesOrRange& Particles, const FSolverReal Dt, const FPBDTriangleMeshCollisions::FTriangleSubMesh& CollidableSubMesh, const SpatialAccelerator& Spatial,
 	const TConstArrayView<FPBDTriangleMeshCollisions::FGIAColor>& VertexGIAColors, const TArray<FPBDTriangleMeshCollisions::FGIAColor>& TriangleGIAColors)
 {
+	PRAGMA_DISABLE_DEPRECATION_WARNINGS
+	const FTriangleMesh& CollidableMesh = CollidableSubMesh.GetCollidableMesh();
+	const TArray<TVec3<int32>>& Elements = CollidableMesh.GetElements();
+
 	if (!Elements.Num())
 	{
 		Constraints.Reset();
 		Barys.Reset();
 		FlipNormal.Reset();
+		ConstraintTypes.Reset();
 		return;
 	}
 	{
+		if (ExistingConstraintLookup.Num() != NumParticles)
+		{
+			ExistingConstraintLookup.Reset();
+			ExistingConstraintLookup.SetNum(NumParticles);
+		}
+
 		bGlobalIntersectionAnalysis = VertexGIAColors.Num() == NumParticles + Offset && TriangleGIAColors.Num() == Elements.Num();
 		TRACE_CPUPROFILER_EVENT_SCOPE(ChaosPBDCollisionSpring_ProximityQuery);
 
+		const int32 NumCollidableParticles = CollidableSubMesh.GetCollidableVertices().IsEmpty() ? NumParticles : CollidableSubMesh.GetCollidableVertices().Num();
+
 		// Preallocate enough space for all possible connections.
 		constexpr int32 MaxConnectionsPerPoint = 3;
-		Constraints.SetNumUninitialized(NumParticles * MaxConnectionsPerPoint);
-		Barys.SetNumUninitialized(NumParticles * MaxConnectionsPerPoint);
-		FlipNormal.SetNumUninitialized(NumParticles * MaxConnectionsPerPoint);
+		Constraints.SetNumUninitialized(NumCollidableParticles * MaxConnectionsPerPoint);
+		Barys.SetNumUninitialized(NumCollidableParticles * MaxConnectionsPerPoint);
+		FlipNormal.SetNumUninitialized(NumCollidableParticles * MaxConnectionsPerPoint);
+		ConstraintTypes.SetNumUninitialized(NumCollidableParticles * MaxConnectionsPerPoint);
 
 		std::atomic<int32> ConstraintIndex(0);
 
-		const FSolverReal HeightSq = FMath::Square(Thickness + Thickness);
+		const FSolverReal MaxSingleSidedThickness = FMath::Max(Thickness, KinematicColliderThickness);
+		const FSolverReal HeightSq = FMath::Square(Thickness + MaxSingleSidedThickness);
 
 		const TConstArrayView<FSolverVec3> ReferencePositionsView = GetConstArrayView(Particles, Offset, NumParticles, ReferencePositions);
 
-		PhysicsParallelFor(NumParticles,
-			[this, &Spatial, &Particles, &ConstraintIndex, HeightSq, MaxConnectionsPerPoint, &VertexGIAColors, &TriangleGIAColors, &ReferencePositionsView](int32 i)
+		PhysicsParallelFor(NumCollidableParticles,
+			[this, Dt, &CollidableSubMesh, &CollidableMesh, &Elements, &Spatial, &Particles, &ConstraintIndex, MaxSingleSidedThickness, HeightSq, MaxConnectionsPerPoint, &VertexGIAColors, &TriangleGIAColors, &ReferencePositionsView](int32 CollidableIndex)
 			{
+				const int32 i = CollidableSubMesh.GetCollidableVertices().IsEmpty() ? CollidableIndex : CollidableSubMesh.GetCollidableVertices()[CollidableIndex];
 				const int32 Index = i + Offset;
 				if (Particles.InvM(Index) == (FSolverReal)0.)
 				{
@@ -151,21 +161,27 @@ void FPBDCollisionSpringConstraintsBase::Init(const SolverParticlesOrRange& Part
 				const bool bVertexHasCollisionLayers = VertexCollisionLayers.IsValidIndex(i) && VertexCollisionLayers[i][0] != INDEX_NONE;
 				check(!bVertexHasCollisionLayers || VertexCollisionLayers[i][0] <= VertexCollisionLayers[i][1]);
 
+				TMap<int32, FExistingConstraintData> PrevExistingConstraintLookup;
+				Swap(PrevExistingConstraintLookup, ExistingConstraintLookup[i]);
+
 				TArray< TTriangleCollisionPoint<FSolverReal> > Result;
-				if (TriangleMesh.PointProximityQuery(Spatial, static_cast<const TArrayView<const FSolverVec3>&>(Particles.XArray()), Index, Particles.X(Index), Thickness * ExtraThicknessMult, Thickness * ExtraThicknessMult,
-					[this, bVertexHasCollisionLayers, &VertexGIAColors, &TriangleGIAColors](const int32 PointIndex, const int32 TriangleIndex)->bool
+				int32 ConstraintsAdded = 0;
+				if (CollidableMesh.PointProximityQuery(Spatial, static_cast<const TArrayView<const FSolverVec3>&>(Particles.XArray()), Index, Particles.X(Index), Thickness * ExtraThicknessMult, MaxSingleSidedThickness * ExtraThicknessMult,
+					[this, bVertexHasCollisionLayers, &Particles, &CollidableSubMesh, &Elements, &VertexGIAColors, &TriangleGIAColors](const int32 PointIndex, const int32 SubMeshTriangleIndex)->bool
 					{
-						const TVector<int32, 3>& Elem = Elements[TriangleIndex];						
+						const TVector<int32, 3>& Elem = Elements[SubMeshTriangleIndex];
+						const int32 FullMeshTriangleIndex = CollidableSubMesh.FullMeshIndexFromSubIndex(SubMeshTriangleIndex);
 
+						const bool bFaceIsKinematic = CollidableSubMesh.GetSubMeshElementIsKinematic()[SubMeshTriangleIndex];
 						bool bUseCollisionLayerOverride = false;
-						if (bVertexHasCollisionLayers && FaceCollisionLayers[TriangleIndex] != INDEX_NONE)
+						if (!bFaceIsKinematic && bVertexHasCollisionLayers && FaceCollisionLayers[FullMeshTriangleIndex] != INDEX_NONE)
 						{
-							if (FaceCollisionLayers[TriangleIndex] < VertexCollisionLayers[PointIndex - Offset][0] || FaceCollisionLayers[TriangleIndex] > VertexCollisionLayers[PointIndex - Offset][1])
+							if (FaceCollisionLayers[FullMeshTriangleIndex] < VertexCollisionLayers[PointIndex - Offset][0] || FaceCollisionLayers[FullMeshTriangleIndex] > VertexCollisionLayers[PointIndex - Offset][1])
 							{
 								bUseCollisionLayerOverride = true;
 							}
 						}
-						if (!bUseCollisionLayerOverride && bGlobalIntersectionAnalysis)
+						if (!bFaceIsKinematic && !bUseCollisionLayerOverride && bGlobalIntersectionAnalysis)
 						{
 							const bool bIsAnyBoundary = VertexGIAColors[PointIndex].IsBoundary()
 								|| VertexGIAColors[Elem[0]].IsBoundary() 
@@ -180,7 +196,7 @@ void FPBDCollisionSpringConstraintsBase::Init(const SolverParticlesOrRange& Part
 								(VertexGIAColors[Elem[0]].IsLoop() 
 									|| VertexGIAColors[Elem[1]].IsLoop() 
 									|| VertexGIAColors[Elem[2]].IsLoop() 
-									|| TriangleGIAColors[TriangleIndex].IsLoop());
+									|| TriangleGIAColors[SubMeshTriangleIndex].IsLoop());
 
 							if (bAreBothLoop)
 							{
@@ -228,19 +244,35 @@ void FPBDCollisionSpringConstraintsBase::Init(const SolverParticlesOrRange& Part
 							}
 						}
 
+
+						FSolverVec3 Bary(CollisionPoint.Bary[1], CollisionPoint.Bary[2], CollisionPoint.Bary[3]);
+
+						const int32 FullMeshTriangleIndex = CollidableSubMesh.FullMeshIndexFromSubIndex(CollisionPoint.Indices[1]);
+
 						bool bFlipNormal = false;
+						// Check kinematic
+						const bool bFaceIsKinematic = CollidableSubMesh.GetSubMeshElementIsKinematic()[CollisionPoint.Indices[1]];
+						if (bFaceIsKinematic)
+						{
+							PrevExistingConstraintLookup.Remove(FullMeshTriangleIndex);
+							// Chaos internal winding order is reversed.
+							bFlipNormal = true;
+						}
+
+						EConstraintType ConstraintType = bFaceIsKinematic ? EConstraintType::Kinematic : EConstraintType::Default;
+
 						// Check collision layers
 						bool bUseCollisionLayerOverride = false;
-						if (bVertexHasCollisionLayers && FaceCollisionLayers[CollisionPoint.Indices[1]] != INDEX_NONE)
+						if (!bFaceIsKinematic && bVertexHasCollisionLayers && FaceCollisionLayers[FullMeshTriangleIndex] != INDEX_NONE)
 						{
-							if (FaceCollisionLayers[CollisionPoint.Indices[1]] < VertexCollisionLayers[i][0])
+							if (FaceCollisionLayers[FullMeshTriangleIndex] < VertexCollisionLayers[i][0])
 							{
 								// Face is lower layer than the vertex. Vertex should always be in front of face (as UE sees it).
 								// NOTE: Chaos internal winding order for normals is reversed, so flip normal in this case.
 								bFlipNormal = true;
 								bUseCollisionLayerOverride = true;
 							}
-							else if (FaceCollisionLayers[CollisionPoint.Indices[1]] > VertexCollisionLayers[i][1])
+							else if (FaceCollisionLayers[FullMeshTriangleIndex] > VertexCollisionLayers[i][1])
 							{
 								// Face is higher layer than the vertex. Vertex should always be behind face (as UE sees it).
 								// NOTE: Chaos internal winding order for normals is reversed, so don't flip normal in this case.
@@ -249,7 +281,7 @@ void FPBDCollisionSpringConstraintsBase::Init(const SolverParticlesOrRange& Part
 							}
 						}
 
-						if (!bUseCollisionLayerOverride)
+						if (!bFaceIsKinematic && !bUseCollisionLayerOverride)
 						{
 							// NOTE: CollisionPoint.Normal has already been flipped to point toward the Point, so need to recalculate here.
 							const TTriangle<FSolverReal> Triangle(Particles.X(Elem[0]), Particles.X(Elem[1]), Particles.X(Elem[2]));
@@ -265,13 +297,40 @@ void FPBDCollisionSpringConstraintsBase::Init(const SolverParticlesOrRange& Part
 							{
 
 								bFlipNormal = !bFlipNormal;
+								ConstraintType = EConstraintType::GIAFlipped;
 							}
 						}
 						const int32 IndexToWrite = ConstraintIndex.fetch_add(1);
 
 						Constraints[IndexToWrite] = { Index, Elem[0], Elem[1], Elem[2] };
-						Barys[IndexToWrite] = { CollisionPoint.Bary[1], CollisionPoint.Bary[2], CollisionPoint.Bary[3] };
+						Barys[IndexToWrite] = Bary;
 						FlipNormal[IndexToWrite] = bFlipNormal;
+						ConstraintTypes[IndexToWrite] = ConstraintType;
+						if (ConstraintType == EConstraintType::Kinematic)
+						{
+							// Remember this constraint
+							ExistingConstraintLookup[i].Add(FullMeshTriangleIndex, {0.f});
+						}
+						++ConstraintsAdded;
+					}
+				}
+				
+				for (TMap<int32, FExistingConstraintData>::TConstIterator PrevData = PrevExistingConstraintLookup.CreateConstIterator(); PrevData && ConstraintsAdded < MaxConnectionsPerPoint; ++PrevData)
+				{
+					if (PrevData.Value().Timer + Dt <= Chaos_CollisionSpring_MaxTimer)
+					{
+						const int32 SubMeshIndex = CollidableSubMesh.SubMeshIndexFromFullIndex(PrevData.Key());
+						if (SubMeshIndex != INDEX_NONE)
+						{
+							const int32 IndexToWrite = ConstraintIndex.fetch_add(1);
+							const TVec3<int32>& Elem = Elements[SubMeshIndex];
+							Constraints[IndexToWrite] = { Index, Elem[0], Elem[1], Elem[2] };
+							Barys[IndexToWrite] = FSolverVec3(1.f, 0.f, 0.f); // Unused
+							FlipNormal[IndexToWrite] = true;
+							ConstraintTypes[IndexToWrite] = EConstraintType::Kinematic;
+							ExistingConstraintLookup[i].Add(PrevData.Key(), {PrevData.Value().Timer + Dt});
+							++ConstraintsAdded;
+						}
 					}
 				}
 			}
@@ -282,20 +341,23 @@ void FPBDCollisionSpringConstraintsBase::Init(const SolverParticlesOrRange& Part
 		Constraints.SetNum(ConstraintNum, EAllowShrinking::No);
 		Barys.SetNum(ConstraintNum, EAllowShrinking::No);
 		FlipNormal.SetNum(ConstraintNum, EAllowShrinking::No);
+		ConstraintTypes.SetNum(ConstraintNum, EAllowShrinking::No);
 	}
+	PRAGMA_ENABLE_DEPRECATION_WARNINGS
 }
-template void CHAOS_API FPBDCollisionSpringConstraintsBase::Init<FTriangleMesh::TBVHType<FSolverReal>>(const FSolverParticles& Particles, const FTriangleMesh::TBVHType<FSolverReal>& Spatial, 
+template void CHAOS_API FPBDCollisionSpringConstraintsBase::Init<FTriangleMesh::TBVHType<FSolverReal>>(const FSolverParticles& Particles, const FSolverReal Dt, const FPBDTriangleMeshCollisions::FTriangleSubMesh& CollidableSubMesh, const FTriangleMesh::TBVHType<FSolverReal>& Spatial,
 	const TConstArrayView<FPBDTriangleMeshCollisions::FGIAColor>& VertexGIAColors, const TArray<FPBDTriangleMeshCollisions::FGIAColor>& TriangleGIAColors);
-template void CHAOS_API FPBDCollisionSpringConstraintsBase::Init<FTriangleMesh::TSpatialHashType<FSolverReal>>(const FSolverParticles& Particles, const FTriangleMesh::TSpatialHashType<FSolverReal>& Spatial, 
+template void CHAOS_API FPBDCollisionSpringConstraintsBase::Init<FTriangleMesh::TSpatialHashType<FSolverReal>>(const FSolverParticles& Particles, const FSolverReal Dt, const FPBDTriangleMeshCollisions::FTriangleSubMesh& CollidableSubMesh, const FTriangleMesh::TSpatialHashType<FSolverReal>& Spatial,
 	const TConstArrayView<FPBDTriangleMeshCollisions::FGIAColor>& VertexGIAColors, const TArray<FPBDTriangleMeshCollisions::FGIAColor>& TriangleGIAColors);
-template void CHAOS_API FPBDCollisionSpringConstraintsBase::Init<FTriangleMesh::TBVHType<FSolverReal>>(const FSolverParticlesRange& Particles, const FTriangleMesh::TBVHType<FSolverReal>& Spatial,
+template void CHAOS_API FPBDCollisionSpringConstraintsBase::Init<FTriangleMesh::TBVHType<FSolverReal>>(const FSolverParticlesRange& Particles, const FSolverReal Dt, const FPBDTriangleMeshCollisions::FTriangleSubMesh& CollidableSubMesh, const FTriangleMesh::TBVHType<FSolverReal>& Spatial,
 	const TConstArrayView<FPBDTriangleMeshCollisions::FGIAColor>& VertexGIAColors, const TArray<FPBDTriangleMeshCollisions::FGIAColor>& TriangleGIAColors);
-template void CHAOS_API FPBDCollisionSpringConstraintsBase::Init<FTriangleMesh::TSpatialHashType<FSolverReal>>(const FSolverParticlesRange& Particles, const FTriangleMesh::TSpatialHashType<FSolverReal>& Spatial,
+template void CHAOS_API FPBDCollisionSpringConstraintsBase::Init<FTriangleMesh::TSpatialHashType<FSolverReal>>(const FSolverParticlesRange& Particles, const FSolverReal Dt, const FPBDTriangleMeshCollisions::FTriangleSubMesh& CollidableSubMesh, const FTriangleMesh::TSpatialHashType<FSolverReal>& Spatial,
 	const TConstArrayView<FPBDTriangleMeshCollisions::FGIAColor>& VertexGIAColors, const TArray<FPBDTriangleMeshCollisions::FGIAColor>& TriangleGIAColors);
 
 template<typename SolverParticlesOrRange>
 FSolverVec3 FPBDCollisionSpringConstraintsBase::GetDelta(const SolverParticlesOrRange& Particles, const int32 ConstraintIndex) const
 {
+	PRAGMA_DISABLE_DEPRECATION_WARNINGS
 	const TVec4<int32>& Constraint = Constraints[ConstraintIndex];
 	const int32 Index1 = Constraint[0];
 	const int32 Index2 = Constraint[1];
@@ -318,31 +380,56 @@ FSolverVec3 FPBDCollisionSpringConstraintsBase::GetDelta(const SolverParticlesOr
 	const FSolverVec3& P3 = Particles.P(Index3);
 	const FSolverVec3& P4 = Particles.P(Index4);
 
-	const FSolverReal Height = Thickness + Thickness;
-	const FSolverVec3 P = Barys[ConstraintIndex][0] * P2 + Barys[ConstraintIndex][1] * P3 + Barys[ConstraintIndex][2] * P4;
-	const FSolverVec3 Difference = P1 - P;
+	const FSolverReal Height = GetConstraintThickness(ConstraintIndex);
 
-	// Normal repulsion with friction
 	const TTriangle<FSolverReal> Triangle(P2, P3, P4);
 	const FSolverVec3 Normal = FlipNormal[ConstraintIndex] ? -Triangle.GetNormal() : Triangle.GetNormal();
 
-	const FSolverReal NormalDifference = Difference.Dot(Normal);
+	FSolverVec3 P;
+	FSolverVec3 Difference;
+	FSolverReal NormalDifference;
+	FSolverReal TangentialFalloff = 1.f;
+	FSolverVec3 Bary;
+	if (ConstraintTypes[ConstraintIndex] == EConstraintType::Kinematic)
+	{
+		P = FindClosestPointAndBaryOnTriangle(P2, P3, P4, P1, Bary);
+		Difference = P1 - P;
+		NormalDifference = Difference.Dot(Normal);
+		const FSolverReal TangentialDifference = (Difference - NormalDifference * Normal).Size();
+		TangentialFalloff = 1.f - TangentialDifference / FMath::Max(1.5f * Height, 1e-4f);
+		if (TangentialFalloff <= 0.f)
+		{
+			return FSolverVec3(0);
+		}
+	}
+	else
+	{
+		Bary = Barys[ConstraintIndex];
+		P = Bary[0] * P2 + Bary[1] * P3 + Bary[2] * P4;
+		Difference = P1 - P;
+		NormalDifference = Difference.Dot(Normal);
+	}
+
+	// Normal repulsion with friction
 	if (NormalDifference > Height)
 	{
 		return FSolverVec3(0);
 	}
 
+	const FSolverReal ConstraintStiffness = GetConstraintStiffness(ConstraintIndex) * TangentialFalloff;
+	const FSolverReal ConstraintFriction = GetConstraintFrictionCoefficient(ConstraintIndex);
+
 	const FSolverReal NormalDelta = Height - NormalDifference;
-	const FSolverVec3 RepulsionDelta = Stiffness * NormalDelta * Normal / CombinedMass;
+	const FSolverVec3 RepulsionDelta = ConstraintStiffness * NormalDelta * Normal / CombinedMass;
 
-	if (FrictionCoefficient > 0)
+	if (ConstraintFriction > 0)
 	{
 		const FSolverVec3& X1 = Particles.X(Index1);
-		const FSolverVec3 X = Barys[ConstraintIndex][0] * Particles.X(Index2) + Barys[ConstraintIndex][1] * Particles.X(Index3) + Barys[ConstraintIndex][2] * Particles.X(Index4);
+		const FSolverVec3 X = Bary[0] * Particles.X(Index2) + Bary[1] * Particles.X(Index3) + Bary[2] * Particles.X(Index4);
 		const FSolverVec3 RelativeDisplacement = (P1 - X1) - (P - X) + (Particles.InvM(Index1) - TrianglePointInvM) * RepulsionDelta;
 		const FSolverVec3 RelativeDisplacementTangent = RelativeDisplacement - RelativeDisplacement.Dot(Normal) * Normal;
 		const FSolverReal RelativeDisplacementTangentLength = RelativeDisplacementTangent.Length();
-		const FSolverReal PositionCorrection = FMath::Min(NormalDelta * FrictionCoefficient, RelativeDisplacementTangentLength);
+		const FSolverReal PositionCorrection = FMath::Min(NormalDelta * ConstraintFriction, RelativeDisplacementTangentLength);
 		const FSolverReal CorrectionRatio = RelativeDisplacementTangentLength < UE_SMALL_NUMBER ? 0.f : PositionCorrection / RelativeDisplacementTangentLength;
 		const FSolverVec3 FrictionDelta = -CorrectionRatio * RelativeDisplacementTangent / CombinedMass;
 		return RepulsionDelta + FrictionDelta;
@@ -351,12 +438,14 @@ FSolverVec3 FPBDCollisionSpringConstraintsBase::GetDelta(const SolverParticlesOr
 	{
 		return RepulsionDelta;
 	}
+	PRAGMA_ENABLE_DEPRECATION_WARNINGS
 }
 template CHAOS_API FSolverVec3 FPBDCollisionSpringConstraintsBase::GetDelta(const FSolverParticles& Particles, const int32 i) const;
 template CHAOS_API FSolverVec3 FPBDCollisionSpringConstraintsBase::GetDelta(const FSolverParticlesRange& Particles, const int32 i) const;
 
 void FPBDCollisionSpringConstraintsBase::UpdateLinearSystem(const FSolverParticlesRange& Particles, const FSolverReal Dt, FEvolutionLinearSystem& LinearSystem) const
 {
+	PRAGMA_DISABLE_DEPRECATION_WARNINGS
 	LinearSystem.ReserveForParallelAdd(Constraints.Num() * 4, Constraints.Num() * 3);
 	for (int32 Index = 0; Index < Constraints.Num(); ++Index)
 	{
@@ -370,7 +459,7 @@ void FPBDCollisionSpringConstraintsBase::UpdateLinearSystem(const FSolverParticl
 		const FSolverVec3& P3 = Particles.P(Index3);
 		const FSolverVec3& P4 = Particles.P(Index4);
 
-		const FSolverReal Height = Thickness + Thickness;
+		const FSolverReal Height = GetConstraintThickness(Index);
 		const FSolverVec3 P = Barys[Index][0] * P2 + Barys[Index][1] * P3 + Barys[Index][2] * P4;
 		const FSolverVec3 Difference = P1 - P;
 
@@ -386,11 +475,15 @@ void FPBDCollisionSpringConstraintsBase::UpdateLinearSystem(const FSolverParticl
 
 		const FSolverReal NormalDelta = Height - NormalDifference;
 
+		const FSolverReal ConstraintFriction = GetConstraintFrictionCoefficient(Index);
 		const FSolverVec3 Force = ProximityStiffness * NormalDelta * Normal;
-		const FSolverMatrix33 DfDx = -ProximityStiffness * (((FSolverReal)1. - FrictionCoefficient) * FSolverMatrix33::OuterProduct(Normal, Normal) + FSolverMatrix33(FrictionCoefficient, FrictionCoefficient, FrictionCoefficient));
+		const FSolverMatrix33 DfDx = -ProximityStiffness * (((FSolverReal)1. - ConstraintFriction) * FSolverMatrix33::OuterProduct(Normal, Normal) + FSolverMatrix33(ConstraintFriction, ConstraintFriction, ConstraintFriction));
 
-		LinearSystem.AddForce(Particles, Force, Index1, Dt);
-		LinearSystem.AddSymmetricForceDerivative(Particles, &DfDx, nullptr, Index1, Index1, Dt);
+		if (Particles.InvM(Index1) > (FSolverReal)0.)
+		{
+			LinearSystem.AddForce(Particles, Force, Index1, Dt);
+			LinearSystem.AddSymmetricForceDerivative(Particles, &DfDx, nullptr, Index1, Index1, Dt);
+		}
 		if (Particles.InvM(Index2) > (FSolverReal)0.)
 		{
 			LinearSystem.AddForce(Particles, -Barys[Index][0] * Force, Index2, Dt);
@@ -416,6 +509,7 @@ void FPBDCollisionSpringConstraintsBase::UpdateLinearSystem(const FSolverParticl
 			LinearSystem.AddSymmetricForceDerivative(Particles, &DfDxScaled, nullptr, Index4, Index4, Dt);
 		}
 	}
+	PRAGMA_ENABLE_DEPRECATION_WARNINGS
 }
 }  // End namespace Chaos::Softs
 
diff --git a/Engine/Source/Runtime/Experimental/Chaos/Private/Chaos/PBDTriangleMeshCollisions.cpp b/Engine/Source/Runtime/Experimental/Chaos/Private/Chaos/PBDTriangleMeshCollisions.cpp
--- a/Engine/Source/Runtime/Experimental/Chaos/Private/Chaos/PBDTriangleMeshCollisions.cpp
+++ b/Engine/Source/Runtime/Experimental/Chaos/Private/Chaos/PBDTriangleMeshCollisions.cpp
@@ -24,7 +24,7 @@ struct FEdgeFaceIntersection
 
 // Returned array has NOT been shrunk
 template<typename SpatialAccelerator, typename SolverParticlesOrRange>
-static void FindEdgeFaceIntersections(const FTriangleMesh& TriangleMesh, const SpatialAccelerator& Spatial, const SolverParticlesOrRange& Particles, TArray<FEdgeFaceIntersection>& Intersections)
+static void FindEdgeFaceIntersections(const FTriangleMesh& TriangleMesh, const TArray<int32>& IntersectableEdges, bool bSkipKinematic, const TArray<bool>& FaceIsKinematic, const SpatialAccelerator& Spatial, const SolverParticlesOrRange& Particles, TArray<FEdgeFaceIntersection>& Intersections)
 {
 	TRACE_CPUPROFILER_EVENT_SCOPE(ChaosFPBDTriangleMeshCollisions_IntersectionQuery);
 
@@ -32,9 +32,11 @@ static void FindEdgeFaceIntersections(const FTriangleMesh& TriangleMesh, const S
 	const TArray<TVec2<int32>>& EdgeToFaces = TriangleMesh.GetEdgeToFaces();
 	const TArray<TVec3<int32>>& Elements = TriangleMesh.GetElements();
 
+	const int32 NumIntersectableEdges = IntersectableEdges.IsEmpty() ? SegmentMesh.GetNumElements() : IntersectableEdges.Num();
+
 	// Preallocate enough space for (more than) typical number of expected intersections.
 	constexpr int32 PreallocatedIntersectionsPerEdge = 3;
-	const int32 PreallocatedIntersectionsNum = PreallocatedIntersectionsPerEdge * SegmentMesh.GetNumElements();
+	const int32 PreallocatedIntersectionsNum = PreallocatedIntersectionsPerEdge * NumIntersectableEdges;
 	Intersections.Reset();
 	Intersections.SetNumUninitialized(PreallocatedIntersectionsNum);
 	std::atomic<int32> IntersectionIndex(0);
@@ -42,9 +44,10 @@ static void FindEdgeFaceIntersections(const FTriangleMesh& TriangleMesh, const S
 	// Extra intersections that require a lock to write to if you have more than PreallocatedIntersectionsPerEdge 
 	TArray<FEdgeFaceIntersection> ExtraIntersections;
 	FCriticalSection CriticalSection;
-	PhysicsParallelFor(SegmentMesh.GetNumElements(),
-		[&Spatial, &TriangleMesh, &Particles, &SegmentMesh, &EdgeToFaces, &Elements, &IntersectionIndex, &Intersections, PreallocatedIntersectionsNum, &ExtraIntersections, &CriticalSection](int32 EdgeIndex)
+	PhysicsParallelFor(NumIntersectableEdges,
+		[&Spatial, &TriangleMesh, &Particles, &SegmentMesh, &EdgeToFaces, &Elements, &IntersectionIndex, &Intersections, PreallocatedIntersectionsNum, &ExtraIntersections, &CriticalSection, &IntersectableEdges, bSkipKinematic, &FaceIsKinematic](int32 IntersectableEdgeIndex)
 		{
+			const int32 EdgeIndex = IntersectableEdges.IsEmpty() ? IntersectableEdgeIndex : IntersectableEdges[IntersectableEdgeIndex];
 
 			TArray< TTriangleCollisionPoint<FSolverReal> > Result;
 
@@ -52,16 +55,25 @@ static void FindEdgeFaceIntersections(const FTriangleMesh& TriangleMesh, const S
 			const int32 EdgePointIndex1 = SegmentMesh.GetElements()[EdgeIndex][1];
 			const FSolverVec3& EdgePosition0 = Particles.X(EdgePointIndex0);
 			const FSolverVec3& EdgePosition1 = Particles.X(EdgePointIndex1);
-
+			const bool bEdgeIsKinematic = Particles.InvM(EdgePointIndex0) == (FSolverReal)0. && Particles.InvM(EdgePointIndex1) == (FSolverReal)0.;
+			if (bSkipKinematic && bEdgeIsKinematic)
+			{
+				return;
+			}
 
 			if (TriangleMesh.EdgeIntersectionQuery(Spatial, static_cast<const TArrayView<const FSolverVec3>&>(Particles.XArray()), EdgeIndex, EdgePosition0, EdgePosition1,
-				[&Elements, EdgePointIndex0, EdgePointIndex1](int32 EdgeIndex, int32 TriangleIndex)
+				[&Elements, EdgePointIndex0, EdgePointIndex1, bSkipKinematic, bEdgeIsKinematic, &FaceIsKinematic,  &Particles](int32 EdgeIndex, int32 TriangleIndex)
 				{
 					if (EdgePointIndex0 == Elements[TriangleIndex][0] || EdgePointIndex0 == Elements[TriangleIndex][1] || EdgePointIndex0 == Elements[TriangleIndex][2] ||
 						EdgePointIndex1 == Elements[TriangleIndex][0] || EdgePointIndex1 == Elements[TriangleIndex][1] || EdgePointIndex1 == Elements[TriangleIndex][2])
 					{
 						return false;
 					}
+
+					if((bSkipKinematic || bEdgeIsKinematic) && FaceIsKinematic[TriangleIndex])
+					{
+						return false;
+					}
 					return true;
 				},
 				Result))
@@ -1274,10 +1286,110 @@ struct FPBDTriangleMeshCollisions::FScratchBuffers
 	}
 };
 
+template<typename SolverParticlesOrRange>
+void FPBDTriangleMeshCollisions::FTriangleSubMesh::Init(const SolverParticlesOrRange& Particles, const TSet<int32>& InDisabledFaces, bool bCollideAgainstAllKinematicVertices, const TSet<int32>& InEnabledKinematicFaces)
+{
+	SubMesh.Init(TArray<TVec3<int32>>());
+	SubMeshToFullElements.Reset();
+	FullToSubMeshElements.Reset();
+	SubMeshElementIsKinematic.Reset();
+	CollidableVertices.Reset();
+	IntersectableSubmeshEdges.Reset();
+	bSubMeshIsFullMesh = true;
+
+	if (InDisabledFaces.IsEmpty() && bCollideAgainstAllKinematicVertices)
+	{
+		const TArray<TVec3<int32>>& FullElements = FullMesh.GetElements();
+		SubMeshElementIsKinematic.SetNumUninitialized(FullMesh.GetNumElements());
+
+		for (int32 FullElementIndex = 0; FullElementIndex < FullMesh.GetNumElements(); ++FullElementIndex)
+		{
+			const bool bIsKinematic =
+				Particles.InvM(FullElements[FullElementIndex][0]) == (FSolverReal)0.f &&
+				Particles.InvM(FullElements[FullElementIndex][1]) == (FSolverReal)0.f &&
+				Particles.InvM(FullElements[FullElementIndex][2]) == (FSolverReal)0.f;
+			SubMeshElementIsKinematic[FullElementIndex] = bIsKinematic;
+		}
+		return;
+	}
+
+	// Build a sub mesh of just colliding faces.
+	bSubMeshIsFullMesh = false;
+	TArray<TVec3<int32>> SubElements;
+	SubElements.Reserve(FullMesh.GetNumElements());
+	FullToSubMeshElements.SetNumUninitialized(FullMesh.GetNumElements());
+	SubMeshToFullElements.Reserve(FullMesh.GetNumElements());
+	SubMeshElementIsKinematic.Reserve(FullMesh.GetNumElements());
+	const TArray<TVec3<int32>>& FullElements = FullMesh.GetElements();
+	for (int32 FullElementIndex = 0; FullElementIndex < FullMesh.GetNumElements(); ++FullElementIndex)
+	{
+		if (InDisabledFaces.Contains(FullElementIndex))
+		{
+			FullToSubMeshElements[FullElementIndex] = INDEX_NONE;
+			continue;
+		}
+		const bool bIsKinematic =
+			Particles.InvM(FullElements[FullElementIndex][0]) == (FSolverReal)0.f &&
+			Particles.InvM(FullElements[FullElementIndex][1]) == (FSolverReal)0.f &&
+			Particles.InvM(FullElements[FullElementIndex][2]) == (FSolverReal)0.f;
+
+		if (!bCollideAgainstAllKinematicVertices &&
+			bIsKinematic &&
+			!InEnabledKinematicFaces.Contains(FullElementIndex))
+		{
+			FullToSubMeshElements[FullElementIndex] = INDEX_NONE;
+			continue;
+		}
+
+		const int32 SubElementIndex = SubElements.Add(FullElements[FullElementIndex]);
+		SubMeshToFullElements.Add(FullElementIndex);
+		SubMeshElementIsKinematic.Add(bIsKinematic);
+		FullToSubMeshElements[FullElementIndex] = SubElementIndex;
+	}
+
+	// Use same Vertex range for submesh.
+	const TVec2<int32> VertexRange = FullMesh.GetVertexRange();
+	constexpr bool bCullDegenerateFalse = false;
+	SubMesh.Init(MoveTemp(SubElements), VertexRange[0], VertexRange[1], bCullDegenerateFalse);
+	check(SubMesh.GetNumElements() == SubMeshToFullElements.Num());
+	check(SubMesh.GetNumElements() == SubMeshElementIsKinematic.Num());
+
+	const TSet<int32> VertexSet = SubMesh.GetVertices();
+	CollidableVertices.Reserve(VertexSet.Num());
+	for (const int32 Vertex : VertexSet)
+	{
+		if (Particles.InvM(Vertex) > (FSolverReal)0.)
+		{
+			CollidableVertices.Add(Vertex);
+		}
+	}
+
+	const FSegmentMesh& SegmentMesh = const_cast<const FTriangleMesh&>(SubMesh).GetSegmentMesh();
+	IntersectableSubmeshEdges.Reserve(SegmentMesh.GetNumElements());
+	for (int32 EdgeIndex = 0; EdgeIndex < SegmentMesh.GetNumElements(); ++EdgeIndex)
+	{
+		const int32 EdgePointIndex0 = SegmentMesh.GetElements()[EdgeIndex][0];
+		const int32 EdgePointIndex1 = SegmentMesh.GetElements()[EdgeIndex][1];
+		if (Particles.InvM(EdgePointIndex0) > (FSolverReal)0. || Particles.InvM(EdgePointIndex1) > (FSolverReal)0)
+		{
+			IntersectableSubmeshEdges.Add(EdgeIndex);
+		}
+	}
+}
+
 template<typename SolverParticlesOrRange>
 void FPBDTriangleMeshCollisions::Init(const SolverParticlesOrRange& Particles, const FSolverReal MinProximityQueryRadius)
 {
-	if (TriangleMesh.GetNumElements() == 0)
+	const bool bDoSelfIntersections = bGlobalIntersectionAnalysis || bContourMinimization;
+	if (bCollidableSubMeshDirty)
+	{		
+		CollidableSubMesh.Init(Particles, DisabledFaces, bSelfCollideAgainstAllKinematicVertices, EnabledKinematicFaces);
+		bCollidableSubMeshDirty = false;
+	}
+
+	const FTriangleMesh& CollidableMesh = CollidableSubMesh.GetCollidableMesh();
+
+	if (CollidableMesh.GetNumElements() == 0)
 	{
 		return;
 	}
@@ -1285,17 +1397,17 @@ void FPBDTriangleMeshCollisions::Init(const SolverParticlesOrRange& Particles, c
 	{
 		TRACE_CPUPROFILER_EVENT_SCOPE(ChaosFPBDTriangleMeshCollisions_BuildSpatialHash);
 		constexpr FSolverReal RadiusToLodSizeMultiplier = 2.; // Radius to Diameter
-		TriangleMesh.BuildSpatialHash(static_cast<const TArrayView<const FSolverVec3>&>(Particles.XArray()), SpatialHash, RadiusToLodSizeMultiplier * MinProximityQueryRadius);
+		CollidableMesh.BuildSpatialHash(static_cast<const TArrayView<const FSolverVec3>&>(Particles.XArray()), SpatialHash, RadiusToLodSizeMultiplier * MinProximityQueryRadius);
 	}
 	ContourMinimizationIntersections.Reset();
 	VertexGIAColors.Reset();
 	VertexGIAColors.SetNumZeroed(NumParticles);
 	TriangleGIAColors.Reset();
-	TriangleGIAColors.SetNumZeroed(TriangleMesh.GetNumElements());
+	TriangleGIAColors.SetNumZeroed(CollidableMesh.GetNumElements());
 	IntersectionContourPoints.Reset();
 	IntersectionContourTypes.Reset();
 
-	if (!bGlobalIntersectionAnalysis && !bContourMinimization)
+	if (!bDoSelfIntersections)
 	{
 		return;
 	}
@@ -1307,7 +1419,8 @@ void FPBDTriangleMeshCollisions::Init(const SolverParticlesOrRange& Particles, c
 	ScratchBuffers->Reset();
 
 	// Detect all EdgeFace Intersections
-	FindEdgeFaceIntersections(TriangleMesh, SpatialHash, Particles, ScratchBuffers->EdgeFaceIntersections);
+	constexpr bool bSkipKinematicTrue = true;
+	FindEdgeFaceIntersections(CollidableMesh, CollidableSubMesh.GetIntersectableSubmeshEdges(), bSkipKinematicTrue, CollidableSubMesh.GetSubMeshElementIsKinematic(), SpatialHash, Particles, ScratchBuffers->EdgeFaceIntersections);
 	if (ScratchBuffers->EdgeFaceIntersections.Num() == 0)
 	{
 		return;
@@ -1316,9 +1429,9 @@ void FPBDTriangleMeshCollisions::Init(const SolverParticlesOrRange& Particles, c
 	if (bGlobalIntersectionAnalysis)
 	{ 
 		// Walk EdgeFace intersections to build global contours
-		GIA::ContourBuilding::BuildIntersectionContours(TriangleMesh, ScratchBuffers->EdgeFaceIntersections, ScratchBuffers->IntersectionContours, IntersectionContourPoints);
+		GIA::ContourBuilding::BuildIntersectionContours(CollidableMesh, ScratchBuffers->EdgeFaceIntersections, ScratchBuffers->IntersectionContours, IntersectionContourPoints);
 		// Flood fill global contours to determine intersecting regions.
-		GIA::FloodFill::FloodFillContours(TriangleMesh, NumParticles, Offset, ScratchBuffers->IntersectionContours, VertexGIAColors, TriangleGIAColors);
+		GIA::FloodFill::FloodFillContours(CollidableMesh, NumParticles, Offset, ScratchBuffers->IntersectionContours, VertexGIAColors, TriangleGIAColors);
 		GIA::FloodFill::AssignContourPointTypes(ScratchBuffers->IntersectionContours, IntersectionContourPoints, IntersectionContourTypes);
 	}
 
@@ -1327,12 +1440,12 @@ void FPBDTriangleMeshCollisions::Init(const SolverParticlesOrRange& Particles, c
 		if (bGlobalIntersectionAnalysis)
 		{
 			// Global contours which are non-closed or loop are handled via ContourMinimization impulses. Build global gradient (by adding contribution across all intersections per contour).
-			ContourMinimization::BuildGlobalContourMinimizationIntersections(TriangleMesh, Particles, ScratchBuffers->IntersectionContours, ContourMinimizationIntersections);
+			ContourMinimization::BuildGlobalContourMinimizationIntersections(CollidableMesh, Particles, ScratchBuffers->IntersectionContours, ContourMinimizationIntersections);
 		}
 		else
 		{
 			// Just build local gradient for all Intersections
-			ContourMinimization::BuildLocalContourMinimizationIntersections(TriangleMesh, Particles, ScratchBuffers->EdgeFaceIntersections, ContourMinimizationIntersections);
+			ContourMinimization::BuildLocalContourMinimizationIntersections(CollidableMesh, Particles, ScratchBuffers->EdgeFaceIntersections, ContourMinimizationIntersections);
 		}
 	}
 }
@@ -1342,7 +1455,8 @@ template CHAOS_API void FPBDTriangleMeshCollisions::Init(const FSolverParticlesR
 template<typename SolverParticlesOrRange>
 void FPBDTriangleMeshCollisions::PostStepInit(const SolverParticlesOrRange& Particles)
 {
-	if (TriangleMesh.GetNumElements() == 0)
+	const FTriangleMesh& CollidableMesh = CollidableSubMesh.GetCollidableMesh();
+	if (CollidableMesh.GetNumElements() == 0)
 	{
 		return;
 	}
@@ -1356,7 +1470,7 @@ void FPBDTriangleMeshCollisions::PostStepInit(const SolverParticlesOrRange& Part
 		{
 			TRACE_CPUPROFILER_EVENT_SCOPE(ChaosFPBDTriangleMeshCollisions_BuildSpatialHash);
 			constexpr FSolverReal MinSpatialLodSize = 1.f;
-			TriangleMesh.BuildSpatialHash(static_cast<const TArrayView<const FSolverVec3>&>(Particles.XArray()), SpatialHash, MinSpatialLodSize);
+			CollidableMesh.BuildSpatialHash(static_cast<const TArrayView<const FSolverVec3>&>(Particles.XArray()), SpatialHash, MinSpatialLodSize);
 
 			if (!ScratchBuffers)
 			{
@@ -1365,17 +1479,18 @@ void FPBDTriangleMeshCollisions::PostStepInit(const SolverParticlesOrRange& Part
 			ScratchBuffers->Reset();
 
 			// Detect all EdgeFace Intersections
-			FindEdgeFaceIntersections(TriangleMesh, SpatialHash, Particles, ScratchBuffers->EdgeFaceIntersections);
+			constexpr bool bSkipKinematicTrue = true;
+			FindEdgeFaceIntersections(CollidableMesh, CollidableSubMesh.GetIntersectableSubmeshEdges(), bSkipKinematicTrue, CollidableSubMesh.GetSubMeshElementIsKinematic(), SpatialHash, Particles, ScratchBuffers->EdgeFaceIntersections);
 
 			if (bUseGlobalPostStepContours)
 			{
-				GIA::ContourBuilding::BuildIntersectionContours(TriangleMesh, ScratchBuffers->EdgeFaceIntersections, ScratchBuffers->IntersectionContours, PostStepIntersectionContourPoints);
-				ContourMinimization::BuildGlobalContourMinimizationIntersections(TriangleMesh, Particles, ScratchBuffers->IntersectionContours, PostStepContourMinimizationIntersections);
+				GIA::ContourBuilding::BuildIntersectionContours(CollidableMesh, ScratchBuffers->EdgeFaceIntersections, ScratchBuffers->IntersectionContours, PostStepIntersectionContourPoints);
+				ContourMinimization::BuildGlobalContourMinimizationIntersections(CollidableMesh, Particles, ScratchBuffers->IntersectionContours, PostStepContourMinimizationIntersections);
 			}
 			else
 			{
 				// Just build local gradient for all Intersections
-				ContourMinimization::BuildLocalContourMinimizationIntersections(TriangleMesh, Particles, ScratchBuffers->EdgeFaceIntersections, PostStepContourMinimizationIntersections);
+				ContourMinimization::BuildLocalContourMinimizationIntersections(CollidableMesh, Particles, ScratchBuffers->EdgeFaceIntersections, PostStepContourMinimizationIntersections);
 			}
 		}
 	}
diff --git a/Engine/Source/Runtime/Experimental/Chaos/Private/Chaos/TriangleMesh.cpp b/Engine/Source/Runtime/Experimental/Chaos/Private/Chaos/TriangleMesh.cpp
--- a/Engine/Source/Runtime/Experimental/Chaos/Private/Chaos/TriangleMesh.cpp
+++ b/Engine/Source/Runtime/Experimental/Chaos/Private/Chaos/TriangleMesh.cpp
@@ -122,6 +122,7 @@ void FTriangleMesh::Init(TArray<TVec3<int32>>&& Elements, const int32 StartIdx,
 	MElements = MoveTemp(Elements);
 	MStartIdx = 0;
 	MNumIndices = 0;
+	ResetAuxiliaryStructures();
 	InitHelper(StartIdx, EndIdx, CullDegenerateElements);
 }
 
@@ -130,6 +131,7 @@ void FTriangleMesh::Init(const TArray<TVec3<int32>>& Elements, const int32 Start
 	MElements = Elements;
 	MStartIdx = 0;
 	MNumIndices = 0;
+	ResetAuxiliaryStructures();
 	InitHelper(StartIdx, EndIdx, CullDegenerateElements);
 }
 
diff --git a/Engine/Source/Runtime/Experimental/Chaos/Public/Chaos/PBDCollisionSpringConstraints.h b/Engine/Source/Runtime/Experimental/Chaos/Public/Chaos/PBDCollisionSpringConstraints.h
--- a/Engine/Source/Runtime/Experimental/Chaos/Public/Chaos/PBDCollisionSpringConstraints.h
+++ b/Engine/Source/Runtime/Experimental/Chaos/Public/Chaos/PBDCollisionSpringConstraints.h
@@ -11,8 +11,6 @@ namespace Chaos::Softs
 class FPBDCollisionSpringConstraints : public FPBDCollisionSpringConstraintsBase
 {
 	typedef FPBDCollisionSpringConstraintsBase Base;
-	using Base::Barys;
-	using Base::Constraints;
 
 public:
 	static constexpr FSolverReal MinFrictionCoefficient = (FSolverReal)0.;
@@ -40,12 +38,18 @@ public:
 			(FSolverReal)FMath::Max(GetSelfCollisionThickness(PropertyCollection, Base::BackCompatThickness), 0.f),
 			(FSolverReal)FMath::Clamp(GetSelfCollisionStiffness(PropertyCollection, Base::BackCompatStiffness), 0.f, 1.f),
 			FMath::Clamp((FSolverReal)GetSelfCollisionFriction(PropertyCollection, Base::BackCompatFrictionCoefficient), MinFrictionCoefficient, MaxFrictionCoefficient),
+			(FSolverReal)FMath::Max(GetSelfCollisionKinematicColliderThickness(PropertyCollection, Base::DefaultKinematicColliderThickness), 0.f),
+			(FSolverReal)FMath::Clamp((FSolverReal)GetSelfCollisionKinematicColliderStiffness(PropertyCollection, Base::DefaultKinematicColliderStiffness), 0.f, 1.f),
+			FMath::Clamp((FSolverReal)GetSelfCollisionKinematicColliderFriction(PropertyCollection, Base::DefaultKinematicColliderFrictionCoefficient), MinFrictionCoefficient, MaxFrictionCoefficient),
 			(FSolverReal)GetSelfCollisionProximityStiffness(PropertyCollection, Base::DefaultProximityStiffness))
 		, SelfCollisionThicknessIndex(PropertyCollection)
 		, SelfCollisionStiffnessIndex(PropertyCollection)
 		, SelfCollisionFrictionIndex(PropertyCollection)
 		, SelfCollisionProximityStiffnessIndex(PropertyCollection)
 		, SelfCollisionLayersIndex(PropertyCollection)
+		, SelfCollisionKinematicColliderThicknessIndex(PropertyCollection)
+		, SelfCollisionKinematicColliderStiffnessIndex(PropertyCollection)
+		, SelfCollisionKinematicColliderFrictionIndex(PropertyCollection)
 	{}
 
 	UE_DEPRECATED(5.4, "Use constructor with FaceIntMaps")
@@ -56,22 +60,14 @@ public:
 		const TArray<FSolverVec3>* InRestPositions,
 		TSet<TVec2<int32>>&& InDisabledCollisionElements,
 		const FCollectionPropertyConstFacade& PropertyCollection)
-		: Base(
+		: FPBDCollisionSpringConstraints(
 			InOffset,
 			InNumParticles,
 			InTriangleMesh,
 			InRestPositions,
-			MoveTemp(InDisabledCollisionElements),
-			TConstArrayView<int32>(),
-			(FSolverReal)FMath::Max(GetSelfCollisionThickness(PropertyCollection, Base::BackCompatThickness), 0.f),
-			(FSolverReal)FMath::Clamp(GetSelfCollisionStiffness(PropertyCollection, Base::BackCompatStiffness), 0.f, 1.f),
-			FMath::Clamp((FSolverReal)GetSelfCollisionFriction(PropertyCollection, Base::BackCompatFrictionCoefficient), MinFrictionCoefficient, MaxFrictionCoefficient),
-			(FSolverReal)GetSelfCollisionProximityStiffness(PropertyCollection, Base::DefaultProximityStiffness))
-		, SelfCollisionThicknessIndex(PropertyCollection)
-		, SelfCollisionStiffnessIndex(PropertyCollection)
-		, SelfCollisionFrictionIndex(PropertyCollection)
-		, SelfCollisionProximityStiffnessIndex(PropertyCollection)
-		, SelfCollisionLayersIndex(PropertyCollection)
+			TMap<FString, TConstArrayView<int32>>(),
+			PropertyCollection
+		)
 	{}
 
 	FPBDCollisionSpringConstraints(
@@ -98,6 +94,9 @@ public:
 		, SelfCollisionFrictionIndex(ForceInit)
 		, SelfCollisionProximityStiffnessIndex(ForceInit)
 		, SelfCollisionLayersIndex(ForceInit)
+		, SelfCollisionKinematicColliderThicknessIndex(ForceInit)
+		, SelfCollisionKinematicColliderStiffnessIndex(ForceInit)
+		, SelfCollisionKinematicColliderFrictionIndex(ForceInit)
 	{}
 
 	virtual ~FPBDCollisionSpringConstraints() override {}
@@ -131,6 +130,10 @@ public:
 				Base::UpdateCollisionLayers(FaceIntMaps.FindRef(FaceIntMapName));
 			}
 		}
+		if (IsSelfCollisionKinematicColliderStiffnessMutable(PropertyCollection))
+		{
+			KinematicColliderStiffness = FMath::Clamp((FSolverReal)GetSelfCollisionKinematicColliderStiffness(PropertyCollection), 0.f, 1.f);
+		}
 	}
 
 	UE_DEPRECATED(5.4, "Use SetProperties with FaceIntMaps")
@@ -172,6 +175,9 @@ private:
 	UE_CHAOS_DECLARE_PROPERTYCOLLECTION_NAME(SelfCollisionProximityStiffness, float);
 	UE_CHAOS_DECLARE_INDEXLESS_PROPERTYCOLLECTION_NAME(SelfCollisionDisableNeighborDistance, int32);
 	UE_CHAOS_DECLARE_PROPERTYCOLLECTION_NAME(SelfCollisionLayers, float);
+	UE_CHAOS_DECLARE_PROPERTYCOLLECTION_NAME(SelfCollisionKinematicColliderThickness, float);
+	UE_CHAOS_DECLARE_PROPERTYCOLLECTION_NAME(SelfCollisionKinematicColliderStiffness, float);
+	UE_CHAOS_DECLARE_PROPERTYCOLLECTION_NAME(SelfCollisionKinematicColliderFriction, float);
 };
 
 }  // End namespace Chaos::Softs
diff --git a/Engine/Source/Runtime/Experimental/Chaos/Public/Chaos/PBDCollisionSpringConstraintsBase.h b/Engine/Source/Runtime/Experimental/Chaos/Public/Chaos/PBDCollisionSpringConstraintsBase.h
--- a/Engine/Source/Runtime/Experimental/Chaos/Public/Chaos/PBDCollisionSpringConstraintsBase.h
+++ b/Engine/Source/Runtime/Experimental/Chaos/Public/Chaos/PBDCollisionSpringConstraintsBase.h
@@ -24,6 +24,9 @@ public:
 	static constexpr FSolverReal BackCompatThickness = (FSolverReal)1.f;
 	static constexpr FSolverReal BackCompatStiffness = (FSolverReal)0.5f;
 	static constexpr FSolverReal BackCompatFrictionCoefficient = (FSolverReal)0.f;
+	static constexpr FSolverReal DefaultKinematicColliderThickness = (FSolverReal)0.f;
+	static constexpr FSolverReal DefaultKinematicColliderStiffness = (FSolverReal)1.f;
+	static constexpr FSolverReal DefaultKinematicColliderFrictionCoefficient = (FSolverReal)0.f;
 	static constexpr FSolverReal DefaultProximityStiffness = (FSolverReal)1.;
 
 	CHAOS_API FPBDCollisionSpringConstraintsBase(
@@ -36,6 +39,9 @@ public:
 		const FSolverReal InThickness = BackCompatThickness,
 		const FSolverReal InStiffness = BackCompatStiffness,
 		const FSolverReal InFrictionCoefficient = BackCompatFrictionCoefficient,
+		const FSolverReal InKinematicColliderThickness = DefaultKinematicColliderThickness,
+		const FSolverReal InKinematicColliderStiffness = DefaultKinematicColliderStiffness,
+		const FSolverReal InKinematicColliderFrictionCoefficient = DefaultKinematicColliderFrictionCoefficient,
 		const FSolverReal InProximityStiffness = DefaultProximityStiffness);
 
 	UE_DEPRECATED(5.4, "Use constructor with SelfCollisionLayers")
@@ -47,28 +53,48 @@ public:
 		TSet<TVec2<int32>>&& InDisabledCollisionElements,
 		const FSolverReal InThickness = BackCompatThickness,
 		const FSolverReal InStiffness = BackCompatStiffness,
-		const FSolverReal InFrictionCoefficient = BackCompatFrictionCoefficient,
-		const FSolverReal InProximityStiffness = DefaultProximityStiffness)
+		const FSolverReal InFrictionCoefficient = BackCompatFrictionCoefficient)
 		: FPBDCollisionSpringConstraintsBase(InOffset, InNumParticles, InTriangleMesh, InReferencePositions, MoveTemp(InDisabledCollisionElements),
-			TConstArrayView<int32>(), InThickness, InStiffness, InFrictionCoefficient, InProximityStiffness)
+			TConstArrayView<int32>(), InThickness, InStiffness, InFrictionCoefficient)
 	{}
 
+	PRAGMA_DISABLE_DEPRECATION_WARNINGS
 	virtual ~FPBDCollisionSpringConstraintsBase() {}
+	PRAGMA_ENABLE_DEPRECATION_WARNINGS
 
-	UE_DEPRECATED(5.0, "Use Init(Particles, Spatial, GIAColors) instead.")
-	CHAOS_API void Init(const FSolverParticles& Particles);
+	template<typename SpatialAccelerator, typename SolverParticlesOrRange>
+	void Init(const SolverParticlesOrRange& Particles, const SpatialAccelerator& Spatial, const TConstArrayView<FPBDTriangleMeshCollisions::FGIAColor>& VertexGIAColors, const TArray<FPBDTriangleMeshCollisions::FGIAColor>& TriangleGIAColors)
+	{
+		constexpr FSolverReal Dt_BigNumber = UE_BIG_NUMBER; // This will disable the timers feature on kinematic colliders.
+		return Init(Particles, Dt_BigNumber, FPBDTriangleMeshCollisions::FTriangleSubMesh(TriangleMesh), Spatial, VertexGIAColors, TriangleGIAColors);
+	}
 
 	template<typename SpatialAccelerator, typename SolverParticlesOrRange>
-	void Init(const SolverParticlesOrRange& Particles, const SpatialAccelerator& Spatial, const TConstArrayView<FPBDTriangleMeshCollisions::FGIAColor>& VertexGIAColors, const TArray<FPBDTriangleMeshCollisions::FGIAColor>& TriangleGIAColors);
+	void Init(const SolverParticlesOrRange& Particles, const FSolverReal Dt, const FPBDTriangleMeshCollisions::FTriangleSubMesh& CollidableSubMesh, const SpatialAccelerator& Spatial, const TConstArrayView<FPBDTriangleMeshCollisions::FGIAColor>& VertexGIAColors, const TArray<FPBDTriangleMeshCollisions::FGIAColor>& TriangleGIAColors);
 
 	template<typename SolverParticlesOrRange>
 	CHAOS_API FSolverVec3 GetDelta(const SolverParticlesOrRange& InParticles, const int32 i) const;
 
-	const TArray<TVec4<int32>>& GetConstraints() const { return Constraints;  }
-	const TArray<FSolverVec3>& GetBarys() const { return Barys; }
+	const TArray<TVec4<int32>>& GetConstraints() const 
+	{
+		PRAGMA_DISABLE_DEPRECATION_WARNINGS
+		return Constraints;  
+		PRAGMA_ENABLE_DEPRECATION_WARNINGS
+	}
+	const TArray<FSolverVec3>& GetBarys() const 
+	{
+		PRAGMA_DISABLE_DEPRECATION_WARNINGS
+		return Barys;
+		PRAGMA_ENABLE_DEPRECATION_WARNINGS
+	}
 	FSolverReal GetThickness() const { return Thickness; }
 	bool GetGlobalIntersectionAnalysis() const { return bGlobalIntersectionAnalysis; }
-	const TArray<bool>& GetFlipNormals() const { return FlipNormal; }
+	const TArray<bool>& GetFlipNormals() const 
+	{
+		PRAGMA_DISABLE_DEPRECATION_WARNINGS
+		return FlipNormal;
+		PRAGMA_ENABLE_DEPRECATION_WARNINGS
+	}
 
 	void SetThickness(FSolverReal InThickness) { Thickness = FMath::Max(InThickness, (FSolverReal)0.);  }
 	void SetFrictionCoefficient(FSolverReal InFrictionCoefficient) { FrictionCoefficient = InFrictionCoefficient; }
@@ -76,6 +102,7 @@ public:
 	template<typename SolverParticlesOrRange>
 	void Apply(SolverParticlesOrRange& Particles, const FSolverReal Dt, const int32 ConstraintIndex) const
 	{
+		PRAGMA_DISABLE_DEPRECATION_WARNINGS
 		const TVector<int32, 4>& Constraint = Constraints[ConstraintIndex];
 		const int32 Index1 = Constraint[0];
 		const int32 Index2 = Constraint[1];
@@ -98,15 +125,18 @@ public:
 		{
 			Particles.P(Index4) -= Particles.InvM(Index4) * Barys[ConstraintIndex][2] * Delta;
 		}
+		PRAGMA_ENABLE_DEPRECATION_WARNINGS
 	}
 
 	template<typename SolverParticlesOrRange>
 	void Apply(SolverParticlesOrRange& InParticles, const FSolverReal Dt) const
 	{
+		PRAGMA_DISABLE_DEPRECATION_WARNINGS
 		for (int32 ConstraintIndex = 0; ConstraintIndex < Constraints.Num(); ++ConstraintIndex)
 		{
 			Apply(InParticles, Dt, ConstraintIndex);
 		}
+		PRAGMA_ENABLE_DEPRECATION_WARNINGS
 	}
 
 	void Apply(FSolverParticles& InParticles, const FSolverReal Dt, const TArray<int32>& InConstraintIndices) const
@@ -122,25 +152,86 @@ public:
 	TConstArrayView<int32> GetFaceCollisionLayers() const { return FaceCollisionLayers; }
 	const TArray<TVector<int32, 2>>& GetVertexCollisionLayers() const { return VertexCollisionLayers; }
 
+	FSolverReal GetConstraintThickness(const int32 ConstraintIndex) const
+	{
+		switch (ConstraintTypes[ConstraintIndex])
+		{
+		default:
+		case EConstraintType::Default:
+		case EConstraintType::GIAFlipped:
+			return 2.f * Thickness;
+		case EConstraintType::Kinematic:
+			return Thickness + KinematicColliderThickness;
+		}		
+	}
+
+	FSolverReal GetConstraintStiffness(const int32 ConstraintIndex) const
+	{
+		switch (ConstraintTypes[ConstraintIndex])
+		{
+		default:
+		case EConstraintType::Default:
+		case EConstraintType::GIAFlipped:
+			return Stiffness;
+		case EConstraintType::Kinematic:
+			return KinematicColliderStiffness;
+		}
+	}
+
+	FSolverReal GetConstraintFrictionCoefficient(const int32 ConstraintIndex) const
+	{
+		switch (ConstraintTypes[ConstraintIndex])
+		{
+		default:
+		case EConstraintType::Default:
+			return FrictionCoefficient;
+		case EConstraintType::Kinematic:
+			return KinematicColliderFrictionCoefficient;
+		case EConstraintType::GIAFlipped:
+			return (FSolverReal)0.f;
+		}
+	}
 protected:
-	TArray<TVec4<int32>> Constraints;
-	TArray<FSolverVec3> Barys;
-	TArray<bool> FlipNormal;
+
 	FSolverReal Thickness;
 	FSolverReal Stiffness; // (0-1 compliance for PBD)
 	FSolverReal FrictionCoefficient;
+	FSolverReal KinematicColliderThickness;
+	FSolverReal KinematicColliderStiffness;
+	FSolverReal KinematicColliderFrictionCoefficient;
 	FSolverReal ProximityStiffness; // (actual spring stiffness for force-based solver)
 
+	UE_DEPRECATED(5.4, "Constraints will be made private")
+	TArray<TVec4<int32>> Constraints;
+	UE_DEPRECATED(5.4, "Barys will be made private")
+	TArray<FSolverVec3> Barys;
+	UE_DEPRECATED(5.4, "FlipNormal will be made private")
+	TArray<bool> FlipNormal;
+
 	CHAOS_API void UpdateCollisionLayers(const TConstArrayView<int32>& InFaceCollisionLayers);
 
 private:
+
 	const FTriangleMesh& TriangleMesh;
-	const TArray<TVec3<int32>>& Elements;
 	const TArray<FSolverVec3>* ReferencePositions;
 	const TSet<TVec2<int32>> DisabledCollisionElements;  // TODO: Make this a bitarray
 	TConstArrayView<int32> FaceCollisionLayers;
 	TArray<TVector<int32, 2>> VertexCollisionLayers; // Only non-empty if FaceCollisionLayers is non-empty. Values are Min and Max layers for that vertex
 
+	enum struct EConstraintType : uint8
+	{
+		Default,
+		Kinematic,
+		GIAFlipped,
+	};
+	TArray<EConstraintType> ConstraintTypes;
+
+	struct FExistingConstraintData
+	{
+		FSolverReal Timer;
+	};
+	TArray<TMap<int32, FExistingConstraintData>> ExistingConstraintLookup; // ArrayIndex = ParticleIndex (without offset), MapKey = FaceIndex (full mesh)
+
 	int32 Offset;
 	int32 NumParticles;
 	bool bGlobalIntersectionAnalysis; // This is set based on which Init is called.
diff --git a/Engine/Source/Runtime/Experimental/Chaos/Public/Chaos/PBDTriangleMeshCollisions.h b/Engine/Source/Runtime/Experimental/Chaos/Public/Chaos/PBDTriangleMeshCollisions.h
--- a/Engine/Source/Runtime/Experimental/Chaos/Public/Chaos/PBDTriangleMeshCollisions.h
+++ b/Engine/Source/Runtime/Experimental/Chaos/Public/Chaos/PBDTriangleMeshCollisions.h
@@ -109,6 +109,73 @@ public:
 		Count
 	};
 
+	class FTriangleSubMesh
+	{
+	public:
+		FTriangleSubMesh(const FTriangleMesh& InFullMesh)
+			: FullMesh(InFullMesh)
+			, bSubMeshIsFullMesh(true)
+		{}
+
+		template<typename SolverParticlesOrRange>
+		void Init(const SolverParticlesOrRange& Particles, const TSet<int32>& DisabledFaces, bool bCollideAgainstAllKinematicVertices, const TSet<int32>& EnabledKinematicFaces);
+
+		const FTriangleMesh& GetCollidableMesh() const
+		{
+			if (bSubMeshIsFullMesh)
+			{
+				return FullMesh;
+			}
+			return SubMesh;
+		}
+
+		const FTriangleMesh& GetFullMesh() const 
+		{
+			return FullMesh;
+		}
+
+		int32 FullMeshIndexFromSubIndex(int32 SubIndex) const
+		{
+			if (bSubMeshIsFullMesh)
+			{
+				return SubIndex;
+			}
+
+			return SubMeshToFullElements[SubIndex];
+		}
+
+		int32 SubMeshIndexFromFullIndex(int32 FullIndex) const
+		{
+			if (bSubMeshIsFullMesh)
+			{
+				return FullIndex;
+			}
+			return FullToSubMeshElements[FullIndex];
+		}
+
+		bool GetSubMeshIsFullMesh() const { return bSubMeshIsFullMesh; }
+
+		const TArray<bool>& GetSubMeshElementIsKinematic() const 
+		{
+			return SubMeshElementIsKinematic;
+		}
+
+		// Will be empty when bSubMeshIsFullMesh
+		const TArray<int32>& GetCollidableVertices() const { return CollidableVertices; }
+		// Will be empty when bSubMeshIsFullMesh
+		const TArray<int32>& GetIntersectableSubmeshEdges() const { return IntersectableSubmeshEdges; }
+
+	private:
+		const FTriangleMesh& FullMesh;
+		FTriangleMesh SubMesh;
+		TArray<int32> SubMeshToFullElements;
+		TArray<int32> FullToSubMeshElements;
+		TArray<bool> SubMeshElementIsKinematic;
+		TArray<int32> CollidableVertices;
+		TArray<int32> IntersectableSubmeshEdges;
+		bool bSubMeshIsFullMesh;
+	};
+
 	static bool IsEnabled(const FCollectionPropertyConstFacade& PropertyCollection)
 	{
 		return PropertyCollection.IsEnabled(FName(TEXT("SelfCollisionStiffness")).ToString(), false);  // Don't use UE_CHAOS_DECLARE_PROPERTYCOLLECTION_NAME here, SelfCollisionStiffness is only needed for activation
@@ -117,10 +184,11 @@ public:
 	FPBDTriangleMeshCollisions(
 		const int32 InOffset,
 		const int32 InNumParticles,
+		const TMap<FString, const TSet<int32>*>& FaceSets,
 		const FTriangleMesh& InTriangleMesh,
 		const FCollectionPropertyConstFacade& PropertyCollection
 	)
-		:TriangleMesh(InTriangleMesh)
+		: CollidableSubMesh(InTriangleMesh)
 		, Offset(InOffset)
 		, NumParticles(InNumParticles)
 		, bUseSelfIntersections(GetUseSelfIntersections(PropertyCollection, false))
@@ -128,11 +196,30 @@ public:
 		, bContourMinimization(GetUseSelfIntersections(PropertyCollection, false) && GetUseContourMinimization(PropertyCollection, true))
 		, NumContourMinimizationPostSteps(GetUseSelfIntersections(PropertyCollection, false) ? GetNumContourMinimizationPostSteps(PropertyCollection, 0) : 0)
 		, bUseGlobalPostStepContours(GetUseGlobalPostStepContours(PropertyCollection, true))
+		, bSelfCollideAgainstAllKinematicVertices(GetSelfCollideAgainstAllKinematicVertices(PropertyCollection, false))
+		, bCollidableSubMeshDirty(true)
 		, UseSelfIntersectionsIndex(PropertyCollection)
 		, UseGlobalIntersectionAnalysisIndex(PropertyCollection)
 		, UseContourMinimizationIndex(PropertyCollection)
 		, NumContourMinimizationPostStepsIndex(PropertyCollection)
 		, UseGlobalPostStepContoursIndex(PropertyCollection)
+		, SelfCollideAgainstAllKinematicVerticesIndex(PropertyCollection)
+		, SelfCollisionDisabledFacesIndex(PropertyCollection)
+		, SelfCollisionEnabledKinematicFacesIndex(PropertyCollection)
+	{
+		if (const TSet<int32>* const InDisabledFaces = FaceSets.FindRef(GetSelfCollisionDisabledFacesString(PropertyCollection, SelfCollisionDisabledFacesName.ToString()), nullptr))
+		{
+			DisabledFaces = *InDisabledFaces;
+		}	
+	}
+
+	UE_DEPRECATED(5.4, "Use Constructor with FaceSets")
+	FPBDTriangleMeshCollisions(
+		const int32 InOffset,
+		const int32 InNumParticles,
+		const FTriangleMesh& InTriangleMesh,
+		const FCollectionPropertyConstFacade& PropertyCollection)
+		: FPBDTriangleMeshCollisions(InOffset, InNumParticles, TMap<FString, const TSet<int32>*>(), InTriangleMesh, PropertyCollection)
 	{}
 
 	FPBDTriangleMeshCollisions(
@@ -142,22 +229,28 @@ public:
 		bool bInGlobalIntersectionAnalysis,
 		bool bInContourMinimization
 	)
-		:TriangleMesh(InTriangleMesh)
+		:CollidableSubMesh(InTriangleMesh)
 		, Offset(InOffset)
 		, NumParticles(InNumParticles)
 		, bUseSelfIntersections(bInGlobalIntersectionAnalysis || bInContourMinimization)
 		, bGlobalIntersectionAnalysis(bInGlobalIntersectionAnalysis)
 		, bContourMinimization(bInContourMinimization)
+		, bSelfCollideAgainstAllKinematicVertices(false)
+		, bCollidableSubMeshDirty(true)
 		, UseSelfIntersectionsIndex(ForceInit)
 		, UseGlobalIntersectionAnalysisIndex(ForceInit)
 		, UseContourMinimizationIndex(ForceInit)
 		, NumContourMinimizationPostStepsIndex(ForceInit)
 		, UseGlobalPostStepContoursIndex(ForceInit)
+		, SelfCollideAgainstAllKinematicVerticesIndex(ForceInit)
+		, SelfCollisionDisabledFacesIndex(ForceInit)
+		, SelfCollisionEnabledKinematicFacesIndex(ForceInit)
 	{}
 
 	virtual ~FPBDTriangleMeshCollisions() = default;
 
-	void SetProperties(const FCollectionPropertyConstFacade& PropertyCollection)
+	void SetProperties(const FCollectionPropertyConstFacade& PropertyCollection,
+		const TMap<FString, const TSet<int32>*>& FaceSets)
 	{
 		const bool bSelfIntersectionsMutable = IsUseSelfIntersectionsMutable(PropertyCollection);
 		if (bSelfIntersectionsMutable)
@@ -188,6 +281,54 @@ public:
 			bGlobalIntersectionAnalysis = bContourMinimization = false;
 			NumContourMinimizationPostSteps = 0;
 		}
+
+		if (IsSelfCollideAgainstAllKinematicVerticesMutable(PropertyCollection))
+		{
+			const bool bNewValue = GetSelfCollideAgainstAllKinematicVertices(PropertyCollection);
+			if (bNewValue != bSelfCollideAgainstAllKinematicVertices)
+			{
+				bSelfCollideAgainstAllKinematicVertices = bNewValue;
+				bCollidableSubMeshDirty = true;
+			}
+		}
+		if (IsSelfCollisionDisabledFacesMutable(PropertyCollection) && IsSelfCollisionDisabledFacesStringDirty(PropertyCollection))
+		{
+			if (const TSet<int32>* const InDisabledFaces = FaceSets.FindRef(GetSelfCollisionDisabledFacesString(PropertyCollection, SelfCollisionDisabledFacesName.ToString()), nullptr))
+			{
+				DisabledFaces = *InDisabledFaces;
+				bCollidableSubMeshDirty = true;
+			}
+			else
+			{
+				if (!DisabledFaces.IsEmpty())
+				{
+					DisabledFaces.Reset();
+					bCollidableSubMeshDirty = true;
+				}
+			}
+		}
+		if (IsSelfCollisionEnabledKinematicFacesMutable(PropertyCollection) && IsSelfCollisionEnabledKinematicFacesStringDirty(PropertyCollection))
+		{
+			if (const TSet<int32>* const InEnabledFaces = FaceSets.FindRef(GetSelfCollisionEnabledKinematicFacesString(PropertyCollection, SelfCollisionEnabledKinematicFacesName.ToString()), nullptr))
+			{
+				EnabledKinematicFaces = *InEnabledFaces;
+				bCollidableSubMeshDirty = true;
+			}
+			else
+			{
+				if (!EnabledKinematicFaces.IsEmpty())
+				{
+					EnabledKinematicFaces.Reset();
+					bCollidableSubMeshDirty = true;
+				}
+			}
+		}
+	}
+
+	UE_DEPRECATED(5.4, "Use version with FaceSets")
+	void SetProperties(const FCollectionPropertyConstFacade& PropertyCollection)
+	{
+		SetProperties(PropertyCollection, TMap<FString, const TSet<int32>*>());
 	}
 
 	template<typename SolverParticlesOrRange>
@@ -204,6 +345,7 @@ public:
 		return NumContourMinimizationPostSteps;
 	}
 
+	const FTriangleSubMesh& GetCollidableSubMesh() const { return CollidableSubMesh; }
 	const FTriangleMesh::TSpatialHashType<FSolverReal>& GetSpatialHash() const { return SpatialHash; }
 	const TArray<FContourMinimizationIntersection>& GetContourMinimizationIntersections() const { return ContourMinimizationIntersections; }
 	const TConstArrayView<FGIAColor> GetVertexGIAColors() const { return bGlobalIntersectionAnalysis && VertexGIAColors.Num() == NumParticles ? TConstArrayView<FGIAColor>(VertexGIAColors.GetData() - Offset, NumParticles + Offset) : TConstArrayView<FGIAColor>(); }
@@ -217,15 +359,19 @@ public:
 	const TArray<TArray<FBarycentricPoint>>& GetPostStepIntersectionContourPoints() const { return PostStepIntersectionContourPoints; }
 private:
 
-	const FTriangleMesh& TriangleMesh;
+	FTriangleSubMesh CollidableSubMesh;
 	int32 Offset;
 	int32 NumParticles;
 	bool bUseSelfIntersections;
 	bool bGlobalIntersectionAnalysis;
-	bool bContourMinimization;
-	
+	bool bContourMinimization;	
 	int32 NumContourMinimizationPostSteps = 0;
 	bool bUseGlobalPostStepContours = true;
+	bool bSelfCollideAgainstAllKinematicVertices;
+	TSet<int32> DisabledFaces;
+	TSet<int32> EnabledKinematicFaces;
+
+	bool bCollidableSubMeshDirty = true;
 	
 	FTriangleMesh::TSpatialHashType<FSolverReal> SpatialHash;
 	TArray<FContourMinimizationIntersection> ContourMinimizationIntersections;
@@ -249,6 +395,9 @@ private:
 	UE_CHAOS_DECLARE_PROPERTYCOLLECTION_NAME(UseContourMinimization, bool);
 	UE_CHAOS_DECLARE_PROPERTYCOLLECTION_NAME(NumContourMinimizationPostSteps, int32);
 	UE_CHAOS_DECLARE_PROPERTYCOLLECTION_NAME(UseGlobalPostStepContours, bool);
+	UE_CHAOS_DECLARE_PROPERTYCOLLECTION_NAME(SelfCollideAgainstAllKinematicVertices, bool);
+	UE_CHAOS_DECLARE_PROPERTYCOLLECTION_NAME(SelfCollisionDisabledFaces, bool);
+	UE_CHAOS_DECLARE_PROPERTYCOLLECTION_NAME(SelfCollisionEnabledKinematicFaces, bool);
 };
 
 }  // End namespace Chaos::Softs
