[ChaosVD] Added Optional CVD Data Channels. This allows opt out of recording some types of data (like specific solver steps, or scene queries).

Also moved most CVD macros to its own file.

The new Data Channels can be enabled/disabled via CVD's Editor interface, commandline argument, or console command

#changelist validated
#virtualized

diff --git a/Engine/Plugins/Experimental/ChaosVD/Source/ChaosVD/Private/Widgets/SChaosVDRecordingControls.cpp b/Engine/Plugins/Experimental/ChaosVD/Source/ChaosVD/Private/Widgets/SChaosVDRecordingControls.cpp
--- a/Engine/Plugins/Experimental/ChaosVD/Source/ChaosVD/Private/Widgets/SChaosVDRecordingControls.cpp
+++ b/Engine/Plugins/Experimental/ChaosVD/Source/ChaosVD/Private/Widgets/SChaosVDRecordingControls.cpp
@@ -11,10 +11,12 @@
 #include "Input/Reply.h"
 #include "Misc/MessageDialog.h"
 #include "StatusBarSubsystem.h"
+#include "Framework/MultiBox/MultiBoxBuilder.h"
 #include "Styling/SlateBrush.h"
 #include "Widgets/Images/SImage.h"
 #include "Widgets/Input/SButton.h"
 #include "Widgets/SChaosVDMainTab.h"
+#include "Widgets/Input/SComboButton.h"
 #include "Widgets/Text/STextBlock.h"
 
 #define LOCTEXT_NAMESPACE "ChaosVisualDebugger"
@@ -54,6 +56,32 @@ void SChaosVDRecordingControls::Construct(const FArguments& InArgs, const TShare
 					.Text_Raw(this, &SChaosVDRecordingControls::GetRecordingTimeText)
 					.ColorAndOpacity(FColor::White)
 			]
+
+			+SHorizontalBox::Slot()
+			.HAlign(HAlign_Left)
+			.VAlign(VAlign_Center)
+			.AutoWidth()
+			.Padding(12.0f, 0.0f, 0.0f, 0.0f)
+			[
+				SNew(SComboButton)
+				.ContentPadding(FMargin(6.0f, 0.0f))
+				.MenuPlacement(MenuPlacement_AboveAnchor)
+				.ComboButtonStyle(&FAppStyle::Get().GetWidgetStyle<FComboButtonStyle>("SimpleComboButton"))
+				.OnGetMenuContent(this, &SChaosVDRecordingControls::GenerateDataChannelsMenu)
+				.HasDownArrow(true)
+				.ButtonContent()
+				[
+					SNew(SHorizontalBox)
+					+ SHorizontalBox::Slot()
+					.VAlign(VAlign_Center)
+					.AutoWidth()
+					[
+						SNew(STextBlock)
+						.Text(LOCTEXT("DataChannelsButton", "Data Channels"))
+						.TextStyle(&FAppStyle::Get().GetWidgetStyle<FTextBlockStyle>("DialogButtonText"))
+					]
+				]
+			]
 		]
 	];
 
@@ -82,6 +110,60 @@ TSharedRef<SButton> SChaosVDRecordingControls::GenerateToggleRecordingStateButto
 		];
 }
 
+TSharedRef<SWidget> SChaosVDRecordingControls::GenerateDataChannelsMenu()
+{
+	using namespace Chaos::VisualDebugger;
+
+	FMenuBuilder MenuBuilder(true, nullptr);
+
+	MenuBuilder.BeginSection("CVDRecordingWidget", LOCTEXT("CVDRecordingMenuChannels", "Data Channels"));
+	{
+		FChaosVDDataChannelsManager::Get().EnumerateChannels([this, &MenuBuilder](const TSharedRef<FCVDDataChannel>& Channel)
+		{
+			MenuBuilder.AddMenuEntry(
+				Channel->GetDisplayName(),
+				FText::Format(LOCTEXT("ChannelDesc", "Enable/disable the {0} channel"), Channel->GetDisplayName()),
+				FSlateIcon(),
+				FUIAction(FExecuteAction::CreateSP(this, &SChaosVDRecordingControls::ToggleChannelEnabledState, Channel.ToWeakPtr()),
+					FCanExecuteAction::CreateSP(this, &SChaosVDRecordingControls::CanChangeChannelEnabledState, Channel.ToWeakPtr()), FIsActionChecked::CreateSP(this, &SChaosVDRecordingControls::IsChannelEnabled,  Channel.ToWeakPtr())), NAME_None, EUserInterfaceActionType::ToggleButton);
+
+			return true;
+		});
+	}
+
+	MenuBuilder.EndSection();
+
+	return MenuBuilder.MakeWidget();
+}
+
+void SChaosVDRecordingControls::ToggleChannelEnabledState(TWeakPtr<FCVDDataChannel> Channel)
+{
+	if (const TSharedPtr<FCVDDataChannel>& ChannelPtr = Channel.Pin())
+	{
+		const bool bIsEnabled = ChannelPtr->IsChannelEnabled();
+		ChannelPtr->SetChannelEnabled(!bIsEnabled);
+	}
+}
+
+bool SChaosVDRecordingControls::IsChannelEnabled(TWeakPtr<FCVDDataChannel> Channel)
+{
+	if (const TSharedPtr<Chaos::VisualDebugger::FChaosVDOptionalDataChannel>& ChannelPtr = Channel.Pin())
+	{
+		return ChannelPtr->IsChannelEnabled();
+	}
+	
+	return false;
+}
+
+bool SChaosVDRecordingControls::CanChangeChannelEnabledState(TWeakPtr<FCVDDataChannel> Channel)
+{
+	if (const TSharedPtr<FCVDDataChannel>& ChannelPtr = Channel.Pin())
+	{
+		return ChannelPtr->CanChangeEnabledState();
+	}
+	
+	return false;
+}
 
 SChaosVDRecordingControls::~SChaosVDRecordingControls()
 {
diff --git a/Engine/Plugins/Experimental/ChaosVD/Source/ChaosVD/Public/Widgets/SChaosVDRecordingControls.h b/Engine/Plugins/Experimental/ChaosVD/Source/ChaosVD/Public/Widgets/SChaosVDRecordingControls.h
--- a/Engine/Plugins/Experimental/ChaosVD/Source/ChaosVD/Public/Widgets/SChaosVDRecordingControls.h
+++ b/Engine/Plugins/Experimental/ChaosVD/Source/ChaosVD/Public/Widgets/SChaosVDRecordingControls.h
@@ -6,6 +6,11 @@
 #include "Templates/SharedPointer.h"
 #include "Widgets/SCompoundWidget.h"
 
+namespace Chaos::VisualDebugger
+{
+	struct FChaosVDOptionalDataChannel;
+}
+
 class SButton;
 class SChaosVDMainTab;
 class FReply;
@@ -18,6 +23,8 @@ enum class EChaosVDRecordingMode
 	Live
 };
 
+typedef Chaos::VisualDebugger::FChaosVDOptionalDataChannel FCVDDataChannel;
+
 class SChaosVDRecordingControls : public SCompoundWidget
 {
 public:
@@ -30,7 +37,12 @@ public:
 
 protected:
 
-	TSharedRef<SButton>GenerateToggleRecordingStateButton(EChaosVDRecordingMode RecordingMode, const FText& StartRecordingTooltip);
+	TSharedRef<SButton> GenerateToggleRecordingStateButton(EChaosVDRecordingMode RecordingMode, const FText& StartRecordingTooltip);
+	TSharedRef<SWidget> GenerateDataChannelsMenu();
+
+	void ToggleChannelEnabledState(TWeakPtr<FCVDDataChannel> Channel);
+	bool IsChannelEnabled(TWeakPtr<FCVDDataChannel> Channel);
+	bool CanChangeChannelEnabledState(TWeakPtr<FCVDDataChannel> Channel);
 
 	const FSlateBrush* GetRecordOrStopButton(EChaosVDRecordingMode RecordingMode) const;
 	
diff --git a/Engine/Source/Runtime/Engine/Private/Collision/ChaosVDSQTraceHelper.cpp b/Engine/Source/Runtime/Engine/Private/Collision/ChaosVDSQTraceHelper.cpp
--- a/Engine/Source/Runtime/Engine/Private/Collision/ChaosVDSQTraceHelper.cpp
+++ b/Engine/Source/Runtime/Engine/Private/Collision/ChaosVDSQTraceHelper.cpp
@@ -56,6 +56,7 @@ int32 GetWorldSolverID(const UWorld* World)
 		
 		FChaosVDContext NewSceneQueryContext;
 		NewSceneQueryContext.Id = FChaosVDRuntimeModule::Get().GenerateUniqueID();
+		NewSceneQueryContext.SetDataChannel(CVDDC_SceneQueries);
 
 		if (CurrentCVDContext.Type ==  static_cast<int32>(EChaosVDContextType::Query) || CurrentCVDContext.Type == static_cast<int32>(EChaosVDContextType::SubTraceQuery))
 		{
diff --git a/Engine/Source/Runtime/Experimental/Chaos/Private/Chaos/Evolution/PBDMinEvolution.cpp b/Engine/Source/Runtime/Experimental/Chaos/Private/Chaos/Evolution/PBDMinEvolution.cpp
--- a/Engine/Source/Runtime/Experimental/Chaos/Private/Chaos/Evolution/PBDMinEvolution.cpp
+++ b/Engine/Source/Runtime/Experimental/Chaos/Private/Chaos/Evolution/PBDMinEvolution.cpp
@@ -124,7 +124,7 @@ namespace Chaos
 	{
 		SCOPE_CYCLE_COUNTER(STAT_MinEvolution_AdvanceOneTimeStep);
 
-		CVD_SCOPE_TRACE_SOLVER_STEP(TEXT("Evolution Advance"));
+		CVD_SCOPE_TRACE_SOLVER_STEP(CVDDC_Default, TEXT("Evolution Advance"));
 
 		Integrate(Dt);
 
diff --git a/Engine/Source/Runtime/Experimental/Chaos/Private/Chaos/PBDRigidsEvolutionGBF.cpp b/Engine/Source/Runtime/Experimental/Chaos/Private/Chaos/PBDRigidsEvolutionGBF.cpp
--- a/Engine/Source/Runtime/Experimental/Chaos/Private/Chaos/PBDRigidsEvolutionGBF.cpp
+++ b/Engine/Source/Runtime/Experimental/Chaos/Private/Chaos/PBDRigidsEvolutionGBF.cpp
@@ -435,7 +435,7 @@ void FPBDRigidsEvolutionGBF::AdvanceOneTimeStepImpl(const FReal Dt, const FSubSt
 #endif
 
 	{
-		CVD_SCOPE_TRACE_SOLVER_STEP(TEXT("Evolution Start"))
+		CVD_SCOPE_TRACE_SOLVER_STEP(CVDDC_EvolutionStart, TEXT("Evolution Start"))
 		CVD_TRACE_PARTICLES_SOA(Particles);
 	}
 
@@ -455,7 +455,7 @@ void FPBDRigidsEvolutionGBF::AdvanceOneTimeStepImpl(const FReal Dt, const FSubSt
 #endif
 
 	{
-		CVD_SCOPE_TRACE_SOLVER_STEP(TEXT("Integrate"))
+		CVD_SCOPE_TRACE_SOLVER_STEP(CVDDC_Integrate, TEXT("Integrate"))
 		SCOPE_CYCLE_COUNTER(STAT_Evolution_Integrate);
 		CSV_SCOPED_TIMING_STAT(PhysicsVerbose, StepSolver_Integrate);
 		Integrate(Particles.GetActiveParticlesView(), Dt);
@@ -469,7 +469,7 @@ void FPBDRigidsEvolutionGBF::AdvanceOneTimeStepImpl(const FReal Dt, const FSubSt
 #endif
 
 	{
-		CVD_SCOPE_TRACE_SOLVER_STEP(TEXT("ApplyKinematicTargets"))
+		CVD_SCOPE_TRACE_SOLVER_STEP(CVDDC_ApplyKinematicTargets, TEXT("ApplyKinematicTargets"))
 		SCOPE_CYCLE_COUNTER(STAT_Evolution_KinematicTargets);
 		CSV_SCOPED_TIMING_STAT(PhysicsVerbose, StepSolver_KinematicTargets);
 		ApplyKinematicTargets(Dt, SubStepInfo.PseudoFraction);
@@ -502,7 +502,7 @@ void FPBDRigidsEvolutionGBF::AdvanceOneTimeStepImpl(const FReal Dt, const FSubSt
 		CollisionDetector.GetBroadPhase().SetSpatialAcceleration(InternalAcceleration);
 
 		{
-			CVD_SCOPE_TRACE_SOLVER_STEP(TEXT("Collision Detection Broad Phase"))
+			CVD_SCOPE_TRACE_SOLVER_STEP(CVDDC_CollisionDetectionBroadPhase, TEXT("Collision Detection Broad Phase"))
 			CollisionDetector.RunBroadPhase(Dt, GetCurrentStepResimCache());
 		}
 
@@ -513,8 +513,7 @@ void FPBDRigidsEvolutionGBF::AdvanceOneTimeStepImpl(const FReal Dt, const FSubSt
 		}
 
 		{
-			CVD_SCOPE_TRACE_SOLVER_STEP(TEXT("Collision Detection Narrow Phase"))
-
+			CVD_SCOPE_TRACE_SOLVER_STEP(CVDDC_CollisionDetectionNarrowPhase, TEXT("Collision Detection Narrow Phase"))
 			CollisionDetector.RunNarrowPhase(Dt, GetCurrentStepResimCache());
 		}
 	}
@@ -697,13 +696,17 @@ void FPBDRigidsEvolutionGBF::AdvanceOneTimeStepImpl(const FReal Dt, const FSubSt
 	}
 	
 	{
-		CVD_SCOPE_TRACE_SOLVER_STEP(TEXT("Evolution End"));
+		CVD_SCOPE_TRACE_SOLVER_STEP(CVDDC_EvolutionEnd, TEXT("Evolution End"));
 
-		CVD_TRACE_MID_PHASES_FROM_COLLISION_CONSTRAINTS(GetCollisionConstraints());
+		{
+			// This needs to be executed withing a CVD Solver step scope,
+			// but we are giving it its inner scope so specify a CVD Data channel so this data can be opted out
+			// without opting out of the whole step
+			CVD_TRACE_STEP_MID_PHASES_FROM_COLLISION_CONSTRAINTS(CVDDC_EndOfEvolutionCollisionConstraints, GetCollisionConstraints());
+		}
 
 		CVD_TRACE_PARTICLES_SOA(Particles);
 	}
-	
 
 #if !UE_BUILD_SHIPPING
 	if(SerializeEvolution)
diff --git a/Engine/Source/Runtime/Experimental/Chaos/Private/ChaosVisualDebugger/ChaosVDOptionalDataChannel.cpp b/Engine/Source/Runtime/Experimental/Chaos/Private/ChaosVisualDebugger/ChaosVDOptionalDataChannel.cpp
new file mode 100644
--- /dev/null
+++ b/Engine/Source/Runtime/Experimental/Chaos/Private/ChaosVisualDebugger/ChaosVDOptionalDataChannel.cpp
@@ -0,0 +1,87 @@
+﻿// Copyright Epic Games, Inc. All Rights Reserved.
+
+#include "ChaosVisualDebugger/ChaosVDOptionalDataChannel.h"
+
+#if WITH_CHAOS_VISUAL_DEBUGGER
+
+#include "ChaosVDRuntimeModule.h"
+#include "HAL/IConsoleManager.h"
+#include "Internationalization/Text.h"
+#include "Internationalization/Internationalization.h"
+#include "Misc/LazySingleton.h"
+
+FAutoConsoleCommand SetCVDDataChannelEnabledCommand(
+	TEXT("p.Chaos.VD.SetCVDDataChannelEnabled"),
+	TEXT("Turn on or off a CVD Data Channel. Argument 1 is true or false, Argument is a comma separated list of channel names. Example: p.Chaos.VD.SetCVDDataChannelEnabled true SceneQueries,Integrate"),
+	FConsoleCommandWithArgsDelegate::CreateLambda([](const TArray<FString>& Args)
+	{
+		using namespace Chaos::VisualDebugger;
+
+		constexpr int32 ExpectedArgumentsNum = 2;
+		if (Args.Num() == ExpectedArgumentsNum)
+		{
+			const bool bNewEnabledState = Args[0] == TEXT("true");
+
+			TArray<FString> ParsedChannels;
+			ParseChannelListFromCommandArgument(ParsedChannels, Args[1]);
+
+			for (const FString& ChannelName : ParsedChannels)
+			{
+				if (const TSharedPtr<FChaosVDOptionalDataChannel> ChannelInstance = FChaosVDDataChannelsManager::Get().GetChannelById(FName(ChannelName)))
+				{
+					ChannelInstance->SetChannelEnabled(bNewEnabledState);
+				}
+			}
+		}
+	})
+);
+
+namespace Chaos::VisualDebugger
+{
+	void FChaosVDOptionalDataChannel::Initialize()
+	{
+		FChaosVDDataChannelsManager::Get().RegisterChannel(AsShared());
+	}
+
+	const FText& FChaosVDOptionalDataChannel::GetDisplayName() const
+	{
+		return *LocalizableChannelName;
+	}
+
+	void FChaosVDOptionalDataChannel::SetChannelEnabled(bool bNewEnabled)
+	{
+		
+	}
+
+	FChaosVDDataChannelsManager& FChaosVDDataChannelsManager::Get()
+	{
+		return TLazySingleton<FChaosVDDataChannelsManager>::Get();
+	}
+
+	void ParseChannelListFromCommandArgument(TArray<FString>& OutParsedChannelList, const FString& InCommandArgument)
+	{
+		constexpr int32 NumReserve = 10; // We don't really know how many channels will be requested, so start with something
+		OutParsedChannelList.Reserve(NumReserve);
+		InCommandArgument.ParseIntoArray(OutParsedChannelList,TEXT(","));
+	}
+
+	TSharedRef<FChaosVDOptionalDataChannel> CreateDataChannel(FName InChannelID, const TSharedRef<FText>& InDisplayName, EChaosVDDataChannelInitializationFlags InitializationFlags)
+	{
+		TSharedRef<FChaosVDOptionalDataChannel> NewChannel = MakeShared<FChaosVDOptionalDataChannel>(MakeShared<FName>(InChannelID), InDisplayName, InitializationFlags);
+		NewChannel->Initialize();
+
+		return NewChannel;
+	}
+}
+
+CVD_DEFINE_OPTIONAL_DATA_CHANNEL(Default, EChaosVDDataChannelInitializationFlags::StartEnabled);
+CVD_DEFINE_OPTIONAL_DATA_CHANNEL(EvolutionStart, EChaosVDDataChannelInitializationFlags::StartEnabled | EChaosVDDataChannelInitializationFlags::CanChangeEnabledState);
+CVD_DEFINE_OPTIONAL_DATA_CHANNEL(Integrate, EChaosVDDataChannelInitializationFlags::CanChangeEnabledState);
+CVD_DEFINE_OPTIONAL_DATA_CHANNEL(ApplyKinematicTargets, EChaosVDDataChannelInitializationFlags::StartEnabled | EChaosVDDataChannelInitializationFlags::CanChangeEnabledState);
+CVD_DEFINE_OPTIONAL_DATA_CHANNEL(CollisionDetectionBroadPhase, EChaosVDDataChannelInitializationFlags::CanChangeEnabledState);
+CVD_DEFINE_OPTIONAL_DATA_CHANNEL(CollisionDetectionNarrowPhase, EChaosVDDataChannelInitializationFlags::CanChangeEnabledState);
+CVD_DEFINE_OPTIONAL_DATA_CHANNEL(EndOfEvolutionCollisionConstraints, EChaosVDDataChannelInitializationFlags::StartEnabled| EChaosVDDataChannelInitializationFlags::CanChangeEnabledState);
+CVD_DEFINE_OPTIONAL_DATA_CHANNEL(EvolutionEnd, EChaosVDDataChannelInitializationFlags::StartEnabled); // Intentionally not allowing changing the enabled state because otherwise there might be nothing to visualize
+CVD_DEFINE_OPTIONAL_DATA_CHANNEL(SceneQueries, EChaosVDDataChannelInitializationFlags::StartEnabled | EChaosVDDataChannelInitializationFlags::CanChangeEnabledState);
+
+#endif
diff --git a/Engine/Source/Runtime/Experimental/Chaos/Private/ChaosVisualDebugger/ChaosVisualDebuggerTrace.cpp b/Engine/Source/Runtime/Experimental/Chaos/Private/ChaosVisualDebugger/ChaosVisualDebuggerTrace.cpp
--- a/Engine/Source/Runtime/Experimental/Chaos/Private/ChaosVisualDebugger/ChaosVisualDebuggerTrace.cpp
+++ b/Engine/Source/Runtime/Experimental/Chaos/Private/ChaosVisualDebugger/ChaosVisualDebuggerTrace.cpp
@@ -2,7 +2,7 @@
 
 #include "ChaosVisualDebugger/ChaosVisualDebuggerTrace.h"
 
-#if CHAOS_VISUAL_DEBUGGER_ENABLED
+#if WITH_CHAOS_VISUAL_DEBUGGER
 
 #include "Chaos/Framework/PhysicsSolverBase.h"
 #include "Chaos/ImplicitObject.h"
@@ -107,13 +107,14 @@ std::atomic<bool> FChaosVisualDebuggerTrace::bIsTracing = false;
 
 void FChaosVisualDebuggerTrace::TraceParticle(const Chaos::FGeometryParticleHandle* ParticleHandle)
 {
+	using namespace Chaos::VisualDebugger::Utils;
 	if (!IsTracing())
 	{
 		return;
 	}
 
 	const FChaosVDContext* CVDContextData = FChaosVDThreadContext::Get().GetCurrentContext(EChaosVDContextType::Solver);
-	if (!ensure(CVDContextData))
+	if (!IsContextEnabledAndValid(CVDContextData))
 	{
 		return;
 	}
@@ -167,13 +168,14 @@ void FChaosVisualDebuggerTrace::TraceParticle(Chaos::FGeometryParticleHandle* Pa
 
 void FChaosVisualDebuggerTrace::TraceParticles(const Chaos::TGeometryParticleHandles<Chaos::FReal, 3>& ParticleHandles)
 {
+	using namespace Chaos::VisualDebugger::Utils;
 	if (!IsTracing())
 	{
 		return;
 	}
 
 	const FChaosVDContext* CVDContextData = FChaosVDThreadContext::Get().GetCurrentContext(EChaosVDContextType::Solver);
-	if (!ensure(CVDContextData))
+	if (!IsContextEnabledAndValid(CVDContextData))
 	{
 		return;
 	}
@@ -187,6 +189,7 @@ void FChaosVisualDebuggerTrace::TraceParticles(const Chaos::TGeometryParticleHan
 
 void FChaosVisualDebuggerTrace::TraceParticleDestroyed(const Chaos::FGeometryParticleHandle* ParticleHandle)
 {
+	using namespace Chaos::VisualDebugger::Utils;
 	if (!IsTracing())
 	{
 		return;
@@ -201,7 +204,7 @@ void FChaosVisualDebuggerTrace::TraceParticleDestroyed(const Chaos::FGeometryPar
 	GeometryTracerObject.RemoveCachedGeometryHash(ParticleHandle->GetGeometry());
 	
 	const FChaosVDContext* CVDContextData = FChaosVDThreadContext::Get().GetCurrentContext(EChaosVDContextType::Solver);
-	if (!ensure(CVDContextData))
+	if (!IsContextEnabledAndValid(CVDContextData))
 	{
 		return;
 	}
@@ -214,13 +217,15 @@ void FChaosVisualDebuggerTrace::TraceParticleDestroyed(const Chaos::FGeometryPar
 
 void FChaosVisualDebuggerTrace::TraceParticlesSoA(const Chaos::FPBDRigidsSOAs& ParticlesSoA)
 {
+	using namespace Chaos::VisualDebugger::Utils;
 	if (!IsTracing())
 	{
 		return;
 	}
 
 	const FChaosVDContext* CVDContextData = FChaosVDThreadContext::Get().GetCurrentContext(EChaosVDContextType::Solver);
-	if (!ensure(CVDContextData))
+
+	if (!IsContextEnabledAndValid(CVDContextData))
 	{
 		return;
 	}
@@ -260,13 +265,15 @@ bool FChaosVisualDebuggerTrace::ShouldPerformFullCapture(int32 SolverID)
 
 void FChaosVisualDebuggerTrace::TraceMidPhase(const Chaos::FParticlePairMidPhase* MidPhase)
 {
+	using namespace Chaos::VisualDebugger::Utils;
 	if (!IsTracing())
 	{
 		return;
 	}
 
 	const FChaosVDContext* CVDContextData = FChaosVDThreadContext::Get().GetCurrentContext(EChaosVDContextType::Solver);
-	if (!ensure(CVDContextData))
+
+	if (!IsContextEnabledAndValid(CVDContextData))
 	{
 		return;
 	}
@@ -290,13 +297,15 @@ void FChaosVisualDebuggerTrace::TraceMidPhase(const Chaos::FParticlePairMidPhase
 
 void FChaosVisualDebuggerTrace::TraceMidPhasesFromCollisionConstraints(Chaos::FPBDCollisionConstraints& InCollisionConstraints)
 {
+	using namespace Chaos::VisualDebugger::Utils;
 	if (!IsTracing())
 	{
 		return;
 	}
 
 	const FChaosVDContext* CVDContextData = FChaosVDThreadContext::Get().GetCurrentContext(EChaosVDContextType::Solver);
-	if (!ensure(CVDContextData))
+
+	if (!IsContextEnabledAndValid(CVDContextData))
 	{
 		return;
 	}
@@ -309,16 +318,17 @@ void FChaosVisualDebuggerTrace::TraceMidPhasesFromCollisionConstraints(Chaos::FP
 	});
 }
 
-
 void FChaosVisualDebuggerTrace::TraceCollisionConstraint(const Chaos::FPBDCollisionConstraint* CollisionConstraint)
 {
+	using namespace Chaos::VisualDebugger::Utils;
 	if (!IsTracing())
 	{
 		return;
 	}
 
 	const FChaosVDContext* CVDContextData = FChaosVDThreadContext::Get().GetCurrentContext(EChaosVDContextType::Solver);
-	if (!ensure(CVDContextData))
+
+	if (!IsContextEnabledAndValid(CVDContextData))
 	{
 		return;
 	}
@@ -336,13 +346,16 @@ void FChaosVisualDebuggerTrace::TraceCollisionConstraint(const Chaos::FPBDCollis
 
 void FChaosVisualDebuggerTrace::TraceCollisionConstraintView(TArrayView<Chaos::FPBDCollisionConstraint* const> CollisionConstraintView)
 {
+	using namespace Chaos::VisualDebugger::Utils;
+
 	if (!IsTracing())
 	{
 		return;
 	}
 
 	const FChaosVDContext* CVDContextData = FChaosVDThreadContext::Get().GetCurrentContext(EChaosVDContextType::Solver);
-	if (!ensure(CVDContextData))
+
+	if (!IsContextEnabledAndValid(CVDContextData))
 	{
 		return;
 	}
@@ -413,13 +426,15 @@ void FChaosVisualDebuggerTrace::TraceSolverFrameEnd(const FChaosVDContext& Conte
 
 void FChaosVisualDebuggerTrace::TraceSolverStepStart(FStringView StepName)
 {
+	using namespace Chaos::VisualDebugger::Utils;
+
 	if (!IsTracing())
 	{
 		return;
 	}
 
 	const FChaosVDContext* CVDContextData = FChaosVDThreadContext::Get().GetCurrentContext(EChaosVDContextType::Solver);
-	if (!ensure(CVDContextData))
+	if (!IsContextEnabledAndValid(CVDContextData))
 	{
 		return;
 	}
@@ -427,18 +442,20 @@ void FChaosVisualDebuggerTrace::TraceSolverStepStart(FStringView StepName)
 	UE_TRACE_LOG(ChaosVDLogger, ChaosVDSolverStepStart, ChaosVDChannel)
 		<< ChaosVDSolverStepStart.Cycle(FPlatformTime::Cycles64())
 		<< ChaosVDSolverStepStart.SolverID(CVDContextData->Id)
-		<< ChaosVDSolverStepStart.StepName(StepName.GetData(), StepName.Len());
+		<< ChaosVDSolverStepStart.StepName(StepName.GetData(), GetNum(StepName));
 }
 
 void FChaosVisualDebuggerTrace::TraceSolverStepEnd()
 {
+	using namespace Chaos::VisualDebugger::Utils;
+
 	if (!IsTracing())
 	{
 		return;
 	}
 
 	const FChaosVDContext* CVDContextData = FChaosVDThreadContext::Get().GetCurrentContext(EChaosVDContextType::Solver);
-	if (!ensure(CVDContextData))
+	if (!IsContextEnabledAndValid(CVDContextData))
 	{
 		return;
 	}
@@ -450,13 +467,15 @@ void FChaosVisualDebuggerTrace::TraceSolverStepEnd()
 
 void FChaosVisualDebuggerTrace::TraceSolverSimulationSpace(const Chaos::FRigidTransform3& Transform)
 {
+	using namespace Chaos::VisualDebugger::Utils;
+
 	if (!IsTracing())
 	{
 		return;
 	}
 
 	const FChaosVDContext* CVDContextData = FChaosVDThreadContext::Get().GetCurrentContext(EChaosVDContextType::Solver);
-	if (!ensure(CVDContextData))
+	if (!IsContextEnabledAndValid(CVDContextData))
 	{
 		return;
 	}
@@ -602,13 +621,20 @@ void FChaosVisualDebuggerTrace::TraceNonSolverTransform(const FTransform& InTran
 
 void FChaosVisualDebuggerTrace::TraceSceneQueryStart(const Chaos::FImplicitObject* InputGeometry, const FQuat& GeometryOrientation,  const FVector& Start, const FVector& End, ECollisionChannel TraceChannel, FChaosVDCollisionQueryParams&& Params, FChaosVDCollisionResponseParams&& ResponseParams, FChaosVDCollisionObjectQueryParams&& ObjectParams, EChaosVDSceneQueryType QueryType, EChaosVDSceneQueryMode QueryMode, int32 SolverID, bool bIsRetry)
 {
+	using namespace Chaos::VisualDebugger::Utils;
+
 	if (!IsTracing())
 	{
 		return;
 	}
 
 	const FChaosVDContext* CVDContextData = FChaosVDThreadContext::Get().GetCurrentContext();
-	const bool bIsQueryContext = CVDContextData && (CVDContextData->Type == static_cast<int32>(EChaosVDContextType::Query) ||  CVDContextData->Type == static_cast<int32>(EChaosVDContextType::SubTraceQuery));
+	if (!IsContextEnabledAndValid(CVDContextData))
+	{
+		return;
+	}
+
+	const bool bIsQueryContext = CVDContextData->Type == static_cast<int32>(EChaosVDContextType::Query) ||  CVDContextData->Type == static_cast<int32>(EChaosVDContextType::SubTraceQuery);
 
 	if (!ensure(bIsQueryContext))
 	{
@@ -656,13 +682,19 @@ void FChaosVisualDebuggerTrace::TraceSceneQueryStart(const Chaos::FImplicitObjec
 
 void FChaosVisualDebuggerTrace::TraceSceneQueryVisit(FChaosVDQueryVisitStep&& InQueryVisitData)
 {
+	using namespace Chaos::VisualDebugger::Utils;
+
 	if (!IsTracing())
 	{
 		return;
 	}
 
 	const FChaosVDContext* CVDContextData = FChaosVDThreadContext::Get().GetCurrentContext();
-	const bool bIsQueryContext = CVDContextData && (CVDContextData->Type == static_cast<int32>(EChaosVDContextType::Query) ||  CVDContextData->Type == static_cast<int32>(EChaosVDContextType::SubTraceQuery));
+	if (!IsContextEnabledAndValid(CVDContextData))
+	{
+		return;
+	}
+	const bool bIsQueryContext = CVDContextData->Type == static_cast<int32>(EChaosVDContextType::Query) ||  CVDContextData->Type == static_cast<int32>(EChaosVDContextType::SubTraceQuery);
 
 	if (!ensure(bIsQueryContext))
 	{
@@ -755,6 +787,32 @@ void FChaosVisualDebuggerTrace::HandleRecordingStop()
 void FChaosVisualDebuggerTrace::HandleRecordingStart()
 {
 	Reset();
+
+	FString CommandlineEnabledCVDChannels;
+	constexpr bool bStopOnSeparator = false;
+	if (FParse::Value(FCommandLine::Get(), TEXT("CVDDataChannelsOverride="), CommandlineEnabledCVDChannels, bStopOnSeparator))
+	{
+		TArray<FString> ParsedChannels;
+		Chaos::VisualDebugger::ParseChannelListFromCommandArgument(ParsedChannels, CommandlineEnabledCVDChannels);
+
+		UE_LOG(LogChaos, Log, TEXT("[%s] Channel list override provided via commandline - Enabling [%d] Requested channels..."), ANSI_TO_TCHAR(__FUNCTION__), ParsedChannels.Num());
+
+		using namespace Chaos::VisualDebugger;
+		FChaosVDDataChannelsManager::Get().EnumerateChannels([&ParsedChannels](const TSharedRef<FChaosVDOptionalDataChannel>& Channel)
+		{
+			if (Channel->CanChangeEnabledState())
+			{
+				// This is far from efficient, but this will be called once when the recording start command is executed, and we only have a handful of channels
+				const FString ChannelIdAsString = Channel->GetId().ToString();
+				const bool bChannelShouldBeEnabled = ParsedChannels.Contains(ChannelIdAsString);
+				Channel->SetChannelEnabled(bChannelShouldBeEnabled);
+
+				UE_LOG(LogChaos, Log, TEXT("[%s] Setting enabled state for channel [%s] to [%s]..."), ANSI_TO_TCHAR(__FUNCTION__), *ChannelIdAsString, bChannelShouldBeEnabled ? TEXT("True") : TEXT("False"));
+			}
+			return true;
+		});
+	}
+	
 	bIsTracing = true;
 }
 
@@ -773,4 +831,4 @@ void FChaosVisualDebuggerTrace::PerformFullCapture(EChaosVDFullCaptureFlags Capt
 	}
 }
 
-#endif
+#endif //WITH_CHAOS_VISUAL_DEBUGGER
diff --git a/Engine/Source/Runtime/Experimental/Chaos/Public/ChaosVisualDebugger/ChaosVDContextProvider.h b/Engine/Source/Runtime/Experimental/Chaos/Public/ChaosVisualDebugger/ChaosVDContextProvider.h
--- a/Engine/Source/Runtime/Experimental/Chaos/Public/ChaosVisualDebugger/ChaosVDContextProvider.h
+++ b/Engine/Source/Runtime/Experimental/Chaos/Public/ChaosVisualDebugger/ChaosVDContextProvider.h
@@ -2,9 +2,11 @@
 
 #pragma once
 
+#include "ChaosVDOptionalDataChannel.h"
 #include "Containers/Array.h"
 #include "Containers/ContainerAllocationPolicies.h"
 #include "HAL/ThreadSingleton.h"
+#include "Templates/SharedPointer.h"
 
 #if WITH_CHAOS_VISUAL_DEBUGGER
 
@@ -24,6 +26,16 @@ struct FChaosVDContext
 	int32 OwnerID = INDEX_NONE;
 	int32 Id = INDEX_NONE;
 	int32 Type = INDEX_NONE;
+
+	FORCEINLINE void SetDataChannel(const TSharedRef<Chaos::VisualDebugger::FChaosVDOptionalDataChannel>& NewDataChannel)
+	{
+		CurrentDataChannel = NewDataChannel;
+	}
+
+	FORCEINLINE bool IsDataChannelEnabled () const { return CurrentDataChannel->IsChannelEnabled(); }
+
+private:
+	TSharedRef<Chaos::VisualDebugger::FChaosVDOptionalDataChannel> CurrentDataChannel = CVDDC_Default;
 };
 
 /** Singleton class that manages the thread local storage used to store CVD Context data */
@@ -88,6 +100,19 @@ protected:
 	friend struct FChaosVDScopedTLSBufferAccessor;
 };
 
+namespace Chaos::VisualDebugger::Utils
+{
+	FORCEINLINE bool IsContextEnabledAndValid(const FChaosVDContext* Context)
+	{
+		if (!ensure(Context))
+		{
+			return false;
+		}
+
+		return Context->IsDataChannelEnabled();
+	}
+}
+
 /** Utility Class that will push the provided CVD Context Data to the local thread storage
  * and remove it when it goes out of scope
  */
@@ -104,6 +129,25 @@ struct FChaosVDScopeContext
 	}
 };
 
+struct FChaosCVDScopedDataChannelOverride
+{
+	FChaosCVDScopedDataChannelOverride(const TSharedRef<Chaos::VisualDebugger::FChaosVDOptionalDataChannel>& NewDataChannel)
+	{
+		if (const FChaosVDContext* CVDContextData = FChaosVDThreadContext::Get().GetCurrentContext())
+		{
+			FChaosVDContext NewContext = *CVDContextData;
+			NewContext.SetDataChannel(NewDataChannel);
+
+			FChaosVDThreadContext::Get().PushContext(MoveTemp(NewContext));
+		}
+	}
+
+	~FChaosCVDScopedDataChannelOverride()
+	{
+		FChaosVDThreadContext::Get().PopContext();
+	}
+};
+
 /** Wrapper class that provides access to the data buffer that should be used to trace CVD (Chaos Visual Debugger) data.
  * It ensures the buffer is cleared on scope exit but keeping a minimum amount of memory allocated
  */
@@ -150,6 +194,11 @@ struct FChaosVDScopedTLSBufferAccessor
 		FChaosVDScopeContext CVDScope(InContext);
 #endif
 
+#ifndef CVD_SCOPED_DATA_CHANNEL_OVERRIDE
+	#define CVD_SCOPED_DATA_CHANNEL_OVERRIDE(DataChannel) \
+	FChaosCVDScopedDataChannelOverride CVDDC_Scope_Override##DataChannel(DataChannel);
+#endif
+
 #else // WITH_CHAOS_VISUAL_DEBUGGER
 
 #ifndef CVD_GET_CURRENT_CONTEXT
@@ -164,6 +213,10 @@ struct FChaosVDScopedTLSBufferAccessor
 	#define CVD_GET_WRAPPED_CURRENT_CONTEXT(OutWrappedContext)
 #endif
 
+#ifndef CVD_SCOPED_CONTEXT_DATA_CHANNEL_OVERRIDE
+	#define CVD_SCOPED_CONTEXT_DATA_CHANNEL_OVERRIDE(DataChannel)
+#endif
+
 #endif // WITH_CHAOS_VISUAL_DEBUGGER
 
 struct FChaosVDContextWrapper
diff --git a/Engine/Source/Runtime/Experimental/Chaos/Public/ChaosVisualDebugger/ChaosVDOptionalDataChannel.h b/Engine/Source/Runtime/Experimental/Chaos/Public/ChaosVisualDebugger/ChaosVDOptionalDataChannel.h
new file mode 100644
--- /dev/null
+++ b/Engine/Source/Runtime/Experimental/Chaos/Public/ChaosVisualDebugger/ChaosVDOptionalDataChannel.h
@@ -0,0 +1,136 @@
+﻿// Copyright Epic Games, Inc. All Rights Reserved.
+
+#pragma once
+
+#if WITH_CHAOS_VISUAL_DEBUGGER
+
+#include "Templates/SharedPointer.h"
+
+class FText;
+class FName;
+
+/** Set of flags to control how a CVD Data Channel gets initialized*/
+enum class EChaosVDDataChannelInitializationFlags : uint8
+{
+	/** If Set, the data channel will start in a enabled state */
+	StartEnabled = 1 << 0,
+	/** If Set, the data channel can be toggle on and off as desired */
+	CanChangeEnabledState = 1 << 1
+};
+ENUM_CLASS_FLAGS(EChaosVDDataChannelInitializationFlags)
+
+namespace Chaos::VisualDebugger
+{
+	/** Structure holding the state of a CVD data channel. Used to enabled and disable recording of specific categories of data */
+	struct FChaosVDOptionalDataChannel : TSharedFromThis<FChaosVDOptionalDataChannel>
+	{
+		explicit FChaosVDOptionalDataChannel(const TSharedRef<FName>& InChannelID, const TSharedRef<FText>& InDisplayName, EChaosVDDataChannelInitializationFlags InitializationFlags)
+				: LocalizableChannelName(InDisplayName), ChannelId(InChannelID)
+		{
+			bIsEnabled = EnumHasAnyFlags(InitializationFlags, EChaosVDDataChannelInitializationFlags::StartEnabled);
+			bCanChangeEnabledState = EnumHasAnyFlags(InitializationFlags, EChaosVDDataChannelInitializationFlags::CanChangeEnabledState);
+		}
+
+		/** Registers itself with the optional channel data manager */
+		void Initialize();
+
+		/** The localizable display name that will be used for this channel in UI elements */
+		CHAOS_API const FText& GetDisplayName() const;
+
+		/** A name used as ID to find this changed using the CVD Channel Manager */
+		FName GetId() const { return *ChannelId; }
+
+		/** Returns true if this data channel is enabled */
+		FORCEINLINE bool IsChannelEnabled() const { return bIsEnabled; }
+
+		/** Enables or disabled this Data Channel */
+		CHAOS_API void SetChannelEnabled(bool bNewEnabled);
+
+		/** Returns true if the enabled state of this channel can be changed */
+		bool CanChangeEnabledState() const { return bCanChangeEnabledState;}
+
+	private:
+		TSharedRef<FText> LocalizableChannelName;
+		TSharedRef<FName> ChannelId;
+		std::atomic<bool> bIsEnabled = true;
+		bool bCanChangeEnabledState = false;
+	};
+
+	/** Manager that provides a way to iterate trough all existing CVD Data channels or get a specific one */
+	class CHAOS_API FChaosVDDataChannelsManager
+	{
+	public:
+
+		static FChaosVDDataChannelsManager& Get();
+
+		/** Iterates trough all the available data channels and executes the provided callback passing each channel as an argument.
+		 * If the callback returns false, the iteration will be stopped.
+		 */
+		template<typename TCallback>
+		void EnumerateChannels(TCallback Callback)
+		{
+			for (const TPair<FName, TSharedPtr<FChaosVDOptionalDataChannel>>& ChannelWithName : AvailableChannels)
+			{
+				if (!Callback(ChannelWithName.Value.ToSharedRef()))
+				{
+					return;
+				}
+			}
+		}
+
+		/** Returns the CVD Data Channel instance for the provided ID*/
+		TSharedPtr<FChaosVDOptionalDataChannel> GetChannelById(FName ChannelId)
+		{
+			if (TSharedPtr<FChaosVDOptionalDataChannel>* ChannelFound = AvailableChannels.Find(ChannelId))
+			{
+				return *ChannelFound;
+			}
+
+			return nullptr;
+		}
+
+	private:
+		
+		/** Registers a CVD Data channel instance - Do not call directly */
+		void RegisterChannel(const TSharedRef<FChaosVDOptionalDataChannel>& NewChannel)
+		{
+			AvailableChannels.Add(NewChannel->GetId(), NewChannel);
+		}
+
+		TMap<FName, TSharedPtr<FChaosVDOptionalDataChannel>> AvailableChannels;
+
+		friend FChaosVDOptionalDataChannel;
+	};
+
+	void ParseChannelListFromCommandArgument(TArray<FString>& OutParsedChannelList, const FString& InCommandArgument);
+
+	/** Creates a CVD Data Channel Instance - Only to be used by the CVD Macros */
+	TSharedRef<FChaosVDOptionalDataChannel> CreateDataChannel(FName InChannelID, const TSharedRef<FText>& InDisplayName, EChaosVDDataChannelInitializationFlags InitializationFlags);
+}
+
+#define CVD_CONCAT_NX(A, B) A ## B
+#define CVD_CONCAT(A, B) CVD_CONCAT_NX(A,B)
+#define CVD_STRINGIZE_NX(A) #A
+#define CVD_STRINGIZE(A) CVD_STRINGIZE_NX(A)
+
+/** Declares an Optional CVD Data channel to be available globally. The data channel can be accessed by using CVDDC_TheNameOfTheChannelUsedWithThisMacro */
+#define CVD_DECLARE_OPTIONAL_DATA_CHANNEL(DataChannelName) \
+			extern CHAOS_API TSharedRef<Chaos::VisualDebugger::FChaosVDOptionalDataChannel> CVDDC_##DataChannelName;
+
+/** Defines and initializes an Optional CVD Data Channel */
+#define CVD_DEFINE_OPTIONAL_DATA_CHANNEL(DataChannelName, InitializationFlags) \
+			TSharedRef<Chaos::VisualDebugger::FChaosVDOptionalDataChannel> CVDDC_##DataChannelName = Chaos::VisualDebugger::CreateDataChannel(#DataChannelName, MakeShared<FText>(NSLOCTEXT(CVD_STRINGIZE(ChaosVisualDebugger), CVD_STRINGIZE(CVD_CONCAT(DataChannelName,_ChannelName)), #DataChannelName)), InitializationFlags);
+
+// Declare CVD's Default set of channels
+
+CVD_DECLARE_OPTIONAL_DATA_CHANNEL(Default);
+CVD_DECLARE_OPTIONAL_DATA_CHANNEL(EvolutionStart);
+CVD_DECLARE_OPTIONAL_DATA_CHANNEL(Integrate);
+CVD_DECLARE_OPTIONAL_DATA_CHANNEL(ApplyKinematicTargets);
+CVD_DECLARE_OPTIONAL_DATA_CHANNEL(CollisionDetectionBroadPhase);
+CVD_DECLARE_OPTIONAL_DATA_CHANNEL(CollisionDetectionNarrowPhase);
+CVD_DECLARE_OPTIONAL_DATA_CHANNEL(EndOfEvolutionCollisionConstraints);
+CVD_DECLARE_OPTIONAL_DATA_CHANNEL(EvolutionEnd);
+CVD_DECLARE_OPTIONAL_DATA_CHANNEL(SceneQueries);
+
+#endif
diff --git a/Engine/Source/Runtime/Experimental/Chaos/Public/ChaosVisualDebugger/ChaosVDTraceMacros.h b/Engine/Source/Runtime/Experimental/Chaos/Public/ChaosVisualDebugger/ChaosVDTraceMacros.h
new file mode 100644
--- /dev/null
+++ b/Engine/Source/Runtime/Experimental/Chaos/Public/ChaosVisualDebugger/ChaosVDTraceMacros.h
@@ -0,0 +1,250 @@
+﻿// Copyright Epic Games, Inc. All Rights Reserved.
+
+#pragma once
+
+#if WITH_CHAOS_VISUAL_DEBUGGER
+
+	#ifndef CVD_DEFINE_TRACE_VECTOR
+		#define CVD_DEFINE_TRACE_VECTOR(Type, Name) \
+		UE_TRACE_EVENT_FIELD(Type, Name##X) \
+		UE_TRACE_EVENT_FIELD(Type, Name##Y) \
+		UE_TRACE_EVENT_FIELD(Type, Name##Z)
+	#endif
+
+	#ifndef CVD_DEFINE_TRACE_ROTATOR
+		#define CVD_DEFINE_TRACE_ROTATOR(Type, Name) \
+		UE_TRACE_EVENT_FIELD(Type, Name##X) \
+		UE_TRACE_EVENT_FIELD(Type, Name##Y) \
+		UE_TRACE_EVENT_FIELD(Type, Name##Z) \
+		UE_TRACE_EVENT_FIELD(Type, Name##W)
+	#endif
+
+	#ifndef CVD_TRACE_VECTOR_ON_EVENT
+		#define CVD_TRACE_VECTOR_ON_EVENT(EventName, Name, Vector) \
+		EventName.Name##X(Vector.X) \
+		<< EventName.Name##Y(Vector.Y) \
+		<< EventName.Name##Z(Vector.Z)
+	#endif
+
+	#ifndef CVD_TRACE_ROTATOR_ON_EVENT
+		#define CVD_TRACE_ROTATOR_ON_EVENT(EventName, Name, Rotator) \
+		EventName.Name##X(Rotator.X) \
+		<< EventName.Name##Y(Rotator.Y) \
+		<< EventName.Name##Z(Rotator.Z) \
+		<< EventName.Name##W(Rotator.W)
+	#endif
+
+	#ifndef CVD_TRACE_PARTICLE
+		#define CVD_TRACE_PARTICLE(ParticleHandle) \
+			FChaosVisualDebuggerTrace::TraceParticle(ParticleHandle);
+	#endif
+
+	#ifndef CVD_TRACE_PARTICLES
+		#define CVD_TRACE_PARTICLES(ParticleHandles) \
+			FChaosVisualDebuggerTrace::TraceParticles(ParticleHandles);
+	#endif
+
+	#ifndef CVD_TRACE_PARTICLES_VIEW
+		#define CVD_TRACE_PARTICLES_VIEW(ParticleHandlesView) \
+		FChaosVisualDebuggerTrace::TraceParticlesView(ParticleHandlesView);
+	#endif
+
+	#ifndef CVD_TRACE_PARTICLES_SOA
+		#define CVD_TRACE_PARTICLES_SOA(ParticleSoA) \
+		FChaosVisualDebuggerTrace::TraceParticlesSoA(ParticleSoA);
+	#endif
+
+	#ifndef CVD_TRACE_SOLVER_START_FRAME
+		#define CVD_TRACE_SOLVER_START_FRAME(SolverType, SolverRef) \
+			FChaosVDContext StartContextData; \
+			FChaosVisualDebuggerTrace::GetCVDContext<SolverType>(SolverRef, StartContextData); \
+			FChaosVisualDebuggerTrace::TraceSolverFrameStart(StartContextData, FChaosVisualDebuggerTrace::GetDebugName<SolverType>(SolverRef));
+	#endif
+
+	#ifndef CVD_TRACE_SOLVER_END_FRAME
+		#define CVD_TRACE_SOLVER_END_FRAME(SolverType, SolverRef) \
+			FChaosVDContext EndContextData; \
+			FChaosVisualDebuggerTrace::GetCVDContext<SolverType>(SolverRef, EndContextData); \
+			FChaosVisualDebuggerTrace::TraceSolverFrameEnd(EndContextData);
+	#endif
+
+	#ifndef CVD_SCOPE_TRACE_SOLVER_FRAME
+		#define CVD_SCOPE_TRACE_SOLVER_FRAME(SolverType, SolverRef) \
+			FChaosVDScopeSolverFrame<SolverType> ScopeSolverFrame(SolverRef);
+	#endif
+
+	#ifndef CVD_TRACE_SOLVER_STEP_START
+		#define CVD_TRACE_SOLVER_STEP_START(StepName) \
+			FChaosVisualDebuggerTrace::TraceSolverStepStart(StepName);
+	#endif
+
+	#ifndef CVD_TRACE_SOLVER_STEP_END
+		#define CVD_TRACE_SOLVER_STEP_END() \
+			FChaosVisualDebuggerTrace::TraceSolverStepEnd();
+	#endif
+
+	#ifndef CVD_SCOPE_TRACE_SOLVER_STEP
+		#define CVD_SCOPE_TRACE_SOLVER_STEP(DataChannel, StepName) \
+			CVD_SCOPED_DATA_CHANNEL_OVERRIDE(DataChannel) \
+			FChaosVDScopeSolverStep ScopeSolverStep(StepName);
+	#endif
+
+	#ifndef CVD_TRACE_BINARY_DATA
+		#define CVD_TRACE_BINARY_DATA(InData, TypeName) \
+		FChaosVisualDebuggerTrace::TraceBinaryData(InData, TypeName);
+	#endif
+
+	#ifndef CVD_TRACE_SOLVER_SIMULATION_SPACE
+		#define CVD_TRACE_SOLVER_SIMULATION_SPACE(InSimulationSpace) \
+		FChaosVisualDebuggerTrace::TraceSolverSimulationSpace(InSimulationSpace);
+	#endif
+
+	#ifndef CVD_TRACE_PARTICLE_DESTROYED
+		#define CVD_TRACE_PARTICLE_DESTROYED(DestroyedParticleHandle) \
+		FChaosVisualDebuggerTrace::TraceParticleDestroyed(DestroyedParticleHandle);
+	#endif
+
+	#ifndef CVD_TRACE_MID_PHASE
+		#define CVD_TRACE_MID_PHASE(MidPhase) \
+		FChaosVisualDebuggerTrace::TraceMidPhase(MidPhase);
+	#endif
+
+	#ifndef CVD_TRACE_COLLISION_CONSTRAINT
+		#define CVD_TRACE_COLLISION_CONSTRAINT(Constraint) \
+		FChaosVisualDebuggerTrace::TraceCollisionConstraint(Constraint);
+	#endif
+
+	#ifndef CVD_TRACE_COLLISION_CONSTRAINT_VIEW
+		#define CVD_TRACE_COLLISION_CONSTRAINT_VIEW(ConstraintView) \
+		FChaosVisualDebuggerTrace::TraceCollisionConstraintView(ConstraintView);
+	#endif
+
+	#ifndef CVD_TRACE_STEP_MID_PHASES_FROM_COLLISION_CONSTRAINTS
+		#define CVD_TRACE_STEP_MID_PHASES_FROM_COLLISION_CONSTRAINTS(DataChannel, CollisionConstraints) \
+		CVD_SCOPED_DATA_CHANNEL_OVERRIDE(DataChannel) \
+		FChaosVisualDebuggerTrace::TraceMidPhasesFromCollisionConstraints(CollisionConstraints);
+	#endif
+
+	#ifndef CVD_TRACE_NON_SOLVER_LOCATION
+		#define CVD_TRACE_NON_SOLVER_LOCATION(Location, DebugName) \
+		FChaosVisualDebuggerTrace::TraceNonSolverLocation(Location, DebugName);
+	#endif
+
+	#ifndef CVD_TRACE_NON_SOLVER_TRANSFORM
+		#define CVD_TRACE_NON_SOLVER_TRANSFORM(Transform, DebugName) \
+		FChaosVisualDebuggerTrace::TraceNonSolverTransform(Transform, DebugName);
+	#endif
+
+	#ifndef CVD_TRACE_INVALIDATE_CACHED_GEOMETRY
+		#define CVD_TRACE_INVALIDATE_CACHED_GEOMETRY(ImplicitObjectPtr) \
+		FChaosVisualDebuggerTrace::InvalidateGeometryFromCache(ImplicitObjectPtr);
+	#endif
+
+	#ifndef CVD_TRACE_SCOPED_SCENE_QUERY
+		#define CVD_TRACE_SCENE_QUERY_START(InputGeometry, GeometryOrientation, Start, End, TraceChannel, Params, ResponseParams, ObjectParams, QueryType, QueryMode, SolverID, bIsRetry) \
+		FChaosVisualDebuggerTrace::TraceSceneQueryStart(InputGeometry, GeometryOrientation, Start, End, TraceChannel, Params, ResponseParams, ObjectParams, QueryType, QueryMode, SolverID, bIsRetry);
+	#endif
+
+	#ifndef CVD_TRACE_SCOPED_SCENE_QUERY_VISIT
+		#define CVD_TRACE_SCOPED_SCENE_QUERY_VISIT(InQueryVisitData) \
+		FChaosVDScopeSceneQueryVisit CVDQueryVisit(InQueryVisitData);
+	#endif
+
+	#ifndef CVD_TRACE_SCENE_QUERY_VISIT
+		#define CVD_TRACE_SCENE_QUERY_VISIT(InQueryVisitData) \
+		FChaosVisualDebuggerTrace::TraceSceneQueryVisit(InQueryVisitData);
+	#endif
+
+#else // WITH_CHAOS_VISUAL_DEBUGGER
+
+	#ifndef CVD_TRACE_PARTICLE
+		#define CVD_TRACE_PARTICLE(ParticleHandle)
+	#endif
+
+	#ifndef CVD_TRACE_PARTICLES
+		#define CVD_TRACE_PARTICLES(ParticleHandles)
+	#endif
+
+	#ifndef CVD_TRACE_SOLVER_START_FRAME
+		#define CVD_TRACE_SOLVER_START_FRAME(SolverType, SolverRef)
+	#endif
+
+	#ifndef CVD_TRACE_SOLVER_END_FRAME
+		#define CVD_TRACE_SOLVER_END_FRAME(SolverType, SolverRef)
+	#endif
+
+	#ifndef CVD_SCOPE_TRACE_SOLVER_FRAME
+		#define CVD_SCOPE_TRACE_SOLVER_FRAME(SolverType, SolverRef)
+	#endif
+
+	#ifndef CVD_TRACE_SOLVER_STEP_START
+		#define CVD_TRACE_SOLVER_STEP_START(StepName)
+	#endif
+
+	#ifndef CVD_TRACE_SOLVER_STEP_END
+		#define CVD_TRACE_SOLVER_STEP_END()
+	#endif
+
+	#ifndef CVD_SCOPE_TRACE_SOLVER_STEP
+		#define CVD_SCOPE_TRACE_SOLVER_STEP(DataChannel, StepName)
+	#endif
+
+	#ifndef CVD_TRACE_BINARY_DATA
+		#define CVD_TRACE_BINARY_DATA(InData, TypeName)
+	#endif
+
+	#ifndef CVD_TRACE_SOLVER_SIMULATION_SPACE
+			#define CVD_TRACE_SOLVER_SIMULATION_SPACE(InSimulationSpace)
+	#endif
+
+	#ifndef CVD_TRACE_PARTICLES_SOA
+		#define CVD_TRACE_PARTICLES_SOA(ParticleSoA)
+	#endif
+
+	#ifndef CVD_TRACE_PARTICLE_DESTROYED
+		#define CVD_TRACE_PARTICLE_DESTROYED(DestroyedParticleHandle)
+	#endif
+
+	#ifndef CVD_TRACE_MID_PHASE
+		#define CVD_TRACE_MID_PHASE(MidPhase)
+	#endif
+	#ifndef CVD_TRACE_COLLISION_CONSTRAINT
+		#define CVD_TRACE_COLLISION_CONSTRAINT(Constraint)
+	#endif
+	#ifndef CVD_TRACE_COLLISION_CONSTRAINT_VIEW
+		#define CVD_TRACE_COLLISION_CONSTRAINT_VIEW(ConstraintView)
+	#endif
+
+	#ifndef CVD_TRACE_PARTICLES_VIEW
+		#define CVD_TRACE_PARTICLES_VIEW(ParticleHandlesView)
+	#endif
+
+	#ifndef CVD_TRACE_STEP_MID_PHASES_FROM_COLLISION_CONSTRAINTS
+		#define CVD_TRACE_STEP_MID_PHASES_FROM_COLLISION_CONSTRAINTS(DataChannel, CollisionConstraints)
+	#endif
+
+	#ifndef CVD_TRACE_NON_SOLVER_LOCATION
+	      #define CVD_TRACE_NON_SOLVER_LOCATION(Location, DebugName)
+	#endif
+
+	#ifndef CVD_TRACE_NON_SOLVER_TRANSFORM
+		#define CVD_TRACE_NON_SOLVER_TRANSFORM(Transform, DebugName)
+	#endif
+
+	#ifndef CVD_TRACE_INVALIDATE_CACHED_GEOMETRY
+		#define CVD_TRACE_INVALIDATE_CACHED_GEOMETRY(ImplicitObjectPtr)
+	#endif
+
+	#ifndef CVD_TRACE_SCOPED_SCENE_QUERY
+		#define CVD_TRACE_SCENE_QUERY_START(InputGeometry, GeometryOrientation, Start, End, TraceChannel, Params, ResponseParams, ObjectParams, QueryType, QueryMode, SolverID)
+	#endif
+
+	#ifndef CVD_TRACE_SCOPED_SCENE_QUERY_VISIT
+		#define CVD_TRACE_SCOPED_SCENE_QUERY_VISIT(InQueryVisitData)
+	#endif
+
+	#ifndef CVD_TRACE_SCENE_QUERY_VISIT
+		#define CVD_TRACE_SCENE_QUERY_VISIT(InQueryVisitData)
+	#endif
+
+#endif // WITH_CHAOS_VISUAL_DEBUGGER
\ No newline at end of file
diff --git a/Engine/Source/Runtime/Experimental/Chaos/Public/ChaosVisualDebugger/ChaosVisualDebuggerTrace.h b/Engine/Source/Runtime/Experimental/Chaos/Public/ChaosVisualDebugger/ChaosVisualDebuggerTrace.h
--- a/Engine/Source/Runtime/Experimental/Chaos/Public/ChaosVisualDebugger/ChaosVisualDebuggerTrace.h
+++ b/Engine/Source/Runtime/Experimental/Chaos/Public/ChaosVisualDebugger/ChaosVisualDebuggerTrace.h
@@ -8,142 +8,23 @@
 
 #include "Chaos/ChaosArchive.h"
 #include "ChaosVDContextProvider.h"
+#include "ChaosVDOptionalDataChannel.h"
 #include "Chaos/ParticleIterator.h"
 #include "HAL/ThreadSafeBool.h"
 
 #ifndef CHAOS_VISUAL_DEBUGGER_ENABLED
-	#define CHAOS_VISUAL_DEBUGGER_ENABLED (WITH_CHAOS_VISUAL_DEBUGGER && UE_TRACE_ENABLED)
+	#define CHAOS_VISUAL_DEBUGGER_ENABLED WITH_CHAOS_VISUAL_DEBUGGER
 #endif
 
-// Define NO-OP versions of our macros if we can't Trace
-#if !CHAOS_VISUAL_DEBUGGER_ENABLED
-	#ifndef CVD_TRACE_PARTICLE
-		#define CVD_TRACE_PARTICLE(ParticleHandle)
-	#endif
+#include "ChaosVisualDebugger/ChaosVDTraceMacros.h"
 
-	#ifndef CVD_TRACE_PARTICLES
-		#define CVD_TRACE_PARTICLES(ParticleHandles)
-	#endif
-
-	#ifndef CVD_TRACE_SOLVER_START_FRAME
-		#define CVD_TRACE_SOLVER_START_FRAME(SolverType, SolverRef)
-	#endif
-
-	#ifndef CVD_TRACE_SOLVER_END_FRAME
-		#define CVD_TRACE_SOLVER_END_FRAME(SolverType, SolverRef)
-	#endif
-
-	#ifndef CVD_SCOPE_TRACE_SOLVER_FRAME
-		#define CVD_SCOPE_TRACE_SOLVER_FRAME(SolverType, SolverRef)
-	#endif
-
-	#ifndef CVD_TRACE_SOLVER_STEP_START
-		#define CVD_TRACE_SOLVER_STEP_START(StepName)
-	#endif
-
-	#ifndef CVD_TRACE_SOLVER_STEP_END
-		#define CVD_TRACE_SOLVER_STEP_END()
-	#endif
-
-	#ifndef CVD_SCOPE_TRACE_SOLVER_STEP
-		#define CVD_SCOPE_TRACE_SOLVER_STEP(StepName)
-	#endif
-
-	#ifndef CVD_TRACE_BINARY_DATA
-		#define CVD_TRACE_BINARY_DATA(InData, TypeName)
-	#endif
-
-#ifndef CVD_TRACE_SOLVER_SIMULATION_SPACE
-		#define CVD_TRACE_SOLVER_SIMULATION_SPACE(InSimulationSpace)
-#endif
-
-#ifndef CVD_TRACE_PARTICLES_SOA
-	#define CVD_TRACE_PARTICLES_SOA(ParticleSoA)
-#endif
-
-#ifndef CVD_TRACE_PARTICLE_DESTROYED
-	#define CVD_TRACE_PARTICLE_DESTROYED(DestroyedParticleHandle)
-#endif
-
-#ifndef CVD_TRACE_MID_PHASE
-	#define CVD_TRACE_MID_PHASE(MidPhase)
-#endif
-#ifndef CVD_TRACE_COLLISION_CONSTRAINT
-	#define CVD_TRACE_COLLISION_CONSTRAINT(Constraint)
-#endif
-#ifndef CVD_TRACE_COLLISION_CONSTRAINT_VIEW
-	#define CVD_TRACE_COLLISION_CONSTRAINT_VIEW(ConstraintView)
-#endif
-
-#ifndef CVD_TRACE_PARTICLES_VIEW
-	#define CVD_TRACE_PARTICLES_VIEW(ParticleHandlesView)
-#endif
-
-#ifndef CVD_TRACE_MID_PHASES_FROM_COLLISION_CONSTRAINTS
-	#define CVD_TRACE_MID_PHASES_FROM_COLLISION_CONSTRAINTS(CollisionConstraints)
-#endif
-
-#ifndef CVD_TRACE_NON_SOLVER_LOCATION
-      #define CVD_TRACE_NON_SOLVER_LOCATION(Location, DebugName)
-#endif
-
-#ifndef CVD_TRACE_NON_SOLVER_TRANSFORM
-	#define CVD_TRACE_NON_SOLVER_TRANSFORM(Transform, DebugName)
-#endif
-
-#ifndef CVD_TRACE_INVALIDATE_CACHED_GEOMETRY
-	#define CVD_TRACE_INVALIDATE_CACHED_GEOMETRY(ImplicitObjectPtr)
-#endif
-
-#ifndef CVD_TRACE_SCOPED_SCENE_QUERY
-	#define CVD_TRACE_SCENE_QUERY_START(InputGeometry, GeometryOrientation, Start, End, TraceChannel, Params, ResponseParams, ObjectParams, QueryType, QueryMode, SolverID)
-#endif
-
-#ifndef CVD_TRACE_SCOPED_SCENE_QUERY_VISIT
-	#define CVD_TRACE_SCOPED_SCENE_QUERY_VISIT(InQueryVisitData)
-#endif
-
-#ifndef CVD_TRACE_SCENE_QUERY_VISIT
-	#define CVD_TRACE_SCENE_QUERY_VISIT(InQueryVisitData)
-#endif
-
-#else
+#if WITH_CHAOS_VISUAL_DEBUGGER
 
 #include "ChaosVDRuntimeModule.h"
 #include "DataWrappers/ChaosVDImplicitObjectDataWrapper.h"
 #include "Trace/Trace.h"
 #include "Trace/Trace.inl"
 
-#ifndef CVD_DEFINE_TRACE_VECTOR
-	#define CVD_DEFINE_TRACE_VECTOR(Type, Name) \
-		UE_TRACE_EVENT_FIELD(Type, Name##X) \
-		UE_TRACE_EVENT_FIELD(Type, Name##Y) \
-		UE_TRACE_EVENT_FIELD(Type, Name##Z)
-#endif
-
-#ifndef CVD_DEFINE_TRACE_ROTATOR
-	#define CVD_DEFINE_TRACE_ROTATOR(Type, Name) \
-	UE_TRACE_EVENT_FIELD(Type, Name##X) \
-	UE_TRACE_EVENT_FIELD(Type, Name##Y) \
-	UE_TRACE_EVENT_FIELD(Type, Name##Z) \
-	UE_TRACE_EVENT_FIELD(Type, Name##W)
-#endif
-
-#ifndef CVD_TRACE_VECTOR_ON_EVENT
-	#define CVD_TRACE_VECTOR_ON_EVENT(EventName, Name, Vector) \
-	EventName.Name##X(Vector.X) \
-	<< EventName.Name##Y(Vector.Y) \
-	<< EventName.Name##Z(Vector.Z)
-#endif
-
-#ifndef CVD_TRACE_ROTATOR_ON_EVENT
-	#define CVD_TRACE_ROTATOR_ON_EVENT(EventName, Name, Rotator) \
-	EventName.Name##X(Rotator.X) \
-	<< EventName.Name##Y(Rotator.Y) \
-	<< EventName.Name##Z(Rotator.Z) \
-	<< EventName.Name##W(Rotator.W)
-#endif
-
 UE_TRACE_CHANNEL_EXTERN(ChaosVDChannel, CHAOS_API)
 
 UE_TRACE_EVENT_BEGIN_EXTERN(ChaosVDLogger, ChaosVDSolverFrameStart)
@@ -227,125 +108,6 @@ UE_TRACE_EVENT_BEGIN_EXTERN(ChaosVDLogger, ChaosVDDummyEvent)
 	UE_TRACE_EVENT_FIELD(int32, SolverID)
 UE_TRACE_EVENT_END()
 
-#ifndef CVD_TRACE_PARTICLE
-	#define CVD_TRACE_PARTICLE(ParticleHandle) \
-		FChaosVisualDebuggerTrace::TraceParticle(ParticleHandle);
-#endif
-
-#ifndef CVD_TRACE_PARTICLES
-	#define CVD_TRACE_PARTICLES(ParticleHandles) \
-		FChaosVisualDebuggerTrace::TraceParticles(ParticleHandles);
-#endif
-
-#ifndef CVD_TRACE_PARTICLES_VIEW
-	#define CVD_TRACE_PARTICLES_VIEW(ParticleHandlesView) \
-	FChaosVisualDebuggerTrace::TraceParticlesView(ParticleHandlesView);
-#endif
-
-#ifndef CVD_TRACE_PARTICLES_SOA
-	#define CVD_TRACE_PARTICLES_SOA(ParticleSoA) \
-	FChaosVisualDebuggerTrace::TraceParticlesSoA(ParticleSoA);
-#endif
-
-#ifndef CVD_TRACE_SOLVER_START_FRAME
-	#define CVD_TRACE_SOLVER_START_FRAME(SolverType, SolverRef) \
-		FChaosVDContext StartContextData; \
-		FChaosVisualDebuggerTrace::GetCVDContext<SolverType>(SolverRef, StartContextData); \
-		FChaosVisualDebuggerTrace::TraceSolverFrameStart(StartContextData, FChaosVisualDebuggerTrace::GetDebugName<SolverType>(SolverRef));
-#endif
-
-#ifndef CVD_TRACE_SOLVER_END_FRAME
-	#define CVD_TRACE_SOLVER_END_FRAME(SolverType, SolverRef) \
-		FChaosVDContext EndContextData; \
-		FChaosVisualDebuggerTrace::GetCVDContext<SolverType>(SolverRef, EndContextData); \
-		FChaosVisualDebuggerTrace::TraceSolverFrameEnd(EndContextData);
-#endif
-
-#ifndef CVD_SCOPE_TRACE_SOLVER_FRAME
-	#define CVD_SCOPE_TRACE_SOLVER_FRAME(SolverType, SolverRef) \
-		FChaosVDScopeSolverFrame<SolverType> ScopeSolverFrame(SolverRef);
-#endif
-
-#ifndef CVD_TRACE_SOLVER_STEP_START
-	#define CVD_TRACE_SOLVER_STEP_START(StepName) \
-		FChaosVisualDebuggerTrace::TraceSolverStepStart(StepName);
-#endif
-
-#ifndef CVD_TRACE_SOLVER_STEP_END
-	#define CVD_TRACE_SOLVER_STEP_END() \
-		FChaosVisualDebuggerTrace::TraceSolverStepEnd();
-#endif
-
-#ifndef CVD_SCOPE_TRACE_SOLVER_STEP
-	#define CVD_SCOPE_TRACE_SOLVER_STEP(StepName) \
-		FChaosVDScopeSolverStep ScopeSolverStep(StepName);
-#endif
-
-#ifndef CVD_TRACE_BINARY_DATA
-	#define CVD_TRACE_BINARY_DATA(InData, TypeName) \
-	FChaosVisualDebuggerTrace::TraceBinaryData(InData, TypeName);
-#endif
-
-#ifndef CVD_TRACE_SOLVER_SIMULATION_SPACE
-	#define CVD_TRACE_SOLVER_SIMULATION_SPACE(InSimulationSpace) \
-	FChaosVisualDebuggerTrace::TraceSolverSimulationSpace(InSimulationSpace);
-#endif
-
-#ifndef CVD_TRACE_PARTICLE_DESTROYED
-	#define CVD_TRACE_PARTICLE_DESTROYED(DestroyedParticleHandle) \
-	FChaosVisualDebuggerTrace::TraceParticleDestroyed(DestroyedParticleHandle);
-#endif
-
-#ifndef CVD_TRACE_MID_PHASE
-	#define CVD_TRACE_MID_PHASE(MidPhase) \
-	FChaosVisualDebuggerTrace::TraceMidPhase(MidPhase);
-#endif
-
-#ifndef CVD_TRACE_COLLISION_CONSTRAINT
-	#define CVD_TRACE_COLLISION_CONSTRAINT(Constraint) \
-	FChaosVisualDebuggerTrace::TraceCollisionConstraint(Constraint);
-#endif
-
-#ifndef CVD_TRACE_COLLISION_CONSTRAINT_VIEW
-	#define CVD_TRACE_COLLISION_CONSTRAINT_VIEW(ConstraintView) \
-	FChaosVisualDebuggerTrace::TraceCollisionConstraintView(ConstraintView);
-#endif
-
-#ifndef CVD_TRACE_MID_PHASES_FROM_COLLISION_CONSTRAINTS
-	#define CVD_TRACE_MID_PHASES_FROM_COLLISION_CONSTRAINTS(CollisionConstraints) \
-	FChaosVisualDebuggerTrace::TraceMidPhasesFromCollisionConstraints(CollisionConstraints);
-#endif
-
-#ifndef CVD_TRACE_NON_SOLVER_LOCATION
-	#define CVD_TRACE_NON_SOLVER_LOCATION(Location, DebugName) \
-	FChaosVisualDebuggerTrace::TraceNonSolverLocation(Location, DebugName);
-#endif
-
-#ifndef CVD_TRACE_NON_SOLVER_TRANSFORM
-	#define CVD_TRACE_NON_SOLVER_TRANSFORM(Transform, DebugName) \
-	FChaosVisualDebuggerTrace::TraceNonSolverTransform(Transform, DebugName);
-#endif
-
-#ifndef CVD_TRACE_INVALIDATE_CACHED_GEOMETRY
-	#define CVD_TRACE_INVALIDATE_CACHED_GEOMETRY(ImplicitObjectPtr) \
-	FChaosVisualDebuggerTrace::InvalidateGeometryFromCache(ImplicitObjectPtr);
-#endif
-
-#ifndef CVD_TRACE_SCOPED_SCENE_QUERY
-	#define CVD_TRACE_SCENE_QUERY_START(InputGeometry, GeometryOrientation, Start, End, TraceChannel, Params, ResponseParams, ObjectParams, QueryType, QueryMode, SolverID, bIsRetry) \
-	FChaosVisualDebuggerTrace::TraceSceneQueryStart(InputGeometry, GeometryOrientation, Start, End, TraceChannel, Params, ResponseParams, ObjectParams, QueryType, QueryMode, SolverID, bIsRetry);
-#endif
-
-#ifndef CVD_TRACE_SCOPED_SCENE_QUERY_VISIT
-	#define CVD_TRACE_SCOPED_SCENE_QUERY_VISIT(InQueryVisitData) \
-	FChaosVDScopeSceneQueryVisit CVDQueryVisit(InQueryVisitData);
-#endif
-
-#ifndef CVD_TRACE_SCENE_QUERY_VISIT
-	#define CVD_TRACE_SCENE_QUERY_VISIT(InQueryVisitData) \
-	FChaosVisualDebuggerTrace::TraceSceneQueryVisit(InQueryVisitData);
-#endif
-
 struct FChaosVDContext;
 struct FChaosVDQueryVisitStep;
 struct FChaosVDCollisionResponseParams;
@@ -382,13 +144,6 @@ public:
 	 * @param ParticleHandle Handle to process and Trace
 	 */
 	static CHAOS_API void TraceParticle(const Chaos::FGeometryParticleHandle* ParticleHandle);
-	
-	/**
-	 * Traces data from a Particle Handle using the provided CVD Context
-	 * @param ParticleHandle Handle to process and Trace
-	 * @param ContextData Context to be used to tied this Trace event to a specific solver frame and step
-	 */
-	static CHAOS_API void TraceParticle(Chaos::FGeometryParticleHandle* ParticleHandle, const FChaosVDContext& ContextData);
 
 	/**
 	 * Traces data from a collection of Particle Handles using. The CVD context currently pushed into will be used to tie this particle data to a specific solver frame and step.
@@ -510,6 +265,13 @@ public:
 	static void UnregisterEventHandlers();
 
 private:
+
+	/**
+	 * Traces data from a Particle Handle using the provided CVD Context
+	 * @param ParticleHandle Handle to process and Trace
+	 * @param ContextData Context to be used to tied this Trace event to a specific solver frame and step
+	 */
+	static CHAOS_API void TraceParticle(Chaos::FGeometryParticleHandle* ParticleHandle, const FChaosVDContext& ContextData);
 	
 	/** Resets the state of the CVD Tracer */
 	static void Reset();
@@ -609,7 +371,6 @@ struct FChaosVDScopeSceneQueryVisit
 {
 	FChaosVDScopeSceneQueryVisit(FChaosVDQueryVisitStep& InVisitData) : VisitData(InVisitData)
 	{
-		
 	}
 
 	~FChaosVDScopeSceneQueryVisit()
@@ -619,4 +380,4 @@ struct FChaosVDScopeSceneQueryVisit
 
 	FChaosVDQueryVisitStep& VisitData;
 };
-#endif // CHAOS_VISUAL_DEBUGGER_ENABLED
+#endif // WITH_CHAOS_VISUAL_DEBUGGER
diff --git a/Engine/Source/Runtime/PhysicsCore/Public/ChaosVDSQVisitorHelpers.h b/Engine/Source/Runtime/PhysicsCore/Public/ChaosVDSQVisitorHelpers.h
--- a/Engine/Source/Runtime/PhysicsCore/Public/ChaosVDSQVisitorHelpers.h
+++ b/Engine/Source/Runtime/PhysicsCore/Public/ChaosVDSQVisitorHelpers.h
@@ -11,6 +11,7 @@
 /** Initializes a CVD Wrapper for a SQ visit with the provider data, and traces it as soon it goes iut of scope if we are recording */
 #ifndef CVD_TRACE_SCOPED_SCENE_QUERY_VISIT_HELPER
 	#define CVD_TRACE_SCOPED_SCENE_QUERY_VISIT_HELPER(Type, ParticleTransform, ParticleIndex, ShapeIndex, CurData) \
+	CVD_SCOPED_DATA_CHANNEL_OVERRIDE(CVDDC_SceneQueries) \
 	FChaosVDQueryVisitStep CVDSQVisitStepData; \
 	CVD_TRACE_SCOPED_SCENE_QUERY_VISIT(CVDSQVisitStepData); \
 	Chaos::VisualDebugger::TraceHelpers::FillCVDQueryVisitData(CVDSQVisitStepData, Type, ParticleTransform, ParticleIndex, ShapeIndex, CurData);
