[ChaosVD] Fix for CIS issue on ChaosVDOptionalDataChannel.h.

There is a macro that was already defined in another CVD module.

#changelist validated
#virtualized

diff --git a/Engine/Plugins/Experimental/ChaosVD/Source/ChaosVD/Public/Trace/ChaosVDTraceAnalyzer.h b/Engine/Plugins/Experimental/ChaosVD/Source/ChaosVD/Public/Trace/ChaosVDTraceAnalyzer.h
--- a/Engine/Plugins/Experimental/ChaosVD/Source/ChaosVD/Public/Trace/ChaosVDTraceAnalyzer.h
+++ b/Engine/Plugins/Experimental/ChaosVD/Source/ChaosVD/Public/Trace/ChaosVDTraceAnalyzer.h
@@ -5,6 +5,7 @@
 #include "ChaosVDRecording.h"
 #include "Trace/Analyzer.h"
 #include "Templates/SharedPointer.h"
+#include "ChaosVisualDebugger/ChaosVDOptionalDataChannel.h"
 
 struct FChaosVDStepData;
 class FChaosVDTraceProvider;
@@ -58,10 +59,6 @@ private:
 	FChaosVDTraceAnalysisComplete ChaosVDTraceAnalysisCompleteDelegate;
 };
 
-#ifndef CVD_STRINGIZE
-	#define CVD_STRINGIZE(x) #x
-#endif
-
 #ifndef CVD_READ_TRACE_VECTOR
 	#define CVD_READ_TRACE_VECTOR(Vector, VectorName, ValueType, EventData) \
 		Vector.X = EventData.GetValue<ValueType>(CVD_STRINGIZE(VectorName##X)); \
diff --git a/Engine/Source/Runtime/Experimental/Chaos/Public/ChaosVisualDebugger/ChaosVDOptionalDataChannel.h b/Engine/Source/Runtime/Experimental/Chaos/Public/ChaosVisualDebugger/ChaosVDOptionalDataChannel.h
--- a/Engine/Source/Runtime/Experimental/Chaos/Public/ChaosVisualDebugger/ChaosVDOptionalDataChannel.h
+++ b/Engine/Source/Runtime/Experimental/Chaos/Public/ChaosVisualDebugger/ChaosVDOptionalDataChannel.h
@@ -108,21 +108,35 @@ namespace Chaos::VisualDebugger
 	TSharedRef<FChaosVDOptionalDataChannel> CreateDataChannel(FName InChannelID, const TSharedRef<FText>& InDisplayName, EChaosVDDataChannelInitializationFlags InitializationFlags);
 }
 
-#define CVD_CONCAT_NX(A, B) A ## B
-#define CVD_CONCAT(A, B) CVD_CONCAT_NX(A,B)
-#define CVD_STRINGIZE_NX(A) #A
-#define CVD_STRINGIZE(A) CVD_STRINGIZE_NX(A)
+#ifndef CVD_CONCAT_NX
+	#define CVD_CONCAT_NX(A, B) A ## B
+#endif
+
+#ifndef CVD_CONCAT
+	#define CVD_CONCAT(A, B) CVD_CONCAT_NX(A,B)
+#endif
+
+#ifndef CVD_STRINGIZE_NX
+	#define CVD_STRINGIZE_NX(A) #A
+#endif
+
+#ifndef CVD_STRINGIZE
+	#define CVD_STRINGIZE(A) CVD_STRINGIZE_NX(A)
+#endif
 
 /** Declares an Optional CVD Data channel to be available globally. The data channel can be accessed by using CVDDC_TheNameOfTheChannelUsedWithThisMacro */
-#define CVD_DECLARE_OPTIONAL_DATA_CHANNEL(DataChannelName) \
-			extern CHAOS_API TSharedRef<Chaos::VisualDebugger::FChaosVDOptionalDataChannel> CVDDC_##DataChannelName;
+#ifndef CVD_DECLARE_OPTIONAL_DATA_CHANNEL
+	#define CVD_DECLARE_OPTIONAL_DATA_CHANNEL(DataChannelName) \
+				extern CHAOS_API TSharedRef<Chaos::VisualDebugger::FChaosVDOptionalDataChannel> CVDDC_##DataChannelName;
+#endif
 
 /** Defines and initializes an Optional CVD Data Channel */
-#define CVD_DEFINE_OPTIONAL_DATA_CHANNEL(DataChannelName, InitializationFlags) \
-			TSharedRef<Chaos::VisualDebugger::FChaosVDOptionalDataChannel> CVDDC_##DataChannelName = Chaos::VisualDebugger::CreateDataChannel(#DataChannelName, MakeShared<FText>(NSLOCTEXT(CVD_STRINGIZE(ChaosVisualDebugger), CVD_STRINGIZE(CVD_CONCAT(DataChannelName,_ChannelName)), #DataChannelName)), InitializationFlags);
+#ifndef CVD_DEFINE_OPTIONAL_DATA_CHANNEL
+	#define CVD_DEFINE_OPTIONAL_DATA_CHANNEL(DataChannelName, InitializationFlags) \
+				TSharedRef<Chaos::VisualDebugger::FChaosVDOptionalDataChannel> CVDDC_##DataChannelName = Chaos::VisualDebugger::CreateDataChannel(#DataChannelName, MakeShared<FText>(NSLOCTEXT(CVD_STRINGIZE(ChaosVisualDebugger), CVD_STRINGIZE(CVD_CONCAT(DataChannelName,_ChannelName)), #DataChannelName)), InitializationFlags);
+#endif
 
 // Declare CVD's Default set of channels
-
 CVD_DECLARE_OPTIONAL_DATA_CHANNEL(Default);
 CVD_DECLARE_OPTIONAL_DATA_CHANNEL(EvolutionStart);
 CVD_DECLARE_OPTIONAL_DATA_CHANNEL(Integrate);
