If bAllowPerInstanceRemoval was set, an instance would be removed if the transform scale is 0 and readded when transform scale is valid.  The process of removing it erases its custom instance data and there is no way to restore it when the instance is readded.

This CL makes a shadow copy of the custom instance data if bAllowPerInstanceRemoval is set which allows the process of readding the instance to also restore its custom instance data.

#changelist validated
#virtualized

diff --git a/Engine/Source/Runtime/Experimental/GeometryCollectionEngine/Private/GeometryCollection/GeometryCollectionISMPoolComponent.cpp b/Engine/Source/Runtime/Experimental/GeometryCollectionEngine/Private/GeometryCollection/GeometryCollectionISMPoolComponent.cpp
--- a/Engine/Source/Runtime/Experimental/GeometryCollectionEngine/Private/GeometryCollection/GeometryCollectionISMPoolComponent.cpp
+++ b/Engine/Source/Runtime/Experimental/GeometryCollectionEngine/Private/GeometryCollection/GeometryCollectionISMPoolComponent.cpp
@@ -38,9 +38,33 @@ FAutoConsoleVariableRef CVarISMPoolComponentFreeListTargetSize(
 	TEXT("Target size for number of ISM components in the recycling free list."));
 
 
-FGeometryCollectionMeshGroup::FMeshId FGeometryCollectionMeshGroup::AddMesh(const FGeometryCollectionStaticMeshInstance& MeshInstance, int32 InstanceCount, const FGeometryCollectionMeshInfo& ISMInstanceInfo)
+void FGeometryCollectionMeshInfo::ShadowCopyCustomData(int32 InstanceCount, int32 NumCustomDataFloatsPerInstance, TArrayView<const float> CustomDataFloats)
+{
+	CustomData.SetNum(InstanceCount * NumCustomDataFloatsPerInstance + NumCustomDataFloatsPerInstance, true);
+
+	for (int32 InstanceIndex = 0; InstanceIndex < InstanceCount; ++InstanceIndex)
+	{
+		int32 Offset = InstanceIndex * NumCustomDataFloatsPerInstance;
+		FMemory::Memcpy(&CustomData[Offset], CustomDataFloats.GetData() + Offset, NumCustomDataFloatsPerInstance * CustomDataFloats.GetTypeSize());
+	}
+}
+
+TArrayView<const float> FGeometryCollectionMeshInfo::CustomDataSlice(int32 InstanceIndex, int32 NumCustomDataFloatsPerInstance)
+{
+	TArrayView<const float> DataView = CustomData;
+	return DataView.Slice(InstanceIndex * NumCustomDataFloatsPerInstance, NumCustomDataFloatsPerInstance);
+}
+
+FGeometryCollectionMeshGroup::FMeshId FGeometryCollectionMeshGroup::AddMesh(const FGeometryCollectionStaticMeshInstance& MeshInstance, int32 InstanceCount, const FGeometryCollectionMeshInfo& ISMInstanceInfo, TArrayView<const float> CustomDataFloats)
 {
 	const FMeshId MeshInfoIndex = MeshInfos.Emplace(ISMInstanceInfo);
+
+	if (bAllowPerInstanceRemoval)
+	{
+		FGeometryCollectionMeshInfo& MeshInfo = MeshInfos[MeshInfoIndex];
+		MeshInfo.ShadowCopyCustomData(InstanceCount, MeshInstance.Desc.NumCustomDataFloats, CustomDataFloats);
+	}
+
 	return MeshInfoIndex;
 }
 
@@ -298,7 +322,9 @@ bool FGeometryCollectionISMPool::BatchUpdateInstancesTransforms(FGeometryCollect
 			else if (!Transform.GetScale3D().IsZero() && !InstanceId.IsValid())
 			{
 				// Re-add the instance to the ISM if the scale becomes non-zero.
-				ISM.InstanceIds[InstanceGroup.Start + InstanceIndex] = ISM.ISMComponent->AddInstanceById(Transform, bWorldSpace);
+				FPrimitiveInstanceId Id = ISM.ISMComponent->AddInstanceById(Transform, bWorldSpace);
+				ISM.InstanceIds[InstanceGroup.Start + InstanceIndex] = Id;
+				ISM.ISMComponent->SetCustomDataById(Id, MeshInfo.CustomDataSlice(InstanceIndex, ISM.ISMComponent->NumCustomDataFloats));
 				continue;
 			}
 		}
@@ -604,7 +630,7 @@ UGeometryCollectionISMPoolComponent::FMeshId UGeometryCollectionISMPoolComponent
 	if (FGeometryCollectionMeshGroup* MeshGroup = MeshGroups.Find(MeshGroupId))
 	{
 		const FGeometryCollectionMeshInfo ISMInstanceInfo = Pool.AddInstancesToISM(this, MeshInstance, InstanceCount, CustomDataFloats);
-		return MeshGroup->AddMesh(MeshInstance, InstanceCount, ISMInstanceInfo);
+		return MeshGroup->AddMesh(MeshInstance, InstanceCount, ISMInstanceInfo, CustomDataFloats);
 	}
 	UE_LOG(LogChaos, Warning, TEXT("UGeometryCollectionISMPoolComponent : Trying to add a mesh to a mesh group (%d) that does not exists"), MeshGroupId);
 	return INDEX_NONE;
diff --git a/Engine/Source/Runtime/Experimental/GeometryCollectionEngine/Public/GeometryCollection/GeometryCollectionISMPoolComponent.h b/Engine/Source/Runtime/Experimental/GeometryCollectionEngine/Public/GeometryCollection/GeometryCollectionISMPoolComponent.h
--- a/Engine/Source/Runtime/Experimental/GeometryCollectionEngine/Public/GeometryCollection/GeometryCollectionISMPoolComponent.h
+++ b/Engine/Source/Runtime/Experimental/GeometryCollectionEngine/Public/GeometryCollection/GeometryCollectionISMPoolComponent.h
@@ -220,6 +220,10 @@ struct FGeometryCollectionMeshInfo
 {
 	int32 ISMIndex;
 	FInstanceGroups::FInstanceGroupId InstanceGroupIndex;
+
+	TArray<float> CustomData;
+	TArrayView<const float> CustomDataSlice(int32 InstanceIndex, int32 NumCustomDataFloatsPerInstance);
+	void ShadowCopyCustomData(int32 InstanceCount, int32 NumCustomDataFloatsPerInstance, TArrayView<const float> CustomDataFloats);
 };
 
 struct FGeometryCollectionISMPool;
@@ -233,7 +237,7 @@ struct FGeometryCollectionMeshGroup
 	using FMeshId = int32;
 
 	/** Adds a new mesh with instance count. We expect to only add a unique mesh instance once to each group. Returns a ID that can be used to update the instances. */
-	FMeshId AddMesh(const FGeometryCollectionStaticMeshInstance& MeshInstance, int32 InstanceCount, const FGeometryCollectionMeshInfo& ISMInstanceInfo);
+	FMeshId AddMesh(const FGeometryCollectionStaticMeshInstance& MeshInstance, int32 InstanceCount, const FGeometryCollectionMeshInfo& ISMInstanceInfo, TArrayView<const float> CustomDataFloats);
 	/** Update instance transforms for a group of instances. */
 	bool BatchUpdateInstancesTransforms(FGeometryCollectionISMPool& ISMPool, FMeshId MeshId, int32 StartInstanceIndex, TArrayView<const FTransform> NewInstancesTransforms, bool bWorldSpace, bool bMarkRenderStateDirty, bool bTeleport);
 	void BatchUpdateInstanceCustomData(FGeometryCollectionISMPool& ISMPool, int32 CustomFloatIndex, float CustomFloatValue);
