Robomerge fix - use default arg for SetNum

Fixes this previous submission...

If bAllowPerInstanceRemoval was set, an instance would be removed if the transform scale is 0 and readded when transform scale is valid.  The process of removing it erases its custom instance data and there is no way to restore it when the instance is readded.

This CL makes a shadow copy of the custom instance data if bAllowPerInstanceRemoval is set which allows the process of readding the instance to also restore its custom instance data.

#changelist validated
#virtualized

diff --git a/Engine/Source/Runtime/Experimental/GeometryCollectionEngine/Private/GeometryCollection/GeometryCollectionISMPoolComponent.cpp b/Engine/Source/Runtime/Experimental/GeometryCollectionEngine/Private/GeometryCollection/GeometryCollectionISMPoolComponent.cpp
--- a/Engine/Source/Runtime/Experimental/GeometryCollectionEngine/Private/GeometryCollection/GeometryCollectionISMPoolComponent.cpp
+++ b/Engine/Source/Runtime/Experimental/GeometryCollectionEngine/Private/GeometryCollection/GeometryCollectionISMPoolComponent.cpp
@@ -40,7 +40,7 @@ FAutoConsoleVariableRef CVarISMPoolComponentFreeListTargetSize(
 
 void FGeometryCollectionMeshInfo::ShadowCopyCustomData(int32 InstanceCount, int32 NumCustomDataFloatsPerInstance, TArrayView<const float> CustomDataFloats)
 {
-	CustomData.SetNum(InstanceCount * NumCustomDataFloatsPerInstance + NumCustomDataFloatsPerInstance, true);
+	CustomData.SetNum(InstanceCount * NumCustomDataFloatsPerInstance + NumCustomDataFloatsPerInstance);
 
 	for (int32 InstanceIndex = 0; InstanceIndex < InstanceCount; ++InstanceIndex)
 	{
