Unshelved from pending changelist '30895702':

Fix geometry collection dynamic data replication
- Some members of the structure were not properly netserialized
- Root data was always ignore
- Dynamic and static rep data were sharing the same VersionProcessed

diff --git a/Engine/Source/Runtime/Experimental/GeometryCollectionEngine/Private/GeometryCollection/GeometryCollectionComponent.cpp b/Engine/Source/Runtime/Experimental/GeometryCollectionEngine/Private/GeometryCollection/GeometryCollectionComponent.cpp
--- a/Engine/Source/Runtime/Experimental/GeometryCollectionEngine/Private/GeometryCollection/GeometryCollectionComponent.cpp
+++ b/Engine/Source/Runtime/Experimental/GeometryCollectionEngine/Private/GeometryCollection/GeometryCollectionComponent.cpp
@@ -420,6 +420,8 @@ bool FGeometryCollectionRepDynamicData::HasChanged(const FGeometryCollectionRepD
 
 bool FGeometryCollectionRepDynamicData::NetSerialize(FArchive& Ar, class UPackageMap* Map, bool& bOutSuccess)
 {
+	Ar << Version;
+
 	// we only support up to 2^16 bones for replication ( bandwidth reasons )
 	ensure(ClusterData.Num() < TNumericLimits<uint16>::Max());
 	uint16 NumClusterData = static_cast<uint16>(ClusterData.Num());
@@ -433,6 +435,7 @@ bool FGeometryCollectionRepDynamicData::NetSerialize(FArchive& Ar, class UPackag
 	for (FClusterData& Data : ClusterData)
 	{
 		Ar << Data.TransformIndex;
+		Ar << Data.bIsInternalCluster;
 		Ar << Data.Position;
 		Ar << Data.EulerRotation; // as rotator
 		Ar << Data.LinearVelocity;
@@ -2056,6 +2059,7 @@ void UGeometryCollectionComponent::ResetRepData()
 	// so at this point those data won't be changed on the Physics Thread. 
 	OneOffActivatedProcessed = 0;
 	VersionProcessed = INDEX_NONE;
+	DynamicRepDataVersionProcessed = INDEX_NONE;
 	LastHardsnapTimeInMs = 0;
 }
 
@@ -2360,7 +2364,8 @@ void UGeometryCollectionComponent::UpdateRepStateAndDynamicData()
 					const FPBDRigidClusteredParticleHandle* ParentHandle = ParticleHandle->Parent();
 
 					// no parent and not root means the particle has been broken  
-					const bool bIsBroken = (ParentHandle == nullptr && TransformIndex != RootIndex);
+					const bool bIsRootIndex = (TransformIndex == RootIndex);
+					const bool bIsBroken = (ParentHandle == nullptr && !bIsRootIndex);
 					if (bIsBroken)
 					{
 						// only record the one at the abandon level and before ( child break of abandon level clusters must be recorded )  
@@ -2387,10 +2392,10 @@ void UGeometryCollectionComponent::UpdateRepStateAndDynamicData()
 						}
 					}
 
-					const bool bIsRoot = bIsBroken && (!ParticleHandle->Disabled());
+					const bool bIsActiveBrokenOrOriginalRoot = (!ParticleHandle->Disabled()) && (bIsBroken || bIsRootIndex);
 					const bool bHasInternalClusterParent = ParentHandle && (ParentHandle->InternalCluster()) && (ParentHandle->PhysicsProxy() == ParticleHandle->PhysicsProxy());
 
-					const FPBDRigidClusteredParticleHandle* RootParticle = bIsRoot ? ParticleHandle : (bHasInternalClusterParent ? ParentHandle : nullptr);
+					const FPBDRigidClusteredParticleHandle* RootParticle = bIsActiveBrokenOrOriginalRoot ? ParticleHandle : (bHasInternalClusterParent ? ParentHandle : nullptr);
 					const int32 RootParticleLevel = (bHasInternalClusterParent)? (Level - 1): (Level);
 					const bool bTrackPosition = (RootParticleLevel <= this->ReplicationMaxPositionAndVelocityCorrectionLevel);
 					if (bTrackPosition && RootParticle)
@@ -2592,7 +2597,7 @@ namespace
 				// It's formulated this way to get a larger correction for a longer time
 				// step, ie. correction velocities are framerate independent.
 				//
-				Cluster.SetV(RepVel + (DeltaX * GeometryCollectionRepLinearMatchStrength * DeltaTime));
+				Cluster.SetV(RepVel + (DeltaX * FMath::Min(1.0f, GeometryCollectionRepLinearMatchStrength * DeltaTime)));
 			}
 
 			//
@@ -3013,7 +3018,7 @@ void UGeometryCollectionComponent::ProcessRepDynamicDataOnPT()
 
 		// enqueue a callback to process the replicated dynamic data on the physics thread 
 		CurrSolver->EnqueueCommandImmediate(
-			[PhysicsProxy = PhysicsProxy, RepDynamicDataCopy = RepDynamicData, &VersionProcessed = VersionProcessed, bReplicateMovement, &LastHardsnapTimeInMs = LastHardsnapTimeInMs]
+			[PhysicsProxy = PhysicsProxy, RepDynamicDataCopy = RepDynamicData, &VersionProcessed = DynamicRepDataVersionProcessed, bReplicateMovement, &LastHardsnapTimeInMs = LastHardsnapTimeInMs]
 			(Chaos::FReal DeltaTime, Chaos::FReal SimTime)
 			{
 				::ProcessRepDynamicDataCommon(PhysicsProxy, RepDynamicDataCopy, VersionProcessed, bReplicateMovement, SimTime, DeltaTime, LastHardsnapTimeInMs);
diff --git a/Engine/Source/Runtime/Experimental/GeometryCollectionEngine/Public/GeometryCollection/GeometryCollectionComponent.h b/Engine/Source/Runtime/Experimental/GeometryCollectionEngine/Public/GeometryCollection/GeometryCollectionComponent.h
--- a/Engine/Source/Runtime/Experimental/GeometryCollectionEngine/Public/GeometryCollection/GeometryCollectionComponent.h
+++ b/Engine/Source/Runtime/Experimental/GeometryCollectionEngine/Public/GeometryCollection/GeometryCollectionComponent.h
@@ -1592,6 +1592,7 @@ protected:
 	bool bAlreadyFullyDecayed = false;
 
 	int32 VersionProcessed = INDEX_NONE;
+	int32 DynamicRepDataVersionProcessed = INDEX_NONE;
 
 	// The last time (in milliseconds) the async physics component tick fired.
 	// We track this on the client to be able to turn off the tick for perf reasons
