Constraints: check existing dependencies between handles to avoid future tick cycles
    - follow up to 31120772

#rnx

diff --git a/Engine/Source/Runtime/Experimental/Animation/Constraints/Private/TransformConstraint.cpp b/Engine/Source/Runtime/Experimental/Animation/Constraints/Private/TransformConstraint.cpp
--- a/Engine/Source/Runtime/Experimental/Animation/Constraints/Private/TransformConstraint.cpp
+++ b/Engine/Source/Runtime/Experimental/Animation/Constraints/Private/TransformConstraint.cpp
@@ -2075,7 +2075,7 @@ void FTransformConstraintUtils::GetChildrenConstraints(
 	// and also has the same target if bIncludeTarget is true
 	const uint32 ParentHash = InHandle->GetHash();
 	const UObject* ParentTarget = InHandle->GetTarget().Get();
-	auto Predicate = [ParentHash, bIncludeTarget, ParentTarget](const ConstraintPtr& Constraint)
+	auto Predicate = [InHandle, ParentHash, bIncludeTarget, ParentTarget, World](const ConstraintPtr& Constraint)
 	{
 		const UTickableTransformConstraint* TransformConstraint = Cast<UTickableTransformConstraint>(Constraint.Get());
 		if (!TransformConstraint)
@@ -2085,16 +2085,29 @@ void FTransformConstraintUtils::GetChildrenConstraints(
 
 		if (TransformConstraint->ParentTRSHandle)
 		{
-			if (TransformConstraint->ParentTRSHandle->GetHash() == ParentHash)
+			const UTransformableHandle* OtherParentHandle = TransformConstraint->ParentTRSHandle;
+			if (OtherParentHandle->GetHash() == ParentHash)
 			{
 				return true;
 			}
 
 			if (bIncludeTarget && ParentTarget)
 			{
-				const UObject* Target = TransformConstraint->ParentTRSHandle->GetTarget().Get();
+				const UObject* Target = OtherParentHandle->GetTarget().Get();
 				if (Target == ParentTarget)
 				{
+					// check direct dependencies to avoid evaluation order issues
+					if (InHandle->HasDirectDependencyWith(*OtherParentHandle))
+					{
+						return false;
+					}
+
+					// check constraints dependencies to avoid cycles
+					if (HasConstraintDependencyWith(World, OtherParentHandle, InHandle))
+					{
+						return false;
+					}
+					
 					return true;
 				}
 			}
