Fixed a null pointer crash if a static mesh had multiple physics proxy mesh entries and one of them was null

#rnx

diff --git a/Engine/Source/Runtime/Experimental/GeometryCollectionEngine/Private/GeometryCollection/GeometryCollectionComponent.cpp b/Engine/Source/Runtime/Experimental/GeometryCollectionEngine/Private/GeometryCollection/GeometryCollectionComponent.cpp
--- a/Engine/Source/Runtime/Experimental/GeometryCollectionEngine/Private/GeometryCollection/GeometryCollectionComponent.cpp
+++ b/Engine/Source/Runtime/Experimental/GeometryCollectionEngine/Private/GeometryCollection/GeometryCollectionComponent.cpp
@@ -3889,25 +3889,27 @@ void UGeometryCollectionComponent::RegisterAndInitializePhysicsProxy()
 	auto CreateTraceCollisionFromStaticMeshGeometry = [this](const FTransform& InToLocal, TArray<Chaos::FImplicitObjectPtr>& OutGeoms, Chaos::FShapesArray& OutShapes) {
 		for (const TObjectPtr<UStaticMesh>& ProxyMesh : RestCollection->RootProxyData.ProxyMeshes)
 		{
-			// If these were null we wouldn't have set this callback in the first place
-			const UStaticMesh* StaticMesh = ProxyMesh.Get();
-			UBodySetup* BodySetup = ProxyMesh->GetBodySetup();
-
-			FBodyCollisionData BodyCollisionData;
-			BodyInstance.BuildBodyFilterData(BodyCollisionData.CollisionFilterData);
-			FBodyInstance::BuildBodyCollisionFlags(BodyCollisionData.CollisionFlags, BodyInstance.GetCollisionEnabled(), BodySetup->GetCollisionTraceFlag() == CTF_UseComplexAsSimple);
-
-			FGeometryAddParams AddParams;
-			AddParams.bDoubleSided = BodySetup->bDoubleSidedGeometry;
-			AddParams.CollisionData = BodyCollisionData;
-			AddParams.CollisionTraceType = BodySetup->GetCollisionTraceFlag();
-			AddParams.WorldTransform = PhysicsProxy->GetWorldTransform_External();
-			AddParams.Scale = AddParams.WorldTransform.GetScale3D();
-			AddParams.SimpleMaterial = GetPhysicalMaterial();
-			AddParams.LocalTransform = InToLocal;
-			
-			AddParams.Geometry = &BodySetup->AggGeom;
-			ChaosInterface::CreateGeometry(AddParams, OutGeoms, OutShapes);
+			if (const UStaticMesh* StaticMesh = ProxyMesh.Get())
+			{
+				if (UBodySetup* BodySetup = ProxyMesh->GetBodySetup())
+				{
+					FBodyCollisionData BodyCollisionData;
+					BodyInstance.BuildBodyFilterData(BodyCollisionData.CollisionFilterData);
+					FBodyInstance::BuildBodyCollisionFlags(BodyCollisionData.CollisionFlags, BodyInstance.GetCollisionEnabled(), BodySetup->GetCollisionTraceFlag() == CTF_UseComplexAsSimple);
+
+					FGeometryAddParams AddParams;
+					AddParams.bDoubleSided = BodySetup->bDoubleSidedGeometry;
+					AddParams.CollisionData = BodyCollisionData;
+					AddParams.CollisionTraceType = BodySetup->GetCollisionTraceFlag();
+					AddParams.WorldTransform = PhysicsProxy->GetWorldTransform_External();
+					AddParams.Scale = AddParams.WorldTransform.GetScale3D();
+					AddParams.SimpleMaterial = GetPhysicalMaterial();
+					AddParams.LocalTransform = InToLocal;
+
+					AddParams.Geometry = &BodySetup->AggGeom;
+					ChaosInterface::CreateGeometry(AddParams, OutGeoms, OutShapes);
+				}
+			}
 		}
 	};
 
