Iris - Don't check on expected behavior. Arrays can have more than a single bit changemask.

#rnx

diff --git a/Engine/Source/Runtime/Experimental/Iris/Core/Private/Iris/ReplicationState/PropertyReplicationState.cpp b/Engine/Source/Runtime/Experimental/Iris/Core/Private/Iris/ReplicationState/PropertyReplicationState.cpp
--- a/Engine/Source/Runtime/Experimental/Iris/Core/Private/Iris/ReplicationState/PropertyReplicationState.cpp
+++ b/Engine/Source/Runtime/Experimental/Iris/Core/Private/Iris/ReplicationState/PropertyReplicationState.cpp
@@ -548,11 +548,14 @@ bool FPropertyReplicationState::IsCustomConditionEnabled(uint32 Index) const
 	const FReplicationStateDescriptor* Descriptor = ReplicationStateDescriptor;
 	const FReplicationStateMemberChangeMaskDescriptor& ChangeMaskInfo = Descriptor->MemberChangeMaskDescriptors[Index];
 
-	// $TODO: Test any bits if check triggers?
-	checkSlow(ChangeMaskInfo.BitCount == 1);
+	if (ChangeMaskInfo.BitCount > 0)
+	{
+		FNetBitArrayView MemberConditionalChangeMask = Private::GetMemberConditionalChangeMask(StateBuffer, Descriptor);
+		return MemberConditionalChangeMask.GetBit(ChangeMaskInfo.BitOffset);
+	}
 
-	FNetBitArrayView MemberConditionalChangeMask = Private::GetMemberConditionalChangeMask(StateBuffer, Descriptor);
-	return MemberConditionalChangeMask.GetBit(ChangeMaskInfo.BitOffset);
+	// If there's no bitmask the property cannot be disabled.
+	return true;
 }
 
 }
