Improve edges generation performances between Geometry Collection
This will allow to enable memory optimization on the Physics Thread

diff --git a/Engine/Source/Runtime/Experimental/Chaos/Private/Chaos/ClusterUnionManager.cpp b/Engine/Source/Runtime/Experimental/Chaos/Private/Chaos/ClusterUnionManager.cpp
--- a/Engine/Source/Runtime/Experimental/Chaos/Private/Chaos/ClusterUnionManager.cpp
+++ b/Engine/Source/Runtime/Experimental/Chaos/Private/Chaos/ClusterUnionManager.cpp
@@ -156,29 +156,38 @@ namespace Chaos
 						{
 							FPBDRigidParticleHandle* OtherParticle = ClusterUnion.ChildParticles[RootObjectIndex];
 							Clustering.CreateNodeConnection(Particle, OtherParticle);
-
-							if (bChaosClusterUnionGenerateInterclusterEdges && Particle && OtherParticle)
+							
+							if (FClusterUnionParticleProperties* Properties = ClusterUnion.ChildProperties.Find(Particle))
 							{
-								// Only generate intercluster edges for main particles. Auxiliary particles that are just bits and pieces of geometry collections
-								// shouldn't also generate intercluster edges.
-								if (FClusterUnionParticleProperties* Properties = ClusterUnion.ChildProperties.Find(Particle); !Properties || !Properties->bIsAuxiliaryParticle)
+								Properties->bEdgesAreGenerated = false;
+								if (FClusterUnionParticleProperties* OtherProperties = ClusterUnion.ChildProperties.Find(OtherParticle))
 								{
-									const TArray<FPBDRigidParticleHandle*>& ParticleChildren = Clustering.GetChildrenMap().FindRef(Particle->CastToClustered());
-									const TArray<FPBDRigidParticleHandle*>& OtherChildren = Clustering.GetChildrenMap().FindRef(OtherParticle->CastToClustered());
+									OtherProperties->bEdgesAreGenerated = false;
+								}
+
+								if (bChaosClusterUnionGenerateInterclusterEdges && Particle && OtherParticle)
+								{
+									// Only generate intercluster edges for main particles. Auxiliary particles that are just bits and pieces of geometry collections
+									// shouldn't also generate intercluster edges.
+									if (!Properties || !Properties->bIsAuxiliaryParticle)
+									{
+										const TArray<FPBDRigidParticleHandle*>& ParticleChildren = Clustering.GetChildrenMap().FindRef(Particle->CastToClustered());
+										const TArray<FPBDRigidParticleHandle*>& OtherChildren = Clustering.GetChildrenMap().FindRef(OtherParticle->CastToClustered());
 
-									TSet<FPBDRigidParticleHandle*> FromSet{ ParticleChildren };
-									TSet<FPBDRigidParticleHandle*> ToSet{ OtherChildren };
+										TSet<FPBDRigidParticleHandle*> FromSet{ ParticleChildren };
+										TSet<FPBDRigidParticleHandle*> ToSet{ OtherChildren };
 
-									TArray<FPBDRigidParticleHandle*> AllParticles;
-									AllParticles.Reserve(FromSet.Num() + ToSet.Num());
-									AllParticles.Append(ParticleChildren);
-									AllParticles.Append(OtherChildren);
+										TArray<FPBDRigidParticleHandle*> AllParticles;
+										AllParticles.Reserve(FromSet.Num() + ToSet.Num());
+										AllParticles.Append(ParticleChildren);
+										AllParticles.Append(OtherChildren);
 
-									FClusterCreationParameters Parameters{ 0.3f, 100, false, false };
-									Parameters.ConnectionMethod = FClusterCreationParameters::EConnectionMethod::BoundsOverlapFilteredDelaunayTriangulation;
-									Parameters.ConnectionGraphBoundsFilteringMargin = 1.0;
+										FClusterCreationParameters Parameters{ 0.3f, 100, false, false };
+										Parameters.ConnectionMethod = FClusterCreationParameters::EConnectionMethod::BoundsOverlapFilteredDelaunayTriangulation;
+										Parameters.ConnectionGraphBoundsFilteringMargin = 1.0;
 
-									Clustering.GenerateConnectionGraph(AllParticles, Parameters, &FromSet, &ToSet);
+										Clustering.GenerateConnectionGraph(AllParticles, Parameters, &FromSet, &ToSet);
+									}
 								}
 							}
 						}
@@ -1023,6 +1032,11 @@ namespace Chaos
 		MClustering.BuildConvexOptimizer(ClusterUnion.InternalCluster);
 	}
 
+	void FClusterUnionManager::AddParticleToConnectionGraphInCluster(FClusterUnion& ClusterUnion, FPBDRigidParticleHandle* Particle)
+	{
+		AddParticleToConnectionGraph(MClustering, ClusterUnion, Particle);
+	}
+
 	DECLARE_CYCLE_STAT(TEXT("FClusterUnionManager::FlushIncrementalConnectivityGraphOperations"), STAT_FlushIncrementalConnectivityGraphOperations, STATGROUP_Chaos);
 	void FClusterUnionManager::FlushIncrementalConnectivityGraphOperations(FClusterUnion& ClusterUnion)
 	{
diff --git a/Engine/Source/Runtime/Experimental/Chaos/Private/Chaos/PBDRigidClustering.cpp b/Engine/Source/Runtime/Experimental/Chaos/Private/Chaos/PBDRigidClustering.cpp
--- a/Engine/Source/Runtime/Experimental/Chaos/Private/Chaos/PBDRigidClustering.cpp
+++ b/Engine/Source/Runtime/Experimental/Chaos/Private/Chaos/PBDRigidClustering.cpp
@@ -1049,8 +1049,14 @@ namespace Chaos
 		ConcreteGCProxy.CreateChildrenGeometry_Internal();
 		if (Chaos::FClusterUnion* ClusterUnion = ClusterUnionManager.FindClusterUnionFromParticle(&ClusteredParticle))
 		{
-			ClusterUnion->PendingConnectivityOperations.Add({ &ClusteredParticle, Chaos::EClusterUnionConnectivityOperation::Add });
+			bool bHasBuiltAllEdges = false;
+			FClusterUnionParticleProperties* Properties = ClusterUnion->ChildProperties.Find(&ClusteredParticle);
+			if (Properties)
+			{
+				bHasBuiltAllEdges = Properties->bEdgesAreGenerated;
+			}
 
+			bool bAllNeighborsHasBuiltEdges = true;
 			const TArray<Chaos::TConnectivityEdge<Chaos::FReal>> Edges = ClusteredParticle.ConnectivityEdges();
 			for (const Chaos::TConnectivityEdge<Chaos::FReal>& Edge : Edges)
 			{
@@ -1061,11 +1067,22 @@ namespace Chaos
 					{
 						FGeometryCollectionPhysicsProxy* GCProxy = GetConcreteProxy<FGeometryCollectionPhysicsProxy>(Sibling);
 						GCProxy->CreateChildrenGeometry_Internal();
+						if (FClusterUnionParticleProperties* SiblingProperties = ClusterUnion->ChildProperties.Find(Sibling))
+						{
+							bAllNeighborsHasBuiltEdges &= SiblingProperties->bEdgesAreGenerated;
+						}
 					}
 				}
 			}
-			ClusterUnionManager.RequestDeferredClusterPropertiesUpdate(ClusterUnion->InternalIndex, Chaos::EUpdateClusterUnionPropertiesFlags::IncrementalGenerateConnectionGraph);
-			ClusterUnionManager.HandleDeferredClusterUnionUpdateProperties();
+			// If has current GC has built all edges or if all neighbors have build all edges don't need to compute neighbors edges. 
+			if (!(bHasBuiltAllEdges || bAllNeighborsHasBuiltEdges))
+			{
+				ClusterUnionManager.AddParticleToConnectionGraphInCluster(*ClusterUnion, &ClusteredParticle);
+				if (Properties)
+				{
+					Properties->bEdgesAreGenerated = true;
+				}
+			}
 		}
 	}
 
diff --git a/Engine/Source/Runtime/Experimental/Chaos/Public/Chaos/ClusterUnionManager.h b/Engine/Source/Runtime/Experimental/Chaos/Public/Chaos/ClusterUnionManager.h
--- a/Engine/Source/Runtime/Experimental/Chaos/Public/Chaos/ClusterUnionManager.h
+++ b/Engine/Source/Runtime/Experimental/Chaos/Public/Chaos/ClusterUnionManager.h
@@ -72,8 +72,13 @@ namespace Chaos
 
 	struct FClusterUnionParticleProperties
 	{
+		FClusterUnionParticleProperties()
+			: bIsAuxiliaryParticle(0)
+			, bEdgesAreGenerated(0)
+		{}
 		// An auxiliary particle will be removed from the cluster union if FRigidClustering::HandleConnectivityOnReleaseClusterParticle detects an island made up of only auxiliary particles.
-		bool bIsAuxiliaryParticle = false;
+		uint8 bIsAuxiliaryParticle: 1;
+		uint8 bEdgesAreGenerated: 1;
 	};
 
 	struct FClusterUnion
@@ -208,6 +213,8 @@ namespace Chaos
 		// Update the cluster union's properties after its set of particle changes.
 		CHAOS_API void UpdateAllClusterUnionProperties(FClusterUnion& ClusterUnion, EUpdateClusterUnionPropertiesFlags Flags = EUpdateClusterUnionPropertiesFlags::All);
 
+		CHAOS_API  void AddParticleToConnectionGraphInCluster(FClusterUnion& ClusterUnion, FPBDRigidParticleHandle* Particle);
+
 		// Returns all cluster unions. Really meant only to be used for debugging.
 		const TMap<FClusterUnionIndex, FClusterUnion>& GetAllClusterUnions() const { return ClusterUnions; }
 
