Iris
- Added explicit configuration of splitsizes for unreliable attachments from server and client
- Fix issue with old moves being delivered from client if split and packet dropped.

diff --git a/Engine/Source/Runtime/Experimental/Iris/Core/Private/Iris/ReplicationSystem/NetBlob/PartialNetObjectAttachmentHandler.cpp b/Engine/Source/Runtime/Experimental/Iris/Core/Private/Iris/ReplicationSystem/NetBlob/PartialNetObjectAttachmentHandler.cpp
--- a/Engine/Source/Runtime/Experimental/Iris/Core/Private/Iris/ReplicationSystem/NetBlob/PartialNetObjectAttachmentHandler.cpp
+++ b/Engine/Source/Runtime/Experimental/Iris/Core/Private/Iris/ReplicationSystem/NetBlob/PartialNetObjectAttachmentHandler.cpp
@@ -3,6 +3,7 @@
 #include "Iris/ReplicationSystem/NetBlob/PartialNetObjectAttachmentHandler.h"
 #include "Iris/ReplicationSystem/NetBlob/PartialNetBlob.h"
 #include "Iris/ReplicationSystem/NetBlob/ShrinkWrapNetBlob.h"
+#include "Iris/ReplicationSystem/ReplicationSystem.h"
 #include "Iris/Serialization/NetSerializationContext.h"
 #include "Iris/Serialization/InternalNetSerializationContext.h"
 #include "Iris/Serialization/NetBitStreamWriter.h"
@@ -29,8 +30,15 @@ void UPartialNetObjectAttachmentHandler::Init(const FPartialNetObjectAttachmentH
 bool UPartialNetObjectAttachmentHandler::PreSerializeAndSplitNetBlob(uint32 ConnectionId, const TRefCountPtr<UE::Net::FNetObjectAttachment>& Blob, TArray<TRefCountPtr<FNetBlob>>& OutPartialBlobs, bool bSerializeWithObject) const
 {
 	using namespace UE::Net;
+	
+	uint32 BitCountSplitThreshold =  GetConfig()->GetBitCountSplitThreshold() & ~31U;	
 
-	const uint32 BitCountSplitThreshold =  GetConfig()->GetBitCountSplitThreshold() & ~31U;
+	const FNetBlobCreationInfo& BlobCreationInfo = Blob->GetCreationInfo();
+	const bool bIsReliable = EnumHasAnyFlags(BlobCreationInfo.Flags, ENetBlobFlags::Reliable);
+	if (!bIsReliable)
+	{
+		BitCountSplitThreshold = (ReplicationSystem->IsServer() ? GetConfig()->GetServerUnreliableBitCountSplitThreshold() : GetConfig()->GetClientUnreliableBitCountSplitThreshold()) & ~31U;
+	}
 
 	TArray<uint32> Payload;
 	Payload.AddUninitialized(BitCountSplitThreshold/32U);
diff --git a/Engine/Source/Runtime/Experimental/Iris/Core/Private/Iris/ReplicationSystem/NetBlob/PartialNetObjectAttachmentHandler.h b/Engine/Source/Runtime/Experimental/Iris/Core/Private/Iris/ReplicationSystem/NetBlob/PartialNetObjectAttachmentHandler.h
--- a/Engine/Source/Runtime/Experimental/Iris/Core/Private/Iris/ReplicationSystem/NetBlob/PartialNetObjectAttachmentHandler.h
+++ b/Engine/Source/Runtime/Experimental/Iris/Core/Private/Iris/ReplicationSystem/NetBlob/PartialNetObjectAttachmentHandler.h
@@ -12,11 +12,21 @@ class UPartialNetObjectAttachmentHandlerConfig : public USequentialPartialNetBlo
 
 public:
 	uint32 GetBitCountSplitThreshold() const { return BitCountSplitThreshold; }
+	uint32 GetClientUnreliableBitCountSplitThreshold() const { return ClientUnreliableBitCountSplitThreshold; }
+	uint32 GetServerUnreliableBitCountSplitThreshold() const { return ServerUnreliableBitCountSplitThreshold; }
 
 private:
 	/** How many bits a payload should have to recommend a split. Should be higher than MaxPartBitCount as splitting adds overhead. */
 	UPROPERTY(Config)
 	uint32 BitCountSplitThreshold = (128 + 64)*8;
+
+	/** How many bits a unreliable payload should have to recommend a split on the client. Should be higher than MaxPartBitCount as splitting adds overhead. */
+	UPROPERTY(Config)
+	uint32 ClientUnreliableBitCountSplitThreshold = (850)*8;
+
+	/** How many bits a unreliable payload should have to recommend a split on the server. Should be higher than MaxPartBitCount as splitting adds overhead. */
+	UPROPERTY(Config)
+	uint32 ServerUnreliableBitCountSplitThreshold = (256)*8;
 };
 
 struct FPartialNetObjectAttachmentHandlerInitParams
