CI fix: Prevent reading potentially unallocated memory (temporary work around)
#[fyi] Benn.Gallagher

diff --git a/Engine/Source/Runtime/Experimental/Chaos/Public/Chaos/ParticleHandle.h b/Engine/Source/Runtime/Experimental/Chaos/Public/Chaos/ParticleHandle.h
--- a/Engine/Source/Runtime/Experimental/Chaos/Public/Chaos/ParticleHandle.h
+++ b/Engine/Source/Runtime/Experimental/Chaos/Public/Chaos/ParticleHandle.h
@@ -2582,7 +2582,8 @@ public:
 	virtual ~TGeometryParticle() //only virtual for easier memory management. Should generally be a static API
 	{
 		SetUniqueIdx(FUniqueIdx{}, false); // Set to an invalid index for dangling handle detection
-		FChaosUserEntityAppend* UserEntityAppend = FChaosUserData::Get<FChaosUserEntityAppend>(MUserData);
+		// TODO: Temporary CI fix: Prevent reading Data on bodyInstances after they may be deallocated, delete custom entities elsewhere
+		FChaosUserEntityAppend* UserEntityAppend = nullptr; // FChaosUserData::Get<FChaosUserEntityAppend>(MUserData); 
 		if (UserEntityAppend)
 		{
 			delete UserEntityAppend;
