Put shadow copy of custom instance data behind a cvar, most likely it will only be needed for debugging purpose.

diff --git a/Engine/Source/Runtime/Experimental/GeometryCollectionEngine/Private/GeometryCollection/GeometryCollectionISMPoolComponent.cpp b/Engine/Source/Runtime/Experimental/GeometryCollectionEngine/Private/GeometryCollection/GeometryCollectionISMPoolComponent.cpp
--- a/Engine/Source/Runtime/Experimental/GeometryCollectionEngine/Private/GeometryCollection/GeometryCollectionISMPoolComponent.cpp
+++ b/Engine/Source/Runtime/Experimental/GeometryCollectionEngine/Private/GeometryCollection/GeometryCollectionISMPoolComponent.cpp
@@ -37,6 +37,12 @@ FAutoConsoleVariableRef CVarISMPoolComponentFreeListTargetSize(
 	GComponentFreeListTargetSize,
 	TEXT("Target size for number of ISM components in the recycling free list."));
 
+static bool GShadowCopyCustomData = false;
+FAutoConsoleVariableRef CVarShadowCopyCustomData(
+	TEXT("r.ISMPool.ShadowCopyCustomData"),
+	GShadowCopyCustomData,
+	TEXT("Keeps a copy of custom instance data so it can be restored if the instance is removed and readded."));
+
 
 void FGeometryCollectionMeshInfo::ShadowCopyCustomData(int32 InstanceCount, int32 NumCustomDataFloatsPerInstance, TArrayView<const float> CustomDataFloats)
 {
@@ -59,7 +65,7 @@ FGeometryCollectionMeshGroup::FMeshId FGeometryCollectionMeshGroup::AddMesh(cons
 {
 	const FMeshId MeshInfoIndex = MeshInfos.Emplace(ISMInstanceInfo);
 
-	if (bAllowPerInstanceRemoval)
+	if (bAllowPerInstanceRemoval && GShadowCopyCustomData)
 	{
 		FGeometryCollectionMeshInfo& MeshInfo = MeshInfos[MeshInfoIndex];
 		MeshInfo.ShadowCopyCustomData(InstanceCount, MeshInstance.Desc.NumCustomDataFloats, CustomDataFloats);
@@ -324,7 +330,11 @@ bool FGeometryCollectionISMPool::BatchUpdateInstancesTransforms(FGeometryCollect
 				// Re-add the instance to the ISM if the scale becomes non-zero.
 				FPrimitiveInstanceId Id = ISM.ISMComponent->AddInstanceById(Transform, bWorldSpace);
 				ISM.InstanceIds[InstanceGroup.Start + InstanceIndex] = Id;
-				ISM.ISMComponent->SetCustomDataById(Id, MeshInfo.CustomDataSlice(InstanceIndex, ISM.ISMComponent->NumCustomDataFloats));
+
+				if (MeshInfo.CustomData.Num())
+				{
+					ISM.ISMComponent->SetCustomDataById(Id, MeshInfo.CustomDataSlice(InstanceIndex, ISM.ISMComponent->NumCustomDataFloats));
+				}
 				continue;
 			}
 		}
