Make particle indirection access more robust.

diff --git a/Engine/Source/Runtime/Experimental/Chaos/Private/PhysicsProxy/GeometryCollectionPhysicsProxy.cpp b/Engine/Source/Runtime/Experimental/Chaos/Private/PhysicsProxy/GeometryCollectionPhysicsProxy.cpp
--- a/Engine/Source/Runtime/Experimental/Chaos/Private/PhysicsProxy/GeometryCollectionPhysicsProxy.cpp
+++ b/Engine/Source/Runtime/Experimental/Chaos/Private/PhysicsProxy/GeometryCollectionPhysicsProxy.cpp
@@ -892,7 +892,7 @@ void FGeometryCollectionPhysicsProxy::CreateGTParticles(TManagedArray<Chaos::FIm
 							Chaos::FImplicitObject* UnionImplicit = new Chaos::FImplicitObjectUnion(MoveTemp(ChildImplicits));
 							Implicits[ParentToFixIndex] = Chaos::FImplicitObjectPtr(UnionImplicit);
 						}
-						if (GTParticles[FromTransformToParticleIndex[ParentToFixIndex]] != nullptr)
+						if (FromTransformToParticleIndex[ParentToFixIndex] != INDEX_NONE && GTParticles[FromTransformToParticleIndex[ParentToFixIndex]] != nullptr)
 						{
 							GTParticles[FromTransformToParticleIndex[ParentToFixIndex]]->SetGeometry(Implicits[ParentToFixIndex]);
 						}
@@ -3777,7 +3777,7 @@ void FGeometryCollectionPhysicsProxy::BufferPhysicsResults_Internal(Chaos::FPBDR
 		// will never hit the logic below that checks for the particle changing from active to inactive.
 		// The parent checks also miss the case where the parent has an external parent (GC particles
 		// historically assume that the root has no parent which is not true any more).
-		if (Parameters.InitialRootIndex != INDEX_NONE)
+		if (Parameters.InitialRootIndex != INDEX_NONE && FromTransformToParticleIndex[Parameters.InitialRootIndex] != INDEX_NONE)
 		{
 			const Chaos::FPBDRigidClusteredParticleHandle* RootHandle = SolverParticleHandles[FromTransformToParticleIndex[Parameters.InitialRootIndex]];
 			if (RootHandle != nullptr)
@@ -4159,7 +4159,7 @@ bool FGeometryCollectionPhysicsProxy::PullNonInterpolatableDataFromSinglePhysics
 
 		if (UpdateValue(GameThreadCollection.Active[TransformGroupIndex], bIsActive))
 		{
-			if (GTParticles[ParticleIndex].IsValid())
+			if (ParticleIndex != INDEX_NONE && GTParticles[ParticleIndex].IsValid())
 			{
 				GTParticles[ParticleIndex]->SetDisabled(!bIsActive);
 				bIsCollectionDirty = true;
@@ -4413,9 +4413,9 @@ bool FGeometryCollectionPhysicsProxy::PullFromPhysicsState(const Chaos::FDirtyGe
 
 							// by default parent is the component's world transform 
 							FTransform ParentWorldTransform = WorldTransform_External;
-							if (ParentTransformIndex != INDEX_NONE)
+							if (ParentTransformIndex != INDEX_NONE && FromTransformToParticleIndex[ParentTransformIndex] != INDEX_NONE)
 							{
-								if (const FParticle* GTParentParticle = (ParentTransformIndex != INDEX_NONE) ? GTParticles[FromTransformToParticleIndex[ParentTransformIndex]].Get() : nullptr)
+								if (const FParticle* GTParentParticle = GTParticles[FromTransformToParticleIndex[ParentTransformIndex]].Get())
 								{
 									const FTransform& ParentMassToLocal = MassToLocal[TransformIndex];
 									ParentWorldTransform = ParentMassToLocal.Inverse() * FTransform { GTParentParticle->R(), GTParentParticle->X() };
@@ -5815,9 +5815,10 @@ void FGeometryCollectionPhysicsProxy::FieldForcesUpdateCallback(Chaos::FPBDRigid
 
 FGeometryCollectionPhysicsProxy::FClusterHandle* FGeometryCollectionPhysicsProxy::GetInitialRootHandle_Internal() const
 {
-	if (SolverParticleHandles.IsValidIndex(FromTransformToParticleIndex[Parameters.InitialRootIndex]))
+	const int32 Index = FromTransformToParticleIndex[Parameters.InitialRootIndex];
+	if (SolverParticleHandles.IsValidIndex(Index))
 	{
-		return SolverParticleHandles[FromTransformToParticleIndex[Parameters.InitialRootIndex]];
+		return SolverParticleHandles[Index];
 	}
 	return nullptr;
 }
