Iris
* PlayerController view target will always be set to be always relevant

#rnx

diff --git a/Engine/Source/Runtime/Experimental/Iris/Core/Public/Iris/ReplicationSystem/NetObjectGroupHandle.h b/Engine/Source/Runtime/Experimental/Iris/Core/Public/Iris/ReplicationSystem/NetObjectGroupHandle.h
--- a/Engine/Source/Runtime/Experimental/Iris/Core/Public/Iris/ReplicationSystem/NetObjectGroupHandle.h
+++ b/Engine/Source/Runtime/Experimental/Iris/Core/Public/Iris/ReplicationSystem/NetObjectGroupHandle.h
@@ -32,6 +32,8 @@ public:
 		NetGroupReplayNetObjectGroupIndex
 	};
 
+	inline static FNetObjectGroupHandle GetInvalid() {return FNetObjectGroupHandle();}
+
 	FNetObjectGroupHandle() : Value(0u) {}
 
 	/** Returns true if the handle is valid, note this does not mean that the group is valid */
diff --git a/Engine/Source/Runtime/Experimental/Iris/Core/Public/Iris/ReplicationSystem/ReplicationSystem.h b/Engine/Source/Runtime/Experimental/Iris/Core/Public/Iris/ReplicationSystem/ReplicationSystem.h
--- a/Engine/Source/Runtime/Experimental/Iris/Core/Public/Iris/ReplicationSystem/ReplicationSystem.h
+++ b/Engine/Source/Runtime/Experimental/Iris/Core/Public/Iris/ReplicationSystem/ReplicationSystem.h
@@ -302,6 +302,7 @@ public:
 	/**
 	 * Create a group which can be used to logically group objects together. The group must be
 	 * destroyed when it's not needed anymore.
+	 * Groups can be used to setup filtering rules on it's members.
 	 * @return A handle to the group, or InvalidNetObjectGroupHandle if no more groups could be created.
 	 * @see DestroyGroup
 	 */
@@ -439,7 +440,7 @@ public:
 	 */
 	IRISCORE_API bool AddInclusionFilterGroup(FNetObjectGroupHandle GroupHandle);
 
-	/** Remove group from filtering system, cancelling all effects of the group. */
+	/** Remove group from filtering system, canceling all effects of the group. */
 	IRISCORE_API void RemoveGroupFilter(FNetObjectGroupHandle GroupHandle);
 
 	/** Set status of GroupFilter for specific connection. */
@@ -448,7 +449,7 @@ public:
 	/** Set status of GroupFilter for connection marked in the Connections BitArray. */
 	IRISCORE_API void SetGroupFilterStatus(FNetObjectGroupHandle GroupHandle, const UE::Net::FNetBitArray& Connections, UE::Net::ENetFilterStatus ReplicationStatus);
 
-	/** Set status of GroupFilter for all connnections. */
+	/** Set status of GroupFilter for all connections. */
 	IRISCORE_API void SetGroupFilterStatus(FNetObjectGroupHandle GroupHandle, UE::Net::ENetFilterStatus ReplicationStatus);
 
 
@@ -472,6 +473,7 @@ public:
 	/**
 	 * Set which connections the object is allowed to be replicated to. This will cancel the effect
 	 * of any previous SetFilter or SetConnectionFilter calls on the object.
+	 * It will also remove the object from any dynamic filters
 	 * @param Handle A valid handle to a replicated object.
 	 * @param Connections Set bits indicates the connection IDs that the object is allowed or not allowed to be replicated to depending on ReplicationStatus.
 	 * @param ReplicationStatus Whether the set bits in Connections indicate if the object is allowed or not allowed to be replicated to those connections.
