Chaos - Don't interpolate network correction if error is greater than a limit

diff --git a/Engine/Source/Runtime/Experimental/Chaos/Private/PhysicsProxy/ClusterUnionPhysicsProxy.cpp b/Engine/Source/Runtime/Experimental/Chaos/Private/PhysicsProxy/ClusterUnionPhysicsProxy.cpp
--- a/Engine/Source/Runtime/Experimental/Chaos/Private/PhysicsProxy/ClusterUnionPhysicsProxy.cpp
+++ b/Engine/Source/Runtime/Experimental/Chaos/Private/PhysicsProxy/ClusterUnionPhysicsProxy.cpp
@@ -430,7 +430,13 @@ namespace Chaos
 
 		if (Error)
 		{
-			const int32 RenderInterpErrorCorrectionDurationTicks = FMath::FloorToInt32(GetRenderInterpErrorCorrectionDuration() / AsyncFixedTimeStep); // Convert duration from seconds to simulation ticks
+			const FReal ErrorMagSq = Error->ErrorX.SizeSquared();
+			const FReal MaxErrorCorrection = GetRenderInterpMaximumErrorCorrectionBeforeSnapping();
+			int32 RenderInterpErrorCorrectionDurationTicks = 0;
+			if (ErrorMagSq < MaxErrorCorrection * MaxErrorCorrection)
+			{
+				RenderInterpErrorCorrectionDurationTicks = FMath::FloorToInt32(GetRenderInterpErrorCorrectionDuration() / AsyncFixedTimeStep); // Convert duration from seconds to simulation ticks
+			}
 			InterpolationData.AccumlateErrorXR(Error->ErrorX, Error->ErrorR, SolverSyncTimestamp, RenderInterpErrorCorrectionDurationTicks);
 		}
 
diff --git a/Engine/Source/Runtime/Experimental/Chaos/Private/PhysicsProxy/GeometryCollectionPhysicsProxy.cpp b/Engine/Source/Runtime/Experimental/Chaos/Private/PhysicsProxy/GeometryCollectionPhysicsProxy.cpp
--- a/Engine/Source/Runtime/Experimental/Chaos/Private/PhysicsProxy/GeometryCollectionPhysicsProxy.cpp
+++ b/Engine/Source/Runtime/Experimental/Chaos/Private/PhysicsProxy/GeometryCollectionPhysicsProxy.cpp
@@ -4288,7 +4288,13 @@ bool FGeometryCollectionPhysicsProxy::PullFromPhysicsState(const Chaos::FDirtyGe
 
 	if (Error)
 	{
-		const int32 RenderInterpErrorCorrectionDurationTicks = FMath::FloorToInt32(GetRenderInterpErrorCorrectionDuration() / AsyncFixedTimeStep); // Convert duration from seconds to simulation ticks
+		const Chaos::FReal ErrorMagSq = Error->ErrorX.SizeSquared();
+		const Chaos::FReal MaxErrorCorrection = GetRenderInterpMaximumErrorCorrectionBeforeSnapping();
+		int32 RenderInterpErrorCorrectionDurationTicks = 0;
+		if (ErrorMagSq < MaxErrorCorrection * MaxErrorCorrection)
+		{
+			RenderInterpErrorCorrectionDurationTicks = FMath::FloorToInt32(GetRenderInterpErrorCorrectionDuration() / AsyncFixedTimeStep); // Convert duration from seconds to simulation ticks
+		}
 		InterpolationData.AccumlateErrorXR(Error->ErrorX, Error->ErrorR, SolverSyncTimestamp, RenderInterpErrorCorrectionDurationTicks);
 	}
 
diff --git a/Engine/Source/Runtime/Experimental/Chaos/Private/PhysicsProxy/SingleParticlePhysicsProxy.cpp b/Engine/Source/Runtime/Experimental/Chaos/Private/PhysicsProxy/SingleParticlePhysicsProxy.cpp
--- a/Engine/Source/Runtime/Experimental/Chaos/Private/PhysicsProxy/SingleParticlePhysicsProxy.cpp
+++ b/Engine/Source/Runtime/Experimental/Chaos/Private/PhysicsProxy/SingleParticlePhysicsProxy.cpp
@@ -291,7 +291,13 @@ bool FSingleParticlePhysicsProxy::PullFromPhysicsState(const Chaos::FDirtyRigidP
 
 		if (Error)
 		{
-			const int32 RenderInterpErrorCorrectionDurationTicks = FMath::FloorToInt32(GetRenderInterpErrorCorrectionDuration() / AsyncFixedTimeStep); // Convert duration from seconds to simulation ticks
+			const FReal ErrorMagSq = Error->ErrorX.SizeSquared();
+			const FReal MaxErrorCorrection = GetRenderInterpMaximumErrorCorrectionBeforeSnapping();
+			int32 RenderInterpErrorCorrectionDurationTicks = 0;
+			if (ErrorMagSq < MaxErrorCorrection * MaxErrorCorrection)
+			{
+				RenderInterpErrorCorrectionDurationTicks = FMath::FloorToInt32(GetRenderInterpErrorCorrectionDuration() / AsyncFixedTimeStep); // Convert duration from seconds to simulation ticks
+			}
 			InterpolationData.AccumlateErrorXR(Error->ErrorX, Error->ErrorR, SolverSyncTimestamp, RenderInterpErrorCorrectionDurationTicks);
 #if RENDERINTERP_ERRORVELOCITYSMOOTHING
 			InterpolationData.SetVelocitySmoothing(Rigid->V(), Rigid->X(), RenderInterpErrorVelocitySmoothingDurationTicks);
diff --git a/Engine/Source/Runtime/Experimental/Chaos/Public/Chaos/Framework/PhysicsProxyBase.h b/Engine/Source/Runtime/Experimental/Chaos/Public/Chaos/Framework/PhysicsProxyBase.h
--- a/Engine/Source/Runtime/Experimental/Chaos/Public/Chaos/Framework/PhysicsProxyBase.h
+++ b/Engine/Source/Runtime/Experimental/Chaos/Public/Chaos/Framework/PhysicsProxyBase.h
@@ -129,6 +129,12 @@ public:
 		static FAutoConsoleVariableRef CVarRenderInterpErrorCorrectionDuration(TEXT("p.RenderInterp.ErrorCorrectionDuration"), RenderInterpErrorCorrectionDuration, TEXT("How long in seconds to apply error correction over."));
 		return RenderInterpErrorCorrectionDuration;
 	}
+	static float GetRenderInterpMaximumErrorCorrectionBeforeSnapping()
+	{
+		static float RenderInterpMaximumErrorCorrectionBeforeSnapping = 100.0f;
+		static FAutoConsoleVariableRef CVarRenderInterpErrorCorrectionMaximumError(TEXT("p.RenderInterp.MaximumErrorCorrectionBeforeSnapping"), RenderInterpMaximumErrorCorrectionBeforeSnapping, TEXT("Maximum error correction in cm before we stop interpolating and snap to target."));
+		return RenderInterpMaximumErrorCorrectionBeforeSnapping;
+	}
 	static float GetRenderInterpErrorVelocitySmoothingDuration()
 	{
 		static float RenderInterpErrorVelocitySmoothingDuration = 0.5f;
