Make consistent API by using GetX, SetX, GetP and SetP for physics particles

diff --git a/Engine/Source/Runtime/Experimental/Chaos/Private/Chaos/Deformable/ChaosDeformableSolver.cpp b/Engine/Source/Runtime/Experimental/Chaos/Private/Chaos/Deformable/ChaosDeformableSolver.cpp
--- a/Engine/Source/Runtime/Experimental/Chaos/Private/Chaos/Deformable/ChaosDeformableSolver.cpp
+++ b/Engine/Source/Runtime/Experimental/Chaos/Private/Chaos/Deformable/ChaosDeformableSolver.cpp
@@ -169,7 +169,7 @@ namespace Chaos::Softs
 			Chaos::FVec3 Position(0.f);
 			Chaos::FVec3 EulerRot(0.f);
 			int32 CollisionParticleOffset = Evolution->AddCollisionParticleRange(1, INDEX_NONE, true);
-			Evolution->CollisionParticles().X(0) = Position;
+			Evolution->CollisionParticles().SetX(0, Position);
 			Evolution->CollisionParticles().SetR(0, Chaos::TRotation<Chaos::FReal, 3>::MakeFromEuler(EulerRot));
 			Evolution->CollisionParticles().SetGeometry(0, MakeImplicitObjectPtr<Chaos::TPlane<Chaos::FReal, 3>>(Chaos::FVec3(0.f, 0.f, 0.f), Chaos::FVec3(0.f, 0.f, 1.f)));
 		}
@@ -223,7 +223,7 @@ namespace Chaos::Softs
 										TTuple<float, float, FVector3f>(
 											Evolution->Particles().InvM(ParticleIndex),
 											Evolution->Particles().PAndInvM(ParticleIndex).InvM,
-											Evolution->Particles().X(ParticleIndex))));
+											Evolution->Particles().GetX(ParticleIndex))));
 
 								Evolution->Particles().InvM(ParticleIndex) = 0.f;
 								Evolution->Particles().PAndInvM(ParticleIndex).InvM = 0.f;
@@ -242,12 +242,12 @@ namespace Chaos::Softs
 							FVector3f SimSpaceTarget = ToSingle(WorldToSim.TransformPosition(ToDouble(WorldSpaceTarget)));
 							const FVector3f& SimSpaceSource = TransientConstraintBuffer[ParticleIndex].Get<2>();
 
-							// Lerp from prevoius particle position to the target over the solver iterations.
-							Evolution->Particles().X(ParticleIndex) =
+							// Lerp from previous particle position to the target over the solver iterations.
+							Evolution->Particles().SetX(ParticleIndex,
 								SimSpaceTarget * CurrentRatio +
-								SimSpaceSource * (static_cast<FSolverReal>(1.) - CurrentRatio);
+								SimSpaceSource * (static_cast<FSolverReal>(1.) - CurrentRatio));
 							Evolution->Particles().PAndInvM(ParticleIndex).P =
-								Evolution->Particles().X(ParticleIndex);
+								Evolution->Particles().GetX(ParticleIndex);
 						}
 #if WITH_EDITOR
 						if (GDeformableDebugParams.IsDebugDrawingEnabled() && GDeformableDebugParams.bDoDrawTransientKinematicParticles)
@@ -257,7 +257,7 @@ namespace Chaos::Softs
 							{
 								int32 LocalIndex = CnstrTargets.GetIndex(i);
 								int32 ParticleIndex = Range[0] + LocalIndex;
-								Chaos::FDebugDrawQueue::GetInstance().DrawDebugPoint(ToDouble(Evolution->Particles().X(ParticleIndex)), FColor::Orange, false, -1.0f, 0, 5);
+								Chaos::FDebugDrawQueue::GetInstance().DrawDebugPoint(ToDouble(Evolution->Particles().GetX(ParticleIndex)), FColor::Orange, false, -1.0f, 0, 5);
 							}
 						}
 #endif
@@ -454,7 +454,7 @@ namespace Chaos::Softs
 			for (uint32 vdx = 0; vdx < NumParticles; ++vdx)
 			{
 				int32 SolverParticleIndex = ParticleStart + vdx;
-				Evolution->Particles().X(SolverParticleIndex) = ChaosVert(InitialPointsXf.TransformPosition(DoubleVert(DynamicVertex[vdx])));
+				Evolution->Particles().SetX(SolverParticleIndex, ChaosVert(InitialPointsXf.TransformPosition(DoubleVert(DynamicVertex[vdx]))));
 				Evolution->Particles().V(SolverParticleIndex) = Chaos::FVec3(0.f, 0.f, 0.f);
 				Evolution->Particles().M(SolverParticleIndex) = MassWithMultiplier[vdx];
 				Evolution->Particles().InvM(SolverParticleIndex) = ChaosInvM(Evolution->Particles().M(SolverParticleIndex));
@@ -466,7 +466,7 @@ namespace Chaos::Softs
 			for (uint32 vdx = 0; vdx < NumParticles; ++vdx)
 			{
 				int32 SolverParticleIndex = ParticleStart + vdx;
-				Evolution->Particles().X(SolverParticleIndex) = DynamicVertex[vdx];
+				Evolution->Particles().SetX(SolverParticleIndex, DynamicVertex[vdx]);
 				Evolution->Particles().V(SolverParticleIndex) = Chaos::FVec3(0.f, 0.f, 0.f);
 				Evolution->Particles().M(SolverParticleIndex) = MassWithMultiplier[vdx];
 				Evolution->Particles().InvM(SolverParticleIndex) = ChaosInvM(Evolution->Particles().M(SolverParticleIndex));
@@ -641,7 +641,7 @@ namespace Chaos::Softs
 							{
 								int32 Index = Evolution->AddCollisionParticle(INDEX_NONE, true);
 								int32 ViewIndex = Evolution->CollisionParticlesActiveView().GetRanges().Num() - 1;
-								Evolution->CollisionParticles().X(Index) = AddBody.Transform.GetTranslation();
+								Evolution->CollisionParticles().SetX(Index, AddBody.Transform.GetTranslation());
 								Evolution->CollisionParticles().SetR(Index, AddBody.Transform.GetRotation());
 								Chaos::FImplicitObjectPtr UniquePtr(AddBody.Shapes); AddBody.Shapes = nullptr;
 								Evolution->CollisionParticles().SetGeometry(Index, MoveTemp(UniquePtr));
@@ -704,7 +704,7 @@ namespace Chaos::Softs
 						if (Proxy.CollisionBodies.Contains(UpdateBody.Key))
 						{
 							FCollisionObjectParticleHandel* ParticleHandle = Proxy.CollisionBodies.Find(UpdateBody.Key);
-							Evolution->CollisionParticles().X(ParticleHandle->ParticleIndex) = UpdateBody.Transform.GetTranslation();
+							Evolution->CollisionParticles().SetX(ParticleHandle->ParticleIndex, UpdateBody.Transform.GetTranslation());
 							Evolution->CollisionParticles().SetR(ParticleHandle->ParticleIndex, UpdateBody.Transform.GetRotation());
 						}
 					}
@@ -733,7 +733,7 @@ namespace Chaos::Softs
 			{
 				auto T = ChaosTet(Tetrahedron[edx], Range[0]);
 				Chaos::FDebugDrawQueue::GetInstance().DrawDebugPoint(
-					DoubleVert(P.X(T[0])), FColor::Blue, false, -1.0f, 0, 5);
+					DoubleVert(P.GetX(T[0])), FColor::Blue, false, -1.0f, 0, 5);
 			}
 		}
 #endif
@@ -1041,7 +1041,7 @@ namespace Chaos::Softs
 											PreviousFleshBuffer = this->PreviousInputPackage->ObjectMap[Owner]->As<FFleshThreadingProxy::FFleshInputBuffer>();
 										}
 
-										MParticles.X(Index) = Chaos::TVector<FSolverReal, 3>((FSolverReal)0.);
+										MParticles.SetX(Index, Chaos::TVector<FSolverReal, 3>((FSolverReal)0.));
 										TVector<FSolverReal, 3> TargetPos((FSolverReal)0.);
 										FSolverReal CurrentRatio = FSolverReal(this->Iteration) / FSolverReal(this->Property.NumSolverSubSteps);
 
@@ -1066,7 +1066,7 @@ namespace Chaos::Softs
 														ComponentPointAtT = ComponentPointAtT * CurrentRatio + BonePreviousTransform.TransformPosition(LocalPoint) * ((FSolverReal)1. - CurrentRatio);
 													}
 
-													MParticles.X(Index) = MParticles.X(Index) + GlobalTransform.TransformPosition(ComponentPointAtT) * BoneWeights[i];
+													MParticles.SetX(Index, MParticles.GetX(Index) + GlobalTransform.TransformPosition(ComponentPointAtT) * BoneWeights[i]);
 
 													bParticleTouched = true;
 												}
@@ -1078,18 +1078,18 @@ namespace Chaos::Softs
 											if (GDeformableDebugParams.IsDebugDrawingEnabled() && GDeformableDebugParams.bDoDrawKinematicParticles)
 											{
 												auto DoubleVert = [](FVector3f V) { return FVector3d(V.X, V.Y, V.Z); };
-												Chaos::FDebugDrawQueue::GetInstance().DrawDebugPoint(DoubleVert(MParticles.X(Index)), FColor::Red, false, -1.0f, 0, 5);
+												Chaos::FDebugDrawQueue::GetInstance().DrawDebugPoint(DoubleVert(MParticles.GetX(Index)), FColor::Red, false, -1.0f, 0, 5);
 											}
 #endif
 										}
-										MParticles.PAndInvM(Index).P = MParticles.X(Index);
+										MParticles.PAndInvM(Index).P = MParticles.GetX(Index);
 									}
 								}
 							}
 							if (!bParticleTouched)
 							{
-								MParticles.X(Index) = GlobalTransform.TransformPosition(ChaosVert(Vertex[Index - Range[0]]));
-								MParticles.PAndInvM(Index).P = MParticles.X(Index);
+								MParticles.SetX(Index, GlobalTransform.TransformPosition(ChaosVert(Vertex[Index - Range[0]])));
+								MParticles.PAndInvM(Index).P = MParticles.GetX(Index);
 
 #if WITH_EDITOR
 								//debug draw
@@ -1098,7 +1098,7 @@ namespace Chaos::Softs
 								if (GDeformableDebugParams.IsDebugDrawingEnabled() && GDeformableDebugParams.bDoDrawKinematicParticles)
 								{
 									auto DoubleVert = [](FVector3f V) { return FVector3d(V.X, V.Y, V.Z); };
-									Chaos::FDebugDrawQueue::GetInstance().DrawDebugPoint(DoubleVert(MParticles.X(Index)), FColor::Red, false, -1.0f, 0, 5);
+									Chaos::FDebugDrawQueue::GetInstance().DrawDebugPoint(DoubleVert(MParticles.GetX(Index)), FColor::Red, false, -1.0f, 0, 5);
 								}
 #endif
 							}
@@ -1647,7 +1647,7 @@ namespace Chaos::Softs
 				{
 					for (int32 vdx = 0; vdx < Position.Num(); vdx++)
 					{
-						const Chaos::FVec3f& Pos = Evolution->Particles().X(vdx + Range[0]);
+						const Chaos::FVec3f& Pos = Evolution->Particles().GetX(vdx + Range[0]);
 						FVector PosD = UEVertd(Pos);
 						Position[vdx] = UEVertf(FinalXf.TransformPosition(PosD));
 					}
@@ -1656,7 +1656,7 @@ namespace Chaos::Softs
 				{
 					for (int32 vdx = 0; vdx < Position.Num(); vdx++)
 					{
-						Position[vdx] = UEVertf(UEVertd(Evolution->Particles().X(vdx + Range[0])));
+						Position[vdx] = UEVertf(UEVertd(Evolution->Particles().GetX(vdx + Range[0])));
 					}
 				}
 			}
@@ -1689,7 +1689,7 @@ namespace Chaos::Softs
 								if (GeomType == ImplicitObjectType::Sphere)
 								{
 									const FSphere& SphereGeometry = Geometry->GetObjectChecked<FSphere>();
-									FVector Center = ToFVector(CollisionParticles.X(Index)) + SphereGeometry.GetCenter();
+									FVector Center = ToFVector(CollisionParticles.GetX(Index)) + SphereGeometry.GetCenter();
 									FReal Radius = SphereGeometry.GetRadius();
 									Chaos::FDebugDrawQueue::GetInstance().DrawDebugSphere(Center, Radius, 12, FColor::Red, false, -1.0f, 0, 1.f);
 								}
@@ -1697,14 +1697,14 @@ namespace Chaos::Softs
 								{
 									const TBox<FReal, 3>& BoxGeometry = Geometry->GetObjectChecked<TBox<FReal, 3>>();
 									FVector Extent = 0.5 * (BoxGeometry.Max() - BoxGeometry.Min());
-									FVector Center = ToFVector(CollisionParticles.X(Index)) + BoxGeometry.GetCenter();
+									FVector Center = ToFVector(CollisionParticles.GetX(Index)) + BoxGeometry.GetCenter();
 									const FQuat& Rotation = ToFQuat(CollisionParticles.GetR(Index));
 									Chaos::FDebugDrawQueue::GetInstance().DrawDebugBox(Center, Extent, Rotation, FColor::Red, false, -1.0f, 0, 1.f);
 								}
 								else if (GeomType == ImplicitObjectType::Convex)
 								{
 									const FConvex& ConvexGeometry = Geometry->GetObjectChecked<FConvex>();
-									FTransform M = FTransform(ToFQuat(CollisionParticles.GetR(Index)), ToFVector(CollisionParticles.X(Index)));
+									FTransform M = FTransform(ToFQuat(CollisionParticles.GetR(Index)), ToFVector(CollisionParticles.GetX(Index)));
 									for (int32 EdgeIndex = 0; EdgeIndex < ConvexGeometry.NumEdges(); ++EdgeIndex)
 									{
 										int32 Index0 = ConvexGeometry.GetEdgeVertex(EdgeIndex, 0);
@@ -1765,7 +1765,7 @@ namespace Chaos::Softs
 
 		for (int32 i = 0; i < Np; i++) {
 
-			FString ParticleInfo = FString::SanitizeFloat(Particles.X(i)[0]) + FString(TEXT(" ")) + FString::SanitizeFloat(Particles.X(i)[1]) + FString(TEXT(" ")) + FString::SanitizeFloat(Particles.X(i)[2]) + FString(TEXT(" ")) + FString::FromInt(1) + FString(TEXT("\r\n"));
+			FString ParticleInfo = FString::SanitizeFloat(Particles.GetX(i)[0]) + FString(TEXT(" ")) + FString::SanitizeFloat(Particles.GetX(i)[1]) + FString(TEXT(" ")) + FString::SanitizeFloat(Particles.GetX(i)[2]) + FString(TEXT(" ")) + FString::FromInt(1) + FString(TEXT("\r\n"));
 			FFileHelper::SaveStringToFile(ParticleInfo, *file, FFileHelper::EEncodingOptions::AutoDetect, &IFileManager::Get(), EFileWrite::FILEWRITE_Append);
 
 		}
diff --git a/Engine/Source/Runtime/Experimental/Chaos/Private/Chaos/Evolution/PBDMinEvolution.cpp b/Engine/Source/Runtime/Experimental/Chaos/Private/Chaos/Evolution/PBDMinEvolution.cpp
--- a/Engine/Source/Runtime/Experimental/Chaos/Private/Chaos/Evolution/PBDMinEvolution.cpp
+++ b/Engine/Source/Runtime/Experimental/Chaos/Private/Chaos/Evolution/PBDMinEvolution.cpp
@@ -180,21 +180,21 @@ namespace Chaos
 		{
 			if (Particle.ObjectState() == EObjectStateType::Dynamic)
 			{
-				Particle.X() = FVec3::Lerp(Particle.Handle()->AuxilaryValue(ParticlePrevXs), Particle.X(), T);
-				Particle.SetR(FRotation3::Slerp(Particle.Handle()->AuxilaryValue(ParticlePrevRs), Particle.GetR(), (decltype(FQuat::X))T));	// LWC_TODO: Remove decltype cast once FQuat supports variants
+				Particle.SetX(FVec3::Lerp(Particle.Handle()->AuxilaryValue(ParticlePrevXs), Particle.GetX(), T));
+				Particle.SetRf(FRotation3f::Slerp(FRotation3f(Particle.Handle()->AuxilaryValue(ParticlePrevRs)), Particle.GetRf(), FRealSingle(T)));
 
 				if (bRewindVelocities)
 				{
-					Particle.SetV(FVec3::Lerp(Particle.GetPreV(), Particle.GetV(), T));
-					Particle.SetW(FVec3::Lerp(Particle.GetPreW(), Particle.GetW(), T));
+					Particle.SetVf(FVec3f::Lerp(Particle.GetPreVf(), Particle.GetVf(), FRealSingle(T)));
+					Particle.SetWf(FVec3f::Lerp(Particle.GetPreWf(), Particle.GetWf(), FRealSingle(T)));
 				}
 			}
 		}
 
 		for (auto& Particle : Particles.GetActiveKinematicParticlesView())
 		{
-			Particle.X() = Particle.X() - Particle.GetV() * RewindDt;
-			Particle.SetR(FRotation3::IntegrateRotationWithAngularVelocity(Particle.GetR(), -Particle.GetW(), RewindDt));
+			Particle.SetX(Particle.GetX() - Particle.GetV() * RewindDt);
+			Particle.SetRf(FRotation3f::IntegrateRotationWithAngularVelocity(Particle.GetRf(), -Particle.GetWf(), FRealSingle(RewindDt)));
 		}
 	}
 
@@ -276,7 +276,7 @@ namespace Chaos
 				// Update cached world space state, including bounds. We use the Swept bounds update so that the bounds includes P,Q and X,Q.
 				// This is because when we have joints, they often pull bodies back to their original positions, so we need to know if there
 				// are contacts at that location.
-				Particle.UpdateWorldSpaceStateSwept(FRigidTransform3(Particle.P(), Particle.GetQ()), BoundsExpansion, -V * Dt);
+				Particle.UpdateWorldSpaceStateSwept(FRigidTransform3(Particle.GetP(), Particle.GetQ()), BoundsExpansion, -V * Dt);
 			}
 		}
 	}
@@ -302,7 +302,7 @@ namespace Chaos
 			}
 
 			TKinematicTarget<FReal, 3>& KinematicTarget = Particle.KinematicTarget();
-			const FVec3 CurrentX = Particle.X();
+			const FVec3 CurrentX = Particle.GetX();
 			const FRotation3 CurrentR = Particle.GetR();
 			constexpr FReal MinDt = 1e-6f;
 
@@ -370,7 +370,7 @@ namespace Chaos
 			{
 				// Move based on velocity
 				bMoved = true;
-				Particle.X() = Particle.X() + Particle.GetV() * Dt;
+				Particle.SetX(Particle.GetX() + Particle.GetV() * Dt);
 				Particle.SetR(FRotation3::IntegrateRotationWithAngularVelocity(Particle.GetR(), Particle.GetW(), Dt));
 				break;
 			}
@@ -378,7 +378,7 @@ namespace Chaos
 			
 			// Set positions and previous velocities if we can
 			// Note: At present kinematics are in fact rigid bodies
-			Particle.P() = Particle.X();
+			Particle.SetP(Particle.GetX());
 			Particle.SetQ(Particle.GetR());
 			Particle.SetPreV(Particle.GetV());
 			Particle.SetPreW(Particle.GetW());
@@ -387,11 +387,11 @@ namespace Chaos
 			{
 				if (!Particle.CCDEnabled())
 				{
-					Particle.UpdateWorldSpaceState(FRigidTransform3(Particle.P(), Particle.GetQ()), FVec3(0));
+					Particle.UpdateWorldSpaceState(FRigidTransform3(Particle.GetP(), Particle.GetQ()), FVec3(0));
 				}
 				else
 				{
-					Particle.UpdateWorldSpaceStateSwept(FRigidTransform3(Particle.P(), Particle.GetQ()), FVec3(0), -Particle.GetV() * Dt);
+					Particle.UpdateWorldSpaceStateSwept(FRigidTransform3(Particle.GetP(), Particle.GetQ()), FVec3(0), -Particle.GetV() * Dt);
 				}
 			}
 		};
@@ -402,7 +402,7 @@ namespace Chaos
 		-> void
 		{
 			TKinematicTarget<FReal, 3>& KinematicTarget = Particle.KinematicTarget();
-			const FVec3 CurrentX = Particle.X();
+			const FVec3 CurrentX = Particle.GetX();
 			const FRotation3 CurrentR = Particle.GetR();
 			constexpr FReal MinDt = 1e-6f;
 
@@ -458,7 +458,7 @@ namespace Chaos
 						NewW = FRotation3::CalculateAngularVelocity(CurrentR, NewR, Dt);
 					}
 				}
-				Particle.X() = NewX;
+				Particle.SetX(NewX);
 				Particle.SetR(NewR);
 				Particle.SetV(NewV);
 				Particle.SetW(NewW);
@@ -470,15 +470,15 @@ namespace Chaos
 			{
 				// Move based on velocity
 				bMoved = true;
-				Particle.X() = Particle.X() + Particle.GetV() * Dt;
-				Particle.SetR(FRotation3::IntegrateRotationWithAngularVelocity(Particle.GetR(), Particle.GetW(), Dt));
+				Particle.SetX(Particle.GetX() + Particle.GetV() * Dt);
+				Particle.SetRf(FRotation3f::IntegrateRotationWithAngularVelocity(Particle.GetRf(), Particle.GetWf(), FRealSingle(Dt)));
 				break;
 			}
 			}
 
 			if (bMoved)
 			{
-				Particle.UpdateWorldSpaceState(FRigidTransform3(Particle.X(), Particle.GetR()), FVec3(0));
+				Particle.UpdateWorldSpaceState(FRigidTransform3(Particle.GetX(), Particle.GetR()), FVec3(0));
 			}
 		};
 
@@ -536,9 +536,9 @@ namespace Chaos
 
 		for (auto& Particle : Particles.GetActiveParticlesView())
 		{
-			Particle.Handle()->AuxilaryValue(ParticlePrevXs) = Particle.X();
+			Particle.Handle()->AuxilaryValue(ParticlePrevXs) = Particle.GetX();
 			Particle.Handle()->AuxilaryValue(ParticlePrevRs) = Particle.GetR();
-			Particle.X() = Particle.P();
+			Particle.SetX(Particle.GetP());
 			Particle.SetR(Particle.GetQ());
 		}
 	}
diff --git a/Engine/Source/Runtime/Experimental/Chaos/Private/Chaos/Evolution/SolverBodyContainer.cpp b/Engine/Source/Runtime/Experimental/Chaos/Private/Chaos/Evolution/SolverBodyContainer.cpp
--- a/Engine/Source/Runtime/Experimental/Chaos/Private/Chaos/Evolution/SolverBodyContainer.cpp
+++ b/Engine/Source/Runtime/Experimental/Chaos/Private/Chaos/Evolution/SolverBodyContainer.cpp
@@ -31,14 +31,14 @@ namespace Chaos
 		else
 		{
 			// Static or kinematic particle
-			SolverBody.SetP(Particle->X());
+			SolverBody.SetP(Particle->GetX());
 			SolverBody.SetQ(Particle->GetR());
 			SolverBody.SetCoM(FVec3(0));
 			SolverBody.SetRoM(FRotation3::FromIdentity());
 			if (KinematicParticle != nullptr)
 			{
 				// Kinematic particle
-				SolverBody.SetX(KinematicParticle->X() - KinematicParticle->GetV() * Dt);
+				SolverBody.SetX(KinematicParticle->GetX() - KinematicParticle->GetV() * Dt);
 				SolverBody.SetR(Particle->GetR());
 				if (!KinematicParticle->GetW().IsNearlyZero())
 				{
@@ -49,7 +49,7 @@ namespace Chaos
 			}
 			else
 			{
-				SolverBody.SetX(Particle->X());
+				SolverBody.SetX(Particle->GetX());
 				SolverBody.SetR(Particle->GetR());
 				SolverBody.SetV(FVec3(0));
 				SolverBody.SetW(FVec3(0));
diff --git a/Engine/Source/Runtime/Experimental/Chaos/Private/Chaos/GeometryParticles.cpp b/Engine/Source/Runtime/Experimental/Chaos/Private/Chaos/GeometryParticles.cpp
--- a/Engine/Source/Runtime/Experimental/Chaos/Private/Chaos/GeometryParticles.cpp
+++ b/Engine/Source/Runtime/Experimental/Chaos/Private/Chaos/GeometryParticles.cpp
@@ -55,7 +55,7 @@ namespace Chaos
 	template <typename T, int d, EGeometryParticlesSimType SimType>
 	void TGeometryParticlesImp<T, d, SimType>::UpdateShapesArray(const int32 Index)
 	{
-		UpdateShapesArrayFromGeometry(MShapesArray[Index], GetGeometry(Index), FRigidTransform3(X(Index), GetR(Index)));
+		UpdateShapesArrayFromGeometry(MShapesArray[Index], GetGeometry(Index), FRigidTransform3(GetX(Index), GetR(Index)));
 	}
 
 	template <typename T, int d, EGeometryParticlesSimType SimType>
diff --git a/Engine/Source/Runtime/Experimental/Chaos/Private/Chaos/ImplicitObject.cpp b/Engine/Source/Runtime/Experimental/Chaos/Private/Chaos/ImplicitObject.cpp
--- a/Engine/Source/Runtime/Experimental/Chaos/Private/Chaos/ImplicitObject.cpp
+++ b/Engine/Source/Runtime/Experimental/Chaos/Private/Chaos/ImplicitObject.cpp
@@ -136,12 +136,12 @@ namespace Chaos
 			TArray<int32> PotentialParticles = Particles->FindAllIntersections(ImplicitBox);
 			for(int32 i : PotentialParticles)
 			{
-				FVec3 LocalPoint = OtherToLocalTransform.TransformPosition(Particles->X(i));
+				FVec3 LocalPoint = OtherToLocalTransform.TransformPosition(Particles->GetX(i));
 				FReal LocalPhi = SignedDistance(LocalPoint);
 				if(LocalPhi < Phi)
 				{
 					Phi = LocalPhi;
-					Point = Particles->X(i);
+					Point = Particles->GetX(i);
 				}
 			}
 		}
@@ -165,12 +165,12 @@ namespace Chaos
 		int32 NumParticles = Particles->Size();
 		for(int32 i = 0; i < NumParticles; ++i)
 		{
-			FVec3 LocalPoint = OtherToLocalTransform.TransformPosition(Particles->X(i));
+			FVec3 LocalPoint = OtherToLocalTransform.TransformPosition(Particles->GetX(i));
 			FReal LocalPhi = SignedDistance(LocalPoint);
 			if(LocalPhi < Phi)
 			{
 				Phi = LocalPhi;
-				Point = Particles->X(i);
+				Point = Particles->GetX(i);
 			}
 		}
 		return MakePair(Point, Phi < Thickness);
diff --git a/Engine/Source/Runtime/Experimental/Chaos/Private/Chaos/Island/IslandManager.cpp b/Engine/Source/Runtime/Experimental/Chaos/Private/Chaos/Island/IslandManager.cpp
--- a/Engine/Source/Runtime/Experimental/Chaos/Private/Chaos/Island/IslandManager.cpp
+++ b/Engine/Source/Runtime/Experimental/Chaos/Private/Chaos/Island/IslandManager.cpp
@@ -125,7 +125,7 @@ namespace Chaos::Private
 				const FKinematicTarget& KinematicTarget = Kinematic->KinematicTarget();
 				if (KinematicTarget.GetMode() == EKinematicTargetMode::Position)
 				{
-					bIsStationary = (Kinematic->X() - KinematicTarget.GetTargetPosition()).IsZero() && (Kinematic->GetR() * KinematicTarget.GetTargetRotation().Inverse()).IsIdentity();
+					bIsStationary = (Kinematic->GetX() - KinematicTarget.GetTargetPosition()).IsZero() && (Kinematic->GetR() * KinematicTarget.GetTargetRotation().Inverse()).IsIdentity();
 				}
 				else
 				{
@@ -262,7 +262,7 @@ namespace Chaos::Private
 		if (Dt > UE_SMALL_NUMBER)
 		{
 			const FReal SmoothRate = FMath::Clamp(CVars::SmoothedPositionLerpRate, 0.0f, 1.0f);
-			const FVec3 VImp = FVec3::CalculateVelocity(Rigid.X(), Rigid.P(), Dt);
+			const FVec3 VImp = FVec3::CalculateVelocity(Rigid.GetX(), Rigid.GetP(), Dt);
 			const FVec3 WImp = FRotation3::CalculateAngularVelocity(Rigid.GetR(), Rigid.GetQ(), Dt);
 			Rigid.SetVSmooth(FMath::Lerp(Rigid.VSmooth(), VImp, SmoothRate));
 			Rigid.SetWSmooth(FMath::Lerp(Rigid.WSmooth(), WImp, SmoothRate));
diff --git a/Engine/Source/Runtime/Experimental/Chaos/Private/Chaos/Levelset.cpp b/Engine/Source/Runtime/Experimental/Chaos/Private/Chaos/Levelset.cpp
--- a/Engine/Source/Runtime/Experimental/Chaos/Private/Chaos/Levelset.cpp
+++ b/Engine/Source/Runtime/Experimental/Chaos/Private/Chaos/Levelset.cpp
@@ -295,7 +295,7 @@ FReal FLevelSet::ComputeLevelSetError(const FParticles& InParticles, const TArra
 
 	ParallelFor(Mesh.GetNumElements(), [&](int32 i) {
 		const TVec3<int32> CurrMeshFace = Faces[i];
-		const FVec3 MeshFaceCenter = (InParticles.X(CurrMeshFace[0]) + InParticles.X(CurrMeshFace[1]) + InParticles.X(CurrMeshFace[2])) / 3.f;
+		const FVec3 MeshFaceCenter = (InParticles.GetX(CurrMeshFace[0]) + InParticles.GetX(CurrMeshFace[1]) + InParticles.GetX(CurrMeshFace[2])) / 3.f;
 
 		//FVec3 GridNormal;
 		//FReal phi = PhiWithNormal(MeshFaceCenter, GridNormal);
@@ -310,7 +310,7 @@ FReal FLevelSet::ComputeLevelSetError(const FParticles& InParticles, const TArra
 
 			for (int j = 0; j < 3; ++j)
 			{
-				DistErrorValues[i] += FMath::Abs(SignedDistance(InParticles.X(CurrMeshFace[j])));
+				DistErrorValues[i] += FMath::Abs(SignedDistance(InParticles.GetX(CurrMeshFace[j])));
 			}
 
 			// per triangle error average of 3 corners and center distance to surface according to MPhi
@@ -323,7 +323,7 @@ FReal FLevelSet::ComputeLevelSetError(const FParticles& InParticles, const TArra
 			//AngleErrorValues[i] = FMath::Acos(FVec3::DotProduct(MeshFaceNormal, GridNormal));
 
 			// triangle area used for weighted average
-			TriangleArea[i] = (FReal)0.5 * sqrt(FVec3::CrossProduct(InParticles.X(CurrMeshFace[1]) - InParticles.X(CurrMeshFace[0]), InParticles.X(CurrMeshFace[2]) - InParticles.X(CurrMeshFace[0])).SizeSquared());
+			TriangleArea[i] = (FReal)0.5 * sqrt(FVec3::CrossProduct(InParticles.GetX(CurrMeshFace[1]) - InParticles.GetX(CurrMeshFace[0]), InParticles.GetX(CurrMeshFace[2]) - InParticles.GetX(CurrMeshFace[0])).SizeSquared());
 		}
 	});
 
@@ -432,9 +432,9 @@ void FLevelSet::OutputDebugData(FErrorReporter& ErrorReporter, const FParticles&
 	for (int i = 0; i < Mesh.GetNumElements(); ++i)
 	{
 		const TVec3<int32> CurrMeshFace = Faces[i];
-		int idx = OutVerts.Add(InParticles.X(CurrMeshFace[0]));
-		OutVerts.Add(InParticles.X(CurrMeshFace[1]));
-		OutVerts.Add(InParticles.X(CurrMeshFace[2]));
+		int idx = OutVerts.Add(InParticles.GetX(CurrMeshFace[0]));
+		OutVerts.Add(InParticles.GetX(CurrMeshFace[1]));
+		OutVerts.Add(InParticles.GetX(CurrMeshFace[2]));
 
 		OutNormals.Add(Normals[i]);
 		OutNormals.Add(Normals[i]);
@@ -772,7 +772,7 @@ bool FLevelSet::ComputeDistancesNearZeroIsocontour(FErrorReporter& ErrorReporter
 	const TArray<TVec3<int32>>& Elements = Mesh.GetSurfaceElements();
 	if (Elements.Num() > 0)
 	{
-		MOriginalLocalBoundingBox = FAABB3(InParticles.X(Elements[0][0]), InParticles.X(Elements[0][0]));
+		MOriginalLocalBoundingBox = FAABB3(InParticles.GetX(Elements[0][0]), InParticles.GetX(Elements[0][0]));
 	}
 	else
 	{
@@ -783,10 +783,10 @@ bool FLevelSet::ComputeDistancesNearZeroIsocontour(FErrorReporter& ErrorReporter
 	for (int32 Index = 0; Index < Elements.Num(); ++Index)
 	{
 		const auto& Element = Elements[Index];
-		TPlane<FReal, 3> TrianglePlane(InParticles.X(Element[0]), Normals[Index]);
-		FAABB3 TriangleBounds(InParticles.X(Element[0]), InParticles.X(Element[0]));
-		TriangleBounds.GrowToInclude(InParticles.X(Element[1]));
-		TriangleBounds.GrowToInclude(InParticles.X(Element[2]));
+		TPlane<FReal, 3> TrianglePlane(InParticles.GetX(Element[0]), Normals[Index]);
+		FAABB3 TriangleBounds(InParticles.GetX(Element[0]), InParticles.GetX(Element[0]));
+		TriangleBounds.GrowToInclude(InParticles.GetX(Element[1]));
+		TriangleBounds.GrowToInclude(InParticles.GetX(Element[2]));
 		MOriginalLocalBoundingBox.GrowToInclude(TriangleBounds); //also save the original bounding box
 
 		TVec3<int32> StartIndex = MGrid.Cell(TriangleBounds.Min() - FVec3((0.5f + UE_KINDA_SMALL_NUMBER) * MGrid.Dx()));
@@ -799,7 +799,7 @@ bool FLevelSet::ComputeDistancesNearZeroIsocontour(FErrorReporter& ErrorReporter
 				{
 					const TVec3<int32> CellIndex(i, j, k);
 					const FVec3 Center = MGrid.Location(CellIndex);
-					const FVec3 Point = FindClosestPointOnTriangle(TrianglePlane, InParticles.X(Element[0]), InParticles.X(Element[1]), InParticles.X(Element[2]), Center);
+					const FVec3 Point = FindClosestPointOnTriangle(TrianglePlane, InParticles.GetX(Element[0]), InParticles.GetX(Element[1]), InParticles.GetX(Element[2]), Center);
 
 					FReal NewPhi = (Point - Center).Size();
 					if (NewPhi < MPhi(CellIndex))
@@ -1337,7 +1337,7 @@ bool FLevelSet::IsIntersectingWithTriangle(const FParticles& Particles, const TV
 	if (Intersection.Second)
 	{
 		const FReal Epsilon = (FReal)1e-1; //todo(ocohen): fattening triangle up is relative to triangle size. Do we care about very large triangles?
-		const FVec2 Bary = ComputeBarycentricInPlane(Particles.X(Element[0]), Particles.X(Element[1]), Particles.X(Element[2]), Intersection.First);
+		const FVec2 Bary = ComputeBarycentricInPlane(Particles.GetX(Element[0]), Particles.GetX(Element[1]), Particles.GetX(Element[2]), Intersection.First);
 
 		if (Bary.X >= -Epsilon && Bary.Y >= -Epsilon && (Bary.Y + Bary.X) <= 1 + Epsilon)
 		{
@@ -1402,7 +1402,7 @@ void FLevelSet::ComputeNormals(const FParticles& InParticles, const FTriangleMes
 	const TArray<TVec3<int32>>& Elements = Mesh.GetSurfaceElements();
 	if (Elements.Num() > 0)
 	{
-		MOriginalLocalBoundingBox = FAABB3(InParticles.X(Elements[0][0]), InParticles.X(Elements[0][0]));
+		MOriginalLocalBoundingBox = FAABB3(InParticles.GetX(Elements[0][0]), InParticles.GetX(Elements[0][0]));
 	}
 	else
 	{
@@ -1411,10 +1411,10 @@ void FLevelSet::ComputeNormals(const FParticles& InParticles, const FTriangleMes
 	for (int32 Index = 0; Index < Elements.Num(); ++Index)
 	{
 		const auto& Element = Elements[Index];
-		TPlane<FReal, 3> TrianglePlane(InParticles.X(Element[0]), Normals[Index]);
-		FAABB3 TriangleBounds(InParticles.X(Element[0]), InParticles.X(Element[0]));
-		TriangleBounds.GrowToInclude(InParticles.X(Element[1]));
-		TriangleBounds.GrowToInclude(InParticles.X(Element[2]));
+		TPlane<FReal, 3> TrianglePlane(InParticles.GetX(Element[0]), Normals[Index]);
+		FAABB3 TriangleBounds(InParticles.GetX(Element[0]), InParticles.GetX(Element[0]));
+		TriangleBounds.GrowToInclude(InParticles.GetX(Element[1]));
+		TriangleBounds.GrowToInclude(InParticles.GetX(Element[2]));
 		MOriginalLocalBoundingBox.GrowToInclude(TriangleBounds); //also save the original bounding box
 
 		TVec3<int32> StartIndex = MGrid.Cell(TriangleBounds.Min() - FVec3((0.5f + UE_KINDA_SMALL_NUMBER) * MGrid.Dx()));
@@ -1431,7 +1431,7 @@ void FLevelSet::ComputeNormals(const FParticles& InParticles, const FTriangleMes
 						continue;
 					}
 					const FVec3 Center = MGrid.Location(CellIndex);
-					const FVec3 Point = FindClosestPointOnTriangle(TrianglePlane, InParticles.X(Element[0]), InParticles.X(Element[1]), InParticles.X(Element[2]), Center);
+					const FVec3 Point = FindClosestPointOnTriangle(TrianglePlane, InParticles.GetX(Element[0]), InParticles.GetX(Element[1]), InParticles.GetX(Element[2]), Center);
 
 					FReal NewPhi = (Point - Center).Size();
 					if (NewPhi < LocalPhi(CellIndex))
diff --git a/Engine/Source/Runtime/Experimental/Chaos/Private/Chaos/NewtonEvolution.cpp b/Engine/Source/Runtime/Experimental/Chaos/Private/Chaos/NewtonEvolution.cpp
--- a/Engine/Source/Runtime/Experimental/Chaos/Private/Chaos/NewtonEvolution.cpp
+++ b/Engine/Source/Runtime/Experimental/Chaos/Private/Chaos/NewtonEvolution.cpp
@@ -320,7 +320,7 @@ void FNewtonEvolution::WriteOutputLog(const int32 Frame)
 		for (int32 i = 0; i < int32(InParticles.Size()); ++i) {
 			for (int32 alpha = 0; alpha < 3; alpha++) {
 				Residual[i][alpha] *= DtSquared;
-				Residual[i][alpha] -= NodalMass[i] * (InParticles.P(i)[alpha] - InParticles.X(i)[alpha] - Dt * Vn[i][alpha]);
+				Residual[i][alpha] -= NodalMass[i] * (InParticles.GetP(i)[alpha] - InParticles.GetX(i)[alpha] - Dt * Vn[i][alpha]);
 			}
 		}
 	}
@@ -668,7 +668,7 @@ void FNewtonEvolution::WriteOutputLog(const int32 Frame)
 						//MParticles.V(Index) *= DampingPowDt;
 
 						//new code:
-						MParticles.P(Index) = MParticles.X(Index) + MParticles.V(Index) * Dt;
+						MParticles.SetP(Index, MParticles.GetX(Index) + MParticles.GetV(Index) * Dt);
 
 						MParticles.V(Index) *= 1.f;
 
@@ -779,7 +779,7 @@ void FNewtonEvolution::WriteOutputLog(const int32 Frame)
 					[this, Dt](FSolverCollisionParticles& CollisionParticles, int32 Index)
 					{
 						// Store active collision particle frames prior to the kinematic update for CCD collisions
-						MCollisionTransforms[Index] = FSolverRigidTransform3(CollisionParticles.X(Index), CollisionParticles.GetR(Index));
+						MCollisionTransforms[Index] = FSolverRigidTransform3(CollisionParticles.GetX(Index), CollisionParticles.GetR(Index));
 
 						// Update collision transform and velocity
 						MCollisionKinematicUpdate(CollisionParticles, Dt, MTime, Index);
@@ -826,9 +826,9 @@ void FNewtonEvolution::WriteOutputLog(const int32 Frame)
 					MParticlesActiveView.ParallelFor(
 						[Dt, this](FSolverParticles& Particles, int32 Index)
 						{
-							Particles.V(Index) = (Particles.P(Index) - Particles.X(Index)) / Dt;
-							Particles.X(Index) = Particles.P(Index);
-							this->MVn[Index] = Particles.V(Index);
+							Particles.SetV(Index, (Particles.P(Index) - Particles.GetX(Index)) / Dt);
+							Particles.SetX(Index, Particles.P(Index));
+							this->MVn[Index] = Particles.GetV(Index);
 						}, MinParallelBatchSize);
 				}
 			}
diff --git a/Engine/Source/Runtime/Experimental/Chaos/Private/Chaos/PBDCollisionSpringConstraintsBase.cpp b/Engine/Source/Runtime/Experimental/Chaos/Private/Chaos/PBDCollisionSpringConstraintsBase.cpp
--- a/Engine/Source/Runtime/Experimental/Chaos/Private/Chaos/PBDCollisionSpringConstraintsBase.cpp
+++ b/Engine/Source/Runtime/Experimental/Chaos/Private/Chaos/PBDCollisionSpringConstraintsBase.cpp
@@ -166,7 +166,7 @@ void FPBDCollisionSpringConstraintsBase::Init(const SolverParticlesOrRange& Part
 
 				TArray< TTriangleCollisionPoint<FSolverReal> > Result;
 				int32 ConstraintsAdded = 0;
-				if (CollidableMesh.PointProximityQuery(Spatial, static_cast<const TArrayView<const FSolverVec3>&>(Particles.XArray()), Index, Particles.X(Index), Thickness * ExtraThicknessMult, MaxSingleSidedThickness * ExtraThicknessMult,
+				if (CollidableMesh.PointProximityQuery(Spatial, static_cast<const TArrayView<const FSolverVec3>&>(Particles.XArray()), Index, Particles.GetX(Index), Thickness * ExtraThicknessMult, MaxSingleSidedThickness * ExtraThicknessMult,
 					[this, bVertexHasCollisionLayers, &Particles, &CollidableSubMesh, &Elements, &VertexGIAColors, &TriangleGIAColors](const int32 PointIndex, const int32 SubMeshTriangleIndex)->bool
 					{
 						const TVector<int32, 3>& Elem = Elements[SubMeshTriangleIndex];
@@ -284,8 +284,8 @@ void FPBDCollisionSpringConstraintsBase::Init(const SolverParticlesOrRange& Part
 						if (!bFaceIsKinematic && !bUseCollisionLayerOverride)
 						{
 							// NOTE: CollisionPoint.Normal has already been flipped to point toward the Point, so need to recalculate here.
-							const TTriangle<FSolverReal> Triangle(Particles.X(Elem[0]), Particles.X(Elem[1]), Particles.X(Elem[2]));
-							bFlipNormal = (Particles.X(Index) - CollisionPoint.Location).Dot(Triangle.GetNormal()) < 0; // Is Point currently behind Triangle?
+							const TTriangle<FSolverReal> Triangle(Particles.GetX(Elem[0]), Particles.GetX(Elem[1]), Particles.GetX(Elem[2]));
+							bFlipNormal = (Particles.GetX(Index) - CollisionPoint.Location).Dot(Triangle.GetNormal()) < 0; // Is Point currently behind Triangle?
 							// Doing a check against ANY (plus the TriangleGIAColors which captures sub-triangle intersections) seems to work better than checking against ALL vertex colors where the triangle must agree.
 							// In particular, it's better at handling thin regions of intersection where a single vertex or line of vertices intersect through faces.
 							// Want Point to push to opposite side of triangle
@@ -424,8 +424,8 @@ FSolverVec3 FPBDCollisionSpringConstraintsBase::GetDelta(const SolverParticlesOr
 
 	if (ConstraintFriction > 0)
 	{
-		const FSolverVec3& X1 = Particles.X(Index1);
-		const FSolverVec3 X = Bary[0] * Particles.X(Index2) + Bary[1] * Particles.X(Index3) + Bary[2] * Particles.X(Index4);
+		const FSolverVec3& X1 = Particles.GetX(Index1);
+		const FSolverVec3 X = Bary[0] * Particles.GetX(Index2) + Bary[1] * Particles.GetX(Index3) + Bary[2] * Particles.GetX(Index4);
 		const FSolverVec3 RelativeDisplacement = (P1 - X1) - (P - X) + (Particles.InvM(Index1) - TrianglePointInvM) * RepulsionDelta;
 		const FSolverVec3 RelativeDisplacementTangent = RelativeDisplacement - RelativeDisplacement.Dot(Normal) * Normal;
 		const FSolverReal RelativeDisplacementTangentLength = RelativeDisplacementTangent.Length();
diff --git a/Engine/Source/Runtime/Experimental/Chaos/Private/Chaos/PBDEvolution.cpp b/Engine/Source/Runtime/Experimental/Chaos/Private/Chaos/PBDEvolution.cpp
--- a/Engine/Source/Runtime/Experimental/Chaos/Private/Chaos/PBDEvolution.cpp
+++ b/Engine/Source/Runtime/Experimental/Chaos/Private/Chaos/PBDEvolution.cpp
@@ -343,7 +343,7 @@ void FPBDEvolution::PreIterationUpdate(
 					const int32 Index = Offset + i;
 					if (MParticles.InvM(Index) != (FSolverReal)0.)  // Process dynamic particles
 					{
-						MParticles.P(Index) = MParticles.X(Index);
+						MParticles.SetP(Index, MParticles.GetX(Index));
 					}
 					else  // Process kinematic particles
 					{
@@ -406,7 +406,7 @@ void FPBDEvolution::PreIterationUpdate(
 						}
 
 						// Euler Step with point damping integration
-						MParticles.P(Index) = MParticles.X(Index) + MParticles.V(Index) * ParticleDampingIntegrated;
+						MParticles.SetP(Index, MParticles.GetX(Index) + MParticles.GetV(Index) * ParticleDampingIntegrated);
 
 						MParticles.V(Index) *= ParticleDampingPowDt;
 					}
@@ -452,7 +452,7 @@ void FPBDEvolution::PreIterationUpdate(
 						}
 
 						// Euler Step with point damping integration
-						MParticles.P(Index) = MParticles.X(Index) + MParticles.V(Index) * DampingIntegrated;
+						MParticles.SetP(Index, MParticles.GetX(Index) + MParticles.GetV(Index) * DampingIntegrated);
 
 						MParticles.V(Index) *= DampingPowDt;
 					}
@@ -553,7 +553,7 @@ void FPBDEvolution::AdvanceOneTimeStep(const FSolverReal Dt)
 				[this, Dt](FSolverCollisionParticles& CollisionParticles, int32 Index)
 				{
 					// Store active collision particle frames prior to the kinematic update for CCD collisions
-					MCollisionTransforms[Index] = FSolverRigidTransform3(CollisionParticles.X(Index), CollisionParticles.GetR(Index));
+					MCollisionTransforms[Index] = FSolverRigidTransform3(CollisionParticles.GetX(Index), CollisionParticles.GetR(Index));
 
 					// Update collision transform and velocity
 					MCollisionKinematicUpdate(CollisionParticles, Dt, MTime, Index); // This expects Sequential update.
@@ -696,8 +696,8 @@ void FPBDEvolution::AdvanceOneTimeStep(const FSolverReal Dt)
 				MParticlesActiveView.ParallelFor(
 					[Dt](FSolverParticles& Particles, int32 Index)
 					{
-						Particles.V(Index) = (Particles.P(Index) - Particles.X(Index)) / Dt;
-						Particles.X(Index) = Particles.P(Index);
+						Particles.SetV(Index, (Particles.GetP(Index) - Particles.GetX(Index)) / Dt);
+						Particles.SetX(Index, Particles.GetP(Index));
 					}, MinParallelBatchSize);
 			}
 		}
diff --git a/Engine/Source/Runtime/Experimental/Chaos/Private/Chaos/PBDJointConstraints.cpp b/Engine/Source/Runtime/Experimental/Chaos/Private/Chaos/PBDJointConstraints.cpp
--- a/Engine/Source/Runtime/Experimental/Chaos/Private/Chaos/PBDJointConstraints.cpp
+++ b/Engine/Source/Runtime/Experimental/Chaos/Private/Chaos/PBDJointConstraints.cpp
@@ -557,11 +557,11 @@ namespace Chaos
 	{
 		FPBDJointSettings JointSettings;
 		JointSettings.ConnectorTransforms[0] = FRigidTransform3(
-			WorldConstraintFrame.GetTranslation() - InConstrainedParticles[0]->X(),
+			WorldConstraintFrame.GetTranslation() - InConstrainedParticles[0]->GetX(),
 			WorldConstraintFrame.GetRotation() * InConstrainedParticles[0]->GetR().Inverse()
 			);
 		JointSettings.ConnectorTransforms[1] = FRigidTransform3(
-			WorldConstraintFrame.GetTranslation() - InConstrainedParticles[1]->X(),
+			WorldConstraintFrame.GetTranslation() - InConstrainedParticles[1]->GetX(),
 			WorldConstraintFrame.GetRotation() * InConstrainedParticles[1]->GetR().Inverse()
 			);
 		return AddConstraint(InConstrainedParticles, JointSettings);
diff --git a/Engine/Source/Runtime/Experimental/Chaos/Private/Chaos/PBDRigidClustering.cpp b/Engine/Source/Runtime/Experimental/Chaos/Private/Chaos/PBDRigidClustering.cpp
--- a/Engine/Source/Runtime/Experimental/Chaos/Private/Chaos/PBDRigidClustering.cpp
+++ b/Engine/Source/Runtime/Experimental/Chaos/Private/Chaos/PBDRigidClustering.cpp
@@ -726,7 +726,7 @@ namespace Chaos
 			{
 				FBreakingData& ClusterBreak = MAllClusterBreakings.AddDefaulted_GetRef();
 				ClusterBreak.Proxy = ClusteredParticle->PhysicsProxy();
-				ClusterBreak.Location = ClusteredParticle->X();
+				ClusterBreak.Location = ClusteredParticle->GetX();
 				ClusterBreak.Velocity = ClusteredParticle->GetV();
 				ClusterBreak.AngularVelocity = ClusteredParticle->GetW();
 				ClusterBreak.Mass = ClusteredParticle->M();
@@ -752,7 +752,7 @@ namespace Chaos
 			{
 				FCrumblingData& ClusterCrumbling = MAllClusterCrumblings.AddDefaulted_GetRef();
 				ClusterCrumbling.Proxy = ClusteredParticle->PhysicsProxy();
-				ClusterCrumbling.Location = ClusteredParticle->X();
+				ClusterCrumbling.Location = ClusteredParticle->GetX();
 				ClusterCrumbling.Orientation = ClusteredParticle->GetR();
 				ClusterCrumbling.LinearVelocity = ClusteredParticle->GetV();
 				ClusterCrumbling.AngularVelocity = ClusteredParticle->GetW();
@@ -860,12 +860,12 @@ namespace Chaos
 
 			ClusteredChild->SetClusterId(ClusterId(nullptr, ClusteredChild->ClusterIds().NumChildren)); // clear Id but retain number of children
 
-			const FRigidTransform3 PreSolveTM(ClusteredParent->P(), ClusteredParent->GetQ());
+			const FRigidTransform3 PreSolveTM(ClusteredParent->GetP(), ClusteredParent->GetQ());
 			const FRigidTransform3 ChildFrame = ClusteredChild->ChildToParent() * PreSolveTM;
 			Child->SetX(ChildFrame.GetTranslation());
 			Child->SetR(ChildFrame.GetRotation());
 
-			Child->SetP(Child->X());
+			Child->SetP(Child->GetX());
 			Child->SetQf(Child->GetRf());
 
 			//todo(ocohen): for now just inherit velocity at new COM. This isn't quite right for rotation
@@ -928,7 +928,7 @@ namespace Chaos
 		}
 		NewClusters.Reserve(NumNewClusters);
 
-		const FRigidTransform3 PreSolveTM = FRigidTransform3(ClusteredParent->P(), ClusteredParent->GetQ());
+		const FRigidTransform3 PreSolveTM = FRigidTransform3(ClusteredParent->GetP(), ClusteredParent->GetQ());
 		
 		TArray<Chaos::FPBDRigidClusteredParticleHandle*> NewClusterHandles = MEvolution.CreateClusteredParticles(NumNewClusters);
 		int32 ClusterHandlesIdx = 0;
@@ -952,7 +952,7 @@ namespace Chaos
 				NewCluster->SetWf(ClusteredParent->GetWf());
 				NewCluster->SetPreVf(ClusteredParent->GetPreVf());
 				NewCluster->SetPreWf(ClusteredParent->GetPreWf());
-				NewCluster->SetP(NewCluster->X());
+				NewCluster->SetP(NewCluster->GetX());
 				NewCluster->SetQf(NewCluster->GetRf());
 
 				UpdateTopLevelParticle(NewCluster);
@@ -2042,7 +2042,7 @@ namespace Chaos
 		FReal ClosestSquaredDist = TNumericLimits<FReal>::Max();
 		for (FPBDRigidParticleHandle* ChildHandle: Particles)
         {
-        	const FReal SquaredDist = (ChildHandle->X() - WorldLocation).SizeSquared();
+        	const FReal SquaredDist = (ChildHandle->GetX() - WorldLocation).SizeSquared();
         	if (SquaredDist < ClosestSquaredDist)
         	{
         		ClosestSquaredDist = SquaredDist;
@@ -2074,7 +2074,7 @@ namespace Chaos
 		const FReal RadiusSquared = Radius * Radius;
 		for (FPBDRigidParticleHandle* ChildHandle: Particles)
 		{
-			const FReal SquaredDist = (ChildHandle->X() - WorldLocation).SizeSquared();
+			const FReal SquaredDist = (ChildHandle->GetX() - WorldLocation).SizeSquared();
 			if (SquaredDist <= RadiusSquared)
 			{
 				Result.Add(ChildHandle);
@@ -2313,7 +2313,7 @@ namespace Chaos
 					}
 					else
 					{
-						const FRigidTransform3 WorldToClusterTM = FRigidTransform3(Cluster->P(), Cluster->GetQ());
+						const FRigidTransform3 WorldToClusterTM = FRigidTransform3(Cluster->GetP(), Cluster->GetQ());
 						const FVec3 ContactLocationClusterLocal = WorldToClusterTM.InverseTransformPosition(ContactWorldLocation);
 						FAABB3 ContactBox(ContactLocationClusterLocal, ContactLocationClusterLocal);
 						ContactBox.Thicken(ClusterDistanceThreshold);
@@ -2652,7 +2652,7 @@ namespace Chaos
 				{
 					ParticlePairArray& ConnectionList = Connections[i];
 
-					const FVec3& Child1X = Child1->X();
+					const FVec3& Child1X = Child1->GetX();
 					FRigidTransform3 TM1 = FRigidTransform3(Child1X, Child1->GetR());
 
 					const int32 Offset = i + 1;
@@ -2665,13 +2665,13 @@ namespace Chaos
 						if (Child2->CollisionParticles())
 						{
 
-							const FVec3& Child2X = Child2->X();
+							const FVec3& Child2X = Child2->GetX();
 							const FRigidTransform3 TM = TM1.GetRelativeTransform(FRigidTransform3(Child2X, Child2->GetR()));
 							const uint32 NumCollisionParticles = Child2->CollisionParticles()->Size();
 							for (uint32 CollisionIdx = 0; CollisionIdx < NumCollisionParticles; ++CollisionIdx)
 							{
 								const FVec3 LocalPoint =
-									TM.TransformPositionNoScale(Child2->CollisionParticles()->X(CollisionIdx));
+									TM.TransformPositionNoScale(Child2->CollisionParticles()->GetX(CollisionIdx));
 								const FReal Phi = Child1->GetGeometry()->SignedDistance(LocalPoint - (LocalPoint * Delta));
 								if (Phi < 0.0)
 								{
@@ -2720,7 +2720,7 @@ namespace Chaos
 		Pts.AddUninitialized(Particles.Num());
 		for (int32 i = 0; i < Particles.Num(); ++i)
 		{
-			Pts[i] = Particles[i]->X();
+			Pts[i] = Particles[i]->GetX();
 		}
 		TArray<TArray<int32>> Neighbors; // Indexes into Particles
 		VoronoiNeighbors(Pts, Neighbors);
@@ -2896,7 +2896,7 @@ namespace Chaos
 		Pts.AddUninitialized(Children.Num());
 		for (int32 i = 0; i < Children.Num(); ++i)
 		{
-			Pts[i] = Children[i]->X();
+			Pts[i] = Children[i]->GetX();
 		}
 		TArray<TArray<int>> Neighbors;
 		VoronoiNeighbors(Pts, Neighbors);
diff --git a/Engine/Source/Runtime/Experimental/Chaos/Private/Chaos/PBDRigidClusteringAlgo.cpp b/Engine/Source/Runtime/Experimental/Chaos/Private/Chaos/PBDRigidClusteringAlgo.cpp
--- a/Engine/Source/Runtime/Experimental/Chaos/Private/Chaos/PBDRigidClusteringAlgo.cpp
+++ b/Engine/Source/Runtime/Experimental/Chaos/Private/Chaos/PBDRigidClusteringAlgo.cpp
@@ -109,7 +109,7 @@ namespace Chaos
 		{
 			MassToLocal.SetTranslation(Cluster->CenterOfMass());
 			Cluster->SetX(Cluster->XCom());
-			Cluster->SetP(Cluster->X());
+			Cluster->SetP(Cluster->GetX());
 			Cluster->SetCenterOfMass(FVec3::ZeroVector);
 		}
 
@@ -275,7 +275,7 @@ namespace Chaos
 		Objects.Reserve(Children.Num());
 		Objects2.Reserve(Children.Num());
 
-		const FRigidTransform3 ClusterWorldTM(Parent->X(), Parent->GetR());
+		const FRigidTransform3 ClusterWorldTM(Parent->GetX(), Parent->GetR());
 
 		TArray<FVec3> OriginalPoints;
 		TArray<FPBDRigidParticleHandle*> ChildParticleHandles;
@@ -307,7 +307,7 @@ namespace Chaos
 				}
 				else
 				{
-					const FRigidTransform3 ChildWorldTM(Child->X(), Child->GetR());
+					const FRigidTransform3 ChildWorldTM(Child->GetX(), Child->GetR());
 					Frame = ChildWorldTM.GetRelativeTransform(ClusterWorldTM);
 				}
 
@@ -331,7 +331,7 @@ namespace Chaos
 					{
 						for (uint32 i = 0; i < CollisionParticles->Size(); ++i)
 						{
-							OriginalPoints.Add(Frame.TransformPosition(CollisionParticles->X(i)));
+							OriginalPoints.Add(Frame.TransformPosition(CollisionParticles->GetX(i)));
 						}
 					}
 				}
@@ -448,7 +448,7 @@ namespace Chaos
 				CollisionParticles->AddParticles(CleanedPoints.Num());
 				for (int32 i = 0; i < CleanedPoints.Num(); ++i)
 				{
-					CollisionParticles->X(i) = CleanedPoints[i];
+					CollisionParticles->SetX(i, CleanedPoints[i]);
 				}
 			}
 
@@ -458,7 +458,7 @@ namespace Chaos
 				Parent->CollisionParticles()->UpdateAccelerationStructures();
 			}
 		}
-		const Chaos::FRigidTransform3 Xf(Parent->X(), Parent->GetR());
+		const Chaos::FRigidTransform3 Xf(Parent->GetX(), Parent->GetR());
 		Parent->UpdateWorldSpaceState(Xf, FVec3(0));
 
 		UpdateCollisionFlags(Parent, bUseParticleImplicit);
diff --git a/Engine/Source/Runtime/Experimental/Chaos/Private/Chaos/PBDRigidDynamicSpringConstraints.cpp b/Engine/Source/Runtime/Experimental/Chaos/Private/Chaos/PBDRigidDynamicSpringConstraints.cpp
--- a/Engine/Source/Runtime/Experimental/Chaos/Private/Chaos/PBDRigidDynamicSpringConstraints.cpp
+++ b/Engine/Source/Runtime/Experimental/Chaos/Private/Chaos/PBDRigidDynamicSpringConstraints.cpp
@@ -35,8 +35,8 @@ namespace Chaos
 
 			const FRotation3 Q0 = bIsRigidDynamic0 ? PBDRigid0->GetQ() : Static0->GetR();
 			const FRotation3 Q1 = bIsRigidDynamic1 ? PBDRigid1->GetQ() : Static1->GetR();
-			const FVec3& P0 = bIsRigidDynamic0 ? PBDRigid0->P() : Static0->X();
-			const FVec3& P1 = bIsRigidDynamic1 ? PBDRigid1->P() : Static1->X();
+			const FVec3 P0 = bIsRigidDynamic0 ? PBDRigid0->GetP() : Static0->GetX();
+			const FVec3 P1 = bIsRigidDynamic1 ? PBDRigid1->GetP() : Static1->GetX();
 
 			// Delete constraints
 			const int32 NumSprings = SpringDistances[ConstraintIndex].Num();
diff --git a/Engine/Source/Runtime/Experimental/Chaos/Private/Chaos/PBDRigidSpringConstraints.cpp b/Engine/Source/Runtime/Experimental/Chaos/Private/Chaos/PBDRigidSpringConstraints.cpp
--- a/Engine/Source/Runtime/Experimental/Chaos/Private/Chaos/PBDRigidSpringConstraints.cpp
+++ b/Engine/Source/Runtime/Experimental/Chaos/Private/Chaos/PBDRigidSpringConstraints.cpp
@@ -96,8 +96,8 @@ namespace Chaos
 		const TGeometryParticleHandle<FReal, 3>* Particle0 = Constraint[0];
 		const TGeometryParticleHandle<FReal, 3>* Particle1 = Constraint[1];
 
-		Distances[ConstraintIndex][0] = Particle0->GetR().Inverse().RotateVector(Location0 - Particle0->X());
-		Distances[ConstraintIndex][1] = Particle1->GetR().Inverse().RotateVector(Location1 - Particle1->X());
+		Distances[ConstraintIndex][0] = Particle0->GetR().Inverse().RotateVector(Location0 - Particle0->GetX());
+		Distances[ConstraintIndex][1] = Particle1->GetR().Inverse().RotateVector(Location1 - Particle1->GetX());
 	}
 
 	FVec3 FPBDRigidSpringConstraints::GetDelta(int32 ConstraintIndex, const FVec3& WorldSpaceX1, const FVec3& WorldSpaceX2) const
diff --git a/Engine/Source/Runtime/Experimental/Chaos/Private/Chaos/PBDRigidsEvolution.cpp b/Engine/Source/Runtime/Experimental/Chaos/Private/Chaos/PBDRigidsEvolution.cpp
--- a/Engine/Source/Runtime/Experimental/Chaos/Private/Chaos/PBDRigidsEvolution.cpp
+++ b/Engine/Source/Runtime/Experimental/Chaos/Private/Chaos/PBDRigidsEvolution.cpp
@@ -1170,7 +1170,7 @@ namespace Chaos
 		// Record removal for event generation
 		FRemovalData& Removal = MAllRemovals.AddDefaulted_GetRef();
 		Removal.Proxy = Particle->PhysicsProxy();
-		Removal.Location = Particle->X();
+		Removal.Location = Particle->GetX();
 		
 		if (Chaos::FPBDRigidParticleHandle* RigidParticle = Particle->CastToRigidParticle())
 		{
@@ -1230,8 +1230,8 @@ namespace Chaos
 			{
 				Data = &TestModeData.Add(Rigid);
 			}
-			Data->X = Rigid->X();
-			Data->P = Rigid->P();
+			Data->X = Rigid->GetX();
+			Data->P = Rigid->GetP();
 			Data->R = Rigid->GetR();
 			Data->Q = Rigid->GetQ();
 			Data->V = Rigid->GetV();
@@ -1273,8 +1273,8 @@ namespace Chaos
 			FTestModeParticleData* Data = TestModeData.Find(Rigid);
 			if (Data != nullptr)
 			{
-				Rigid->X() = Data->X;
-				Rigid->P() = Data->P;
+				Rigid->SetX(Data->X);
+				Rigid->SetP(Data->P);
 				Rigid->SetR(Data->R);
 				Rigid->SetQ(Data->Q);
 				Rigid->SetV(Data->V);
diff --git a/Engine/Source/Runtime/Experimental/Chaos/Private/Chaos/PBDRigidsEvolutionGBF.cpp b/Engine/Source/Runtime/Experimental/Chaos/Private/Chaos/PBDRigidsEvolutionGBF.cpp
--- a/Engine/Source/Runtime/Experimental/Chaos/Private/Chaos/PBDRigidsEvolutionGBF.cpp
+++ b/Engine/Source/Runtime/Experimental/Chaos/Private/Chaos/PBDRigidsEvolutionGBF.cpp
@@ -795,7 +795,7 @@ FPBDRigidsEvolutionGBF::FPBDRigidsEvolutionGBF(
 	{
 		ParticlesInput.ParallelFor([&](auto& Particle, int32 Index)
 		{
-			Particle.X() = Particle.P();
+			Particle.SetX(Particle.GetP());
 			Particle.SetR(Particle.GetQ());
 
 			//TODO: rename this function since it's not just updating position
@@ -1033,7 +1033,7 @@ void FPBDRigidsEvolutionGBF::DestroyTransientConstraints(FGeometryParticleHandle
 
 CHAOS_API void FPBDRigidsEvolutionGBF::SetParticleTransform(FGeometryParticleHandle* InParticle, const FVec3& InPos, const FRotation3& InRot, const bool bIsTeleport)
 {
-	const FVec3 PrevX = InParticle->X();
+	const FVec3 PrevX = InParticle->GetX();
 	const FRotation3 PrevR = InParticle->GetR();
 
 	FGenericParticleHandle(InParticle)->SetTransform(InPos, InRot);
@@ -1056,7 +1056,7 @@ void FPBDRigidsEvolutionGBF::SetParticleTransformSwept(FGeometryParticleHandle*
 	// that are not centerd on their center of mass (either it's multi-shape body, or the user modified the CoM).
 	const FVec3 CoMOffset = FConstGenericParticleHandle(InParticle)->CenterOfMass();
 	const FVec3 EndCoM = InPos + InRot * CoMOffset;
-	const FVec3 StartCoM = InParticle->X() + InParticle->GetR() * CoMOffset;
+	const FVec3 StartCoM = InParticle->GetX() + InParticle->GetR() * CoMOffset;
 	FVec3 SweepDir = EndCoM - StartCoM;
 	const FReal SweepLength = SweepDir.SafeNormalize();
 	const FVec3 SweepStart = InPos - SweepDir * SweepLength;
@@ -1114,7 +1114,7 @@ void FPBDRigidsEvolutionGBF::OnParticleMoved(FGeometryParticleHandle* InParticle
 	// is different by very small amounts around 1e-7 in both position and rotation when switching from dynamic to kinematic.
 	const FReal CollisionPositionTolerance = FReal(1.e-4);
 	const FReal CollisionRotationTolerance = FReal(1.e-6);
-	if (!FVec3::IsNearlyEqual(PrevX, InParticle->X(), CollisionPositionTolerance) || !FRotation3::IsNearlyEqual(PrevR, InParticle->GetR(), CollisionRotationTolerance))
+	if (!FVec3::IsNearlyEqual(PrevX, InParticle->GetX(), CollisionPositionTolerance) || !FRotation3::IsNearlyEqual(PrevR, InParticle->GetR(), CollisionRotationTolerance))
 	{
 		GetIslandManager().WakeParticleIslands(InParticle);
 
@@ -1137,7 +1137,7 @@ void FPBDRigidsEvolutionGBF::OnParticleMoved(FGeometryParticleHandle* InParticle
 			{
 				// @todo(chaos): better sleep system! We should probably have a sleep accumulator per particle
 				const FReal InvDt = FReal(30.0);
-				const FVec3 DV = (PrevX - Rigid->X()) * InvDt;
+				const FVec3 DV = (PrevX - Rigid->GetX()) * InvDt;
 				const FReal SmoothRate = FMath::Clamp(CVars::SmoothedPositionLerpRate, 0.0f, 1.0f);
 				Rigid->VSmooth() = FMath::Lerp(Rigid->VSmooth(), Rigid->GetV() + DV, SmoothRate);
 			}
diff --git a/Engine/Source/Runtime/Experimental/Chaos/Private/Chaos/PBDSoftBodyCollisionConstraint.cpp b/Engine/Source/Runtime/Experimental/Chaos/Private/Chaos/PBDSoftBodyCollisionConstraint.cpp
--- a/Engine/Source/Runtime/Experimental/Chaos/Private/Chaos/PBDSoftBodyCollisionConstraint.cpp
+++ b/Engine/Source/Runtime/Experimental/Chaos/Private/Chaos/PBDSoftBodyCollisionConstraint.cpp
@@ -292,7 +292,7 @@ void FPBDSoftBodyCollisionConstraintBase::ApplyInternal(FSolverParticlesRange& P
 					continue;
 				}
 
-				const FSolverRigidTransform3 Frame(CollisionParticlesRange.X(CollisionIndex), CollisionParticlesRange.R(CollisionIndex));
+				const FSolverRigidTransform3 Frame(CollisionParticlesRange.GetX(CollisionIndex), CollisionParticlesRange.R(CollisionIndex));
 				const FVec3 RigidSpacePosition(Frame.InverseTransformPosition(PAndInvM[Index].P));  // PhiWithNormal requires FReal based arguments
 				FVec3 ImplicitNormal;                                                                // since implicits don't use FSolverReal
 				FSolverReal Phi;
@@ -339,7 +339,7 @@ void FPBDSoftBodyCollisionConstraintBase::ApplyInternal(FSolverParticlesRange& P
 
 					if constexpr (bWithFriction)
 					{
-						const FSolverVec3 VectorToPoint = PAndInvM[Index].P - CollisionParticlesRange.X(VelocityBone);
+						const FSolverVec3 VectorToPoint = PAndInvM[Index].P - CollisionParticlesRange.GetX(VelocityBone);
 						const FSolverVec3 RelativeDisplacement = (PAndInvM[Index].P - X[Index]) - (CollisionParticlesRange.V(VelocityBone) + FSolverVec3::CrossProduct(CollisionParticlesRange.W(VelocityBone), VectorToPoint)) * Dt; // This corresponds to the tangential velocity multiplied by dt (friction will drive this to zero if it is high enough)
 						const FSolverVec3 RelativeDisplacementTangent = RelativeDisplacement - FSolverVec3::DotProduct(RelativeDisplacement, NormalWorld) * NormalWorld; // Project displacement into the tangential plane
 						const FSolverReal RelativeDisplacementTangentLength = RelativeDisplacementTangent.Size();
@@ -379,7 +379,7 @@ void FPBDSoftBodyCollisionConstraintBase::ApplyInternalCCD(FSolverParticlesRange
 					continue;
 				}
 
-				const FSolverRigidTransform3 Frame(CollisionParticlesRange.X(CollisionIndex), CollisionParticlesRange.R(CollisionIndex));
+				const FSolverRigidTransform3 Frame(CollisionParticlesRange.GetX(CollisionIndex), CollisionParticlesRange.R(CollisionIndex));
 
 				const Pair<FVec3, bool> PointPair = CollisionParticlesRange.GetGeometry(CollisionIndex)->FindClosestIntersection(  // Geometry operates in FReal
 					FVec3(CollisionTransforms[CollisionIndex].InverseTransformPositionNoScale(X[Index])),        // hence the back and forth
@@ -435,7 +435,7 @@ void FPBDSoftBodyCollisionConstraintBase::ApplyInternalCCD(FSolverParticlesRange
 							}
 						}
 
-						const FSolverVec3 VectorToPoint = PAndInvM[Index].P - CollisionParticlesRange.X(VelocityBone);
+						const FSolverVec3 VectorToPoint = PAndInvM[Index].P - CollisionParticlesRange.GetX(VelocityBone);
 						const FSolverVec3 RelativeDisplacement = (PAndInvM[Index].P - X[Index]) - (CollisionParticlesRange.V(VelocityBone) + FSolverVec3::CrossProduct(CollisionParticlesRange.W(VelocityBone), VectorToPoint)) * Dt;  // This corresponds to the tangential velocity multiplied by dt (friction will drive this to zero if it is high enough)
 						const FSolverVec3 RelativeDisplacementTangent = RelativeDisplacement - FSolverVec3::DotProduct(RelativeDisplacement, NormalWorld) * NormalWorld;  // Project displacement into the tangential plane
 						const FSolverReal RelativeDisplacementTangentLength = RelativeDisplacementTangent.Size();
@@ -545,7 +545,7 @@ void FPBDSoftBodyCollisionConstraintBase::UpdateLinearSystem(const FSolverPartic
 					continue;
 				}
 
-				const FSolverRigidTransform3 Frame(CollisionParticlesRange.X(CollisionIndex), CollisionParticlesRange.R(CollisionIndex));
+				const FSolverRigidTransform3 Frame(CollisionParticlesRange.GetX(CollisionIndex), CollisionParticlesRange.R(CollisionIndex));
 				const FVec3 RigidSpacePosition(Frame.InverseTransformPosition(PAndInvM[Index].P));  // PhiWithNormal requires FReal based arguments
 				FVec3 ImplicitNormal;                                                                // since implicits don't use FSolverReal
 				const FSolverReal Phi = (FSolverReal)CollisionParticlesRange.GetGeometry(CollisionIndex)->PhiWithNormal(RigidSpacePosition, ImplicitNormal);
diff --git a/Engine/Source/Runtime/Experimental/Chaos/Private/Chaos/PBDTriangleMeshCollisions.cpp b/Engine/Source/Runtime/Experimental/Chaos/Private/Chaos/PBDTriangleMeshCollisions.cpp
--- a/Engine/Source/Runtime/Experimental/Chaos/Private/Chaos/PBDTriangleMeshCollisions.cpp
+++ b/Engine/Source/Runtime/Experimental/Chaos/Private/Chaos/PBDTriangleMeshCollisions.cpp
@@ -53,8 +53,8 @@ static void FindEdgeFaceIntersections(const FTriangleMesh& TriangleMesh, const T
 
 			const int32 EdgePointIndex0 = SegmentMesh.GetElements()[EdgeIndex][0];
 			const int32 EdgePointIndex1 = SegmentMesh.GetElements()[EdgeIndex][1];
-			const FSolverVec3& EdgePosition0 = Particles.X(EdgePointIndex0);
-			const FSolverVec3& EdgePosition1 = Particles.X(EdgePointIndex1);
+			const FSolverVec3& EdgePosition0 = Particles.GetX(EdgePointIndex0);
+			const FSolverVec3& EdgePosition1 = Particles.GetX(EdgePointIndex1);
 			const bool bEdgeIsKinematic = Particles.InvM(EdgePointIndex0) == (FSolverReal)0. && Particles.InvM(EdgePointIndex1) == (FSolverReal)0.;
 			if (bSkipKinematic && bEdgeIsKinematic)
 			{
@@ -1144,8 +1144,8 @@ namespace ContourMinimization
 		const TArray<TVec2<int32>>& EdgeToFaces = TriangleMesh.GetEdgeToFaces();
 		const TArray<TVec3<int32>>& Elements = TriangleMesh.GetElements();
 
-		const FSolverVec3& EdgePosition0 = Particles.X(ContourIntersection.EdgeVertices[0]);
-		const FSolverVec3& EdgePosition1 = Particles.X(ContourIntersection.EdgeVertices[1]);
+		const FSolverVec3& EdgePosition0 = Particles.GetX(ContourIntersection.EdgeVertices[0]);
+		const FSolverVec3& EdgePosition1 = Particles.GetX(ContourIntersection.EdgeVertices[1]);
 		const FSolverVec3 EdgeDir = EdgePosition1 - EdgePosition0; // Unnormalized
 		const FSolverReal EdgeDirDotNormal = FSolverVec3::DotProduct(EdgeDir, EdgeFaceIntersection.FaceNormal);
 		if (FMath::Abs(EdgeDirDotNormal) >= UE_SMALL_NUMBER)
@@ -1160,9 +1160,9 @@ namespace ContourMinimization
 					continue;
 				}
 
-				const FSolverVec3& P0 = Particles.X(Elements[EdgeFace_i][0]);
-				const FSolverVec3& P1 = Particles.X(Elements[EdgeFace_i][1]);
-				const FSolverVec3& P2 = Particles.X(Elements[EdgeFace_i][2]);
+				const FSolverVec3& P0 = Particles.GetX(Elements[EdgeFace_i][0]);
+				const FSolverVec3& P1 = Particles.GetX(Elements[EdgeFace_i][1]);
+				const FSolverVec3& P2 = Particles.GetX(Elements[EdgeFace_i][2]);
 
 				// B_i
 				const TTriangle<FSolverReal> Triangle_i(P0, P1, P2);
diff --git a/Engine/Source/Runtime/Experimental/Chaos/Private/Chaos/PerParticleDampVelocity.cpp b/Engine/Source/Runtime/Experimental/Chaos/Private/Chaos/PerParticleDampVelocity.cpp
--- a/Engine/Source/Runtime/Experimental/Chaos/Private/Chaos/PerParticleDampVelocity.cpp
+++ b/Engine/Source/Runtime/Experimental/Chaos/Private/Chaos/PerParticleDampVelocity.cpp
@@ -57,7 +57,7 @@ void FPerParticleDampVelocity::UpdatePositionBasedState(const FSolverParticlesRa
 				continue;
 			}
 
-			Xcm += Particles.X(Index) * Particles.M(Index);
+			Xcm += Particles.GetX(Index) * Particles.M(Index);
 			Vcm += Particles.V(Index) * Particles.M(Index);
 			Mcm += Particles.M(Index);
 		}
@@ -77,7 +77,7 @@ void FPerParticleDampVelocity::UpdatePositionBasedState(const FSolverParticlesRa
 				continue;
 			}
 
-			const FSolverVec3 V = Particles.X(Index) - Xcm;
+			const FSolverVec3 V = Particles.GetX(Index) - Xcm;
 			L += FSolverVec3::CrossProduct(V, Particles.M(Index) * Particles.V(Index));
 			const FSolverMatrix33 M(0, V[2], -V[1], -V[2], 0, V[0], V[1], -V[0], 0);
 			I += M.GetTransposed() * M * Particles.M(Index);
diff --git a/Engine/Source/Runtime/Experimental/Chaos/Private/Chaos/PhysicsObjectInterface.cpp b/Engine/Source/Runtime/Experimental/Chaos/Private/Chaos/PhysicsObjectInterface.cpp
--- a/Engine/Source/Runtime/Experimental/Chaos/Private/Chaos/PhysicsObjectInterface.cpp
+++ b/Engine/Source/Runtime/Experimental/Chaos/Private/Chaos/PhysicsObjectInterface.cpp
@@ -98,14 +98,14 @@ namespace
 
 		if (ParticleHandle)
 		{
-			FVec3 ParticleX = ParticleHandle->X();
+			FVec3 ParticleX = ParticleHandle->GetX();
 			if (ParticleHandle->Disabled())
 			{
 				if (FPBDRigidClusteredParticleHandle* ClusteredParticle = ParticleHandle->CastToClustered())
 				{
 					if (const FPBDRigidClusteredParticleHandle* ClusteredParent = ClusteredParticle->Parent())
 					{
-						const FTransform ParentTransform(ClusteredParent->GetR(), ClusteredParent->X());
+						const FTransform ParentTransform(ClusteredParent->GetR(), ClusteredParent->GetX());
 						ParticleX = ParentTransform.TransformPosition(ClusteredParticle->ChildToParent().GetTranslation());
 					}
 				}
@@ -258,7 +258,7 @@ namespace Chaos
 
 		if (TThreadParticle<Id>* Particle = Object->GetParticle<Id>())
 		{
-			return Particle->X();
+			return Particle->GetX();
 		}
 
 		return FVector::Zero();
@@ -782,7 +782,7 @@ namespace Chaos
 			FBox ParticleBox(ForceInit);
 			if (const FImplicitObjectRef Geometry = Particle->GetGeometry(); Geometry && Geometry->HasBoundingBox())
 			{
-				const Chaos::FAABB3 WorldBox = Geometry->CalculateTransformedBounds(TRigidTransform<FReal, 3>(Particle->X(), Particle->GetR()));
+				const Chaos::FAABB3 WorldBox = Geometry->CalculateTransformedBounds(TRigidTransform<FReal, 3>(Particle->GetX(), Particle->GetR()));
 				ParticleBox = FBox{ WorldBox.Min(), WorldBox.Max() };
 			}
 
diff --git a/Engine/Source/Runtime/Experimental/Chaos/Private/Chaos/SegmentMesh.cpp b/Engine/Source/Runtime/Experimental/Chaos/Private/Chaos/SegmentMesh.cpp
--- a/Engine/Source/Runtime/Experimental/Chaos/Private/Chaos/SegmentMesh.cpp
+++ b/Engine/Source/Runtime/Experimental/Chaos/Private/Chaos/SegmentMesh.cpp
@@ -97,7 +97,7 @@ namespace Chaos
 			for(int32 i = 0; i < MElements.Num(); i++)
 			{
 				const TVec2<int32>& edge = MElements[i];
-				lengths[i] = (InParticles.X(edge[0]) - InParticles.X(edge[1])).SizeSquared();
+				lengths[i] = (InParticles.GetX(edge[0]) - InParticles.GetX(edge[1])).SizeSquared();
 			}
 		}
 		else
@@ -105,7 +105,7 @@ namespace Chaos
 			for(int32 i = 0; i < MElements.Num(); i++)
 			{
 				const TVec2<int32>& edge = MElements[i];
-				lengths[i] = (InParticles.X(edge[0]) - InParticles.X(edge[1])).Size();
+				lengths[i] = (InParticles.GetX(edge[0]) - InParticles.GetX(edge[1])).Size();
 			}
 		}
 		return lengths;
diff --git a/Engine/Source/Runtime/Experimental/Chaos/Private/Chaos/TriangleMeshImplicitObject.cpp b/Engine/Source/Runtime/Experimental/Chaos/Private/Chaos/TriangleMeshImplicitObject.cpp
--- a/Engine/Source/Runtime/Experimental/Chaos/Private/Chaos/TriangleMeshImplicitObject.cpp
+++ b/Engine/Source/Runtime/Experimental/Chaos/Private/Chaos/TriangleMeshImplicitObject.cpp
@@ -145,9 +145,9 @@ struct FTriangleMeshRaycastVisitor
 		constexpr FReal Epsilon = 1e-4f;
 
 		const int32 FaceIndex = TriIdx.Payload;
-		const FVec3& A = Particles.X(Elements[FaceIndex][0]);
-		const FVec3& B = Particles.X(Elements[FaceIndex][1]);
-		const FVec3& C = Particles.X(Elements[FaceIndex][2]);
+		const FVec3& A = Particles.GetX(Elements[FaceIndex][0]);
+		const FVec3& B = Particles.GetX(Elements[FaceIndex][1]);
+		const FVec3& C = Particles.GetX(Elements[FaceIndex][2]);
 
 		// Note: the math here needs to match FTriangleMeshImplicitObject::GetFaceNormal
 		// @todo(chaos) we should really preprocess the face and remove the degenerated ones to avoid paying this runtime cost
@@ -191,9 +191,9 @@ struct FTriangleMeshRaycastVisitor
 		const FReal R2 = R * R;
 
 		const int32 FaceIndex = TriIdx.Payload;		
-		const FVec3& A = Particles.X(Elements[FaceIndex][0]);
-		const FVec3& B = Particles.X(Elements[FaceIndex][1]);
-		const FVec3& C = Particles.X(Elements[FaceIndex][2]);
+		const FVec3& A = Particles.GetX(Elements[FaceIndex][0]);
+		const FVec3& B = Particles.GetX(Elements[FaceIndex][1]);
+		const FVec3& C = Particles.GetX(Elements[FaceIndex][2]);
 
 		// Note: the math here needs to match FTriangleMeshImplicitObject::GetFaceNormal
 		const FVec3 AB = B - A;
@@ -436,9 +436,9 @@ template <typename IdxType, typename ParticlesType>
 inline void TriangleMeshTransformVertsHelperSimd(const VectorRegister4Float& TriMeshScale, int32 TriIdx, const ParticlesType& Particles,
 	const TArray<TVector<IdxType, 3>>& Elements, VectorRegister4Float& OutA, VectorRegister4Float& OutB, VectorRegister4Float& OutC)
 {
-	OutA = VectorLoadFloat3(&Particles.X(Elements[TriIdx][0]).X);
-	OutB = VectorLoadFloat3(&Particles.X(Elements[TriIdx][1]).X);
-	OutC = VectorLoadFloat3(&Particles.X(Elements[TriIdx][2]).X);
+	OutA = VectorLoadFloat3(&Particles.GetX(Elements[TriIdx][0]).X);
+	OutB = VectorLoadFloat3(&Particles.GetX(Elements[TriIdx][1]).X);
+	OutC = VectorLoadFloat3(&Particles.GetX(Elements[TriIdx][2]).X);
 	OutA = VectorMultiply(OutA, TriMeshScale);
 	OutB = VectorMultiply(OutB, TriMeshScale);
 	OutC = VectorMultiply(OutC, TriMeshScale);
@@ -1010,9 +1010,9 @@ bool FTriangleMeshImplicitObject::OverlapImp(const TArray<TVec3<IdxType>>& Eleme
 
 	for (int32 TriIdx : PotentialIntersections)
 	{
-		const FVec3& A = MParticles.X(Elements[TriIdx][0]);
-		const FVec3& B = MParticles.X(Elements[TriIdx][1]);
-		const FVec3& C = MParticles.X(Elements[TriIdx][2]);
+		const FVec3& A = MParticles.GetX(Elements[TriIdx][0]);
+		const FVec3& B = MParticles.GetX(Elements[TriIdx][1]);
+		const FVec3& C = MParticles.GetX(Elements[TriIdx][2]);
 
 		const FVec3 AB = B - A;
 		const FVec3 AC = C - A;
@@ -1245,9 +1245,9 @@ struct FTriangleMeshSweepVisitor
 
 		const TParticles<FRealSingle, 3>& Particles = TriMesh.MParticles;
 
-		const TVector<FRealSingle, 3>& AVec = Particles.X(Elements[TriIdx][0]);
-		const TVector<FRealSingle, 3>& BVec = Particles.X(Elements[TriIdx][1]);
-		const TVector<FRealSingle, 3>& CVec = Particles.X(Elements[TriIdx][2]);
+		const TVector<FRealSingle, 3>& AVec = Particles.GetX(Elements[TriIdx][0]);
+		const TVector<FRealSingle, 3>& BVec = Particles.GetX(Elements[TriIdx][1]);
+		const TVector<FRealSingle, 3>& CVec = Particles.GetX(Elements[TriIdx][2]);
 
 		VectorRegister4Float A = MakeVectorRegister(AVec.X, AVec.Y, AVec.Z, 0.0f);
 		VectorRegister4Float B = MakeVectorRegister(BVec.X, BVec.Y, BVec.Z, 0.0f);
@@ -1410,9 +1410,9 @@ struct FTriangleMeshSweepVisitorCCD
 
 		const TParticles<FRealSingle, 3>& Particles = TriMesh.MParticles;
 
-		const TVector<FRealSingle, 3>& AVec = Particles.X(Elements[TriIdx][0]);
-		const TVector<FRealSingle, 3>& BVec = Particles.X(Elements[TriIdx][1]);
-		const TVector<FRealSingle, 3>& CVec = Particles.X(Elements[TriIdx][2]);
+		const TVector<FRealSingle, 3>& AVec = Particles.GetX(Elements[TriIdx][0]);
+		const TVector<FRealSingle, 3>& BVec = Particles.GetX(Elements[TriIdx][1]);
+		const TVector<FRealSingle, 3>& CVec = Particles.GetX(Elements[TriIdx][2]);
 
 		VectorRegister4Float A = MakeVectorRegister(AVec.X, AVec.Y, AVec.Z, 0.0f);
 		VectorRegister4Float B = MakeVectorRegister(BVec.X, BVec.Y, BVec.Z, 0.0f);
@@ -1755,9 +1755,9 @@ int32 FTriangleMeshImplicitObject::FindMostOpposingFace(const TArray<TVec3<IdxTy
 
 	for (int32 TriIdx : PotentialIntersections)
 	{
-		const FVec3& A = MParticles.X(Elements[TriIdx][0]) * Scale;
-		const FVec3& B = MParticles.X(Elements[TriIdx][1]) * Scale;
-		const FVec3& C = MParticles.X(Elements[TriIdx][2]) * Scale;
+		const FVec3& A = MParticles.GetX(Elements[TriIdx][0]) * Scale;
+		const FVec3& B = MParticles.GetX(Elements[TriIdx][1]) * Scale;
+		const FVec3& C = MParticles.GetX(Elements[TriIdx][2]) * Scale;
 
 		const FVec3 AB = B - A;
 		const FVec3 AC = C - A;
@@ -1905,9 +1905,9 @@ FVec3 FTriangleMeshImplicitObject::GetFaceNormal(const int32 FaceIdx) const
 	{
 		auto LambdaHelper = [&](const auto& Elements)
 		{
-			const ParticleVecType& A = MParticles.X(Elements[FaceIdx][0]);
-			const ParticleVecType& B = MParticles.X(Elements[FaceIdx][1]);
-			const ParticleVecType& C = MParticles.X(Elements[FaceIdx][2]);
+			const ParticleVecType& A = MParticles.GetX(Elements[FaceIdx][0]);
+			const ParticleVecType& B = MParticles.GetX(Elements[FaceIdx][1]);
+			const ParticleVecType& C = MParticles.GetX(Elements[FaceIdx][2]);
 
 			const ParticleVecType AB = B - A;
 			const ParticleVecType AC = C - A;
@@ -2016,7 +2016,7 @@ void FTriangleMeshImplicitObject::UpdateVertices(const TArray<FVector>& NewPosit
 		int32 InternalIdx = bRemapIndices ? (*ExternalVertexIndexMap.Get())[i] : i;
 		if (InternalIdx < (int32)MParticles.Size())
 		{
-			MParticles.X(InternalIdx) = Chaos::FVec3(NewPositions[i]);
+			MParticles.SetX(InternalIdx, Chaos::FVec3(NewPositions[i]));
 		}
 	}
 	RebuildFastBVH();
diff --git a/Engine/Source/Runtime/Experimental/Chaos/Private/Chaos/XPBDBendingConstraints.cpp b/Engine/Source/Runtime/Experimental/Chaos/Private/Chaos/XPBDBendingConstraints.cpp
--- a/Engine/Source/Runtime/Experimental/Chaos/Private/Chaos/XPBDBendingConstraints.cpp
+++ b/Engine/Source/Runtime/Experimental/Chaos/Private/Chaos/XPBDBendingConstraints.cpp
@@ -215,10 +215,10 @@ void FXPBDBendingConstraints::ApplyHelper(SolverParticlesOrRange& Particles, con
 	const FSolverReal Angle = CalcAngle(Particles.P(Index1), Particles.P(Index2), Particles.P(Index3), Particles.P(Index4));
 	const TStaticArray<FSolverVec3, 4> Grads = Base::GetGradients(Particles, ConstraintIndex);
 
-	const FSolverVec3 V1TimesDt = Particles.P(Index1) - Particles.X(Index1);
-	const FSolverVec3 V2TimesDt = Particles.P(Index2) - Particles.X(Index2);
-	const FSolverVec3 V3TimesDt = Particles.P(Index3) - Particles.X(Index3);
-	const FSolverVec3 V4TimesDt = Particles.P(Index4) - Particles.X(Index4);
+	const FSolverVec3 V1TimesDt = Particles.P(Index1) - Particles.GetX(Index1);
+	const FSolverVec3 V2TimesDt = Particles.P(Index2) - Particles.GetX(Index2);
+	const FSolverVec3 V3TimesDt = Particles.P(Index3) - Particles.GetX(Index3);
+	const FSolverVec3 V4TimesDt = Particles.P(Index4) - Particles.GetX(Index4);
 
 	FSolverReal& Lambda = bDampingOnly ? LambdasDamping[ConstraintIndex] : Lambdas[ConstraintIndex];
 	const FSolverReal AlphaInv = bDampingOnly ? (FSolverReal)0.f : BiphasicStiffnessValue * Dt * Dt;
diff --git a/Engine/Source/Runtime/Experimental/Chaos/Private/Chaos/XPBDStretchBiasElementConstraints.cpp b/Engine/Source/Runtime/Experimental/Chaos/Private/Chaos/XPBDStretchBiasElementConstraints.cpp
--- a/Engine/Source/Runtime/Experimental/Chaos/Private/Chaos/XPBDStretchBiasElementConstraints.cpp
+++ b/Engine/Source/Runtime/Experimental/Chaos/Private/Chaos/XPBDStretchBiasElementConstraints.cpp
@@ -286,9 +286,9 @@ void FXPBDStretchBiasElementConstraints::InitConstraintsAndRestData(const Solver
 			const FSolverMatrix22& DeltaUVInv = FMath::Abs(DeltaUVDet) < UE_SMALL_NUMBER ? DeltaUVInverse.Add_GetRef(FSolverMatrix22(1.f, 0.f, 0.f, 1.f)) :
 				DeltaUVInverse.Add_GetRef(FSolverMatrix22(DeltaVNormalized[1] / DeltaUVDet, -DeltaVNormalized[0] / DeltaUVDet, -DeltaUNormalized[1] / DeltaUVDet, DeltaUNormalized[0] / DeltaUVDet));
 
-			const FSolverVec3& X0 = InParticles.X(Constraint[0]);
-			const FSolverVec3& X1 = InParticles.X(Constraint[1]);
-			const FSolverVec3& X2 = InParticles.X(Constraint[2]);
+			const FSolverVec3& X0 = InParticles.GetX(Constraint[0]);
+			const FSolverVec3& X1 = InParticles.GetX(Constraint[1]);
+			const FSolverVec3& X2 = InParticles.GetX(Constraint[2]);
 			const FSolverVec3 X01 = X1 - X0;
 			const FSolverVec3 X02 = X2 - X0; 
 			const FSolverVec3 DXDu = X01 * DeltaUVInv.M[0] + X02 * DeltaUVInv.M[1];
@@ -829,9 +829,9 @@ void FXPBDStretchBiasElementConstraints::ApplyHelper(SolverParticlesOrRange& Par
 	const FSolverVec3 GradCs2 = (DXDvNormalized * OneOverDXDuLen * DeltaUVInv.M[1] + DXDuNormalized * OneOverDXDvLen * DeltaUVInv.M[3]);
 	const FSolverVec3 GradCs0 = -GradCs1 - GradCs2;
 
-	const FSolverVec3 V0TimesDt = Particles.P(i0) - Particles.X(i0);
-	const FSolverVec3 V1TimesDt = Particles.P(i1) - Particles.X(i1);
-	const FSolverVec3 V2TimesDt = Particles.P(i2) - Particles.X(i2);
+	const FSolverVec3 V0TimesDt = Particles.GetP(i0) - Particles.GetX(i0);
+	const FSolverVec3 V1TimesDt = Particles.GetP(i1) - Particles.GetX(i1);
+	const FSolverVec3 V2TimesDt = Particles.GetP(i2) - Particles.GetX(i2);
 
 	// These scale factors make everything resolution independent
 	const FSolverVec3 FinalStiffnesses = StiffnessScales[ConstraintIndex] * ExpStiffnessValue;
@@ -864,8 +864,8 @@ void FXPBDStretchBiasElementConstraints::InitializeDmInvAndMeasures(const FSolve
 	RestDmArray.Init(PMatrix<FSolverReal, 3, 2>(0.f),Constraints.Num());
 	for (int32 e = 0; e < Constraints.Num(); e++)
 	{
-		const TVec3<FSolverReal> X1X0 = Particles.X(Constraints[e][1]) - Particles.X(Constraints[e][0]);
-		const TVec3<FSolverReal> X2X0 = Particles.X(Constraints[e][2]) - Particles.X(Constraints[e][0]);
+		const TVec3<FSolverReal> X1X0 = Particles.GetX(Constraints[e][1]) - Particles.GetX(Constraints[e][0]);
+		const TVec3<FSolverReal> X2X0 = Particles.GetX(Constraints[e][2]) - Particles.GetX(Constraints[e][0]);
 		PMatrix<FSolverReal, 3, 2> RestDm(0.f);
 		PMatrix<FSolverReal, 2, 2> Dm(0.f, 0.f, 0.f);
 		Dm.M[0] = X1X0.Size();
diff --git a/Engine/Source/Runtime/Experimental/Chaos/Private/EventDefaults.cpp b/Engine/Source/Runtime/Experimental/Chaos/Private/EventDefaults.cpp
--- a/Engine/Source/Runtime/Experimental/Chaos/Private/EventDefaults.cpp
+++ b/Engine/Source/Runtime/Experimental/Chaos/Private/EventDefaults.cpp
@@ -393,13 +393,13 @@ namespace Chaos
 						ActiveParticle->GetGeometry() &&
 						ActiveParticle->GetGeometry()->HasBoundingBox())
 					{
-						if (ensure(!ActiveParticle->X().ContainsNaN() &&
+						if (ensure(!ActiveParticle->GetX().ContainsNaN() &&
 							!ActiveParticle->GetV().ContainsNaN() &&
 							!ActiveParticle->GetW().ContainsNaN() &&
 							FMath::IsFinite(ActiveParticle->M())))
 						{
 							FTrailingData TrailingData;
-							TrailingData.Location = ActiveParticle->X();
+							TrailingData.Location = ActiveParticle->GetX();
 							TrailingData.Velocity = ActiveParticle->GetV();
 							TrailingData.AngularVelocity = ActiveParticle->GetW();
 							TrailingData.Mass = ActiveParticle->M();
diff --git a/Engine/Source/Runtime/Experimental/Chaos/Private/GeometryCollection/GeometryCollectionCollisionStructureManager.cpp b/Engine/Source/Runtime/Experimental/Chaos/Private/GeometryCollection/GeometryCollectionCollisionStructureManager.cpp
--- a/Engine/Source/Runtime/Experimental/Chaos/Private/GeometryCollection/GeometryCollectionCollisionStructureManager.cpp
+++ b/Engine/Source/Runtime/Experimental/Chaos/Private/GeometryCollection/GeometryCollectionCollisionStructureManager.cpp
@@ -63,7 +63,7 @@ FCollisionStructureManager::NewSimplicial(
 			OutsideVertices.AddUninitialized(IndicesArray.Num());
 			for (int32 Idx : IndicesArray)
 			{
-				const Chaos::FVec3& SamplePoint = Vertices.X(Idx);
+				const Chaos::FVec3& SamplePoint = Vertices.GetX(Idx);
 				if (Implicit->SignedDistance(SamplePoint) > Threshold)
 				{
 					OutsideVertices[LSVCounter] = SamplePoint;
@@ -82,8 +82,8 @@ FCollisionStructureManager::NewSimplicial(
 			OutsideVertices.AddUninitialized(IndicesArray.Num());
 			for (int32 Idx=0;Idx<IndicesArray.Num();Idx++)
 			{
-				Bounds += FVector(Vertices.X(IndicesArray[Idx]));
-				OutsideVertices[Idx] = Vertices.X(IndicesArray[Idx]);
+				Bounds += FVector(Vertices.GetX(IndicesArray[Idx]));
+				OutsideVertices[Idx] = Vertices.GetX(IndicesArray[Idx]);
 			}
 			Extent = Bounds.GetExtent().Size();
 		}
@@ -104,7 +104,7 @@ FCollisionStructureManager::NewSimplicial(
 			{
 				if (!OutsideVertices[i].ContainsNaN())
 				{
-					Simplicial->X(i) = OutsideVertices[i];
+					Simplicial->SetX(i, OutsideVertices[i]);
 					VertexCounter++;
 				}
 			}
@@ -114,7 +114,7 @@ FCollisionStructureManager::NewSimplicial(
 		if(!Simplicial->Size())
 		{
 			Simplicial->AddParticles(1);
-			Simplicial->X(0) = Chaos::FVec3(0);
+			Simplicial->SetX(0, Chaos::FVec3(0));
 		}
 
 		Simplicial->UpdateAccelerationStructures();
@@ -140,7 +140,7 @@ FCollisionStructureManager::NewSimplicial(
 		// @todo : Clean collision particles need to operate on the collision mask from the DynamicCollection,
 		//         then transfer only the good collision particles during the initialization. `
 		FCollisionStructureManager::FSimplicial * Simplicial = new FCollisionStructureManager::FSimplicial();
-		const TArrayView<const Chaos::FVec3> ArrayView(&AllParticles.X(0), AllParticles.Size());
+		const TArrayView<const Chaos::FVec3> ArrayView(&AllParticles.GetX(0), AllParticles.Size());
 		const TArray<Chaos::FVec3>& Result = Chaos::CleanCollisionParticles(TriMesh, ArrayView, CollisionParticlesFraction);
 
 		if (Result.Num())
@@ -151,7 +151,7 @@ FCollisionStructureManager::NewSimplicial(
 			{
 				if (!Result[Index].ContainsNaN())
 				{
-					Simplicial->X(Index) = Result[Index];
+					Simplicial->SetX(Index, Result[Index]);
 					VertexCounter++;
 				}
 			}
@@ -161,7 +161,7 @@ FCollisionStructureManager::NewSimplicial(
 		if (!Simplicial->Size())
 		{
 			Simplicial->AddParticles(1);
-			Simplicial->X(0) = Chaos::FVec3(0);
+			Simplicial->SetX(0, Chaos::FVec3(0));
 		}
 
 		Simplicial->UpdateAccelerationStructures();
diff --git a/Engine/Source/Runtime/Experimental/Chaos/Private/PBDRigidsSolver.cpp b/Engine/Source/Runtime/Experimental/Chaos/Private/PBDRigidsSolver.cpp
--- a/Engine/Source/Runtime/Experimental/Chaos/Private/PBDRigidsSolver.cpp
+++ b/Engine/Source/Runtime/Experimental/Chaos/Private/PBDRigidsSolver.cpp
@@ -1905,12 +1905,12 @@ namespace Chaos
 					if(!Handle.IsSleeping())
 					{
 						UE_LOG(LogChaos, Log, TEXT("Particle Dynamic At Position = %s | Velocity = %s | Quaternion = %s | Omega = %s"),
-							*Handle.X().ToString(), *Handle.GetV().ToString(), *Handle.GetR().ToString(), *Handle.GetW().ToString());
+							*Handle.GetX().ToString(), *Handle.GetV().ToString(), *Handle.GetR().ToString(), *Handle.GetW().ToString());
 					}
 					else
 					{
 						UE_LOG(LogChaos, Log, TEXT("Particle Sleeping At Position = %s | Velocity = %s | Quaternion = %s | Omega = %s"), 
-							*Handle.X().ToString(), *Handle.GetV().ToString(), *Handle.GetR().ToString(), *Handle.GetW().ToString());
+							*Handle.GetX().ToString(), *Handle.GetV().ToString(), *Handle.GetR().ToString(), *Handle.GetW().ToString());
 					}
 				}
 			}
@@ -1963,14 +1963,14 @@ namespace Chaos
 										if (!Handle.IsSleeping())
 										{
 											UE_LOG(LogChaos, Log, TEXT("Particle Dynamic At Position = %s | Velocity = %s | Quaternion = %s | Omega = %s | Resim Frame = %d | Sync State = %d | Needs Resim = %d | Unique Idx = %d"),
-												*Handle.X().ToString(), *Handle.GetV().ToString(), *Handle.GetR().ToString(), *Handle.GetW().ToString(),
+												*Handle.GetX().ToString(), *Handle.GetV().ToString(), *Handle.GetR().ToString(), *Handle.GetW().ToString(),
 												GetEvolution()->GetIslandManager().GetParticleIsland(Handle.Handle())->GetResimFrame(), (uint8)Handle.SyncState(),
 												GetEvolution()->GetIslandManager().GetParticleIsland(Handle.Handle())->NeedsResim(), Handle.UniqueIdx().Idx);
 										}
 										else
 										{
 											UE_LOG(LogChaos, Log, TEXT("Particle Sleeping At Position = %s | Velocity = %s | Quaternion = %s | Omega = %s | Resim Frame = %d | Sync State = %d | Needs Resim = %d | Unique Idx = %d"),
-												*Handle.X().ToString(), *Handle.GetV().ToString(), *Handle.GetR().ToString(), *Handle.GetW().ToString(),
+												*Handle.GetX().ToString(), *Handle.GetV().ToString(), *Handle.GetR().ToString(), *Handle.GetW().ToString(),
 												GetEvolution()->GetIslandManager().GetParticleIsland(Handle.Handle())->GetResimFrame(), (uint8)Handle.SyncState(),
 												GetEvolution()->GetIslandManager().GetParticleIsland(Handle.Handle())->NeedsResim(), Handle.UniqueIdx().Idx);
 										}
diff --git a/Engine/Source/Runtime/Experimental/Chaos/Private/PhysicsProxy/ClusterUnionPhysicsProxy.cpp b/Engine/Source/Runtime/Experimental/Chaos/Private/PhysicsProxy/ClusterUnionPhysicsProxy.cpp
--- a/Engine/Source/Runtime/Experimental/Chaos/Private/PhysicsProxy/ClusterUnionPhysicsProxy.cpp
+++ b/Engine/Source/Runtime/Experimental/Chaos/Private/PhysicsProxy/ClusterUnionPhysicsProxy.cpp
@@ -42,7 +42,7 @@ namespace Chaos
 			}
 
 			BufferData.SetProxy(*Proxy);
-			BufferData.X = Particle->X();
+			BufferData.X = Particle->GetX();
 			BufferData.R = Particle->GetR();
 			BufferData.V = Particle->GetV();
 			BufferData.W = Particle->GetW();
@@ -134,7 +134,7 @@ namespace Chaos
 			Particle_Internal->SetR(InitData.InitialTransform.GetRotation());
 			Particle_Internal->SetVf(Chaos::FVec3f(0.f));
 			Particle_Internal->SetWf(Chaos::FVec3f(0.f));
-			Particle_Internal->SetP(Particle_Internal->X());
+			Particle_Internal->SetP(Particle_Internal->GetX());
 			Particle_Internal->SetQf(Particle_Internal->GetRf());
 			Particle_Internal->SetCenterOfMass(FVector3f::ZeroVector);
 			Particle_Internal->SetRotationOfMass(FQuat::Identity);
@@ -380,14 +380,14 @@ namespace Chaos
 						{
 							ParticlesToUpdate.Add(Particle);
 
-							const FRigidTransform3 ChildWorldTM(Particle->X(), Particle->GetR());
+							const FRigidTransform3 ChildWorldTM(Particle->GetX(), Particle->GetR());
 							NewChildToParent.Add(ChildWorldTM.GetRelativeTransform(NewTransform));
 						}
 					}
 				}
 			}
 
-			Evolution.SetParticleTransform(Particle_Internal, NewXR->X(), NewXR->R(), true);
+			Evolution.SetParticleTransform(Particle_Internal, NewXR->GetX(), NewXR->R(), true);
 
 			if (!ParticlesToUpdate.IsEmpty() && !NewChildToParent.IsEmpty())
 			{
@@ -400,7 +400,7 @@ namespace Chaos
 			// but this is primarily for proxies within the cluster union. They (i.e. GCs) need to see this change to the particle's transform.
 			Evolution.GetParticles().MarkTransientDirtyParticle(Particle_Internal);
 
-			const FRigidTransform3 WorldTransform{ Particle_Internal->X(), Particle_Internal->GetR() };
+			const FRigidTransform3 WorldTransform{ Particle_Internal->GetX(), Particle_Internal->GetR() };
 			Particle_Internal->UpdateWorldSpaceState(WorldTransform, FVec3(0));
 
 			Evolution.DirtyParticle(*Particle_Internal);
@@ -536,7 +536,7 @@ namespace Chaos
 			{
 				Chaos::FDebugDrawQueue::GetInstance().DrawDebugBox(NextPullData->X, FVector(2, 1, 1), NextPullData->R, FColor::Yellow, false, 5.f, 0, 0.5f);
 				Chaos::FDebugDrawQueue::GetInstance().DrawDebugDirectionalArrow(PullData.X, NextPullData->X, 0.5f, FColor::Yellow, false, 5.0f, 0, 0.5f);
-				Chaos::FDebugDrawQueue::GetInstance().DrawDebugBox(Particle_External->X(), FVector(2, 1, 1), Particle_External->R(), FColor::Green, false, 5.f, 0, 0.5f);
+				Chaos::FDebugDrawQueue::GetInstance().DrawDebugBox(Particle_External->GetX(), FVector(2, 1, 1), Particle_External->R(), FColor::Green, false, 5.f, 0, 0.5f);
 
 				if (bIsReplicationErrorSmoothing)
 				{
@@ -546,7 +546,7 @@ namespace Chaos
 						Chaos::FDebugDrawQueue::GetInstance().DrawDebugDirectionalArrow(PullData.X, (PullData.X + InterpolationData.GetErrorX(0)), 1, FColor::Red, false, 5.0f, 0, 0.5f);
 					}
 
-					Chaos::FDebugDrawQueue::GetInstance().DrawDebugDirectionalArrow((Particle_External->X() - InterpolationData.GetErrorX(*Alpha)), Particle_External->X(), 1, FColor::Blue, false, 5.0f, 0, 0.5f);
+					Chaos::FDebugDrawQueue::GetInstance().DrawDebugDirectionalArrow((Particle_External->GetX() - InterpolationData.GetErrorX(*Alpha)), Particle_External->GetX(), 1, FColor::Blue, false, 5.0f, 0, 0.5f);
 				}
 			}
 #endif // CHAOS_DEBUG_DRAW
diff --git a/Engine/Source/Runtime/Experimental/Chaos/Private/PhysicsProxy/GeometryCollectionPhysicsProxy.cpp b/Engine/Source/Runtime/Experimental/Chaos/Private/PhysicsProxy/GeometryCollectionPhysicsProxy.cpp
--- a/Engine/Source/Runtime/Experimental/Chaos/Private/PhysicsProxy/GeometryCollectionPhysicsProxy.cpp
+++ b/Engine/Source/Runtime/Experimental/Chaos/Private/PhysicsProxy/GeometryCollectionPhysicsProxy.cpp
@@ -337,7 +337,7 @@ void SetImplicitToPTParticles(Chaos::TPBDRigidParticleHandle<Chaos::FReal, 3>* H
 			Handle->SetLocalBounds(SharedImplicitTS->BoundingBox());
 		}
 		Handle->SetHasBounds(true);
-		const Chaos::FRigidTransform3 Xf(Handle->X(), Handle->GetR());
+		const Chaos::FRigidTransform3 Xf(Handle->GetX(), Handle->GetR());
 		Handle->UpdateWorldSpaceState(Xf, Chaos::FVec3(0));
 	}
 
@@ -363,18 +363,18 @@ void SetImplicitToPTParticles(Chaos::TPBDRigidParticleHandle<Chaos::FReal, 3>* H
 		// we need to account for scale and check if the particle is still within its domain
 		for (int32 ParticleIndex = 0; ParticleIndex < (int32)CollisionParticles->Size(); ++ParticleIndex)
 		{
-			CollisionParticles->X(ParticleIndex) *= Scale;
+			CollisionParticles->SetX(ParticleIndex, CollisionParticles->GetX(ParticleIndex) * Scale);
 
 			// Make sure the collision particles are at least in the domain 
 			// of the implicit shape.
-			ensure(ImplicitShapeDomain.Contains(CollisionParticles->X(ParticleIndex)));
+			ensure(ImplicitShapeDomain.Contains(CollisionParticles->GetX(ParticleIndex)));
 		}
 
 		// @todo(remove): IF there is no simplicial we should not be forcing one. 
 		if (!CollisionParticles->Size())
 		{
 			CollisionParticles->AddParticles(1);
-			CollisionParticles->X(0) = Chaos::FVec3(0);
+			CollisionParticles->SetX(0, Chaos::FVec3(0));
 		}
 		CollisionParticles->UpdateAccelerationStructures();
 	}
@@ -426,7 +426,7 @@ void PopulateSimulatedParticle(
 	Handle->SetVf(Chaos::FVec3f(0.f));
 	Handle->SetR(WorldTransform.GetRotation().GetNormalized());
 	Handle->SetWf(Chaos::FVec3f(0.f));
-	Handle->SetP(Handle->X());
+	Handle->SetP(Handle->GetX());
 	Handle->SetQf(Handle->GetRf());
 	Handle->SetCenterOfMass(FVector3f::ZeroVector);
 	Handle->SetRotationOfMass(FQuat::Identity);
@@ -1998,7 +1998,7 @@ Chaos::TPBDGeometryCollectionParticleHandle<Chaos::FReal, 3>* FGeometryCollectio
 	{
 		Handle->SetHasBounds(true);
 		Handle->SetLocalBounds(Handle->GetGeometry()->BoundingBox());
-		const Chaos::FRigidTransform3 Xf(Handle->X(), Handle->GetR());
+		const Chaos::FRigidTransform3 Xf(Handle->GetX(), Handle->GetR());
 		Handle->UpdateWorldSpaceState(Xf, Chaos::FVec3(0));
 
 		static_cast<Chaos::FPBDRigidsSolver*>(Solver)->GetEvolution()->DirtyParticle(*Handle);
@@ -2147,7 +2147,7 @@ FGeometryCollectionPhysicsProxy::BuildClusters_Internal(
 	{
 		Parent->SetHasBounds(true);
 		Parent->SetLocalBounds(Parent->GetGeometry()->BoundingBox());
-		const Chaos::FRigidTransform3 Xf(Parent->X(), Parent->GetR());
+		const Chaos::FRigidTransform3 Xf(Parent->GetX(), Parent->GetR());
 		Parent->UpdateWorldSpaceState(Xf, Chaos::FVec3(0));
 
 		static_cast<Chaos::FPBDRigidsSolver*>(Solver)->GetEvolution()->DirtyParticle(*Parent);
@@ -2386,7 +2386,7 @@ void FGeometryCollectionPhysicsProxy::ApplyForceAt_External(FVector Force, FVect
 					{
 						if (ClusteredHandle->IsDynamic())
 						{
-							const Chaos::FReal DistanceSquared = (WorldLocation - ClusteredHandle->X()).SquaredLength();
+							const Chaos::FReal DistanceSquared = (WorldLocation - ClusteredHandle->GetX()).SquaredLength();
 							if (DistanceSquared < ClosestDistanceSquared)
 							{
 								ClosestDistanceSquared = DistanceSquared;
@@ -2424,7 +2424,7 @@ void FGeometryCollectionPhysicsProxy::ApplyImpulseAt_External(FVector Impulse, F
 					{
 						if (ClusteredHandle->IsDynamic())
 						{
-							const Chaos::FReal DistanceSquared = (WorldLocation - ClusteredHandle->X()).SquaredLength();
+							const Chaos::FReal DistanceSquared = (WorldLocation - ClusteredHandle->GetX()).SquaredLength();
 							if (DistanceSquared < ClosestDistanceSquared)
 							{
 								ClosestDistanceSquared = DistanceSquared;
@@ -2552,7 +2552,7 @@ void FGeometryCollectionPhysicsProxy::SetAnchoredByTransformedBox_External(const
 				{
 					if (FClusterHandle* ParticleHandle = SolverParticleHandles[ParticleIndex])
 					{
-						const FVec3 PositionInBoxSpace = BoxTransform.InverseTransformPosition(ParticleHandle->X());
+						const FVec3 PositionInBoxSpace = BoxTransform.InverseTransformPosition(ParticleHandle->GetX());
 						if (BoundsToCheck.Contains(PositionInBoxSpace))
 						{
 							// if we have a max level , we make sure to anchor a parent of the right level 
@@ -3306,7 +3306,7 @@ void FGeometryCollectionPhysicsProxy::SetWorldTransform_Internal(const FTransfor
 
 				if (KinematicRootHandle)
 				{
-					const FTransform RootWorldTransform(KinematicRootHandle->GetR(), KinematicRootHandle->X());
+					const FTransform RootWorldTransform(KinematicRootHandle->GetR(), KinematicRootHandle->GetX());
 					const FTransform RootRelativeTransform = RootWorldTransform.GetRelativeTransform(Parameters.PrevWorldTransform);
 					const FTransform WorldTransform = RootRelativeTransform * ActorToWorld;
 
@@ -3322,7 +3322,7 @@ void FGeometryCollectionPhysicsProxy::SetWorldTransform_Internal(const FTransfor
 					if (ClusterUnionIndex != INDEX_NONE)
 					{
 						const int32 TransformGroupIndex = FromParticleToTransformIndex[ParticleIndex];
-						const FTransform ParentWorldTransform{ ParentHandle->GetR(), ParentHandle->X() };
+						const FTransform ParentWorldTransform{ ParentHandle->GetR(), ParentHandle->GetX() };
 						const FTransform NewWorldTransform = MassToLocal[TransformGroupIndex] * FTransform(PhysicsThreadCollection.GetTransform(TransformGroupIndex)) * Parameters.WorldTransform;
 						const FTransform RelativeTransform = NewWorldTransform.GetRelativeTransform(ParentWorldTransform);
 
@@ -3709,7 +3709,7 @@ void FGeometryCollectionPhysicsProxy::BufferPhysicsResults_External(Chaos::FDirt
 
 static void UpdateParticleHandleTransform(Chaos::FPBDRigidsSolver& CurrentSolver, Chaos::FPBDRigidClusteredParticleHandle& Handle, const Chaos::FRigidTransform3& NewTransform)
 {
-	Handle.X() = NewTransform.GetTranslation();
+	Handle.SetX(NewTransform.GetTranslation());
 	Handle.SetR(NewTransform.GetRotation());
 	Handle.UpdateWorldSpaceState(NewTransform, Chaos::FVec3{0});
 	CurrentSolver.GetEvolution()->DirtyParticle(Handle);
@@ -3718,7 +3718,7 @@ static void UpdateParticleHandleTransform(Chaos::FPBDRigidsSolver& CurrentSolver
 static void UpdateParticleHandleTransformIfNeeded(Chaos::FPBDRigidsSolver& CurrentSolver, Chaos::FPBDRigidClusteredParticleHandle& Handle, const Chaos::FRigidTransform3& NewTransform)
 {
 	const Chaos::FVec3 NewX = NewTransform.GetTranslation();
-	const Chaos::FVec3 OldX = Handle.X();
+	const Chaos::FVec3 OldX = Handle.GetX();
 	
 	const Chaos::FRotation3 NewR  = NewTransform.GetRotation();
 	const Chaos::FRotation3 OldR = Handle.GetR();
@@ -3801,7 +3801,7 @@ void FGeometryCollectionPhysicsProxy::BufferPhysicsResults_Internal(Chaos::FPBDR
 			const bool bIsActive = !Handle->Disabled();
 
 			FGeometryCollectionResults::FPositionData PositionData;
-			PositionData.ParticleX = Handle->X();
+			PositionData.ParticleX = Handle->GetX();
 			PositionData.ParticleR = Handle->GetR();
 
 			FGeometryCollectionResults::FVelocityData VelocityData;
@@ -3875,7 +3875,7 @@ void FGeometryCollectionPhysicsProxy::BufferPhysicsResults_Internal(Chaos::FPBDR
 
 				// child of internal clusters may not have their position up to date 
 				// so we need to compute it from the ChildToParent transform 
-				const FTransform ParticleToWorld = Handle->ChildToParent() * FRigidTransform3(ClusterParent->X(), ClusterParent->GetR()); // aka ClusterChildToWorld
+				const FTransform ParticleToWorld = Handle->ChildToParent() * FRigidTransform3(ClusterParent->GetX(), ClusterParent->GetR()); // aka ClusterChildToWorld
 				PositionData.ParticleX = ParticleToWorld.GetTranslation();
 				PositionData.ParticleR = ParticleToWorld.GetRotation();
 				
@@ -3895,7 +3895,7 @@ void FGeometryCollectionPhysicsProxy::BufferPhysicsResults_Internal(Chaos::FPBDR
 			{
 				if (ClusterParent && !ClusterParent->Disabled())
 				{
-					const FRigidTransform3 ChildToWorld = Handle->ChildToParent() * FRigidTransform3(ClusterParent->X(), ClusterParent->GetR());
+					const FRigidTransform3 ChildToWorld = Handle->ChildToParent() * FRigidTransform3(ClusterParent->GetX(), ClusterParent->GetR());
 					UpdateParticleHandleTransformIfNeeded(*CurrentSolver, *Handle, ChildToWorld);
 					// fields may have applied velocities, we need to make sure to clear that up, so that we don't accumulate
 					VelocityData.ParticleV = FVec3::ZeroVector;
@@ -3975,7 +3975,7 @@ static inline bool UpdateGTParticleXR(Chaos::FPBDRigidParticle& GTParticle, cons
 {
 	GC_PHYSICSPROXY_CHECK_FOR_NAN(NewX);
 
-	const Chaos::FVec3 OldX = GTParticle.X();
+	const Chaos::FVec3 OldX = GTParticle.GetX();
 	const bool bNeedUpdateX = (!NewX.Equals(OldX, GeometryCollectionPositionUpdateTolerance));
 	if (bNeedUpdateX)
 	{
@@ -5135,7 +5135,7 @@ void FGeometryCollectionPhysicsProxy::InitializeSharedCollisionStructures(
 	MassSpaceParticles.AddParticles(Vertex.Num());
 	for (int32 Idx = 0; Idx < Vertex.Num(); ++Idx)
 	{
-		MassSpaceParticles.X(Idx) = Vertex[Idx];	//mass space computation done later down
+		MassSpaceParticles.SetX(Idx, Vertex[Idx]);	//mass space computation done later down
 	}
 	
 	
@@ -5157,7 +5157,7 @@ void FGeometryCollectionPhysicsProxy::InitializeSharedCollisionStructures(
 				const int32 IdxEnd = IdxStart + VertexCount[GeometryIndex];
 				for (int32 Idx = IdxStart; Idx < IdxEnd; ++Idx)
 				{
-					MassSpaceParticles.X(Idx) = MassToLocalTransform.InverseTransformPositionNoScale(MassSpaceParticles.X(Idx));
+					MassSpaceParticles.SetX(Idx, MassToLocalTransform.InverseTransformPositionNoScale(MassSpaceParticles.GetX(Idx)));
 				}
 			}
 
@@ -5169,7 +5169,7 @@ void FGeometryCollectionPhysicsProxy::InitializeSharedCollisionStructures(
 				const TSet<int32> MeshVertices = TriMesh->GetVertices();
 				for (const int32 Idx : MeshVertices)
 				{
-					InstanceBoundingBox += MassSpaceParticles.X(Idx);
+					InstanceBoundingBox += MassSpaceParticles.GetX(Idx);
 				}
 			}
 			else if(VertexCount[GeometryIndex])
@@ -5178,7 +5178,7 @@ void FGeometryCollectionPhysicsProxy::InitializeSharedCollisionStructures(
 				const int32 IdxEnd = IdxStart + VertexCount[GeometryIndex];
 				for (int32 Idx = IdxStart; Idx < IdxEnd; ++Idx)
 				{
-					InstanceBoundingBox += MassSpaceParticles.X(Idx);
+					InstanceBoundingBox += MassSpaceParticles.GetX(Idx);
 				}
 			}
 			else
@@ -5269,7 +5269,7 @@ void FGeometryCollectionPhysicsProxy::InitializeSharedCollisionStructures(
 		// Init to -FLT_MAX for debugging purposes
 		for (int32 Idx = 0; Idx < NumTransforms; Idx++)
 		{
-			CollectionSpaceParticles->X(Idx) = Chaos::FVec3(-TNumericLimits<FReal>::Max());
+			CollectionSpaceParticles->SetX(Idx, Chaos::FVec3(-TNumericLimits<FReal>::Max()));
 		}
 
 		//
@@ -5418,9 +5418,9 @@ void FGeometryCollectionPhysicsProxy::InitializeSharedCollisionStructures(
 							for (const int32 VertIdx : VertsAdded)
 							{
 								//Update particles so they are in the cluster's mass space
-								MassSpaceParticles.X(VertIdx) =
-									ChildMassToClusterMass.TransformPosition(MassSpaceParticles.X(VertIdx));
-								InstanceBoundingBox += MassSpaceParticles.X(VertIdx);
+								MassSpaceParticles.SetX(VertIdx, 
+									ChildMassToClusterMass.TransformPosition(MassSpaceParticles.GetX(VertIdx)));
+								InstanceBoundingBox += MassSpaceParticles.GetX(VertIdx);
 							}
 
 							// Reset vert set for next child so we don't repeatedly transform verts and over extend our calculated bounds
diff --git a/Engine/Source/Runtime/Experimental/Chaos/Private/PhysicsProxy/SingleParticlePhysicsProxy.cpp b/Engine/Source/Runtime/Experimental/Chaos/Private/PhysicsProxy/SingleParticlePhysicsProxy.cpp
--- a/Engine/Source/Runtime/Experimental/Chaos/Private/PhysicsProxy/SingleParticlePhysicsProxy.cpp
+++ b/Engine/Source/Runtime/Experimental/Chaos/Private/PhysicsProxy/SingleParticlePhysicsProxy.cpp
@@ -104,7 +104,7 @@ void PushToPhysicsStateImp(const Chaos::FDirtyPropertiesManager& Manager, Chaos:
 			// Update world-space cached state like the bounds
 			// @todo(chaos): do we need to do this here? It should be done in Integrate and ApplyKinematicTarget so only really Statics need this...
 			const bool bHasKinematicTarget = (NewKinematicTargetGT != nullptr) && (NewKinematicTargetGT->GetMode() == EKinematicTargetMode::Position);
-			const FRigidTransform3 WorldTransform = !bHasKinematicTarget ? FRigidTransform3(Handle->X(), Handle->GetR()) : NewKinematicTargetGT->GetTarget();
+			const FRigidTransform3 WorldTransform = !bHasKinematicTarget ? FRigidTransform3(Handle->GetX(), Handle->GetR()) : NewKinematicTargetGT->GetTarget();
 			Handle->UpdateWorldSpaceState(WorldTransform, FVec3(0));
 
 			Evolution.DirtyParticle(*Handle);
@@ -226,7 +226,7 @@ void FSingleParticlePhysicsProxy::ClearAccumulatedData()
 template <typename T>
 void BufferPhysicsResultsImp(Chaos::FDirtyRigidParticleData& PullData, T* Particle)
 {
-	PullData.X = Particle->X();
+	PullData.X = Particle->GetX();
 	PullData.R = Particle->GetR();
 	PullData.V = Particle->GetV();
 	PullData.W = Particle->GetW();
