Make consistent API by using GetX, SetX, GetP and SetP for physics particles

diff --git a/Engine/Source/Runtime/Experimental/Chaos/Private/RewindData.cpp b/Engine/Source/Runtime/Experimental/Chaos/Private/RewindData.cpp
--- a/Engine/Source/Runtime/Experimental/Chaos/Private/RewindData.cpp
+++ b/Engine/Source/Runtime/Experimental/Chaos/Private/RewindData.cpp
@@ -255,7 +255,7 @@ void FRewindData::ApplyTargets(const int32 Frame, const bool bResetSimulation)
 #if !(UE_BUILD_SHIPPING || UE_BUILD_TEST)
 		if (!History.TargetPositions.IsClean(RewindFrameAndPhase) && Chaos::FPhysicsSolverBase::CanDebugNetworkPhysicsPrediction())
 		{
-			UE_LOG(LogChaos, Log, TEXT("Reset particle %d position to the target %s at frame %d"), PTParticle->UniqueIdx().Idx, *PTParticle->X().ToString(), Frame);
+			UE_LOG(LogChaos, Log, TEXT("Reset particle %d position to the target %s at frame %d"), PTParticle->UniqueIdx().Idx, *PTParticle->GetX().ToString(), Frame);
 		}
 #endif
 	}
@@ -417,7 +417,7 @@ void FRewindData::AccumulateErrorIfNecessary(FGeometryParticleHandle& Obj, const
 	}
 
 	// Get the error offset after a correction
-	const FVec3 ErrorX = DirtyInfo->GetHistory().PreCorrectionXR.X() - Obj.X();
+	const FVec3 ErrorX = DirtyInfo->GetHistory().PreCorrectionXR.X() - Obj.GetX();
 	FQuat ErrorR = DirtyInfo->GetHistory().PreCorrectionXR.R() * Obj.GetR().Inverse();
 	ErrorR.EnforceShortestArcWith(FQuat::Identity);
 	ErrorR.Normalize();
diff --git a/Engine/Source/Runtime/Experimental/Chaos/Public/Chaos/BoundingVolumeUtilities.h b/Engine/Source/Runtime/Experimental/Chaos/Public/Chaos/BoundingVolumeUtilities.h
--- a/Engine/Source/Runtime/Experimental/Chaos/Public/Chaos/BoundingVolumeUtilities.h
+++ b/Engine/Source/Runtime/Experimental/Chaos/Public/Chaos/BoundingVolumeUtilities.h
@@ -167,14 +167,14 @@ template<class T, int d>
 TAABB<T, d> ComputeWorldSpaceBoundingBox(const TParticles<T, d>& Objects, const int32 i, bool bUseVelocity = false, T Dt = 0)
 {
 	ensure(!bUseVelocity);
-	return TAABB<T, d>(Objects.X(i), Objects.X(i));
+	return TAABB<T, d>(Objects.GetX(i), Objects.GetX(i));
 }
 
 template<class T, int d>
 TAABB<T, d> ComputeWorldSpaceBoundingBox(const TGeometryParticles<T, d>& Objects, const int32 i, bool bUseVelocity = false, T Dt = 0)
 {
 	ensure(!bUseVelocity);
-	TRigidTransform<T, d> LocalToWorld(Objects.X(i), Objects.GetR(i));
+	TRigidTransform<T, d> LocalToWorld(Objects.GetX(i), Objects.GetR(i));
 	const auto& LocalBoundingBox = Objects.GetGeometry(i)->BoundingBox();
 	return LocalBoundingBox.TransformedAABB(LocalToWorld);
 }
@@ -182,7 +182,7 @@ TAABB<T, d> ComputeWorldSpaceBoundingBox(const TGeometryParticles<T, d>& Objects
 template<class T, int d>
 TAABB<T, d> ComputeWorldSpaceBoundingBox(const TPBDRigidParticles<T, d>& Objects, const int32 i, bool bUseVelocity = false, T Dt = 0)
 {
-	TRigidTransform<T, d> LocalToWorld(Objects.P(i), Objects.GetQ(i));
+	TRigidTransform<T, d> LocalToWorld(Objects.GetP(i), Objects.GetQ(i));
 	TAABB<T, d> WorldSpaceBox;
 	if (Objects.GetGeometry(i))
 	{
@@ -192,10 +192,10 @@ TAABB<T, d> ComputeWorldSpaceBoundingBox(const TPBDRigidParticles<T, d>& Objects
 	else
 	{
 		check(Objects.CollisionParticles(i) && Objects.CollisionParticles(i)->Size());
-		TAABB<T, d> LocalBoundingBox(Objects.CollisionParticles(i)->X(0), Objects.CollisionParticles(i)->X(0));
+		TAABB<T, d> LocalBoundingBox(Objects.CollisionParticles(i)->GetX(0), Objects.CollisionParticles(i)->GetX(0));
 		for (uint32 j = 1; j < Objects.CollisionParticles(i)->Size(); ++j)
 		{
-			LocalBoundingBox.GrowToInclude(Objects.CollisionParticles(i)->X(j));
+			LocalBoundingBox.GrowToInclude(Objects.CollisionParticles(i)->GetX(j));
 		}
 		WorldSpaceBox = LocalBoundingBox.TransformedAABB(LocalToWorld);
 	}
diff --git a/Engine/Source/Runtime/Experimental/Chaos/Public/Chaos/Convex.h b/Engine/Source/Runtime/Experimental/Chaos/Public/Chaos/Convex.h
--- a/Engine/Source/Runtime/Experimental/Chaos/Public/Chaos/Convex.h
+++ b/Engine/Source/Runtime/Experimental/Chaos/Public/Chaos/Convex.h
@@ -929,7 +929,7 @@ namespace Chaos
 				Vertices.SetNum(NumVertices);
 				for (int32 VertexIndex = 0; VertexIndex < NumVertices; ++VertexIndex)
 				{
-					Vertices[VertexIndex] = TmpSurfaceParticles.X(VertexIndex);
+					Vertices[VertexIndex] = TmpSurfaceParticles.GetX(VertexIndex);
 				}
 			}
 			else
diff --git a/Engine/Source/Runtime/Experimental/Chaos/Public/Chaos/Deformable/GaussSeidelCorotatedCodimensionalConstraints.h b/Engine/Source/Runtime/Experimental/Chaos/Public/Chaos/Deformable/GaussSeidelCorotatedCodimensionalConstraints.h
--- a/Engine/Source/Runtime/Experimental/Chaos/Public/Chaos/Deformable/GaussSeidelCorotatedCodimensionalConstraints.h
+++ b/Engine/Source/Runtime/Experimental/Chaos/Public/Chaos/Deformable/GaussSeidelCorotatedCodimensionalConstraints.h
@@ -134,8 +134,8 @@ namespace Chaos::Softs
 					&& MeshConstraints[e][1] < (int32)Particles.Size() && MeshConstraints[e][1] > INDEX_NONE
 					&& MeshConstraints[e][2] < (int32)Particles.Size() && MeshConstraints[e][2] > INDEX_NONE)
 				{
-					const TVec3<T> X1X0 = Particles.X(MeshConstraints[e][1]) - Particles.X(MeshConstraints[e][0]);
-					const TVec3<T> X2X0 = Particles.X(MeshConstraints[e][2]) - Particles.X(MeshConstraints[e][0]);
+					const TVec3<T> X1X0 = Particles.GetX(MeshConstraints[e][1]) - Particles.GetX(MeshConstraints[e][0]);
+					const TVec3<T> X2X0 = Particles.GetX(MeshConstraints[e][2]) - Particles.GetX(MeshConstraints[e][0]);
 					PMatrix<T, 2, 2> Dm((T)0., (T)0., (T)0.);
 					Dm.M[0] = X1X0.Size();
 					Dm.M[2] = X1X0.Dot(X2X0) / Dm.M[0];
diff --git a/Engine/Source/Runtime/Experimental/Chaos/Public/Chaos/Deformable/GaussSeidelWeakConstraints.h b/Engine/Source/Runtime/Experimental/Chaos/Public/Chaos/Deformable/GaussSeidelWeakConstraints.h
--- a/Engine/Source/Runtime/Experimental/Chaos/Public/Chaos/Deformable/GaussSeidelWeakConstraints.h
+++ b/Engine/Source/Runtime/Experimental/Chaos/Public/Chaos/Deformable/GaussSeidelWeakConstraints.h
@@ -505,7 +505,7 @@ namespace Chaos::Softs
 					const int32 Index = BoundaryVertices[i];
 
 					TArray< TTriangleCollisionPoint<FSolverReal> > Result;
-					if (TriangleMesh.PointProximityQuery(Spatial, static_cast<const TArrayView<const FSolverVec3>&>(Particles.XArray()), Index, Particles.X(Index), DetectRadius, DetectRadius,
+					if (TriangleMesh.PointProximityQuery(Spatial, static_cast<const TArrayView<const FSolverVec3>&>(Particles.XArray()), Index, Particles.GetX(Index), DetectRadius, DetectRadius,
 						[this, &ComponentIndex, &Elements](const int32 PointIndex, const int32 TriangleIndex)->bool
 						{
 							const TVector<int32, 3>& Elem = Elements[TriangleIndex];
@@ -538,8 +538,8 @@ namespace Chaos::Softs
 						
 
 							// NOTE: CollisionPoint.Normal has already been flipped to point toward the Point, so need to recalculate here.
-							const TTriangle<FSolverReal> Triangle(Particles.X(Elem[0]), Particles.X(Elem[1]), Particles.X(Elem[2]));
-							if ((Particles.X(Index) - CollisionPoint.Location).Dot(-Triangle.GetNormal()) < 0) //Is point inside boundary? Normal should point outwards
+							const TTriangle<FSolverReal> Triangle(Particles.GetX(Elem[0]), Particles.GetX(Elem[1]), Particles.GetX(Elem[2]));
+							if ((Particles.GetX(Index) - CollisionPoint.Location).Dot(-Triangle.GetNormal()) < 0) //Is point inside boundary? Normal should point outwards
 							{
 								const int32 IndexToWrite = ConstraintIndex.fetch_add(1);
 
diff --git a/Engine/Source/Runtime/Experimental/Chaos/Public/Chaos/EvolutionResimCache.h b/Engine/Source/Runtime/Experimental/Chaos/Public/Chaos/EvolutionResimCache.h
--- a/Engine/Source/Runtime/Experimental/Chaos/Public/Chaos/EvolutionResimCache.h
+++ b/Engine/Source/Runtime/Experimental/Chaos/Public/Chaos/EvolutionResimCache.h
@@ -44,7 +44,7 @@ namespace Chaos
 		void SaveParticlePostSolve(const FPBDRigidParticleHandle& Particle)
 		{
 			FPBDSolveCache& Cache = ParticleToCachedSolve.FindOrAdd(Particle.UniqueIdx());
-			Cache.P = Particle.P();
+			Cache.P = Particle.GetP();
 			Cache.Q = Particle.GetQ();
 			Cache.V = Particle.GetV();
 			Cache.W = Particle.GetW();
@@ -56,7 +56,7 @@ namespace Chaos
 			const FPBDSolveCache* Cache = ParticleToCachedSolve.Find(Particle.UniqueIdx());
 			if (Cache)
 			{
-				Particle.P() = Cache->P;
+				Particle.SetP(Cache->P);
 				Particle.SetQ(Cache->Q);
 				Particle.SetV(Cache->V);
 				Particle.SetW(Cache->W);
diff --git a/Engine/Source/Runtime/Experimental/Chaos/Public/Chaos/GeometryParticles.h b/Engine/Source/Runtime/Experimental/Chaos/Public/Chaos/GeometryParticles.h
--- a/Engine/Source/Runtime/Experimental/Chaos/Public/Chaos/GeometryParticles.h
+++ b/Engine/Source/Runtime/Experimental/Chaos/Public/Chaos/GeometryParticles.h
@@ -154,7 +154,7 @@ namespace Chaos
 	public:
 
 		using TArrayCollection::Size;
-		using TParticles<T,d>::X;
+		using TParticles<T,d>::GetX;
 		using TSimpleGeometryParticles<T, d>::GetR;
 		using TSimpleGeometryParticles<T, d>::GetGeometry;
 		using TSimpleGeometryParticles<T, d>::SetGeometry;
@@ -320,7 +320,7 @@ namespace Chaos
 				// Update the world-space stat of all the shapes - must be called after UpdateShapesArray
 				// world space inflated bounds needs to take expansion into account - this is done in integrate for dynamics anyway, so
 				// this computation is mainly for statics
-				UpdateWorldSpaceState(Index, TRigidTransform<FReal, 3>(X(Index), GetR(Index)), FVec3(0));
+				UpdateWorldSpaceState(Index, TRigidTransform<FReal, 3>(GetX(Index), GetR(Index)), FVec3(0));
 			}
 		}
 
@@ -587,7 +587,7 @@ public:
 				{
 					for (int32 Idx = 0; Idx < MShapesArray.Num(); ++Idx)
 					{
-						UpdateWorldSpaceState(Idx, FRigidTransform3(X(Idx), GetR(Idx)), FVec3(0));
+						UpdateWorldSpaceState(Idx, FRigidTransform3(GetX(Idx), GetR(Idx)), FVec3(0));
 					}
 				}
 			}
@@ -601,7 +601,7 @@ public:
 					{
 						MLocalBounds[Idx] = TAABB<T, d>(GetGeometry(Idx)->BoundingBox());
 						//ignore velocity too, really just trying to get something reasonable)
-						UpdateWorldSpaceState(Idx, FRigidTransform3(X(Idx), GetR(Idx)), FVec3(0));
+						UpdateWorldSpaceState(Idx, FRigidTransform3(GetX(Idx), GetR(Idx)), FVec3(0));
 					}
 				}
 			}
diff --git a/Engine/Source/Runtime/Experimental/Chaos/Public/Chaos/MPMTransfer.h b/Engine/Source/Runtime/Experimental/Chaos/Public/Chaos/MPMTransfer.h
--- a/Engine/Source/Runtime/Experimental/Chaos/Public/Chaos/MPMTransfer.h
+++ b/Engine/Source/Runtime/Experimental/Chaos/Public/Chaos/MPMTransfer.h
@@ -75,7 +75,7 @@ public:
 			//use physics parallel for:
 			for (int32 p = 0; p < N; p++)
 			{
-				Grid.BaseNodeIndex(InParticles.X(p), Indices[p], Weights[p]);
+				Grid.BaseNodeIndex(InParticles.GetX(p), Indices[p], Weights[p]);
 			}
 
 		}
@@ -173,7 +173,7 @@ public:
 
 	void APICP2G(const int32 p, const int32 GlobIndex, const T mip, const TDynamicParticles<T, 3>& Particles, TArray<T>& GridData)
 	{
-		TVector<T, 3> xi_minus_xp = Grid.Node(GlobIndex) - Particles.X(p);
+		TVector<T, 3> xi_minus_xp = Grid.Node(GlobIndex) - Particles.GetX(p);
 		for (int32 l = 0; l < 3; ++l) 
 		{
 			T vl = T(0);
@@ -337,7 +337,7 @@ public:
 								}
 								else
 								{
-									TVector<T, 3> rip = Grid.Node(GlobIndexFlat) - Particles.X(p);
+									TVector<T, 3> rip = Grid.Node(GlobIndexFlat) - Particles.GetX(p);
 									T NiProdbInv = (T(4) * Ni[0] * Ni[1] * Ni[2]) / (Grid.GetDx()[0] * Grid.GetDx()[0]);
 									for (int32 l = 0; l < 3; l++)
 									{
diff --git a/Engine/Source/Runtime/Experimental/Chaos/Public/Chaos/NewtonElasticFEM.h b/Engine/Source/Runtime/Experimental/Chaos/Public/Chaos/NewtonElasticFEM.h
--- a/Engine/Source/Runtime/Experimental/Chaos/Public/Chaos/NewtonElasticFEM.h
+++ b/Engine/Source/Runtime/Experimental/Chaos/Public/Chaos/NewtonElasticFEM.h
@@ -76,7 +76,7 @@ public:
 		PMatrix<T, 3, 3> Result((T)0.);
 		for (int i = 0; i < 3; i++) {
 			for (int c = 0; c < 3; c++) {
-				Result.SetAt(c, i, InParticles.X(Mesh[e][i + 1])[c] - InParticles.X(Mesh[e][0])[c]);
+				Result.SetAt(c, i, InParticles.GetX(Mesh[e][i + 1])[c] - InParticles.GetX(Mesh[e][0])[c]);
 			}
 		}
 		return Result;
diff --git a/Engine/Source/Runtime/Experimental/Chaos/Public/Chaos/PBDAltitudeSpringConstraints.h b/Engine/Source/Runtime/Experimental/Chaos/Public/Chaos/PBDAltitudeSpringConstraints.h
--- a/Engine/Source/Runtime/Experimental/Chaos/Public/Chaos/PBDAltitudeSpringConstraints.h
+++ b/Engine/Source/Runtime/Experimental/Chaos/Public/Chaos/PBDAltitudeSpringConstraints.h
@@ -25,17 +25,17 @@ namespace Chaos::Softs
 				TVec4<int32> Constraint = Constraints[ElementIndex];
 				for (int32 i = 0; i < 4; i++)
 				{
-					const FSolverVec3& P1 = InParticles.X(Constraint[i]);
-					const FSolverVec3& P2 = InParticles.X(Constraint[(i + 1) % 4]);
-					const FSolverVec3& P3 = InParticles.X(Constraint[(i + 2) % 4]);
-					const FSolverVec3& P4 = InParticles.X(Constraint[(i + 3) % 4]);
+					const FSolverVec3& P1 = InParticles.GetX(Constraint[i]);
+					const FSolverVec3& P2 = InParticles.GetX(Constraint[(i + 1) % 4]);
+					const FSolverVec3& P3 = InParticles.GetX(Constraint[(i + 2) % 4]);
+					const FSolverVec3& P4 = InParticles.GetX(Constraint[(i + 3) % 4]);
 					RestLength[4* ElementIndex +i] = FSolverVec3::DotProduct(FSolverVec3::CrossProduct(P3 - P2, P4 - P2), P1 - P2) / FSolverVec3::CrossProduct(P3 - P2, P4 - P2).Size();
 
 				}
-				const FSolverVec3& P1 = InParticles.X(Constraint[0]);
-				const FSolverVec3& P2 = InParticles.X(Constraint[1]);
-				const FSolverVec3& P3 = InParticles.X(Constraint[2]);
-				const FSolverVec3& P4 = InParticles.X(Constraint[3]);
+				const FSolverVec3& P1 = InParticles.GetX(Constraint[0]);
+				const FSolverVec3& P2 = InParticles.GetX(Constraint[1]);
+				const FSolverVec3& P3 = InParticles.GetX(Constraint[2]);
+				const FSolverVec3& P4 = InParticles.GetX(Constraint[3]);
 				Volumes[ElementIndex] = FSolverVec3::DotProduct(FSolverVec3::CrossProduct(P2 - P1, P3 - P1), P4 - P1) / (FSolverReal)6.;
 			}
 			
diff --git a/Engine/Source/Runtime/Experimental/Chaos/Public/Chaos/PBDAnimDriveConstraint.h b/Engine/Source/Runtime/Experimental/Chaos/Public/Chaos/PBDAnimDriveConstraint.h
--- a/Engine/Source/Runtime/Experimental/Chaos/Public/Chaos/PBDAnimDriveConstraint.h
+++ b/Engine/Source/Runtime/Experimental/Chaos/Public/Chaos/PBDAnimDriveConstraint.h
@@ -236,7 +236,7 @@ namespace Chaos::Softs
 
 			ParticlePosition -= InStiffness * (ParticlePosition - AnimPosition);
 
-			const FSolverVec3 ParticleDisplacement = ParticlePosition - Particles.X(ParticleIndex);
+			const FSolverVec3 ParticleDisplacement = ParticlePosition - Particles.GetX(ParticleIndex);
 			const FSolverVec3 AnimationDisplacement = (AnimVelocity) * Dt;
 			const FSolverVec3 RelativeDisplacement = ParticleDisplacement - AnimationDisplacement;
 
diff --git a/Engine/Source/Runtime/Experimental/Chaos/Public/Chaos/PBDAxialSpringConstraintsBase.h b/Engine/Source/Runtime/Experimental/Chaos/Public/Chaos/PBDAxialSpringConstraintsBase.h
--- a/Engine/Source/Runtime/Experimental/Chaos/Public/Chaos/PBDAxialSpringConstraintsBase.h
+++ b/Engine/Source/Runtime/Experimental/Chaos/Public/Chaos/PBDAxialSpringConstraintsBase.h
@@ -115,9 +115,9 @@ private:
 	template<typename SolverParticlesOrRange>
 	FSolverReal FindBary(const SolverParticlesOrRange& Particles, const int32 i1, const int32 i2, const int32 i3)
 	{
-		const FSolverVec3& P1 = Particles.X(i1);
-		const FSolverVec3& P2 = Particles.X(i2);
-		const FSolverVec3& P3 = Particles.X(i3);
+		const FSolverVec3& P1 = Particles.GetX(i1);
+		const FSolverVec3& P2 = Particles.GetX(i2);
+		const FSolverVec3& P3 = Particles.GetX(i3);
 		const FSolverVec3& P32 = P3 - P2;
 		const FSolverReal Bary = FSolverVec3::DotProduct(P32, P3 - P1) / P32.SizeSquared();
 		return FMath::Clamp(Bary, (FSolverReal)0., (FSolverReal)1.);
@@ -189,9 +189,9 @@ private:
 			i1 = Constraint[0];
 			i2 = Constraint[1];
 			i3 = Constraint[2];
-			const FSolverVec3& P1 = Particles.X(i1);
-			const FSolverVec3& P2 = Particles.X(i2);
-			const FSolverVec3& P3 = Particles.X(i3);
+			const FSolverVec3& P1 = Particles.GetX(i1);
+			const FSolverVec3& P2 = Particles.GetX(i2);
+			const FSolverVec3& P3 = Particles.GetX(i3);
 			const FSolverVec3 P = (P2 - P3) * Bary + P3;
 
 			if (bTrimKinematicConstraints)
diff --git a/Engine/Source/Runtime/Experimental/Chaos/Public/Chaos/PBDParticles.h b/Engine/Source/Runtime/Experimental/Chaos/Public/Chaos/PBDParticles.h
--- a/Engine/Source/Runtime/Experimental/Chaos/Public/Chaos/PBDParticles.h
+++ b/Engine/Source/Runtime/Experimental/Chaos/Public/Chaos/PBDParticles.h
@@ -26,8 +26,11 @@ class TPBDParticles : public TDynamicParticles<T, d>
 
 	const TVector<T, d>& P(const int32 index) const { return MP[index]; }
 	TVector<T, d>& P(const int32 index) { return MP[index]; }
+	const TVector<T, d>& GetP(const int32 index)const  { return MP[index]; }
+	void SetP(const int32 index, const TVector<T, d>& InP) { MP[index] = InP; }
 	TArrayCollectionArray<TVector<T, d>>& GetP() { return MP; }
 
+
   private:
 	TArrayCollectionArray<TVector<T, d>> MP;
 };
diff --git a/Engine/Source/Runtime/Experimental/Chaos/Public/Chaos/PBDRigidParticles.h b/Engine/Source/Runtime/Experimental/Chaos/Public/Chaos/PBDRigidParticles.h
--- a/Engine/Source/Runtime/Experimental/Chaos/Public/Chaos/PBDRigidParticles.h
+++ b/Engine/Source/Runtime/Experimental/Chaos/Public/Chaos/PBDRigidParticles.h
@@ -59,8 +59,12 @@ class TPBDRigidParticles : public TRigidParticles<T, d>
 		TArrayCollection::AddArray(&MSolverBodyIndex);
 	}
 
+	UE_DEPRECATED(5.4, "Use GetP instead")
 	FORCEINLINE const TVector<T, d>& P(const int32 index) const { return MP[index]; }
+	UE_DEPRECATED(5.4, "Use GetP or SetP instead")
 	FORCEINLINE TVector<T, d>& P(const int32 index) { return MP[index]; }
+	FORCEINLINE const TVector<T, d>& GetP(const int32 index) const { return MP[index]; }
+	FORCEINLINE void SetP(const int32 index, const TVector<T, d>& InP) { MP[index] = InP; }
 
 	UE_DEPRECATED(5.4, "Use GetQ instead")
 	FORCEINLINE const TRotation<T, d> Q(const int32 index) const { return TRotation<T, d>(MQ[index]); }
@@ -90,8 +94,8 @@ class TPBDRigidParticles : public TRigidParticles<T, d>
 	void SetPreWf(const int32 index, const TVector<FRealSingle, d>& InPreW) { MPreW[index] = InPreW; }
 
 	// World-space center of mass location
-	const TVector<T, d> XCom(const int32 index) const { return this->X(index) + this->GetR(index).RotateVector(CenterOfMass(index)); }
-	const TVector<T, d> PCom(const int32 index) const { return this->P(index) + this->GetQ(index).RotateVector(CenterOfMass(index)); }
+	const TVector<T, d> XCom(const int32 index) const { return this->GetX(index) + this->GetR(index).RotateVector(CenterOfMass(index)); }
+	const TVector<T, d> PCom(const int32 index) const { return this->GetP(index) + this->GetQ(index).RotateVector(CenterOfMass(index)); }
 
 	// World-space center of mass rotation
 	const TRotation<T, d> RCom(const int32 index) const { return this->GetR(index) * RotationOfMass(index); }
@@ -99,8 +103,8 @@ class TPBDRigidParticles : public TRigidParticles<T, d>
 
 	void SetTransformPQCom(const int32 index, const TVector<T, d>& InPCom, const TRotation<T, d>& InQCom)
 	{
-		SetQ(index,InQCom * RotationOfMass(index).Inverse());
-		P(index) = InPCom - GetQ(index) * CenterOfMass(index);
+		SetQ(index, InQCom * RotationOfMass(index).Inverse());
+		SetP(index, InPCom - GetQ(index) * CenterOfMass(index));
 	}
 
 	// The index into an FSolverBodyContainer (for dynamic particles only), or INDEX_NONE.
@@ -181,8 +185,8 @@ class TPBDRigidParticles : public TRigidParticles<T, d>
 				1.f / this->I(Index)[1],
 				1.f / this->I(Index)[2]);
 
-			this->P(Index) = this->X(Index);
-			this->SetQ(Index, this->GetR(Index));
+			this->SetP(Index, this->GetX(Index));
+			this->SetQf(Index, this->GetRf(Index));
 		}
 		else if (InObjectState == EObjectStateType::Sleeping)
 		{
@@ -214,7 +218,7 @@ class TPBDRigidParticles : public TRigidParticles<T, d>
 	FString ToString(int32 index) const
 	{
 		FString BaseString = TRigidParticles<T, d>::ToString(index);
-		return FString::Printf(TEXT("%s, MP:%s, MQ:%s, MPreV:%s, MPreW:%s"), *BaseString, *P(index).ToString(), *GetQ(index).ToString(), *GetPreV(index).ToString(), *GetPreW(index).ToString());
+		return FString::Printf(TEXT("%s, MP:%s, MQ:%s, MPreV:%s, MPreW:%s"), *BaseString, *GetP(index).ToString(), *GetQ(index).ToString(), *GetPreV(index).ToString(), *GetPreW(index).ToString());
 	}
 
 	virtual void Serialize(FChaosArchive& Ar) override
diff --git a/Engine/Source/Runtime/Experimental/Chaos/Public/Chaos/PBDRigidsEvolution.h b/Engine/Source/Runtime/Experimental/Chaos/Public/Chaos/PBDRigidsEvolution.h
--- a/Engine/Source/Runtime/Experimental/Chaos/Public/Chaos/PBDRigidsEvolution.h
+++ b/Engine/Source/Runtime/Experimental/Chaos/Public/Chaos/PBDRigidsEvolution.h
@@ -751,7 +751,7 @@ public:
 			CVD_SCOPE_CONTEXT(CVDContext.Context);
 
 			TKinematicTarget<FReal, 3>& KinematicTarget = Particle.KinematicTarget();
-			const FVec3 CurrentX = Particle.X();
+			const FVec3 CurrentX = Particle.GetX();
 			const FRotation3 CurrentR = Particle.GetR();
 			constexpr FReal MinDt = 1e-6f;
 
@@ -765,8 +765,8 @@ public:
 			case EKinematicTargetMode::Reset:
 			{
 				// Reset velocity and then switch to do-nothing mode
-				Particle.SetVf(FVec3f(0, 0, 0));
-				Particle.SetWf(FVec3f(0, 0, 0));
+				Particle.SetVf(FVec3f(0.0f, 0.0f, 0.0f));
+				Particle.SetWf(FVec3f(0.0f, 0.0f, 0.0f));
 				Particle.ClearIsMovingKinematic();
 				KinematicTarget.SetMode(EKinematicTargetMode::None);
 				break;
@@ -808,7 +808,7 @@ public:
 						NewW = FRotation3::CalculateAngularVelocity(CurrentR, NewR, Dt);
 					}
 				}
-				Particle.X() = NewX;
+				Particle.SetX(NewX);
 				Particle.SetR(NewR);
 				Particle.SetV(NewV);
 				Particle.SetW(NewW);
@@ -821,8 +821,8 @@ public:
 			{
 				// Move based on velocity
 				bMoved = true;
-				Particle.X() = Particle.X() + Particle.GetV() * Dt;
-				Particle.SetR(FRotation3::IntegrateRotationWithAngularVelocity(Particle.GetR(), Particle.GetW(), Dt));
+				Particle.SetX(Particle.GetX() + Particle.GetV() * Dt);
+				Particle.SetRf(FRotation3f::IntegrateRotationWithAngularVelocity(Particle.GetRf(), Particle.GetWf(), FRealSingle(Dt)));
 				Particle.SetIsMovingKinematic();
 
 				break;
@@ -831,7 +831,7 @@ public:
 			
 			// Set positions and previous velocities if we can
 			// Note: At present kinematics are in fact rigid bodies
-			Particle.P() = Particle.X();
+			Particle.SetP(Particle.GetX());
 			Particle.SetQf(Particle.GetRf());
 			Particle.SetPreVf(Particle.GetVf());
 			Particle.SetPreWf(Particle.GetWf());
@@ -840,11 +840,11 @@ public:
 			{
 				if (!Particle.CCDEnabled())
 				{
-					Particle.UpdateWorldSpaceState(FRigidTransform3(Particle.P(), Particle.GetQ()), FVec3(0));
+					Particle.UpdateWorldSpaceState(FRigidTransform3(Particle.GetP(), Particle.GetQ()), FVec3(0));
 				}
 				else
 				{
-					Particle.UpdateWorldSpaceStateSwept(FRigidTransform3(Particle.P(), Particle.GetQ()), FVec3(0), -Particle.GetV() * Dt);
+					Particle.UpdateWorldSpaceStateSwept(FRigidTransform3(Particle.GetP(), Particle.GetQ()), FVec3(0), -Particle.GetV() * Dt);
 				}
 			}
 
diff --git a/Engine/Source/Runtime/Experimental/Chaos/Public/Chaos/PBDRigidsEvolutionGBF.h b/Engine/Source/Runtime/Experimental/Chaos/Public/Chaos/PBDRigidsEvolutionGBF.h
--- a/Engine/Source/Runtime/Experimental/Chaos/Public/Chaos/PBDRigidsEvolutionGBF.h
+++ b/Engine/Source/Runtime/Experimental/Chaos/Public/Chaos/PBDRigidsEvolutionGBF.h
@@ -323,7 +323,7 @@ namespace Chaos
 					{
 						// Expand bounds about P/Q by a small amount. This can still result in missed collisions, especially
 						// when we have joints that pull the body back to X/R, if P-X is greater than the BoundsThickness
-						Particle.UpdateWorldSpaceStateSwept(FRigidTransform3(Particle.P(), Particle.GetQ()), FVec3(BoundsThickness), VelocityBoundsDelta);
+						Particle.UpdateWorldSpaceStateSwept(FRigidTransform3(Particle.GetP(), Particle.GetQ()), FVec3(BoundsThickness), VelocityBoundsDelta);
 					}
 					else
 					{
@@ -331,11 +331,11 @@ namespace Chaos
 #if CHAOS_DEBUG_DRAW
 						if (CVars::ChaosSolverDrawCCDThresholds)
 						{
-							DebugDraw::DrawCCDAxisThreshold(Particle.X(), Particle.CCDAxisThreshold(), Particle.P() - Particle.X(), Particle.GetQ());
+							DebugDraw::DrawCCDAxisThreshold(Particle.GetX(), Particle.CCDAxisThreshold(), Particle.GetP() - Particle.GetX(), Particle.GetQ());
 						}
 #endif
 
-						if (FCCDHelpers::DeltaExceedsThreshold(Particle.CCDAxisThreshold(), Particle.P() - Particle.X(), Particle.GetQ()))
+						if (FCCDHelpers::DeltaExceedsThreshold(Particle.CCDAxisThreshold(), Particle.GetP() - Particle.GetX(), Particle.GetQ()))
 						{
 							// We sweep the bounds from P back along the velocity and expand by a small amount.
 							// If not using tight bounds we also expand the bounds in all directions by Velocity. This is necessary only for secondary CCD collisions
@@ -346,11 +346,11 @@ namespace Chaos
 							{
 								CCDBoundsExpansion += VDt.GetAbsMax();
 							}
-							Particle.UpdateWorldSpaceStateSwept(FRigidTransform3(Particle.P(), Particle.GetQ()), FVec3(CCDBoundsExpansion), -VDt);
+							Particle.UpdateWorldSpaceStateSwept(FRigidTransform3(Particle.GetP(), Particle.GetQ()), FVec3(CCDBoundsExpansion), -VDt);
 						}
 						else
 						{
-							Particle.UpdateWorldSpaceStateSwept(FRigidTransform3(Particle.P(), Particle.GetQ()), FVec3(BoundsThickness), VelocityBoundsDelta);
+							Particle.UpdateWorldSpaceStateSwept(FRigidTransform3(Particle.GetP(), Particle.GetQ()), FVec3(BoundsThickness), VelocityBoundsDelta);
 						}
 					}
 
diff --git a/Engine/Source/Runtime/Experimental/Chaos/Public/Chaos/PBDSoftsSolverParticles.h b/Engine/Source/Runtime/Experimental/Chaos/Public/Chaos/PBDSoftsSolverParticles.h
--- a/Engine/Source/Runtime/Experimental/Chaos/Public/Chaos/PBDSoftsSolverParticles.h
+++ b/Engine/Source/Runtime/Experimental/Chaos/Public/Chaos/PBDSoftsSolverParticles.h
@@ -35,9 +35,21 @@ public:
 
 	const FSolverVec3& P(const int32 index) const { return MPAndInvM[index].P; }
 	FSolverVec3& P(const int32 index) { return MPAndInvM[index].P; }
+	const FSolverVec3& X(const int32 index) const { 
+		PRAGMA_DISABLE_DEPRECATION_WARNINGS
+		return Base::X(index); 
+		PRAGMA_ENABLE_DEPRECATION_WARNINGS
+	}
+	FSolverVec3& X(const int32 index) { 
+		PRAGMA_DISABLE_DEPRECATION_WARNINGS
+		return Base::X(index); 
+		PRAGMA_ENABLE_DEPRECATION_WARNINGS
+	}
+	const FSolverVec3& GetP(const int32 index) const { return MPAndInvM[index].P; }
+	void SetP(const int32 index, const FSolverVec3& InP) { MPAndInvM[index].P = InP; }
 	// When PBD is solving, V = last time step's V and we use the latest P to estimate velocity, (VelocityEstimate = (P - X)/Dt.
 	// After PBD (e.g., during implicit timestepping CG solve), V = current velocity estimate, VPrev = last time step's velocity
-	// Warning: VPrev and Acceleration use the same arrays as Accleration is only used by PBD and VPrev is only used by the linear solver
+	// Warning: VPrev and Acceleration use the same arrays as Acceleration is only used by PBD and VPrev is only used by the linear solver
 	const FSolverVec3& VPrev(const int32 Index) const { return Acceleration(Index); }
 	FSolverVec3& VPrev(const int32 Index) { return Acceleration(Index); }
 
@@ -46,7 +58,7 @@ public:
 
 	TArrayCollectionArray<FPAndInvM>& GetPAndInvM() { return MPAndInvM; }
 	const TArrayCollectionArray<FPAndInvM>& GetPAndInvM() const { return MPAndInvM; }
-	// Warning: VPrev and Acceleration use the same arrays as Accleration is only used by PBD and VPrev is only used by the linear solver
+	// Warning: VPrev and Acceleration use the same arrays as Acceleration is only used by PBD and VPrev is only used by the linear solver
 	TArrayCollectionArray<FSolverVec3>& GetVPrev() { return GetAcceleration(); }
 	const TArrayCollectionArray<FSolverVec3>& GetVPrev() const { return  GetAcceleration(); }
 
diff --git a/Engine/Source/Runtime/Experimental/Chaos/Public/Chaos/PBDSpringConstraintsBase.h b/Engine/Source/Runtime/Experimental/Chaos/Public/Chaos/PBDSpringConstraintsBase.h
--- a/Engine/Source/Runtime/Experimental/Chaos/Public/Chaos/PBDSpringConstraintsBase.h
+++ b/Engine/Source/Runtime/Experimental/Chaos/Public/Chaos/PBDSpringConstraintsBase.h
@@ -82,8 +82,8 @@ public:
 		Dists.Reset(Constraints.Num());
 		for (const TVec2<int32>& Constraint : Constraints)
 		{
-			const FSolverVec3& P0 = Particles.X(Constraint[0]);
-			const FSolverVec3& P1 = Particles.X(Constraint[1]);
+			const FSolverVec3& P0 = Particles.GetX(Constraint[0]);
+			const FSolverVec3& P1 = Particles.GetX(Constraint[1]);
 			Dists.Add((P1 - P0).Size());
 		}
 	}
diff --git a/Engine/Source/Runtime/Experimental/Chaos/Public/Chaos/PBDTetConstraintsBase.h b/Engine/Source/Runtime/Experimental/Chaos/Public/Chaos/PBDTetConstraintsBase.h
--- a/Engine/Source/Runtime/Experimental/Chaos/Public/Chaos/PBDTetConstraintsBase.h
+++ b/Engine/Source/Runtime/Experimental/Chaos/Public/Chaos/PBDTetConstraintsBase.h
@@ -19,10 +19,10 @@ public:
 	{
 		for (const TVec4<int32>& Constraint : Constraints)
 		{
-			const FSolverVec3& P1 = InParticles.X(Constraint[0]);
-			const FSolverVec3& P2 = InParticles.X(Constraint[1]);
-			const FSolverVec3& P3 = InParticles.X(Constraint[2]);
-			const FSolverVec3& P4 = InParticles.X(Constraint[3]);
+			const FSolverVec3& P1 = InParticles.GetX(Constraint[0]);
+			const FSolverVec3& P2 = InParticles.GetX(Constraint[1]);
+			const FSolverVec3& P3 = InParticles.GetX(Constraint[2]);
+			const FSolverVec3& P4 = InParticles.GetX(Constraint[3]);
 			Volumes.Add(FSolverVec3::DotProduct(FSolverVec3::CrossProduct(P2 - P1, P3 - P1), P4 - P1) / (FSolverReal)6.);
 		}
 	}
diff --git a/Engine/Source/Runtime/Experimental/Chaos/Public/Chaos/PBDVolumeConstraintBase.h b/Engine/Source/Runtime/Experimental/Chaos/Public/Chaos/PBDVolumeConstraintBase.h
--- a/Engine/Source/Runtime/Experimental/Chaos/Public/Chaos/PBDVolumeConstraintBase.h
+++ b/Engine/Source/Runtime/Experimental/Chaos/Public/Chaos/PBDVolumeConstraintBase.h
@@ -20,15 +20,15 @@ class FPBDVolumeConstraintBase
 		FSolverVec3 Com = FSolverVec3(0.);
 		for (uint32 i = 0; i < InParticles.Size(); ++i)
 		{
-			Com += InParticles.X(i);
+			Com += InParticles.GetX(i);
 		}
 		Com /= (FSolverReal)InParticles.Size();
 		RefVolume = (FSolverReal)0.;
 		for (const TVec3<int32>& Constraint : Constraints)
 		{
-			const FSolverVec3& P1 = InParticles.X(Constraint[0]);
-			const FSolverVec3& P2 = InParticles.X(Constraint[1]);
-			const FSolverVec3& P3 = InParticles.X(Constraint[2]);
+			const FSolverVec3& P1 = InParticles.GetX(Constraint[0]);
+			const FSolverVec3& P2 = InParticles.GetX(Constraint[1]);
+			const FSolverVec3& P3 = InParticles.GetX(Constraint[2]);
 			RefVolume += GetVolume(P1, P2, P3, Com);
 		}
 		RefVolume /= (FSolverReal)9.;
@@ -44,11 +44,11 @@ class FPBDVolumeConstraintBase
 		FSolverReal WlDenom = (FSolverReal)0.;
 		for (uint32 i = 0; i < InParticles.Size(); ++i)
 		{
-			WlDenom += (InParticles.P(i) - InParticles.X(i)).Size();
+			WlDenom += (InParticles.P(i) - InParticles.GetX(i)).Size();
 		}
 		for (uint32 i = 0; i < InParticles.Size(); ++i)
 		{
-			const FSolverReal Wl = (InParticles.P(i) - InParticles.X(i)).Size() / WlDenom;
+			const FSolverReal Wl = (InParticles.P(i) - InParticles.GetX(i)).Size() / WlDenom;
 			W[i] = OneMinusAlpha * Wl + Alpha * Wg;
 		}
 		return W;
diff --git a/Engine/Source/Runtime/Experimental/Chaos/Public/Chaos/Particle/ParticleUtilities.h b/Engine/Source/Runtime/Experimental/Chaos/Public/Chaos/Particle/ParticleUtilities.h
--- a/Engine/Source/Runtime/Experimental/Chaos/Public/Chaos/Particle/ParticleUtilities.h
+++ b/Engine/Source/Runtime/Experimental/Chaos/Public/Chaos/Particle/ParticleUtilities.h
@@ -20,16 +20,16 @@ namespace Chaos
 	{
 	public:
 		template<typename T_PARTICLEHANDLE>
-		static inline FVec3& Position(T_PARTICLEHANDLE Particle) { return Particle->X(); }
+		static inline FVec3& Position(T_PARTICLEHANDLE Particle) { return Particle->GetX(); }
 		template<typename T_PARTICLEHANDLE>
-		static inline const FVec3& GetPosition(T_PARTICLEHANDLE Particle) { return Particle->X(); }
+		static inline const FVec3& GetPosition(T_PARTICLEHANDLE Particle) { return Particle->GetX(); }
 		template<typename T_PARTICLEHANDLE>
 		UE_DEPRECATED(5.4, "Use Particle->SetR() instead.")
 		static inline FRotation3 Rotation(T_PARTICLEHANDLE Particle) { return Particle->GetR(); }
 		template<typename T_PARTICLEHANDLE>
 		static inline const FRotation3 GetRotation(T_PARTICLEHANDLE Particle) { return Particle->GetR(); }
 
-		static inline FVec3& Position(TPBDRigidParticles<FReal, 3>& Particles, const int32 Index) { return Particles.X(Index); }
+		static inline const FVec3& Position(TPBDRigidParticles<FReal, 3>& Particles, const int32 Index) { return Particles.GetX(Index); }
 		static inline const FRotation3 Rotation(TPBDRigidParticles<FReal, 3>& Particles, const int32 Index) { return Particles.GetR(Index); }
 	};
 
@@ -39,13 +39,13 @@ namespace Chaos
 		template<typename T_PARTICLEHANDLE>
 		static inline FVec3& Position(T_PARTICLEHANDLE Particle) { return Particle->P(); }
 		template<typename T_PARTICLEHANDLE>
-		static inline const FVec3& GetPosition(T_PARTICLEHANDLE Particle) { return Particle->P(); }
+		static inline const FVec3& GetPosition(T_PARTICLEHANDLE Particle) { return Particle->GetP(); }
 		template<typename T_PARTICLEHANDLE>
 		static inline FRotation3& Rotation(T_PARTICLEHANDLE Particle) { return Particle->GetQ(); }
 		template<typename T_PARTICLEHANDLE>
 		static inline const FRotation3 GetRotation(T_PARTICLEHANDLE Particle) { return Particle->GetQ(); }
 
-		static inline FVec3& Position(TPBDRigidParticles<FReal, 3>& Particles, const int32 Index) { return Particles.P(Index); }
+		static inline const FVec3& Position(TPBDRigidParticles<FReal, 3>& Particles, const int32 Index) { return Particles.GetP(Index); }
 		static inline FRotation3 Rotation(TPBDRigidParticles<FReal, 3>& Particles, const int32 Index) { return Particles.GetQ(Index); }
 	};
 
diff --git a/Engine/Source/Runtime/Experimental/Chaos/Public/Chaos/ParticleDirtyFlags.h b/Engine/Source/Runtime/Experimental/Chaos/Public/Chaos/ParticleDirtyFlags.h
--- a/Engine/Source/Runtime/Experimental/Chaos/Public/Chaos/ParticleDirtyFlags.h
+++ b/Engine/Source/Runtime/Experimental/Chaos/Public/Chaos/ParticleDirtyFlags.h
@@ -80,7 +80,7 @@ public:
 	template <typename TOther>
 	void CopyFrom(const TOther& Other)
 	{
-		MX = Other.X();
+		MX = Other.GetX();
 		MR = FRotation3f(Other.GetR());
 	}
 
@@ -96,6 +96,7 @@ public:
 	}
 
 	const FVec3& X() const { return MX; }
+	const FVec3& GetX() const { return MX; }
 	void SetX(const FVec3& InX){ MX = InX; }
 
 	const FRotation3 R() const { return FRotation3(MR); }
diff --git a/Engine/Source/Runtime/Experimental/Chaos/Public/Chaos/ParticleHandle.h b/Engine/Source/Runtime/Experimental/Chaos/Public/Chaos/ParticleHandle.h
--- a/Engine/Source/Runtime/Experimental/Chaos/Public/Chaos/ParticleHandle.h
+++ b/Engine/Source/Runtime/Experimental/Chaos/Public/Chaos/ParticleHandle.h
@@ -519,9 +519,15 @@ public:
 
 	TGeometryParticleHandleImp(const TGeometryParticleHandleImp&) = delete;
 
-	const TVector<T, d>& X() const { return GeometryParticles->X(ParticleIdx); }
-	TVector<T, d>& X() { return GeometryParticles->X(ParticleIdx); }
-	void SetX(const TVector<T, d>& InX, bool bInvalidate = false) { GeometryParticles->X(ParticleIdx) = InX; }
+	UE_DEPRECATED(5.4, "Use GetX instead")
+	const TVector<T, d>& X() const { return GeometryParticles->GetX(ParticleIdx); }
+	UE_DEPRECATED(5.4, "Use GetX or SetX instead")
+	TVector<T, d>& X() { 
+		PRAGMA_DISABLE_DEPRECATION_WARNINGS 
+			return GeometryParticles->X(ParticleIdx);  
+		PRAGMA_ENABLE_DEPRECATION_WARNINGS }
+	const TVector<T, d>& GetX() const { return GeometryParticles->GetX(ParticleIdx); }
+	void SetX(const TVector<T, d>& InX, bool bInvalidate = false) { GeometryParticles->SetX(ParticleIdx, InX); }
 
 	FUniqueIdx UniqueIdx() const { return GeometryParticles->UniqueIdx(ParticleIdx); }
 	void SetUniqueIdx(const FUniqueIdx UniqueIdx, bool bInvalidate = false) const { GeometryParticles->UniqueIdx(ParticleIdx) = UniqueIdx; }
@@ -535,7 +541,7 @@ public:
 	const TRotation<FRealSingle, d> GetRf() const { return GeometryParticles->GetRf(ParticleIdx); }
 	void SetRf(const TRotation<FRealSingle, d>& InR, bool bInvalidate = false) { GeometryParticles->SetRf(ParticleIdx, InR); }
 
-	FRigidTransform3 GetTransformXR() const { return FRigidTransform3(X(), GetR()); }
+	FRigidTransform3 GetTransformXR() const { return FRigidTransform3(GetX(), GetR()); }
 
 	// Initialize the transform
 	void InitTransform(const FVec3& InP, const FRotation3& InQ)
@@ -979,7 +985,7 @@ protected:
 		SetPreVf(this->GetVf());
 		SetPreWf(this->GetWf());
 		SetSolverBodyIndex(INDEX_NONE);
-		SetP(this->X());
+		SetP(this->GetX());
 		SetQf(this->GetRf());
 		SetVSmooth(this->GetV());
 		SetWSmooth(this->GetW());
@@ -1057,9 +1063,16 @@ public:
 	int32 SolverBodyIndex() const { return PBDRigidParticles->SolverBodyIndex(ParticleIdx); }
 	void SetSolverBodyIndex(const int32 InSolverBodyIndex) { PBDRigidParticles->SetSolverBodyIndex(ParticleIdx, InSolverBodyIndex); }
 
-	const TVector<T, d>& P() const { return PBDRigidParticles->P(ParticleIdx); }
-	TVector<T, d>& P() { return PBDRigidParticles->P(ParticleIdx); }
-	void SetP(const TVector<T, d>& InP) { PBDRigidParticles->P(ParticleIdx) = InP; }
+	UE_DEPRECATED(5.4, "Use GetP instead")
+	const TVector<T, d>& P() const { return PBDRigidParticles->GetP(ParticleIdx); }
+	UE_DEPRECATED(5.4, "Use GetP or SetP instead")
+	TVector<T, d>& P() { 
+		PRAGMA_DISABLE_DEPRECATION_WARNINGS
+		return PBDRigidParticles->P(ParticleIdx); 
+		PRAGMA_ENABLE_DEPRECATION_WARNINGS
+	}
+	const TVector<T, d>& GetP() const { return PBDRigidParticles->GetP(ParticleIdx); }
+	void SetP(const TVector<T, d>& InP) { PBDRigidParticles->SetP(ParticleIdx, InP); }
 
 	UE_DEPRECATED(5.4, "Use GetQ instead")
 	const TRotation<T, d> Q() const { return PBDRigidParticles->GetQ(ParticleIdx); }
@@ -1090,7 +1103,7 @@ public:
 	// Set world-space center of mass transform
 	void SetTransformPQCom(const TVector<T, d>& InPCom, const TRotation<T, d>& InQCom) { PBDRigidParticles->SetTransformPQCom(ParticleIdx, InPCom, InQCom); }
 
-	FRigidTransform3 GetTransformPQ() const { return FRigidTransform3(P(), GetQ()); }
+	FRigidTransform3 GetTransformPQ() const { return FRigidTransform3(GetP(), GetQ()); }
 
 	FRigidTransform3 GetTransformXRCom() const { return FRigidTransform3(XCom(), RCom()); }
 
@@ -1390,7 +1403,7 @@ void TGeometryParticleHandleImp<T,d,bPersistent>::SetXR(const FParticlePositionR
 	SetR(XR.R());
 	if(auto Rigid = CastToRigidParticle())
 	{
-		Rigid->SetP(X());
+		Rigid->SetP(GetX());
 		Rigid->SetQf(GetRf());
 	}
 }
@@ -1698,18 +1711,26 @@ public:
 
 	void SetTransform(const FVec3& Pos, const FRotation3& Rot)
 	{
-		MHandle->X() = Pos;
+		MHandle->SetX(Pos);
 		MHandle->SetR(Rot);
 		if (FPBDRigidParticleHandle* Dynamic = CastToRigidParticle())
 		{
-			Dynamic->P() = Pos;
+			Dynamic->SetP(Pos);
 			Dynamic->SetQ(Rot);
 		}
 	}
 
 	// Static Particles
-	FVec3& X() { return MHandle->X(); }
-	const FVec3& X() const { return MHandle->X(); }
+	UE_DEPRECATED(5.4, "Use GetX or SetX instead")
+	FVec3& X() { 
+		PRAGMA_DISABLE_DEPRECATION_WARNINGS
+		return MHandle->X(); 
+		PRAGMA_ENABLE_DEPRECATION_WARNINGS
+	}
+	UE_DEPRECATED(5.4, "Use GetX instead")
+	const FVec3& X() const { return MHandle->GetX(); }
+	void SetX(const FVec3& InX) { MHandle->SetX(InX); }
+	const FVec3& GetX() const { return MHandle->GetX(); }
 	UE_DEPRECATED(5.4, "Use GetR or SetR instead")
 	FRotation3 R() { return MHandle->GetR(); }
 	UE_DEPRECATED(5.4, "Use GetR instead")
@@ -1843,24 +1864,29 @@ public:
 		}
 	}
 
-	FVec3& P()
+	void SetP(const FVec3& InP)
 	{
 		if (IsDynamic())
 		{
-			return MHandle->CastToRigidParticle()->P();
+			return MHandle->CastToRigidParticle()->SetP(InP);
 		}
 
-		return X();
+		return SetX(InP);
 	}
 
 	const FVec3& P() const
 	{
 		if (IsDynamic())
 		{
-			return MHandle->CastToRigidParticle()->P();
+			return MHandle->CastToRigidParticle()->GetP();
 		}
 
-		return X();
+		return GetX();
+	}
+
+	const FVec3& GetP() const
+	{
+		return P();
 	}
 
 	void SetQ(const FRotation3& InQ)
@@ -1895,7 +1921,7 @@ public:
 		{
 			return MHandle->CastToRigidParticle()->XCom();
 		}
-		return X();
+		return GetX();
 	}
 	const FVec3 PCom() const
 	{
@@ -1903,7 +1929,7 @@ public:
 		{
 			return MHandle->CastToRigidParticle()->PCom();
 		}
-		return X();
+		return GetX();
 	}
 
 	// World-space center of mass rotation
@@ -2618,6 +2644,7 @@ public:
 	static TGeometryParticle<T, d>* SerializationFactory(FChaosArchive& Ar, TGeometryParticle<T, d>* Serializable);
 
 	const TVector<T, d>& X() const { return MXR.Read().X(); }
+	const TVector<T, d>& GetX() const { return MXR.Read().X(); }
 	void SetX(const TVector<T, d>& InX, bool bInvalidate = true);
 
 	FUniqueIdx UniqueIdx() const { return MNonFrequentData.Read().UniqueIdx(); }
diff --git a/Engine/Source/Runtime/Experimental/Chaos/Public/Chaos/Particles.h b/Engine/Source/Runtime/Experimental/Chaos/Public/Chaos/Particles.h
--- a/Engine/Source/Runtime/Experimental/Chaos/Public/Chaos/Particles.h
+++ b/Engine/Source/Runtime/Experimental/Chaos/Public/Chaos/Particles.h
@@ -141,19 +141,31 @@ namespace Chaos
 			return MX;
 		}
 		
+		UE_DEPRECATED(5.4, "Use GetX instead")
 		const TVector<T, d>& X(const int32 Index) const
 		{
 			return MX[Index];
 		}
 
+		UE_DEPRECATED(5.4, "Use GetX or SetX instead")
 		TVector<T, d>& X(const int32 Index)
 		{
 			return MX[Index];
 		}
 
+		const TVector<T, d>& GetX(const int32 Index) const
+		{
+			return MX[Index];
+		}
+
+		void SetX(const int32 Index, const TVector<T, d>& InX)
+		{
+			MX[Index] = InX;
+		}
+
 		FString ToString(int32 index) const
 		{
-			return FString::Printf(TEXT("MX:%s"), *X(index).ToString());
+			return FString::Printf(TEXT("MX:%s"), *GetX(index).ToString());
 		}
 
 		uint32 GetTypeHash() const
diff --git a/Engine/Source/Runtime/Experimental/Chaos/Public/Chaos/PerCellBoundaryConditions.h b/Engine/Source/Runtime/Experimental/Chaos/Public/Chaos/PerCellBoundaryConditions.h
--- a/Engine/Source/Runtime/Experimental/Chaos/Public/Chaos/PerCellBoundaryConditions.h
+++ b/Engine/Source/Runtime/Experimental/Chaos/Public/Chaos/PerCellBoundaryConditions.h
@@ -22,7 +22,7 @@ class TPerCellBoundaryConditions
 		for (uint32 i = 0; i < MParticles.Size(); ++i)
 		{
 			const TVector<T, d>& X = Grid.Location(Index);
-			TRigidTransform<T, d> Frame(MParticles.X(i), MParticles.GetR(i));
+			TRigidTransform<T, d> Frame(MParticles.GetX(i), MParticles.GetR(i));
 			if (MParticles.GetGeometry(i)->SignedDistance(Frame.InverseTransformPosition(X)) < 0)
 			{
 				BoundaryConditions(Index) = true;
@@ -33,7 +33,7 @@ class TPerCellBoundaryConditions
 		for (uint32 i = 0; i < MSources.Size(); ++i)
 		{
 			const TVector<T, d>& X = Grid.Location(Index);
-			TRigidTransform<T, d> Frame(MSources.X(i), MSources.GetR(i));
+			TRigidTransform<T, d> Frame(MSources.GetX(i), MSources.GetR(i));
 			if (MSources.GetGeometry(i)->SignedDistance(Frame.InverseTransformPosition(X)) < 0)
 			{
 				BoundaryConditions(Index) = true;
@@ -50,7 +50,7 @@ class TPerCellBoundaryConditions
 		for (uint32 i = 0; i < MSources.Size(); ++i)
 		{
 			const TVector<T, d>& X = Grid.Location(Index);
-			TRigidTransform<T, d> Frame(MSources.X(i), MSources.GetR(i));
+			TRigidTransform<T, d> Frame(MSources.GetX(i), MSources.GetR(i));
 			if (MSources.GetGeometry(i)->SignedDistance(Frame.InverseTransformPosition(X)) < 0)
 			{
 				Density(Index) = 1.f;
diff --git a/Engine/Source/Runtime/Experimental/Chaos/Public/Chaos/PerParticleDampVelocity.h b/Engine/Source/Runtime/Experimental/Chaos/Public/Chaos/PerParticleDampVelocity.h
--- a/Engine/Source/Runtime/Experimental/Chaos/Public/Chaos/PerParticleDampVelocity.h
+++ b/Engine/Source/Runtime/Experimental/Chaos/Public/Chaos/PerParticleDampVelocity.h
@@ -23,7 +23,7 @@ public:
 	// Apply damping without first checking for kinematic particles
 	void ApplyFast(FSolverParticles& Particles, const FSolverReal /*Dt*/, const int32 Index) const
 	{
-		Apply(Particles.X(Index), Particles.V(Index));
+		Apply(Particles.GetX(Index), Particles.V(Index));
 	}
 
 	void Apply(const FSolverVec3& X, FSolverVec3& V) const
diff --git a/Engine/Source/Runtime/Experimental/Chaos/Public/Chaos/PerParticlePBDCCDCollisionConstraint.h b/Engine/Source/Runtime/Experimental/Chaos/Public/Chaos/PerParticlePBDCCDCollisionConstraint.h
--- a/Engine/Source/Runtime/Experimental/Chaos/Public/Chaos/PerParticlePBDCCDCollisionConstraint.h
+++ b/Engine/Source/Runtime/Experimental/Chaos/Public/Chaos/PerParticlePBDCCDCollisionConstraint.h
@@ -87,10 +87,10 @@ private:
 					return; // Bail out if the collision groups doesn't match the particle group id, or use INDEX_NONE (= global collision that affects all particle)
 				}
 
-				const FSolverRigidTransform3 Frame(CollisionParticles.X(CollisionIndex), CollisionParticles.GetR(CollisionIndex));
+				const FSolverRigidTransform3 Frame(CollisionParticles.GetX(CollisionIndex), CollisionParticles.GetR(CollisionIndex));
 
 				const Pair<FVec3, bool> PointPair = CollisionParticles.GetGeometry(CollisionIndex)->FindClosestIntersection(  // Geometry operates in FReal
-					FVec3(CollisionTransforms[CollisionIndex].InverseTransformPositionNoScale(Particles.X(Index))),        // hence the back and forth
+					FVec3(CollisionTransforms[CollisionIndex].InverseTransformPositionNoScale(Particles.GetX(Index))),        // hence the back and forth
 					FVec3(Frame.InverseTransformPositionNoScale(Particles.P(Index))), (FReal)Thickness);                   // FVec3/FReal conversions
 
 				if (PointPair.Second)
@@ -141,8 +141,8 @@ private:
 						}
 					}
 
-					const FSolverVec3 VectorToPoint = Particles.P(Index) - CollisionParticles.X(VelocityBone);
-					const FSolverVec3 RelativeDisplacement = (Particles.P(Index) - Particles.X(Index)) - (CollisionParticles.V(VelocityBone) + FSolverVec3::CrossProduct(CollisionParticles.W(VelocityBone), VectorToPoint)) * Dt;  // This corresponds to the tangential velocity multiplied by dt (friction will drive this to zero if it is high enough)
+					const FSolverVec3 VectorToPoint = Particles.P(Index) - CollisionParticles.GetX(VelocityBone);
+					const FSolverVec3 RelativeDisplacement = (Particles.P(Index) - Particles.GetX(Index)) - (CollisionParticles.V(VelocityBone) + FSolverVec3::CrossProduct(CollisionParticles.W(VelocityBone), VectorToPoint)) * Dt;  // This corresponds to the tangential velocity multiplied by dt (friction will drive this to zero if it is high enough)
 					const FSolverVec3 RelativeDisplacementTangent = RelativeDisplacement - FSolverVec3::DotProduct(RelativeDisplacement, NormalWorld) * NormalWorld;  // Project displacement into the tangential plane
 					const FSolverReal RelativeDisplacementTangentLength = RelativeDisplacementTangent.Size();
 					if (RelativeDisplacementTangentLength >= UE_SMALL_NUMBER)
diff --git a/Engine/Source/Runtime/Experimental/Chaos/Public/Chaos/PerParticlePBDCollisionConstraint.h b/Engine/Source/Runtime/Experimental/Chaos/Public/Chaos/PerParticlePBDCollisionConstraint.h
--- a/Engine/Source/Runtime/Experimental/Chaos/Public/Chaos/PerParticlePBDCollisionConstraint.h
+++ b/Engine/Source/Runtime/Experimental/Chaos/Public/Chaos/PerParticlePBDCollisionConstraint.h
@@ -129,7 +129,7 @@ private:
 					{
 						return; // Bail out if the collision groups doesn't match the particle group id, or use INDEX_NONE (= global collision that affects all particle)
 					}
-					const FSolverRigidTransform3 Frame(CollisionParticles.X(i), CollisionParticles.GetR(i));
+					const FSolverRigidTransform3 Frame(CollisionParticles.GetX(i), CollisionParticles.GetR(i));
 					const FVec3 RigidSpacePosition(Frame.InverseTransformPosition(Particles.P(Index)));  // PhiWithNormal requires FReal based arguments
 					FVec3 ImplicitNormal;                                                                // since implicits don't use FSolverReal
 					FSolverReal Phi;
@@ -176,8 +176,8 @@ private:
 
 						if (bFastPositionBasedFriction)
 						{
-							FSolverVec3 VectorToPoint = Particles.P(Index) - CollisionParticles.X(VelocityBone);
-							const FSolverVec3 RelativeDisplacement = (Particles.P(Index) - Particles.X(Index)) - (CollisionParticles.V(VelocityBone) + FSolverVec3::CrossProduct(CollisionParticles.W(VelocityBone), VectorToPoint)) * Dt; // This corresponds to the tangential velocity multiplied by dt (friction will drive this to zero if it is high enough)
+							FSolverVec3 VectorToPoint = Particles.GetP(Index) - CollisionParticles.GetX(VelocityBone);
+							const FSolverVec3 RelativeDisplacement = (Particles.GetP(Index) - Particles.GetX(Index)) - (CollisionParticles.V(VelocityBone) + FSolverVec3::CrossProduct(CollisionParticles.W(VelocityBone), VectorToPoint)) * Dt; // This corresponds to the tangential velocity multiplied by dt (friction will drive this to zero if it is high enough)
 							const FSolverVec3 RelativeDisplacementTangent = RelativeDisplacement - FSolverVec3::DotProduct(RelativeDisplacement, NormalWorld) * NormalWorld; // Project displacement into the tangential plane
 							const FSolverReal RelativeDisplacementTangentLength = RelativeDisplacementTangent.Size();
 							if (RelativeDisplacementTangentLength >= UE_SMALL_NUMBER)
@@ -191,7 +191,7 @@ private:
 						{
 							// Note, to fix: Only use fast position based friction for now, since adding to TMaps here is not thread safe when calling Apply on multiple threads (will cause crash)
 							FVelocityConstraint Constraint;
-							FSolverVec3 VectorToPoint = Particles.P(Index) - CollisionParticles.X(VelocityBone);
+							FSolverVec3 VectorToPoint = Particles.GetP(Index) - CollisionParticles.GetX(VelocityBone);
 							Constraint.Velocity = CollisionParticles.V(VelocityBone) + FSolverVec3::CrossProduct(CollisionParticles.W(VelocityBone), VectorToPoint);
 							Constraint.Normal = Frame.TransformVector(Normal);
 						
@@ -220,7 +220,7 @@ private:
 					{
 						return; // Bail out if the collision groups doesn't match the particle group id, or use INDEX_NONE (= global collision that affects all particle)
 					}
-					const FSolverRigidTransform3 Frame(CollisionParticles.X(i), CollisionParticles.GetR(i));
+					const FSolverRigidTransform3 Frame(CollisionParticles.GetX(i), CollisionParticles.GetR(i));
 					const FVec3 RigidSpacePosition(Frame.InverseTransformPosition(Particles.P(Index)));  // PhiWithNormal requires FReal based arguments
 					FVec3 ImplicitNormal;                                                                // since implicits don't use FSolverReal
 					const FSolverReal Phi = (FSolverReal)CollisionParticles.GetGeometry(i)->PhiWithNormal(RigidSpacePosition, ImplicitNormal);
diff --git a/Engine/Source/Runtime/Experimental/Chaos/Public/Chaos/PerParticlePBDEulerStep.h b/Engine/Source/Runtime/Experimental/Chaos/Public/Chaos/PerParticlePBDEulerStep.h
--- a/Engine/Source/Runtime/Experimental/Chaos/Public/Chaos/PerParticlePBDEulerStep.h
+++ b/Engine/Source/Runtime/Experimental/Chaos/Public/Chaos/PerParticlePBDEulerStep.h
@@ -19,7 +19,7 @@ class FPerParticlePBDEulerStep : public FPerParticleRule
 	template<class T_PARTICLES>
 	inline void ApplyHelper(T_PARTICLES& InParticles, const FReal Dt, const int32 Index) const
 	{
-		InParticles.P(Index) = InParticles.X(Index) + InParticles.V(Index) * Dt;
+		InParticles.SetP(Index, InParticles.GetX(Index) + InParticles.GetV(Index) * Dt);
 	}
 
 	inline void Apply(FPBDParticles& InParticles, const FReal Dt, const int32 Index) const override //-V762
diff --git a/Engine/Source/Runtime/Experimental/Chaos/Public/Chaos/PerParticlePBDUpdateFromDeltaPosition.h b/Engine/Source/Runtime/Experimental/Chaos/Public/Chaos/PerParticlePBDUpdateFromDeltaPosition.h
--- a/Engine/Source/Runtime/Experimental/Chaos/Public/Chaos/PerParticlePBDUpdateFromDeltaPosition.h
+++ b/Engine/Source/Runtime/Experimental/Chaos/Public/Chaos/PerParticlePBDUpdateFromDeltaPosition.h
@@ -17,14 +17,13 @@ class FPerParticlePBDUpdateFromDeltaPosition : public FPerParticleRule
 	template<class T_PARTICLES>
 	inline void ApplyHelper(T_PARTICLES& InParticles, const FReal Dt, const int32 Index) const
 	{
-		InParticles.SetV(Index, (InParticles.P(Index) - InParticles.X(Index)) / Dt);
-		//InParticles.X(Index) = InParticles.P(Index);
+		InParticles.SetV(Index, (InParticles.GetP(Index) - InParticles.GetX(Index)) / Dt);
 	}
 
 	inline void Apply(FPBDParticles& InParticles, const FReal Dt, const int32 Index) const override //-V762
 	{
-		InParticles.V(Index) = (InParticles.P(Index) - InParticles.X(Index)) / Dt;
-		InParticles.X(Index) = InParticles.P(Index);
+		InParticles.V(Index) = (InParticles.GetP(Index) - InParticles.GetX(Index)) / Dt;
+		InParticles.SetX(Index,  InParticles.GetP(Index));
 	}
 
 	inline void Apply(TPBDRigidParticles<FReal, 3>& InParticles, const FReal Dt, const int32 Index) const override //-V762
diff --git a/Engine/Source/Runtime/Experimental/Chaos/Public/Chaos/SimpleGeometryParticles.h b/Engine/Source/Runtime/Experimental/Chaos/Public/Chaos/SimpleGeometryParticles.h
--- a/Engine/Source/Runtime/Experimental/Chaos/Public/Chaos/SimpleGeometryParticles.h
+++ b/Engine/Source/Runtime/Experimental/Chaos/Public/Chaos/SimpleGeometryParticles.h
@@ -17,7 +17,7 @@ namespace Chaos
 	public:
 
 		using TArrayCollection::Size;
-		using TParticles<T,d>::X;
+		using TParticles<T,d>::GetX;
 		
 		TSimpleGeometryParticles()
 		    : TParticles<T, d>()
diff --git a/Engine/Source/Runtime/Experimental/Chaos/Public/Chaos/SoftsSolverCollisionParticles.h b/Engine/Source/Runtime/Experimental/Chaos/Public/Chaos/SoftsSolverCollisionParticles.h
--- a/Engine/Source/Runtime/Experimental/Chaos/Public/Chaos/SoftsSolverCollisionParticles.h
+++ b/Engine/Source/Runtime/Experimental/Chaos/Public/Chaos/SoftsSolverCollisionParticles.h
@@ -8,6 +8,8 @@ namespace Chaos::Softs
 {
 class FSolverCollisionParticles : public TSimpleGeometryParticles<FSolverReal, 3>
 {
+	typedef TSimpleGeometryParticles<FSolverReal, 3> Base;
+
   public:
 	FSolverCollisionParticles()
 	    : TSimpleGeometryParticles<FSolverReal, 3>()
@@ -26,6 +28,17 @@ class FSolverCollisionParticles : public TSimpleGeometryParticles<FSolverReal, 3
 	FSolverCollisionParticles& operator=(const FSolverCollisionParticles& Other) = delete;
 	FSolverCollisionParticles& operator=(FSolverCollisionParticles&& Other) = delete;
 
+	const TVector<FSolverReal, 3>& X(const int32 Index) const	{ 
+		PRAGMA_DISABLE_DEPRECATION_WARNINGS
+		return Base::X(Index);
+		PRAGMA_ENABLE_DEPRECATION_WARNINGS
+	}
+	TVector<FSolverReal, 3>& X(const int32 Index) {
+		PRAGMA_DISABLE_DEPRECATION_WARNINGS
+		return Base::X(Index);
+		PRAGMA_ENABLE_DEPRECATION_WARNINGS
+		}
+
 	const FSolverVec3& V(const int32 Index) const { return MV[Index]; }
 	FSolverVec3& V(const int32 Index) { return MV[Index]; }
 	const TArrayCollectionArray<FSolverVec3>& GetV() const { return MV; }
diff --git a/Engine/Source/Runtime/Experimental/Chaos/Public/Chaos/SoftsSolverCollisionParticlesRange.h b/Engine/Source/Runtime/Experimental/Chaos/Public/Chaos/SoftsSolverCollisionParticlesRange.h
--- a/Engine/Source/Runtime/Experimental/Chaos/Public/Chaos/SoftsSolverCollisionParticlesRange.h
+++ b/Engine/Source/Runtime/Experimental/Chaos/Public/Chaos/SoftsSolverCollisionParticlesRange.h
@@ -34,7 +34,8 @@ public:
 	// SimpleGeometryParticles data
 	const FSolverRotation3 R(const int32 Index) const { return GetParticles().GetR(Index + Offset); }
 	void SetR(const int32 Index, const FSolverRotation3& InR) { GetParticles().SetR(Index + Offset, InR); }
-	// FSolverRotation3 R(const int32 Index) { return GetParticles().SetR(Index + Offset); }
+	UE_DEPRECATED(5.4, "Use GetR or SetR instead")
+	FSolverRotation3 R(const int32 Index) { return GetParticles().GetR(Index + Offset); }
 	TConstArrayView<FSolverRotation3> GetR() const { return GetConstArrayView(GetParticles().GetR()); }
 	TArrayView<FSolverRotation3> GetR() { return GetArrayView(GetParticles().GetR()); }
 
@@ -43,8 +44,14 @@ public:
 	TConstArrayView<FImplicitObjectPtr> GetAllGeometry() const { return GetConstArrayView(GetParticles().GetAllGeometry()); }
 
 	// Particles data
-	const FSolverVec3& X(const int32 Index) const { return GetParticles().X(Index + Offset); }
-	FSolverVec3& X(const int32 Index) { return GetParticles().X(Index + Offset); }
+	const FSolverVec3& X(const int32 Index) const { return GetParticles().GetX(Index + Offset); }
+	FSolverVec3& X(const int32 Index) { 
+		PRAGMA_DISABLE_DEPRECATION_WARNINGS
+		return GetParticles().X(Index + Offset); 
+		PRAGMA_ENABLE_DEPRECATION_WARNINGS
+	}
+	const FSolverVec3& GetX(const int32 Index) const { return GetParticles().GetX(Index + Offset); }
+	void SetX(const int32 Index, const FSolverVec3& InX) { GetParticles().SetX(Index + Offset, InX); }
 	TConstArrayView<FSolverVec3> XArray() const { return GetConstArrayView(GetParticles().XArray()); }
 	TArrayView<FSolverVec3> XArray() { return GetArrayView(GetParticles().XArray()); }
 };
diff --git a/Engine/Source/Runtime/Experimental/Chaos/Public/Chaos/SoftsSolverParticlesRange.h b/Engine/Source/Runtime/Experimental/Chaos/Public/Chaos/SoftsSolverParticlesRange.h
--- a/Engine/Source/Runtime/Experimental/Chaos/Public/Chaos/SoftsSolverParticlesRange.h
+++ b/Engine/Source/Runtime/Experimental/Chaos/Public/Chaos/SoftsSolverParticlesRange.h
@@ -29,6 +29,8 @@ public:
 	// SolverParticles data
 	const FSolverVec3& P(const int32 Index) const { return GetParticles().P(Index + Offset); }
 	FSolverVec3& P(const int32 Index) { return GetParticles().P(Index + Offset); }
+	const FSolverVec3& GetP(const int32 Index) const { return GetParticles().P(Index + Offset); }
+	void SetP(const int32 Index, const FSolverVec3& InP) { GetParticles().P(Index + Offset) = InP; }
 	const FPAndInvM& PAndInvM(const int32 Index) const { return GetParticles().PAndInvM(Index + Offset); }
 	FPAndInvM& PAndInvM(const int32 Index) { return GetParticles().PAndInvM(Index + Offset); }
 	TConstArrayView<FPAndInvM> GetPAndInvM() const { return GetConstArrayView(GetParticles().GetPAndInvM()); }
@@ -57,8 +59,15 @@ public:
 	TArrayView<FSolverReal> GetInvM() { return GetArrayView(GetParticles().GetInvM()); }
 
 	// Particles data
-	const FSolverVec3& X(const int32 Index) const { return GetParticles().X(Index + Offset); }
-	FSolverVec3& X(const int32 Index) { return GetParticles().X(Index + Offset); }
+	const FSolverVec3& X(const int32 Index) const { return GetParticles().GetX(Index + Offset); }
+	FSolverVec3& X(const int32 Index) { 
+		PRAGMA_DISABLE_DEPRECATION_WARNINGS
+		return GetParticles().X(Index + Offset); 
+		PRAGMA_ENABLE_DEPRECATION_WARNINGS
+	}
+	const FSolverVec3& GetX(const int32 Index) const { return GetParticles().GetX(Index + Offset); }
+	void SetX(const int32 Index, const FSolverVec3& InX) { GetParticles().SetX(Index + Offset, InX); }
+
 	TConstArrayView<FSolverVec3> XArray() const { return GetConstArrayView(GetParticles().XArray()); }
 	TArrayView<FSolverVec3> XArray() { return GetArrayView(GetParticles().XArray()); }
 };
diff --git a/Engine/Source/Runtime/Experimental/Chaos/Public/Chaos/SoftsSpring.h b/Engine/Source/Runtime/Experimental/Chaos/Public/Chaos/SoftsSpring.h
--- a/Engine/Source/Runtime/Experimental/Chaos/Public/Chaos/SoftsSpring.h
+++ b/Engine/Source/Runtime/Experimental/Chaos/Public/Chaos/SoftsSpring.h
@@ -56,8 +56,8 @@ FSolverVec3 GetXPBDSpringDeltaWithDamping(const SolverParticlesOrRange& Particle
 	const FSolverReal Distance = Direction.SafeNormalize();
 	const FSolverReal Offset = Distance - RestLength;
 
-	const FSolverVec3& X1 = Particles.X(Index1);
-	const FSolverVec3& X2 = Particles.X(Index2);
+	const FSolverVec3& X1 = Particles.GetX(Index1);
+	const FSolverVec3& X2 = Particles.GetX(Index2);
 
 	const FSolverVec3 RelativeVelocityTimesDt = P1 - X1 - P2 + X2;
 
@@ -89,8 +89,8 @@ FSolverVec3 GetXPBDSpringDampingDelta(const SolverParticlesOrRange& Particles, c
 	FSolverVec3 Direction = (P1 - P2);
 	Direction.SafeNormalize();
 
-	const FSolverVec3& X1 = Particles.X(Index1);
-	const FSolverVec3& X2 = Particles.X(Index2);
+	const FSolverVec3& X1 = Particles.GetX(Index1);
+	const FSolverVec3& X2 = Particles.GetX(Index2);
 	const FSolverVec3 RelativeVelocityTimesDt = P1 - X1 - P2 + X2;
 	const FSolverReal BetaDt = Damping * Dt;
 	const FSolverReal DLambda = (-BetaDt * FSolverVec3::DotProduct(Direction, RelativeVelocityTimesDt) - Lambda) / (BetaDt * CombinedInvMass + (FSolverReal)1.);
@@ -154,9 +154,9 @@ FSolverVec3 GetXPBDAxialSpringDeltaWithDamping(const SolverParticlesOrRange& Par
 	const FSolverReal Distance = Direction.SafeNormalize();
 	const FSolverReal Offset = Distance - RestLength;
 
-	const FSolverVec3& X1 = Particles.X(Index1);
-	const FSolverVec3& X2 = Particles.X(Index2);
-	const FSolverVec3& X3 = Particles.X(Index3);
+	const FSolverVec3& X1 = Particles.GetX(Index1);
+	const FSolverVec3& X2 = Particles.GetX(Index2);
+	const FSolverVec3& X3 = Particles.GetX(Index3);
 	const FSolverVec3 X = (X2 - X3) * Bary + X3;
 
 	const FSolverVec3 RelativeVelocityTimesDt = P1 - X1 - P + X;
@@ -195,9 +195,9 @@ FSolverVec3 GetXPBDAxialSpringDampingDelta(const SolverParticlesOrRange& Particl
 	const FSolverReal Distance = Direction.SafeNormalize();
 	const FSolverReal Offset = Distance - RestLength;
 
-	const FSolverVec3& X1 = Particles.X(Index1);
-	const FSolverVec3& X2 = Particles.X(Index2);
-	const FSolverVec3& X3 = Particles.X(Index3);
+	const FSolverVec3& X1 = Particles.GetX(Index1);
+	const FSolverVec3& X2 = Particles.GetX(Index2);
+	const FSolverVec3& X3 = Particles.GetX(Index3);
 	const FSolverVec3 X = (X2 - X3) * Bary + X3;
 	const FSolverVec3 RelativeVelocityTimesDt = P1 - X1 - P + X;
 	const FSolverReal BetaDt = Damping * Dt;
diff --git a/Engine/Source/Runtime/Experimental/Chaos/Public/Chaos/TriangleMesh.h b/Engine/Source/Runtime/Experimental/Chaos/Public/Chaos/TriangleMesh.h
--- a/Engine/Source/Runtime/Experimental/Chaos/Public/Chaos/TriangleMesh.h
+++ b/Engine/Source/Runtime/Experimental/Chaos/Public/Chaos/TriangleMesh.h
@@ -271,9 +271,9 @@ namespace Chaos
 			const int32 Idx = Particles.Size();
 			Particles.AddParticles(3);
 			// Left handed
-			Particles.X(Idx + 0) = FVec3((T)0., (T)0., (T)0.8083);
-			Particles.X(Idx + 1) = FVec3((T)0., (T)0.7, (T)-0.4041);
-			Particles.X(Idx + 2) = FVec3((T)0., (T)-0.7, (T)-0.4041);
+			Particles.SetX(Idx + 0, FVec3((T)0., (T)0., (T)0.8083));
+			Particles.SetX(Idx + 1, FVec3((T)0., (T)0.7, (T)-0.4041));
+			Particles.SetX(Idx + 2, FVec3((T)0., (T)-0.7, (T)-0.4041));
 
 			TArray<TVec3<int32>> Elements;
 			Elements.SetNum(1);
diff --git a/Engine/Source/Runtime/Experimental/Chaos/Public/Chaos/TriangleMeshImplicitObject.h b/Engine/Source/Runtime/Experimental/Chaos/Public/Chaos/TriangleMeshImplicitObject.h
--- a/Engine/Source/Runtime/Experimental/Chaos/Public/Chaos/TriangleMeshImplicitObject.h
+++ b/Engine/Source/Runtime/Experimental/Chaos/Public/Chaos/TriangleMeshImplicitObject.h
@@ -458,9 +458,9 @@ namespace Chaos
 	inline void TriangleMeshTransformVertsHelper(const FVec3& TriMeshScale, int32 TriIdx, const ParticlesType& Particles,
 		const TArray<TVector<IdxType, 3>>& Elements, FVec3& OutA, FVec3& OutB, FVec3& OutC)
 	{
-		OutA = Particles.X(Elements[TriIdx][0]) * TriMeshScale;
-		OutB = Particles.X(Elements[TriIdx][1]) * TriMeshScale;
-		OutC = Particles.X(Elements[TriIdx][2]) * TriMeshScale;
+		OutA = Particles.GetX(Elements[TriIdx][0]) * TriMeshScale;
+		OutB = Particles.GetX(Elements[TriIdx][1]) * TriMeshScale;
+		OutC = Particles.GetX(Elements[TriIdx][2]) * TriMeshScale;
 	}
 
 	template <typename IdxType, typename ParticlesType>
@@ -472,9 +472,9 @@ namespace Chaos
 		const int32 VertexIndexB = Elements[TriIdx][1];
 		const int32 VertexIndexC = Elements[TriIdx][2];
 		// Note: deliberately using scaled transform here. See VisitTriangles
-		OutA = Transform.TransformPosition(FVector(Particles.X(VertexIndexA)));
-		OutB = Transform.TransformPosition(FVector(Particles.X(VertexIndexB)));
-		OutC = Transform.TransformPosition(FVector(Particles.X(VertexIndexC)));
+		OutA = Transform.TransformPosition(FVector(Particles.GetX(VertexIndexA)));
+		OutB = Transform.TransformPosition(FVector(Particles.GetX(VertexIndexB)));
+		OutC = Transform.TransformPosition(FVector(Particles.GetX(VertexIndexC)));
 		OutVertexIndexA = VertexIndexA;
 		OutVertexIndexB = VertexIndexB;
 		OutVertexIndexC = VertexIndexC;
@@ -507,16 +507,16 @@ namespace Chaos
 				const TArray<TVec3<IdxType>>& Tris = MElements.GetIndexBuffer<IdxType>();
 				const TVec3<IdxType>& FirstTri = Tris[0];
 
-				MLocalBoundingBox = FAABB3(MParticles.X(FirstTri[0]), MParticles.X(FirstTri[0]));
-				MLocalBoundingBox.GrowToInclude(MParticles.X(FirstTri[1]));
-				MLocalBoundingBox.GrowToInclude(MParticles.X(FirstTri[2]));
+				MLocalBoundingBox = FAABB3(MParticles.GetX(FirstTri[0]), MParticles.GetX(FirstTri[0]));
+				MLocalBoundingBox.GrowToInclude(MParticles.GetX(FirstTri[1]));
+				MLocalBoundingBox.GrowToInclude(MParticles.GetX(FirstTri[2]));
 
 				for(int32 TriangleIndex = 1; TriangleIndex < NumTriangles; ++TriangleIndex)
 				{
 					const TVec3<IdxType>& Tri = Tris[TriangleIndex];
-					MLocalBoundingBox.GrowToInclude(MParticles.X(Tri[0]));
-					MLocalBoundingBox.GrowToInclude(MParticles.X(Tri[1]));
-					MLocalBoundingBox.GrowToInclude(MParticles.X(Tri[2]));
+					MLocalBoundingBox.GrowToInclude(MParticles.GetX(Tri[0]));
+					MLocalBoundingBox.GrowToInclude(MParticles.GetX(Tri[1]));
+					MLocalBoundingBox.GrowToInclude(MParticles.GetX(Tri[2]));
 				}
 			}
 			
@@ -805,7 +805,7 @@ namespace Chaos
 			: FImplicitObject(EImplicitObject::HasBoundingBox | EImplicitObject::DisableCollisions, ImplicitObjectType::TriangleMesh)
 			, MParticles(MoveTemp(Particles))
 			, MElements(MoveTemp(Elements))
-			, MLocalBoundingBox(MParticles.X(0), MParticles.X(0))
+			, MLocalBoundingBox(MParticles.GetX(0), MParticles.GetX(0))
 			, MaterialIndices(MoveTemp(InMaterialIndices))
 			, ExternalFaceIndexMap(MoveTemp(InExternalFaceIndexMap))
 			, ExternalVertexIndexMap(MoveTemp(InExternalVertexIndexMap))
@@ -817,16 +817,16 @@ namespace Chaos
 				const TArray<TVec3<IdxType>>& Tris = MElements.GetIndexBuffer<IdxType>();
 				const TVec3<IdxType>& FirstTri = Tris[0];
 
-				MLocalBoundingBox = FAABB3(MParticles.X(FirstTri[0]), MParticles.X(FirstTri[0]));
-				MLocalBoundingBox.GrowToInclude(MParticles.X(FirstTri[1]));
-				MLocalBoundingBox.GrowToInclude(MParticles.X(FirstTri[2]));
+				MLocalBoundingBox = FAABB3(MParticles.GetX(FirstTri[0]), MParticles.GetX(FirstTri[0]));
+				MLocalBoundingBox.GrowToInclude(MParticles.GetX(FirstTri[1]));
+				MLocalBoundingBox.GrowToInclude(MParticles.GetX(FirstTri[2]));
 
 				for(int32 TriangleIndex = 1; TriangleIndex < NumTriangles; ++TriangleIndex)
 				{
 					const TVec3<IdxType>& Tri = Tris[TriangleIndex];
-					MLocalBoundingBox.GrowToInclude(MParticles.X(Tri[0]));
-					MLocalBoundingBox.GrowToInclude(MParticles.X(Tri[1]));
-					MLocalBoundingBox.GrowToInclude(MParticles.X(Tri[2]));
+					MLocalBoundingBox.GrowToInclude(MParticles.GetX(Tri[0]));
+					MLocalBoundingBox.GrowToInclude(MParticles.GetX(Tri[1]));
+					MLocalBoundingBox.GrowToInclude(MParticles.GetX(Tri[2]));
 				}
 			}
 			
@@ -850,10 +850,10 @@ namespace Chaos
 			{
 				auto LambdaHelper = [&](const auto& Elements)
 				{
-					TAABB<FRealSingle,3> Bounds(TmData->MParticles.X(Elements[Index][0]), TmData->MParticles.X(Elements[Index][0]));
+					TAABB<FRealSingle,3> Bounds(TmData->MParticles.GetX(Elements[Index][0]), TmData->MParticles.GetX(Elements[Index][0]));
 
-					Bounds.GrowToInclude(TmData->MParticles.X(Elements[Index][1]));
-					Bounds.GrowToInclude(TmData->MParticles.X(Elements[Index][2]));
+					Bounds.GrowToInclude(TmData->MParticles.GetX(Elements[Index][1]));
+					Bounds.GrowToInclude(TmData->MParticles.GetX(Elements[Index][2]));
 
 					return Bounds;
 				};
diff --git a/Engine/Source/Runtime/Experimental/Chaos/Public/Chaos/VelocityField.h b/Engine/Source/Runtime/Experimental/Chaos/Public/Chaos/VelocityField.h
--- a/Engine/Source/Runtime/Experimental/Chaos/Public/Chaos/VelocityField.h
+++ b/Engine/Source/Runtime/Experimental/Chaos/Public/Chaos/VelocityField.h
@@ -180,8 +180,8 @@ private:
 
 		// Calculate the normal and the area of the surface exposed to the flow
 		FSolverVec3 N = FSolverVec3::CrossProduct(
-			InParticles.X(Element[1]) - InParticles.X(Element[0]),
-			InParticles.X(Element[2]) - InParticles.X(Element[0]));
+			InParticles.GetX(Element[1]) - InParticles.GetX(Element[0]),
+			InParticles.GetX(Element[2]) - InParticles.GetX(Element[0]));
 		const FSolverReal DoubleArea = N.SafeNormalize();
 
 		// Calculate the direction and the relative velocity of the triangle to the flow
@@ -209,8 +209,8 @@ private:
 
 		// Calculate the normal and the area of the surface exposed to the flow
 		FSolverVec3 N = FSolverVec3::CrossProduct(
-			InParticles.X(Element[1]) - InParticles.X(Element[0]),
-			InParticles.X(Element[2]) - InParticles.X(Element[0]));
+			InParticles.GetX(Element[1]) - InParticles.GetX(Element[0]),
+			InParticles.GetX(Element[2]) - InParticles.GetX(Element[0]));
 		const FSolverReal DoubleArea = N.SafeNormalize();
 
 		// Calculate the direction and the relative velocity of the triangle to the flow
diff --git a/Engine/Source/Runtime/Experimental/Chaos/Public/Chaos/XPBDCorotatedConstraints.h b/Engine/Source/Runtime/Experimental/Chaos/Public/Chaos/XPBDCorotatedConstraints.h
--- a/Engine/Source/Runtime/Experimental/Chaos/Public/Chaos/XPBDCorotatedConstraints.h
+++ b/Engine/Source/Runtime/Experimental/Chaos/Public/Chaos/XPBDCorotatedConstraints.h
@@ -207,7 +207,7 @@ namespace Chaos::Softs
 			PMatrix<T, 3, 3> Result((T)0.);
 			for (int i = 0; i < 3; i++) {
 				for (int c = 0; c < 3; c++) {
-					Result.SetAt(c, i, InParticles.X(MeshConstraints[e][i + 1])[c] - InParticles.X(MeshConstraints[e][0])[c]);
+					Result.SetAt(c, i, InParticles.GetX(MeshConstraints[e][i + 1])[c] - InParticles.GetX(MeshConstraints[e][0])[c]);
 				}
 			}
 			return Result;
@@ -218,7 +218,7 @@ namespace Chaos::Softs
 			PMatrix<T, 3, 3> Result((T)0.);
 			for (int i = 0; i < 3; i++) {
 				for (int c = 0; c < 3; c++) {
-					Result.SetAt(c, i, InParticles.P(MeshConstraints[e][i+1])[c] - InParticles.P(MeshConstraints[e][0])[c]);
+					Result.SetAt(c, i, InParticles.GetP(MeshConstraints[e][i+1])[c] - InParticles.GetP(MeshConstraints[e][0])[c]);
 				}
 			}
 			return Result;
diff --git a/Engine/Source/Runtime/Experimental/Chaos/Public/Chaos/XPBDVolumeConstraints.h b/Engine/Source/Runtime/Experimental/Chaos/Public/Chaos/XPBDVolumeConstraints.h
--- a/Engine/Source/Runtime/Experimental/Chaos/Public/Chaos/XPBDVolumeConstraints.h
+++ b/Engine/Source/Runtime/Experimental/Chaos/Public/Chaos/XPBDVolumeConstraints.h
@@ -33,10 +33,10 @@ namespace Chaos::Softs
 			for (int32 ElementIndex = 0; ElementIndex < InMesh.Num(); ElementIndex++)
 			{
 				TVec4<int32> Constraint = InMesh[ElementIndex];
-				const FSolverVec3& P1 = InParticles.X(Constraint[0]);
-				const FSolverVec3& P2 = InParticles.X(Constraint[1]);
-				const FSolverVec3& P3 = InParticles.X(Constraint[2]);
-				const FSolverVec3& P4 = InParticles.X(Constraint[3]);
+				const FSolverVec3& P1 = InParticles.GetX(Constraint[0]);
+				const FSolverVec3& P2 = InParticles.GetX(Constraint[1]);
+				const FSolverVec3& P3 = InParticles.GetX(Constraint[2]);
+				const FSolverVec3& P4 = InParticles.GetX(Constraint[3]);
 				Volumes[ElementIndex] = FSolverVec3::DotProduct(FSolverVec3::CrossProduct(P2 - P1, P3 - P1), P4 - P1) / (FSolverReal)6.;
 			}
 
diff --git a/Engine/Source/Runtime/Experimental/Chaos/Public/PhysicsProxy/FieldSystemProxyHelper.h b/Engine/Source/Runtime/Experimental/Chaos/Public/PhysicsProxy/FieldSystemProxyHelper.h
--- a/Engine/Source/Runtime/Experimental/Chaos/Public/PhysicsProxy/FieldSystemProxyHelper.h
+++ b/Engine/Source/Runtime/Experimental/Chaos/Public/PhysicsProxy/FieldSystemProxyHelper.h
@@ -95,7 +95,7 @@ namespace Chaos
 						{
 							if (ParentHandle->Disabled() == false)
 							{
-								const FRigidTransform3 ParentWorldTM(ParentHandle->P(), ParentHandle->GetQ());
+								const FRigidTransform3 ParentWorldTM(ParentHandle->GetP(), ParentHandle->GetQ());
 								const FRigidTransform3 ChildFrame = ClusterHandle->ChildToParent() * ParentWorldTM;
 								SamplePosition = ChildFrame.GetTranslation();
 							}
@@ -131,7 +131,7 @@ namespace Chaos
 				{
 					if (Chaos::FGeometryParticleHandle* FilteredHandle = FilteredHandles[Idx])
 					{
-						const FVec3& SamplePosition = FilteredHandle->X();
+						const FVec3& SamplePosition = FilteredHandle->GetX();
 						FillExecutionDatas(SamplePosition, FilteredHandle, HandleIndex);
 					}
 				}
@@ -454,12 +454,12 @@ namespace Chaos
 						if (TargetedParticles.Contains(Index.Sample))
 						{
 							const int32 ConstraintIndex = TargetedParticles[Index.Sample];
-							PositionTarget.Replace(ConstraintIndex, ParticleHandles[Index.Sample]->X());
+							PositionTarget.Replace(ConstraintIndex, ParticleHandles[Index.Sample]->GetX());
 						}
 						else
 						{
 							const int32 ConstraintIndex = PositionTarget.NumConstraints();
-							PositionTarget.AddConstraint(RigidHandle, RigidHandle->X());
+							PositionTarget.AddConstraint(RigidHandle, RigidHandle->GetX());
 							TargetedParticles.Add(Index.Sample, ConstraintIndex);
 						}
 					}
diff --git a/Engine/Source/Runtime/Experimental/Chaos/Public/PhysicsProxy/SingleParticlePhysicsProxy.h b/Engine/Source/Runtime/Experimental/Chaos/Public/PhysicsProxy/SingleParticlePhysicsProxy.h
--- a/Engine/Source/Runtime/Experimental/Chaos/Public/PhysicsProxy/SingleParticlePhysicsProxy.h
+++ b/Engine/Source/Runtime/Experimental/Chaos/Public/PhysicsProxy/SingleParticlePhysicsProxy.h
@@ -208,7 +208,8 @@ public:
 	}
 
 	//API for static particle
-	const FVec3& X() const { return ReadRef([](auto* Particle) -> const auto& { return Particle->X(); }); }
+	const FVec3& X() const { return ReadRef([](auto* Particle) -> const auto& { return Particle->GetX(); }); }
+	const FVec3& GetX() const { return ReadRef([](auto* Particle) -> const auto& { return Particle->GetX(); }); }
 
 protected:
 	void SetXBase(const FVec3& InX, bool bInvalidate = true) { Write([&InX, bInvalidate, this](auto* Particle)
diff --git a/Engine/Source/Runtime/Experimental/Chaos/Public/RewindData.h b/Engine/Source/Runtime/Experimental/Chaos/Public/RewindData.h
--- a/Engine/Source/Runtime/Experimental/Chaos/Public/RewindData.h
+++ b/Engine/Source/Runtime/Experimental/Chaos/Public/RewindData.h
@@ -1002,7 +1002,7 @@ struct FGeometryParticleStateBase
 	template<typename TParticle>
 	void CachePreCorrectionState(const TParticle& Particle)
 	{
-		PreCorrectionXR.SetX(Particle.X());
+		PreCorrectionXR.SetX(Particle.GetX());
 		PreCorrectionXR.SetR(Particle.GetR());
 	}
 
@@ -1043,7 +1043,7 @@ public:
 	}
 
 
-	REWIND_PARTICLE_STATIC_PROPERTY(ParticlePositionRotation, X)
+	REWIND_PARTICLE_STATIC_PROPERTY(ParticlePositionRotation, GetX)
 	REWIND_PARTICLE_STATIC_PROPERTY(ParticlePositionRotation, GetR)
 
 	REWIND_PARTICLE_KINEMATIC_PROPERTY(Velocities, GetV)
@@ -1099,7 +1099,7 @@ public:
 		//TODO: use macro to define api and the to string
 		FString Out = FString::Printf(TEXT("ParticleID:[Global: %d Local: %d]\n"), Particle.ParticleID().GlobalID, Particle.ParticleID().LocalID);
 
-		REWIND_PARTICLE_TO_STR(X)
+		REWIND_PARTICLE_TO_STR(GetX)
 		REWIND_PARTICLE_TO_STR(GetR)
 		//REWIND_PARTICLE_TO_STR(Geometry)
 		//REWIND_PARTICLE_TO_STR(UniqueIdx)
diff --git a/Engine/Source/Runtime/Experimental/ChaosCore/Private/Rotation.cpp b/Engine/Source/Runtime/Experimental/ChaosCore/Private/Rotation.cpp
--- a/Engine/Source/Runtime/Experimental/ChaosCore/Private/Rotation.cpp
+++ b/Engine/Source/Runtime/Experimental/ChaosCore/Private/Rotation.cpp
@@ -21,6 +21,11 @@ namespace Chaos
 #endif
 	}
 
+	TRotation<FRealSingle, 3> TRotation<FRealSingle, 3>::Conjugate(const ::Chaos::TRotation<FRealSingle, 3>& InR)
+	{
+		return InR.Inverse();
+	}
+
 	/**
 	 * Negate all values of the quaternion (note: not the inverse rotation. See Conjugate)
 	 */
@@ -99,6 +104,25 @@ namespace Chaos
 		return TVector<FReal, 3>(W.X, W.Y, W.Z);
 	}
 
+	TVector<FRealSingle, 3> TRotation<FRealSingle, 3>::CalculateAngularVelocity1(const TRotation<FRealSingle, 3>& InR0, const TRotation<FRealSingle, 3>& InR1, const FRealSingle InDt)
+	{
+		if (!ensure(InDt > UE_SMALL_NUMBER))
+		{
+			return TVector<FRealSingle, 3>(0);
+		}
+
+		const TRotation<FRealSingle, 3> R0 = InR0;
+		TRotation<FRealSingle, 3> R1 = InR1;
+		R1.EnforceShortestArcWith(R0);
+
+		// W = 2 * dQ/dT * Qinv
+		const TRotation<FRealSingle, 3> DRDt = (R1 - R0) / InDt;
+		const TRotation<FRealSingle, 3> RInv = Conjugate(R0);
+		const TRotation<FRealSingle, 3> W = (DRDt * RInv) * 2.0f;
+
+		return TVector<FRealSingle, 3>(W.X, W.Y, W.Z);
+	}
+
 	//PRAGMA_ENABLE_OPTIMIZATION
 
 
diff --git a/Engine/Source/Runtime/Experimental/ChaosVehicles/ChaosVehiclesCore/Private/SimModule/DeferredForcesModular.cpp b/Engine/Source/Runtime/Experimental/ChaosVehicles/ChaosVehiclesCore/Private/SimModule/DeferredForcesModular.cpp
--- a/Engine/Source/Runtime/Experimental/ChaosVehicles/ChaosVehiclesCore/Private/SimModule/DeferredForcesModular.cpp
+++ b/Engine/Source/Runtime/Experimental/ChaosVehicles/ChaosVehiclesCore/Private/SimModule/DeferredForcesModular.cpp
@@ -116,7 +116,7 @@ Chaos::FPBDRigidParticleHandle* FDeferredForcesModular::GetParticle(const FTrans
 
 	if (ClusterParticle != nullptr && Child != nullptr)
 	{
-		const FRigidTransform3 ClusterWorldTM(ClusterParticle->X(), ClusterParticle->GetR());
+		const FRigidTransform3 ClusterWorldTM(ClusterParticle->GetX(), ClusterParticle->GetR());
 		
 		Proxy->GetParticle_Internal(TransformIndex);
 
@@ -128,7 +128,7 @@ Chaos::FPBDRigidParticleHandle* FDeferredForcesModular::GetParticle(const FTrans
 		{
 			//FTransform ChildT = Child->GetTransformXRCom();
 			//const FRigidTransform3 ChildWorldTM(ChildT.GetLocation(), ChildT.GetRotation());
-			const FRigidTransform3 ChildWorldTM(Child->X(), Child->GetR());
+			const FRigidTransform3 ChildWorldTM(Child->GetX(), Child->GetR());
 			Frame = ChildWorldTM.GetRelativeTransform(ClusterWorldTM);
 			Frame.SetLocation(Frame.GetLocation() + OffsetTransform.TransformVector(PositionalOffset));
 
@@ -154,7 +154,7 @@ Chaos::FPBDRigidParticleHandle* FDeferredForcesModular::GetParticle(const FTrans
 	const int SingleChassisIndex = 0;
 	TransformOut = FTransform::Identity;
 
-	const FRigidTransform3 ClusterWorldTM(ClusterParticles[SingleChassisIndex]->X(), ClusterParticles[SingleChassisIndex]->GetR());
+	const FRigidTransform3 ClusterWorldTM(ClusterParticles[SingleChassisIndex]->GetX(), ClusterParticles[SingleChassisIndex]->GetR());
 	FRigidTransform3 Frame = FRigidTransform3::Identity;
 
 	if (Chaos::FPBDRigidParticleHandle* Child = GetParticleFromUniqueIndex(ParticleIdx, Particles))
@@ -167,7 +167,7 @@ Chaos::FPBDRigidParticleHandle* FDeferredForcesModular::GetParticle(const FTrans
 		{
 			//FTransform ChildT = Child->GetTransformXRCom();
 			//const FRigidTransform3 ChildWorldTM(ChildT.GetLocation(), ChildT.GetRotation());
-			const FRigidTransform3 ChildWorldTM(Child->X(), Child->GetR());
+			const FRigidTransform3 ChildWorldTM(Child->GetX(), Child->GetR());
 			Frame = ChildWorldTM.GetRelativeTransform(ClusterWorldTM);
 			Frame.SetLocation(Frame.GetLocation() + OffsetTransform.TransformVector(PositionalOffset));
 		}
@@ -321,7 +321,7 @@ void AddForce_Implementation(Chaos::FPBDRigidParticleHandle* RigidHandle, const
 	const Chaos::FVec3 WorldCOM = Chaos::FParticleUtilitiesGT::GetCoMWorldPosition(RigidHandle);
 
 	// local to world
-	const FTransform WorldTM(RigidHandle->GetR(), RigidHandle->X());
+	const FTransform WorldTM(RigidHandle->GetR(), RigidHandle->GetX());
 	Chaos::FVec3 Position = WorldTM.TransformPosition(OffsetTransform.GetLocation());
 	Chaos::FVec3 Force = WorldTM.TransformVector(ParticleOffsetTransform.TransformVector(OffsetTransform.TransformVector(LocalForce)));
 
@@ -349,7 +349,7 @@ void AddForce_Implementation(Chaos::FPBDRigidParticleHandle* RigidHandle, const
 		if (GCoreModularVehicleDebugParams.ShowMass)
 		{
 			FString OutputText = FString::Format(TEXT("{0} {1}"), { RP->M(), RP->I().ToString() });
-			Chaos::FDebugDrawQueue::GetInstance().DrawDebugString(RigidHandle->X() + FVector(0,0,200), OutputText, nullptr, FColor::White, -1, false, -1.0f);
+			Chaos::FDebugDrawQueue::GetInstance().DrawDebugString(RigidHandle->GetX() + FVector(0,0,200), OutputText, nullptr, FColor::White, -1, false, -1.0f);
 		}
 #endif
 	}
@@ -371,7 +371,7 @@ void AddTorque_Implementation(Chaos::FPBDRigidParticleHandle* RigidHandle, const
 	if (RP)
 	{
 		// local to world
-		const FTransform WorldTM(RigidHandle->GetR(), RigidHandle->X());
+		const FTransform WorldTM(RigidHandle->GetR(), RigidHandle->GetX());
 		FVector T1 = OffsetTransform.TransformVector(LocalTorque);
 		FVector T2 = ParticleOffsetTransform.TransformVector(T1);
 		Chaos::FVec3 WorldTorque = WorldTM.TransformVector(T2);
diff --git a/Engine/Source/Runtime/Experimental/ChaosVehicles/ChaosVehiclesCore/Private/SimModule/SimModuleTree.cpp b/Engine/Source/Runtime/Experimental/ChaosVehicles/ChaosVehiclesCore/Private/SimModule/SimModuleTree.cpp
--- a/Engine/Source/Runtime/Experimental/ChaosVehicles/ChaosVehiclesCore/Private/SimModule/SimModuleTree.cpp
+++ b/Engine/Source/Runtime/Experimental/ChaosVehicles/ChaosVehiclesCore/Private/SimModule/SimModuleTree.cpp
@@ -301,7 +301,7 @@ void FSimModuleTree::UpdateModuleVelocites(FGeometryCollectionPhysicsProxy* Phys
 		{
 			if (Chaos::FPBDRigidClusteredParticleHandle* ParentParticle = PhysicsProxy->GetSolverClusterHandle_Internal(Module->GetTransformIndex()))
 			{
-				const FTransform BodyTransform(ParentParticle->GetR(), ParentParticle->X());
+				const FTransform BodyTransform(ParentParticle->GetR(), ParentParticle->GetX());
 
 				if (Module->IsBehaviourType(eSimModuleTypeFlags::Velocity))
 				{
@@ -318,7 +318,7 @@ void FSimModuleTree::UpdateModuleVelocites(FGeometryCollectionPhysicsProxy* Phys
 					if (Particle)
 					{
 						FVector WorldLocation = BodyTransform.TransformPosition(Module->GetParentRelativeTransform().GetLocation());
-						const Chaos::FVec3 Arm = WorldLocation - Particle->X();
+						const Chaos::FVec3 Arm = WorldLocation - Particle->GetX();
 
 						FVector WorldVelocity = Particle->GetV() - Chaos::FVec3::CrossProduct(Arm, Particle->GetW());
 						FVector LocalLinearVelocity = BodyTransform.InverseTransformVector(WorldVelocity);
@@ -358,7 +358,7 @@ void FSimModuleTree::UpdateModuleVelocites(FClusterUnionPhysicsProxy* PhysicsPro
 		{
 			if (ISimulationModuleBase* Module = SimulationModuleTree[i].SimModule)
 			{
-				const FTransform BodyTransform(ParentParticle->GetR(), ParentParticle->X());
+				const FTransform BodyTransform(ParentParticle->GetR(), ParentParticle->GetX());
 
 				if (Module->IsBehaviourType(eSimModuleTypeFlags::Velocity))
 				{
@@ -377,7 +377,7 @@ void FSimModuleTree::UpdateModuleVelocites(FClusterUnionPhysicsProxy* PhysicsPro
 						const FTransform& OffsetTransform = Module->GetComponentTransform();
 						FVector LocalPos = Module->GetParentRelativeTransform().GetLocation();
 						FVector WorldLocation = BodyTransform.TransformPosition(LocalPos);
-						const Chaos::FVec3 Arm = WorldLocation - Particle->X();
+						const Chaos::FVec3 Arm = WorldLocation - Particle->GetX();
 
 						//Chaos::FDebugDrawQueue::GetInstance().DrawDebugLine(Particle->X(), Particle->X() + Arm, FColor::Yellow, false, -1.f, 0, 2.f);
 
@@ -407,7 +407,7 @@ void FSimModuleTree::UpdateVehicleState(FClusterUnionPhysicsProxy* PhysicsProxy)
 
 	if (Chaos::FClusterUnionPhysicsProxy::FInternalParticle* ParentParticle = PhysicsProxy->GetParticle_Internal())
 	{
-		const FTransform BodyTransform(ParentParticle->GetR(), ParentParticle->X());
+		const FTransform BodyTransform(ParentParticle->GetR(), ParentParticle->GetX());
 
 		VehicleState.ForwardDir = BodyTransform.GetUnitAxis(EAxis::X);
 		VehicleState.UpDir = BodyTransform.GetUnitAxis(EAxis::Z);
diff --git a/Engine/Source/Runtime/Experimental/ChaosVisualDebugger/Public/DataWrappers/ChaosVDParticleDataWrapper.h b/Engine/Source/Runtime/Experimental/ChaosVisualDebugger/Public/DataWrappers/ChaosVDParticleDataWrapper.h
--- a/Engine/Source/Runtime/Experimental/ChaosVisualDebugger/Public/DataWrappers/ChaosVDParticleDataWrapper.h
+++ b/Engine/Source/Runtime/Experimental/ChaosVisualDebugger/Public/DataWrappers/ChaosVDParticleDataWrapper.h
@@ -126,7 +126,7 @@ struct FChaosVDParticlePositionRotation : public FChaosVDWrapperDataBase
 	template <typename TOther>
 	void CopyFrom(const TOther& Other)
 	{
-		MX = Other.X();
+		MX = Other.GetX();
 		MR = Other.GetR();
 		
 		bHasValidData = true;
diff --git a/Engine/Source/Runtime/Experimental/GeometryCollectionEngine/Private/GeometryCollection/GeometryCollectionComponent.cpp b/Engine/Source/Runtime/Experimental/GeometryCollectionEngine/Private/GeometryCollection/GeometryCollectionComponent.cpp
--- a/Engine/Source/Runtime/Experimental/GeometryCollectionEngine/Private/GeometryCollection/GeometryCollectionComponent.cpp
+++ b/Engine/Source/Runtime/Experimental/GeometryCollectionEngine/Private/GeometryCollection/GeometryCollectionComponent.cpp
@@ -2241,7 +2241,7 @@ void UGeometryCollectionComponent::UpdateRepData()
 				{
 					FGeometryCollectionClusterRep& ClusterRep = LocalRepData.Clusters.AddDefaulted_GetRef();
 
-					ClusterRep.Position = Cluster->X();
+					ClusterRep.Position = Cluster->GetX();
 					ClusterRep.Rotation = Cluster->GetR();
 					ClusterRep.LinearVelocity = Cluster->GetV();
 					ClusterRep.AngularVelocity = Cluster->GetW();
@@ -2416,7 +2416,7 @@ void UGeometryCollectionComponent::UpdateRepStateAndDynamicData()
 				FGeometryCollectionRepDynamicData::FClusterData& Data = LocalRepDynamicData.ClusterData.AddDefaulted_GetRef();
 				Data.TransformIndex = Root.TransformIndex;
 				Data.bIsInternalCluster = Root.Handle->InternalCluster();
-				Data.Position = Root.Handle->X();
+				Data.Position = Root.Handle->GetX();
 				Data.EulerRotation = Root.Handle->GetR().Rotator().Euler();
 				Data.LinearVelocity = Root.Handle->GetV();
 				Data.AngularVelocityInDegreesPerSecond = FMath::RadiansToDegrees(Root.Handle->GetW());
@@ -2585,7 +2585,7 @@ namespace
 			//
 			const FVector RepVel = LinearVelocity;
 			const FVector RepExtrapPos = Position + (RepVel * RepExtrapTime);
-			const Chaos::FVec3 DeltaX = RepExtrapPos - Cluster.X();
+			const Chaos::FVec3 DeltaX = RepExtrapPos - Cluster.GetX();
 			const float DeltaXMagSq = DeltaX.SizeSquared();
 			if (DeltaXMagSq > SMALL_NUMBER && GeometryCollectionRepLinearMatchStrength > SMALL_NUMBER)
 			{
@@ -2936,7 +2936,7 @@ bool UGeometryCollectionComponent::ProcessRepData(const float DeltaTime, const f
 		ForEachClusterPair([RepExtrapTime](const FGeometryCollectionClusterRep& RepCluster, Chaos::FPBDRigidParticleHandle& Cluster)
 		{
 			// Don't bother debug drawing if the delta is too small
-			if ((Cluster.X() - RepCluster.Position).SizeSquared() < .1f)
+			if ((Cluster.GetX() - RepCluster.Position).SizeSquared() < .1f)
 			{
 				FVector Axis;
 				float Angle;
@@ -2948,8 +2948,8 @@ bool UGeometryCollectionComponent::ProcessRepData(const float DeltaTime, const f
 			}
 
 			Chaos::FDebugDrawQueue& DrawQueue = Chaos::FDebugDrawQueue::GetInstance();
-			DrawQueue.DrawDebugCoordinateSystem(Cluster.X(), FRotator(Cluster.GetR()), 100.f, false, -1, -1, 1.f);
-			DrawQueue.DrawDebugBox(Cluster.X() + Cluster.LocalBounds().Center(), Cluster.LocalBounds().Extents(), Cluster.GetR(), FColor::White, false, -1, -1, 1.f);
+			DrawQueue.DrawDebugCoordinateSystem(Cluster.GetX(), FRotator(Cluster.GetR()), 100.f, false, -1, -1, 1.f);
+			DrawQueue.DrawDebugBox(Cluster.GetX() + Cluster.LocalBounds().Center(), Cluster.LocalBounds().Extents(), Cluster.GetR(), FColor::White, false, -1, -1, 1.f);
 			DrawQueue.DrawDebugBox(RepCluster.Position + Cluster.LocalBounds().Center(), Cluster.LocalBounds().Extents(), RepCluster.Rotation, FColor::Green, false, -1, -1, 1.f);
 
 			if (bGeometryCollectionRepUseClusterVelocityMatch)
@@ -2959,7 +2959,7 @@ bool UGeometryCollectionComponent::ProcessRepData(const float DeltaTime, const f
 				const FVector RepExtrapPos = RepCluster.Position + (RepVel * RepExtrapTime);
 				const Chaos::FRotation3 RepExtrapAng = Chaos::FRotation3::IntegrateRotationWithAngularVelocity(RepCluster.Rotation, RepAngVel, RepExtrapTime);
 				DrawQueue.DrawDebugCoordinateSystem(RepExtrapPos, FRotator(RepExtrapAng), 100.f, false, -1, -1, 1.f);
-				DrawQueue.DrawDebugDirectionalArrow(Cluster.X(), RepExtrapPos, 10.f, FColor::White, false, -1, -1, 1.f);
+				DrawQueue.DrawDebugDirectionalArrow(Cluster.GetX(), RepExtrapPos, 10.f, FColor::White, false, -1, -1, 1.f);
 				DrawQueue.DrawDebugBox(RepExtrapPos + Cluster.LocalBounds().Center(), Cluster.LocalBounds().Extents(), RepExtrapAng, FColor::Orange, false, -1, -1, 1.f);
 			}
 			else
diff --git a/Engine/Source/Runtime/NavigationSystem/Private/NavMesh/RecastNavMeshGenerator.cpp b/Engine/Source/Runtime/NavigationSystem/Private/NavMesh/RecastNavMeshGenerator.cpp
--- a/Engine/Source/Runtime/NavigationSystem/Private/NavMesh/RecastNavMeshGenerator.cpp
+++ b/Engine/Source/Runtime/NavigationSystem/Private/NavMesh/RecastNavMeshGenerator.cpp
@@ -384,7 +384,7 @@ void ExportChaosTriMesh(const Chaos::FTriangleMeshImplicitObject* const TriMesh,
 		{
 			for (int32 i = 0; i < 3; i++)
 			{
-				const FVector UnrealCoords = LocalToWorld.TransformPosition((FVector)Vertices.X(Triangles[TriIdx][i]));
+				const FVector UnrealCoords = LocalToWorld.TransformPosition((FVector)Vertices.GetX(Triangles[TriIdx][i]));
 				UnrealBounds += UnrealCoords;
 
 				VertexBuffer.Add(UnrealCoords.X);
diff --git a/Engine/Source/Runtime/PhysicsCore/Private/PhysicsInterfaceUtilsCore.cpp b/Engine/Source/Runtime/PhysicsCore/Private/PhysicsInterfaceUtilsCore.cpp
--- a/Engine/Source/Runtime/PhysicsCore/Private/PhysicsInterfaceUtilsCore.cpp
+++ b/Engine/Source/Runtime/PhysicsCore/Private/PhysicsInterfaceUtilsCore.cpp
@@ -29,7 +29,7 @@ FCollisionFilterData ToUnrealFilterData(const FChaosFilterData& FilterData)
 template <typename THitLocation>
 uint32 FindFaceIndexImp(const THitLocation& PHit, const FVector& UnitDir)
 {
-	const FTransform WorldTM(PHit.Actor->GetR(), PHit.Actor->X());
+	const FTransform WorldTM(PHit.Actor->GetR(), PHit.Actor->GetX());
 	const FVector LocalPosition = WorldTM.InverseTransformPositionNoScale(PHit.WorldPosition);
 	const FVector LocalUnitDir = WorldTM.InverseTransformVectorNoScale(UnitDir);
 	return PHit.Shape->GetGeometry()->FindMostOpposingFace(LocalPosition, LocalUnitDir, PHit.FaceIndex, 1);	//todo:this number matches the one above, but is it right?
diff --git a/Engine/Source/Runtime/PhysicsCore/Public/SQVisitor.h b/Engine/Source/Runtime/PhysicsCore/Public/SQVisitor.h
--- a/Engine/Source/Runtime/PhysicsCore/Public/SQVisitor.h
+++ b/Engine/Source/Runtime/PhysicsCore/Public/SQVisitor.h
@@ -201,7 +201,7 @@ private:
 		const bool bTestShapeBounds = Shapes.Num() > 1;
 		bool bContinue = true;
 
-		const FRigidTransform3 ActorTM(GeometryParticle->X(), GeometryParticle->GetR());
+		const FRigidTransform3 ActorTM(GeometryParticle->GetX(), GeometryParticle->GetR());
 		const TAABB<FReal, 3> QueryGeomWorldBounds = QueryGeom ? QueryGeom->CalculateTransformedBounds(StartTM) : TAABB<FReal, 3>(-HalfExtents, HalfExtents);
 
 #if CHAOS_DEBUG_DRAW
@@ -215,7 +215,7 @@ private:
 
 			const FImplicitObject* Geom = Shape->GetGeometry();
 
-			CVD_TRACE_SCOPED_SCENE_QUERY_VISIT_HELPER(EChaosVDSceneQueryVisitorType::NarrowPhase, FTransform(GeometryParticle->GetR(), GeometryParticle->X()), Payload.UniqueIdx().Idx, ShapeIndex, CurData);
+			CVD_TRACE_SCOPED_SCENE_QUERY_VISIT_HELPER(EChaosVDSceneQueryVisitorType::NarrowPhase, FTransform(GeometryParticle->GetR(), GeometryParticle->GetX()), Payload.UniqueIdx().Idx, ShapeIndex, CurData);
 
 			if (bTestShapeBounds)
 			{
@@ -564,7 +564,7 @@ private:
 		{
 			const TUniquePtr<FPerShapeData>& Shape = Shapes[ShapeIndex];
 
-			CVD_TRACE_SCOPED_SCENE_QUERY_VISIT_HELPER(EChaosVDSceneQueryVisitorType::BroadPhase, FTransform(GeometryParticle->GetR(), GeometryParticle->X()), Payload.UniqueIdx().Idx, ShapeIndex, CurData);
+			CVD_TRACE_SCOPED_SCENE_QUERY_VISIT_HELPER(EChaosVDSceneQueryVisitorType::BroadPhase, FTransform(GeometryParticle->GetR(), GeometryParticle->GetX()), Payload.UniqueIdx().Idx, ShapeIndex, CurData);
 
 			ECollisionQueryHitType HitType = QueryFilterData.flags & FPhysicsQueryFlag::ePREFILTER ? QueryCallback.PreFilter(QueryFilterDataConcrete, *Shape, *GeometryParticle) : ECollisionQueryHitType::Block;
 			if (HitType != ECollisionQueryHitType::None)
