[ChaosVD] Fix issue where part of an old recording might be added to a new recording if it was done between PIE sessions

diff --git a/Engine/Source/Runtime/Experimental/ChaosVisualDebugger/Private/ChaosVDRuntimeModule.cpp b/Engine/Source/Runtime/Experimental/ChaosVisualDebugger/Private/ChaosVDRuntimeModule.cpp
--- a/Engine/Source/Runtime/Experimental/ChaosVisualDebugger/Private/ChaosVDRuntimeModule.cpp
+++ b/Engine/Source/Runtime/Experimental/ChaosVisualDebugger/Private/ChaosVDRuntimeModule.cpp
@@ -170,12 +170,15 @@ void FChaosVDRuntimeModule::StartRecording(TConstArrayView<FString> Args)
 	UE::Trace::ToggleChannel(TEXT("ChaosVDChannel"), true); 
 	UE::Trace::ToggleChannel(TEXT("Frame"), true);
 
+	FTraceAuxiliary::FOptions TracingOptions;
+	TracingOptions.bExcludeTail = true;
+
 	if (Args.Num() == 0 || Args[0] == TEXT("File"))
 	{
 		ActiveRecordingFileName.Empty();
 		GenerateRecordingFileName(ActiveRecordingFileName);
 
-		bIsRecording = FTraceAuxiliary::Start(FTraceAuxiliary::EConnectionType::File, *ActiveRecordingFileName);
+		bIsRecording = FTraceAuxiliary::Start(FTraceAuxiliary::EConnectionType::File, *ActiveRecordingFileName, nullptr, &TracingOptions);
 	}
 	else if(Args[0] == TEXT("Server"))
 	{
@@ -189,7 +192,7 @@ void FChaosVDRuntimeModule::StartRecording(TConstArrayView<FString> Args)
 		bIsRecording = FTraceAuxiliary::Start(
 		FTraceAuxiliary::EConnectionType::Network,
 		*Target,
-		nullptr);
+		nullptr, &TracingOptions);
 	}
 #endif
 	
