Chaos Cloth: Kinematic colliders optimizations.

Moved kinematic colliders into their own spatial hash acceleration structure.
Store kinematic collider constraints in a way that they can be applied with ISPC.

diff --git a/Engine/Plugins/Experimental/ChaosCloth/Source/Chaos/Private/ChaosCloth/ChaosClothConstraints.cpp b/Engine/Plugins/Experimental/ChaosCloth/Source/Chaos/Private/ChaosCloth/ChaosClothConstraints.cpp
--- a/Engine/Plugins/Experimental/ChaosCloth/Source/Chaos/Private/ChaosCloth/ChaosClothConstraints.cpp
+++ b/Engine/Plugins/Experimental/ChaosCloth/Source/Chaos/Private/ChaosCloth/ChaosClothConstraints.cpp
@@ -1595,7 +1595,7 @@ void FClothConstraints::CreateForceBasedRules()
 			}
 			// Thickness * 2 to account for collision radius for both particles
 			SelfCollisionInit->Init(Particles, SelfCollisionConstraints->GetThickness() * (Softs::FSolverReal)2.f);
-			SelfCollisionConstraints->Init(Particles, Dt, SelfCollisionInit->GetCollidableSubMesh(), SelfCollisionInit->GetSpatialHash(), SelfCollisionInit->GetVertexGIAColors(), SelfCollisionInit->GetTriangleGIAColors());
+			SelfCollisionConstraints->Init(Particles, Dt, SelfCollisionInit->GetCollidableSubMesh(), SelfCollisionInit->GetDynamicSpatialHash(), SelfCollisionInit->GetKinematicColliderSpatialHash(), SelfCollisionInit->GetVertexGIAColors(), SelfCollisionInit->GetTriangleGIAColors());
 		});
 
 		constexpr bool bPostCollisions = true;
@@ -2001,7 +2001,7 @@ void FClothConstraints::CreatePBDRules()
 				}
 				// Thickness * 2 to account for collision radius for both particles
 				SelfCollisionInit->Init(Particles, SelfCollisionConstraints->GetThickness() * (Softs::FSolverReal)2.f);
-				SelfCollisionConstraints->Init(Particles, Dt, SelfCollisionInit->GetCollidableSubMesh(), SelfCollisionInit->GetSpatialHash(), SelfCollisionInit->GetVertexGIAColors(), SelfCollisionInit->GetTriangleGIAColors());
+				SelfCollisionConstraints->Init(Particles, Dt, SelfCollisionInit->GetCollidableSubMesh(), SelfCollisionInit->GetDynamicSpatialHash(), SelfCollisionInit->GetKinematicColliderSpatialHash(), SelfCollisionInit->GetVertexGIAColors(), SelfCollisionInit->GetTriangleGIAColors());
 			};
 
 		PostCollisionConstraintRules[PostCollisionConstraintRuleIndex++] =
diff --git a/Engine/Source/Runtime/Core/Public/Math/Vector.isph b/Engine/Source/Runtime/Core/Public/Math/Vector.isph
--- a/Engine/Source/Runtime/Core/Public/Math/Vector.isph
+++ b/Engine/Source/Runtime/Core/Public/Math/Vector.isph
@@ -743,6 +743,28 @@ inline void VectorScatter(uniform FVector4f *varying DstPtr, const varying FVect
 	}
 }
 
+inline varying FIntVector VectorGather(const uniform FIntVector *varying SrcPtr)
+{
+	varying FIntVector Result;
+
+	if (((1 << TARGET_WIDTH) - 1 ^ lanemask()) == 0)
+	{
+#if TARGET_WIDTH == 4
+		AosToSoa3Explicit((uniform int32* varying)SrcPtr, (uniform int32* uniform)&Result);
+#else
+		UniformLoad(SrcPtr, (uniform FIntVector* uniform)&Result);
+		aos_to_soa3_ispc((uniform int32* uniform)&Result, &Result.V[0], &Result.V[1], &Result.V[2]);
+#endif
+	}
+	else
+	{
+		#pragma ignore warning(perf)
+		Result = *SrcPtr;
+	}
+
+	return Result;
+}
+
 inline FVector4d operator+(const FVector4d &A, const FVector4d &B)
 {
 	FVector4d Result;
@@ -4018,20 +4040,30 @@ inline uniform T VectorReduceMax(const T& A)
 	return SetVector(reduce_max(A.V[0]), reduce_max(A.V[1]), reduce_max(A.V[2]));
 }
 
-inline FVector3d VectorGetSafeNormal(const FVector3d &V)
+inline FVector3d VectorGetSafeNormal(const FVector3d &V, const uniform FVector3d &ResultIfZero)
 {
 	//We want N / ||N|| and to avoid inf
 	//So we want N / ||N|| < 1 / eps => N eps < ||N||, but this is clearly true for all eps < 1 and N > 0
 	double SizeSqr = VectorSizeSquared(V);
-	return VectorSelect(SizeSqr <= DBL_MIN, DoubleForwardVector, V / sqrt(SizeSqr));
+	return VectorSelect(SizeSqr <= DBL_MIN, ResultIfZero, V / sqrt(SizeSqr));
 }
 
-inline FVector3f VectorGetSafeNormal(const FVector3f &V)
+inline FVector3f VectorGetSafeNormal(const FVector3f &V, const uniform FVector3f &ResultIfZero)
 {
 	//We want N / ||N|| and to avoid inf
 	//So we want N / ||N|| < 1 / eps => N eps < ||N||, but this is clearly true for all eps < 1 and N > 0
 	float SizeSqr = VectorSizeSquared(V);
-	return VectorSelect(SizeSqr <= FLT_MIN, FloatForwardVector, V / sqrt(SizeSqr));
+	return VectorSelect(SizeSqr <= FLT_MIN, ResultIfZero, V / sqrt(SizeSqr));
+}
+
+inline FVector3d VectorGetSafeNormal(const FVector3d &V)
+{
+	return VectorGetSafeNormal(V, DoubleForwardVector);
+}
+
+inline FVector3f VectorGetSafeNormal(const FVector3f &V)
+{
+	return VectorGetSafeNormal(V, FloatForwardVector);
 }
 
 #endif
diff --git a/Engine/Source/Runtime/Experimental/Chaos/Private/Chaos/Deformable/ChaosDeformableSolver.cpp b/Engine/Source/Runtime/Experimental/Chaos/Private/Chaos/Deformable/ChaosDeformableSolver.cpp
--- a/Engine/Source/Runtime/Experimental/Chaos/Private/Chaos/Deformable/ChaosDeformableSolver.cpp
+++ b/Engine/Source/Runtime/Experimental/Chaos/Private/Chaos/Deformable/ChaosDeformableSolver.cpp
@@ -1351,7 +1351,7 @@ namespace Chaos::Softs
 				
 				this->TriangleMeshCollisions->Init(InParticles);
 				TArray<FPBDTriangleMeshCollisions::FGIAColor> EmptyGIAColors;
-				this->GSWeakConstraints->CollisionDetectionSpatialHash(this->Evolution->Particles(), *SurfaceTriangleMesh, ParticleComponentIndex, TriangleMeshCollisions->GetSpatialHash(), Property.CollisionSearchRadius, Property.CollisionSpringStiffness, Property.bAllowSliding);
+				this->GSWeakConstraints->CollisionDetectionSpatialHash(this->Evolution->Particles(), *SurfaceTriangleMesh, ParticleComponentIndex, TriangleMeshCollisions->GetDynamicSpatialHash(), Property.CollisionSearchRadius, Property.CollisionSpringStiffness, Property.bAllowSliding);
 				
 			};
 
diff --git a/Engine/Source/Runtime/Experimental/Chaos/Private/Chaos/PBDCollisionSpringConstraints.ispc b/Engine/Source/Runtime/Experimental/Chaos/Private/Chaos/PBDCollisionSpringConstraints.ispc
new file mode 100644
--- /dev/null
+++ b/Engine/Source/Runtime/Experimental/Chaos/Private/Chaos/PBDCollisionSpringConstraints.ispc
@@ -0,0 +1,269 @@
+// Copyright Epic Games, Inc. All Rights Reserved.
+
+#define EXPLICIT_VECTOR4 1
+
+#include "Math/Vector.isph"
+#include "Chaos/PBDSofts.isph"
+
+static inline FVector3f ProjectToP01(const FVector3f& P0, const FVector3f& P1, float Size10, float ProjP1, FVector3f& Bary)
+{
+	Bary.V[2] = 0.f;
+	Bary.V[1] = clamp(ProjP1 / Size10, 0.f, 1.f);
+	Bary.V[0] = 1.f - Bary.V[1];
+	return Bary.V[0] * P0 + Bary.V[1] * P1;
+}
+
+static inline FVector3f ProjectToP02(const FVector3f& P0, const FVector3f& P2, float Size20, float ProjP2, FVector3f& Bary)
+{
+	Bary.V[1] = 0.f;
+	Bary.V[2] = clamp(ProjP2 / Size20, 0.f, 1.f);
+	Bary.V[0] = 1.f - Bary.V[2];
+	return Bary.V[0] * P0 + Bary.V[2] * P2;
+}
+
+static inline FVector3f ProjectToP12(const FVector3f& P1, const FVector3f& P2, const FVector3f& P, FVector3f& Bary)
+{
+	const FVector3f P2P1 = P2 - P1;
+
+	Bary.V[0] = 0.f;
+	Bary.V[2] = clamp(VectorDot(P2P1, P - P1) / VectorSizeSquared(P2P1), 0.f, 1.f);
+	Bary.V[1] = 1.f - Bary.V[2];
+	return Bary.V[1] * P1 + Bary.V[2] * P2;
+}
+
+static inline FVector3f GetTriangleNormalFlipped(const FVector3f&P0, const FVector3f& P1, const FVector3f& P2)
+{
+	// This will return -Chaos::TTriangle<float>::GetNormal.
+	return VectorGetSafeNormal(VectorCross(P2 - P0, P1 - P0), FloatZeroVector);
+}
+
+static inline FVector3f FindClosestPointAndBaryOnTriangle(const FVector3f& P0, const FVector3f& P1, const FVector3f& P2, const FVector3f& P, FVector3f& Bary)
+{
+	const FVector3f P10 = P1 - P0;
+	const FVector3f P20 = P2 - P0;
+	const FVector3f P21 = P2 - P1;
+	const FVector3f PP0 = P - P0;
+
+	const float Size10 = VectorSizeSquared(P10);
+	const float Size20 = VectorSizeSquared(P20);
+	const float Size21 = VectorSizeSquared(P21);
+
+	const float ProjP1 = VectorDot(PP0, P10);
+	const float ProjP2 = VectorDot(PP0, P20);
+
+	
+	// Degenerate triangles
+	if (Size10 < FLOAT_SMALL_NUMBER)
+	{
+		if (Size20 < FLOAT_SMALL_NUMBER)
+		{
+			// Triangle is (nearly) a single point.
+			Bary.V[0] = 1.f;
+			Bary.V[1] = Bary.V[2] = 0.f;
+			return P0;
+		}
+
+		// Triangle is a line segment from P0(=P1) to P2. Project to that line segment.
+		return ProjectToP02(P0, P2, Size20, ProjP2, Bary);
+	}
+	if (Size20 < FLOAT_SMALL_NUMBER)
+	{
+		// Triangle is a line segment from P0(=P2) to P1. Project to that line segment.
+		return ProjectToP01(P0, P1, Size10, ProjP1, Bary);
+	}
+
+
+	const float ProjSides = VectorDot(P10, P20);
+	const float Denom = Size10 * Size20 - ProjSides * ProjSides;
+
+	if (Denom < FLOAT_SMALL_NUMBER)
+	{
+		// Triangle is a line segment from P0 to P1(=P2), or otherwise the 3 points are (nearly) colinear. Project to the longest edge.
+		if (Size21 > Size20)
+		{
+			return ProjectToP12(P1, P2, P, Bary);
+		}
+		else if (Size20 > Size10)
+		{
+			return ProjectToP02(P2, P2, Size20, ProjP2, Bary);
+		}
+		else
+		{
+			return ProjectToP01(P0, P1, Size10, ProjP1, Bary);
+		}
+	}
+
+	// Non-degenerate triangle
+
+	// Numerators of barycentric coordinates if P is inside triangle
+	const float BaryYNum = (Size20 * ProjP1 - ProjSides * ProjP2);
+	const float BaryZNum = (Size10 * ProjP2 - ProjSides * ProjP1);
+
+	// Using this specific epsilon to have parity with FindClosestPointOnTriangle
+	const uniform float Epsilon = 1e-4f;
+	const float EpsilonScaledByDenom = Epsilon * Denom;
+
+	if (BaryYNum + BaryZNum > Denom * (1 + Epsilon)) // i.e., Bary.X < -Epsilon.
+	{
+		if (BaryYNum < -EpsilonScaledByDenom)
+		{
+			// Bary.X < 0, Bary.Y < 0, (Bary.Z > 1)
+
+			// We're outside the triangle in the region opposite angle P2. (if angle at P2 is > 90 degrees, closest point could be on one of these segments)
+			// Should project onto line segment P02 or P12
+			if (ProjP2 < Size20)
+			{
+				// Closer to P02 
+				// Project to line segment from P02
+				return ProjectToP02(P0, P2, Size20, ProjP2, Bary);
+			}
+		} 
+		else if (BaryZNum < -EpsilonScaledByDenom)
+		{
+			// Bary.X < 0, Bary.Y > 1, Bary.Z < 0
+			
+			// We're outside the triangle in the region opposite the angle P1.
+			// Should project onto line segment P01 or P12
+			if (ProjP1 < Size10)
+			{
+				// Closer to P01.
+				// Project to line segment from P01
+				return ProjectToP01(P0, P1, Size10, ProjP1, Bary);
+			}
+		}
+
+		// Bary.X < 0, Bary.Y >=0, Bary.Z >= 0
+
+		// We're outside the triangle in the region closest to P12
+		// Project to line segment from P12
+		return ProjectToP12(P1, P2, P, Bary);
+	}
+	else if (BaryYNum < -EpsilonScaledByDenom)
+	{
+		if (BaryZNum < -EpsilonScaledByDenom)
+		{
+			// Bary.X > 1, Bary.Y < 0, Bary.Z < 0
+
+			// We're outside the triangle in the region opposite P0.
+			// Should project onto line segment P01 or P02 
+			if (ProjP1 > 0.f)
+			{
+				// Closer to P01.
+				// Project to line segment from P01
+				return ProjectToP01(P0, P1, Size10, ProjP1, Bary);
+			}
+		}
+
+		// Bary.X >= 0, Bary.Y < 0, Bary.Z >= 0
+		// We're outside the triangle in region closest to P02.
+		// Project to line segment from P02
+		return ProjectToP02(P0, P2, Size20, ProjP2, Bary);
+	}
+	else if (BaryZNum < -EpsilonScaledByDenom)
+	{
+		// Bary.X >=0, Bary.Y >=0, Bary.Z < 0
+		// We're outside the triangle in the region closest to P01
+		// Project to line segment from P0 to P1
+		return ProjectToP01(P0, P1, Size10, ProjP1, Bary);
+	}
+	
+	// Bary.X >= 0, Bary.Y >= 0, Bary.Z >= 0
+	// Interior of triangle!
+	Bary.V[1] = BaryYNum / Denom;
+	Bary.V[2] = BaryZNum / Denom;
+	Bary.V[0] = 1.f - Bary.V[1] - Bary.V[2];
+	return Bary.V[0] * P0 + Bary.V[1] * P1 + Bary.V[2] * P2;
+}
+
+export void ApplyKinematicCollisionSpringConstraints(
+			uniform FVector4f PandInvM[],
+			const uniform FVector3f ParticlesX[],
+			const uniform int32 KinematicCollidingParticles[],
+			const uniform FIntVector KinematicColliderElements[],
+			const uniform FIntVector TriangleMeshElements[],
+			const uniform float Height,
+			const uniform float KinematicColliderStiffness, 
+			const uniform float KinematicColliderFrictionCoefficient, 
+			const uniform int32 NumConstraints)
+{
+	foreach(Index = 0...NumConstraints)
+	{
+		const int32 Index1 = KinematicCollidingParticles[Index];
+
+		const FVector4f PandInvM1 = VectorGather(&PandInvM[Index1]);
+		varying FVector3f P1;
+		varying float IM;
+		UnzipPandInvM(PandInvM1, P1, IM);
+
+		const varying FIntVector Elements = VectorLoad(&KinematicColliderElements[extract(Index,0)]);
+
+		for(uniform int32 EIndex = 0; EIndex < 3; ++EIndex)
+		{
+			const varying int32 ElementIndex = Elements.V[EIndex];
+			if(ElementIndex == -1)
+			{
+				break;
+			}
+
+			const FIntVector Element = VectorGather(&TriangleMeshElements[ElementIndex]);
+
+			const varying int32 Index2 = Element.V[0];
+			const varying int32 Index3 = Element.V[1];
+			const varying int32 Index4 = Element.V[2];
+		
+			const varying FVector4f PandInvM2 = VectorGather(&PandInvM[Index2]);
+			const varying FVector4f PandInvM3 = VectorGather(&PandInvM[Index3]);
+			const varying FVector4f PandInvM4 = VectorGather(&PandInvM[Index4]);
+
+			varying FVector3f P2, P3, P4;
+			varying float IM2, IM3, IM4;
+		
+			UnzipPandInvM(PandInvM2, P2, IM2);
+			UnzipPandInvM(PandInvM3, P3, IM3);
+			UnzipPandInvM(PandInvM4, P4, IM4);
+
+			const varying FVector3f Normal = GetTriangleNormalFlipped(P2, P3, P4);
+
+			varying FVector3f Bary;
+			const varying FVector3f P = FindClosestPointAndBaryOnTriangle(P2, P3, P4, P1, Bary);
+			const varying FVector3f Difference = P1 - P;
+			const varying float NormalDifference = VectorDot(Difference, Normal);
+			if(NormalDifference > Height)
+			{
+				continue;
+			}
+
+			const varying float TangentialDifference = VectorSize(Difference - NormalDifference * Normal);
+			const uniform float TangentialFalloffMultiplier = 1.5f;
+			const varying float TangentialFalloff = 1.f - TangentialDifference / max(TangentialFalloffMultiplier * Height, FLOAT_KINDA_SMALL_NUMBER);
+			if(TangentialFalloff <= 0.f)
+			{
+				continue;
+			}
+
+			const float NormalDelta = Height - NormalDifference;
+			const FVector3f RepulsionDelta = KinematicColliderStiffness * NormalDelta * Normal;
+
+			P1 = P1 + RepulsionDelta;
+
+			if (KinematicColliderFrictionCoefficient > 0)
+			{
+				const FVector3f X1 = VectorGather(&ParticlesX[Index1]);
+				const FVector3f X2 = VectorGather(&ParticlesX[Index2]);
+				const FVector3f X3 = VectorGather(&ParticlesX[Index3]);
+				const FVector3f X4 = VectorGather(&ParticlesX[Index4]);
+				const FVector3f X = Bary.V[0] * X2 + Bary.V[1] * X3 + Bary.V[2] * X4;
+				const FVector3f RelativeDisplacement = (P1 - X1) - (P - X);
+				const FVector3f RelativeDisplacementTangent = RelativeDisplacement - VectorDot(RelativeDisplacement,Normal) * Normal;
+				const float RelativeDisplacementTangentLength = VectorSize(RelativeDisplacementTangent);
+				const float PositionCorrection = min(NormalDelta * KinematicColliderFrictionCoefficient, RelativeDisplacementTangentLength);
+				const float CorrectionRatio = RelativeDisplacementTangentLength < FLOAT_SMALL_NUMBER ? 0.f : PositionCorrection / RelativeDisplacementTangentLength;
+				const FVector3f FrictionDelta = -CorrectionRatio * RelativeDisplacementTangent;
+
+				P1 = P1 + FrictionDelta;
+			}
+		}
+
+		VectorScatter(&PandInvM[Index1], SetVector4(P1, IM));
+	}
+}
diff --git a/Engine/Source/Runtime/Experimental/Chaos/Private/Chaos/PBDCollisionSpringConstraintsBase.cpp b/Engine/Source/Runtime/Experimental/Chaos/Private/Chaos/PBDCollisionSpringConstraintsBase.cpp
--- a/Engine/Source/Runtime/Experimental/Chaos/Private/Chaos/PBDCollisionSpringConstraintsBase.cpp
+++ b/Engine/Source/Runtime/Experimental/Chaos/Private/Chaos/PBDCollisionSpringConstraintsBase.cpp
@@ -7,20 +7,31 @@
 #include "Chaos/PBDSoftsSolverParticles.h"
 #include "Chaos/SoftsEvolutionLinearSystem.h"
 #include "Chaos/SoftsSolverParticlesRange.h"
+#if INTEL_ISPC
+#include "PBDCollisionSpringConstraints.ispc.generated.h"
+#endif
 
 #if !COMPILE_WITHOUT_UNREAL_SUPPORT
 #include "Chaos/Framework/Parallel.h"
 #include <atomic>
 
-namespace Chaos::Softs {
-
-int32 DebugParticleIndex = INDEX_NONE;
-FAutoConsoleVariableRef CVarChaosClothDebugParticleIndex(TEXT("p.ChaosCloth.DebugParticleIndex"), DebugParticleIndex, TEXT("DebugParticleIndex"));
+#if INTEL_ISPC && !UE_BUILD_SHIPPING
+static_assert(sizeof(ispc::FVector3f) == sizeof(Chaos::Softs::FSolverVec3), "sizeof(ispc::FVector3f) != sizeof(Chaos::Softs::FSolverVec3)");
+static_assert(sizeof(ispc::FVector4f) == sizeof(Chaos::Softs::FPAndInvM), "sizeof(ispc::FVector4f) != sizeof(Chaos::Softs::FPAndInvM");
+static_assert(sizeof(ispc::FIntVector) == sizeof(Chaos::TVec3<int32>), "sizeof(ispc::FIntVector) != sizeof(Chaos::TVec3<int32>");
 
-FSolverReal Chaos_CollisionSpring_MaxTimer = (FSolverReal)0.1f;
-FAutoConsoleVariableRef CVarChaosCollisionSPringMaxTimer(TEXT("p.Chaos.CollisionSpring.MaxTimer"), Chaos_CollisionSpring_MaxTimer, TEXT("Max Timer"));
+bool bChaos_CollisionSpring_ISPC_Enabled = true;
+FAutoConsoleVariableRef CVarChaosCollisionSpringsISPCEnabled(TEXT("p.Chaos.CollisionSpring.ISPC"), bChaos_CollisionSpring_ISPC_Enabled, TEXT("Whether to use ISPC optimizations in collision spring constraints"));
+#endif
 
+#if UE_BUILD_SHIPPING
+static constexpr Chaos::Softs::FSolverReal KinematicColliderMaxTimer = (Chaos::Softs::FSolverReal)0.1f;
+#else
+Chaos::Softs::FSolverReal KinematicColliderMaxTimer = (Chaos::Softs::FSolverReal)0.1f;
+FAutoConsoleVariableRef CVarChaosCollisionSpringMaxTimer(TEXT("p.Chaos.CollisionSpring.MaxTimer"), KinematicColliderMaxTimer, TEXT("Amount of time (in seconds) to remember a kinematic collision connection after it has moved more than Thickness away. Increasing this can reduce jitter at the cost of more computation."));
+#endif
 
+namespace Chaos::Softs {
 
 static TConstArrayView<FSolverVec3> GetConstArrayView(const FSolverParticles& Particles, int32 Offset, int32 NumParticles, const TArray<FSolverVec3>* const Data)
 {
@@ -106,32 +117,50 @@ void FPBDCollisionSpringConstraintsBase::UpdateCollisionLayers(const TConstArray
 }
 
 template<typename SpatialAccelerator, typename SolverParticlesOrRange>
-void FPBDCollisionSpringConstraintsBase::Init(const SolverParticlesOrRange& Particles, const FSolverReal Dt, const FPBDTriangleMeshCollisions::FTriangleSubMesh& CollidableSubMesh, const SpatialAccelerator& Spatial,
-	const TConstArrayView<FPBDTriangleMeshCollisions::FGIAColor>& VertexGIAColors, const TArray<FPBDTriangleMeshCollisions::FGIAColor>& TriangleGIAColors)
+void FPBDCollisionSpringConstraintsBase::Init(const SolverParticlesOrRange& Particles, const SpatialAccelerator& Spatial, const TConstArrayView<FPBDTriangleMeshCollisions::FGIAColor>& VertexGIAColors, const TArray<FPBDTriangleMeshCollisions::FGIAColor>& TriangleGIAColors)
+{
+	FPBDTriangleMeshCollisions::FTriangleSubMesh SubMesh(TriangleMesh);
+	SubMesh.InitAllDynamic();
+	SpatialAccelerator UnusedKinematicSpatial;
+	constexpr FSolverReal LargeDt = UE_BIG_NUMBER; // This will disable all kinematic collider timers.
+	Init(Particles, LargeDt, SubMesh, Spatial, UnusedKinematicSpatial, VertexGIAColors, TriangleGIAColors);
+}
+template void CHAOS_API FPBDCollisionSpringConstraintsBase::Init<FTriangleMesh::TBVHType<FSolverReal>>(const FSolverParticles& Particles, const FTriangleMesh::TBVHType<FSolverReal>& Spatial, 
+	const TConstArrayView<FPBDTriangleMeshCollisions::FGIAColor>& VertexGIAColors, const TArray<FPBDTriangleMeshCollisions::FGIAColor>& TriangleGIAColors);
+template void CHAOS_API FPBDCollisionSpringConstraintsBase::Init<FTriangleMesh::TSpatialHashType<FSolverReal>>(const FSolverParticles& Particles, const FTriangleMesh::TSpatialHashType<FSolverReal>& Spatial,
+	const TConstArrayView<FPBDTriangleMeshCollisions::FGIAColor>& VertexGIAColors, const TArray<FPBDTriangleMeshCollisions::FGIAColor>& TriangleGIAColors);
+template void CHAOS_API FPBDCollisionSpringConstraintsBase::Init<FTriangleMesh::TBVHType<FSolverReal>>(const FSolverParticlesRange& Particles, const FTriangleMesh::TBVHType<FSolverReal>& Spatial,
+	const TConstArrayView<FPBDTriangleMeshCollisions::FGIAColor>& VertexGIAColors, const TArray<FPBDTriangleMeshCollisions::FGIAColor>& TriangleGIAColors);
+template void CHAOS_API FPBDCollisionSpringConstraintsBase::Init<FTriangleMesh::TSpatialHashType<FSolverReal>>(const FSolverParticlesRange& Particles, const FTriangleMesh::TSpatialHashType<FSolverReal>& Spatial,
+	const TConstArrayView<FPBDTriangleMeshCollisions::FGIAColor>& VertexGIAColors, const TArray<FPBDTriangleMeshCollisions::FGIAColor>& TriangleGIAColors);
+
+template<typename SpatialAccelerator, typename SolverParticlesOrRange>
+void FPBDCollisionSpringConstraintsBase::Init(const SolverParticlesOrRange& Particles, const FSolverReal Dt, const FPBDTriangleMeshCollisions::FTriangleSubMesh& CollidableSubMesh,
+	const SpatialAccelerator& DynamicSpatial, const SpatialAccelerator& KinematicColliderSpatial, const TConstArrayView<FPBDTriangleMeshCollisions::FGIAColor>& VertexGIAColors, const TArray<FPBDTriangleMeshCollisions::FGIAColor>& TriangleGIAColors)
 {
 	PRAGMA_DISABLE_DEPRECATION_WARNINGS
-	const FTriangleMesh& CollidableMesh = CollidableSubMesh.GetCollidableMesh();
-	const TArray<TVec3<int32>>& Elements = CollidableMesh.GetElements();
+	const int32 NumDynamicElements = CollidableSubMesh.GetDynamicSubMesh().GetNumElements();
 
-	if (!Elements.Num())
+	if (!NumDynamicElements)
 	{
 		Constraints.Reset();
 		Barys.Reset();
 		FlipNormal.Reset();
 		ConstraintTypes.Reset();
+		KinematicCollidingParticles.Reset();
+		KinematicColliderElements.Reset();
 		return;
 	}
 	{
-		if (ExistingConstraintLookup.Num() != NumParticles)
+		if (KinematicColliderTimers.Num() != NumParticles)
 		{
-			ExistingConstraintLookup.Reset();
-			ExistingConstraintLookup.SetNum(NumParticles);
+			KinematicColliderTimers.Init(TMap<int32, FSolverReal>(), NumParticles);
 		}
 
-		bGlobalIntersectionAnalysis = VertexGIAColors.Num() == NumParticles + Offset && TriangleGIAColors.Num() == Elements.Num();
+		bGlobalIntersectionAnalysis = VertexGIAColors.Num() == NumParticles + Offset && TriangleGIAColors.Num() == NumDynamicElements;
 		TRACE_CPUPROFILER_EVENT_SCOPE(ChaosPBDCollisionSpring_ProximityQuery);
 
-		const int32 NumCollidableParticles = CollidableSubMesh.GetCollidableVertices().IsEmpty() ? NumParticles : CollidableSubMesh.GetCollidableVertices().Num();
+		const int32 NumCollidableParticles = CollidableSubMesh.GetDynamicVertices().IsEmpty() ? NumParticles : CollidableSubMesh.GetDynamicVertices().Num();
 
 		// Preallocate enough space for all possible connections.
 		constexpr int32 MaxConnectionsPerPoint = 3;
@@ -140,17 +169,18 @@ void FPBDCollisionSpringConstraintsBase::Init(const SolverParticlesOrRange& Part
 		FlipNormal.SetNumUninitialized(NumCollidableParticles * MaxConnectionsPerPoint);
 		ConstraintTypes.SetNumUninitialized(NumCollidableParticles * MaxConnectionsPerPoint);
 
-		std::atomic<int32> ConstraintIndex(0);
+		KinematicCollidingParticles.SetNumUninitialized(NumCollidableParticles);
+		KinematicColliderElements.SetNumUninitialized(NumCollidableParticles);
 
-		const FSolverReal MaxSingleSidedThickness = FMath::Max(Thickness, KinematicColliderThickness);
-		const FSolverReal HeightSq = FMath::Square(Thickness + MaxSingleSidedThickness);
+		std::atomic<int32> ConstraintIndex(0);
+		std::atomic<int32> KinematicConstraintIndex(0);
 
 		const TConstArrayView<FSolverVec3> ReferencePositionsView = GetConstArrayView(Particles, Offset, NumParticles, ReferencePositions);
 
 		PhysicsParallelFor(NumCollidableParticles,
-			[this, Dt, &CollidableSubMesh, &CollidableMesh, &Elements, &Spatial, &Particles, &ConstraintIndex, MaxSingleSidedThickness, HeightSq, MaxConnectionsPerPoint, &VertexGIAColors, &TriangleGIAColors, &ReferencePositionsView](int32 CollidableIndex)
+			[this, Dt, &CollidableSubMesh, &DynamicSpatial, &KinematicColliderSpatial, &Particles, &ConstraintIndex, &KinematicConstraintIndex, MaxConnectionsPerPoint, &VertexGIAColors, &TriangleGIAColors, &ReferencePositionsView](int32 CollidableIndex)
 			{
-				const int32 i = CollidableSubMesh.GetCollidableVertices().IsEmpty() ? CollidableIndex : CollidableSubMesh.GetCollidableVertices()[CollidableIndex];
+				const int32 i = CollidableSubMesh.GetDynamicVertices().IsEmpty() ? CollidableIndex : CollidableSubMesh.GetDynamicVertices()[CollidableIndex] - Offset; // DynamicVertices already has offset applied.
 				const int32 Index = i + Offset;
 				if (Particles.InvM(Index) == (FSolverReal)0.)
 				{
@@ -161,27 +191,23 @@ void FPBDCollisionSpringConstraintsBase::Init(const SolverParticlesOrRange& Part
 				const bool bVertexHasCollisionLayers = VertexCollisionLayers.IsValidIndex(i) && VertexCollisionLayers[i][0] != INDEX_NONE;
 				check(!bVertexHasCollisionLayers || VertexCollisionLayers[i][0] <= VertexCollisionLayers[i][1]);
 
-				TMap<int32, FExistingConstraintData> PrevExistingConstraintLookup;
-				Swap(PrevExistingConstraintLookup, ExistingConstraintLookup[i]);
-
-				TArray< TTriangleCollisionPoint<FSolverReal> > Result;
-				int32 ConstraintsAdded = 0;
-				if (CollidableMesh.PointProximityQuery(Spatial, static_cast<const TArrayView<const FSolverVec3>&>(Particles.XArray()), Index, Particles.GetX(Index), Thickness * ExtraThicknessMult, MaxSingleSidedThickness * ExtraThicknessMult,
-					[this, bVertexHasCollisionLayers, &Particles, &CollidableSubMesh, &Elements, &VertexGIAColors, &TriangleGIAColors](const int32 PointIndex, const int32 SubMeshTriangleIndex)->bool
+				// Dynamic collisions
+				TArray< TTriangleCollisionPoint<FSolverReal> > DynamicResult;
+				if (CollidableSubMesh.GetDynamicSubMesh().PointProximityQuery(DynamicSpatial, static_cast<const TArrayView<const FSolverVec3>&>(Particles.XArray()), Index, Particles.X(Index), Thickness * ExtraThicknessMult, Thickness * ExtraThicknessMult,
+					[this, bVertexHasCollisionLayers, &Particles, &CollidableSubMesh, &VertexGIAColors, &TriangleGIAColors](const int32 PointIndex, const int32 SubMeshTriangleIndex)->bool
 					{
-						const TVector<int32, 3>& Elem = Elements[SubMeshTriangleIndex];
-						const int32 FullMeshTriangleIndex = CollidableSubMesh.FullMeshIndexFromSubIndex(SubMeshTriangleIndex);
+						const TVector<int32, 3>& Elem = CollidableSubMesh.GetDynamicSubMesh().GetElements()[SubMeshTriangleIndex];
+						const int32 FullMeshTriangleIndex = CollidableSubMesh.GetFullMeshElementIndexFromDynamicElement(SubMeshTriangleIndex);
 
-						const bool bFaceIsKinematic = CollidableSubMesh.GetSubMeshElementIsKinematic()[SubMeshTriangleIndex];
 						bool bUseCollisionLayerOverride = false;
-						if (!bFaceIsKinematic && bVertexHasCollisionLayers && FaceCollisionLayers[FullMeshTriangleIndex] != INDEX_NONE)
+						if (bVertexHasCollisionLayers && FaceCollisionLayers[FullMeshTriangleIndex] != INDEX_NONE)
 						{
 							if (FaceCollisionLayers[FullMeshTriangleIndex] < VertexCollisionLayers[PointIndex - Offset][0] || FaceCollisionLayers[FullMeshTriangleIndex] > VertexCollisionLayers[PointIndex - Offset][1])
 							{
 								bUseCollisionLayerOverride = true;
 							}
 						}
-						if (!bFaceIsKinematic && !bUseCollisionLayerOverride && bGlobalIntersectionAnalysis)
+						if (!bUseCollisionLayerOverride && bGlobalIntersectionAnalysis)
 						{
 							const bool bIsAnyBoundary = VertexGIAColors[PointIndex].IsBoundary()
 								|| VertexGIAColors[Elem[0]].IsBoundary() 
@@ -213,24 +239,24 @@ void FPBDCollisionSpringConstraintsBase::Init(const SolverParticlesOrRange& Part
 
 						return true;
 					},
-					Result))
+					DynamicResult))
 				{
 
-					if (Result.Num() > MaxConnectionsPerPoint)
+					if (DynamicResult.Num() > MaxConnectionsPerPoint)
 					{
 						// TODO: once we have a PartialSort, use that instead here.
-						Result.Sort(
+						DynamicResult.Sort(
 							[](const TTriangleCollisionPoint<FSolverReal>& First, const TTriangleCollisionPoint<FSolverReal>& Second)->bool
 							{
 								return First.Phi < Second.Phi;
 							}
 						);
-						Result.SetNum(MaxConnectionsPerPoint, EAllowShrinking::No);
+						DynamicResult.SetNum(MaxConnectionsPerPoint, EAllowShrinking::No);
 					}
 
-					for (const TTriangleCollisionPoint<FSolverReal>& CollisionPoint : Result)
+					for (const TTriangleCollisionPoint<FSolverReal>& CollisionPoint : DynamicResult)
 					{
-						const TVector<int32, 3>& Elem = Elements[CollisionPoint.Indices[1]];
+						const TVector<int32, 3>& Elem = CollidableSubMesh.GetDynamicSubMesh().GetElements()[CollisionPoint.Indices[1]];
 						if (ReferencePositionsView.Num())
 						{
 							const FSolverVec3& RefP = ReferencePositionsView[Index];
@@ -238,32 +264,22 @@ void FPBDCollisionSpringConstraintsBase::Init(const SolverParticlesOrRange& Part
 							const FSolverVec3& RefP1 = ReferencePositionsView[Elem[1]];
 							const FSolverVec3& RefP2 = ReferencePositionsView[Elem[2]];
 							const FSolverVec3 RefDiff = RefP - CollisionPoint.Bary[1] * RefP0 - CollisionPoint.Bary[2] * RefP1 - CollisionPoint.Bary[3] * RefP2;
-							if (RefDiff.SizeSquared() < HeightSq)
+							if (RefDiff.SizeSquared() < FMath::Square(Thickness))
 							{
 								continue;
 							}
 						}
 
-
 						FSolverVec3 Bary(CollisionPoint.Bary[1], CollisionPoint.Bary[2], CollisionPoint.Bary[3]);
 
-						const int32 FullMeshTriangleIndex = CollidableSubMesh.FullMeshIndexFromSubIndex(CollisionPoint.Indices[1]);
+						const int32 FullMeshTriangleIndex = CollidableSubMesh.GetFullMeshElementIndexFromDynamicElement(CollisionPoint.Indices[1]);
 
 						bool bFlipNormal = false;
-						// Check kinematic
-						const bool bFaceIsKinematic = CollidableSubMesh.GetSubMeshElementIsKinematic()[CollisionPoint.Indices[1]];
-						if (bFaceIsKinematic)
-						{
-							PrevExistingConstraintLookup.Remove(FullMeshTriangleIndex);
-							// Chaos internal winding order is reversed.
-							bFlipNormal = true;
-						}
-
-						EConstraintType ConstraintType = bFaceIsKinematic ? EConstraintType::Kinematic : EConstraintType::Default;
+						EConstraintType ConstraintType = EConstraintType::Default;
 
 						// Check collision layers
 						bool bUseCollisionLayerOverride = false;
-						if (!bFaceIsKinematic && bVertexHasCollisionLayers && FaceCollisionLayers[FullMeshTriangleIndex] != INDEX_NONE)
+						if (bVertexHasCollisionLayers && FaceCollisionLayers[FullMeshTriangleIndex] != INDEX_NONE)
 						{
 							if (FaceCollisionLayers[FullMeshTriangleIndex] < VertexCollisionLayers[i][0])
 							{
@@ -281,7 +297,7 @@ void FPBDCollisionSpringConstraintsBase::Init(const SolverParticlesOrRange& Part
 							}
 						}
 
-						if (!bFaceIsKinematic && !bUseCollisionLayerOverride)
+						if (!bUseCollisionLayerOverride)
 						{
 							// NOTE: CollisionPoint.Normal has already been flipped to point toward the Point, so need to recalculate here.
 							const TTriangle<FSolverReal> Triangle(Particles.GetX(Elem[0]), Particles.GetX(Elem[1]), Particles.GetX(Elem[2]));
@@ -306,32 +322,71 @@ void FPBDCollisionSpringConstraintsBase::Init(const SolverParticlesOrRange& Part
 						Barys[IndexToWrite] = Bary;
 						FlipNormal[IndexToWrite] = bFlipNormal;
 						ConstraintTypes[IndexToWrite] = ConstraintType;
-						if (ConstraintType == EConstraintType::Kinematic)
+					}
+				}
+
+				// Kinematic colliders
+				// Increment existing timers and remove any elements that are too old.
+				for(TMap<int32, FSolverReal>::TIterator TimerIter = KinematicColliderTimers[i].CreateIterator(); TimerIter; ++TimerIter)
+				{
+					TimerIter.Value() += Dt;
+					if (TimerIter.Value() > KinematicColliderMaxTimer)
+					{
+						TimerIter.RemoveCurrent();
+					}
+				}
+
+				if(CollidableSubMesh.GetKinematicColliderSubMesh().GetNumElements() > 0)
+				{
+					TArray< TTriangleCollisionPoint<FSolverReal> > KinematicResult;
+					if (CollidableSubMesh.GetKinematicColliderSubMesh().PointProximityQuery(KinematicColliderSpatial, static_cast<const TArrayView<const FSolverVec3>&>(Particles.XArray()), Index, Particles.X(Index), Thickness * ExtraThicknessMult, KinematicColliderThickness * ExtraThicknessMult, [](const int32 PointIndex, const int32 SubMeshTriangleIndex)->bool { return true; }, KinematicResult))
+					{
+						if (KinematicResult.Num() > MaxKinematicConnectionsPerPoint)
 						{
-							// Remember this constraint
-							ExistingConstraintLookup[i].Add(FullMeshTriangleIndex, {0.f});
+							// TODO: once we have a PartialSort, use that instead here.
+							KinematicResult.Sort(
+								[](const TTriangleCollisionPoint<FSolverReal>& First, const TTriangleCollisionPoint<FSolverReal>& Second)->bool
+							{
+								return First.Phi < Second.Phi;
+							}
+							);
+							KinematicResult.SetNum(MaxKinematicConnectionsPerPoint, EAllowShrinking::No);
+						}
+						for (const TTriangleCollisionPoint<FSolverReal>& CollisionPoint : KinematicResult)
+						{
+							const int32 FullMeshTriangleIndex = CollidableSubMesh.GetFullMeshElementIndexFromKinematicElement(CollisionPoint.Indices[1]);
+							KinematicColliderTimers[i].FindOrAdd(FullMeshTriangleIndex, (FSolverReal)0.f);
 						}
-						++ConstraintsAdded;
 					}
 				}
-				
-				for (TMap<int32, FExistingConstraintData>::TConstIterator PrevData = PrevExistingConstraintLookup.CreateConstIterator(); PrevData && ConstraintsAdded < MaxConnectionsPerPoint; ++PrevData)
+
+				// Get MaxKinematicConnectionsPerPoint most recent
+				if (KinematicColliderTimers[i].Num() > MaxKinematicConnectionsPerPoint)
 				{
-					if (PrevData.Value().Timer + Dt <= Chaos_CollisionSpring_MaxTimer)
+					KinematicColliderTimers[i].ValueSort([](const FSolverReal A, const FSolverReal B) { return A < B; });
+				}
+
+				if (KinematicColliderTimers[i].Num() > 0)
+				{
+					const int32 IndexToWrite = KinematicConstraintIndex.fetch_add(1);
+					KinematicCollidingParticles[IndexToWrite] = Index;
+
+					int32 LocalIndex = 0;
+					for (TMap<int32, FSolverReal>::TIterator TimerIter = KinematicColliderTimers[i].CreateIterator(); TimerIter; ++TimerIter)
 					{
-						const int32 SubMeshIndex = CollidableSubMesh.SubMeshIndexFromFullIndex(PrevData.Key());
-						if (SubMeshIndex != INDEX_NONE)
+						if (LocalIndex < MaxKinematicConnectionsPerPoint)
+						{
+							KinematicColliderElements[IndexToWrite][LocalIndex++] = TimerIter.Key();
+						}
+						else
 						{
-							const int32 IndexToWrite = ConstraintIndex.fetch_add(1);
-							const TVec3<int32>& Elem = Elements[SubMeshIndex];
-							Constraints[IndexToWrite] = { Index, Elem[0], Elem[1], Elem[2] };
-							Barys[IndexToWrite] = FSolverVec3(1.f, 0.f, 0.f); // Unused
-							FlipNormal[IndexToWrite] = true;
-							ConstraintTypes[IndexToWrite] = EConstraintType::Kinematic;
-							ExistingConstraintLookup[i].Add(PrevData.Key(), {PrevData.Value().Timer + Dt});
-							++ConstraintsAdded;
+							TimerIter.RemoveCurrent();
 						}
 					}
+					for (; LocalIndex < MaxKinematicConnectionsPerPoint; ++LocalIndex)
+					{
+						KinematicColliderElements[IndexToWrite][LocalIndex] = INDEX_NONE;
+					}
 				}
 			}
 		);
@@ -342,16 +397,24 @@ void FPBDCollisionSpringConstraintsBase::Init(const SolverParticlesOrRange& Part
 		Barys.SetNum(ConstraintNum, EAllowShrinking::No);
 		FlipNormal.SetNum(ConstraintNum, EAllowShrinking::No);
 		ConstraintTypes.SetNum(ConstraintNum, EAllowShrinking::No);
+
+		const int32 KinematicConstraintNum = KinematicConstraintIndex.load();
+		KinematicCollidingParticles.SetNum(KinematicConstraintNum, EAllowShrinking::No);
+		KinematicColliderElements.SetNum(KinematicConstraintNum, EAllowShrinking::No);
 	}
 	PRAGMA_ENABLE_DEPRECATION_WARNINGS
 }
-template void CHAOS_API FPBDCollisionSpringConstraintsBase::Init<FTriangleMesh::TBVHType<FSolverReal>>(const FSolverParticles& Particles, const FSolverReal Dt, const FPBDTriangleMeshCollisions::FTriangleSubMesh& CollidableSubMesh, const FTriangleMesh::TBVHType<FSolverReal>& Spatial,
+template void CHAOS_API FPBDCollisionSpringConstraintsBase::Init<FTriangleMesh::TBVHType<FSolverReal>>(const FSolverParticles& Particles, const FSolverReal Dt,
+	const FPBDTriangleMeshCollisions::FTriangleSubMesh& CollidableSubMesh, const FTriangleMesh::TBVHType<FSolverReal>& DynamicSpatial, const FTriangleMesh::TBVHType<FSolverReal>& KinematicColliderSpatial,
 	const TConstArrayView<FPBDTriangleMeshCollisions::FGIAColor>& VertexGIAColors, const TArray<FPBDTriangleMeshCollisions::FGIAColor>& TriangleGIAColors);
-template void CHAOS_API FPBDCollisionSpringConstraintsBase::Init<FTriangleMesh::TSpatialHashType<FSolverReal>>(const FSolverParticles& Particles, const FSolverReal Dt, const FPBDTriangleMeshCollisions::FTriangleSubMesh& CollidableSubMesh, const FTriangleMesh::TSpatialHashType<FSolverReal>& Spatial,
+template void CHAOS_API FPBDCollisionSpringConstraintsBase::Init<FTriangleMesh::TSpatialHashType<FSolverReal>>(const FSolverParticles& Particles, const FSolverReal Dt,
+	const FPBDTriangleMeshCollisions::FTriangleSubMesh& CollidableSubMesh, const FTriangleMesh::TSpatialHashType<FSolverReal>& DynamicSpatial, const FTriangleMesh::TSpatialHashType<FSolverReal>& KinematicColliderSpatial,
 	const TConstArrayView<FPBDTriangleMeshCollisions::FGIAColor>& VertexGIAColors, const TArray<FPBDTriangleMeshCollisions::FGIAColor>& TriangleGIAColors);
-template void CHAOS_API FPBDCollisionSpringConstraintsBase::Init<FTriangleMesh::TBVHType<FSolverReal>>(const FSolverParticlesRange& Particles, const FSolverReal Dt, const FPBDTriangleMeshCollisions::FTriangleSubMesh& CollidableSubMesh, const FTriangleMesh::TBVHType<FSolverReal>& Spatial,
+template void CHAOS_API FPBDCollisionSpringConstraintsBase::Init<FTriangleMesh::TBVHType<FSolverReal>>(const FSolverParticlesRange& Particles, const FSolverReal Dt,
+	const FPBDTriangleMeshCollisions::FTriangleSubMesh& CollidableSubMesh, const FTriangleMesh::TBVHType<FSolverReal>& DynamicSpatial, const FTriangleMesh::TBVHType<FSolverReal>& KinematicColliderSpatial,
 	const TConstArrayView<FPBDTriangleMeshCollisions::FGIAColor>& VertexGIAColors, const TArray<FPBDTriangleMeshCollisions::FGIAColor>& TriangleGIAColors);
-template void CHAOS_API FPBDCollisionSpringConstraintsBase::Init<FTriangleMesh::TSpatialHashType<FSolverReal>>(const FSolverParticlesRange& Particles, const FSolverReal Dt, const FPBDTriangleMeshCollisions::FTriangleSubMesh& CollidableSubMesh, const FTriangleMesh::TSpatialHashType<FSolverReal>& Spatial,
+template void CHAOS_API FPBDCollisionSpringConstraintsBase::Init<FTriangleMesh::TSpatialHashType<FSolverReal>>(const FSolverParticlesRange& Particles, const FSolverReal Dt,
+	const FPBDTriangleMeshCollisions::FTriangleSubMesh& CollidableSubMesh, const FTriangleMesh::TSpatialHashType<FSolverReal>& DynamicSpatial, const FTriangleMesh::TSpatialHashType<FSolverReal>& KinematicColliderSpatial,
 	const TConstArrayView<FPBDTriangleMeshCollisions::FGIAColor>& VertexGIAColors, const TArray<FPBDTriangleMeshCollisions::FGIAColor>& TriangleGIAColors);
 
 template<typename SolverParticlesOrRange>
@@ -385,30 +448,10 @@ FSolverVec3 FPBDCollisionSpringConstraintsBase::GetDelta(const SolverParticlesOr
 	const TTriangle<FSolverReal> Triangle(P2, P3, P4);
 	const FSolverVec3 Normal = FlipNormal[ConstraintIndex] ? -Triangle.GetNormal() : Triangle.GetNormal();
 
-	FSolverVec3 P;
-	FSolverVec3 Difference;
-	FSolverReal NormalDifference;
-	FSolverReal TangentialFalloff = 1.f;
-	FSolverVec3 Bary;
-	if (ConstraintTypes[ConstraintIndex] == EConstraintType::Kinematic)
-	{
-		P = FindClosestPointAndBaryOnTriangle(P2, P3, P4, P1, Bary);
-		Difference = P1 - P;
-		NormalDifference = Difference.Dot(Normal);
-		const FSolverReal TangentialDifference = (Difference - NormalDifference * Normal).Size();
-		TangentialFalloff = 1.f - TangentialDifference / FMath::Max(1.5f * Height, 1e-4f);
-		if (TangentialFalloff <= 0.f)
-		{
-			return FSolverVec3(0);
-		}
-	}
-	else
-	{
-		Bary = Barys[ConstraintIndex];
-		P = Bary[0] * P2 + Bary[1] * P3 + Bary[2] * P4;
-		Difference = P1 - P;
-		NormalDifference = Difference.Dot(Normal);
-	}
+	const FSolverVec3& Bary = Barys[ConstraintIndex];
+	const FSolverVec3 P = Bary[0] * P2 + Bary[1] * P3 + Bary[2] * P4;
+	const FSolverVec3 Difference = P1 - P;
+	const FSolverReal NormalDifference = Difference.Dot(Normal);
 
 	// Normal repulsion with friction
 	if (NormalDifference > Height)
@@ -416,7 +459,7 @@ FSolverVec3 FPBDCollisionSpringConstraintsBase::GetDelta(const SolverParticlesOr
 		return FSolverVec3(0);
 	}
 
-	const FSolverReal ConstraintStiffness = GetConstraintStiffness(ConstraintIndex) * TangentialFalloff;
+	const FSolverReal ConstraintStiffness = GetConstraintStiffness(ConstraintIndex);
 	const FSolverReal ConstraintFriction = GetConstraintFrictionCoefficient(ConstraintIndex);
 
 	const FSolverReal NormalDelta = Height - NormalDifference;
@@ -511,6 +554,120 @@ void FPBDCollisionSpringConstraintsBase::UpdateLinearSystem(const FSolverParticl
 	}
 	PRAGMA_ENABLE_DEPRECATION_WARNINGS
 }
+
+
+template<typename SolverParticlesOrRange>
+void FPBDCollisionSpringConstraintsBase::Apply(SolverParticlesOrRange& InParticles, const FSolverReal Dt) const
+{
+	PRAGMA_DISABLE_DEPRECATION_WARNINGS
+	// Dynamic collisions
+	{
+		TRACE_CPUPROFILER_EVENT_SCOPE(ChaosPBDCollisionSpring_ApplyDynamic);
+		for (int32 ConstraintIndex = 0; ConstraintIndex < Constraints.Num(); ++ConstraintIndex)
+		{
+			Apply(InParticles, Dt, ConstraintIndex);
+		}
+	}
+	PRAGMA_ENABLE_DEPRECATION_WARNINGS
+
+	ApplyKinematicConstraints(InParticles, Dt);
+}
+template void CHAOS_API FPBDCollisionSpringConstraintsBase::Apply(FSolverParticles& Particles, const FSolverReal Dt) const;
+template void CHAOS_API FPBDCollisionSpringConstraintsBase::Apply(FSolverParticlesRange& Particles, const FSolverReal Dt) const;
+
+template<typename SolverParticlesOrRange>
+void FPBDCollisionSpringConstraintsBase::ApplyKinematicConstraints(SolverParticlesOrRange& Particles, const FSolverReal Dt) const
+{
+	if (!KinematicCollidingParticles.Num())
+	{
+		return;
+	}
+	TRACE_CPUPROFILER_EVENT_SCOPE(ChaosPBDCollisionSpring_ApplyKinematicConstraints);
+
+	const FSolverReal Height = Thickness + KinematicColliderThickness;
+#if INTEL_ISPC
+	static_assert(sizeof(ispc::FIntVector) == sizeof(TVector<int32, MaxKinematicConnectionsPerPoint>), "sizeof(ispc::FIntVector) != sizeof(TVector<int32, MaxKinematicConnectionsPerPoint>)");
+	if (bRealTypeCompatibleWithISPC && bChaos_CollisionSpring_ISPC_Enabled)
+	{
+		ispc::ApplyKinematicCollisionSpringConstraints(
+			(ispc::FVector4f*)Particles.GetPAndInvM().GetData(),
+			(const ispc::FVector3f*)Particles.XArray().GetData(),
+			KinematicCollidingParticles.GetData(),
+			(const ispc::FIntVector*)KinematicColliderElements.GetData(),
+			(const ispc::FIntVector*)TriangleMesh.GetElements().GetData(),
+			Height,
+			KinematicColliderStiffness,
+			KinematicColliderFrictionCoefficient,
+			KinematicCollidingParticles.Num()
+		);
+	}
+	else
+#endif
+	{
+		for (int32 Index = 0; Index < KinematicCollidingParticles.Num(); ++Index)
+		{
+			const int32 Index1 = KinematicCollidingParticles[Index];
+			for (int32 EIndex = 0; EIndex < MaxKinematicConnectionsPerPoint; ++EIndex)
+			{
+				const int32 ElemIndex = KinematicColliderElements[Index][EIndex];
+				if (ElemIndex == INDEX_NONE)
+				{
+					break;
+				}
+
+				const int32 Index2 = TriangleMesh.GetElements()[ElemIndex][0];
+				const int32 Index3 = TriangleMesh.GetElements()[ElemIndex][1];
+				const int32 Index4 = TriangleMesh.GetElements()[ElemIndex][2];
+
+				FSolverVec3& P1 = Particles.P(Index1);
+				const FSolverVec3& P2 = Particles.P(Index2);
+				const FSolverVec3& P3 = Particles.P(Index3);
+				const FSolverVec3& P4 = Particles.P(Index4);
+
+				const TTriangle<FSolverReal> Triangle(P2, P3, P4);
+				const FSolverVec3  Normal = -Triangle.GetNormal(); // normals are flipped from UE
+
+				FSolverVec3 Bary;
+				const FSolverVec3 P = FindClosestPointAndBaryOnTriangle(P2, P3, P4, P1, Bary);
+				const FSolverVec3 Difference = P1 - P;
+				const FSolverReal NormalDifference = Difference.Dot(Normal);
+
+				if (NormalDifference > Height)
+				{
+					continue;
+				}
+
+				const FSolverReal TangentialDifference = (Difference - NormalDifference * Normal).Size();
+				constexpr FSolverReal TangentialFalloffMultiplier(1.5f);
+				const FSolverReal TangentialFalloff = (FSolverReal)1.f - TangentialDifference / FMath::Max(TangentialFalloffMultiplier * Height, UE_KINDA_SMALL_NUMBER);
+				if (TangentialFalloff <= 0.f)
+				{
+					continue;
+				}
+
+				const FSolverReal NormalDelta = Height - NormalDifference;
+				const FSolverVec3 RepulsionDelta = KinematicColliderStiffness * NormalDelta * Normal;
+
+				P1 += RepulsionDelta;
+
+				if (KinematicColliderFrictionCoefficient > 0)
+				{
+					const FSolverVec3& X1 = Particles.X(Index1);
+					const FSolverVec3 X = Bary[0] * Particles.X(Index2) + Bary[1] * Particles.X(Index3) + Bary[2] * Particles.X(Index4);
+					const FSolverVec3 RelativeDisplacement = (P1 - X1) - (P - X);
+					const FSolverVec3 RelativeDisplacementTangent = RelativeDisplacement - RelativeDisplacement.Dot(Normal) * Normal;
+					const FSolverReal RelativeDisplacementTangentLength = RelativeDisplacementTangent.Length();
+					const FSolverReal PositionCorrection = FMath::Min(NormalDelta * KinematicColliderFrictionCoefficient, RelativeDisplacementTangentLength);
+					const FSolverReal CorrectionRatio = RelativeDisplacementTangentLength < UE_SMALL_NUMBER ? 0.f : PositionCorrection / RelativeDisplacementTangentLength;
+					const FSolverVec3 FrictionDelta = -CorrectionRatio * RelativeDisplacementTangent;
+
+					P1 += FrictionDelta;
+				}
+			}
+		}
+	}
+}
+
 }  // End namespace Chaos::Softs
 
 #endif
diff --git a/Engine/Source/Runtime/Experimental/Chaos/Private/Chaos/PBDTriangleMeshCollisions.cpp b/Engine/Source/Runtime/Experimental/Chaos/Private/Chaos/PBDTriangleMeshCollisions.cpp
--- a/Engine/Source/Runtime/Experimental/Chaos/Private/Chaos/PBDTriangleMeshCollisions.cpp
+++ b/Engine/Source/Runtime/Experimental/Chaos/Private/Chaos/PBDTriangleMeshCollisions.cpp
@@ -24,7 +24,7 @@ struct FEdgeFaceIntersection
 
 // Returned array has NOT been shrunk
 template<typename SpatialAccelerator, typename SolverParticlesOrRange>
-static void FindEdgeFaceIntersections(const FTriangleMesh& TriangleMesh, const TArray<int32>& IntersectableEdges, bool bSkipKinematic, const TArray<bool>& FaceIsKinematic, const SpatialAccelerator& Spatial, const SolverParticlesOrRange& Particles, TArray<FEdgeFaceIntersection>& Intersections)
+static void FindEdgeFaceIntersections(const FTriangleMesh& TriangleMesh, const SpatialAccelerator& Spatial, const SolverParticlesOrRange& Particles, TArray<FEdgeFaceIntersection>& Intersections)
 {
 	TRACE_CPUPROFILER_EVENT_SCOPE(ChaosFPBDTriangleMeshCollisions_IntersectionQuery);
 
@@ -32,7 +32,7 @@ static void FindEdgeFaceIntersections(const FTriangleMesh& TriangleMesh, const T
 	const TArray<TVec2<int32>>& EdgeToFaces = TriangleMesh.GetEdgeToFaces();
 	const TArray<TVec3<int32>>& Elements = TriangleMesh.GetElements();
 
-	const int32 NumIntersectableEdges = IntersectableEdges.IsEmpty() ? SegmentMesh.GetNumElements() : IntersectableEdges.Num();
+	const int32 NumIntersectableEdges = SegmentMesh.GetNumElements();
 
 	// Preallocate enough space for (more than) typical number of expected intersections.
 	constexpr int32 PreallocatedIntersectionsPerEdge = 3;
@@ -45,24 +45,23 @@ static void FindEdgeFaceIntersections(const FTriangleMesh& TriangleMesh, const T
 	TArray<FEdgeFaceIntersection> ExtraIntersections;
 	FCriticalSection CriticalSection;
 	PhysicsParallelFor(NumIntersectableEdges,
-		[&Spatial, &TriangleMesh, &Particles, &SegmentMesh, &EdgeToFaces, &Elements, &IntersectionIndex, &Intersections, PreallocatedIntersectionsNum, &ExtraIntersections, &CriticalSection, &IntersectableEdges, bSkipKinematic, &FaceIsKinematic](int32 IntersectableEdgeIndex)
+		[&Spatial, &TriangleMesh, &Particles, &SegmentMesh, &EdgeToFaces, &Elements, &IntersectionIndex, &Intersections, PreallocatedIntersectionsNum, &ExtraIntersections, &CriticalSection](int32 EdgeIndex)
 		{
-			const int32 EdgeIndex = IntersectableEdges.IsEmpty() ? IntersectableEdgeIndex : IntersectableEdges[IntersectableEdgeIndex];
-
 			TArray< TTriangleCollisionPoint<FSolverReal> > Result;
 
 			const int32 EdgePointIndex0 = SegmentMesh.GetElements()[EdgeIndex][0];
 			const int32 EdgePointIndex1 = SegmentMesh.GetElements()[EdgeIndex][1];
-			const FSolverVec3& EdgePosition0 = Particles.GetX(EdgePointIndex0);
-			const FSolverVec3& EdgePosition1 = Particles.GetX(EdgePointIndex1);
+			const FSolverVec3& EdgePosition0 = Particles.X(EdgePointIndex0);
+			const FSolverVec3& EdgePosition1 = Particles.X(EdgePointIndex1);
+			// No faces can be kinematic, but edges can still be since a face is only kinematic if all 3 vertices are kinematic
 			const bool bEdgeIsKinematic = Particles.InvM(EdgePointIndex0) == (FSolverReal)0. && Particles.InvM(EdgePointIndex1) == (FSolverReal)0.;
-			if (bSkipKinematic && bEdgeIsKinematic)
+			if (bEdgeIsKinematic)
 			{
 				return;
 			}
 
 			if (TriangleMesh.EdgeIntersectionQuery(Spatial, static_cast<const TArrayView<const FSolverVec3>&>(Particles.XArray()), EdgeIndex, EdgePosition0, EdgePosition1,
-				[&Elements, EdgePointIndex0, EdgePointIndex1, bSkipKinematic, bEdgeIsKinematic, &FaceIsKinematic,  &Particles](int32 EdgeIndex, int32 TriangleIndex)
+				[&Elements, EdgePointIndex0, EdgePointIndex1, &Particles](int32 EdgeIndex, int32 TriangleIndex)
 				{
 					if (EdgePointIndex0 == Elements[TriangleIndex][0] || EdgePointIndex0 == Elements[TriangleIndex][1] || EdgePointIndex0 == Elements[TriangleIndex][2] ||
 						EdgePointIndex1 == Elements[TriangleIndex][0] || EdgePointIndex1 == Elements[TriangleIndex][1] || EdgePointIndex1 == Elements[TriangleIndex][2])
@@ -70,10 +69,6 @@ static void FindEdgeFaceIntersections(const FTriangleMesh& TriangleMesh, const T
 						return false;
 					}
 
-					if((bSkipKinematic || bEdgeIsKinematic) && FaceIsKinematic[TriangleIndex])
-					{
-						return false;
-					}
 					return true;
 				},
 				Result))
@@ -1289,43 +1284,26 @@ struct FPBDTriangleMeshCollisions::FScratchBuffers
 template<typename SolverParticlesOrRange>
 void FPBDTriangleMeshCollisions::FTriangleSubMesh::Init(const SolverParticlesOrRange& Particles, const TSet<int32>& InDisabledFaces, bool bCollideAgainstAllKinematicVertices, const TSet<int32>& InEnabledKinematicFaces)
 {
-	SubMesh.Init(TArray<TVec3<int32>>());
-	SubMeshToFullElements.Reset();
-	FullToSubMeshElements.Reset();
-	SubMeshElementIsKinematic.Reset();
-	CollidableVertices.Reset();
-	IntersectableSubmeshEdges.Reset();
-	bSubMeshIsFullMesh = true;
-
-	if (InDisabledFaces.IsEmpty() && bCollideAgainstAllKinematicVertices)
-	{
-		const TArray<TVec3<int32>>& FullElements = FullMesh.GetElements();
-		SubMeshElementIsKinematic.SetNumUninitialized(FullMesh.GetNumElements());
+	FullMeshToSubMeshIndices.Reset();
+	DynamicSubMeshToFullMeshIndices.Reset();
+	KinematicColliderSubMeshToFullMeshIndices.Reset();
+	DynamicVertices.Reset();
 
-		for (int32 FullElementIndex = 0; FullElementIndex < FullMesh.GetNumElements(); ++FullElementIndex)
-		{
-			const bool bIsKinematic =
-				Particles.InvM(FullElements[FullElementIndex][0]) == (FSolverReal)0.f &&
-				Particles.InvM(FullElements[FullElementIndex][1]) == (FSolverReal)0.f &&
-				Particles.InvM(FullElements[FullElementIndex][2]) == (FSolverReal)0.f;
-			SubMeshElementIsKinematic[FullElementIndex] = bIsKinematic;
-		}
-		return;
-	}
+	TArray<TVec3<int32>> DynamicMeshElements;
+	TArray<TVec3<int32>> KinematicMeshElements;
+	DynamicMeshElements.Reserve(FullMesh.GetNumElements());
+	KinematicMeshElements.Reserve(InEnabledKinematicFaces.Num());
+
+	FullMeshToSubMeshIndices.SetNumZeroed(FullMesh.GetNumElements());
+	DynamicSubMeshToFullMeshIndices.Reserve(FullMesh.GetNumElements());
+	KinematicColliderSubMeshToFullMeshIndices.Reserve(InEnabledKinematicFaces.Num());
 
-	// Build a sub mesh of just colliding faces.
-	bSubMeshIsFullMesh = false;
-	TArray<TVec3<int32>> SubElements;
-	SubElements.Reserve(FullMesh.GetNumElements());
-	FullToSubMeshElements.SetNumUninitialized(FullMesh.GetNumElements());
-	SubMeshToFullElements.Reserve(FullMesh.GetNumElements());
-	SubMeshElementIsKinematic.Reserve(FullMesh.GetNumElements());
 	const TArray<TVec3<int32>>& FullElements = FullMesh.GetElements();
 	for (int32 FullElementIndex = 0; FullElementIndex < FullMesh.GetNumElements(); ++FullElementIndex)
 	{
 		if (InDisabledFaces.Contains(FullElementIndex))
 		{
-			FullToSubMeshElements[FullElementIndex] = INDEX_NONE;
+			FullMeshToSubMeshIndices[FullElementIndex].SubMeshType = ESubMeshType::Invalid;
 			continue;
 		}
 		const bool bIsKinematic =
@@ -1333,47 +1311,66 @@ void FPBDTriangleMeshCollisions::FTriangleSubMesh::Init(const SolverParticlesOrR
 			Particles.InvM(FullElements[FullElementIndex][1]) == (FSolverReal)0.f &&
 			Particles.InvM(FullElements[FullElementIndex][2]) == (FSolverReal)0.f;
 
-		if (!bCollideAgainstAllKinematicVertices &&
-			bIsKinematic &&
-			!InEnabledKinematicFaces.Contains(FullElementIndex))
+		if (bIsKinematic)
 		{
-			FullToSubMeshElements[FullElementIndex] = INDEX_NONE;
-			continue;
+			if (bCollideAgainstAllKinematicVertices || InEnabledKinematicFaces.Contains(FullElementIndex))
+			{
+				const int32 SubMeshIndex = KinematicMeshElements.Add(FullElements[FullElementIndex]);
+				KinematicColliderSubMeshToFullMeshIndices.Add(FullElementIndex);
+				FullMeshToSubMeshIndices[FullElementIndex].SubMeshIndex = SubMeshIndex;
+				FullMeshToSubMeshIndices[FullElementIndex].SubMeshType = ESubMeshType::Kinematic;
+			}
+			else
+			{
+				FullMeshToSubMeshIndices[FullElementIndex].SubMeshType = ESubMeshType::Invalid;
+			}
+		}
+		else
+		{
+			const int32 SubMeshIndex = DynamicMeshElements.Add(FullElements[FullElementIndex]);
+			DynamicSubMeshToFullMeshIndices.Add(FullElementIndex);
+			FullMeshToSubMeshIndices[FullElementIndex].SubMeshIndex = SubMeshIndex;
+			FullMeshToSubMeshIndices[FullElementIndex].SubMeshType = ESubMeshType::Dynamic;
 		}
-
-		const int32 SubElementIndex = SubElements.Add(FullElements[FullElementIndex]);
-		SubMeshToFullElements.Add(FullElementIndex);
-		SubMeshElementIsKinematic.Add(bIsKinematic);
-		FullToSubMeshElements[FullElementIndex] = SubElementIndex;
 	}
 
-	// Use same Vertex range for submesh.
+	// Use same Vertex range for submeshes.
 	const TVec2<int32> VertexRange = FullMesh.GetVertexRange();
 	constexpr bool bCullDegenerateFalse = false;
-	SubMesh.Init(MoveTemp(SubElements), VertexRange[0], VertexRange[1], bCullDegenerateFalse);
-	check(SubMesh.GetNumElements() == SubMeshToFullElements.Num());
-	check(SubMesh.GetNumElements() == SubMeshElementIsKinematic.Num());
 
-	const TSet<int32> VertexSet = SubMesh.GetVertices();
-	CollidableVertices.Reserve(VertexSet.Num());
-	for (const int32 Vertex : VertexSet)
+	DynamicSubMesh.Init(MoveTemp(DynamicMeshElements), VertexRange[0], VertexRange[1], bCullDegenerateFalse);
+	KinematicColliderSubMesh.Init(MoveTemp(KinematicMeshElements), VertexRange[0], VertexRange[1], bCullDegenerateFalse);
+
+	check(DynamicSubMesh.GetNumElements() == DynamicSubMeshToFullMeshIndices.Num());
+	check(KinematicColliderSubMesh.GetNumElements() == KinematicColliderSubMeshToFullMeshIndices.Num());
+
+	const TSet<int32> DynamicVertexSet = DynamicSubMesh.GetVertices();
+	DynamicVertices.Reserve(DynamicVertexSet.Num());
+	for (const int32 PossiblyDynamicVertex : DynamicVertexSet)
 	{
-		if (Particles.InvM(Vertex) > (FSolverReal)0.)
+		// Some of these vertices may be kinematic since a triangle is dynamic if any of its vertices are.
+		if (Particles.InvM(PossiblyDynamicVertex) > (FSolverReal)0.)
 		{
-			CollidableVertices.Add(Vertex);
+			DynamicVertices.Add(PossiblyDynamicVertex);
 		}
 	}
+}
 
-	const FSegmentMesh& SegmentMesh = const_cast<const FTriangleMesh&>(SubMesh).GetSegmentMesh();
-	IntersectableSubmeshEdges.Reserve(SegmentMesh.GetNumElements());
-	for (int32 EdgeIndex = 0; EdgeIndex < SegmentMesh.GetNumElements(); ++EdgeIndex)
+void FPBDTriangleMeshCollisions::FTriangleSubMesh::InitAllDynamic()
+{
+	constexpr bool bCullDegenerateFalse = false;
+	DynamicSubMesh.Init(FullMesh.GetElements(), FullMesh.GetVertexRange()[0], FullMesh.GetVertexRange()[1], bCullDegenerateFalse);
+	KinematicColliderSubMesh.Init(TArray<TVec3<int32>>());
+
+	FullMeshToSubMeshIndices.SetNumUninitialized(FullMesh.GetNumElements());
+	DynamicSubMeshToFullMeshIndices.SetNumUninitialized(FullMesh.GetNumElements());
+	KinematicColliderSubMeshToFullMeshIndices.Reset();
+	DynamicVertices.Reset();
+	for (int32 Index = 0; Index < FullMesh.GetNumElements(); ++Index)
 	{
-		const int32 EdgePointIndex0 = SegmentMesh.GetElements()[EdgeIndex][0];
-		const int32 EdgePointIndex1 = SegmentMesh.GetElements()[EdgeIndex][1];
-		if (Particles.InvM(EdgePointIndex0) > (FSolverReal)0. || Particles.InvM(EdgePointIndex1) > (FSolverReal)0)
-		{
-			IntersectableSubmeshEdges.Add(EdgeIndex);
-		}
+		FullMeshToSubMeshIndices[Index].SubMeshIndex = Index;
+		FullMeshToSubMeshIndices[Index].SubMeshType = ESubMeshType::Dynamic;
+		DynamicSubMeshToFullMeshIndices[Index] = Index;
 	}
 }
 
@@ -1387,9 +1384,9 @@ void FPBDTriangleMeshCollisions::Init(const SolverParticlesOrRange& Particles, c
 		bCollidableSubMeshDirty = false;
 	}
 
-	const FTriangleMesh& CollidableMesh = CollidableSubMesh.GetCollidableMesh();
+	const FTriangleMesh& DynamicSubMesh = CollidableSubMesh.GetDynamicSubMesh();
 
-	if (CollidableMesh.GetNumElements() == 0)
+	if (DynamicSubMesh.GetNumElements() == 0)
 	{
 		return;
 	}
@@ -1397,13 +1394,17 @@ void FPBDTriangleMeshCollisions::Init(const SolverParticlesOrRange& Particles, c
 	{
 		TRACE_CPUPROFILER_EVENT_SCOPE(ChaosFPBDTriangleMeshCollisions_BuildSpatialHash);
 		constexpr FSolverReal RadiusToLodSizeMultiplier = 2.; // Radius to Diameter
-		CollidableMesh.BuildSpatialHash(static_cast<const TArrayView<const FSolverVec3>&>(Particles.XArray()), SpatialHash, RadiusToLodSizeMultiplier * MinProximityQueryRadius);
+		DynamicSubMesh.BuildSpatialHash(static_cast<const TArrayView<const FSolverVec3>&>(Particles.XArray()), DynamicSubMeshSpatialHash, RadiusToLodSizeMultiplier * MinProximityQueryRadius);
+		if (CollidableSubMesh.GetKinematicColliderSubMesh().GetNumElements() > 0)
+		{
+			CollidableSubMesh.GetKinematicColliderSubMesh().BuildSpatialHash(static_cast<const TArrayView<const FSolverVec3>&>(Particles.XArray()), KinematicSubMeshSpatialHash, RadiusToLodSizeMultiplier * MinProximityQueryRadius);
+		}
 	}
 	ContourMinimizationIntersections.Reset();
 	VertexGIAColors.Reset();
 	VertexGIAColors.SetNumZeroed(NumParticles);
 	TriangleGIAColors.Reset();
-	TriangleGIAColors.SetNumZeroed(CollidableMesh.GetNumElements());
+	TriangleGIAColors.SetNumZeroed(DynamicSubMesh.GetNumElements());
 	IntersectionContourPoints.Reset();
 	IntersectionContourTypes.Reset();
 
@@ -1420,7 +1421,7 @@ void FPBDTriangleMeshCollisions::Init(const SolverParticlesOrRange& Particles, c
 
 	// Detect all EdgeFace Intersections
 	constexpr bool bSkipKinematicTrue = true;
-	FindEdgeFaceIntersections(CollidableMesh, CollidableSubMesh.GetIntersectableSubmeshEdges(), bSkipKinematicTrue, CollidableSubMesh.GetSubMeshElementIsKinematic(), SpatialHash, Particles, ScratchBuffers->EdgeFaceIntersections);
+	FindEdgeFaceIntersections(DynamicSubMesh, DynamicSubMeshSpatialHash, Particles, ScratchBuffers->EdgeFaceIntersections);
 	if (ScratchBuffers->EdgeFaceIntersections.Num() == 0)
 	{
 		return;
@@ -1429,9 +1430,9 @@ void FPBDTriangleMeshCollisions::Init(const SolverParticlesOrRange& Particles, c
 	if (bGlobalIntersectionAnalysis)
 	{ 
 		// Walk EdgeFace intersections to build global contours
-		GIA::ContourBuilding::BuildIntersectionContours(CollidableMesh, ScratchBuffers->EdgeFaceIntersections, ScratchBuffers->IntersectionContours, IntersectionContourPoints);
+		GIA::ContourBuilding::BuildIntersectionContours(DynamicSubMesh, ScratchBuffers->EdgeFaceIntersections, ScratchBuffers->IntersectionContours, IntersectionContourPoints);
 		// Flood fill global contours to determine intersecting regions.
-		GIA::FloodFill::FloodFillContours(CollidableMesh, NumParticles, Offset, ScratchBuffers->IntersectionContours, VertexGIAColors, TriangleGIAColors);
+		GIA::FloodFill::FloodFillContours(DynamicSubMesh, NumParticles, Offset, ScratchBuffers->IntersectionContours, VertexGIAColors, TriangleGIAColors);
 		GIA::FloodFill::AssignContourPointTypes(ScratchBuffers->IntersectionContours, IntersectionContourPoints, IntersectionContourTypes);
 	}
 
@@ -1440,12 +1441,12 @@ void FPBDTriangleMeshCollisions::Init(const SolverParticlesOrRange& Particles, c
 		if (bGlobalIntersectionAnalysis)
 		{
 			// Global contours which are non-closed or loop are handled via ContourMinimization impulses. Build global gradient (by adding contribution across all intersections per contour).
-			ContourMinimization::BuildGlobalContourMinimizationIntersections(CollidableMesh, Particles, ScratchBuffers->IntersectionContours, ContourMinimizationIntersections);
+			ContourMinimization::BuildGlobalContourMinimizationIntersections(DynamicSubMesh, Particles, ScratchBuffers->IntersectionContours, ContourMinimizationIntersections);
 		}
 		else
 		{
 			// Just build local gradient for all Intersections
-			ContourMinimization::BuildLocalContourMinimizationIntersections(CollidableMesh, Particles, ScratchBuffers->EdgeFaceIntersections, ContourMinimizationIntersections);
+			ContourMinimization::BuildLocalContourMinimizationIntersections(DynamicSubMesh, Particles, ScratchBuffers->EdgeFaceIntersections, ContourMinimizationIntersections);
 		}
 	}
 }
@@ -1455,8 +1456,8 @@ template CHAOS_API void FPBDTriangleMeshCollisions::Init(const FSolverParticlesR
 template<typename SolverParticlesOrRange>
 void FPBDTriangleMeshCollisions::PostStepInit(const SolverParticlesOrRange& Particles)
 {
-	const FTriangleMesh& CollidableMesh = CollidableSubMesh.GetCollidableMesh();
-	if (CollidableMesh.GetNumElements() == 0)
+	const FTriangleMesh& DynamicSubMesh = CollidableSubMesh.GetDynamicSubMesh();
+	if (DynamicSubMesh.GetNumElements() == 0)
 	{
 		return;
 	}
@@ -1470,7 +1471,7 @@ void FPBDTriangleMeshCollisions::PostStepInit(const SolverParticlesOrRange& Part
 		{
 			TRACE_CPUPROFILER_EVENT_SCOPE(ChaosFPBDTriangleMeshCollisions_BuildSpatialHash);
 			constexpr FSolverReal MinSpatialLodSize = 1.f;
-			CollidableMesh.BuildSpatialHash(static_cast<const TArrayView<const FSolverVec3>&>(Particles.XArray()), SpatialHash, MinSpatialLodSize);
+			DynamicSubMesh.BuildSpatialHash(static_cast<const TArrayView<const FSolverVec3>&>(Particles.XArray()), DynamicSubMeshSpatialHash, MinSpatialLodSize);
 
 			if (!ScratchBuffers)
 			{
@@ -1479,18 +1480,17 @@ void FPBDTriangleMeshCollisions::PostStepInit(const SolverParticlesOrRange& Part
 			ScratchBuffers->Reset();
 
 			// Detect all EdgeFace Intersections
-			constexpr bool bSkipKinematicTrue = true;
-			FindEdgeFaceIntersections(CollidableMesh, CollidableSubMesh.GetIntersectableSubmeshEdges(), bSkipKinematicTrue, CollidableSubMesh.GetSubMeshElementIsKinematic(), SpatialHash, Particles, ScratchBuffers->EdgeFaceIntersections);
+			FindEdgeFaceIntersections(DynamicSubMesh, DynamicSubMeshSpatialHash, Particles, ScratchBuffers->EdgeFaceIntersections);
 
 			if (bUseGlobalPostStepContours)
 			{
-				GIA::ContourBuilding::BuildIntersectionContours(CollidableMesh, ScratchBuffers->EdgeFaceIntersections, ScratchBuffers->IntersectionContours, PostStepIntersectionContourPoints);
-				ContourMinimization::BuildGlobalContourMinimizationIntersections(CollidableMesh, Particles, ScratchBuffers->IntersectionContours, PostStepContourMinimizationIntersections);
+				GIA::ContourBuilding::BuildIntersectionContours(DynamicSubMesh, ScratchBuffers->EdgeFaceIntersections, ScratchBuffers->IntersectionContours, PostStepIntersectionContourPoints);
+				ContourMinimization::BuildGlobalContourMinimizationIntersections(DynamicSubMesh, Particles, ScratchBuffers->IntersectionContours, PostStepContourMinimizationIntersections);
 			}
 			else
 			{
 				// Just build local gradient for all Intersections
-				ContourMinimization::BuildLocalContourMinimizationIntersections(CollidableMesh, Particles, ScratchBuffers->EdgeFaceIntersections, PostStepContourMinimizationIntersections);
+				ContourMinimization::BuildLocalContourMinimizationIntersections(DynamicSubMesh, Particles, ScratchBuffers->EdgeFaceIntersections, PostStepContourMinimizationIntersections);
 			}
 		}
 	}
diff --git a/Engine/Source/Runtime/Experimental/Chaos/Public/Chaos/PBDCollisionSpringConstraintsBase.h b/Engine/Source/Runtime/Experimental/Chaos/Public/Chaos/PBDCollisionSpringConstraintsBase.h
--- a/Engine/Source/Runtime/Experimental/Chaos/Public/Chaos/PBDCollisionSpringConstraintsBase.h
+++ b/Engine/Source/Runtime/Experimental/Chaos/Public/Chaos/PBDCollisionSpringConstraintsBase.h
@@ -63,14 +63,12 @@ public:
 	PRAGMA_ENABLE_DEPRECATION_WARNINGS
 
 	template<typename SpatialAccelerator, typename SolverParticlesOrRange>
-	void Init(const SolverParticlesOrRange& Particles, const SpatialAccelerator& Spatial, const TConstArrayView<FPBDTriangleMeshCollisions::FGIAColor>& VertexGIAColors, const TArray<FPBDTriangleMeshCollisions::FGIAColor>& TriangleGIAColors)
-	{
-		constexpr FSolverReal Dt_BigNumber = UE_BIG_NUMBER; // This will disable the timers feature on kinematic colliders.
-		return Init(Particles, Dt_BigNumber, FPBDTriangleMeshCollisions::FTriangleSubMesh(TriangleMesh), Spatial, VertexGIAColors, TriangleGIAColors);
-	}
+	UE_DEPRECATED(5.4, "Use Init with CollidableSubMesh. This method is much less efficient as it recreates the CollidableSubMesh each call.")
+	void Init(const SolverParticlesOrRange& Particles, const SpatialAccelerator& Spatial, const TConstArrayView<FPBDTriangleMeshCollisions::FGIAColor>& VertexGIAColors, const TArray<FPBDTriangleMeshCollisions::FGIAColor>& TriangleGIAColors);
 
 	template<typename SpatialAccelerator, typename SolverParticlesOrRange>
-	void Init(const SolverParticlesOrRange& Particles, const FSolverReal Dt, const FPBDTriangleMeshCollisions::FTriangleSubMesh& CollidableSubMesh, const SpatialAccelerator& Spatial, const TConstArrayView<FPBDTriangleMeshCollisions::FGIAColor>& VertexGIAColors, const TArray<FPBDTriangleMeshCollisions::FGIAColor>& TriangleGIAColors);
+	void Init(const SolverParticlesOrRange& Particles, const FSolverReal Dt, const FPBDTriangleMeshCollisions::FTriangleSubMesh& CollidableSubMesh,
+		const SpatialAccelerator& DynamicSpatial, const SpatialAccelerator& KinematicColliderSpatial, const TConstArrayView<FPBDTriangleMeshCollisions::FGIAColor>& VertexGIAColors, const TArray<FPBDTriangleMeshCollisions::FGIAColor>& TriangleGIAColors);
 
 	template<typename SolverParticlesOrRange>
 	CHAOS_API FSolverVec3 GetDelta(const SolverParticlesOrRange& InParticles, const int32 i) const;
@@ -87,6 +85,7 @@ public:
 		return Barys;
 		PRAGMA_ENABLE_DEPRECATION_WARNINGS
 	}
+
 	FSolverReal GetThickness() const { return Thickness; }
 	bool GetGlobalIntersectionAnalysis() const { return bGlobalIntersectionAnalysis; }
 	const TArray<bool>& GetFlipNormals() const 
@@ -129,15 +128,7 @@ public:
 	}
 
 	template<typename SolverParticlesOrRange>
-	void Apply(SolverParticlesOrRange& InParticles, const FSolverReal Dt) const
-	{
-		PRAGMA_DISABLE_DEPRECATION_WARNINGS
-		for (int32 ConstraintIndex = 0; ConstraintIndex < Constraints.Num(); ++ConstraintIndex)
-		{
-			Apply(InParticles, Dt, ConstraintIndex);
-		}
-		PRAGMA_ENABLE_DEPRECATION_WARNINGS
-	}
+	void Apply(SolverParticlesOrRange& InParticles, const FSolverReal Dt) const;
 
 	void Apply(FSolverParticles& InParticles, const FSolverReal Dt, const TArray<int32>& InConstraintIndices) const
 	{
@@ -212,6 +203,9 @@ protected:
 
 private:
 
+	template<typename SolverParticlesOrRange>
+	void ApplyKinematicConstraints(SolverParticlesOrRange& InParticles, const FSolverReal Dt) const;
+
 	const FTriangleMesh& TriangleMesh;
 	const TArray<FSolverVec3>* ReferencePositions;
 	const TSet<TVec2<int32>> DisabledCollisionElements;  // TODO: Make this a bitarray
@@ -226,17 +220,27 @@ private:
 	};
 	TArray<EConstraintType> ConstraintTypes;
 
-	struct FExistingConstraintData
-	{
-		FSolverReal Timer;
-	};
-	TArray<TMap<int32, FExistingConstraintData>> ExistingConstraintLookup; // ArrayIndex = ParticleIndex (without offset), MapKey = FaceIndex (full mesh)
+	static constexpr int32 MaxKinematicConnectionsPerPoint = 3;
+	// Parallel arrays (for ISPC SOA)
+	TArray<int32> KinematicCollidingParticles;
+	TArray<TVector<int32, MaxKinematicConnectionsPerPoint>> KinematicColliderElements;
+
+	TArray<TMap<int32, FSolverReal>> KinematicColliderTimers; // Keep constraints for a cvar-defined time after it's moved out of proximity. Helps reduce jitter.
 
 	int32 Offset;
 	int32 NumParticles;
 	bool bGlobalIntersectionAnalysis; // This is set based on which Init is called.
 };
-
 }  // End namespace Chaos::Softs
 
+// Support ISPC enable/disable in non-shipping builds
+#if !INTEL_ISPC
+const bool bChaos_CollisionSpring_ISPC_Enabled = false;
+#elif UE_BUILD_SHIPPING
+const bool bChaos_CollisionSpring_ISPC_Enabled = true;
+#else
+extern CHAOS_API bool bChaos_CollisionSpring_ISPC_Enabled;
+#endif
+
+
 #endif
diff --git a/Engine/Source/Runtime/Experimental/Chaos/Public/Chaos/PBDTriangleMeshCollisions.h b/Engine/Source/Runtime/Experimental/Chaos/Public/Chaos/PBDTriangleMeshCollisions.h
--- a/Engine/Source/Runtime/Experimental/Chaos/Public/Chaos/PBDTriangleMeshCollisions.h
+++ b/Engine/Source/Runtime/Experimental/Chaos/Public/Chaos/PBDTriangleMeshCollisions.h
@@ -114,66 +114,76 @@ public:
 	public:
 		FTriangleSubMesh(const FTriangleMesh& InFullMesh)
 			: FullMesh(InFullMesh)
-			, bSubMeshIsFullMesh(true)
 		{}
 
 		template<typename SolverParticlesOrRange>
 		void Init(const SolverParticlesOrRange& Particles, const TSet<int32>& DisabledFaces, bool bCollideAgainstAllKinematicVertices, const TSet<int32>& EnabledKinematicFaces);
 
-		const FTriangleMesh& GetCollidableMesh() const
-		{
-			if (bSubMeshIsFullMesh)
-			{
-				return FullMesh;
-			}
-			return SubMesh;
-		}
+		void InitAllDynamic();
 
 		const FTriangleMesh& GetFullMesh() const 
 		{
 			return FullMesh;
 		}
 
-		int32 FullMeshIndexFromSubIndex(int32 SubIndex) const
+		const FTriangleMesh& GetDynamicSubMesh() const
 		{
-			if (bSubMeshIsFullMesh)
-			{
-				return SubIndex;
-			}
+			return DynamicSubMesh;
+		}
 
-			return SubMeshToFullElements[SubIndex];
+		const FTriangleMesh& GetKinematicColliderSubMesh() const
+		{
+			return KinematicColliderSubMesh;
 		}
 
-		int32 SubMeshIndexFromFullIndex(int32 FullIndex) const
+		bool IsElementDynamic(int32 FullMeshIndex) const
 		{
-			if (bSubMeshIsFullMesh)
-			{
-				return FullIndex;
-			}
-			return FullToSubMeshElements[FullIndex];
+			return FullMeshToSubMeshIndices[FullMeshIndex].SubMeshType == ESubMeshType::Dynamic;
 		}
 
-		bool GetSubMeshIsFullMesh() const { return bSubMeshIsFullMesh; }
+		bool IsElementKinematicCollider(int32 FullMeshIndex) const
+		{
+			return FullMeshToSubMeshIndices[FullMeshIndex].SubMeshType == ESubMeshType::Kinematic;
+		}
+
+		int32 GetSubMeshElementIndex(int32 FullMeshIndex) const
+		{
+			check(FullMeshToSubMeshIndices[FullMeshIndex].SubMeshType != ESubMeshType::Invalid);
+			return FullMeshToSubMeshIndices[FullMeshIndex].SubMeshIndex;
+		}
+
+		int32 GetFullMeshElementIndexFromDynamicElement(int32 DynamicMeshIndex) const
+		{
+			return DynamicSubMeshToFullMeshIndices[DynamicMeshIndex];
+		}
 
-		const TArray<bool>& GetSubMeshElementIsKinematic() const 
+		int32 GetFullMeshElementIndexFromKinematicElement(int32 KinematicMeshIndex) const
 		{
-			return SubMeshElementIsKinematic;
+			return KinematicColliderSubMeshToFullMeshIndices[KinematicMeshIndex];
 		}
 
-		// Will be empty when bSubMeshIsFullMesh
-		const TArray<int32>& GetCollidableVertices() const { return CollidableVertices; }
-		// Will be empty when bSubMeshIsFullMesh
-		const TArray<int32>& GetIntersectableSubmeshEdges() const { return IntersectableSubmeshEdges; }
+		const TArray<int32>& GetDynamicVertices() const { return DynamicVertices; }
 
 	private:
 		const FTriangleMesh& FullMesh;
-		FTriangleMesh SubMesh;
-		TArray<int32> SubMeshToFullElements;
-		TArray<int32> FullToSubMeshElements;
-		TArray<bool> SubMeshElementIsKinematic;
-		TArray<int32> CollidableVertices;
-		TArray<int32> IntersectableSubmeshEdges;
-		bool bSubMeshIsFullMesh;
+		FTriangleMesh DynamicSubMesh;
+		FTriangleMesh KinematicColliderSubMesh;
+		
+		enum struct ESubMeshType : int32
+		{
+			Invalid = 0,
+			Dynamic = 1,
+			Kinematic = 2
+		};
+		struct FFullToSubMeshIndex
+		{
+			int32 SubMeshIndex : 30;
+			ESubMeshType SubMeshType : 2;
+		};
+		TArray<FFullToSubMeshIndex> FullMeshToSubMeshIndices;
+		TArray<int32> DynamicSubMeshToFullMeshIndices;
+		TArray<int32> KinematicColliderSubMeshToFullMeshIndices;
+		TArray<int32> DynamicVertices; // Will be empty if InitAllDynamic was used to initialize
 	};
 
 	static bool IsEnabled(const FCollectionPropertyConstFacade& PropertyCollection)
@@ -346,7 +356,10 @@ public:
 	}
 
 	const FTriangleSubMesh& GetCollidableSubMesh() const { return CollidableSubMesh; }
-	const FTriangleMesh::TSpatialHashType<FSolverReal>& GetSpatialHash() const { return SpatialHash; }
+	UE_DEPRECATED(5.4, "Use GetDynamicSpatialHash or GetKinematicColliderSpatialHash")
+	const FTriangleMesh::TSpatialHashType<FSolverReal>& GetSpatialHash() const { return DynamicSubMeshSpatialHash; }
+	const FTriangleMesh::TSpatialHashType<FSolverReal>& GetDynamicSpatialHash() const { return DynamicSubMeshSpatialHash; }
+	const FTriangleMesh::TSpatialHashType<FSolverReal>& GetKinematicColliderSpatialHash() const { return KinematicSubMeshSpatialHash; }
 	const TArray<FContourMinimizationIntersection>& GetContourMinimizationIntersections() const { return ContourMinimizationIntersections; }
 	const TConstArrayView<FGIAColor> GetVertexGIAColors() const { return bGlobalIntersectionAnalysis && VertexGIAColors.Num() == NumParticles ? TConstArrayView<FGIAColor>(VertexGIAColors.GetData() - Offset, NumParticles + Offset) : TConstArrayView<FGIAColor>(); }
 	const TArray<FGIAColor>& GetTriangleGIAColors() const { return TriangleGIAColors; }
@@ -373,7 +386,8 @@ private:
 
 	bool bCollidableSubMeshDirty = true;
 	
-	FTriangleMesh::TSpatialHashType<FSolverReal> SpatialHash;
+	FTriangleMesh::TSpatialHashType<FSolverReal> DynamicSubMeshSpatialHash;
+	FTriangleMesh::TSpatialHashType<FSolverReal> KinematicSubMeshSpatialHash;
 	TArray<FContourMinimizationIntersection> ContourMinimizationIntersections;
 	TArray<FGIAColor> VertexGIAColors;
 	TArray<FGIAColor> TriangleGIAColors;
