Follow up from new physics particle api

diff --git a/Engine/Source/Runtime/Experimental/Chaos/Private/Chaos/ContactModification.cpp b/Engine/Source/Runtime/Experimental/Chaos/Private/Chaos/ContactModification.cpp
--- a/Engine/Source/Runtime/Experimental/Chaos/Private/Chaos/ContactModification.cpp
+++ b/Engine/Source/Runtime/Experimental/Chaos/Private/Chaos/ContactModification.cpp
@@ -392,7 +392,7 @@ namespace Chaos
 			EObjectStateType State(RigidHandle->ObjectState());
 			if (State == EObjectStateType::Dynamic || State == EObjectStateType::Sleeping)
 			{
-				RigidHandle->SetQf(Rotation);
+				RigidHandle->SetQ(Rotation);
 
 				if (bMaintainVelocity)
 				{
@@ -401,16 +401,17 @@ namespace Chaos
 				else if (Modifier->Dt > 0.0f)
 				{
 					// Update W to new implicit velocity
-					RigidHandle->SetWf(FRotation3::CalculateAngularVelocity(RigidHandle->GetRf(), RigidHandle->GetQf(), FRealSingle(Modifier->Dt)));
+					RigidHandle->SetWf(FRotation3f::CalculateAngularVelocity(RigidHandle->GetRf(), RigidHandle->GetQf(), FRealSingle(Modifier->Dt)));
 				}
 				UpdateConstraintShapeTransforms();
 				return;
 			}
 			else
 			{
+				FRotation3f Rotationf(Rotation);
 				// Kinematic must keep Q/R in sync
-				RigidHandle->SetRf(Rotation);
-				RigidHandle->SetQf(Rotation);
+				RigidHandle->SetRf(Rotationf);
+				RigidHandle->SetQf(Rotationf);
 				UpdateConstraintShapeTransforms();
 				return;
 			}
@@ -419,7 +420,7 @@ namespace Chaos
 		FKinematicGeometryParticleHandle* KinematicHandle = Particle->CastToKinematicParticle();
 		if (KinematicHandle)
 		{
-			KinematicHandle->SetRf(Rotation);
+			KinematicHandle->SetR(Rotation);
 			UpdateConstraintShapeTransforms();
 			return;
 		}
diff --git a/Engine/Source/Runtime/Experimental/Chaos/Public/Chaos/PerParticlePBDUpdateFromDeltaPosition.h b/Engine/Source/Runtime/Experimental/Chaos/Public/Chaos/PerParticlePBDUpdateFromDeltaPosition.h
--- a/Engine/Source/Runtime/Experimental/Chaos/Public/Chaos/PerParticlePBDUpdateFromDeltaPosition.h
+++ b/Engine/Source/Runtime/Experimental/Chaos/Public/Chaos/PerParticlePBDUpdateFromDeltaPosition.h
@@ -29,7 +29,7 @@ class FPerParticlePBDUpdateFromDeltaPosition : public FPerParticleRule
 	inline void Apply(TPBDRigidParticles<FReal, 3>& InParticles, const FReal Dt, const int32 Index) const override //-V762
 	{
 		ApplyHelper(InParticles, Dt, Index);
-		InParticles.SetW(Index, FRotation3::CalculateAngularVelocity(InParticles.GetR(Index), InParticles.GetQ(Index), Dt));
+		InParticles.SetW(Index, FRotation3f::CalculateAngularVelocity(InParticles.GetRf(Index), InParticles.GetQf(Index), FRealSingle(Dt)));
 	}
 
 	inline void Apply(FPBDRigidParticleHandle* Particle, const FReal Dt) const override //-V762
@@ -38,7 +38,7 @@ class FPerParticlePBDUpdateFromDeltaPosition : public FPerParticleRule
 		const FVec3 CenteredX = Particle->XCom();
 		const FVec3 CenteredP = Particle->PCom();
 		Particle->SetV(FVec3::CalculateVelocity(CenteredX, CenteredP, Dt));
-		Particle->SetW(FRotation3::CalculateAngularVelocity(Particle->GetR(), Particle->GetQ(), Dt));
+		Particle->SetWf(FRotation3f::CalculateAngularVelocity(Particle->GetRf(), Particle->GetQf(), FRealSingle(Dt)));
 	}
 
 	inline void Apply(TTransientPBDRigidParticleHandle<FReal, 3>& Particle, const FReal Dt) const override //-V762
@@ -47,7 +47,7 @@ class FPerParticlePBDUpdateFromDeltaPosition : public FPerParticleRule
 		const FVec3 CenteredX = Particle.XCom();
 		const FVec3 CenteredP = Particle.PCom();
 		Particle.SetV(FVec3::CalculateVelocity(CenteredX, CenteredP, Dt));
-		Particle.SetW(FRotation3::CalculateAngularVelocity(Particle.GetR(), Particle.GetQ(), Dt));
+		Particle.SetWf(FRotation3f::CalculateAngularVelocity(Particle.GetRf(), Particle.GetQf(), FRealSingle(Dt)));
 	}
 };
 
