[ChaosVD] Serialization optimizations

- Created serializable name table, that serializes unique FNames with their IDs, deduplicating string data and therefore reducing the amount of data we need to serialize
- Packed several Booleans into integers before serializing them. This reduces the amount of data being serialized as each bool is serialized as int32 and requires a virtual call for each Boolean serialized, therefore by packing them and serializing them as integer skips part of the work
- Skip serialization of "invalid" data wrappers. Data wrappers not marked as valid, are structs that were not populated with data so we can just not serialize them and load the default values in the CVD editor
- Changed a cvar to be a console variable ref to avoid a virtual call.

diff --git a/Engine/Plugins/Experimental/ChaosVD/Source/ChaosVD/Private/ChaosVDScene.cpp b/Engine/Plugins/Experimental/ChaosVD/Source/ChaosVD/Private/ChaosVDScene.cpp
--- a/Engine/Plugins/Experimental/ChaosVD/Source/ChaosVD/Private/ChaosVDScene.cpp
+++ b/Engine/Plugins/Experimental/ChaosVD/Source/ChaosVD/Private/ChaosVDScene.cpp
@@ -397,11 +397,7 @@ AChaosVDParticleActor* FChaosVDScene::SpawnParticleFromRecordedData(const TShare
 		return nullptr;
 	}
 
-	FActorSpawnParameters Params;
-	Params.Name = *InParticleData->DebugName;
-	Params.NameMode = FActorSpawnParameters::ESpawnActorNameMode::Requested;
-
-	if (AChaosVDParticleActor* NewActor = PhysicsVDWorld->SpawnActor<AChaosVDParticleActor>(Params))
+	if (AChaosVDParticleActor* NewActor = PhysicsVDWorld->SpawnActor<AChaosVDParticleActor>())
 	{
 		NewActor->SetIsActive(true);
 		NewActor->SetScene(AsShared());
diff --git a/Engine/Plugins/Experimental/ChaosVD/Source/ChaosVD/Private/Trace/ChaosVDTraceProvider.cpp b/Engine/Plugins/Experimental/ChaosVD/Source/ChaosVD/Private/Trace/ChaosVDTraceProvider.cpp
--- a/Engine/Plugins/Experimental/ChaosVD/Source/ChaosVD/Private/Trace/ChaosVDTraceProvider.cpp
+++ b/Engine/Plugins/Experimental/ChaosVD/Source/ChaosVD/Private/Trace/ChaosVDTraceProvider.cpp
@@ -6,6 +6,8 @@
 #include "ChaosVDRecording.h"
 
 #include "Chaos/ChaosArchive.h"
+#include "ChaosVisualDebugger/ChaosVDSerializedNameTable.h"
+#include "ChaosVisualDebugger/ChaosVisualDebuggerTrace.h"
 
 #include "Compression/OodleDataCompressionUtil.h"
 #include "Serialization/MemoryReader.h"
@@ -13,6 +15,7 @@
 #include "Trace/DataProcessors/ChaosVDMidPhaseDataProcessor.h"
 #include "Trace/DataProcessors/ChaosVDSceneQueryDataProcessor.h"
 #include "Trace/DataProcessors/ChaosVDSceneQueryVisitDataProcessor.h"
+#include "Trace/DataProcessors/ChaosVDSerializedNameEntryDataProcessor.h"
 #include "Trace/DataProcessors/ChaosVDTraceImplicitObjectProcessor.h"
 #include "Trace/DataProcessors/ChaosVDTraceParticleDataProcessor.h"
 
@@ -181,6 +184,11 @@ void FChaosVDTraceProvider::RegisterDataProcessor(TSharedPtr<IChaosVDDataProcess
 	RegisteredDataProcessors.Add(InDataProcessor->GetCompatibleTypeName(), InDataProcessor);
 }
 
+TSharedPtr<Chaos::VisualDebugger::FChaosVDSerializableNameTable> FChaosVDTraceProvider::GetNameTable() const
+{
+	return InternalRecording.IsValid() ? InternalRecording->GetNameTableInstance() : nullptr; 
+}
+
 void FChaosVDTraceProvider::RegisterDefaultDataProcessorsIfNeeded()
 {
 	if (bDefaultDataProcessorsRegistered)
@@ -212,6 +220,10 @@ void FChaosVDTraceProvider::RegisterDefaultDataProcessorsIfNeeded()
 	SceneQueryVisitDataProcessor->SetTraceProvider(AsShared());
 	RegisterDataProcessor(SceneQueryVisitDataProcessor);
 
+	TSharedPtr<FChaosVDSerializedNameEntryDataProcessor> NameEntryDataProcessor = MakeShared<FChaosVDSerializedNameEntryDataProcessor>();
+	NameEntryDataProcessor->SetTraceProvider(AsShared());
+	RegisterDataProcessor(NameEntryDataProcessor);
+
 	bDefaultDataProcessorsRegistered = true;
 }
 
diff --git a/Engine/Plugins/Experimental/ChaosVD/Source/ChaosVD/Private/Trace/DataProcessors/ChaosVDConstraintDataProcessor.cpp b/Engine/Plugins/Experimental/ChaosVD/Source/ChaosVD/Private/Trace/DataProcessors/ChaosVDConstraintDataProcessor.cpp
--- a/Engine/Plugins/Experimental/ChaosVD/Source/ChaosVD/Private/Trace/DataProcessors/ChaosVDConstraintDataProcessor.cpp
+++ b/Engine/Plugins/Experimental/ChaosVD/Source/ChaosVD/Private/Trace/DataProcessors/ChaosVDConstraintDataProcessor.cpp
@@ -3,6 +3,7 @@
 #include "Trace/DataProcessors/ChaosVDConstraintDataProcessor.h"
 
 #include "ChaosVDRecording.h"
+#include "ChaosVisualDebugger/ChaosVDMemWriterReader.h"
 #include "DataWrappers/ChaosVDCollisionDataWrappers.h"
 #include "Serialization/MemoryReader.h"
 #include "Trace/ChaosVDTraceProvider.h"
@@ -19,21 +20,28 @@ bool FChaosVDConstraintDataProcessor::ProcessRawData(const TArray<uint8>& InData
 		return false;
 	}
 
+	if (!ensure(ProviderSharedPtr->GetNameTable().IsValid()))
+	{
+		return false;
+	}
+
 	FChaosVDConstraint RecordedConstraint;
-	FMemoryReader MemReader(InData);
-	RecordedConstraint.Serialize(MemReader);
+	const bool bSuccess = ReadDataFromBuffer(InData, RecordedConstraint, ProviderSharedPtr->GetNameTable().ToSharedRef());
 
-	// This can be null if the recording started Mid-Frame. In this case we just discard the data for now
-	if (FChaosVDSolverFrameData* FrameData = ProviderSharedPtr->GetCurrentSolverFrame(RecordedConstraint.SolverID))
+	if (bSuccess)
 	{
-		if (ensureMsgf(FrameData->SolverSteps.Num() > 0, TEXT("A MidPhase was traced without a valid step scope")))
+		// This can be null if the recording started Mid-Frame. In this case we just discard the data for now
+		if (FChaosVDSolverFrameData* FrameData = ProviderSharedPtr->GetCurrentSolverFrame(RecordedConstraint.SolverID))
 		{
-			AddConstraintToParticleIDMap(RecordedConstraint, RecordedConstraint.Particle0Index, *FrameData);
-			AddConstraintToParticleIDMap(RecordedConstraint, RecordedConstraint.Particle1Index, *FrameData);
+			if (ensureMsgf(FrameData->SolverSteps.Num() > 0, TEXT("A MidPhase was traced without a valid step scope")))
+			{
+				AddConstraintToParticleIDMap(RecordedConstraint, RecordedConstraint.Particle0Index, *FrameData);
+				AddConstraintToParticleIDMap(RecordedConstraint, RecordedConstraint.Particle1Index, *FrameData);
+			}
 		}
 	}
 
-	return !MemReader.IsError() && !MemReader.IsCriticalError();
+	return bSuccess;
 }
 
 void FChaosVDConstraintDataProcessor::AddConstraintToParticleIDMap(const FChaosVDConstraint& InConstraintData, int32 ParticleID, FChaosVDSolverFrameData& InFrameData)
diff --git a/Engine/Plugins/Experimental/ChaosVD/Source/ChaosVD/Private/Trace/DataProcessors/ChaosVDMidPhaseDataProcessor.cpp b/Engine/Plugins/Experimental/ChaosVD/Source/ChaosVD/Private/Trace/DataProcessors/ChaosVDMidPhaseDataProcessor.cpp
--- a/Engine/Plugins/Experimental/ChaosVD/Source/ChaosVD/Private/Trace/DataProcessors/ChaosVDMidPhaseDataProcessor.cpp
+++ b/Engine/Plugins/Experimental/ChaosVD/Source/ChaosVD/Private/Trace/DataProcessors/ChaosVDMidPhaseDataProcessor.cpp
@@ -3,7 +3,7 @@
 #include "Trace/DataProcessors/ChaosVDMidPhaseDataProcessor.h"
 #include "DataWrappers/ChaosVDCollisionDataWrappers.h"
 #include "ChaosVDRecording.h"
-#include "Serialization/MemoryReader.h"
+#include "ChaosVisualDebugger/ChaosVDMemWriterReader.h"
 #include "Trace/ChaosVDTraceProvider.h"
 
 FChaosVDMidPhaseDataProcessor::FChaosVDMidPhaseDataProcessor() : IChaosVDDataProcessor(FChaosVDParticlePairMidPhase::WrapperTypeName)
@@ -17,25 +17,33 @@ bool FChaosVDMidPhaseDataProcessor::ProcessRawData(const TArray<uint8>& InData)
 	{
 		return false;
 	}
+	
+	if (!ensure(ProviderSharedPtr->GetNameTable().IsValid()))
+	{
+		return false;
+	}
+
 
 	TSharedPtr<FChaosVDParticlePairMidPhase> MidPhase = MakeShared<FChaosVDParticlePairMidPhase>();
-	FMemoryReader MemReader(InData);
 
-	MidPhase->Serialize(MemReader);
+	const bool bSuccess = ReadDataFromBuffer(InData, *MidPhase, ProviderSharedPtr->GetNameTable().ToSharedRef());
 
-	// This can be null if the recording started Mid-Frame. In this case we just discard the data for now
-	if (FChaosVDSolverFrameData* FrameData = ProviderSharedPtr->GetCurrentSolverFrame(MidPhase->SolverID))
+	if (bSuccess)
 	{
-		if (ensureMsgf(FrameData->SolverSteps.Num() > 0, TEXT("A MidPhase was traced without a valid step scope")))
+		// This can be null if the recording started Mid-Frame. In this case we just discard the data for now
+		if (FChaosVDSolverFrameData* FrameData = ProviderSharedPtr->GetCurrentSolverFrame(MidPhase->SolverID))
 		{
-			FrameData->SolverSteps.Last().RecordedMidPhases.Add(MidPhase);
+			if (ensureMsgf(FrameData->SolverSteps.Num() > 0, TEXT("A MidPhase was traced without a valid step scope")))
+			{
+				FrameData->SolverSteps.Last().RecordedMidPhases.Add(MidPhase);
 
-			AddMidPhaseToParticleIDMap(MidPhase, MidPhase->Particle0Idx, *FrameData);
-			AddMidPhaseToParticleIDMap(MidPhase, MidPhase->Particle1Idx, *FrameData);
+				AddMidPhaseToParticleIDMap(MidPhase, MidPhase->Particle0Idx, *FrameData);
+				AddMidPhaseToParticleIDMap(MidPhase, MidPhase->Particle1Idx, *FrameData);
+			}
 		}
 	}
 
-	return !MemReader.IsError() && !MemReader.IsCriticalError();
+	return bSuccess;
 }
 
 void FChaosVDMidPhaseDataProcessor::AddMidPhaseToParticleIDMap(const TSharedPtr<FChaosVDParticlePairMidPhase>& MidPhaseData, int32 ParticleID, FChaosVDSolverFrameData& InFrameData)
diff --git a/Engine/Plugins/Experimental/ChaosVD/Source/ChaosVD/Private/Trace/DataProcessors/ChaosVDSceneQueryDataProcessor.cpp b/Engine/Plugins/Experimental/ChaosVD/Source/ChaosVD/Private/Trace/DataProcessors/ChaosVDSceneQueryDataProcessor.cpp
--- a/Engine/Plugins/Experimental/ChaosVD/Source/ChaosVD/Private/Trace/DataProcessors/ChaosVDSceneQueryDataProcessor.cpp
+++ b/Engine/Plugins/Experimental/ChaosVD/Source/ChaosVD/Private/Trace/DataProcessors/ChaosVDSceneQueryDataProcessor.cpp
@@ -3,11 +3,12 @@
 #include "Trace/DataProcessors/ChaosVDSceneQueryDataProcessor.h"
 
 #include "ChaosVDRecording.h"
+#include "ChaosVisualDebugger/ChaosVDMemWriterReader.h"
+#include "ChaosVisualDebugger/ChaosVDSerializedNameTable.h"
+#include "ChaosVisualDebugger/ChaosVisualDebuggerTrace.h"
 #include "DataWrappers/ChaosVDQueryDataWrappers.h"
-#include "Serialization/MemoryReader.h"
 #include "Trace/ChaosVDTraceProvider.h"
 
-
 FChaosVDSceneQueryDataProcessor::FChaosVDSceneQueryDataProcessor() : IChaosVDDataProcessor(FChaosVDQueryDataWrapper::WrapperTypeName)
 {
 }
@@ -19,26 +20,31 @@ bool FChaosVDSceneQueryDataProcessor::ProcessRawData(const TArray<uint8>& InData
 	{
 		return false;
 	}
-
-	if (const TSharedPtr<FChaosVDGameFrameData> CurrentFrameData = ProviderSharedPtr->GetCurrentGameFrame().Pin())
+	if (!ensure(ProviderSharedPtr->GetNameTable().IsValid()))
 	{
-		const TSharedPtr<FChaosVDQueryDataWrapper> QueryData = MakeShared<FChaosVDQueryDataWrapper>();
-		FMemoryReader MemReader(InData);
+		return false;
+	}
 
-		QueryData->Serialize(MemReader);
+	const TSharedPtr<FChaosVDQueryDataWrapper> QueryData = MakeShared<FChaosVDQueryDataWrapper>();
+	const bool bSuccess = ReadDataFromBuffer(InData, *QueryData, ProviderSharedPtr->GetNameTable().ToSharedRef());
 
-		// If ParentQueryID was set, this is a sub query, so find the parent add it to the sub-queries list so we can navigate trough the query "hierarchy" later on
-		if (QueryData->ParentQueryID != INDEX_NONE)
+	if (bSuccess)
+	{
+		if (const TSharedPtr<FChaosVDGameFrameData> CurrentFrameData = ProviderSharedPtr->GetCurrentGameFrame().Pin())
 		{
-			if (const TSharedPtr<FChaosVDQueryDataWrapper>* ParentQueryData = CurrentFrameData->RecordedSceneQueries.Find(QueryData->ParentQueryID))
+			// If ParentQueryID was set, this is a sub query, so find the parent add it to the sub-queries list so we can navigate trough the query "hierarchy" later on
+			if (QueryData->ParentQueryID != INDEX_NONE)
 			{
-				(*ParentQueryData)->SubQueriesIDs.Add(QueryData->ID);
+				if (const TSharedPtr<FChaosVDQueryDataWrapper>* ParentQueryData = CurrentFrameData->RecordedSceneQueries.Find(QueryData->ParentQueryID))
+				{
+					(*ParentQueryData)->SubQueriesIDs.Add(QueryData->ID);
+				}
 			}
-		}
 
-		CurrentFrameData->RecordedSceneQueries.Add(QueryData->ID, QueryData);
+			CurrentFrameData->RecordedSceneQueries.Add(QueryData->ID, QueryData);
+		}
 	}
 
-	return true;
+	return bSuccess;
 }
 
diff --git a/Engine/Plugins/Experimental/ChaosVD/Source/ChaosVD/Private/Trace/DataProcessors/ChaosVDSceneQueryVisitDataProcessor.cpp b/Engine/Plugins/Experimental/ChaosVD/Source/ChaosVD/Private/Trace/DataProcessors/ChaosVDSceneQueryVisitDataProcessor.cpp
--- a/Engine/Plugins/Experimental/ChaosVD/Source/ChaosVD/Private/Trace/DataProcessors/ChaosVDSceneQueryVisitDataProcessor.cpp
+++ b/Engine/Plugins/Experimental/ChaosVD/Source/ChaosVD/Private/Trace/DataProcessors/ChaosVDSceneQueryVisitDataProcessor.cpp
@@ -3,11 +3,12 @@
 #include "Trace/DataProcessors/ChaosVDSceneQueryVisitDataProcessor.h"
 
 #include "ChaosVDRecording.h"
+#include "ChaosVisualDebugger/ChaosVDMemWriterReader.h"
+#include "ChaosVisualDebugger/ChaosVDSerializedNameTable.h"
+#include "ChaosVisualDebugger/ChaosVisualDebuggerTrace.h"
 #include "DataWrappers/ChaosVDQueryDataWrappers.h"
-#include "Serialization/MemoryReader.h"
 #include "Trace/ChaosVDTraceProvider.h"
 
-
 FChaosVDSceneQueryVisitDataProcessor::FChaosVDSceneQueryVisitDataProcessor() : IChaosVDDataProcessor(FChaosVDQueryVisitStep::WrapperTypeName)
 {
 }
@@ -20,29 +21,34 @@ bool FChaosVDSceneQueryVisitDataProcessor::ProcessRawData(const TArray<uint8>& I
 		return false;
 	}
 
-	if (const TSharedPtr<FChaosVDGameFrameData> CurrentFrameData = ProviderSharedPtr->GetCurrentGameFrame().Pin())
+	if (!ensure(ProviderSharedPtr->GetNameTable().IsValid()))
 	{
-		FChaosVDQueryVisitStep VisitStepData;
+		return false;
+	}
 
-		FMemoryReader MemReader(InData);
+	FChaosVDQueryVisitStep VisitStepData;
+	const bool bSuccess = ReadDataFromBuffer(InData, VisitStepData, ProviderSharedPtr->GetNameTable().ToSharedRef());
 
-		VisitStepData.Serialize(MemReader);
-		
-		if (TSharedPtr<FChaosVDQueryDataWrapper>* QueryDataPtrPtr = CurrentFrameData->RecordedSceneQueries.Find(VisitStepData.OwningQueryID))
-		{
-			TSharedPtr<FChaosVDQueryDataWrapper> QueryDataPtr = *QueryDataPtrPtr;
-			if (QueryDataPtrPtr->IsValid())
+	if (bSuccess)
+	{
+		if (const TSharedPtr<FChaosVDGameFrameData> CurrentFrameData = ProviderSharedPtr->GetCurrentGameFrame().Pin())
+		{	
+			if (TSharedPtr<FChaosVDQueryDataWrapper>* QueryDataPtrPtr = CurrentFrameData->RecordedSceneQueries.Find(VisitStepData.OwningQueryID))
 			{
-				if (VisitStepData.HitData.HasValidData())
+				TSharedPtr<FChaosVDQueryDataWrapper> QueryDataPtr = *QueryDataPtrPtr;
+				if (QueryDataPtrPtr->IsValid())
 				{
-					// Quick and dirty way of show the hits in the details panel. If copying this data around becomes a bottle neck we can write a customization layout for it
-					QueryDataPtr->Hits.Add(VisitStepData);
-				}
+					if (VisitStepData.HitData.HasValidData())
+					{
+						// Quick and dirty way of show the hits in the details panel. If copying this data around becomes a bottle neck we can write a customization layout for it
+						QueryDataPtr->Hits.Add(VisitStepData);
+					}
 	
-				QueryDataPtr->SQVisitData.Emplace(MoveTemp(VisitStepData));
+					QueryDataPtr->SQVisitData.Emplace(MoveTemp(VisitStepData));
+				}
 			}
 		}
 	}
 
-	return true;
+	return bSuccess;
 }
\ No newline at end of file
diff --git a/Engine/Plugins/Experimental/ChaosVD/Source/ChaosVD/Private/Trace/DataProcessors/ChaosVDSerializedNameEntryDataProcessor.cpp b/Engine/Plugins/Experimental/ChaosVD/Source/ChaosVD/Private/Trace/DataProcessors/ChaosVDSerializedNameEntryDataProcessor.cpp
new file mode 100644
--- /dev/null
+++ b/Engine/Plugins/Experimental/ChaosVD/Source/ChaosVD/Private/Trace/DataProcessors/ChaosVDSerializedNameEntryDataProcessor.cpp
@@ -0,0 +1,31 @@
+﻿// Copyright Epic Games, Inc. All Rights Reserved.
+
+#include "Trace/DataProcessors/ChaosVDSerializedNameEntryDataProcessor.h"
+
+#include "ChaosVisualDebugger/ChaosVDSerializedNameTable.h"
+#include "ChaosVisualDebugger/ChaosVisualDebuggerTrace.h"
+#include "Serialization/MemoryReader.h"
+#include "Trace/ChaosVDTraceProvider.h"
+
+FChaosVDSerializedNameEntryDataProcessor::FChaosVDSerializedNameEntryDataProcessor()
+	: IChaosVDDataProcessor(Chaos::VisualDebugger::FChaosVDSerializedNameEntry::WrapperTypeName)
+{
+}
+
+bool FChaosVDSerializedNameEntryDataProcessor::ProcessRawData(const TArray<uint8>& InData)
+{
+	const TSharedPtr<FChaosVDTraceProvider> ProviderSharedPtr = TraceProvider.Pin();
+	if (!ensure(ProviderSharedPtr.IsValid()))
+	{
+		return false;
+	}
+
+	Chaos::VisualDebugger::FChaosVDSerializedNameEntry NameEntry;
+	
+	FMemoryReader MemReader(InData);
+	MemReader << NameEntry;
+
+	ProviderSharedPtr->GetNameTable()->AddNameToTable(NameEntry);
+
+	return true;
+}
diff --git a/Engine/Plugins/Experimental/ChaosVD/Source/ChaosVD/Private/Trace/DataProcessors/ChaosVDTraceImplicitObjectProcessor.cpp b/Engine/Plugins/Experimental/ChaosVD/Source/ChaosVD/Private/Trace/DataProcessors/ChaosVDTraceImplicitObjectProcessor.cpp
--- a/Engine/Plugins/Experimental/ChaosVD/Source/ChaosVD/Private/Trace/DataProcessors/ChaosVDTraceImplicitObjectProcessor.cpp
+++ b/Engine/Plugins/Experimental/ChaosVD/Source/ChaosVD/Private/Trace/DataProcessors/ChaosVDTraceImplicitObjectProcessor.cpp
@@ -4,6 +4,8 @@
 
 #include "Chaos/ImplicitObject.h"
 #include "ChaosVDRecording.h"
+#include "ChaosVisualDebugger/ChaosVDMemWriterReader.h"
+#include "ChaosVisualDebugger/ChaosVDSerializedNameTable.h"
 #include "ChaosVisualDebugger/ChaosVisualDebuggerTrace.h"
 #include "DataWrappers/ChaosVDImplicitObjectDataWrapper.h"
 #include "Serialization/MemoryReader.h"
@@ -23,16 +25,21 @@ bool FChaosVDTraceImplicitObjectProcessor::ProcessRawData(const TArray<uint8>& I
 		return false;
 	}
 
-	FMemoryReader MemeReader(InData);
-	Chaos::FChaosArchive Ar(MemeReader);
+	if (!ensure(ProviderSharedPtr->GetNameTable().IsValid()))
+	{
+		return false;
+	}
 
 	FChaosVDImplicitObjectWrapper WrappedGeometryData;
-	WrappedGeometryData.Serialize(Ar);
+	const bool bSuccess = ReadDataFromBuffer<FChaosVDImplicitObjectWrapper, Chaos::FChaosArchive>(InData, WrappedGeometryData, ProviderSharedPtr->GetNameTable().ToSharedRef());
 
-	if (TSharedPtr<FChaosVDRecording> Recording = ProviderSharedPtr->GetRecordingForSession())
+	if (bSuccess)
 	{
-		Recording->AddImplicitObject(WrappedGeometryData.Hash, WrappedGeometryData.ImplicitObject);
+		if (TSharedPtr<FChaosVDRecording> Recording = ProviderSharedPtr->GetRecordingForSession())
+		{
+			Recording->AddImplicitObject(WrappedGeometryData.Hash, WrappedGeometryData.ImplicitObject);
+		}
 	}
 
-	return !Ar.IsError() && !Ar.IsCriticalError();;
+	return bSuccess;
 }
diff --git a/Engine/Plugins/Experimental/ChaosVD/Source/ChaosVD/Private/Trace/DataProcessors/ChaosVDTraceParticleDataProcessor.cpp b/Engine/Plugins/Experimental/ChaosVD/Source/ChaosVD/Private/Trace/DataProcessors/ChaosVDTraceParticleDataProcessor.cpp
--- a/Engine/Plugins/Experimental/ChaosVD/Source/ChaosVD/Private/Trace/DataProcessors/ChaosVDTraceParticleDataProcessor.cpp
+++ b/Engine/Plugins/Experimental/ChaosVD/Source/ChaosVD/Private/Trace/DataProcessors/ChaosVDTraceParticleDataProcessor.cpp
@@ -3,8 +3,10 @@
 #include "Trace/DataProcessors/ChaosVDTraceParticleDataProcessor.h"
 
 #include "ChaosVDRecording.h"
+#include "ChaosVisualDebugger/ChaosVDMemWriterReader.h"
+#include "ChaosVisualDebugger/ChaosVDSerializedNameTable.h"
+#include "ChaosVisualDebugger/ChaosVisualDebuggerTrace.h"
 #include "DataWrappers/ChaosVDParticleDataWrapper.h"
-#include "Serialization/MemoryReader.h"
 #include "Trace/ChaosVDTraceProvider.h"
 
 FChaosVDTraceParticleDataProcessor::FChaosVDTraceParticleDataProcessor(): IChaosVDDataProcessor(FChaosVDParticleDataWrapper::WrapperTypeName)
@@ -19,22 +21,25 @@ bool FChaosVDTraceParticleDataProcessor::ProcessRawData(const TArray<uint8>& InD
 		return false;
 	}
 
-	FMemoryReader MemeReader(InData);
-	MemeReader.SetUseUnversionedPropertySerialization(true);
-
-	Chaos::FChaosArchive Ar(MemeReader);
+	if (!ensure(ProviderSharedPtr->GetNameTable().IsValid()))
+	{
+		return false;
+	}
 
 	TSharedPtr<FChaosVDParticleDataWrapper> ParticleData = MakeShared<FChaosVDParticleDataWrapper>();
-	ParticleData->Serialize(Ar);
+	const bool bSuccess = ReadDataFromBuffer(InData, *ParticleData, ProviderSharedPtr->GetNameTable().ToSharedRef());
 
-	// This can be null if the recording started Mid-Frame. In this case we just discard the data for now
-	if (FChaosVDSolverFrameData* FrameData = ProviderSharedPtr->GetCurrentSolverFrame(ParticleData->SolverID))
+	if (bSuccess)
 	{
-		if (ensureMsgf(FrameData->SolverSteps.Num() > 0, TEXT("A particle was traced without a valid step scope")))
+		// This can be null if the recording started Mid-Frame. In this case we just discard the data for now
+		if (FChaosVDSolverFrameData* FrameData = ProviderSharedPtr->GetCurrentSolverFrame(ParticleData->SolverID))
 		{
-			FrameData->SolverSteps.Last().RecordedParticlesData.Add(ParticleData);
+			if (ensureMsgf(FrameData->SolverSteps.Num() > 0, TEXT("A particle was traced without a valid step scope")))
+			{
+				FrameData->SolverSteps.Last().RecordedParticlesData.Add(ParticleData);
+			}
 		}
 	}
 
-	return !Ar.IsError() && !Ar.IsCriticalError();
+	return bSuccess;
 }
diff --git a/Engine/Plugins/Experimental/ChaosVD/Source/ChaosVD/Public/Trace/ChaosVDTraceProvider.h b/Engine/Plugins/Experimental/ChaosVD/Source/ChaosVD/Public/Trace/ChaosVDTraceProvider.h
--- a/Engine/Plugins/Experimental/ChaosVD/Source/ChaosVD/Public/Trace/ChaosVDTraceProvider.h
+++ b/Engine/Plugins/Experimental/ChaosVD/Source/ChaosVD/Public/Trace/ChaosVDTraceProvider.h
@@ -13,6 +13,11 @@
 #include "Templates/SharedPointer.h"
 #include "TraceServices/Model/AnalysisSession.h"
 
+namespace Chaos::VisualDebugger
+{
+	class FChaosVDSerializableNameTable;
+}
+
 struct FChaosVDGameFrameData;
 class FChaosVDEngine;
 struct FChaosVDSolverFrameData;
@@ -70,6 +75,8 @@ public:
 
 	void RegisterDataProcessor(TSharedPtr<IChaosVDDataProcessor> InDataProcessor);
 
+	TSharedPtr<Chaos::VisualDebugger::FChaosVDSerializableNameTable> GetNameTable() const;
+
 private:
 
 	void RegisterDefaultDataProcessorsIfNeeded();
diff --git a/Engine/Plugins/Experimental/ChaosVD/Source/ChaosVD/Public/Trace/DataProcessors/ChaosVDSerializedNameEntryDataProcessor.h b/Engine/Plugins/Experimental/ChaosVD/Source/ChaosVD/Public/Trace/DataProcessors/ChaosVDSerializedNameEntryDataProcessor.h
new file mode 100644
--- /dev/null
+++ b/Engine/Plugins/Experimental/ChaosVD/Source/ChaosVD/Public/Trace/DataProcessors/ChaosVDSerializedNameEntryDataProcessor.h
@@ -0,0 +1,16 @@
+﻿// Copyright Epic Games, Inc. All Rights Reserved.
+
+#pragma once
+
+#include "Trace/DataProcessors/IChaosVDDataProcessor.h"
+
+/**
+ * Data processor implementation that is able to deserialize traced Name Entries
+ */
+class FChaosVDSerializedNameEntryDataProcessor final : public IChaosVDDataProcessor
+{
+public:
+	explicit FChaosVDSerializedNameEntryDataProcessor();
+
+	virtual bool ProcessRawData(const TArray<uint8>& InData) override;
+};
diff --git a/Engine/Plugins/Experimental/ChaosVD/Source/ChaosVD/Public/Trace/DataProcessors/IChaosVDDataProcessor.h b/Engine/Plugins/Experimental/ChaosVD/Source/ChaosVD/Public/Trace/DataProcessors/IChaosVDDataProcessor.h
--- a/Engine/Plugins/Experimental/ChaosVD/Source/ChaosVD/Public/Trace/DataProcessors/IChaosVDDataProcessor.h
+++ b/Engine/Plugins/Experimental/ChaosVD/Source/ChaosVD/Public/Trace/DataProcessors/IChaosVDDataProcessor.h
@@ -1,6 +1,7 @@
 ﻿// Copyright Epic Games, Inc. All Rights Reserved.
 
 #pragma once
+#include "ChaosVisualDebugger/ChaosVDMemWriterReader.h"
 #include "Containers/Array.h"
 #include "Containers/StringFwd.h"
 #include "HAL/Platform.h"
@@ -8,6 +9,32 @@
 
 class FChaosVDTraceProvider;
 
+namespace Chaos::VisualDebugger
+{
+	template<typename TDataToSerialize>
+	bool ReadDataFromBuffer(const TArray<uint8>& InDataBuffer, TDataToSerialize& Data, const TSharedRef<FChaosVDSerializableNameTable>& InNameTableInstance)
+	{
+		FChaosVDMemoryReader MemReader(InDataBuffer, InNameTableInstance);
+		MemReader.SetShouldSkipUpdateCustomVersion(true);
+
+		Data.Serialize(MemReader);
+
+		return !MemReader.IsError() && !MemReader.IsCriticalError();
+	}
+
+	template<typename TDataToSerialize, typename TArchive>
+	bool ReadDataFromBuffer(const TArray<uint8>& InDataBuffer, TDataToSerialize& Data, const TSharedRef<FChaosVDSerializableNameTable>& InNameTableInstance)
+	{
+		FChaosVDMemoryReader MemReader(InDataBuffer, InNameTableInstance);
+		TArchive Ar(MemReader);
+		Ar.SetShouldSkipUpdateCustomVersion(true);
+
+		Data.Serialize(Ar);
+
+		return !Ar.IsError() && !Ar.IsCriticalError();
+	}
+}
+
 /** Interface for all used for any class that is able to process traced Chaos Visual Debugger binary data */
 class IChaosVDDataProcessor
 {
diff --git a/Engine/Source/Runtime/Engine/Private/Collision/ChaosVDSQTraceHelper.cpp b/Engine/Source/Runtime/Engine/Private/Collision/ChaosVDSQTraceHelper.cpp
--- a/Engine/Source/Runtime/Engine/Private/Collision/ChaosVDSQTraceHelper.cpp
+++ b/Engine/Source/Runtime/Engine/Private/Collision/ChaosVDSQTraceHelper.cpp
@@ -83,10 +83,10 @@ int32 GetWorldSolverID(const UWorld* World)
 		FChaosVDCollisionQueryParams CVDCollisionParams;
 		CVDCollisionParams.CopyFrom(Params);
 
-		CVDCollisionParams.IgnoredActorsNames.Reserve(CVDCollisionParams.IgnoredActorsIDs.Num());
-		CVDCollisionParams.IgnoredComponentsNames.Reserve(CVDCollisionParams.IgnoredComponentsIDs.Num());
-		Algo::Transform(CVDCollisionParams.IgnoredActorsIDs, CVDCollisionParams.IgnoredActorsNames, [](uint32 ID){ return FName(PhysicsLogUtil::MakeActorNameString(GetObjetPointerFromID<AActor>(ID))); } );
-		Algo::Transform(CVDCollisionParams.IgnoredComponentsIDs, CVDCollisionParams.IgnoredComponentsNames, [](uint32 ID){ return FName(PhysicsLogUtil::MakeActorNameString(GetObjetPointerFromID<UActorComponent>(ID))); } );
+		CVDCollisionParams.IgnoredActorsNames.Reserve(Params.GetIgnoredActors().Num());
+		CVDCollisionParams.IgnoredComponentsNames.Reserve(Params.GetIgnoredComponents().Num());
+		Algo::Transform(Params.GetIgnoredActors(), CVDCollisionParams.IgnoredActorsNames, [](uint32 ID){ return GetFNameSafe(GetObjetPointerFromID<AActor>(ID)); } );
+		Algo::Transform(Params.GetIgnoredComponents(), CVDCollisionParams.IgnoredComponentsNames, [](uint32 ID){ return GetFNameSafe(GetObjetPointerFromID<UActorComponent>(ID)); } );
 
 		FChaosVDCollisionResponseParams CVDCollisionResponseParams;
 		CVDCollisionResponseParams.CopyFrom(ResponseParams);
diff --git a/Engine/Source/Runtime/Experimental/Chaos/Private/ChaosVisualDebugger/ChaosVDDataWrapperUtils.cpp b/Engine/Source/Runtime/Experimental/Chaos/Private/ChaosVisualDebugger/ChaosVDDataWrapperUtils.cpp
--- a/Engine/Source/Runtime/Experimental/Chaos/Private/ChaosVisualDebugger/ChaosVDDataWrapperUtils.cpp
+++ b/Engine/Source/Runtime/Experimental/Chaos/Private/ChaosVisualDebugger/ChaosVDDataWrapperUtils.cpp
@@ -4,12 +4,12 @@
 
 #if WITH_CHAOS_VISUAL_DEBUGGER
 
-#include "Algo/Copy.h"
 #include "Chaos/Collision/ParticlePairMidPhase.h"
+#include "Chaos/Collision/CollisionConstraintAllocator.h"
 #include "Chaos/ParticleHandle.h"
+#include "ChaosVisualDebugger/ChaosVDSerializedNameTable.h"
 #include "DataWrappers/ChaosVDCollisionDataWrappers.h"
 #include "DataWrappers/ChaosVDParticleDataWrapper.h"
-#include "Chaos/Collision/CollisionConstraintAllocator.h"
 
 namespace Chaos::VisualDebugger::Utils
 {
@@ -68,7 +68,7 @@ void FChaosVDDataWrapperUtils::CopyCollisionMaterialToDataWrapper(const Chaos::F
 	OutCopyTo.InvInertiaScale1 = InCopyFrom.InvInertiaScale1;
 }
 
-FChaosVDParticleDataWrapper FChaosVDDataWrapperUtils::BuildParticleDataWrapperFromParticle(const Chaos::FGeometryParticleHandle* ParticleHandlePtr)
+FChaosVDParticleDataWrapper FChaosVDDataWrapperUtils::BuildParticleDataWrapperFromParticle(const Chaos::FGeometryParticleHandle* ParticleHandlePtr, const TSharedRef<Chaos::VisualDebugger::FChaosVDSerializableNameTable>& InNameTableInstance)
 {
 	check(ParticleHandlePtr);
 
@@ -78,7 +78,10 @@ FChaosVDParticleDataWrapper FChaosVDDataWrapperUtils::BuildParticleDataWrapperFr
 	WrappedParticleData.Type =  static_cast<EChaosVDParticleType>(ParticleHandlePtr->Type);
 
 #if CHAOS_DEBUG_NAME
-	WrappedParticleData.DebugNamePtr = ParticleHandlePtr->DebugName();
+	// Passing it as a Ptr because from here until it is serialized right after this function ends this string does not change
+	// Passing it as a sharedptr has an additional 20% cost as it has to increment the reference counter, which adds up
+	// TODO: We should switch to FName to take advantage of the new CVD Serializable name table so they can be de-duplicated, but to do so we need to change how we create our debug names to not be unique strings
+	WrappedParticleData.DebugNamePtr = ParticleHandlePtr->DebugName().IsValid() ? ParticleHandlePtr->DebugName().Get() : nullptr;
 #endif
 
 	WrappedParticleData.ParticlePositionRotation.CopyFrom(*ParticleHandlePtr);
@@ -100,6 +103,8 @@ FChaosVDParticleDataWrapper FChaosVDDataWrapperUtils::BuildParticleDataWrapperFr
 		WrappedParticleData.ParticleCluster.CopyFrom(*ClusteredParticle);
 	}
 
+	WrappedParticleData.MarkAsValid();
+
 	return MoveTemp(WrappedParticleData);
 }
 
diff --git a/Engine/Source/Runtime/Experimental/Chaos/Private/ChaosVisualDebugger/ChaosVDSerializedNameTable.cpp b/Engine/Source/Runtime/Experimental/Chaos/Private/ChaosVisualDebugger/ChaosVDSerializedNameTable.cpp
new file mode 100644
--- /dev/null
+++ b/Engine/Source/Runtime/Experimental/Chaos/Private/ChaosVisualDebugger/ChaosVDSerializedNameTable.cpp
@@ -0,0 +1,84 @@
+﻿// Copyright Epic Games, Inc. All Rights Reserved.
+
+#include "ChaosVisualDebugger/ChaosVDSerializedNameTable.h"
+
+#include "ChaosVisualDebugger/ChaosVDTraceMacros.h"
+#include "ChaosVisualDebugger/ChaosVisualDebuggerTrace.h"
+#include "Misc/ScopeRWLock.h"
+#include "Serialization/MemoryWriter.h"
+
+namespace Chaos::VisualDebugger
+{
+	uint64 FChaosVDSerializableNameTable::AddNameToTable(FName Name)
+	{
+		const uint64 NameID = Name.ToUnstableInt();
+		if (NameID == 0)
+		{
+			// 0 means an empty name, so don't bother reading the name table
+			return NameID;
+		}
+
+		{
+			FReadScopeLock ReadLock(NamesByIDLock);
+			if (NamesByID.Contains(NameID))
+			{
+				return NameID;
+			}
+		}
+		{
+			FWriteScopeLock WriteLock(NamesByIDLock);			
+			NamesByID.Add(NameID, Name);
+		}
+
+		{
+			FChaosVDSerializedNameEntry NameEntry {NameID, Name.GetNumber(), Name.GetPlainNameString()};
+
+			FMemMark StackMarker(FMemStack::Get());
+			TArray<uint8, TInlineAllocator<256, TMemStackAllocator<>>> NameBuffer;
+
+			TMemoryWriterBase MemWriterAr(NameBuffer);
+			MemWriterAr.SetShouldSkipUpdateCustomVersion(true);
+
+			MemWriterAr << NameEntry;
+
+			CVD_TRACE_BINARY_DATA(NameBuffer, FChaosVDSerializedNameEntry::WrapperTypeName);
+		}
+		
+		return NameID;
+	}
+
+	uint64 FChaosVDSerializableNameTable::AddNameToTable(const FChaosVDSerializedNameEntry& InNameEntry)
+	{
+		FWriteScopeLock WriteLock(NamesByIDLock);
+		NamesByID.Add(InNameEntry.NameID, FName(InNameEntry.Name, InNameEntry.NameNumber));
+
+		return InNameEntry.NameID;
+	}
+
+	FName FChaosVDSerializableNameTable::GetNameFromTable(uint64 NameID)
+	{
+		if (NameID == 0)
+		{
+			// 0 means an empty name, so don't bother reading the name table
+			return FName();
+		}
+
+		{
+			FReadScopeLock ReadLock(NamesByIDLock);
+			if (const FName* FoundName = NamesByID.Find(NameID))
+			{
+				return *FoundName;
+			}
+		}
+
+		return FName();
+	}
+
+	void FChaosVDSerializableNameTable::ResetTable()
+	{
+		{
+			FWriteScopeLock WriteLock(NamesByIDLock);
+			NamesByID.Empty();
+		}
+	}
+}
diff --git a/Engine/Source/Runtime/Experimental/Chaos/Private/ChaosVisualDebugger/ChaosVisualDebuggerTrace.cpp b/Engine/Source/Runtime/Experimental/Chaos/Private/ChaosVisualDebugger/ChaosVisualDebuggerTrace.cpp
--- a/Engine/Source/Runtime/Experimental/Chaos/Private/ChaosVisualDebugger/ChaosVisualDebuggerTrace.cpp
+++ b/Engine/Source/Runtime/Experimental/Chaos/Private/ChaosVisualDebugger/ChaosVisualDebuggerTrace.cpp
@@ -10,6 +10,8 @@
 #include "Chaos/PBDCollisionConstraints.h"
 #include "Chaos/PBDRigidsSOAs.h"
 #include "ChaosVisualDebugger/ChaosVDDataWrapperUtils.h"
+#include "ChaosVisualDebugger/ChaosVDMemWriterReader.h"
+#include "ChaosVisualDebugger/ChaosVDSerializedNameTable.h"
 #include "Compression/OodleDataCompressionUtil.h"
 #include "DataWrappers/ChaosVDCollisionDataWrappers.h"
 #include "DataWrappers/ChaosVDImplicitObjectDataWrapper.h"
@@ -17,7 +19,6 @@
 #include "DataWrappers/ChaosVDQueryDataWrappers.h"
 #include "HAL/CriticalSection.h"
 #include "Misc/ScopeRWLock.h"
-#include "Serialization/MemoryWriter.h"
 #include "Templates/UniquePtr.h"
 
 UE_TRACE_EVENT_DEFINE(ChaosVDLogger, ChaosVDSolverFrameStart)
@@ -37,15 +38,20 @@ UE_TRACE_EVENT_DEFINE(ChaosVDLogger, ChaosVDDummyEvent)
 UE_TRACE_EVENT_DEFINE(ChaosVDLogger, ChaosVDNonSolverLocation)
 UE_TRACE_EVENT_DEFINE(ChaosVDLogger, ChaosVDNonSolverTransform)
 
-static FAutoConsoleVariable CVarChaosVDCompressBinaryData(
+namespace Chaos::VisualDebugger::Cvars
+{
+	static bool bCompressBinaryData = false;
+	FAutoConsoleVariableRef CVarChaosVDCompressBinaryData(
 	TEXT("p.Chaos.VD.CompressBinaryData"),
-	false,
+	bCompressBinaryData,
 	TEXT("If true, serialized binary data will be compressed using Oodle on the fly before being traced"));
 
-static FAutoConsoleVariable CVarChaosVDCompressionMode(
+	static int32 CompressionMode = 2;
+	FAutoConsoleVariableRef CVarChaosVDCompressionMode(
 	TEXT("p.Chaos.VD.CompressionMode"),
-	2,
+	CompressionMode,
 	TEXT("Oodle compression mode to use, 4 is by default which equsals to ECompressionLevel::VeryFast"));
+}
 
 /** Struct where we keep track of the geometry we are tracing */
 struct FChaosVDGeometryTraceContext
@@ -103,6 +109,7 @@ FDelegateHandle FChaosVisualDebuggerTrace::RecordingFullCaptureRequestedHandle =
 FRWLock FChaosVisualDebuggerTrace::DeltaRecordingStatesLock = FRWLock();
 TSet<int32> FChaosVisualDebuggerTrace::SolverIDsForDeltaRecording = TSet<int32>();
 TSet<int32> FChaosVisualDebuggerTrace::RequestedFullCaptureSolverIDs = TSet<int32>();
+TSharedRef<FChaosVDSerializableNameTable> FChaosVisualDebuggerTrace::CVDNameTable = MakeShared<FChaosVDSerializableNameTable>();
 std::atomic<bool> FChaosVisualDebuggerTrace::bIsTracing = false;
 
 void FChaosVisualDebuggerTrace::TraceParticle(const Chaos::FGeometryParticleHandle* ParticleHandle)
@@ -140,7 +147,7 @@ void FChaosVisualDebuggerTrace::TraceParticle(Chaos::FGeometryParticleHandle* Pa
 	TraceImplicitObject({ GeometryHash, ParticleHandle->GetGeometry() });
 
 	{
-		FChaosVDParticleDataWrapper ParticleDataWrapper = FChaosVDDataWrapperUtils::BuildParticleDataWrapperFromParticle(ParticleHandle);
+		FChaosVDParticleDataWrapper ParticleDataWrapper = FChaosVDDataWrapperUtils::BuildParticleDataWrapperFromParticle(ParticleHandle, CVDNameTable);
 		ParticleDataWrapper.GeometryHash = GeometryHash;
 		ParticleDataWrapper.SolverID = ContextData.Id;
 		
@@ -156,11 +163,7 @@ void FChaosVisualDebuggerTrace::TraceParticle(Chaos::FGeometryParticleHandle* Pa
 		
 		FChaosVDScopedTLSBufferAccessor TLSDataBuffer;
 
-		FMemoryWriter MemWriterAr(TLSDataBuffer.BufferRef);
-		MemWriterAr.SetShouldSkipUpdateCustomVersion(true);
-		MemWriterAr.SetUseUnversionedPropertySerialization(true);
-
-		ParticleDataWrapper.Serialize(MemWriterAr);
+		Chaos::VisualDebugger::WriteDataToBuffer(TLSDataBuffer.BufferRef, ParticleDataWrapper);
 
 		TraceBinaryData(TLSDataBuffer.BufferRef, FChaosVDParticleDataWrapper::WrapperTypeName);
 	}
@@ -286,15 +289,12 @@ void FChaosVisualDebuggerTrace::TraceMidPhase(const Chaos::FParticlePairMidPhase
 	FChaosVDParticlePairMidPhase CVDMidPhase = FChaosVDDataWrapperUtils::BuildMidPhaseDataWrapperFromMidPhase(*MidPhase);
 	CVDMidPhase.SolverID = CVDContextData->Id;
 
-	FChaosVDScopedTLSBufferAccessor CVDBuffer;
-	FMemoryWriter MemWriter(CVDBuffer.BufferRef);
-		
-	CVDMidPhase.Serialize(MemWriter);
+	FChaosVDScopedTLSBufferAccessor TLSDataBuffer;
+	Chaos::VisualDebugger::WriteDataToBuffer(TLSDataBuffer.BufferRef, CVDMidPhase);
 
-	TraceBinaryData(CVDBuffer.BufferRef, FChaosVDParticlePairMidPhase::WrapperTypeName);
+	TraceBinaryData(TLSDataBuffer.BufferRef, FChaosVDParticlePairMidPhase::WrapperTypeName);
 }
 
-
 void FChaosVisualDebuggerTrace::TraceMidPhasesFromCollisionConstraints(Chaos::FPBDCollisionConstraints& InCollisionConstraints)
 {
 	using namespace Chaos::VisualDebugger::Utils;
@@ -336,12 +336,10 @@ void FChaosVisualDebuggerTrace::TraceCollisionConstraint(const Chaos::FPBDCollis
 	FChaosVDConstraint CVDConstraint = FChaosVDDataWrapperUtils::BuildConstraintDataWrapperFromConstraint(*CollisionConstraint);
 	CVDConstraint.SolverID = CVDContextData->Id;
 
-	FChaosVDScopedTLSBufferAccessor CVDBuffer;
-	FMemoryWriter MemWriter(CVDBuffer.BufferRef);
-		
-	CVDConstraint.Serialize(MemWriter);
+	FChaosVDScopedTLSBufferAccessor TLSDataBuffer;
+	Chaos::VisualDebugger::WriteDataToBuffer(TLSDataBuffer.BufferRef, CVDConstraint);
 
-	TraceBinaryData(CVDBuffer.BufferRef, FChaosVDConstraint::WrapperTypeName);
+	TraceBinaryData(TLSDataBuffer.BufferRef, FChaosVDConstraint::WrapperTypeName);
 }
 
 void FChaosVisualDebuggerTrace::TraceCollisionConstraintView(TArrayView<Chaos::FPBDCollisionConstraint* const> CollisionConstraintView)
@@ -487,7 +485,7 @@ void FChaosVisualDebuggerTrace::TraceSolverSimulationSpace(const Chaos::FRigidTr
 		<< CVD_TRACE_ROTATOR_ON_EVENT(ChaosVDSolverSimulationSpace, Rotation, Transform.GetRotation());
 }
 
-void FChaosVisualDebuggerTrace::TraceBinaryData(const TArray<uint8>& InData, FStringView TypeName)
+void FChaosVisualDebuggerTrace::TraceBinaryData(TConstArrayView<uint8> InData, FStringView TypeName)
 {
 	if (!IsTracing())
 	{
@@ -501,21 +499,21 @@ void FChaosVisualDebuggerTrace::TraceBinaryData(const TArray<uint8>& InData, FSt
 
 	ensure(DataID < TNumericLimits<int32>::Max());
 
-	const TArray<uint8>* DataToTrace = &InData;
+	TConstArrayView<uint8> DataViewToTrace = InData;
 
 	// Handle Compression if enabled
-	const bool bIsCompressed = CVarChaosVDCompressBinaryData->GetBool();
+	const bool bIsCompressed = Chaos::VisualDebugger::Cvars::bCompressBinaryData;
 	TArray<uint8> CompressedData;
 	if (bIsCompressed)
 	{
 		CompressedData.Reserve(CompressedData.Num());
-		FOodleCompressedArray::CompressTArray(CompressedData, InData, FOodleDataCompression::ECompressor::Kraken,
-			static_cast<FOodleDataCompression::ECompressionLevel>(CVarChaosVDCompressionMode->GetInt()));
+		FOodleCompressedArray::CompressData(CompressedData, InData.GetData(),InData.Num(), FOodleDataCompression::ECompressor::Kraken,
+			static_cast<FOodleDataCompression::ECompressionLevel>(Chaos::VisualDebugger::Cvars::CompressionMode));
 
-		DataToTrace = &CompressedData;
+		DataViewToTrace = CompressedData;
 	}
 
-	const uint32 DataSize = static_cast<uint32>(DataToTrace->Num());
+	const uint32 DataSize = static_cast<uint32>(DataViewToTrace.Num());
 	constexpr uint32 MaxChunkSize = TNumericLimits<uint16>::Max();
 	const uint32 ChunkNum = (DataSize + MaxChunkSize - 1) / MaxChunkSize;
 
@@ -531,7 +529,7 @@ void FChaosVisualDebuggerTrace::TraceBinaryData(const TArray<uint8>& InData, FSt
 	for (uint32 Index = 0; Index < ChunkNum; ++Index)
 	{
 		const uint16 Size = static_cast<uint16>(FMath::Min(RemainingSize, MaxChunkSize));
-		const uint8* ChunkData = DataToTrace->GetData() + MaxChunkSize * Index;
+		const uint8* ChunkData = DataViewToTrace.GetData() + MaxChunkSize * Index;
 
 		UE_TRACE_LOG(ChaosVDLogger, ChaosVDBinaryDataContent, ChaosVDChannel)
 			<< ChaosVDBinaryDataContent.Cycle(FPlatformTime::Cycles64())
@@ -570,13 +568,7 @@ void FChaosVisualDebuggerTrace::TraceImplicitObject(FChaosVDImplicitObjectWrappe
 	}
 
 	FChaosVDScopedTLSBufferAccessor TLSDataBuffer;
-
-	FMemoryWriter MemWriterAr(TLSDataBuffer.BufferRef);
-	Chaos::FChaosArchive Ar(MemWriterAr);
-
-	Ar.SetShouldSkipUpdateCustomVersion(true);
-
-	WrappedGeometryData.Serialize(Ar);
+	Chaos::VisualDebugger::WriteDataToBuffer<FChaosVDImplicitObjectWrapper, Chaos::FChaosArchive>(TLSDataBuffer.BufferRef, WrappedGeometryData);
 
 	TraceBinaryData(TLSDataBuffer.BufferRef, FChaosVDImplicitObjectWrapper::WrapperTypeName);
 }
@@ -670,12 +662,7 @@ void FChaosVisualDebuggerTrace::TraceSceneQueryStart(const Chaos::FImplicitObjec
 	WrappedQueryData.Type = QueryType;
 
 	FChaosVDScopedTLSBufferAccessor TLSDataBuffer;
-
-	FMemoryWriter MemWriterAr(TLSDataBuffer.BufferRef);
-	MemWriterAr.SetShouldSkipUpdateCustomVersion(true);
-	MemWriterAr.SetUseUnversionedPropertySerialization(true);
-
-	WrappedQueryData.Serialize(MemWriterAr);
+	Chaos::VisualDebugger::WriteDataToBuffer(TLSDataBuffer.BufferRef, WrappedQueryData);
 
 	TraceBinaryData(TLSDataBuffer.BufferRef, FChaosVDQueryDataWrapper::WrapperTypeName);
 }
@@ -704,12 +691,7 @@ void FChaosVisualDebuggerTrace::TraceSceneQueryVisit(FChaosVDQueryVisitStep&& In
 	InQueryVisitData.OwningQueryID = CVDContextData->Id;
 
 	FChaosVDScopedTLSBufferAccessor TLSDataBuffer;
-
-	FMemoryWriter MemWriterAr(TLSDataBuffer.BufferRef);
-	MemWriterAr.SetShouldSkipUpdateCustomVersion(true);
-	MemWriterAr.SetUseUnversionedPropertySerialization(true);
-
-	InQueryVisitData.Serialize(MemWriterAr);
+	Chaos::VisualDebugger::WriteDataToBuffer(TLSDataBuffer.BufferRef, InQueryVisitData);
 
 	TraceBinaryData(TLSDataBuffer.BufferRef, FChaosVDQueryVisitStep::WrapperTypeName);
 }
@@ -765,6 +747,8 @@ void FChaosVisualDebuggerTrace::UnregisterEventHandlers()
 
 void FChaosVisualDebuggerTrace::Reset()
 {
+	CVDNameTable->ResetTable();
+
 	{
 		FWriteScopeLock WriteLock(DeltaRecordingStatesLock);
 		RequestedFullCaptureSolverIDs.Reset();
diff --git a/Engine/Source/Runtime/Experimental/Chaos/Public/ChaosVisualDebugger/ChaosVDDataWrapperUtils.h b/Engine/Source/Runtime/Experimental/Chaos/Public/ChaosVisualDebugger/ChaosVDDataWrapperUtils.h
--- a/Engine/Source/Runtime/Experimental/Chaos/Public/ChaosVisualDebugger/ChaosVDDataWrapperUtils.h
+++ b/Engine/Source/Runtime/Experimental/Chaos/Public/ChaosVisualDebugger/ChaosVDDataWrapperUtils.h
@@ -10,6 +10,11 @@
 
 namespace Chaos
 {
+	namespace VisualDebugger
+	{
+		class FChaosVDSerializableNameTable;
+	}
+
 	class FParticlePairMidPhase;
 	class FPBDCollisionConstraint;
 	class FPBDCollisionConstraintMaterial;
@@ -42,7 +47,7 @@ private:
 	static void CopyCollisionMaterialToDataWrapper(const Chaos::FPBDCollisionConstraintMaterial& InCopyFrom, FChaosVDCollisionMaterial& OutCopyTo);
 
 	/** Creates and populates a FChaosVDParticleDataWrapper with the data of the provided FGeometryParticleHandle */
-	static FChaosVDParticleDataWrapper BuildParticleDataWrapperFromParticle(const Chaos::FGeometryParticleHandle* ParticleHandlePtr);
+	static FChaosVDParticleDataWrapper BuildParticleDataWrapperFromParticle(const Chaos::FGeometryParticleHandle* ParticleHandlePtr, const TSharedRef<Chaos::VisualDebugger::FChaosVDSerializableNameTable>& InNameTableInstance);
 
 	/** Creates and populates a FChaosVDConstraint with the data of the provided FPBDCollisionConstraint */
 	static FChaosVDConstraint BuildConstraintDataWrapperFromConstraint(const Chaos::FPBDCollisionConstraint& InConstraint);
diff --git a/Engine/Source/Runtime/Experimental/Chaos/Public/ChaosVisualDebugger/ChaosVDMemWriterReader.h b/Engine/Source/Runtime/Experimental/Chaos/Public/ChaosVisualDebugger/ChaosVDMemWriterReader.h
new file mode 100644
--- /dev/null
+++ b/Engine/Source/Runtime/Experimental/Chaos/Public/ChaosVisualDebugger/ChaosVDMemWriterReader.h
@@ -0,0 +1,67 @@
+// Copyright Epic Games, Inc. All Rights Reserved.
+
+#pragma once
+#include "ChaosVDSerializedNameTable.h"
+#include "HAL/Platform.h"
+#include "Serialization/MemoryReader.h"
+#include "Serialization/MemoryWriter.h"
+#include "Templates/SharedPointer.h"
+
+class FName;
+class FArchive;
+
+namespace Chaos::VisualDebugger
+{
+	FORCEINLINE void NameTableFNameSerializer(FArchive& Ar ,FName& Name, const TSharedRef<FChaosVDSerializableNameTable>& InNameTable)
+	{
+		// Serialize the FName as a CVD Name ID
+		if (Ar.IsLoading())
+		{
+			uint64 NameID = 0;
+			Ar << NameID;
+				
+			Name = InNameTable->GetNameFromTable(NameID);
+		}
+		else
+		{
+			uint64 NameID = InNameTable->AddNameToTable(Name);
+			Ar << NameID;
+		}
+	}
+
+	/** Preferred memory writer for CVD recorded structs. It has support for FName serialization and de-duplication */
+	class FChaosVDMemoryWriter final : public FMemoryWriter
+	{
+	public:
+		FChaosVDMemoryWriter(TArray<uint8>& InBytes, const TSharedRef<FChaosVDSerializableNameTable>& InNameTableInstance)
+			: FMemoryWriter(InBytes), NameTableInstance(InNameTableInstance)
+		{
+		}
+
+		virtual FArchive& operator<<(FName& Name) override
+		{
+			NameTableFNameSerializer(*this, Name, NameTableInstance);
+			return *this;
+		}
+
+		TSharedRef<FChaosVDSerializableNameTable> NameTableInstance;
+	};
+
+	/** Preferred memory reader for CVD recorded structs. It has support for FName serialization and de-duplication */
+	class FChaosVDMemoryReader final : public FMemoryReader
+	{
+	public:
+		FChaosVDMemoryReader(const TArray<uint8>& InBytes, const TSharedRef<FChaosVDSerializableNameTable>& InNameTableInstance)
+			: FMemoryReader(InBytes), NameTableInstance(InNameTableInstance)
+		{
+		}
+
+		virtual FArchive& operator<<(FName& Name) override
+		{
+			NameTableFNameSerializer(*this, Name, NameTableInstance);
+			return *this;
+		}
+
+		TSharedRef<FChaosVDSerializableNameTable> NameTableInstance;
+	};
+}
diff --git a/Engine/Source/Runtime/Experimental/Chaos/Public/ChaosVisualDebugger/ChaosVDSerializedNameTable.h b/Engine/Source/Runtime/Experimental/Chaos/Public/ChaosVisualDebugger/ChaosVDSerializedNameTable.h
new file mode 100644
--- /dev/null
+++ b/Engine/Source/Runtime/Experimental/Chaos/Public/ChaosVisualDebugger/ChaosVDSerializedNameTable.h
@@ -0,0 +1,57 @@
+﻿// Copyright Epic Games, Inc. All Rights Reserved.
+
+#pragma once
+
+#include "Containers/Map.h"
+#include "Containers/StringView.h"
+#include "HAL/Platform.h"
+#include "Templates/SharedPointer.h"
+
+namespace Chaos::VisualDebugger
+{
+	class FChaosVDMemoryReader;
+}
+
+namespace Chaos::VisualDebugger
+{
+	/** Serialized Name entry that can be loaded by its recorded ID */
+	struct CHAOS_API FChaosVDSerializedNameEntry
+	{
+		inline static FStringView WrapperTypeName = TEXT("FChaosVDSerializedNameEntry");
+
+		uint64 NameID = 0;
+		int32 NameNumber = NAME_NO_NUMBER;
+		FString Name;
+	};
+
+	/** Simple Name table that keeps track of FNames by their current ID, which can be rebuilt later on with these IDs*/
+	class CHAOS_API FChaosVDSerializableNameTable : public TSharedFromThis<FChaosVDSerializableNameTable>
+	{
+	public:
+		/** Adds a FName to the name table */
+		uint64 AddNameToTable(FName Name);
+
+		/** Adds a Serialized Name entry to the name table. Used by the CVD Data processor to rebuild the name table on load */
+		uint64 AddNameToTable(const FChaosVDSerializedNameEntry& InNameEntry);
+
+		/** Returns the FName associated with the provided ID in this table */
+		FName GetNameFromTable(uint64 NameID);
+
+		/** Clears all the FNames tracked by this table */
+		void ResetTable();
+
+	private:
+			
+		TMap<uint64, FName> NamesByID;
+		FRWLock NamesByIDLock;
+	};
+
+	inline FArchive& operator<<(FArchive& Ar, FChaosVDSerializedNameEntry& NameEntry)
+	{
+		Ar << NameEntry.NameID;
+		Ar << NameEntry.NameNumber;
+		Ar << NameEntry.Name;
+
+		return Ar;
+	}
+}
diff --git a/Engine/Source/Runtime/Experimental/Chaos/Public/ChaosVisualDebugger/ChaosVisualDebuggerTrace.h b/Engine/Source/Runtime/Experimental/Chaos/Public/ChaosVisualDebugger/ChaosVisualDebuggerTrace.h
--- a/Engine/Source/Runtime/Experimental/Chaos/Public/ChaosVisualDebugger/ChaosVisualDebuggerTrace.h
+++ b/Engine/Source/Runtime/Experimental/Chaos/Public/ChaosVisualDebugger/ChaosVisualDebuggerTrace.h
@@ -20,6 +20,7 @@
 
 #if WITH_CHAOS_VISUAL_DEBUGGER
 
+#include "ChaosVDMemWriterReader.h"
 #include "ChaosVDRuntimeModule.h"
 #include "DataWrappers/ChaosVDImplicitObjectDataWrapper.h"
 #include "Trace/Trace.h"
@@ -122,6 +123,11 @@ enum ECollisionChannel : int;
 
 namespace Chaos
 {
+	namespace VisualDebugger
+	{
+		class FChaosVDSerializableNameTable;
+	}
+
 	class FPBDCollisionConstraints;
 	class FPBDRigidsSOAs;
 	class FImplicitObject;
@@ -134,6 +140,7 @@ namespace Chaos
 }
 
 using FChaosVDImplicitObjectWrapper = FChaosVDImplicitObjectDataWrapper<Chaos::FImplicitObjectPtr, Chaos::FChaosArchive>;
+using FChaosVDSerializableNameTable = Chaos::VisualDebugger::FChaosVDSerializableNameTable;
 
 /** Class containing  all the Tracing logic to record data for the Chaos Visual Debugger tool */
 class FChaosVisualDebuggerTrace
@@ -203,7 +210,7 @@ public:
 	 * @param InData Data to trace
 	 * @param TypeName Type name the data represents. It is used during Trace Analysis serialize it back (this is not automatic)
 	 */
-	static CHAOS_API void TraceBinaryData(const TArray<uint8>& InData, FStringView TypeName);
+	static CHAOS_API void TraceBinaryData(TConstArrayView<uint8> InData, FStringView TypeName);
 
 	/**
 	 * Serializes the implicit object contained in the wrapper and trace its it as binary data
@@ -264,6 +271,8 @@ public:
 	/** Unbinds to the static events triggered by the ChaosVD Runtime module */
 	static void UnregisterEventHandlers();
 
+	static TSharedRef<FChaosVDSerializableNameTable>& GetNameTableInstance() { return CVDNameTable; }
+
 private:
 
 	/**
@@ -291,6 +300,7 @@ private:
 
 	static TSet<int32> SolverIDsForDeltaRecording;
 	static TSet<int32> RequestedFullCaptureSolverIDs;
+	static TSharedRef<FChaosVDSerializableNameTable> CVDNameTable;
 
 	static std::atomic<bool> bIsTracing;
 
@@ -380,4 +390,26 @@ struct FChaosVDScopeSceneQueryVisit
 
 	FChaosVDQueryVisitStep& VisitData;
 };
+
+namespace Chaos::VisualDebugger
+{
+	template<typename TDataToSerialize>
+	void WriteDataToBuffer(TArray<uint8>& InOutDataBuffer, TDataToSerialize& Data)
+	{
+		FChaosVDMemoryWriter MemWriterAr(InOutDataBuffer, FChaosVisualDebuggerTrace::GetNameTableInstance());
+		MemWriterAr.SetShouldSkipUpdateCustomVersion(true);
+
+		Data.Serialize(MemWriterAr);
+	}
+
+	template<typename TDataToSerialize, typename TArchive>
+	void WriteDataToBuffer(TArray<uint8>& InOutDataBuffer, TDataToSerialize& Data)
+	{
+		FChaosVDMemoryWriter MemWriterAr(InOutDataBuffer, FChaosVisualDebuggerTrace::GetNameTableInstance());
+		TArchive Ar(MemWriterAr);
+		Ar.SetShouldSkipUpdateCustomVersion(true);
+
+		Data.Serialize(Ar);
+	}
+}
 #endif // WITH_CHAOS_VISUAL_DEBUGGER
diff --git a/Engine/Source/Runtime/Experimental/ChaosVDData/Private/ChaosVDRecording.cpp b/Engine/Source/Runtime/Experimental/ChaosVDData/Private/ChaosVDRecording.cpp
--- a/Engine/Source/Runtime/Experimental/ChaosVDData/Private/ChaosVDRecording.cpp
+++ b/Engine/Source/Runtime/Experimental/ChaosVDData/Private/ChaosVDRecording.cpp
@@ -2,6 +2,12 @@
 
 #include "ChaosVDRecording.h"
 #include "Chaos/ImplicitObject.h"
+#include "ChaosVisualDebugger/ChaosVDSerializedNameTable.h"
+
+FChaosVDRecording::FChaosVDRecording()
+{
+	NameTable = MakeShared<Chaos::VisualDebugger::FChaosVDSerializableNameTable>();
+}
 
 int32 FChaosVDRecording::GetAvailableGameFramesNumber() const
 {
diff --git a/Engine/Source/Runtime/Experimental/ChaosVDData/Public/ChaosVDRecording.h b/Engine/Source/Runtime/Experimental/ChaosVDData/Public/ChaosVDRecording.h
--- a/Engine/Source/Runtime/Experimental/ChaosVDData/Public/ChaosVDRecording.h
+++ b/Engine/Source/Runtime/Experimental/ChaosVDData/Public/ChaosVDRecording.h
@@ -12,6 +12,11 @@
 #include "Chaos/ImplicitObject.h"
 #include "DataWrappers/ChaosVDQueryDataWrappers.h"
 
+namespace Chaos::VisualDebugger
+{
+class FChaosVDSerializableNameTable;
+}
+
 DECLARE_MULTICAST_DELEGATE(FChaosVDRecordingUpdated)
 DECLARE_MULTICAST_DELEGATE_TwoParams(FChaosVDGeometryDataLoaded, const Chaos::FConstImplicitObjectPtr&, const uint32 GeometryID)
 
@@ -96,7 +101,9 @@ struct FChaosVDGameFrameData
  * It is currently populated while analyzing a Trace session
  */
 struct CHAOSVDDATA_API FChaosVDRecording
-{	
+{
+	FChaosVDRecording();
+
 	/** Returns the current available recorded solvers number */
 	int32 GetAvailableSolversNumber_AssumesLocked() const { return RecordedFramesDataPerSolver.Num(); }
 	
@@ -257,6 +264,8 @@ struct CHAOSVDDATA_API FChaosVDRecording
 	/** Sets if this recording is being populated from a live session */
 	void SetIsLive(bool bNewIsLive) { bIsLive = bNewIsLive; }
 
+	TSharedPtr<Chaos::VisualDebugger::FChaosVDSerializableNameTable> GetNameTableInstance() const { return NameTable; }
+
 protected:
 
 	/** Adds an Implicit Object to the recording and takes ownership of it */
@@ -281,6 +290,8 @@ protected:
 	/** Id to Ptr map of all shared geometry data required to visualize */
 	TMap<uint32, Chaos::FConstImplicitObjectPtr> ImplicitObjects;
 
+	TSharedPtr<Chaos::VisualDebugger::FChaosVDSerializableNameTable> NameTable;
+
 	mutable FRWLock RecordingDataLock;
 
 	/** True if this recording is being populated from a live session */
diff --git a/Engine/Source/Runtime/Experimental/ChaosVisualDebugger/Private/DataWrappers/ChaosVDCollisionDataWrappers.cpp b/Engine/Source/Runtime/Experimental/ChaosVisualDebugger/Private/DataWrappers/ChaosVDCollisionDataWrappers.cpp
--- a/Engine/Source/Runtime/Experimental/ChaosVisualDebugger/Private/DataWrappers/ChaosVDCollisionDataWrappers.cpp
+++ b/Engine/Source/Runtime/Experimental/ChaosVisualDebugger/Private/DataWrappers/ChaosVDCollisionDataWrappers.cpp
@@ -1,17 +1,7 @@
 ﻿// Copyright Epic Games, Inc. All Rights Reserved.
 
 #include "DataWrappers/ChaosVDCollisionDataWrappers.h"
-
-#ifndef CVD_SERIALIZE_STATIC_ARRAY
-	#define CVD_SERIALIZE_STATIC_ARRAY(Archive, Array) \
-		{ \
-			constexpr int32 Size = UE_ARRAY_COUNT(Array) ; \
-			for (int32 Index = 0; Index < Size; Index++)\
-			{\
-				Archive << Array[Index]; \
-			}\
-		}
-#endif
+#include "DataWrappers/ChaosVDDataSerializationMacros.h"
 
 bool FChaosVDContactPoint::Serialize(FArchive& Ar)
 {
@@ -26,13 +16,31 @@ bool FChaosVDContactPoint::Serialize(FArchive& Ar)
 
 bool FChaosVDManifoldPoint::Serialize(FArchive& Ar)
 {
-	FArchive_Serialize_BitfieldBool(Ar, bDisabled);
-	FArchive_Serialize_BitfieldBool(Ar, bWasRestored);
-	FArchive_Serialize_BitfieldBool(Ar, bWasReplaced);
+	EChaosVDManifoldPointFlags PackedFlags = EChaosVDManifoldPointFlags::None;
 
-	FArchive_Serialize_BitfieldBool(Ar, bHasStaticFrictionAnchor);
-	FArchive_Serialize_BitfieldBool(Ar, bIsValid);
-	FArchive_Serialize_BitfieldBool(Ar, bInsideStaticFrictionCone);
+	
+	// Note: When the UI is done to show the enums in the same way we show these bitfiedls as read only booleans in CVD, we can remove all
+	// this macro boilerplate and just serialize the the flags directly
+	if (Ar.IsLoading())
+	{
+		Ar << PackedFlags;
+		CVD_UNPACK_BITFIELD_DATA(bDisabled, PackedFlags, EChaosVDManifoldPointFlags::Disabled);
+		CVD_UNPACK_BITFIELD_DATA(bWasRestored, PackedFlags, EChaosVDManifoldPointFlags::WasRestored);
+		CVD_UNPACK_BITFIELD_DATA(bWasReplaced, PackedFlags, EChaosVDManifoldPointFlags::WasReplaced);
+		CVD_UNPACK_BITFIELD_DATA(bHasStaticFrictionAnchor, PackedFlags, EChaosVDManifoldPointFlags::HasStaticFrictionAnchor);
+		CVD_UNPACK_BITFIELD_DATA(bIsValid, PackedFlags, EChaosVDManifoldPointFlags::IsValid);
+		CVD_UNPACK_BITFIELD_DATA(bInsideStaticFrictionCone, PackedFlags, EChaosVDManifoldPointFlags::InsideStaticFrictionCone);
+	}
+	else
+	{
+		CVD_PACK_BITFIELD_DATA(bDisabled, PackedFlags, EChaosVDManifoldPointFlags::Disabled);
+		CVD_PACK_BITFIELD_DATA(bWasRestored, PackedFlags, EChaosVDManifoldPointFlags::WasRestored);
+		CVD_PACK_BITFIELD_DATA(bWasReplaced, PackedFlags, EChaosVDManifoldPointFlags::WasReplaced);
+		CVD_PACK_BITFIELD_DATA(bHasStaticFrictionAnchor, PackedFlags, EChaosVDManifoldPointFlags::HasStaticFrictionAnchor);
+		CVD_PACK_BITFIELD_DATA(bIsValid, PackedFlags, EChaosVDManifoldPointFlags::IsValid);
+		CVD_PACK_BITFIELD_DATA(bInsideStaticFrictionCone, PackedFlags, EChaosVDManifoldPointFlags::InsideStaticFrictionCone);
+		Ar << PackedFlags;
+	}
 
 	Ar << NetPushOut;
 	Ar << NetImpulse;
@@ -68,20 +76,44 @@ bool FChaosVDCollisionMaterial::Serialize(FArchive& Ar)
 }
 
 bool FChaosVDConstraint::Serialize(FArchive& Ar)
-{
-	FArchive_Serialize_BitfieldBool(Ar, bIsCurrent);
-	FArchive_Serialize_BitfieldBool(Ar, bDisabled);
-	FArchive_Serialize_BitfieldBool(Ar, bUseManifold);
-	FArchive_Serialize_BitfieldBool(Ar, bUseIncrementalManifold);
-	FArchive_Serialize_BitfieldBool(Ar, bCanRestoreManifold);
-	FArchive_Serialize_BitfieldBool(Ar, bWasManifoldRestored);
-	FArchive_Serialize_BitfieldBool(Ar, bIsQuadratic0);
-	FArchive_Serialize_BitfieldBool(Ar, bIsQuadratic1);
-	FArchive_Serialize_BitfieldBool(Ar, bIsProbe);
-	FArchive_Serialize_BitfieldBool(Ar, bCCDEnabled);
-	FArchive_Serialize_BitfieldBool(Ar, bCCDSweepEnabled);
-	FArchive_Serialize_BitfieldBool(Ar, bModifierApplied);
-	FArchive_Serialize_BitfieldBool(Ar, bMaterialSet);
+{	
+	EChaosVDConstraintFlags PackedFlags = EChaosVDConstraintFlags::None;
+
+	
+	// Note: When the UI is done to show the enums in the same way we show these bitfiedls as read only booleans in CVD, we can remove all
+	// this macro boilerplate and just serialize the the flags directly
+	if (Ar.IsLoading())
+	{
+		Ar << PackedFlags;
+		CVD_UNPACK_BITFIELD_DATA(bIsCurrent, PackedFlags, EChaosVDConstraintFlags::IsCurrent);
+		CVD_UNPACK_BITFIELD_DATA(bDisabled, PackedFlags, EChaosVDConstraintFlags::Disabled);
+		CVD_UNPACK_BITFIELD_DATA(bUseManifold, PackedFlags, EChaosVDConstraintFlags::UseManifold);
+		CVD_UNPACK_BITFIELD_DATA(bUseIncrementalManifold, PackedFlags, EChaosVDConstraintFlags::UseIncrementalManifold);
+		CVD_UNPACK_BITFIELD_DATA(bCanRestoreManifold, PackedFlags, EChaosVDConstraintFlags::CanRestoreManifold);
+		CVD_UNPACK_BITFIELD_DATA(bWasManifoldRestored, PackedFlags, EChaosVDConstraintFlags::WasManifoldRestored);
+		CVD_UNPACK_BITFIELD_DATA(bIsQuadratic0, PackedFlags, EChaosVDConstraintFlags::IsQuadratic0);
+		CVD_UNPACK_BITFIELD_DATA(bIsProbe, PackedFlags, EChaosVDConstraintFlags::IsProbe);
+		CVD_UNPACK_BITFIELD_DATA(bCCDEnabled, PackedFlags, EChaosVDConstraintFlags::CCDEnabled);
+		CVD_UNPACK_BITFIELD_DATA(bCCDSweepEnabled, PackedFlags, EChaosVDConstraintFlags::CCDSweepEnabled);
+		CVD_UNPACK_BITFIELD_DATA(bModifierApplied, PackedFlags, EChaosVDConstraintFlags::ModifierApplied);
+		CVD_UNPACK_BITFIELD_DATA(bMaterialSet, PackedFlags, EChaosVDConstraintFlags::MaterialSet);
+	}
+	else
+	{
+		CVD_PACK_BITFIELD_DATA(bIsCurrent, PackedFlags, EChaosVDConstraintFlags::IsCurrent);
+		CVD_PACK_BITFIELD_DATA(bDisabled, PackedFlags, EChaosVDConstraintFlags::Disabled);
+		CVD_PACK_BITFIELD_DATA(bUseManifold, PackedFlags, EChaosVDConstraintFlags::UseManifold);
+		CVD_PACK_BITFIELD_DATA(bUseIncrementalManifold, PackedFlags, EChaosVDConstraintFlags::UseIncrementalManifold);
+		CVD_PACK_BITFIELD_DATA(bCanRestoreManifold, PackedFlags, EChaosVDConstraintFlags::CanRestoreManifold);
+		CVD_PACK_BITFIELD_DATA(bWasManifoldRestored, PackedFlags, EChaosVDConstraintFlags::WasManifoldRestored);
+		CVD_PACK_BITFIELD_DATA(bIsQuadratic0, PackedFlags, EChaosVDConstraintFlags::IsQuadratic0);
+		CVD_PACK_BITFIELD_DATA(bIsProbe, PackedFlags, EChaosVDConstraintFlags::IsProbe);
+		CVD_PACK_BITFIELD_DATA(bCCDEnabled, PackedFlags, EChaosVDConstraintFlags::CCDEnabled);
+		CVD_PACK_BITFIELD_DATA(bCCDSweepEnabled, PackedFlags, EChaosVDConstraintFlags::CCDSweepEnabled);
+		CVD_PACK_BITFIELD_DATA(bModifierApplied, PackedFlags, EChaosVDConstraintFlags::ModifierApplied);
+		CVD_PACK_BITFIELD_DATA(bMaterialSet, PackedFlags, EChaosVDConstraintFlags::MaterialSet);
+		Ar << PackedFlags;
+	}
 
 	Ar << Material;
 	Ar << AccumulatedImpulse;
@@ -111,16 +143,36 @@ bool FChaosVDConstraint::Serialize(FArchive& Ar)
 	return true;
 }
 
-
 bool FChaosVDParticlePairMidPhase::Serialize(FArchive& Ar)
 {
 	Ar << SolverID;
+
+	
+	EChaosVDMidPhaseFlags PackedFlags = EChaosVDMidPhaseFlags::None;
+
 	
-	FArchive_Serialize_BitfieldBool(Ar, bIsActive);
-	FArchive_Serialize_BitfieldBool(Ar, bIsCCD);
-	FArchive_Serialize_BitfieldBool(Ar, bIsCCDActive);
-	FArchive_Serialize_BitfieldBool(Ar, bIsSleeping);
-	FArchive_Serialize_BitfieldBool(Ar, bIsModified);
+	// Note: When the UI is done to show the enums in the same way we show these bitfiedls as read only booleans in CVD, we can remove all
+	// this macro boilerplate and just serialize the the flags directly
+	if (Ar.IsLoading())
+	{
+		Ar << PackedFlags;
+		CVD_UNPACK_BITFIELD_DATA(bIsActive, PackedFlags, EChaosVDMidPhaseFlags::IsActive);
+		CVD_UNPACK_BITFIELD_DATA(bIsCCD, PackedFlags, EChaosVDMidPhaseFlags::IsCCD);
+		CVD_UNPACK_BITFIELD_DATA(bIsCCDActive, PackedFlags, EChaosVDMidPhaseFlags::IsCCDActive);
+		CVD_UNPACK_BITFIELD_DATA(bIsCCDActive, PackedFlags, EChaosVDMidPhaseFlags::IsCCDActive);
+		CVD_UNPACK_BITFIELD_DATA(bIsSleeping, PackedFlags, EChaosVDMidPhaseFlags::IsSleeping);
+		CVD_UNPACK_BITFIELD_DATA(bIsModified, PackedFlags, EChaosVDMidPhaseFlags::IsModified);
+	}
+	else
+	{
+		CVD_PACK_BITFIELD_DATA(bIsActive, PackedFlags, EChaosVDMidPhaseFlags::IsActive);
+		CVD_PACK_BITFIELD_DATA(bIsCCD, PackedFlags, EChaosVDMidPhaseFlags::IsCCD);
+		CVD_PACK_BITFIELD_DATA(bIsCCDActive, PackedFlags, EChaosVDMidPhaseFlags::IsCCDActive);
+		CVD_PACK_BITFIELD_DATA(bIsCCDActive, PackedFlags, EChaosVDMidPhaseFlags::IsCCDActive);
+		CVD_PACK_BITFIELD_DATA(bIsSleeping, PackedFlags, EChaosVDMidPhaseFlags::IsSleeping);
+		CVD_PACK_BITFIELD_DATA(bIsModified, PackedFlags, EChaosVDMidPhaseFlags::IsModified);
+		Ar << PackedFlags;
+	}
 		
 	Ar << LastUsedEpoch;
 
@@ -145,10 +197,26 @@ bool FChaosVDCollisionFilterData::Serialize(FArchive& Ar)
 bool FChaosVDShapeCollisionData::Serialize(FArchive& Ar)
 {
 	Ar << CollisionTraceType;
+
+	EChaosVDCollisionShapeDataFlags PackedFlags = EChaosVDCollisionShapeDataFlags::None;
+
 	
-	FArchive_Serialize_BitfieldBool(Ar, bSimCollision);
-	FArchive_Serialize_BitfieldBool(Ar, bQueryCollision);
-	FArchive_Serialize_BitfieldBool(Ar, bIsProbe);
+	// Note: When the UI is done to show the enums in the same way we show these bitfiedls as read only booleans in CVD, we can remove all
+	// this macro boilerplate and just serialize the the flags directly
+	if (Ar.IsLoading())
+	{
+		Ar << PackedFlags;
+		CVD_UNPACK_BITFIELD_DATA(bSimCollision, PackedFlags, EChaosVDCollisionShapeDataFlags::SimCollision);
+		CVD_UNPACK_BITFIELD_DATA(bQueryCollision, PackedFlags, EChaosVDCollisionShapeDataFlags::QueryCollision);
+		CVD_UNPACK_BITFIELD_DATA(bIsProbe, PackedFlags, EChaosVDCollisionShapeDataFlags::IsProbe);
+	}
+	else
+	{
+		CVD_PACK_BITFIELD_DATA(bSimCollision, PackedFlags, EChaosVDCollisionShapeDataFlags::SimCollision);
+		CVD_PACK_BITFIELD_DATA(bQueryCollision, PackedFlags, EChaosVDCollisionShapeDataFlags::QueryCollision);
+		CVD_PACK_BITFIELD_DATA(bIsProbe, PackedFlags, EChaosVDCollisionShapeDataFlags::IsProbe);
+		Ar << PackedFlags;
+	}
 
 	return true;
 }
diff --git a/Engine/Source/Runtime/Experimental/ChaosVisualDebugger/Private/DataWrappers/ChaosVDParticleDataWrapper.cpp b/Engine/Source/Runtime/Experimental/ChaosVisualDebugger/Private/DataWrappers/ChaosVDParticleDataWrapper.cpp
--- a/Engine/Source/Runtime/Experimental/ChaosVisualDebugger/Private/DataWrappers/ChaosVDParticleDataWrapper.cpp
+++ b/Engine/Source/Runtime/Experimental/ChaosVisualDebugger/Private/DataWrappers/ChaosVDParticleDataWrapper.cpp
@@ -14,53 +14,79 @@ bool FChaosVDFRigidParticleControlFlags::Serialize(FArchive& Ar)
 	Ar << bInertiaConditioningEnabled;
 	Ar << GravityGroupIndex;
 
-	return true;
+	return !Ar.IsError();
 }
 
 bool FChaosVDParticlePositionRotation::Serialize(FArchive& Ar)
 {
+	Ar << bHasValidData;
+	if (!bHasValidData)
+	{
+		return !Ar.IsError();
+	}
+
 	Ar << MX;
 	Ar << MR;
-	Ar << bHasValidData;
 
-	return true;
+	return !Ar.IsError();
 }
 
 bool FChaosVDParticleVelocities::Serialize(FArchive& Ar)
 {
+	Ar << bHasValidData;
+	if (!bHasValidData)
+	{
+		return !Ar.IsError();
+	}
+
 	Ar << MV;
 	Ar << MW;
-	Ar << bHasValidData;
 
-	return true;
+	return !Ar.IsError();
 }
 
 bool FChaosVDParticleDynamics::Serialize(FArchive& Ar)
 {
+	Ar << bHasValidData;
+	if (!bHasValidData)
+	{
+		return !Ar.IsError();
+	}
+
 	Ar << MAcceleration;
 	Ar << MAngularAcceleration;
 	Ar << MAngularImpulseVelocity;
 	Ar << MLinearImpulseVelocity;
-	Ar << bHasValidData;
 
-	return true;
+	return !Ar.IsError();
 }
 
 bool FChaosVDParticleMassProps::Serialize(FArchive& Ar)
 {
+	Ar << bHasValidData;
+	if (!bHasValidData)
+	{
+		return !Ar.IsError();
+	}
+
 	Ar << MCenterOfMass;
 	Ar << MRotationOfMass;
 	Ar << MI;
 	Ar << MInvI;
 	Ar << MM;
 	Ar << MInvM;
-	Ar << bHasValidData;
 
-	return true;
+	return !Ar.IsError();
 }
 
 bool FChaosVDParticleDynamicMisc::Serialize(FArchive& Ar)
 {
+	Ar << bHasValidData;
+	if (!bHasValidData)
+	{
+		return !Ar.IsError();
+	}
+
 	Ar << MAngularEtherDrag;
 	Ar << MMaxLinearSpeedSq;
 	Ar << MMaxAngularSpeedSq;
@@ -69,7 +95,6 @@ bool FChaosVDParticleDynamicMisc::Serialize(FArchive& Ar)
 	Ar << MObjectState;
 	Ar << MSleepType;
 	Ar << bDisabled;
-	Ar << bHasValidData;
 
 	MControlFlags.Serialize(Ar);
 
@@ -78,6 +103,12 @@ bool FChaosVDParticleDynamicMisc::Serialize(FArchive& Ar)
 
 bool FChaosVDParticleCluster::Serialize(FArchive& Ar)
 {
+	Ar << bHasValidData;
+	if (!bHasValidData)
+	{
+		return !Ar.IsError();
+	}
+
 	Ar << ParentParticleID;	
 	Ar << NumChildren;
 	Ar << ChildToParent;
@@ -91,17 +122,22 @@ bool FChaosVDParticleCluster::Serialize(FArchive& Ar)
 	Ar << bIsAnchored;
 	Ar << bUnbreakable;
 	Ar << bIsChildToParentLocked;
-	
-	Ar << bHasValidData;
-	return true;
+
+	return !Ar.IsError();
 }
 
 bool FChaosVDParticleDataWrapper::Serialize(FArchive& Ar)
 {
+	Ar << bHasValidData;
+	if (!bHasValidData)
+	{
+		return !Ar.IsError();
+	}
+
 	Ar << Type;
 	Ar << GeometryHash;
 
-	bHasDebugName = DebugNamePtr.IsValid();
+	bHasDebugName = DebugNamePtr != nullptr;
 	Ar << bHasDebugName;
 
 	if (bHasDebugName)
@@ -112,7 +148,10 @@ bool FChaosVDParticleDataWrapper::Serialize(FArchive& Ar)
 		}
 		else
 		{
-			FString& DebugNameRef = *DebugNamePtr.Get();
+			// Keep the static analyser happy, we already checked this above
+			checkSlow(DebugNamePtr != nullptr);
+
+			FString& DebugNameRef = *DebugNamePtr;
 			Ar << DebugNameRef;
 		}
 	}
@@ -126,11 +165,9 @@ bool FChaosVDParticleDataWrapper::Serialize(FArchive& Ar)
 	Ar << ParticleDynamicsMisc;
 	Ar << ParticleMassProps;
 
-	Ar << bHasValidData;
-
 	Ar << CollisionDataPerShape;
 
 	Ar << ParticleCluster;
 
-	return true;
+	return !Ar.IsError();
 }
diff --git a/Engine/Source/Runtime/Experimental/ChaosVisualDebugger/Private/DataWrappers/ChaosVDQueryDataWrappers.cpp b/Engine/Source/Runtime/Experimental/ChaosVisualDebugger/Private/DataWrappers/ChaosVDQueryDataWrappers.cpp
--- a/Engine/Source/Runtime/Experimental/ChaosVisualDebugger/Private/DataWrappers/ChaosVDQueryDataWrappers.cpp
+++ b/Engine/Source/Runtime/Experimental/ChaosVisualDebugger/Private/DataWrappers/ChaosVDQueryDataWrappers.cpp
@@ -2,67 +2,130 @@
 
 #include "DataWrappers/ChaosVDQueryDataWrappers.h"
 
+#include "DataWrappers/ChaosVDDataSerializationMacros.h"
+
 bool FChaosVDCollisionResponseParams::Serialize(FArchive& Ar)
 {
-	Ar << FlagsPerChannel;
 	Ar << bHasValidData;
+	if (!bHasValidData)
+	{
+		return !Ar.IsError();
+	}
+
+	Ar << FlagsPerChannel;
 
 	return !Ar.IsError();
 }
 
 bool FChaosVDCollisionObjectQueryParams::Serialize(FArchive& Ar)
 {
+	Ar << bHasValidData;
+	if (!bHasValidData)
+	{
+		return !Ar.IsError();
+	}
+
 	Ar << ObjectTypesToQuery;
 	Ar << IgnoreMask;
 
-	Ar << bHasValidData;
-
 	return !Ar.IsError();
 }
 
 bool FChaosVDCollisionQueryParams::Serialize(FArchive& Ar)
 {
+	Ar << bHasValidData;
+	if (!bHasValidData)
+	{
+		return !Ar.IsError();
+	}
+
 	Ar << TraceTag;
 	Ar << OwnerTag;
-	Ar << bTraceComplex;
-	Ar << bFindInitialOverlaps;
-	Ar << bReturnFaceIndex;
-	Ar << bReturnPhysicalMaterial;
-	Ar << bIgnoreBlocks;
-	Ar << bIgnoreTouches;
-	Ar << bSkipNarrowPhase;
-	Ar << bTraceIntoSubComponents;
-	Ar << bReplaceHitWithSubComponents;
-	Ar << IgnoreMask;
-	Ar << IgnoredComponentsIDs;
 	Ar << IgnoredActorsIDs;
 	Ar << IgnoredComponentsNames;
 	Ar << IgnoredActorsNames;
 
-	Ar << bHasValidData;
-	
+	EChaosVDCollisionQueryParamsFlags PackedQueryFlags = EChaosVDCollisionQueryParamsFlags::None;
+
+	// Note: When the UI is done to show the enums in the same way we show these bitfiedls as read only booleans in CVD, we can remove all
+	// this macro boilerplate and just serialize the the flags directly
+	if (Ar.IsLoading())
+	{
+		Ar << PackedQueryFlags;
+		CVD_UNPACK_BITFIELD_DATA(bTraceComplex, PackedQueryFlags, EChaosVDCollisionQueryParamsFlags::TraceComplex);
+		CVD_UNPACK_BITFIELD_DATA(bFindInitialOverlaps, PackedQueryFlags, EChaosVDCollisionQueryParamsFlags::FindInitialOverlaps);
+		CVD_UNPACK_BITFIELD_DATA(bReturnFaceIndex, PackedQueryFlags, EChaosVDCollisionQueryParamsFlags::ReturnFaceIndex);
+		CVD_UNPACK_BITFIELD_DATA(bReturnPhysicalMaterial, PackedQueryFlags, EChaosVDCollisionQueryParamsFlags::ReturnPhysicalMaterial);
+		CVD_UNPACK_BITFIELD_DATA(bIgnoreBlocks, PackedQueryFlags, EChaosVDCollisionQueryParamsFlags::IgnoreBlocks);
+		CVD_UNPACK_BITFIELD_DATA(bIgnoreTouches, PackedQueryFlags, EChaosVDCollisionQueryParamsFlags::IgnoreTouches);
+		CVD_UNPACK_BITFIELD_DATA(bSkipNarrowPhase, PackedQueryFlags, EChaosVDCollisionQueryParamsFlags::SkipNarrowPhase);
+		CVD_UNPACK_BITFIELD_DATA(bTraceIntoSubComponents, PackedQueryFlags, EChaosVDCollisionQueryParamsFlags::TraceIntoSubComponents);
+		CVD_UNPACK_BITFIELD_DATA(bReplaceHitWithSubComponents, PackedQueryFlags, EChaosVDCollisionQueryParamsFlags::ReplaceHitWithSubComponents);
+	}
+	else
+	{
+		CVD_PACK_BITFIELD_DATA(bTraceComplex, PackedQueryFlags, EChaosVDCollisionQueryParamsFlags::TraceComplex);
+		CVD_PACK_BITFIELD_DATA(bFindInitialOverlaps, PackedQueryFlags, EChaosVDCollisionQueryParamsFlags::FindInitialOverlaps);
+		CVD_PACK_BITFIELD_DATA(bReturnFaceIndex, PackedQueryFlags, EChaosVDCollisionQueryParamsFlags::ReturnFaceIndex);
+		CVD_PACK_BITFIELD_DATA(bReturnPhysicalMaterial, PackedQueryFlags, EChaosVDCollisionQueryParamsFlags::ReturnPhysicalMaterial);
+		CVD_PACK_BITFIELD_DATA(bIgnoreBlocks, PackedQueryFlags, EChaosVDCollisionQueryParamsFlags::IgnoreBlocks);
+		CVD_PACK_BITFIELD_DATA(bIgnoreTouches, PackedQueryFlags, EChaosVDCollisionQueryParamsFlags::IgnoreTouches);
+		CVD_PACK_BITFIELD_DATA(bSkipNarrowPhase, PackedQueryFlags, EChaosVDCollisionQueryParamsFlags::SkipNarrowPhase);
+		CVD_PACK_BITFIELD_DATA(bTraceIntoSubComponents, PackedQueryFlags, EChaosVDCollisionQueryParamsFlags::TraceIntoSubComponents);
+		CVD_PACK_BITFIELD_DATA(bReplaceHitWithSubComponents, PackedQueryFlags, EChaosVDCollisionQueryParamsFlags::ReplaceHitWithSubComponents);
+		Ar << PackedQueryFlags;
+	}
+
 	return !Ar.IsError();
 }
 
 bool FChaosVDQueryFastData::Serialize(FArchive& Ar)
 {
+	Ar << bHasValidData;
+	if (!bHasValidData)
+	{
+		return !Ar.IsError();
+	}
+
 	Ar << Dir;
 	Ar << InvDir;
 	Ar << CurrentLength;
 	Ar << InvDir;
 
-	//TODO : Pack these flags
-	Ar << bParallel0;
-	Ar << bParallel1;
-	Ar << bParallel2;
+	EChaosVDQueryFastDataParallelFlags PackedFlags = EChaosVDQueryFastDataParallelFlags::None;
 
-	Ar << bHasValidData;
+	
+	// Note: When the UI is done to show the enums in the same way we show these bitfiedls as read only booleans in CVD, we can remove all
+	// this macro boilerplate and just serialize the the flags directly
+	if (Ar.IsLoading())
+	{
+		Ar << PackedFlags;
+
+		CVD_UNPACK_BITFIELD_DATA(bParallel0, PackedFlags, EChaosVDQueryFastDataParallelFlags::Parallel0);
+		CVD_UNPACK_BITFIELD_DATA(bParallel1, PackedFlags, EChaosVDQueryFastDataParallelFlags::Parallel1);
+		CVD_UNPACK_BITFIELD_DATA(bParallel2, PackedFlags, EChaosVDQueryFastDataParallelFlags::Parallel2);
+	}
+	else
+	{
+		CVD_PACK_BITFIELD_DATA(bParallel0, PackedFlags, EChaosVDQueryFastDataParallelFlags::Parallel0);
+		CVD_PACK_BITFIELD_DATA(bParallel1, PackedFlags, EChaosVDQueryFastDataParallelFlags::Parallel1);
+		CVD_PACK_BITFIELD_DATA(bParallel2, PackedFlags, EChaosVDQueryFastDataParallelFlags::Parallel2);
+
+		Ar << PackedFlags;
+	}
 
 	return !Ar.IsError();
 }
 
 bool FChaosVDQueryHitData::Serialize(FArchive& Ar)
 {
+	Ar << bHasValidData;
+
+	if (!bHasValidData)
+	{
+		return !Ar.IsError();
+	}
+
 	Ar << Distance;
 	Ar << FaceIdx;
 	Ar << Flags;
@@ -76,6 +139,12 @@ bool FChaosVDQueryHitData::Serialize(FArchive& Ar)
 
 bool FChaosVDQueryVisitStep::Serialize(FArchive& Ar)
 {
+	Ar << bHasValidData;
+	if (!bHasValidData)
+	{
+		return !Ar.IsError();
+	}
+
 	Ar << OwningQueryID;
 	Ar << Type;
 	Ar << ShapeIndex;
@@ -86,8 +155,6 @@ bool FChaosVDQueryVisitStep::Serialize(FArchive& Ar)
 
 	Ar << QueryFastData;
 
-	Ar << bHasValidData;
-
 	return !Ar.IsError();
 }
 
diff --git a/Engine/Source/Runtime/Experimental/ChaosVisualDebugger/Public/DataWrappers/ChaosVDCollisionDataWrappers.h b/Engine/Source/Runtime/Experimental/ChaosVisualDebugger/Public/DataWrappers/ChaosVDCollisionDataWrappers.h
--- a/Engine/Source/Runtime/Experimental/ChaosVisualDebugger/Public/DataWrappers/ChaosVDCollisionDataWrappers.h
+++ b/Engine/Source/Runtime/Experimental/ChaosVisualDebugger/Public/DataWrappers/ChaosVDCollisionDataWrappers.h
@@ -94,6 +94,19 @@ inline FArchive& operator<<(FArchive& Ar, FChaosVDContactPoint& Data)
 	return Ar;
 }
 
+UENUM()
+enum class EChaosVDManifoldPointFlags : uint8
+{
+	None = 0,
+	Disabled = 1 << 0,
+	WasRestored = 1 << 1,
+	WasReplaced = 1 << 2,
+	HasStaticFrictionAnchor = 1 << 3,
+	IsValid = 1 << 4,
+	InsideStaticFrictionCone = 1 << 5,
+};
+ENUM_CLASS_FLAGS(EChaosVDManifoldPointFlags)
+
 USTRUCT()
 struct CHAOSVDRUNTIME_API FChaosVDManifoldPoint
 {
@@ -211,6 +224,24 @@ inline FArchive& operator<<(FArchive& Ar, FChaosVDCollisionMaterial& Data)
 	return Ar;
 }
 
+UENUM()
+enum class EChaosVDConstraintFlags : uint16
+{
+	None = 0,
+	IsCurrent = 1 << 0,
+	Disabled = 1 << 1,
+	UseManifold = 1 << 2,
+	UseIncrementalManifold = 1 << 3,
+	CanRestoreManifold = 1 << 4,
+	WasManifoldRestored = 1 << 5,
+	IsQuadratic0 = 1 << 6,
+	IsQuadratic1 = 1 << 7,
+	IsProbe = 1 << 8,
+	CCDEnabled = 1 << 9,
+	CCDSweepEnabled = 1 << 10,
+	ModifierApplied = 1 << 11,
+	MaterialSet = 1 << 12,
+};
 
 USTRUCT()
 struct CHAOSVDRUNTIME_API FChaosVDConstraint
@@ -329,6 +360,17 @@ inline FArchive& operator<<(FArchive& Ar, FChaosVDConstraint& Data)
 	return Ar;
 }
 
+UENUM()
+enum class EChaosVDMidPhaseFlags : uint8
+{
+	None = 0,
+	IsActive = 1 << 0,
+	IsCCD = 1 << 1,
+	IsCCDActive = 1 << 2,
+	IsSleeping = 1 << 3,
+	IsModified = 1 << 4,
+};
+
 USTRUCT()
 struct CHAOSVDRUNTIME_API FChaosVDParticlePairMidPhase
 {
@@ -426,6 +468,14 @@ inline FArchive& operator<<(FArchive& Ar, FChaosVDCollisionFilterData& Data)
 	Data.Serialize(Ar);
 	return Ar;
 }
+UENUM()
+enum class EChaosVDCollisionShapeDataFlags : uint8
+{
+	None = 0,
+	SimCollision = 1 << 0,
+	QueryCollision = 1 << 1,
+	IsProbe = 1 << 2,
+};
 
 USTRUCT()
 struct CHAOSVDRUNTIME_API FChaosVDShapeCollisionData
diff --git a/Engine/Source/Runtime/Experimental/ChaosVisualDebugger/Public/DataWrappers/ChaosVDDataSerializationMacros.h b/Engine/Source/Runtime/Experimental/ChaosVisualDebugger/Public/DataWrappers/ChaosVDDataSerializationMacros.h
new file mode 100644
--- /dev/null
+++ b/Engine/Source/Runtime/Experimental/ChaosVisualDebugger/Public/DataWrappers/ChaosVDDataSerializationMacros.h
@@ -0,0 +1,27 @@
+﻿// Copyright Epic Games, Inc. All Rights Reserved.
+
+#pragma once
+
+#ifndef CVD_READ_PACKED_BITFIELD_SERIALIZER
+	#define CVD_UNPACK_BITFIELD_DATA(Value, PackedBits, EnumFlag) \
+				Value = EnumHasAnyFlags(PackedBits, EnumFlag);
+#endif
+
+#ifndef CVD_WRITE_PACKED_BITFIELD_SERIALIZER
+	#define CVD_PACK_BITFIELD_DATA(Value, PackedBits, EnumFlag) \
+				if (Value) \
+				{ \
+					EnumAddFlags(PackedBits, EnumFlag);\
+				}
+#endif
+
+#ifndef CVD_SERIALIZE_STATIC_ARRAY
+	#define CVD_SERIALIZE_STATIC_ARRAY(Archive, Array) \
+	{ \
+		constexpr int32 Size = UE_ARRAY_COUNT(Array) ; \
+		for (int32 Index = 0; Index < Size; Index++)\
+		{\
+			Archive << Array[Index]; \
+		}\
+	}
+#endif
\ No newline at end of file
diff --git a/Engine/Source/Runtime/Experimental/ChaosVisualDebugger/Public/DataWrappers/ChaosVDParticleDataWrapper.h b/Engine/Source/Runtime/Experimental/ChaosVisualDebugger/Public/DataWrappers/ChaosVDParticleDataWrapper.h
--- a/Engine/Source/Runtime/Experimental/ChaosVisualDebugger/Public/DataWrappers/ChaosVDParticleDataWrapper.h
+++ b/Engine/Source/Runtime/Experimental/ChaosVisualDebugger/Public/DataWrappers/ChaosVDParticleDataWrapper.h
@@ -587,15 +587,13 @@ struct FChaosVDParticleDataWrapper : public FChaosVDWrapperDataBase
 
 	UPROPERTY()
 	TArray<FChaosVDShapeCollisionData> CollisionDataPerShape;
-	
+
 	/** Only used during recording */
-	TSharedPtr<FString> DebugNamePtr;
+	FString* DebugNamePtr = nullptr;
 	bool bHasDebugName = false;
 
 	virtual bool HasValidData() const override { return bHasValidData; }
 
-	void SetHasValidData(bool bNewIsValid){ bHasValidData = bNewIsValid; }
-
 	CHAOSVDRUNTIME_API bool Serialize(FArchive& Ar);
 };
 
diff --git a/Engine/Source/Runtime/Experimental/ChaosVisualDebugger/Public/DataWrappers/ChaosVDQueryDataWrappers.h b/Engine/Source/Runtime/Experimental/ChaosVisualDebugger/Public/DataWrappers/ChaosVDQueryDataWrappers.h
--- a/Engine/Source/Runtime/Experimental/ChaosVisualDebugger/Public/DataWrappers/ChaosVDQueryDataWrappers.h
+++ b/Engine/Source/Runtime/Experimental/ChaosVisualDebugger/Public/DataWrappers/ChaosVDQueryDataWrappers.h
@@ -94,6 +94,32 @@ struct TStructOpsTypeTraits<FChaosVDCollisionObjectQueryParams> : public TStruct
 	};
 };
 
+UENUM()
+enum class EChaosVDCollisionQueryParamsFlags : uint16
+{
+	None = 0,
+	/** Whether we should trace against complex collision */
+	TraceComplex = 1 << 0,
+	/** Whether we want to find out initial overlap or not. If true, it will return if this was initial overlap. */
+	FindInitialOverlaps = 1 << 1,
+	/** Whether we want to return the triangle face index for complex static mesh traces */
+	ReturnFaceIndex = 1 << 2,
+	/** Whether we want to include the physical material in the results. */
+	ReturnPhysicalMaterial = 1 << 3,
+	/** Whether to ignore blocking results. */
+	IgnoreBlocks = 1 << 4,
+	/** Whether to ignore touch/overlap results. */
+	IgnoreTouches = 1 << 5,
+	/** Whether to skip narrow phase checks (only for overlaps). */
+	SkipNarrowPhase = 1 << 6,
+	/** Whether to ignore traces to the cluster union and trace against its children instead. */
+	TraceIntoSubComponents = 1 << 7,
+	/** Extra filtering done on the query. See declaration for filtering logic */
+	ReplaceHitWithSubComponents = 1 << 8
+};
+
+ENUM_CLASS_FLAGS(EChaosVDCollisionQueryParamsFlags)
+
 USTRUCT()
 struct FChaosVDCollisionQueryParams : public FChaosVDWrapperDataBase
 {
@@ -117,9 +143,6 @@ struct FChaosVDCollisionQueryParams : public FChaosVDWrapperDataBase
 		bReplaceHitWithSubComponents = Other.bReplaceHitWithSubComponents;
 		IgnoreMask = Other.IgnoreMask;
 
-		IgnoredComponentsIDs = Other.GetIgnoredComponents();
-		IgnoredActorsIDs = Other.GetIgnoredSourceObjects();
-
 		bHasValidData = true;
 	}
 	
@@ -133,39 +156,39 @@ struct FChaosVDCollisionQueryParams : public FChaosVDWrapperDataBase
 
 	/** Whether we should trace against complex collision */
 	UPROPERTY(VisibleAnywhere, Category=QueryData)
-	bool bTraceComplex = false;
+	uint16 bTraceComplex : 1 = false;
 
 	/** Whether we want to find out initial overlap or not. If true, it will return if this was initial overlap. */
 	UPROPERTY(VisibleAnywhere, Category=QueryData)
-	bool bFindInitialOverlaps = false;
+	uint16 bFindInitialOverlaps : 1= false;
 
 	/** Whether we want to return the triangle face index for complex static mesh traces */
 	UPROPERTY(VisibleAnywhere, Category=QueryData)
-	bool bReturnFaceIndex = false;
+	uint16 bReturnFaceIndex : 1 = false;
 
 	/** Whether we want to include the physical material in the results. */
 	UPROPERTY(VisibleAnywhere, Category=QueryData)
-	bool bReturnPhysicalMaterial = false;
+	uint16 bReturnPhysicalMaterial : 1 = false;
 
 	/** Whether to ignore blocking results. */
 	UPROPERTY(VisibleAnywhere, Category=QueryData)
-	bool bIgnoreBlocks = false;
+	uint16 bIgnoreBlocks : 1 = false;
 
 	/** Whether to ignore touch/overlap results. */
 	UPROPERTY(VisibleAnywhere, Category=QueryData)
-	bool bIgnoreTouches = false;
+	uint16 bIgnoreTouches : 1 = false;
 
 	/** Whether to skip narrow phase checks (only for overlaps). */
 	UPROPERTY(VisibleAnywhere, Category=QueryData)
-	bool bSkipNarrowPhase = false;
+	uint16 bSkipNarrowPhase : 1 = false;
 
 	/** Whether to ignore traces to the cluster union and trace against its children instead. */
 	UPROPERTY(VisibleAnywhere, Category=QueryData)
-	bool bTraceIntoSubComponents = false;
+	uint16 bTraceIntoSubComponents : 1 = false;
 
 	/** If bTraceIntoSubComponents is true, whether to replace the hit of the cluster union with its children instead. */
 	UPROPERTY(VisibleAnywhere, Category=QueryData)
-	bool bReplaceHitWithSubComponents = false;
+	uint16 bReplaceHitWithSubComponents : 1 = false;
 
 	/** Extra filtering done on the query. See declaration for filtering logic */
 	UPROPERTY(VisibleAnywhere, Category=QueryData)
@@ -197,6 +220,15 @@ struct TStructOpsTypeTraits<FChaosVDCollisionQueryParams> : public TStructOpsTyp
 	};
 };
 
+UENUM()
+enum class EChaosVDQueryFastDataParallelFlags : uint8
+{
+	None = 0,
+	Parallel0 = 1 << 0,
+	Parallel1 = 1 << 1,
+	Parallel2 = 1 << 2
+};
+
 USTRUCT()
 struct FChaosVDQueryFastData : public FChaosVDWrapperDataBase
 {
@@ -217,13 +249,13 @@ struct FChaosVDQueryFastData : public FChaosVDWrapperDataBase
 	double InvCurrentLength = 0.0;
 
 	UPROPERTY(VisibleAnywhere, Category=QueryFastData)
-	bool bParallel0 = false;
+	uint8 bParallel0 : 1 = false;
 
 	UPROPERTY(VisibleAnywhere, Category=QueryFastData)
-	bool bParallel1 = false;
+	uint8 bParallel1 : 1 = false;
 
 	UPROPERTY(VisibleAnywhere, Category=QueryFastData)
-	bool bParallel2 = false;
+	uint8 bParallel2 : 1 = false;
 };
 
 inline FArchive& operator<<(FArchive& Ar, FChaosVDQueryFastData& Data)
