Undo changelist 31279946

Roll back the TAABB/TBox index fixes - it is causing some issues with box-mesh collisions.

diff --git a/Engine/Source/Programs/HeadlessChaos/Private/HeadlessChaosTestSupport.cpp b/Engine/Source/Programs/HeadlessChaos/Private/HeadlessChaosTestSupport.cpp
deleted file mode 100644
diff --git a/Engine/Source/Runtime/Experimental/Chaos/Private/Chaos/Box.cpp b/Engine/Source/Runtime/Experimental/Chaos/Private/Chaos/Box.cpp
--- a/Engine/Source/Runtime/Experimental/Chaos/Private/Chaos/Box.cpp
+++ b/Engine/Source/Runtime/Experimental/Chaos/Private/Chaos/Box.cpp
@@ -5,39 +5,16 @@
 
 namespace Chaos
 {
-	
-	// Indices below must match TAABB::GetVertex (ascii art below was copied from there)
-	// because TBox uses TAABB::SupportCore which returns a vertex index
-	// NOTE: This shows a right-handed coordinate system, but we are left-handed...
-	//
-	//	   6---------7
-	//	  /|        /|
-	//	 / |       / |
-	//	4---------5  |
-	//	|  |      |  |
-	//	|  2------|--3
-	//	| /       | /
-	//	|/        |/
-	//	0---------1
-	//
-	//  Z
-	//  |  Y
-	//  | /
-	//  |/
-	//  o-----X
-	// 
-	//
-
 	FConvexHalfEdgeStructureDataS16 MakeBoxStructureData()
 	{
 		TArray<TArray<int32>> PlaneVertices
 		{
-			{ 0, 4, 6, 2 },	// -X,
-			{ 0, 1, 5, 4 },	// -Y
-			{ 0, 2, 3, 1 },	// -Z
-			{ 1, 3, 7, 5 },	//  X
-			{ 2, 6, 7, 3 },	//  Y
-			{ 4, 5, 7, 6 },	//  Z
+			{7, 5, 4, 6},	// X
+			{2, 3, 7, 6},	// Y
+			{5, 7, 3, 1},	// Z,
+			{2, 0, 1, 3},	// -X
+			{0, 4, 5, 1},	// -Y
+			{0, 2, 6, 4},	// -Z,
 		};
 
 		return FConvexHalfEdgeStructureDataS16::MakePlaneVertices(PlaneVertices, 8);
@@ -45,24 +22,24 @@ namespace Chaos
 
 	template<typename T, int D> TArray<FVec3> TBox<T, D>::SNormals =
 	{
-		FVec3(-1,0,0),	// -X
-		FVec3(0,-1,0),	// -Y
-		FVec3(0,0,-1),	// -Z
-		FVec3(1,0,0),	//  X
-		FVec3(0,1,0),	//  Y
-		FVec3(0,0,1),	//  Z
+		FVec3(1,0,0),	 // X
+		FVec3(0,1,0),	 // Y
+		FVec3(0,0,1),	 // Z
+		FVec3(-1,0,0), // -X
+		FVec3(0,-1,0), // -Y
+		FVec3(0,0,-1), // -Z
 	};
 	
 	template<typename T, int D> TArray<FVec3> TBox<T, D>::SVertices =
 	{
-		FVec3(-1,-1,-1),	// 0 
-		FVec3( 1,-1,-1),	// 1 
-		FVec3(-1, 1,-1),	// 2 
-		FVec3( 1, 1,-1),	// 3 
-		FVec3(-1,-1, 1),	// 4 
-		FVec3( 1,-1, 1),	// 5 
-		FVec3(-1, 1, 1),	// 6 
-		FVec3( 1, 1, 1),	// 7 
+		FVec3(-1,-1,-1),
+		FVec3(-1,-1,1),
+		FVec3(-1,1,-1),
+		FVec3(-1,1,1),
+		FVec3(1,-1,-1),
+		FVec3(1,-1,1),
+		FVec3(1,1,-1),
+		FVec3(1,1,1),
 	};
 
 	template<typename T, int D> FConvexHalfEdgeStructureDataS16 TBox<T, D>::SStructureData = MakeBoxStructureData();
diff --git a/Engine/Source/Runtime/Experimental/Chaos/Public/Chaos/AABB.h b/Engine/Source/Runtime/Experimental/Chaos/Public/Chaos/AABB.h
--- a/Engine/Source/Runtime/Experimental/Chaos/Public/Chaos/AABB.h
+++ b/Engine/Source/Runtime/Experimental/Chaos/Public/Chaos/AABB.h
@@ -311,7 +311,7 @@ namespace Chaos
 			return BestNormal;
 		}
 		
-		FORCEINLINE_DEBUGGABLE TVector<T, d> Support(const TVector<T, d>& Direction, const T Thickness, int32& OutVertexIndex) const
+		FORCEINLINE_DEBUGGABLE TVector<T, d> Support(const TVector<T, d>& Direction, const T Thickness, int32& VertexIndex) const
 		{
 			TVector<T, d> ChosenPt;
 			FIntVector ChosenAxis;
@@ -328,8 +328,7 @@ namespace Chaos
 					ChosenAxis[Axis] = 1;
 				}
 			}
-
-			OutVertexIndex = GetIndex(ChosenAxis);
+			VertexIndex =  ChosenAxis[0] * 4 + ChosenAxis[1] * 2 + ChosenAxis[2];
 
 			if (Thickness != (T)0)
 			{
@@ -350,7 +349,7 @@ namespace Chaos
 		}
 
 		// Support vertex in the specified direction, assuming each face has been moved inwards by InMargin
-		FORCEINLINE_DEBUGGABLE FVec3 SupportCore(const FVec3& Direction, const FReal InMargin, FReal* OutSupportDelta, int32& OutVertexIndex) const
+		FORCEINLINE_DEBUGGABLE FVec3 SupportCore(const FVec3& Direction, const FReal InMargin, FReal* OutSupportDelta, int32& VertexIndex) const
 		{
 			FVec3 ChosenPt;
 			FIntVector ChosenAxis;
@@ -367,9 +366,7 @@ namespace Chaos
 					ChosenAxis[Axis] = 1;
 				}
 			}
-
-			OutVertexIndex = GetIndex(ChosenAxis);
-
+			VertexIndex = ChosenAxis[0] * 4 + ChosenAxis[1] * 2 + ChosenAxis[2];
 			// Maximum distance between the Core+Margin position and the original outer vertex
 			constexpr FReal RootThreeMinusOne = FReal(1.7320508075688772935274463415059 - 1.0);
 			if (OutSupportDelta != nullptr)
@@ -389,7 +386,7 @@ namespace Chaos
 			return MakeVectorRegisterFloatFromDouble(MakeVectorRegister(SupportVert.X, SupportVert.Y, SupportVert.Z, 0.0));
 		}
 
-		FORCEINLINE_DEBUGGABLE TVector<T, d> SupportCoreScaled(const TVector<T, d>& Direction, const T InMargin, const TVector<T, d>& Scale, T* OutSupportDelta, int32& OutVertexIndex) const
+		FORCEINLINE_DEBUGGABLE TVector<T, d> SupportCoreScaled(const TVector<T, d>& Direction, const T InMargin, const TVector<T, d>& Scale, T* OutSupportDelta, int32& VertexIndex) const
 		{
 			const TVector<T, d> ScaledDirection = Direction * Scale; // Compensate for Negative scales, only the signs really matter here
 
@@ -408,9 +405,8 @@ namespace Chaos
 					ChosenAxis[Axis] = 1;
 				}
 			}
-
-			OutVertexIndex = GetIndex(ChosenAxis);
-
+			VertexIndex = ChosenAxis[0] * 4 + ChosenAxis[1] * 2 + ChosenAxis[2];
+			
 			constexpr T RootThreeMinusOne = T(1.7320508075688772935274463415059 - 1.0);
 			if (OutSupportDelta != nullptr)
 			{
@@ -499,26 +495,12 @@ namespace Chaos
 		FORCEINLINE TVector<T, d> GetVertex(const int32 Index) const
 		{
 			check(0 <= Index && Index < 8);
-
-			// See GetIndex() for reverse logic
 			return TVector<T, d>(
 				(Index & (1 << 0)) == 0 ? MMin.X : MMax.X,
 				(Index & (1 << 1)) == 0 ? MMin.Y : MMax.Y,
 				(Index & (1 << 2)) == 0 ? MMin.Z : MMax.Z);
 		}
 
-		/**
-		* Given a point on a unit cube expressed as an IntVector of 0s and 1s, get the vertex index of that point.
-		* 0 represents a negative axis and 1 represents a positive axis. All other values are invalid, but we do 
-		* not check this as it would add overhead to support functions. This is used by the support functions and 
-		* must invert the index logic in GetVertex().
-		*/
-		FORCEINLINE int32 GetIndex(const FIntVector& AxisSelector) const
-		{
-			// See GetVertex() for reverse logic
-			return AxisSelector[0] + AxisSelector[1] * 2 + AxisSelector[2] * 4;
-		}
-
 		/**
 		*	Get an array of two indices into the vertex list for one of the
 		*	twelve edges of the box. Edges are ordered by increasing vertex
