Iris
- Changed Iris FastArraySerializer implementation to not rely ony ItemMap when issues FastArraySerialzierCallbacks as the itemmap is reset from callbacks.

diff --git a/Engine/Source/Runtime/Experimental/Iris/Core/Public/Iris/ReplicationSystem/Private/FastArrayReplicationFragmentInternal.h b/Engine/Source/Runtime/Experimental/Iris/Core/Public/Iris/ReplicationSystem/Private/FastArrayReplicationFragmentInternal.h
--- a/Engine/Source/Runtime/Experimental/Iris/Core/Public/Iris/ReplicationSystem/Private/FastArrayReplicationFragmentInternal.h
+++ b/Engine/Source/Runtime/Experimental/Iris/Core/Public/Iris/ReplicationSystem/Private/FastArrayReplicationFragmentInternal.h
@@ -153,6 +153,8 @@ void FFastArrayReplicationFragmentHelper::ConditionalRebuildItemMap(FastArrayTyp
 
 	if (ArraySerializer.ItemMap.Num() != Items.Num())
 	{
+		UE_LOG(LogNetFastTArray, Verbose, TEXT("FastArrayDeltaSerialize: Recreating Items map. Items.Num: %d Map.Num: %d"), Items.Num(), ArraySerializer.ItemMap.Num());
+
 		ArraySerializer.ItemMap.Reset();
 			
 		const ItemType* SrcItems = Items.GetData();
@@ -204,9 +206,6 @@ void FFastArrayReplicationFragmentHelper::ApplyReplicatedState(FastArrayType* Ds
 			{
 				UE_LOG(LogNetFastTArray, Log, TEXT("   Removed ID: %d local Idx: %d"), ReplicationID, It);
 				RemovedIndices.Add(It);
-
-				// Remove callback
-				DstItems[It].PreReplicatedRemove(*DstArraySerializer);
 			}
 		}
 	}
@@ -250,13 +249,28 @@ void FFastArrayReplicationFragmentHelper::ApplyReplicatedState(FastArrayType* Ds
 		}
 	}
 
+	// Increment keys so that a client can re-serialize the array if needed, such as for client replay recording.
+	DstArraySerializer->IncrementArrayReplicationKey();
+
 	// Added and changed callbacks to FastArraySerializer
 	const int32 PreRemoveSize = DstWrappedArray->Num();
 	const int32 FinalSize = PreRemoveSize - RemovedIndices.Num();
 
+	// Remove callback
+	for (int32 RemovedIndex : RemovedIndices)
+	{
+		(*DstWrappedArray)[RemovedIndex].PreReplicatedRemove(*DstArraySerializer);
+	}
+
 	// Remove callback to FastArraySerializer - done after adding new elements
 	DstArraySerializer->PreReplicatedRemove(MakeArrayView(RemovedIndices), FinalSize);
 
+	if (PreRemoveSize != DstWrappedArray->Num())
+	{
+		UE_LOG(LogNetFastTArray, Error, TEXT("Item size changed after PreReplicatedRemove! PremoveSize: %d  Item.Num: %d"),
+			PreRemoveSize, DstWrappedArray->Num());
+	}
+
 	// Add callbacks
 	for (int32 AddedIndex : AddedIndices)
 	{
@@ -271,6 +285,12 @@ void FFastArrayReplicationFragmentHelper::ApplyReplicatedState(FastArrayType* Ds
 	}
 	DstArraySerializer->PostReplicatedChange(MakeArrayView(ModifiedIndices), FinalSize);
 
+	if (PreRemoveSize != DstWrappedArray->Num())
+	{
+		UE_LOG(LogNetFastTArray, Error, TEXT("Item size changed after PostReplicatedAdd/PostReplicatedChange! PreRemoveSize: %d  Item.Num: %d"),
+			PreRemoveSize, DstWrappedArray->Num());
+	}
+
 	// Remove indices
 	if (RemovedIndices.Num() > 0)
 	{
@@ -289,9 +309,6 @@ void FFastArrayReplicationFragmentHelper::ApplyReplicatedState(FastArrayType* Ds
 		DstArraySerializer->ItemMap.Empty();
 	}
 
-	// Increment keys so that a client can re-serialize the array if needed, such as for client replay recording.
-	DstArraySerializer->IncrementArrayReplicationKey();
-
 	// Invoke PostReplicatedReceive if is defined by the serializer
 	CallPostReplicatedReceiveOrNot(*DstArraySerializer, Context.bHasUnresolvableReferences);
 }
