Iris
* Added FIrisDelegates::GetCriticalErrorDetectedDelegate
* Currenly triggered when a protocol mismatch occurs

#rnx

diff --git a/Engine/Source/Runtime/Experimental/Iris/Core/Private/Iris/Core/IrisDelegates.cpp b/Engine/Source/Runtime/Experimental/Iris/Core/Private/Iris/Core/IrisDelegates.cpp
new file mode 100644
--- /dev/null
+++ b/Engine/Source/Runtime/Experimental/Iris/Core/Private/Iris/Core/IrisDelegates.cpp
@@ -0,0 +1,16 @@
+// Copyright Epic Games, Inc. All Rights Reserved.
+
+#include "Iris/Core/IrisDelegates.h"
+
+
+namespace UE::Net
+{
+
+FIrisCriticalErrorDetected& FIrisDelegates::GetCriticalErrorDetectedDelegate()
+{
+	static FIrisCriticalErrorDetected CriticalErrorDelegate;
+	return CriticalErrorDelegate;
+}
+
+} // end namespace UE::Net
+
diff --git a/Engine/Source/Runtime/Experimental/Iris/Core/Private/Iris/ReplicationSystem/ObjectReplicationBridge.cpp b/Engine/Source/Runtime/Experimental/Iris/Core/Private/Iris/ReplicationSystem/ObjectReplicationBridge.cpp
--- a/Engine/Source/Runtime/Experimental/Iris/Core/Private/Iris/ReplicationSystem/ObjectReplicationBridge.cpp
+++ b/Engine/Source/Runtime/Experimental/Iris/Core/Private/Iris/ReplicationSystem/ObjectReplicationBridge.cpp
@@ -11,6 +11,7 @@
 #include "Iris/Core/IrisMemoryTracker.h"
 #include "Iris/Core/IrisProfiler.h"
 #include "Iris/Core/IrisDebugging.h"
+#include "Iris/Core/IrisDelegates.h"
 
 #include "Net/Core/NetBitArrayPrinter.h"
 #include "Net/Core/NetHandle/NetHandleManager.h"
@@ -737,6 +738,8 @@ FReplicationBridgeCreateNetRefHandleResult UObjectReplicationBridge::CreateNetRe
 			UE::Net::Private::ObjectBridgeDebugging::RemoteProtocolMismatchDetected(ReplicationSystem, Context.ConnectionId, RegisteredFragments, ArchetypeOrCDOUsedAsKey, InstancePtr);
 		}
 
+		FIrisDelegates::GetCriticalErrorDetectedDelegate().Broadcast(ReplicationSystem);
+
 		OnProtocolMismatchDetected(WantedNetHandle);
 	}
 	else
diff --git a/Engine/Source/Runtime/Experimental/Iris/Core/Public/Iris/Core/IrisDelegates.h b/Engine/Source/Runtime/Experimental/Iris/Core/Public/Iris/Core/IrisDelegates.h
new file mode 100644
--- /dev/null
+++ b/Engine/Source/Runtime/Experimental/Iris/Core/Public/Iris/Core/IrisDelegates.h
@@ -0,0 +1,25 @@
+// Copyright Epic Games, Inc. All Rights Reserved.
+
+#pragma once
+
+#include "Delegates/Delegate.h"
+
+
+class UReplicationSystem;
+
+namespace UE::Net
+{
+
+DECLARE_MULTICAST_DELEGATE_OneParam(FIrisCriticalErrorDetected, UReplicationSystem*);
+
+class FIrisDelegates
+{
+public:
+	IRISCORE_API static FIrisCriticalErrorDetected& GetCriticalErrorDetectedDelegate();
+};
+
+	
+	
+
+} // end namespace UE::Net
+
