Chaos:
- User defined Entity Append data is now destroyed at appropriate time (just before BodyInstances)

#[fyi] Benn.Gallagher

diff --git a/Engine/Source/Runtime/Engine/Private/Components/PrimitiveComponent.cpp b/Engine/Source/Runtime/Engine/Private/Components/PrimitiveComponent.cpp
--- a/Engine/Source/Runtime/Engine/Private/Components/PrimitiveComponent.cpp
+++ b/Engine/Source/Runtime/Engine/Private/Components/PrimitiveComponent.cpp
@@ -961,6 +961,10 @@ void UPrimitiveComponent::OnDestroyPhysicsState()
 	UnWeldFromParent();
 	UnWeldChildren();
 
+	// Remove all user defined entities here
+	TArray<Chaos::FPhysicsObject*> PhysicsObjects = GetAllPhysicsObjects();
+	FPhysicsObjectExternalInterface::LockWrite(PhysicsObjects)->SetUserDefinedEntity(PhysicsObjects, nullptr);
+
 	// clean up physics engine representation
 	if(BodyInstance.IsValidBodyInstance())
 	{
diff --git a/Engine/Source/Runtime/Experimental/Chaos/Private/Chaos/PhysicsObjectInterface.cpp b/Engine/Source/Runtime/Experimental/Chaos/Private/Chaos/PhysicsObjectInterface.cpp
--- a/Engine/Source/Runtime/Experimental/Chaos/Private/Chaos/PhysicsObjectInterface.cpp
+++ b/Engine/Source/Runtime/Experimental/Chaos/Private/Chaos/PhysicsObjectInterface.cpp
@@ -862,14 +862,26 @@ namespace Chaos
 					FChaosUserEntityAppend* UserEntityAppend = FChaosUserData::Get<FChaosUserEntityAppend>(Particle->UserData());
 					if (!UserEntityAppend)
 					{
-						UserEntityAppend = new FChaosUserEntityAppend;
-						UserEntityAppend->ChaosUserData = reinterpret_cast<FChaosUserData*>(Particle->UserData());
-						UserEntityAppend->UserDefinedEntity = UserDefinedEntity;
-						Particle->SetUserData(UserEntityAppend);
+						if (UserDefinedEntity)
+						{
+							UserEntityAppend = new FChaosUserEntityAppend;
+							UserEntityAppend->ChaosUserData = reinterpret_cast<FChaosUserData*>(Particle->UserData());
+							UserEntityAppend->UserDefinedEntity = UserDefinedEntity;
+							Particle->SetUserData(UserEntityAppend);
+						}
 					}
 					else
 					{
-						UserEntityAppend->UserDefinedEntity = UserDefinedEntity; // Overwrite previous used defined entity
+						if (UserDefinedEntity)
+						{
+							UserEntityAppend->UserDefinedEntity = UserDefinedEntity; // Overwrite previous used defined entity
+						}
+						else
+						{
+							// Restore the particle user data and delete unused memory
+							Particle->SetUserData(UserEntityAppend->ChaosUserData);
+							delete UserEntityAppend;
+						}
 					}
 				}
 			}
diff --git a/Engine/Source/Runtime/Experimental/Chaos/Public/Chaos/ParticleHandle.h b/Engine/Source/Runtime/Experimental/Chaos/Public/Chaos/ParticleHandle.h
--- a/Engine/Source/Runtime/Experimental/Chaos/Public/Chaos/ParticleHandle.h
+++ b/Engine/Source/Runtime/Experimental/Chaos/Public/Chaos/ParticleHandle.h
@@ -2608,12 +2608,6 @@ public:
 	virtual ~TGeometryParticle() //only virtual for easier memory management. Should generally be a static API
 	{
 		SetUniqueIdx(FUniqueIdx{}, false); // Set to an invalid index for dangling handle detection
-		// TODO: Temporary CI fix: Prevent reading Data on bodyInstances after they may be deallocated, delete custom entities elsewhere
-		FChaosUserEntityAppend* UserEntityAppend = nullptr; // FChaosUserData::Get<FChaosUserEntityAppend>(MUserData); 
-		if (UserEntityAppend)
-		{
-			delete UserEntityAppend;
-		}
 	}
 
 	TGeometryParticle(const TGeometryParticle&) = delete;
diff --git a/Engine/Source/Runtime/Experimental/Chaos/Public/Chaos/PhysicsObjectInterface.h b/Engine/Source/Runtime/Experimental/Chaos/Public/Chaos/PhysicsObjectInterface.h
--- a/Engine/Source/Runtime/Experimental/Chaos/Public/Chaos/PhysicsObjectInterface.h
+++ b/Engine/Source/Runtime/Experimental/Chaos/Public/Chaos/PhysicsObjectInterface.h
@@ -125,7 +125,7 @@ namespace Chaos
 	class FWritePhysicsObjectInterface: public FReadPhysicsObjectInterface<Id>
 	{
 	public:
-		CHAOS_API void SetUserDefinedEntity(TArrayView<const FPhysicsObjectHandle> InObjects, FChaosUserDefinedEntity* UserDefinedEntity);
+		CHAOS_API void SetUserDefinedEntity(TArrayView<const FPhysicsObjectHandle> InObjects, FChaosUserDefinedEntity* UserDefinedEntity); // Set the user defined entity, use nullptr to remove the Entity and release the memory
 		CHAOS_API void PutToSleep(TArrayView<const FPhysicsObjectHandle> InObjects);
 		CHAOS_API void WakeUp(TArrayView<const FPhysicsObjectHandle> InObjects);
 		CHAOS_API void ForceKinematic(TArrayView<const FPhysicsObjectHandle> InObjects);
