Iris
* Added NetRefHandle::GetInvalid() and replaced existing calls to FNetRefHandle() to GetInvalid()

#rnx

diff --git a/Engine/Plugins/Runtime/ReplicationSystemTestPlugin/Source/Private/Tests/ReplicationSystem/ReplicatedTestObject.h b/Engine/Plugins/Runtime/ReplicationSystemTestPlugin/Source/Private/Tests/ReplicationSystem/ReplicatedTestObject.h
--- a/Engine/Plugins/Runtime/ReplicationSystemTestPlugin/Source/Private/Tests/ReplicationSystem/ReplicatedTestObject.h
+++ b/Engine/Plugins/Runtime/ReplicationSystemTestPlugin/Source/Private/Tests/ReplicationSystem/ReplicatedTestObject.h
@@ -499,7 +499,7 @@ public:
 	// In this example we have methods that directly uses UReplicatedTestObject;
 	FNetRefHandle BeginReplication(UReplicatedTestObject* Instance);
 	FNetRefHandle BeginReplication(UReplicatedTestObject* Instance, const UObjectReplicationBridge::FCreateNetRefHandleParams& Params);
-	FNetRefHandle BeginReplication(FNetRefHandle OwnerHandle, UReplicatedTestObject* Instance, FNetRefHandle InsertRelativeToSubObjectHandle = FNetRefHandle(), ESubObjectInsertionOrder InsertionOrder = UReplicationBridge::ESubObjectInsertionOrder::None);
+	FNetRefHandle BeginReplication(FNetRefHandle OwnerHandle, UReplicatedTestObject* Instance, FNetRefHandle InsertRelativeToSubObjectHandle = FNetRefHandle::GetInvalid(), ESubObjectInsertionOrder InsertionOrder = UReplicationBridge::ESubObjectInsertionOrder::None);
 
 	// For testing we expose some things that normally are not accessible
 	const UE::Net::FReplicationInstanceProtocol* GetReplicationInstanceProtocol(FNetRefHandle Handle) const;
diff --git a/Engine/Plugins/Runtime/ReplicationSystemTestPlugin/Source/Private/Tests/ReplicationSystem/ReplicationSystemServerClientTestFixture.h b/Engine/Plugins/Runtime/ReplicationSystemTestPlugin/Source/Private/Tests/ReplicationSystem/ReplicationSystemServerClientTestFixture.h
--- a/Engine/Plugins/Runtime/ReplicationSystemTestPlugin/Source/Private/Tests/ReplicationSystem/ReplicationSystemServerClientTestFixture.h
+++ b/Engine/Plugins/Runtime/ReplicationSystemTestPlugin/Source/Private/Tests/ReplicationSystem/ReplicationSystemServerClientTestFixture.h
@@ -95,7 +95,7 @@ public:
 	}
 
 	template<typename T>
-	T* CreateSubObject(FNetRefHandle OwnerHandle, FNetRefHandle InsertRelativeToSubObjectHandle = FNetRefHandle(), UReplicationBridge::ESubObjectInsertionOrder InsertionOrder = UReplicationBridge::ESubObjectInsertionOrder::None)
+	T* CreateSubObject(FNetRefHandle OwnerHandle, FNetRefHandle InsertRelativeToSubObjectHandle = FNetRefHandle::GetInvalid(), UReplicationBridge::ESubObjectInsertionOrder InsertionOrder = UReplicationBridge::ESubObjectInsertionOrder::None)
 	{
 		T* CreatedObject = NewObject<T>();
 		if (Cast<UReplicatedTestObject>(CreatedObject))
diff --git a/Engine/Source/Runtime/Engine/Private/Net/Iris/ReplicationSystem/ActorReplicationBridge.cpp b/Engine/Source/Runtime/Engine/Private/Net/Iris/ReplicationSystem/ActorReplicationBridge.cpp
--- a/Engine/Source/Runtime/Engine/Private/Net/Iris/ReplicationSystem/ActorReplicationBridge.cpp
+++ b/Engine/Source/Runtime/Engine/Private/Net/Iris/ReplicationSystem/ActorReplicationBridge.cpp
@@ -168,69 +168,69 @@ UE::Net::FNetRefHandle UActorReplicationBridge::BeginReplication(AActor* Actor,
 
 	if (!ShouldUseIrisReplication(Actor))
 	{
-		return FNetRefHandle();
+		return FNetRefHandle::GetInvalid();
 	}
 
 	if (!ensureMsgf(Actor == nullptr || !(Actor->HasAnyFlags(RF_ClassDefaultObject | RF_ArchetypeObject)), TEXT("Actor %s is a CDO or Archetype and should not be replicated."), ToCStr(GetFullNameSafe(Actor))))
 	{
-		return FNetRefHandle();
+		return FNetRefHandle::GetInvalid();
 	}
 
 	const bool bIsNetActor = ULevel::IsNetActor(Actor);
 	if (!bIsNetActor)
 	{
 		UE_LOG_ACTORREPLICATIONBRIDGE(VeryVerbose, TEXT("Actor %s doesn't have a NetRole."), ToCStr(GetFullNameSafe(Actor)));
-		return FNetRefHandle();
+		return FNetRefHandle::GetInvalid();
 	}
 
 	if (Actor->GetLocalRole() != ROLE_Authority)
 	{
 		UE_LOG_ACTORREPLICATIONBRIDGE(VeryVerbose, TEXT("Actor %s NetRole isn't Authority."), ToCStr(GetFullNameSafe(Actor)));
-		return FNetRefHandle();
+		return FNetRefHandle::GetInvalid();
 	}
 
 	if (Actor->IsActorBeingDestroyed() || !IsValid(Actor) || Actor->IsUnreachable())
 	{
 		UE_LOG_ACTORREPLICATIONBRIDGE(Verbose, TEXT("Actor %s is being destroyed or unreachable and can't be replicated."), ToCStr(GetFullNameSafe(Actor)));
-		return FNetRefHandle();
+		return FNetRefHandle::GetInvalid();
 	}
 
 	if (!Actor->GetIsReplicated())
 	{
 		UE_LOG_ACTORREPLICATIONBRIDGE(Verbose, TEXT("Actor %s is not supposed to be replicated."), ToCStr(GetFullNameSafe(Actor)));
-		return FNetRefHandle();
+		return FNetRefHandle::GetInvalid();
 	}
 
 	if (Actor->GetTearOff())
 	{
 		UE_LOG_ACTORREPLICATIONBRIDGE(Verbose, TEXT("Actor %s is torn off and should not be replicated."), ToCStr(GetFullNameSafe(Actor)));
-		return FNetRefHandle();
+		return FNetRefHandle::GetInvalid();
 	}
 
 	if (!Actor->IsActorInitialized())
 	{
 		UE_LOG_ACTORREPLICATIONBRIDGE(Warning, TEXT("Actor %s is not initialized and won't be replicated."), ToCStr(GetFullNameSafe(Actor)));
-		return FNetRefHandle();
+		return FNetRefHandle::GetInvalid();
 	}
 
 	if (!NetDriver)
 	{
 		UE_LOG_ACTORREPLICATIONBRIDGE(VeryVerbose, TEXT("There's no NetDriver so nothing can be replicated."));
-		return FNetRefHandle();
+		return FNetRefHandle::GetInvalid();
 	}
 
 
 	if (!NetDriver->ShouldReplicateActor(Actor))
 	{
 		UE_LOG_ACTORREPLICATIONBRIDGE(VeryVerbose, TEXT("Actor %s doesn't want to replicate with NetDriver %s."), ToCStr(GetFullNameSafe(Actor)), ToCStr(NetDriver->GetName()));
-		return FNetRefHandle();
+		return FNetRefHandle::GetInvalid();
 	}
 
 	// Initially dormant actors begin replication when their dormancy is flushed
 	const ENetDormancy Dormancy = Actor->NetDormancy;	
 	if (Actor->IsNetStartupActor() && (Dormancy == DORM_Initial))
 	{
-		return FNetRefHandle();
+		return FNetRefHandle::GetInvalid();
 	}
 
 	FNetRefHandle ExistingHandle = GetReplicatedRefHandle(Actor);
@@ -263,7 +263,7 @@ UE::Net::FNetRefHandle UActorReplicationBridge::BeginReplication(AActor* Actor,
 	if (!ActorRefHandle.IsValid())
 	{
 		ensureMsgf(false, TEXT("Failed to create NetRefHandle for Actor Named %s"), ToCStr(Actor->GetName()));
-		return FNetRefHandle();
+		return FNetRefHandle::GetInvalid();
 	}
 
 	// Set owning connection filtering if actor is only relevant to owner
@@ -356,7 +356,7 @@ UE::Net::FNetRefHandle UActorReplicationBridge::BeginReplication(FNetRefHandle O
 
 	if (!OwnerHandle.IsValid() || !ShouldUseIrisReplication(SubObject))
 	{
-		return FNetRefHandle();
+		return FNetRefHandle::GetInvalid();
 	}
 
 	AActor* Owner = SubObject->GetOwner();
@@ -372,7 +372,7 @@ UE::Net::FNetRefHandle UActorReplicationBridge::BeginReplication(FNetRefHandle O
 			|| SubObject->HasAnyFlags(RF_ArchetypeObject | RF_ClassDefaultObject)
 		)
 		{
-			return FNetRefHandle();
+			return FNetRefHandle::GetInvalid();
 		}
 
 		if (!SubObject->IsUsingRegisteredSubObjectList())
@@ -382,7 +382,7 @@ UE::Net::FNetRefHandle UActorReplicationBridge::BeginReplication(FNetRefHandle O
 
 		if (RepComponentInfo == nullptr || RepComponentInfo->NetCondition == ELifetimeCondition::COND_Never)
 		{
-			return FNetRefHandle();
+			return FNetRefHandle::GetInvalid();
 		}
 
 		// Start replicating the subobject with its owner.
@@ -392,7 +392,7 @@ UE::Net::FNetRefHandle UActorReplicationBridge::BeginReplication(FNetRefHandle O
 	if (!ReplicatedComponentHandle.IsValid())
 	{
 		ensureMsgf(false, TEXT("Failed to create or find NetRefHandle for ActorComponent Named %s"), ToCStr(SubObject->GetName()));
-		return FNetRefHandle();
+		return FNetRefHandle::GetInvalid();
 	}
 
 	// Update or set any conditionals
diff --git a/Engine/Source/Runtime/Experimental/Iris/Core/Private/Iris/ReplicationSystem/NetRefHandleManager.h b/Engine/Source/Runtime/Experimental/Iris/Core/Private/Iris/ReplicationSystem/NetRefHandleManager.h
--- a/Engine/Source/Runtime/Experimental/Iris/Core/Private/Iris/ReplicationSystem/NetRefHandleManager.h
+++ b/Engine/Source/Runtime/Experimental/Iris/Core/Private/Iris/ReplicationSystem/NetRefHandleManager.h
@@ -451,7 +451,7 @@ FNetRefHandle FNetRefHandleManager::GetCompleteNetRefHandle(FNetRefHandle Incomp
 	}
 	else
 	{
-		return FNetRefHandle();
+		return FNetRefHandle::GetInvalid();
 	}
 }
 
@@ -489,7 +489,7 @@ uint32 FNetRefHandleManager::GetOriginalDestroyedStartupObjectIndex(FInternalNet
 
 inline bool FNetRefHandleManager::AddSubObject(FNetRefHandle OwnerHandle, FNetRefHandle SubObjectHandle, EAddSubObjectFlags Flags)
 {
-	return AddSubObject(OwnerHandle, SubObjectHandle, FNetRefHandle(), Flags);
+	return AddSubObject(OwnerHandle, SubObjectHandle, FNetRefHandle::GetInvalid(), Flags);
 }
 
 TArrayView<const FInternalNetRefIndex> FNetRefHandleManager::GetSubObjects(FInternalNetRefIndex OwnerIndex) const
diff --git a/Engine/Source/Runtime/Experimental/Iris/Core/Private/Iris/ReplicationSystem/ObjectReferenceCache.cpp b/Engine/Source/Runtime/Experimental/Iris/Core/Private/Iris/ReplicationSystem/ObjectReferenceCache.cpp
--- a/Engine/Source/Runtime/Experimental/Iris/Core/Private/Iris/ReplicationSystem/ObjectReferenceCache.cpp
+++ b/Engine/Source/Runtime/Experimental/Iris/Core/Private/Iris/ReplicationSystem/ObjectReferenceCache.cpp
@@ -343,7 +343,7 @@ bool FObjectReferenceCache::CreateObjectReferenceInternal(const UObject* Object,
 			constexpr bool bReading = false;
 			// The connection ID isn't used unless reading.
 			RenamePathForPie(UE::Net::InvalidConnectionId, ObjectPath, bReading);
-			OutReference = MakeNetObjectReference(FNetRefHandle(), StringTokenStore->GetOrCreateToken(ObjectPath));	
+			OutReference = MakeNetObjectReference(FNetRefHandle::GetInvalid(), StringTokenStore->GetOrCreateToken(ObjectPath));	
 			return true;
 		}
 	}
@@ -423,7 +423,7 @@ FNetRefHandle FObjectReferenceCache::GetObjectReferenceHandleFromObject(const UO
 {
 	if (Object == nullptr)
 	{
-		return FNetRefHandle();
+		return FNetRefHandle::GetInvalid();
 	}
 
 	// Check if we already know about this object
@@ -436,7 +436,7 @@ FNetRefHandle FObjectReferenceCache::GetObjectReferenceHandleFromObject(const UO
 		{
 			UE_LOG_REFERENCECACHE(Verbose, TEXT("ObjectReferenceCache::GetObjectReferenceHandleFromObject removed %s from ObjectToNetReferenceHandle due to ReferenceHandleToCachedReference not holding the handle. For Object %s (0x%p) "), *Reference->ToString(), *GetNameSafe(Object), Object);
 			const_cast<TMap<const UObject*, FNetRefHandle>&>(ObjectToNetReferenceHandle).Remove(Object);
-			return FNetRefHandle();
+			return FNetRefHandle::GetInvalid();
 		}
 
 		const FCachedNetObjectReference& CachedObject = *CachedObjectPtr;
@@ -455,7 +455,7 @@ FNetRefHandle FObjectReferenceCache::GetObjectReferenceHandleFromObject(const UO
 		}
 	}
 
-	return FNetRefHandle();
+	return FNetRefHandle::GetInvalid();
 }
 
 void FObjectReferenceCache::AddRemoteReference(FNetRefHandle RefHandle, const UObject* Object)
@@ -504,7 +504,7 @@ void FObjectReferenceCache::AddRemoteReference(FNetRefHandle RefHandle, const UO
 			CachedObject.ObjectKey = Object;
 			CachedObject.NetRefHandle = RefHandle;
 			CachedObject.RelativePath = FNetToken();
-			CachedObject.OuterNetRefHandle = FNetRefHandle();
+			CachedObject.OuterNetRefHandle = FNetRefHandle::GetInvalid();
 			CachedObject.bIsPackage = false;
 			CachedObject.bNoLoad = false;
 			CachedObject.bIsBroken = false;
@@ -1113,7 +1113,7 @@ FNetObjectReference FObjectReferenceCache::GetReplicatedOuter(const FNetObjectRe
 	}
 
 	const FCachedNetObjectReference* CachedObject = ReferenceHandleToCachedReference.Find(Reference.GetRefHandle());
-	FNetRefHandle OuterNetRefHandle = CachedObject ? CachedObject->OuterNetRefHandle : FNetRefHandle();
+	FNetRefHandle OuterNetRefHandle = CachedObject ? CachedObject->OuterNetRefHandle : FNetRefHandle::GetInvalid();
 
 	while (OuterNetRefHandle.IsValid())
 	{
@@ -1123,7 +1123,7 @@ FNetObjectReference FObjectReferenceCache::GetReplicatedOuter(const FNetObjectRe
 		}
 
 		CachedObject = ReferenceHandleToCachedReference.Find(OuterNetRefHandle);
-		OuterNetRefHandle = CachedObject ? CachedObject->OuterNetRefHandle : FNetRefHandle();
+		OuterNetRefHandle = CachedObject ? CachedObject->OuterNetRefHandle : FNetRefHandle::GetInvalid();
 	};
 
 	return FNetObjectReference();
@@ -1143,7 +1143,7 @@ FNetRefHandle FObjectReferenceCache::GetDynamicRoot(const FNetRefHandle Handle)
 	}
 	else
 	{
-		return FNetRefHandle();
+		return FNetRefHandle::GetInvalid();
 	}
 }
 
@@ -1248,7 +1248,7 @@ void FObjectReferenceCache::WriteFullReferenceInternal(FNetSerializationContext&
 	{
 		UE_CLOG(RefHandle.IsValid(), LogIris, Log, TEXT("ObjectReferenceCache::WriteFullReference Trying to write Stale handle, %s"), *RefHandle.ToString());
 
-		WriteNetRefHandle(Context, FNetRefHandle());
+		WriteNetRefHandle(Context, FNetRefHandle::GetInvalid());
 		return;			
 	}
 
@@ -1297,7 +1297,7 @@ void FObjectReferenceCache::ReadFullReferenceInternal(FNetSerializationContext&
 		return;
 	}
 
-	UE_NET_TRACE_NAMED_OBJECT_SCOPE(ReferenceScope, FNetRefHandle(), *Context.GetBitStreamReader(), Context.GetTraceCollector(), ENetTraceVerbosity::Verbose);
+	UE_NET_TRACE_NAMED_OBJECT_SCOPE(ReferenceScope, FNetRefHandle::GetInvalid(), *Context.GetBitStreamReader(), Context.GetTraceCollector(), ENetTraceVerbosity::Verbose);
 
 	const FNetRefHandle NetRefHandle = ReadNetRefHandle(Context);
 
@@ -1442,7 +1442,7 @@ void FObjectReferenceCache::ReadFullReference(FNetSerializationContext& Context,
 
 	FNetObjectReference ObjectRef;
 
-	UE_NET_TRACE_NAMED_OBJECT_SCOPE(ReferenceScope, FNetRefHandle(), *Reader, Context.GetTraceCollector(), ENetTraceVerbosity::Verbose);
+	UE_NET_TRACE_NAMED_OBJECT_SCOPE(ReferenceScope, FNetRefHandle::GetInvalid(), *Reader, Context.GetTraceCollector(), ENetTraceVerbosity::Verbose);
 	ReadFullReferenceInternal(Context, ObjectRef, 0U);
 
 	UE_NET_TRACE_SET_SCOPE_OBJECTID(ReferenceScope, ObjectRef.GetRefHandle());
@@ -1508,7 +1508,7 @@ void FObjectReferenceCache::WriteReference(FNetSerializationContext& Context, FN
 	{
 		UE_CLOG(RefHandle.IsValid(), LogIris, Verbose, TEXT("ObjectReferenceCache::WriteReference Trying to write Stale handle, %s"), *RefHandle.ToString());
 
-		WriteNetRefHandle(Context, FNetRefHandle());
+		WriteNetRefHandle(Context, FNetRefHandle::GetInvalid());
 
 		return;
 	}
@@ -1524,7 +1524,7 @@ void FObjectReferenceCache::WriteReference(FNetSerializationContext& Context, FN
 void FObjectReferenceCache::ReadReference(FNetSerializationContext& Context, FNetObjectReference& OutRef)
 {
 	UE_NET_TRACE_SCOPE(NetObjectReference, *Context.GetBitStreamReader(), Context.GetTraceCollector(), ENetTraceVerbosity::Trace);
-	UE_NET_TRACE_NAMED_OBJECT_SCOPE(ReferenceScope, FNetRefHandle(), *Context.GetBitStreamReader(), Context.GetTraceCollector(), ENetTraceVerbosity::Verbose);
+	UE_NET_TRACE_NAMED_OBJECT_SCOPE(ReferenceScope, FNetRefHandle::GetInvalid(), *Context.GetBitStreamReader(), Context.GetTraceCollector(), ENetTraceVerbosity::Verbose);
 
 	FNetBitStreamReader* Reader = Context.GetBitStreamReader();
 
@@ -1729,7 +1729,7 @@ void FObjectReferenceCache::ReadMustBeMappedExports(FNetSerializationContext& Co
 
 	while (bHasExportsToRead && !Context.HasErrorOrOverflow())
 	{
-		UE_NET_TRACE_NAMED_OBJECT_SCOPE(ReferenceScope, FNetRefHandle(), *Context.GetBitStreamReader(), Context.GetTraceCollector(), ENetTraceVerbosity::Verbose);
+		UE_NET_TRACE_NAMED_OBJECT_SCOPE(ReferenceScope, FNetRefHandle::GetInvalid(), *Context.GetBitStreamReader(), Context.GetTraceCollector(), ENetTraceVerbosity::Verbose);
 		const FNetRefHandle MustBeMappedHandle = ReadNetRefHandle(Context);
 		UE_NET_TRACE_SET_SCOPE_OBJECTID(ReferenceScope, MustBeMappedHandle);	
 
diff --git a/Engine/Source/Runtime/Experimental/Iris/Core/Private/Iris/ReplicationSystem/ObjectReplicationBridge.cpp b/Engine/Source/Runtime/Experimental/Iris/Core/Private/Iris/ReplicationSystem/ObjectReplicationBridge.cpp
--- a/Engine/Source/Runtime/Experimental/Iris/Core/Private/Iris/ReplicationSystem/ObjectReplicationBridge.cpp
+++ b/Engine/Source/Runtime/Experimental/Iris/Core/Private/Iris/ReplicationSystem/ObjectReplicationBridge.cpp
@@ -205,7 +205,7 @@ UE::Net::FNetRefHandle UObjectReplicationBridge::GetReplicatedRefHandle(const UO
 {
 	FNetRefHandle Handle = GetObjectReferenceCache()->GetObjectReferenceHandleFromObject(Object);
 
-	return IsReplicatedHandle(Handle) ? Handle : FNetRefHandle();
+	return IsReplicatedHandle(Handle) ? Handle : FNetRefHandle::GetInvalid();
 }
 
 UE::Net::FNetRefHandle UObjectReplicationBridge::GetReplicatedRefHandle(FNetHandle Handle) const
@@ -217,7 +217,7 @@ UE::Net::FNetRefHandle UObjectReplicationBridge::GetReplicatedRefHandle(FNetHand
 	FInternalNetRefIndex ObjectInternalIndex = NetRefHandleManager->GetInternalIndexFromNetHandle(Handle);
 	if (ObjectInternalIndex == FNetRefHandleManager::InvalidInternalIndex)
 	{
-		return FNetRefHandle();
+		return FNetRefHandle::GetInvalid();
 	}
 
 	const FNetRefHandleManager::FReplicatedObjectData& ObjectData = NetRefHandleManager->GetReplicatedObjectDataNoCheck(ObjectInternalIndex);
@@ -234,7 +234,7 @@ UE::Net::FNetRefHandle UObjectReplicationBridge::BeginReplication(UObject* Insta
 	if (bBlockBeginReplication)
 	{
 		ensureMsgf(false, TEXT("BeginReplication is not allowed during this operation. %s will not be replicated"), *GetNameSafe(Instance));
-		return FNetRefHandle();
+		return FNetRefHandle::GetInvalid();
 	}
 
 	FNetRefHandle AllocatedRefHandle = ObjectReferenceCache->CreateObjectReferenceHandle(Instance);
@@ -242,7 +242,7 @@ UE::Net::FNetRefHandle UObjectReplicationBridge::BeginReplication(UObject* Insta
 	// If we failed to assign a handle, or if the Handle already is replicating, just return the handle
 	if (!AllocatedRefHandle.IsValid())
 	{
-		return FNetRefHandle();
+		return FNetRefHandle::GetInvalid();
 	}
 
 	if (IsReplicatedHandle(AllocatedRefHandle))
@@ -306,7 +306,7 @@ UE::Net::FNetRefHandle UObjectReplicationBridge::BeginReplication(UObject* Insta
 		if (!bIsValidProtocol)
 		{
 			UE_LOG_OBJECTREPLICATIONBRIDGE(Error, TEXT("BeginReplication Found invalid protocol ProtocolId:0x%" UINT64_x_FMT " for Object named %s"), ReplicationProtocol->ProtocolIdentifier, *Instance->GetName());
-			return FNetRefHandle();
+			return FNetRefHandle::GetInvalid();
 		}
 	}
 #endif
@@ -398,7 +398,7 @@ UE::Net::FNetRefHandle UObjectReplicationBridge::BeginReplication(UObject* Insta
 	// If we get here, it means that we failed to assign an internal handle for the object. We've probably run out of handles which currently is a fatal error.
 	UE_LOG(LogIris, Error, TEXT("UObjectReplicationBridge::BeginReplication - Failed to create NetRefHandle for object %s"), ToCStr(Instance->GetPathName()));
 
-	return FNetRefHandle();
+	return FNetRefHandle::GetInvalid();
 }
 
 void UObjectReplicationBridge::AssignDynamicFilter(UObject* Instance, const FCreateNetRefHandleParams& Params, FNetRefHandle RefHandle)
@@ -485,7 +485,7 @@ UE::Net::FNetRefHandle UObjectReplicationBridge::BeginReplication(FNetRefHandle
 		return SubObjectRefHandle;
 	}
 
-	return FNetRefHandle();
+	return FNetRefHandle::GetInvalid();
 }
 
 void UObjectReplicationBridge::SetSubObjectNetCondition(FNetRefHandle SubObjectRefHandle, ELifetimeCondition Condition)
diff --git a/Engine/Source/Runtime/Experimental/Iris/Core/Private/Iris/ReplicationSystem/ReplicationReader.cpp b/Engine/Source/Runtime/Experimental/Iris/Core/Private/Iris/ReplicationSystem/ReplicationReader.cpp
--- a/Engine/Source/Runtime/Experimental/Iris/Core/Private/Iris/ReplicationSystem/ReplicationReader.cpp
+++ b/Engine/Source/Runtime/Experimental/Iris/Core/Private/Iris/ReplicationSystem/ReplicationReader.cpp
@@ -263,7 +263,7 @@ void FReplicationReader::Deinit()
 // Read incomplete handle
 FNetRefHandle FReplicationReader::ReadNetRefHandleId(FNetSerializationContext& Context, FNetBitStreamReader& Reader) const
 {
-	UE_NET_TRACE_NAMED_OBJECT_SCOPE(ReferenceScope, FNetRefHandle(), *Context.GetBitStreamReader(), Context.GetTraceCollector(), ENetTraceVerbosity::Verbose);
+	UE_NET_TRACE_NAMED_OBJECT_SCOPE(ReferenceScope, FNetRefHandle::GetInvalid(), *Context.GetBitStreamReader(), Context.GetTraceCollector(), ENetTraceVerbosity::Verbose);
 
 	const uint64 NetId = ReadPackedUint64(&Reader);
 	FNetRefHandle RefHandle = FNetRefHandleManager::MakeNetRefHandleFromId(NetId);
@@ -273,7 +273,7 @@ FNetRefHandle FReplicationReader::ReadNetRefHandleId(FNetSerializationContext& C
 	if (RefHandle.GetId() != NetId)
 	{
 		Context.SetError(GNetError_InvalidNetHandle);
-		return FNetRefHandle();
+		return FNetRefHandle::GetInvalid();
 	}
 
 	return RefHandle;
@@ -294,7 +294,7 @@ uint32 FReplicationReader::ReadObjectsPendingDestroy(FNetSerializationContext& C
 	
 		for (uint32 It = 0; It < ObjectsToRead; ++It)
 		{
-			UE_NET_TRACE_NAMED_OBJECT_SCOPE(DestroyedObjectScope, FNetRefHandle(), Reader, Context.GetTraceCollector(), ENetTraceVerbosity::Trace);
+			UE_NET_TRACE_NAMED_OBJECT_SCOPE(DestroyedObjectScope, FNetRefHandle::GetInvalid(), Reader, Context.GetTraceCollector(), ENetTraceVerbosity::Trace);
 
 			FNetRefHandle IncompleteHandle = ReadNetRefHandleId(Context, Reader);
 			FNetRefHandle SubObjectRootOrHandle = IncompleteHandle;
diff --git a/Engine/Source/Runtime/Experimental/Iris/Core/Public/Iris/ReplicationSystem/NetRefHandle.h b/Engine/Source/Runtime/Experimental/Iris/Core/Public/Iris/ReplicationSystem/NetRefHandle.h
--- a/Engine/Source/Runtime/Experimental/Iris/Core/Public/Iris/ReplicationSystem/NetRefHandle.h
+++ b/Engine/Source/Runtime/Experimental/Iris/Core/Public/Iris/ReplicationSystem/NetRefHandle.h
@@ -24,18 +24,23 @@ namespace UE::Net
 class FNetRefHandle
 {
 public:
-	enum { Invalid = 0 };
+	
+	inline static FNetRefHandle GetInvalid() { return FNetRefHandle(); }
+
+private:
+	enum { InvalidValue = 0 };
 	enum { IdBits = 60 };
 	enum { ReplicationSystemIdBits = 4 };
 
 public:
-	FNetRefHandle() : Value(Invalid) {}
+
+	FNetRefHandle() : Value(InvalidValue) {}
 
 	uint64 GetId() const { return Id; }
 	uint32 GetReplicationSystemId() const { check(ReplicationSystemId != 0); return (uint32)(ReplicationSystemId - 1); }
-	bool IsValid() const { return Value != Invalid; }
+	bool IsValid() const { return Value != InvalidValue; }
 
-	bool IsCompleteHandle() const { return Value != Invalid && ReplicationSystemId != 0U; }
+	bool IsCompleteHandle() const { return Value != InvalidValue && ReplicationSystemId != 0U; }
 
 	bool IsStatic() const { return Id & StaticIdMask; }
 	bool IsDynamic() const { return IsValid() && !IsStatic(); }
diff --git a/Engine/Source/Runtime/Experimental/Iris/Core/Public/Iris/ReplicationSystem/ObjectReplicationBridge.h b/Engine/Source/Runtime/Experimental/Iris/Core/Public/Iris/ReplicationSystem/ObjectReplicationBridge.h
--- a/Engine/Source/Runtime/Experimental/Iris/Core/Public/Iris/ReplicationSystem/ObjectReplicationBridge.h
+++ b/Engine/Source/Runtime/Experimental/Iris/Core/Public/Iris/ReplicationSystem/ObjectReplicationBridge.h
@@ -446,5 +446,5 @@ protected:
 
 inline UE::Net::FNetRefHandle UObjectReplicationBridge::BeginReplication(UE::Net::FNetRefHandle OwnerHandle, UObject* Instance, const FCreateNetRefHandleParams& Params)
 {
-	return BeginReplication(OwnerHandle, Instance, FNetRefHandle(), ESubObjectInsertionOrder::None, Params);
+	return BeginReplication(OwnerHandle, Instance, FNetRefHandle::GetInvalid(), ESubObjectInsertionOrder::None, Params);
 }
