Iris
* Added print to log root object of subobject when protocol mismatch occurs

#rnx

diff --git a/Engine/Source/Runtime/Experimental/Iris/Core/Private/Iris/ReplicationSystem/ObjectReplicationBridge.cpp b/Engine/Source/Runtime/Experimental/Iris/Core/Private/Iris/ReplicationSystem/ObjectReplicationBridge.cpp
--- a/Engine/Source/Runtime/Experimental/Iris/Core/Private/Iris/ReplicationSystem/ObjectReplicationBridge.cpp
+++ b/Engine/Source/Runtime/Experimental/Iris/Core/Private/Iris/ReplicationSystem/ObjectReplicationBridge.cpp
@@ -1946,7 +1946,7 @@ void UObjectReplicationBridge::OnProtocolMismatchReported(FNetRefHandle RefHandl
 	UObject* ObjInstance = NetRefHandleManager->GetReplicatedObjectInstance(ObjectInternalIndex);
 	UObject* ObjArchetype = ObjInstance ? ObjInstance->GetArchetype() : nullptr;
 
-	UE_LOG(LogIris, Error, TEXT("OnProtocolMismatchReported from client:%u when instancing %s. CDO:%s ReplicatedObject:%s"), ConnectionId, *RefHandle.ToString(), *GetNameSafe(ObjArchetype), *GetNameSafe(ObjInstance));
+	UE_LOG(LogIris, Error, TEXT("OnProtocolMismatchReported from client:%u when instancing %s. CDO:%s ReplicatedObject:%s NetObject:%s"), ConnectionId, *RefHandle.ToString(), *GetNameSafe(ObjArchetype), *GetNameSafe(ObjInstance), *NetRefHandleManager->PrintObjectFromIndex(ObjectInternalIndex));
 
 	if (UE_LOG_ACTIVE(LogIris, Error))
 	{
