[NetPhysics] Resim improvements -

- Fix bug with rewind target buffer where the Insert logic broke if the inserted data would be at index 0 of the buffer.
- Changed so that rewind target buffer is not expecting a continuous buffer with entries for each frame, instead it will expect there to be holes between target frames.
- If allowing rewind to resimulated frames, block resim from rewind frame so we don't rewind to same targets again.
- Interpolate incoming targets from the server to fill gaps in the rewind target buffer
- Default to not allowing rewind to IncompleteHistory since we now use target interpolation prediction.
- Set high quantization precision on any actor using NetworkPhysicsComponent, not only pawns.

diff --git a/Engine/Source/Runtime/Engine/Private/PhysicsEngine/NetworkPhysicsComponent.cpp b/Engine/Source/Runtime/Engine/Private/PhysicsEngine/NetworkPhysicsComponent.cpp
--- a/Engine/Source/Runtime/Engine/Private/PhysicsEngine/NetworkPhysicsComponent.cpp
+++ b/Engine/Source/Runtime/Engine/Private/PhysicsEngine/NetworkPhysicsComponent.cpp
@@ -406,13 +406,10 @@ void UNetworkPhysicsComponent::InitPhysics()
 			bCompareInputToTriggerRewind = PhysicsSettings->ResimulationSettings.GetCompareInputToTriggerRewind(PhysicsReplicationCVars::ResimulationCVars::bCompareInputToTriggerRewind);
 		}
 
-		if (APawn* Pawn = Cast<APawn>(Owner))
-		{
-			FRepMovement& RepMovement = Pawn->GetReplicatedMovement_Mutable();
-			RepMovement.LocationQuantizationLevel = EVectorQuantization::RoundTwoDecimals;
-			RepMovement.RotationQuantizationLevel = ERotatorQuantization::ShortComponents;
-			RepMovement.VelocityQuantizationLevel = EVectorQuantization::RoundTwoDecimals;
-		}
+		FRepMovement& RepMovement = Owner->GetReplicatedMovement_Mutable();
+		RepMovement.LocationQuantizationLevel = EVectorQuantization::RoundTwoDecimals;
+		RepMovement.RotationQuantizationLevel = ERotatorQuantization::ShortComponents;
+		RepMovement.VelocityQuantizationLevel = EVectorQuantization::RoundTwoDecimals;
 
 		if (UPrimitiveComponent* RootPrimComp = Cast<UPrimitiveComponent>(Owner->GetRootComponent()))
 		{
diff --git a/Engine/Source/Runtime/Experimental/Chaos/Private/RewindData.cpp b/Engine/Source/Runtime/Experimental/Chaos/Private/RewindData.cpp
--- a/Engine/Source/Runtime/Experimental/Chaos/Private/RewindData.cpp
+++ b/Engine/Source/Runtime/Experimental/Chaos/Private/RewindData.cpp
@@ -204,10 +204,14 @@ void FRewindData::ApplyTargets(const int32 Frame, const bool bResetSimulation)
 	for (FDirtyParticleInfo& DirtyParticleInfo : DirtyParticles)
 	{
 		FGeometryParticleHandle* PTParticle = DirtyParticleInfo.GetObjectPtr();
-		FGeometryParticleStateBase& History = DirtyParticleInfo.GetHistory();	
+		FGeometryParticleStateBase& History = DirtyParticleInfo.GetHistory();
 
 		const bool bResimAsFollower = DirtyParticleInfo.bResimAsFollower;
-		RewindHelper(PTParticle, bResimAsFollower, History.TargetPositions, [](auto Particle, const auto& Data) {Particle->SetXR(Data); });
+
+		RewindHelper(PTParticle, bResimAsFollower, History.TargetPositions, [](auto Particle, const auto& Data) 
+		{
+			Particle->SetXR(Data); 
+		});
 		RewindHelper(PTParticle->CastToKinematicParticle(), bResimAsFollower, History.TargetVelocities, [](auto Particle, const auto& Data)
 		{
 			Particle->SetV(Data.V());
@@ -292,10 +296,7 @@ bool FRewindData::RewindToFrame(int32 Frame)
 	FFrameAndPhase RewindFrameAndPhase{ Frame, FFrameAndPhase::PostPushData };
 	FFrameAndPhase CurFrameAndPhase{ CurFrame, FFrameAndPhase::PrePushData };
 
-	if (!bResimAllowRewindToResimulatedFrames)
-	{
-		BlockResimFrame = CurFrame;
-	}
+	BlockResimFrame = bResimAllowRewindToResimulatedFrames ? Frame : CurFrame;
 
 	auto RewindHelper = [RewindFrameAndPhase, CurFrameAndPhase, this](auto Obj, bool bResimAsFollower, auto& Property, const auto& RewindFunc) -> bool
 	{
@@ -809,7 +810,7 @@ void FRewindData::MarkDirtyFromPT(FGeometryParticleHandle& Handle)
 
 	const FFrameAndPhase FrameAndPhase{ CurFrame, FFrameAndPhase::PostPushData };
 
-	if(bRecordingHistory || Latest.ParticlePositionRotation.IsClean(FrameAndPhase))
+	if (bRecordingHistory || Latest.ParticlePositionRotation.IsClean(FrameAndPhase))
 	{
 		if (auto Data = Latest.ParticlePositionRotation.WriteAccessNonDecreasing(FrameAndPhase, PropertiesPool))
 		{
@@ -883,7 +884,7 @@ enum class EResimFrameValidation : int32
 };
 CHAOS_API int32 ResimFrameValidation = (int32)EResimFrameValidation::IslandValidation;
 FAutoConsoleVariableRef CVarResimFrameValidationLeniency(TEXT("p.Resim.ResimFrameValidation"), ResimFrameValidation, TEXT("0 = no leniency, all dirty particles need a valid target. 1 = Island leniency, all particles in resim islands need a valid target. 2 = Full leniency, only the particle triggering the resim need a valid target."));
-CHAOS_API bool bResimIncompleteHistory = true;
+CHAOS_API bool bResimIncompleteHistory = false;
 FAutoConsoleVariableRef CVarResimIncompleteHistory(TEXT("p.Resim.IncompleteHistory"), bResimIncompleteHistory, TEXT("If a valid resim frame can't be found, use the requested resim frame and perform a resimulation with incomplete data."));
 
 int32 FRewindData::FindValidResimFrame(const int32 RequestedFrame)
@@ -987,7 +988,6 @@ int32 FRewindData::FindValidResimFrame(const int32 RequestedFrame)
 					bHasTargetHistory = false;
 					break;
 				}
-				break;
 			}
 		}
 		
@@ -1030,14 +1030,13 @@ int32 FRewindData::FindValidResimFrame(const int32 RequestedFrame)
 		ValidFrame = bResimIncompleteHistory ? RequestedFrame : INDEX_NONE;
 
 #if DEBUG_REWIND_DATA
-		UE_LOG(LogTemp, Warning, TEXT("COMMON | PT | FindValidResimFrame | No valid resim frame found | RequestedFrame: %d | ValidFrame: %d | EarliestFrame: %d | HasTargetHistory: %d | EarliestHistoryFrame: %d | CurrentFrame: %d | FramesSaved: %d, ResimFrameValidation:%d"), RequestedFrame, ValidFrame, EarliestFrame, bHasTargetHistory, GetEarliestFrame_Internal(), CurrentFrame(), FramesSaved, ResimFrameValidation);
+		UE_LOG(LogTemp, Warning, TEXT("COMMON | PT | FindValidResimFrame | No valid resim frame found | RequestedFrame: %d | ValidFrame: %d | EarliestFrame: %d | HasTargetHistory: %d | EarliestHistoryFrame: %d | CurrentFrame: %d | FramesSaved: %d | ResimFrameValidation: %d"), RequestedFrame, ValidFrame, EarliestFrame, bHasTargetHistory, GetEarliestFrame_Internal(), CurrentFrame(), FramesSaved, ResimFrameValidation);
 #endif
 	}
 
 	return ValidFrame;
 }
 
-
 void FRewindData::PushStateAtFrame(FGeometryParticleHandle& Handle, int32 Frame, FFrameAndPhase::EParticleHistoryPhase Phase, 
 	const FVector& Position, const FQuat& Quaternion, const FVector& LinVelocity, const FVector& AngVelocity, const bool bShouldSleep)
 {
@@ -1069,7 +1068,7 @@ void FRewindData::PushStateAtFrame(FGeometryParticleHandle& Handle, int32 Frame,
 	}
 }
 
-void FRewindData::PushPTDirtyData(TPBDRigidParticleHandle<FReal,3>& Handle,const int32 SrcDataIdx)
+void FRewindData::PushPTDirtyData(TPBDRigidParticleHandle<FReal, 3>& Handle, const int32 SrcDataIdx)
 {
 	const bool bRecordingHistory = !IsResimAndInSync(Handle);
 
@@ -1077,7 +1076,7 @@ void FRewindData::PushPTDirtyData(TPBDRigidParticleHandle<FReal,3>& Handle,const
 	FGeometryParticleStateBase& Latest = Info.AddFrame(CurFrame);
 
 	const FFrameAndPhase FrameAndPhase{ CurFrame, FFrameAndPhase::PostCallbacks };
-	
+
 	if (bRecordingHistory || Latest.ParticlePositionRotation.IsClean(FrameAndPhase))
 	{
 		if (FParticlePositionRotation* PreXR = Latest.ParticlePositionRotation.WriteAccessNonDecreasing(FrameAndPhase, PropertiesPool))
@@ -1115,9 +1114,44 @@ FJointState FRewindData::GetPastJointStateAtFrame(const FPBDJointConstraintHandl
 	return GetPastStateAtFrameImp<FJointState>(DirtyJoints, Handle, Frame, Phase);
 }
 
+CHAOS_API int32 bInterpolateTargetGaps = 5;
+FAutoConsoleVariableRef CVarResimInterpolateTargetGaps(TEXT("p.Resim.InterpolateTargetGaps"), bInterpolateTargetGaps, TEXT("How many frame gaps in replicated targets we should fill by interpolating between the previous and the new target received. Value in max number of frames to interpolate, deactivate by setting to 0."));
+
 void FRewindData::SetTargetStateAtFrame(FGeometryParticleHandle& Handle, const int32 Frame, FFrameAndPhase::EParticleHistoryPhase Phase,
 	const FVector& Position, const FQuat& Quaternion, const FVector& LinVelocity, const FVector& AngVelocity, const bool bShouldSleep)
 {
+	if (bInterpolateTargetGaps)
+	{
+		FDirtyParticleInfo& Info = FindOrAddDirtyObj(Handle);
+		FGeometryParticleStateBase& Latest = Info.GetHistory();
+		FFrameAndPhase FrameAndPhase;
+
+		if (Latest.TargetPositions.GetHeadFrameAndPhase(FrameAndPhase))
+		{
+			const int32 FrameDiff = Frame - FrameAndPhase.Frame;
+			if (FrameDiff > 1 && FrameDiff <= bInterpolateTargetGaps)
+			{
+				const FParticlePositionRotation* TargetXR = Latest.TargetPositions.Read(FrameAndPhase, PropertiesPool);
+				const FParticleVelocities* TargetVW = Latest.TargetVelocities.Read(FrameAndPhase, PropertiesPool);
+				const FParticleDynamicMisc* TargetDynamic = Latest.TargetStates.Read(FrameAndPhase, PropertiesPool);
+				if (TargetXR && TargetVW && TargetDynamic)
+				{
+					for (int32 InterpFrame = 1; InterpFrame < FrameDiff; InterpFrame++)
+					{
+						const float Alpha = (1.0f / (float)FrameDiff) * (float)InterpFrame;
+
+						PushStateAtFrame(Handle, FrameAndPhase.Frame + InterpFrame, Phase,
+							FMath::Lerp(TargetXR->GetX(), Position, Alpha),
+							FRotation3::Slerp(TargetXR->GetR(), Quaternion, Alpha),
+							FMath::Lerp(TargetVW->GetV(), LinVelocity, Alpha),
+							FMath::Lerp(TargetVW->GetW(), AngVelocity, Alpha),
+							bShouldSleep && TargetDynamic->ObjectState() == EObjectStateType::Sleeping);
+					}
+				}
+			}
+		}
+	}
+
 	PushStateAtFrame(Handle, Frame, Phase,
 		Position, Quaternion, LinVelocity, AngVelocity, bShouldSleep);
 }
diff --git a/Engine/Source/Runtime/Experimental/Chaos/Public/RewindData.h b/Engine/Source/Runtime/Experimental/Chaos/Public/RewindData.h
--- a/Engine/Source/Runtime/Experimental/Chaos/Public/RewindData.h
+++ b/Engine/Source/Runtime/Experimental/Chaos/Public/RewindData.h
@@ -433,14 +433,26 @@ public:
 		return Idx != INDEX_NONE ? &GetPool(Manager).GetElement(Buffer[Idx].Ref) : nullptr;
 	}
 
+	//Get the FrameAndPhase of the head / last entry
+	const bool GetHeadFrameAndPhase(FFrameAndPhase& OutFrameAndPhase) const
+	{
+		if (NumValid)
+		{
+			const int32 Prev = Next == 0 ? Buffer.Num() - 1 : Next - 1;
+			OutFrameAndPhase = Buffer[Prev].FrameAndPhase;
+			return true;
+		}
+		return false;
+	}
+
 	//Releases data back into the pool
 	void Release(FDirtyPropertiesPool& Manager)
 	{
 		TPropertyPool<T>& Pool = GetPool(Manager);
 		for(FPropertyInterval& Interval : Buffer)
-	{
+		{
 			Pool.RemoveElement(Interval.Ref);
-	}
+		}
 
 		Buffer.Empty();
 		NumValid = 0;
@@ -459,11 +471,11 @@ public:
 	void ClearEntryAndFuture(const FFrameAndPhase FrameAndPhase)
 	{
 		//Move next backwards until FrameAndPhase and anything more future than it is gone
-		while(NumValid)
+		while (NumValid)
 		{
 			const int32 PotentialNext = Next - 1 >= 0 ? Next - 1 : Buffer.Num() - 1;
 
-			if(Buffer[PotentialNext].FrameAndPhase < FrameAndPhase)
+			if (Buffer[PotentialNext].FrameAndPhase < FrameAndPhase)
 			{
 				break;
 			}
@@ -508,14 +520,14 @@ public:
 		T* Result = nullptr;
 
 		int32 FrameIndex = FindIdx(FrameAndPhase);
-		if(FrameIndex != INDEX_NONE)
+		if (FrameIndex != INDEX_NONE)
 		{
 			Result = &GetPool(Manager).GetElement(Buffer[FrameIndex].Ref);
 		}
 		else
 		{
 			FPropertyIdx ElementRef;
-			if(Next >= Buffer.Num())
+			if (Next >= Buffer.Num())
 			{
 				GetPool(Manager).AddElement(ElementRef);
 				Buffer.Add({ ElementRef, FrameAndPhase });
@@ -545,6 +557,13 @@ public:
 				else
 				{
 					Buffer[NextFrame] = Buffer[PrevFrame];
+
+					if (Count == NumValid - 1)
+					{ 
+						// If we shift back and reach the end of the buffer, insert here
+						Buffer[PrevFrame].FrameAndPhase = FrameAndPhase;
+						Buffer[PrevFrame].Ref = ElementRef;
+					}
 				}
 			}
 
@@ -562,24 +581,25 @@ private:
 	{
 		int32 Cur = Next;	//go in reverse order because hopefully we don't rewind too far back
 		int32 Result = INDEX_NONE;
-		for(int32 Count = 0; Count < NumValid; ++Count)
+		for (int32 Count = 0; Count < NumValid; ++Count)
 		{
 			--Cur;
 			if (Cur < 0) { Cur = Buffer.Num() - 1; }
 
 			const FPropertyInterval& Interval = Buffer[Cur];
-			if(Interval.FrameAndPhase < FrameAndPhase)
+			
+			if (Interval.FrameAndPhase < FrameAndPhase)
 			{
 				//no reason to keep searching, frame is bigger than everything before this
 				break;
 			}
 			else
-	{
+			{
 				Result = Cur;
 			}
 		}
 
-		if(bNoEntryIsHead || Result == INDEX_NONE)
+		if (bNoEntryIsHead || Result == INDEX_NONE)
 		{
 			//in this mode we consider the entire interval as one entry
 			return Result;
@@ -653,17 +673,17 @@ private:
 #if VALIDATE_REWIND_DATA
 		int32 Val = Next;
 		FFrameAndPhase PrevVal;
-		for(int32 Count = 0; Count < NumValid; ++Count)
-	{
+		for (int32 Count = 0; Count < NumValid; ++Count)
+		{
 			--Val;
 			if (Val < 0) { Val = Buffer.Num() - 1; }
 			if (Count == 0)
-		{
+			{
 				PrevVal = Buffer[Val].FrameAndPhase;
-		}
-		else
-		{
-				ensure(Buffer[Val].FrameAndPhase < PrevVal);
+			}
+			else
+			{
+				ensureMsgf(Buffer[Val].FrameAndPhase < PrevVal, TEXT("ValidateOrder Idx: %d TailFrame: %d/%d, HeadFrame: %d/%d"), Val, Buffer[Val].FrameAndPhase.Frame, Buffer[Val].FrameAndPhase.Phase, PrevVal.Frame, PrevVal.Phase);
 				PrevVal = Buffer[Val].FrameAndPhase;
 			}
 		}
@@ -1014,9 +1034,9 @@ struct FGeometryParticleStateBase
 	TParticlePropertyBuffer<FParticleMassProps,EChaosProperty::MassProps> MassProps;
 	TParticlePropertyBuffer<FKinematicTarget, EChaosProperty::KinematicTarget> KinematicTarget;
 
-	TParticlePropertyBuffer<FParticlePositionRotation, EChaosProperty::XR> TargetPositions;
-	TParticlePropertyBuffer<FParticleVelocities, EChaosProperty::Velocities> TargetVelocities;
-	TParticlePropertyBuffer<FParticleDynamicMisc, EChaosProperty::DynamicMisc> TargetStates;
+	TParticlePropertyBuffer<FParticlePositionRotation, EChaosProperty::XR, /*bNoEntryIsHead=*/false> TargetPositions;
+	TParticlePropertyBuffer<FParticleVelocities, EChaosProperty::Velocities, /*bNoEntryIsHead=*/false> TargetVelocities;
+	TParticlePropertyBuffer<FParticleDynamicMisc, EChaosProperty::DynamicMisc, /*bNoEntryIsHead=*/false> TargetStates;
 
 	FShapesArrayStateBase ShapesArrayState;
 
