Chaos - fix ProjectOntoAxis(Capsule)

This could cause end-cap collisions to be lost

diff --git a/Engine/Source/Runtime/Experimental/Chaos/Public/Chaos/Collision/ConvexContactPointUtilities.h b/Engine/Source/Runtime/Experimental/Chaos/Public/Chaos/Collision/ConvexContactPointUtilities.h
--- a/Engine/Source/Runtime/Experimental/Chaos/Public/Chaos/Collision/ConvexContactPointUtilities.h
+++ b/Engine/Source/Runtime/Experimental/Chaos/Public/Chaos/Collision/ConvexContactPointUtilities.h
@@ -69,7 +69,7 @@ namespace Chaos
 		else
 		{
 			MinVertexIndex = 1;
-			MaxVertexIndex = 01;
+			MaxVertexIndex = 0;
 			PMin = D1;
 			PMax = D0;
 		}
