Chaos
- shape bounds are now tight-fitting (particle bounds are still padded by cull distance and some fraction of movement)
- debugdraw pre/post integrate shapes
- improve evolution test mode (for reproducible visual collision detection tests)
- move integrate to cpp file to reduce compilation

diff --git a/Engine/Source/Runtime/Experimental/Chaos/Private/Chaos/ChaosDebugDraw.cpp b/Engine/Source/Runtime/Experimental/Chaos/Private/Chaos/ChaosDebugDraw.cpp
--- a/Engine/Source/Runtime/Experimental/Chaos/Private/Chaos/ChaosDebugDraw.cpp
+++ b/Engine/Source/Runtime/Experimental/Chaos/Private/Chaos/ChaosDebugDraw.cpp
@@ -237,6 +237,51 @@ namespace Chaos
 			return ChaosDefaultBoundsColorsByState;
 		}
 
+		CHAOS_API const FChaosDebugDrawColorsByState& GetDefaultShapesColorsPreIntegrate()
+		{
+			// default colors by state for shapes
+			static FChaosDebugDrawColorsByState ChaosDefaultShapesColors
+			{
+				/* InDynamicColor =	  */ FColor(64, 64, 0),
+				/* InSleepingColor =  */ FColor(0, 0, 0),
+				/* InKinematicColor = */ FColor(0, 32, 64),
+				/* InStaticColor =	  */ FColor(64, 0, 0),
+				/* InDebrisColor =	  */ FColor(0, 64, 0),
+			};
+
+			return ChaosDefaultShapesColors;
+		}
+
+		const FChaosDebugDrawColorsByState& GetDefaultShapesColorsPostIntegrate()
+		{
+			// default colors by state for shapes
+			static FChaosDebugDrawColorsByState ChaosDefaultShapesColors
+			{
+				/* InDynamicColor =	  */ FColor(255, 255, 0),
+				/* InSleepingColor =  */ FColor(128, 128, 128),
+				/* InKinematicColor = */ FColor(0, 128, 255),
+				/* InStaticColor =	  */ FColor(255, 0, 0),
+				/* InDebrisColor =	  */ FColor(0, 255, 0),
+			};
+
+			return ChaosDefaultShapesColors;
+		}
+
+		const FChaosDebugDrawColorsByState& GetDefaultShapesColorsCollisionDetection()
+		{
+			// default colors by state for shapes
+			static FChaosDebugDrawColorsByState ChaosDefaultShapesColors
+			{
+				/* InDynamicColor =	  */ FColor(0, 255, 0),
+				/* InSleepingColor =  */ FColor(0, 255, 0),
+				/* InKinematicColor = */ FColor(0, 0, 0),
+				/* InStaticColor =	  */ FColor(0, 0, 0),
+				/* InDebrisColor =	  */ FColor(0, 255, 0),
+			};
+
+			return ChaosDefaultShapesColors;
+		}
+
 		FColor GetIndexColor(const int32 Index)
 		{
 			static const FColor Colors[] =
@@ -927,13 +972,11 @@ namespace Chaos
 			return false;
 		}
 
-		void DrawParticleShapesImpl(const FRigidTransform3& SpaceTransform, const FGeometryParticleHandle* Particle, const FColor& InColor, const FChaosDebugDrawSettings& Settings)
+		void DrawParticleShapesImpl(const FRigidTransform3& SpaceTransform, const FGeometryParticleHandle* Particle, const FRigidTransform3& ParticleTransform, const FColor& InColor, const FChaosDebugDrawSettings& Settings)
 		{
-			const FVec3 P = SpaceTransform.TransformPosition(Particle->ObjectState() == EObjectStateType::Dynamic ? Particle->CastToRigidParticle()->GetP() : Particle->GetX());
-			const FRotation3 Q = SpaceTransform.GetRotation() * (Particle->ObjectState() == EObjectStateType::Dynamic ? Particle->CastToRigidParticle()->GetQ() : Particle->GetR());
-			const FRigidTransform3 ParticleSpaceTransform = FRigidTransform3(P, Q);
+			const FRigidTransform3 ParticleSpaceTransform = ParticleTransform * SpaceTransform;
 
-			if(!DrawConvexOptimizerShapes(P, Q, Particle, InColor, Settings))
+			if(!DrawConvexOptimizerShapes(ParticleTransform.GetTranslation(), ParticleTransform.GetRotation(), Particle, InColor, Settings))
 			{
 				for (const FShapeInstancePtr& ShapeInstance : Particle->ShapeInstances())
 				{
@@ -942,11 +985,6 @@ namespace Chaos
 			}
 		}
 
-		void DrawParticleShapesImpl(const FRigidTransform3& SpaceTransform, const FGeometryParticle* Particle, const FColor& InColor, const FChaosDebugDrawSettings& Settings)
-		{
-			DrawParticleShapesImpl(SpaceTransform, Particle->Handle(), InColor, Settings);
-		}
-
 		void DrawBVHImpl(const FGeometryParticleHandle* Particle, const FRigidTransform3& ShapeTransform, const Private::FImplicitBVH* BVH, const FColor& UnusedColor, const FRealSingle Duration, const FChaosDebugDrawSettings& Settings)
 		{
 			const FReal BoundsDrawShrink = 1.0f;	// 1cm
@@ -1052,7 +1090,7 @@ namespace Chaos
 				const bool bIsComplex = (ShapeType == ImplicitObjectType::TriangleMesh) || (ShapeType == ImplicitObjectType::HeightField);
 				if (!bIsComplex)
 				{
-					const FAABB3 ShapeBox = Shape->GetWorldSpaceInflatedShapeBounds();
+					const FAABB3 ShapeBox = Shape->GetWorldSpaceShapeBounds();
 					const FVec3 ShapeP = SpaceTransform.TransformPosition(ShapeBox.GetCenter());
 					const FRotation3 ShapeQ = SpaceTransform.GetRotation();
 					const FColor ShapeColor = (bChaosDebugDebugDrawColorBoundsByShapeType) ? Settings.BoundsColorsPerShapeType.GetColorFromShapeType(ShapeType) : Color;
@@ -1341,8 +1379,8 @@ namespace Chaos
 
 				const FColor C3 = FColor::White;
 				const FMatrix Axes = FRotationMatrix::MakeFromX(Normal);
-				const FVec3 P0 = SpaceTransform.TransformPosition(Contact.GetShape0()->GetWorldSpaceInflatedShapeBounds().GetCenter());
-				const FVec3 P1 = SpaceTransform.TransformPosition(Contact.GetShape1()->GetWorldSpaceInflatedShapeBounds().GetCenter());
+				const FVec3 P0 = SpaceTransform.TransformPosition(Contact.GetShape0()->GetWorldSpaceShapeBounds().GetCenter());
+				const FVec3 P1 = SpaceTransform.TransformPosition(Contact.GetShape1()->GetWorldSpaceShapeBounds().GetCenter());
 				FDebugDrawQueue::GetInstance().DrawDebugLine(Location, P0, C3, false, UE_KINDA_SMALL_NUMBER, uint8(Settings.DrawPriority), Settings.LineThickness * 0.5f);
 				FDebugDrawQueue::GetInstance().DrawDebugLine(Location, P1, C3, false, UE_KINDA_SMALL_NUMBER, uint8(Settings.DrawPriority), Settings.LineThickness * 0.5f);
 			}
@@ -2007,7 +2045,7 @@ namespace Chaos
 				for (auto& Particle : ParticlesView)
 				{				
 					const FColor Color = GetParticleShapeDefaultColor(Particle.Handle(), DebugDrawSettings);
-					DrawParticleShapesImpl(SpaceTransform, GetHandleHelper(&Particle), Color, DebugDrawSettings);
+					DrawParticleShapesImpl(SpaceTransform, GetHandleHelper(&Particle), FConstGenericParticleHandle(Particle.Handle())->GetTransformPQ(), Color, DebugDrawSettings);
 				}
 			}
 		}
@@ -2020,7 +2058,7 @@ namespace Chaos
 				for (auto& Particle : ParticlesView)
 				{
 					const FColor Color = GetParticleShapeDefaultColor(Particle.Handle(), DebugDrawSettings);
-					DrawParticleShapesImpl(SpaceTransform, GetHandleHelper(&Particle), Color, DebugDrawSettings);
+					DrawParticleShapesImpl(SpaceTransform, GetHandleHelper(&Particle), FConstGenericParticleHandle(Particle.Handle())->GetTransformPQ(), Color, DebugDrawSettings);
 				}
 			}
 		}
@@ -2033,16 +2071,24 @@ namespace Chaos
 				for (auto& Particle : ParticlesView)
 				{
 					const FColor Color = GetParticleShapeDefaultColor(Particle.Handle(), DebugDrawSettings);
-					DrawParticleShapesImpl(SpaceTransform, GetHandleHelper(&Particle), Color, DebugDrawSettings);
+					DrawParticleShapesImpl(SpaceTransform, GetHandleHelper(&Particle), FConstGenericParticleHandle(Particle.Handle())->GetTransformPQ(), Color, DebugDrawSettings);
 				}
 			}
 		}
 
+		void DrawParticleShapes(const FRigidTransform3& SpaceTransform, const FGeometryParticleHandle* Particle, const FRigidTransform3& ParticleTransform, const FColor& Color, const FChaosDebugDrawSettings* Settings)
+		{
+			if (FDebugDrawQueue::IsDebugDrawingEnabled())
+			{
+				DrawParticleShapesImpl(SpaceTransform, Particle, ParticleTransform, Color, GetChaosDebugDrawSettings(Settings));
+			}
+		}
+
 		void DrawParticleShapes(const FRigidTransform3& SpaceTransform, const FGeometryParticleHandle* Particle, const FColor& Color, const FChaosDebugDrawSettings* Settings)
 		{
 			if (FDebugDrawQueue::IsDebugDrawingEnabled())
 			{
-				DrawParticleShapesImpl(SpaceTransform, Particle, Color, GetChaosDebugDrawSettings(Settings));
+				DrawParticleShapesImpl(SpaceTransform, Particle, FConstGenericParticleHandle(Particle)->GetTransformPQ(), Color, GetChaosDebugDrawSettings(Settings));
 			}
 		}
 
@@ -2050,7 +2096,7 @@ namespace Chaos
 		{
 			if (FDebugDrawQueue::IsDebugDrawingEnabled())
 			{
-				DrawParticleShapesImpl(SpaceTransform, Particle, Color, GetChaosDebugDrawSettings(Settings));
+				DrawParticleShapesImpl(SpaceTransform, Particle->Handle(), FConstGenericParticleHandle(Particle->Handle())->GetTransformPQ(), Color, GetChaosDebugDrawSettings(Settings));
 			}
 		}
 
diff --git a/Engine/Source/Runtime/Experimental/Chaos/Private/Chaos/Collision/ParticlePairMidPhase.cpp b/Engine/Source/Runtime/Experimental/Chaos/Private/Chaos/Collision/ParticlePairMidPhase.cpp
--- a/Engine/Source/Runtime/Experimental/Chaos/Private/Chaos/Collision/ParticlePairMidPhase.cpp
+++ b/Engine/Source/Runtime/Experimental/Chaos/Private/Chaos/Collision/ParticlePairMidPhase.cpp
@@ -325,13 +325,15 @@ namespace Chaos
 	{
 		PHYSICS_CSV_SCOPED_EXPENSIVE(PhysicsVerbose, NarrowPhase_ShapeBounds);
 
-		const FAABB3& ShapeWorldBounds0 = Shape0->GetWorldSpaceInflatedShapeBounds();
-		const FAABB3& ShapeWorldBounds1 = Shape1->GetWorldSpaceInflatedShapeBounds();
+		const FAABB3& ShapeWorldBounds0 = Shape0->GetWorldSpaceShapeBounds();
+		const FAABB3& ShapeWorldBounds1 = Shape1->GetWorldSpaceShapeBounds();
 
 		// World-space expanded bounds check
 		if (BoundsTestFlags.bEnableAABBCheck)
 		{
-			if (!ShapeWorldBounds0.Intersects(ShapeWorldBounds1))
+			// @todo(chaos): ideally this is a swept bounds test with a smaller cull distance (see FParticlePairMidPhase::GenerateCollisions)
+			const FAABB3 ExpandedShapeWorldBounds0 = FAABB3(ShapeWorldBounds0).ThickenSymmetrically(FVec3(CullDistance));
+			if (!ExpandedShapeWorldBounds0.Intersects(ShapeWorldBounds1))
 			{
 				return false;
 			}
@@ -1413,10 +1415,10 @@ namespace Chaos
 		const FRigidTransform3 ParticleWorldTransformB = PB->GetTransformPQ();
 		const FShapeInstanceArray& ShapeInstancesA = ParticleA->ShapeInstances();
 
-		// ImplicitB transforms/bounds
+		// ImplicitB transforms/bounds (expanded by cull distance)
 		const FRigidTransform3 ImplicitTransformB = RelativeTransformB * ParticleWorldTransformB;
 		const FRigidTransform3 ImplicitTransformBToA = ImplicitTransformB.GetRelativeTransform(ParticleWorldTransformA);
-		const FAABB3 ImplicitBoundsBInA = ImplicitB->CalculateTransformedBounds(ImplicitTransformBToA);
+		const FAABB3 ImplicitBoundsBInA = ImplicitB->CalculateTransformedBounds(ImplicitTransformBToA).ThickenSymmetrically(FVec3(CullDistance));
 
 		// If ImplicitB has a built-in BVH (Heightfield or TriMesh) we handle the test against the BVH differently
 		const bool bHasInternalBVHB = ImplicitB->template IsA<FHeightField>();
diff --git a/Engine/Source/Runtime/Experimental/Chaos/Private/Chaos/PBDRigidsEvolution.cpp b/Engine/Source/Runtime/Experimental/Chaos/Private/Chaos/PBDRigidsEvolution.cpp
--- a/Engine/Source/Runtime/Experimental/Chaos/Private/Chaos/PBDRigidsEvolution.cpp
+++ b/Engine/Source/Runtime/Experimental/Chaos/Private/Chaos/PBDRigidsEvolution.cpp
@@ -26,7 +26,8 @@ namespace Chaos
 {
 	namespace CVars
 	{
-		extern bool bChaos_Solver_TestMode;
+		extern bool bChaos_Solver_TestMode_Enabled;
+		extern int32 Chaos_Solver_TestMode_Step;
 	}
 
 	CHAOS_API int32 FixBadAccelerationStructureRemoval = 1;
@@ -1195,6 +1196,15 @@ namespace Chaos
 	}
 
 #if CHAOS_EVOLUTION_COLLISION_TESTMODE
+	void FPBDRigidsEvolutionBase::TestModeStep()
+	{
+		if (CVars::Chaos_Solver_TestMode_Step > 0)
+		{
+			TestModeData.Reset();
+			--CVars::Chaos_Solver_TestMode_Step;
+		}
+	}
+
 	void FPBDRigidsEvolutionBase::TestModeParticleDisabled(FGeometryParticleHandle* Particle)
 	{
 		if (FPBDRigidParticleHandle* Rigid = Particle->CastToRigidParticle())
@@ -1205,7 +1215,7 @@ namespace Chaos
 
 	void FPBDRigidsEvolutionBase::TestModeSaveParticles()
 	{
-		if (!CVars::bChaos_Solver_TestMode)
+		if (!CVars::bChaos_Solver_TestMode_Enabled)
 		{
 			return;
 		}
@@ -1218,7 +1228,7 @@ namespace Chaos
 
 	void FPBDRigidsEvolutionBase::TestModeSaveParticle(FGeometryParticleHandle* Particle)
 	{
-		if (!CVars::bChaos_Solver_TestMode)
+		if (!CVars::bChaos_Solver_TestMode_Enabled)
 		{
 			return;
 		}
@@ -1228,20 +1238,37 @@ namespace Chaos
 			FTestModeParticleData* Data = TestModeData.Find(Rigid);
 			if (Data == nullptr)
 			{
-				Data = &TestModeData.Add(Rigid);
+				TestModeData.Add(Rigid);
+				TestModeUpdateSavedParticle(Particle);
+			}
+		}
+	}
+
+	void FPBDRigidsEvolutionBase::TestModeUpdateSavedParticle(FGeometryParticleHandle* Particle)
+	{
+		if (!CVars::bChaos_Solver_TestMode_Enabled)
+		{
+			return;
+		}
+
+		if (FPBDRigidParticleHandle* Rigid = Particle->CastToRigidParticle())
+		{
+			FTestModeParticleData* Data = TestModeData.Find(Rigid);
+			if (Data != nullptr)
+			{
+				Data->X = Rigid->GetX();
+				Data->P = Rigid->GetP();
+				Data->R = Rigid->GetR();
+				Data->Q = Rigid->GetQ();
+				Data->V = Rigid->GetV();
+				Data->W = Rigid->GetW();
 			}
-			Data->X = Rigid->GetX();
-			Data->P = Rigid->GetP();
-			Data->R = Rigid->GetR();
-			Data->Q = Rigid->GetQ();
-			Data->V = Rigid->GetV();
-			Data->W = Rigid->GetW();
 		}
 	}
 
 	void FPBDRigidsEvolutionBase::TestModeRestoreParticles()
 	{
-		if (!CVars::bChaos_Solver_TestMode)
+		if (!CVars::bChaos_Solver_TestMode_Enabled)
 		{
 			return;
 		}
@@ -1250,11 +1277,7 @@ namespace Chaos
 		{
 			// If this is the first time we have seen this particle, save its data, otherwise restore it from the cache
 			FTestModeParticleData* Data = TestModeData.Find(Rigid.Handle());
-			if (Data == nullptr)
-			{
-				TestModeSaveParticle(Rigid.Handle());
-			}
-			else
+			if (Data != nullptr)
 			{
 				TestModeRestoreParticle(Rigid.Handle());
 			}
@@ -1263,7 +1286,7 @@ namespace Chaos
 
 	void FPBDRigidsEvolutionBase::TestModeRestoreParticle(FGeometryParticleHandle* Particle)
 	{
-		if (!CVars::bChaos_Solver_TestMode)
+		if (!CVars::bChaos_Solver_TestMode_Enabled)
 		{
 			return;
 		}
diff --git a/Engine/Source/Runtime/Experimental/Chaos/Private/Chaos/PBDRigidsEvolutionGBF.cpp b/Engine/Source/Runtime/Experimental/Chaos/Private/Chaos/PBDRigidsEvolutionGBF.cpp
--- a/Engine/Source/Runtime/Experimental/Chaos/Private/Chaos/PBDRigidsEvolutionGBF.cpp
+++ b/Engine/Source/Runtime/Experimental/Chaos/Private/Chaos/PBDRigidsEvolutionGBF.cpp
@@ -26,7 +26,7 @@
 
 #include "ChaosVisualDebugger/ChaosVisualDebuggerTrace.h"
 
-//PRAGMA_DISABLE_OPTIMIZATION
+//UE_DISABLE_OPTIMIZATION
 
 namespace Chaos
 {
@@ -169,9 +169,15 @@ namespace Chaos
 		// This is used to test collision detection and - it will be removed
 		// NOTE: You should also set the following for dragging in PIE to work while test mode is active:
 		// 		p.DisableEditorPhysicsHandle 1
-		bool bChaos_Solver_TestMode  = false;
-		FAutoConsoleVariableRef CVarChaosSolverTestMode(TEXT("p.Chaos.Solver.TestMode"), bChaos_Solver_TestMode, TEXT(""));
+		bool bChaos_Solver_TestMode_Enabled  = false;
+		FAutoConsoleVariableRef CVarChaosSolverTestModeEnabled(TEXT("p.Chaos.Solver.TestMode.Enabled"), bChaos_Solver_TestMode_Enabled, TEXT(""));
 
+		bool bChaos_Solver_TestMode_ShowInitialTransforms = false;
+		FAutoConsoleVariableRef CVarChaosSolverTestModeShowInitialTransforms(TEXT("p.Chaos.Solver.TestMode.ShowInitialTransforms"), bChaos_Solver_TestMode_ShowInitialTransforms, TEXT(""));
+		
+		int32 Chaos_Solver_TestMode_Step  = 0;
+		FAutoConsoleVariableRef CVarChaosSolverTestModeStep(TEXT("p.Chaos.Solver.TestMode.Step"), Chaos_Solver_TestMode_Step, TEXT(""));
+		
 		// Set to true to enable some debug validation of the Particle Views every frame
 		bool bChaosSolverCheckParticleViews = false;
 		FAutoConsoleVariableRef CVarChaosSolverCheckParticleViews(TEXT("p.Chaos.Solver.CheckParticleViews"), bChaosSolverCheckParticleViews, TEXT(""));
@@ -442,31 +448,27 @@ void FPBDRigidsEvolutionGBF::AdvanceOneTimeStepImpl(const FReal Dt, const FSubSt
 	// Update the collision solver type (used to support runtime comparisons of solver types for debugging/testing)
 	UpdateCollisionSolverType();
 
-	if (PreIntegrateCallback != nullptr)
-	{
-		SCOPE_CYCLE_COUNTER(STAT_Evolution_PreIntegrateCallback);
-		PreIntegrateCallback();
-	}
-
 #if CHAOS_EVOLUTION_COLLISION_TESTMODE
 	{
+		TestModeStep();
 		TestModeSaveParticles();
+		TestModeRestoreParticles();
+		TestModeResetCollisions();
 	}
 #endif
 
+	if (PreIntegrateCallback != nullptr)
 	{
-		CVD_SCOPE_TRACE_SOLVER_STEP(CVDDC_Integrate, TEXT("Integrate"))
-		SCOPE_CYCLE_COUNTER(STAT_Evolution_Integrate);
-		CSV_SCOPED_TIMING_STAT(PhysicsVerbose, StepSolver_Integrate);
-		Integrate(Particles.GetActiveParticlesView(), Dt);
+		SCOPE_CYCLE_COUNTER(STAT_Evolution_PreIntegrateCallback);
+		PreIntegrateCallback(Dt);
 	}
 
-#if CHAOS_EVOLUTION_COLLISION_TESTMODE
 	{
-		TestModeRestoreParticles();
-		TestModeResetCollisions();
+		CVD_SCOPE_TRACE_SOLVER_STEP(CVDDC_Integrate, TEXT("Integrate"))
+		SCOPE_CYCLE_COUNTER(STAT_Evolution_Integrate);
+		CSV_SCOPED_TIMING_STAT(PhysicsVerbose, StepSolver_Integrate);
+		Integrate(Dt);
 	}
-#endif
 
 	{
 		CVD_SCOPE_TRACE_SOLVER_STEP(CVDDC_ApplyKinematicTargets, TEXT("ApplyKinematicTargets"))
@@ -478,7 +480,7 @@ void FPBDRigidsEvolutionGBF::AdvanceOneTimeStepImpl(const FReal Dt, const FSubSt
 	if (PostIntegrateCallback != nullptr)
 	{
 		SCOPE_CYCLE_COUNTER(STAT_Evolution_PostIntegrateCallback);
-		PostIntegrateCallback();
+		PostIntegrateCallback(Dt);
 	}
 
 	{
@@ -521,7 +523,7 @@ void FPBDRigidsEvolutionGBF::AdvanceOneTimeStepImpl(const FReal Dt, const FSubSt
 	if (PostDetectCollisionsCallback != nullptr)
 	{
 		SCOPE_CYCLE_COUNTER(STAT_Evolution_PostDetectCollisionsCallback);
-		PostDetectCollisionsCallback();
+		PostDetectCollisionsCallback(Dt);
 	}
 
 	{
@@ -578,7 +580,7 @@ void FPBDRigidsEvolutionGBF::AdvanceOneTimeStepImpl(const FReal Dt, const FSubSt
 	if (PreSolveCallback != nullptr)
 	{
 		SCOPE_CYCLE_COUNTER(STAT_Evolution_PreSolveCallback);
-		PreSolveCallback();
+		PreSolveCallback(Dt);
 	}
 
 	// todo(chaos) : we are using the main gravity ( index 0 ) we should revise this to account for the various gravities based on the constraint ? 
@@ -630,7 +632,7 @@ void FPBDRigidsEvolutionGBF::AdvanceOneTimeStepImpl(const FReal Dt, const FSubSt
 	if (PostSolveCallback != nullptr)
 	{
 		SCOPE_CYCLE_COUNTER(STAT_Evolution_PostSolveCallback);
-		PostSolveCallback();
+		PostSolveCallback(Dt);
 	}
 
 	{
@@ -682,6 +684,7 @@ void FPBDRigidsEvolutionGBF::AdvanceOneTimeStepImpl(const FReal Dt, const FSubSt
 	}
 
 #if CHAOS_EVOLUTION_COLLISION_TESTMODE
+	if (CVars::bChaos_Solver_TestMode_Enabled && CVars::bChaos_Solver_TestMode_ShowInitialTransforms)
 	{
 		TestModeRestoreParticles();
 	}
@@ -716,6 +719,161 @@ void FPBDRigidsEvolutionGBF::AdvanceOneTimeStepImpl(const FReal Dt, const FSubSt
 #endif
 }
 
+void FPBDRigidsEvolutionGBF::Integrate(FReal Dt)
+{
+	TParticleView<FPBDRigidParticles>& ParticlesView = Particles.GetActiveParticlesView();
+
+	//SCOPE_CYCLE_COUNTER(STAT_Integrate);
+	CHAOS_SCOPED_TIMER(Integrate);
+
+	const FReal BoundsThickness = GetCollisionConstraints().GetDetectorSettings().BoundsExpansion;
+	const FReal VelocityBoundsMultiplier = GetCollisionConstraints().GetDetectorSettings().BoundsVelocityInflation;
+	const FReal MaxVelocityBoundsExpansion = GetCollisionConstraints().GetDetectorSettings().MaxVelocityBoundsExpansion;
+	const FReal HackMaxAngularSpeedSq = CVars::HackMaxAngularVelocity * CVars::HackMaxAngularVelocity;
+	const FReal HackMaxLinearSpeedSq = CVars::HackMaxVelocity * CVars::HackMaxVelocity;
+
+	FChaosVDContextWrapper CVDContext;
+	CVD_GET_WRAPPED_CURRENT_CONTEXT(CVDContext);
+
+	ParticlesView.ParallelFor([&](auto& GeomParticle, int32 Index)
+		{
+			CVD_SCOPE_CONTEXT(CVDContext.Context);
+
+			//question: can we enforce this at the API layer? Right now islands contain non dynamic which makes this hard
+			auto PBDParticle = GeomParticle.CastToRigidParticle();
+			if (PBDParticle && PBDParticle->ObjectState() == EObjectStateType::Dynamic)
+			{
+				auto& Particle = *PBDParticle;
+
+				FVec3 V = Particle.GetV();
+				FVec3 W = Particle.GetW();
+
+				//save off previous velocities
+				Particle.SetPreV(V);
+				Particle.SetPreW(W);
+
+				for (FForceRule ForceRule : ForceRules)
+				{
+					ForceRule(Particle, Dt);
+				}
+
+				//EulerStepVelocityRule.Apply(Particle, Dt);
+				V += Particle.Acceleration() * Dt;
+				W += Particle.AngularAcceleration() * Dt;
+
+				//AddImpulsesRule.Apply(Particle, Dt);
+				V += Particle.LinearImpulseVelocity();
+				W += Particle.AngularImpulseVelocity();
+				Particle.LinearImpulseVelocity() = FVec3(0);
+				Particle.AngularImpulseVelocity() = FVec3(0);
+
+				//EtherDragRule.Apply(Particle, Dt);
+
+				const FReal LinearDrag = LinearEtherDragOverride >= 0 ? LinearEtherDragOverride : Particle.LinearEtherDrag() * Dt;
+				const FReal LinearMultiplier = FMath::Max(FReal(0), FReal(1) - LinearDrag);
+				V *= LinearMultiplier;
+
+				const FReal AngularDrag = AngularEtherDragOverride >= 0 ? AngularEtherDragOverride : Particle.AngularEtherDrag() * Dt;
+				const FReal AngularMultiplier = FMath::Max(FReal(0), FReal(1) - AngularDrag);
+				W *= AngularMultiplier;
+
+				FReal LinearSpeedSq = V.SizeSquared();
+				if (LinearSpeedSq > Particle.MaxLinearSpeedSq())
+				{
+					V *= FMath::Sqrt(Particle.MaxLinearSpeedSq() / LinearSpeedSq);
+				}
+
+				FReal AngularSpeedSq = W.SizeSquared();
+				if (AngularSpeedSq > Particle.MaxAngularSpeedSq())
+				{
+					W *= FMath::Sqrt(Particle.MaxAngularSpeedSq() / AngularSpeedSq);
+				}
+
+				if (CVars::HackMaxAngularVelocity >= 0.f)
+				{
+					AngularSpeedSq = W.SizeSquared();
+					if (AngularSpeedSq > HackMaxAngularSpeedSq)
+					{
+						W = W * (CVars::HackMaxAngularVelocity / FMath::Sqrt(AngularSpeedSq));
+					}
+				}
+
+				if (CVars::HackMaxVelocity >= 0.f)
+				{
+					LinearSpeedSq = V.SizeSquared();
+					if (LinearSpeedSq > HackMaxLinearSpeedSq)
+					{
+						V = V * (CVars::HackMaxVelocity / FMath::Sqrt(LinearSpeedSq));
+					}
+				}
+
+				FVec3 PCoM = Particle.XCom();
+				FRotation3 QCoM = Particle.RCom();
+
+				PCoM = PCoM + V * Dt;
+				QCoM = FRotation3::IntegrateRotationWithAngularVelocity(QCoM, W, Dt);
+
+				Particle.SetTransformPQCom(PCoM, QCoM);
+				Particle.SetV(V);
+				Particle.SetW(W);
+
+				// We need to expand the bounds back along velocity otherwise we can miss collisions when a box
+				// lands on another box, imparting velocity to the lower box and causing the boxes to be
+				// separated by more than Cull Distance at collision detection time.
+				FVec3 VelocityBoundsDelta = FVec3(0);
+				if ((VelocityBoundsMultiplier > 0) && (MaxVelocityBoundsExpansion > 0))
+				{
+					VelocityBoundsDelta = (-VelocityBoundsMultiplier * Dt) * V;
+
+					// Box clamp to avoid sqrt
+					VelocityBoundsDelta.BoundToCube(MaxVelocityBoundsExpansion);
+				}
+
+				if (!Particle.CCDEnabled())
+				{
+					// Expand bounds about P/Q by a small amount. This can still result in missed collisions, especially
+					// when we have joints that pull the body back to X/R, if P-X is greater than the BoundsThickness
+					Particle.UpdateWorldSpaceStateSwept(FRigidTransform3(Particle.GetP(), Particle.GetQ()), FVec3(BoundsThickness), VelocityBoundsDelta);
+				}
+				else
+				{
+
+#if CHAOS_DEBUG_DRAW
+					if (CVars::ChaosSolverDrawCCDThresholds)
+					{
+						DebugDraw::DrawCCDAxisThreshold(Particle.GetX(), Particle.CCDAxisThreshold(), Particle.GetP() - Particle.GetX(), Particle.GetQ());
+					}
+#endif
+
+					if (FCCDHelpers::DeltaExceedsThreshold(Particle.CCDAxisThreshold(), Particle.GetP() - Particle.GetX(), Particle.GetQ()))
+					{
+						// We sweep the bounds from P back along the velocity and expand by a small amount.
+						// If not using tight bounds we also expand the bounds in all directions by Velocity. This is necessary only for secondary CCD collisions
+						// @todo(chaos): expanding the bounds by velocity is very expensive - revisit this
+						const FVec3 VDt = V * Dt;
+						FReal CCDBoundsExpansion = BoundsThickness;
+						if (!CVars::bChaosCollisionCCDUseTightBoundingBox && (CVars::ChaosCollisionCCDConstraintMaxProcessCount > 1))
+						{
+							CCDBoundsExpansion += VDt.GetAbsMax();
+						}
+						Particle.UpdateWorldSpaceStateSwept(FRigidTransform3(Particle.GetP(), Particle.GetQ()), FVec3(CCDBoundsExpansion), -VDt);
+					}
+					else
+					{
+						Particle.UpdateWorldSpaceStateSwept(FRigidTransform3(Particle.GetP(), Particle.GetQ()), FVec3(BoundsThickness), VelocityBoundsDelta);
+					}
+				}
+
+				CVD_TRACE_PARTICLE(PBDParticle->Handle())
+			}
+		});
+
+	for (auto& Particle : ParticlesView)
+	{
+		Base::DirtyParticle(Particle);
+	}
+}
+
 void FPBDRigidsEvolutionGBF::SetIsDeterministic(const bool bInIsDeterministic)
 {
 	// We detect collisions in parallel, so order is non-deterministic without additional processing
@@ -1043,7 +1201,7 @@ CHAOS_API void FPBDRigidsEvolutionGBF::SetParticleTransform(FGeometryParticleHan
 #if CHAOS_EVOLUTION_COLLISION_TESTMODE
 	{
 		// Update the test mode cache so we can move particles in PIE to test collisions
-		TestModeSaveParticle(InParticle);
+		TestModeUpdateSavedParticle(InParticle);
 	}
 #endif
 }
@@ -1260,6 +1418,11 @@ void FPBDRigidsEvolutionGBF::UpdateInertiaConditioning()
 #if CHAOS_EVOLUTION_COLLISION_TESTMODE
 void FPBDRigidsEvolutionGBF::TestModeResetCollisions()
 {
+	if (!CVars::bChaos_Solver_TestMode_Enabled)
+	{
+		return;
+	}
+	
 	for (FPBDCollisionConstraintHandle* Collision : CollisionConstraints.GetConstraintHandles())
 	{
 		if (Collision != nullptr)
diff --git a/Engine/Source/Runtime/Experimental/Chaos/Private/Chaos/ShapeInstance.cpp b/Engine/Source/Runtime/Experimental/Chaos/Private/Chaos/ShapeInstance.cpp
--- a/Engine/Source/Runtime/Experimental/Chaos/Private/Chaos/ShapeInstance.cpp
+++ b/Engine/Source/Runtime/Experimental/Chaos/Private/Chaos/ShapeInstance.cpp
@@ -69,7 +69,7 @@ namespace Chaos
 			{
 				for (auto& Shape : ShapesArray)
 				{
-					Shape->UpdateShapeBounds(ActorTM, FVec3(0));
+					Shape->UpdateShapeBounds(ActorTM);
 				}
 			}
 		}
@@ -266,12 +266,12 @@ namespace Chaos
 
 		if (Ar.CustomVer(FExternalPhysicsCustomObjectVersion::GUID) >= FExternalPhysicsCustomObjectVersion::SerializeShapeWorldSpaceBounds)
 		{
-			TBox<FReal, 3>::SerializeAsAABB(Ar, WorldSpaceInflatedShapeBounds);
+			TBox<FReal, 3>::SerializeAsAABB(Ar, WorldSpaceShapeBounds);
 		}
 		else
 		{
 			// This should be set by particle serializing this FPerShapeData.
-			WorldSpaceInflatedShapeBounds = FAABB3(FVec3(0.0f, 0.0f, 0.0f), FVec3(0.0f, 0.0f, 0.0f));
+			WorldSpaceShapeBounds = FAABB3(FVec3(0.0f, 0.0f, 0.0f), FVec3(0.0f, 0.0f, 0.0f));
 		}
 	}
 
@@ -333,21 +333,21 @@ namespace Chaos
 	///////////////////////////////////////////////////////////////////////////////////////////////
 
 
-	void FShapeInstanceProxy::UpdateShapeBounds(const FRigidTransform3& WorldTransform, const FVec3& BoundsExpansion)
+	void FShapeInstanceProxy::UpdateShapeBounds(const FRigidTransform3& WorldTransform)
 	{
 		if (Geometry && Geometry->HasBoundingBox())
 		{
-			WorldSpaceInflatedShapeBounds = Geometry->CalculateTransformedBounds(WorldTransform).ThickenSymmetrically(BoundsExpansion);
+			WorldSpaceShapeBounds = Geometry->CalculateTransformedBounds(WorldTransform);
 		}
 		else
 		{
-			WorldSpaceInflatedShapeBounds = FAABB3(WorldTransform.GetLocation(), WorldTransform.GetLocation()).ThickenSymmetrically(BoundsExpansion);
+			WorldSpaceShapeBounds = FAABB3(WorldTransform.GetLocation(), WorldTransform.GetLocation());
 		}
 	}
 
-	void FShapeInstanceProxy::UpdateWorldSpaceState(const FRigidTransform3& WorldTransform, const FVec3& BoundsExpansion)
+	void FShapeInstanceProxy::UpdateWorldSpaceState(const FRigidTransform3& WorldTransform)
 	{
-		UpdateShapeBounds(WorldTransform, BoundsExpansion);
+		UpdateShapeBounds(WorldTransform);
 	}
 
 	const FImplicitObject* FShapeInstanceProxy::GetLeafGeometry() const
@@ -397,19 +397,19 @@ namespace Chaos
 	///////////////////////////////////////////////////////////////////////////////////////////////
 
 
-	void FShapeInstance::UpdateShapeBounds(const FRigidTransform3& WorldTransform, const FVec3& BoundsExpansion)
+	void FShapeInstance::UpdateShapeBounds(const FRigidTransform3& WorldTransform)
 	{
 		if (Geometry && Geometry->HasBoundingBox())
 		{
-			WorldSpaceInflatedShapeBounds = Geometry->CalculateTransformedBounds(WorldTransform).ThickenSymmetrically(BoundsExpansion);
+			WorldSpaceShapeBounds = Geometry->CalculateTransformedBounds(WorldTransform);
 		}
 		else
 		{
-			WorldSpaceInflatedShapeBounds = FAABB3(WorldTransform.GetLocation(), WorldTransform.GetLocation()).ThickenSymmetrically(BoundsExpansion);
+			WorldSpaceShapeBounds = FAABB3(WorldTransform.GetLocation(), WorldTransform.GetLocation());
 		}
 	}
 
-	void FShapeInstance::UpdateWorldSpaceState(const FRigidTransform3& WorldTransform, const FVec3& BoundsExpansion)
+	void FShapeInstance::UpdateWorldSpaceState(const FRigidTransform3& WorldTransform)
 	{
 		FRigidTransform3 LeafWorldTransform = WorldTransform;
 		const FRigidTransform3* LeafRelativeTransform = nullptr;
@@ -430,11 +430,11 @@ namespace Chaos
 		// Update the bounds at the world transform
 		if ((LeafGeometry != nullptr) && LeafGeometry->HasBoundingBox())
 		{
-			WorldSpaceInflatedShapeBounds = LeafGeometry->CalculateTransformedBounds(LeafWorldTransform).ThickenSymmetrically(BoundsExpansion);;
+			WorldSpaceShapeBounds = LeafGeometry->CalculateTransformedBounds(LeafWorldTransform);
 		}
 		else
 		{
-			WorldSpaceInflatedShapeBounds = FAABB3(WorldTransform.GetLocation(), WorldTransform.GetLocation()).ThickenSymmetrically(BoundsExpansion);
+			WorldSpaceShapeBounds = FAABB3(WorldTransform.GetLocation(), WorldTransform.GetLocation());
 		}
 	}
 
diff --git a/Engine/Source/Runtime/Experimental/Chaos/Private/PBDRigidsSolver.cpp b/Engine/Source/Runtime/Experimental/Chaos/Private/PBDRigidsSolver.cpp
--- a/Engine/Source/Runtime/Experimental/Chaos/Private/PBDRigidsSolver.cpp
+++ b/Engine/Source/Runtime/Experimental/Chaos/Private/PBDRigidsSolver.cpp
@@ -109,6 +109,13 @@ TRACE_DECLARE_INT_COUNTER(ChaosTraceCounter_NumProjectionIterations, TEXT("Chaos
 
 namespace Chaos
 {
+	namespace DebugDraw
+	{
+		CHAOS_API const FChaosDebugDrawColorsByState& GetDefaultShapesColorsPreIntegrate();
+		CHAOS_API const FChaosDebugDrawColorsByState& GetDefaultShapesColorsPostIntegrate();
+		CHAOS_API const FChaosDebugDrawColorsByState& GetDefaultShapesColorsCollisionDetection();
+	}
+
 	namespace CVars
 	{
 		int32 ChaosSolverDebugDrawShapes = CHAOS_SOLVER_ENABLE_DEBUG_DRAW;
@@ -173,6 +180,11 @@ namespace Chaos
 		FAutoConsoleVariableRef CVarChaosSolverDebugDrawShowServer(TEXT("p.Chaos.Solver.DebugDraw.ShowServer"), ChaosSolverDebugDrawShowServer, TEXT("Draw server related debug data"));
 		FAutoConsoleVariableRef CVarChaosSolverDebugDrawShowClient(TEXT("p.Chaos.Solver.DebugDraw.ShowClient"), ChaosSolverDebugDrawShowClient, TEXT("Draw client related debug data"));
 
+		int32 ChaosSolverDebugDrawPreIntegrationShapes = 0;
+		int32 ChaosSolverDebugDrawPreIntegrationCollisions = 0;
+		FAutoConsoleVariableRef CVarChaosSolverDrawPreIntegrationShapes(TEXT("p.Chaos.Solver.DebugDrawPreIntegrationShapes"), ChaosSolverDebugDrawPreIntegrationShapes, TEXT("Draw Shapes prior to integrate."));
+		FAutoConsoleVariableRef CVarChaosSolverDrawPreIntegrationCollisions(TEXT("p.Chaos.Solver.DebugDrawPreIntegrationCollisions"), ChaosSolverDebugDrawPreIntegrationCollisions, TEXT("Draw Collisions prior to integrate."));
+
 		int32 ChaosSolverDebugDrawPostIntegrationShapes = 0;
 		int32 ChaosSolverDebugDrawPostIntegrationCollisions = 0;
 		FAutoConsoleVariableRef CVarChaosSolverDrawPostIntegrationShapes(TEXT("p.Chaos.Solver.DebugDrawPostIntegrationShapes"), ChaosSolverDebugDrawPostIntegrationShapes, TEXT("Draw Shapes prior to constraint solve phase."));
@@ -627,7 +639,7 @@ namespace Chaos
 			});
 
 		MEvolution->SetPreIntegrateCallback(
-			[this]()
+			[this](FReal Dt)
 			{
 				for (ISimCallbackObject* Callback : SimCallbackObjects)
 				{
@@ -637,10 +649,12 @@ namespace Chaos
 						Callback->PreIntegrate_Internal();
 					}
 				}
+
+				PreIntegrateDebugDraw(Dt);
 			});
 
 		MEvolution->SetPostIntegrateCallback(
-			[this]()
+			[this](FReal Dt)
 			{
 				for (ISimCallbackObject* Callback : SimCallbackObjects)
 				{
@@ -653,7 +667,7 @@ namespace Chaos
 			});
 
 		MEvolution->SetPreSolveCallback(
-			[this]()
+			[this](FReal Dt)
 			{
 				for (ISimCallbackObject* Callback : SimCallbackObjects)
 				{
@@ -664,11 +678,11 @@ namespace Chaos
 					}
 				}
 
-				PreSolveDebugDraw();
+				PreSolveDebugDraw(Dt);
 			});
 
 		MEvolution->SetPostSolveCallback(
-			[this]()
+			[this](FReal Dt)
 			{
 				for (ISimCallbackObject* Callback : SimCallbackObjects)
 				{
@@ -2376,60 +2390,57 @@ TRACE_COUNTER_SET(ChaosTraceCounter_##Name, Value)
 		CHAOS_COUNTER_STAT(NumDynamicShapes, NumDynamicShapes);
 	}
 
-	void FPBDRigidsSolver::PreSolveDebugDraw() const
+	void FPBDRigidsSolver::DebugDrawShapes(const bool bShowStatic, const bool bShowKinematic, const bool bShowDynamic) const
 	{
 #if CHAOS_DEBUG_DRAW
-		QUICK_SCOPE_CYCLE_COUNTER(SolverDebugDraw);
-
-#if !UE_BUILD_SHIPPING && !UE_BUILD_TEST
-		const bool bIsServer = DebugName.ToString().StartsWith(TEXT("Server"));
-		if (bIsServer && !ChaosSolverDebugDrawShowServer)
+		if (bShowStatic)
 		{
-			return;
+			DebugDraw::DrawParticleShapes(FRigidTransform3(), Particles.GetActiveStaticParticlesView(), 1.0f, &ChaosSolverDebugDebugDrawSettings);
 		}
-		if (!bIsServer && !ChaosSolverDebugDrawShowClient)
+		if (bShowKinematic)
 		{
-			return;
+			DebugDraw::DrawParticleShapes(FRigidTransform3(), Particles.GetActiveKinematicParticlesView(), 1.0f, &ChaosSolverDebugDebugDrawSettings);
 		}
+		if (bShowDynamic)
+		{
+			DebugDraw::DrawParticleShapes(FRigidTransform3(), Particles.GetNonDisabledDynamicView(), 1.0f, &ChaosSolverDebugDebugDrawSettings);
+		}
+#endif
+	}
 
-		if (ChaosSolverDebugDrawColorShapeByClientServer)
+	void FPBDRigidsSolver::PreIntegrateDebugDraw(FReal Dt) const
+	{
+#if CHAOS_DEBUG_DRAW
+		QUICK_SCOPE_CYCLE_COUNTER(SolverDebugDraw);
+
+		if (ChaosSolverDebugDrawPreIntegrationShapes == 1)
 		{
-			if (bIsServer)
-			{
-				ChaosSolverDebugDebugDrawSettings.ShapesColorsPerState = GetSolverShapesColorsByState_Server();
-			}
-			else
-			{
-				ChaosSolverDebugDebugDrawSettings.ShapesColorsPerState = GetSolverShapesColorsByState_Client();
-			}
+			ChaosSolverDebugDebugDrawSettings.ShapesColorsPerState = DebugDraw::GetDefaultShapesColorsPreIntegrate();
+
+			DebugDrawShapes(false, !!ChaosSolverDrawShapesShowKinematic, !!ChaosSolverDrawShapesShowDynamic);
 		}
-		else
+		if (ChaosSolverDebugDrawPreIntegrationCollisions == 1)
 		{
-			ChaosSolverDebugDebugDrawSettings.ShapesColorsPerState = DebugDraw::GetDefaultShapesColorsByState();
+			DebugDraw::DrawCollisions(FRigidTransform3(), GetEvolution()->GetCollisionConstraints().GetConstraintAllocator(), 1.f, &ChaosSolverDebugDebugDrawSettings);
 		}
 #endif
+	}
+
+	void FPBDRigidsSolver::PreSolveDebugDraw(FReal Dt) const
+	{
+#if CHAOS_DEBUG_DRAW
+		QUICK_SCOPE_CYCLE_COUNTER(SolverDebugDraw);
 
 		if (ChaosSolverDebugDrawPostIntegrationShapes == 1)
 		{
-			if (ChaosSolverDrawShapesShowStatic)
-			{
-				DebugDraw::DrawParticleShapes(FRigidTransform3(), Particles.GetActiveStaticParticlesView(), 1.0f, &ChaosSolverDebugDebugDrawSettings);
-			}
-			if (ChaosSolverDrawShapesShowKinematic)
-			{
-				DebugDraw::DrawParticleShapes(FRigidTransform3(), Particles.GetActiveKinematicParticlesView(), 1.0f, &ChaosSolverDebugDebugDrawSettings);
-			}
-			if (ChaosSolverDrawShapesShowDynamic)
-			{
-				DebugDraw::DrawParticleShapes(FRigidTransform3(), Particles.GetNonDisabledDynamicView(), 1.0f, &ChaosSolverDebugDebugDrawSettings);
-			}
-		}
+			ChaosSolverDebugDebugDrawSettings.ShapesColorsPerState = DebugDraw::GetDefaultShapesColorsPostIntegrate();
 
+			DebugDrawShapes(!!ChaosSolverDrawShapesShowStatic, !!ChaosSolverDrawShapesShowKinematic, !!ChaosSolverDrawShapesShowDynamic);
+		}
 		if (ChaosSolverDebugDrawPostIntegrationCollisions == 1)
 		{
 			DebugDraw::DrawCollisions(FRigidTransform3(), GetEvolution()->GetCollisionConstraints().GetConstraintAllocator(), 1.f, &ChaosSolverDebugDebugDrawSettings);
 		}
-
 #endif
 	}
 
@@ -2438,47 +2449,39 @@ TRACE_COUNTER_SET(ChaosTraceCounter_##Name, Value)
 #if CHAOS_DEBUG_DRAW
 		QUICK_SCOPE_CYCLE_COUNTER(SolverDebugDraw);
 
-#if !UE_BUILD_SHIPPING && !UE_BUILD_TEST
-		const bool bIsServer = DebugName.ToString().StartsWith(TEXT("Server"));
-		if (bIsServer && !ChaosSolverDebugDrawShowServer)
-		{
-			return;
-		}
-		if (!bIsServer && !ChaosSolverDebugDrawShowClient)
-		{
-			return;
-		}
-
-		if (ChaosSolverDebugDrawColorShapeByClientServer)
+		if (ChaosSolverDebugDrawShapes == 1)
 		{
-			if (bIsServer)
+			const bool bIsServer = GetDebugName().ToString().StartsWith(TEXT("Server"));
+			if (bIsServer && !ChaosSolverDebugDrawShowServer)
 			{
-				ChaosSolverDebugDebugDrawSettings.ShapesColorsPerState = GetSolverShapesColorsByState_Server();
+				return;
 			}
-			else
+			if (!bIsServer && !ChaosSolverDebugDrawShowClient)
 			{
-				ChaosSolverDebugDebugDrawSettings.ShapesColorsPerState = GetSolverShapesColorsByState_Client();
-			}
-		}
-		else
-		{
-			ChaosSolverDebugDebugDrawSettings.ShapesColorsPerState = DebugDraw::GetDefaultShapesColorsByState();
-		}
-#endif
-		if (ChaosSolverDebugDrawShapes == 1)
-		{
-			if (ChaosSolverDrawShapesShowStatic)
-			{
-				DebugDraw::DrawParticleShapes(FRigidTransform3(), Particles.GetActiveStaticParticlesView(), 1.0f, &ChaosSolverDebugDebugDrawSettings);
+				return;
 			}
-			if (ChaosSolverDrawShapesShowKinematic)
+
+			if (ChaosSolverDebugDrawColorShapeByClientServer)
 			{
-				DebugDraw::DrawParticleShapes(FRigidTransform3(), Particles.GetActiveKinematicParticlesView(), 1.0f, &ChaosSolverDebugDebugDrawSettings);
+				if (bIsServer)
+				{
+					ChaosSolverDebugDebugDrawSettings.ShapesColorsPerState = GetSolverShapesColorsByState_Server();
+				}
+				else
+				{
+					ChaosSolverDebugDebugDrawSettings.ShapesColorsPerState = GetSolverShapesColorsByState_Client();
+				}
 			}
-			if (ChaosSolverDrawShapesShowDynamic)
+			else
 			{
-				DebugDraw::DrawParticleShapes(FRigidTransform3(), Particles.GetNonDisabledDynamicView(), 1.0f, &ChaosSolverDebugDebugDrawSettings);
+				ChaosSolverDebugDebugDrawSettings.ShapesColorsPerState = DebugDraw::GetDefaultShapesColorsByState();
 			}
+
+			DebugDrawShapes(!!ChaosSolverDrawShapesShowStatic, !!ChaosSolverDrawShapesShowKinematic, !!ChaosSolverDrawShapesShowDynamic);
+		}
+		if (ChaosSolverDebugDrawCollisions == 1)
+		{
+			DebugDraw::DrawCollisions(FRigidTransform3(), GetEvolution()->GetCollisionConstraints().GetConstraintAllocator(), 1.f, &ChaosSolverDebugDebugDrawSettings);
 		}
 		if (ChaosSolverDebugDrawMass == 1)
 		{
@@ -2492,10 +2495,6 @@ TRACE_COUNTER_SET(ChaosTraceCounter_##Name, Value)
 				DebugDraw::DrawParticleBVH(FRigidTransform3(), Particle.Handle(), FColor::Silver, &ChaosSolverDebugDebugDrawSettings);
 			}
 		}
-		if (ChaosSolverDebugDrawCollisions == 1) 
-		{
-			DebugDraw::DrawCollisions(FRigidTransform3(), GetEvolution()->GetCollisionConstraints().GetConstraintAllocator(), 1.f, &ChaosSolverDebugDebugDrawSettings);
-		}
 		if (ChaosSolverDebugDrawBounds == 1)
 		{
 			DebugDraw::DrawParticleBounds(FRigidTransform3(), Particles.GetActiveStaticParticlesView(), Dt, &ChaosSolverDebugDebugDrawSettings);
diff --git a/Engine/Source/Runtime/Experimental/Chaos/Public/Chaos/ChaosDebugDraw.h b/Engine/Source/Runtime/Experimental/Chaos/Public/Chaos/ChaosDebugDraw.h
--- a/Engine/Source/Runtime/Experimental/Chaos/Public/Chaos/ChaosDebugDraw.h
+++ b/Engine/Source/Runtime/Experimental/Chaos/Public/Chaos/ChaosDebugDraw.h
@@ -203,6 +203,7 @@ namespace Chaos
 		CHAOS_API void DrawParticleShapes(const FRigidTransform3& SpaceTransform, const TParticleView<FGeometryParticles>& ParticlesView,  FReal ColorScale, const FChaosDebugDrawSettings* Settings = nullptr);
 		CHAOS_API void DrawParticleShapes(const FRigidTransform3& SpaceTransform, const TParticleView<FKinematicGeometryParticles>& ParticlesView,  FReal ColorScale, const FChaosDebugDrawSettings* Settings = nullptr);
 		CHAOS_API void DrawParticleShapes(const FRigidTransform3& SpaceTransform, const TParticleView<FPBDRigidParticles>& ParticlesView,  FReal ColorScale, const FChaosDebugDrawSettings* Settings = nullptr);
+		CHAOS_API void DrawParticleShapes(const FRigidTransform3& SpaceTransform, const FGeometryParticleHandle* Particle, const FRigidTransform3& ParticleTransform, const FColor& Color, const FChaosDebugDrawSettings* Settings = nullptr);
 		CHAOS_API void DrawParticleShapes(const FRigidTransform3& SpaceTransform, const FGeometryParticleHandle* Particle, const FColor& Color, const FChaosDebugDrawSettings* Settings = nullptr);
 		CHAOS_API void DrawParticleShapes(const FRigidTransform3& SpaceTransform, const FGeometryParticle* Particle, const FColor& Color, const FChaosDebugDrawSettings* Settings = nullptr);
 		CHAOS_API void DrawParticleBVH(const FRigidTransform3& SpaceTransform, const FGeometryParticleHandle* Particle, const FColor& Color, const FChaosDebugDrawSettings* Settings = nullptr);
diff --git a/Engine/Source/Runtime/Experimental/Chaos/Public/Chaos/GeometryParticles.h b/Engine/Source/Runtime/Experimental/Chaos/Public/Chaos/GeometryParticles.h
--- a/Engine/Source/Runtime/Experimental/Chaos/Public/Chaos/GeometryParticles.h
+++ b/Engine/Source/Runtime/Experimental/Chaos/Public/Chaos/GeometryParticles.h
@@ -451,19 +451,23 @@ namespace Chaos
 			// If we have no shapes we are have a point bounds at our local origin
 			FAABB3 WorldBounds = (Shapes.Num() > 0) ? FAABB3::EmptyAABB() : FAABB3::ZeroAABB();
 
+			// NOTE: Individual shape bounds are not expanded. We only require that the particle bounds is expanded because that is used
+			// in the broadphase. The midphase is what requires the shape bounds and it also has all the information required to expand the shape bounds as needed.
 			for (const auto& Shape : Shapes)
 			{
-				Shape->UpdateWorldSpaceState(WorldTransform, BoundsExpansion);
-				WorldBounds.GrowToInclude(Shape->GetWorldSpaceInflatedShapeBounds());
+				Shape->UpdateWorldSpaceState(WorldTransform);
+				WorldBounds.GrowToInclude(Shape->GetWorldSpaceShapeBounds());
 			}
 
-			MWorldSpaceInflatedBounds[Index] = TAABB<T, d>(WorldBounds);
+			MWorldSpaceInflatedBounds[Index] = TAABB<T, d>(WorldBounds).ThickenSymmetrically(BoundsExpansion);
 		}
 
-		void UpdateWorldSpaceStateSwept(const int32 Index, const FRigidTransform3& EndWorldTransform, const FVec3& BoundsExpansion, const FVec3& DeltaX)
+		void UpdateWorldSpaceStateSwept(const int32 Index, const FRigidTransform3& WorldTransform, const FVec3& BoundsExpansion, const FVec3& DeltaX)
 		{
-			// NOTE: Individual shape bounds are not expanded by the DeltaX (velocity term). Maybe they should be...
-			UpdateWorldSpaceState(Index, EndWorldTransform, BoundsExpansion);
+			// Update the bounds of all shapes (individual shape bounds are not expanded) and accumulate the net bounds
+			UpdateWorldSpaceState(Index, WorldTransform, BoundsExpansion);
+
+			// Apply the swept bounds delta
 			MWorldSpaceInflatedBounds[Index].GrowByVector(DeltaX);
 		}
 
diff --git a/Engine/Source/Runtime/Experimental/Chaos/Public/Chaos/PBDRigidsEvolution.h b/Engine/Source/Runtime/Experimental/Chaos/Public/Chaos/PBDRigidsEvolution.h
--- a/Engine/Source/Runtime/Experimental/Chaos/Public/Chaos/PBDRigidsEvolution.h
+++ b/Engine/Source/Runtime/Experimental/Chaos/Public/Chaos/PBDRigidsEvolution.h
@@ -1147,9 +1147,11 @@ protected:
 
 #if CHAOS_EVOLUTION_COLLISION_TESTMODE
 	// Test Mode for Collision issues (resets particle positions every tick for repeatable testing)
+	CHAOS_API void TestModeStep();
 	CHAOS_API void TestModeParticleDisabled(FGeometryParticleHandle* Particle);
 	CHAOS_API void TestModeSaveParticles();
 	CHAOS_API void TestModeSaveParticle(FGeometryParticleHandle* Particle);
+	CHAOS_API void TestModeUpdateSavedParticle(FGeometryParticleHandle* Particle);
 	CHAOS_API void TestModeRestoreParticles();
 	CHAOS_API void TestModeRestoreParticle(FGeometryParticleHandle* Particle);
 
diff --git a/Engine/Source/Runtime/Experimental/Chaos/Public/Chaos/PBDRigidsEvolutionGBF.h b/Engine/Source/Runtime/Experimental/Chaos/Public/Chaos/PBDRigidsEvolutionGBF.h
--- a/Engine/Source/Runtime/Experimental/Chaos/Public/Chaos/PBDRigidsEvolutionGBF.h
+++ b/Engine/Source/Runtime/Experimental/Chaos/Public/Chaos/PBDRigidsEvolutionGBF.h
@@ -37,7 +37,7 @@ namespace Chaos
 		CHAOS_API extern int32 ChaosSolverDrawCCDThresholds;
 	}
 
-	using FPBDRigidsEvolutionCallback = TFunction<void()>;
+	using FPBDRigidsEvolutionCallback = TFunction<void(FReal Dt)>;
 
 	using FPBDRigidsEvolutionIslandCallback = TFunction<void(int32 Island)>;
 
@@ -208,161 +208,7 @@ namespace Chaos
 
 		CHAOS_API const FChaosPhysicsMaterial* GetFirstClusteredPhysicsMaterial(const FGeometryParticleHandle* Particle) const;
 
-		template<typename TParticleView>
-		void Integrate(const TParticleView& InParticles, FReal Dt)
-		{
-			//SCOPE_CYCLE_COUNTER(STAT_Integrate);
-			CHAOS_SCOPED_TIMER(Integrate);
-
-			const FReal BoundsThickness = GetCollisionConstraints().GetDetectorSettings().BoundsExpansion;
-			const FReal VelocityBoundsMultiplier = GetCollisionConstraints().GetDetectorSettings().BoundsVelocityInflation;
-			const FReal MaxVelocityBoundsExpansion = GetCollisionConstraints().GetDetectorSettings().MaxVelocityBoundsExpansion;
-			const FReal MaxAngularSpeedSq = CVars::HackMaxAngularVelocity * CVars::HackMaxAngularVelocity;
-			const FReal MaxSpeedSq = CVars::HackMaxVelocity * CVars::HackMaxVelocity;
-			
-			FChaosVDContextWrapper CVDContext;
-			CVD_GET_WRAPPED_CURRENT_CONTEXT(CVDContext);
-			
-			InParticles.ParallelFor([&](auto& GeomParticle, int32 Index) 
-			{
-				CVD_SCOPE_CONTEXT(CVDContext.Context);
-
-				//question: can we enforce this at the API layer? Right now islands contain non dynamic which makes this hard
-				auto PBDParticle = GeomParticle.CastToRigidParticle();
-				if (PBDParticle && PBDParticle->ObjectState() == EObjectStateType::Dynamic)
-				{
-					auto& Particle = *PBDParticle;
-
-					//save off previous velocities
-					Particle.SetPreVf(Particle.GetVf());
-					Particle.SetPreWf(Particle.GetWf());
-
-					for (FForceRule ForceRule : ForceRules)
-					{
-						ForceRule(Particle, Dt);
-					}
-
-					//EulerStepVelocityRule.Apply(Particle, Dt);
-					Particle.SetV(Particle.GetV() + Particle.Acceleration() * Dt);
-					Particle.SetW(Particle.GetW() + Particle.AngularAcceleration() * Dt);
-
-					//AddImpulsesRule.Apply(Particle, Dt);
-					Particle.SetV(Particle.GetV() + Particle.LinearImpulseVelocity());
-					Particle.SetW(Particle.GetW() + Particle.AngularImpulseVelocity());
-					Particle.LinearImpulseVelocity() = FVec3(0);
-					Particle.AngularImpulseVelocity() = FVec3(0);
-
-					//EtherDragRule.Apply(Particle, Dt);
-					{
-						FVec3 V = Particle.GetV();
-						FVec3 W = Particle.GetW();
-
-						const FReal LinearDrag = LinearEtherDragOverride >= 0 ? LinearEtherDragOverride : Particle.LinearEtherDrag() * Dt;
-						const FReal LinearMultiplier = FMath::Max(FReal(0), FReal(1) - LinearDrag);
-						V *= LinearMultiplier;
-
-						const FReal AngularDrag = AngularEtherDragOverride >= 0 ? AngularEtherDragOverride : Particle.AngularEtherDrag() * Dt;
-						const FReal AngularMultiplier = FMath::Max(FReal(0), FReal(1) - AngularDrag);
-						W *= AngularMultiplier;
-
-						const FReal LinearSpeedSq = V.SizeSquared();
-						const FReal AngularSpeedSq = W.SizeSquared();
-
-						if (LinearSpeedSq > Particle.MaxLinearSpeedSq())
-						{
-							V *= FMath::Sqrt(Particle.MaxLinearSpeedSq() / LinearSpeedSq);
-						}
-
-						if (AngularSpeedSq > Particle.MaxAngularSpeedSq())
-						{
-							W *= FMath::Sqrt(Particle.MaxAngularSpeedSq() / AngularSpeedSq);
-						}
-						Particle.SetV(V);
-						Particle.SetW(W);
-					}
-
-					if (CVars::HackMaxAngularVelocity >= 0.f)
-					{
-						const FReal AngularSpeedSq = Particle.GetW().SizeSquared();
-						if (AngularSpeedSq > MaxAngularSpeedSq)
-						{
-							Particle.SetW(Particle.GetW() * (CVars::HackMaxAngularVelocity / FMath::Sqrt(AngularSpeedSq)));
-						}
-					}
-
-					if (CVars::HackMaxVelocity >= 0.f)
-					{
-						const FReal SpeedSq = Particle.GetV().SizeSquared();
-						if (SpeedSq > MaxSpeedSq)
-						{
-							Particle.SetV(Particle.GetV() * (CVars::HackMaxVelocity / FMath::Sqrt(SpeedSq)));
-						}
-					}
-
-					FVec3 PCoM = Particle.XCom();
-					FRotation3 QCoM = Particle.RCom();
-
-					PCoM = PCoM + Particle.GetV() * Dt;
-					QCoM = FRotation3::IntegrateRotationWithAngularVelocity(QCoM, Particle.GetW(), Dt);
-
-					Particle.SetTransformPQCom(PCoM, QCoM);
-
-					// We need to expand the bounds back along velocity otherwise we can miss collisions when a box
-					// lands on another box, imparting velocity to the lower box and causing the boxes to be
-					// separated by more than Cull Distance at collision detection time.
-					FVec3 VelocityBoundsDelta = FVec3(0);
-					if ((VelocityBoundsMultiplier > 0) && (MaxVelocityBoundsExpansion > 0))
-					{
-						VelocityBoundsDelta = (-VelocityBoundsMultiplier * Dt) * Particle.GetV();
-						
-						// Box clamp to avoid sqrt
-						VelocityBoundsDelta.BoundToCube(MaxVelocityBoundsExpansion);
-					}
-
-					if (!Particle.CCDEnabled())
-					{
-						// Expand bounds about P/Q by a small amount. This can still result in missed collisions, especially
-						// when we have joints that pull the body back to X/R, if P-X is greater than the BoundsThickness
-						Particle.UpdateWorldSpaceStateSwept(FRigidTransform3(Particle.GetP(), Particle.GetQ()), FVec3(BoundsThickness), VelocityBoundsDelta);
-					}
-					else
-					{
-
-#if CHAOS_DEBUG_DRAW
-						if (CVars::ChaosSolverDrawCCDThresholds)
-						{
-							DebugDraw::DrawCCDAxisThreshold(Particle.GetX(), Particle.CCDAxisThreshold(), Particle.GetP() - Particle.GetX(), Particle.GetQ());
-						}
-#endif
-
-						if (FCCDHelpers::DeltaExceedsThreshold(Particle.CCDAxisThreshold(), Particle.GetP() - Particle.GetX(), Particle.GetQ()))
-						{
-							// We sweep the bounds from P back along the velocity and expand by a small amount.
-							// If not using tight bounds we also expand the bounds in all directions by Velocity. This is necessary only for secondary CCD collisions
-							// @todo(chaos): expanding the bounds by velocity is very expensive - revisit this
-							const FVec3 VDt = Particle.GetV() * Dt;
-							FReal CCDBoundsExpansion = BoundsThickness;
-							if (!CVars::bChaosCollisionCCDUseTightBoundingBox && (CVars::ChaosCollisionCCDConstraintMaxProcessCount > 1))
-							{
-								CCDBoundsExpansion += VDt.GetAbsMax();
-							}
-							Particle.UpdateWorldSpaceStateSwept(FRigidTransform3(Particle.GetP(), Particle.GetQ()), FVec3(CCDBoundsExpansion), -VDt);
-						}
-						else
-						{
-							Particle.UpdateWorldSpaceStateSwept(FRigidTransform3(Particle.GetP(), Particle.GetQ()), FVec3(BoundsThickness), VelocityBoundsDelta);
-						}
-					}
-
-					CVD_TRACE_PARTICLE(PBDParticle->Handle())
-				}
-			});
-
-			for (auto& Particle : InParticles)
-			{
-				Base::DirtyParticle(Particle);
-			}
-		}
+		CHAOS_API void Integrate(FReal Dt);
 
 		CHAOS_API void Serialize(FChaosArchive& Ar);
 
@@ -393,7 +239,11 @@ namespace Chaos
 				Particle.WSmooth() = FMath::Lerp(Particle.WSmooth(), PredictedAngularVelocity, SmoothRate);
 			}
 		}
-		
+
+		template<typename TParticleView> 
+		UE_DEPRECATED(5.4, "Use Integrate(Dt)")
+		void Integrate(const TParticleView& InParticles, FReal Dt) { Integrate(Dt); }
+
 	protected:
 
 		CHAOS_API void AdvanceOneTimeStepImpl(const FReal dt, const FSubStepInfo& SubStepInfo);
diff --git a/Engine/Source/Runtime/Experimental/Chaos/Public/Chaos/ParticleHandle.h b/Engine/Source/Runtime/Experimental/Chaos/Public/Chaos/ParticleHandle.h
--- a/Engine/Source/Runtime/Experimental/Chaos/Public/Chaos/ParticleHandle.h
+++ b/Engine/Source/Runtime/Experimental/Chaos/Public/Chaos/ParticleHandle.h
@@ -631,16 +631,16 @@ public:
 	}
 
 	/**
-	 * @brief Update any cached state that depends on world-space transform
-	 * @param EndWorldTransform The transform at the end of the sweep
-	 * @param BoundsExpansion A uniform expansion applied to the bounds of the particle and the all shapes
-	 * @param DeltaX A directional expansion applied to the bounds of the particle, but not the shapes
+	 * @brief Update any cached state that depends on world-space transform for a sweep between WorldTransform and WorldTransform + DeltaX
+	 * @param WorldTransform The transform at the end of the sweep
+	 * @param BoundsExpansion A uniform expansion applied to the bounds of the particle, but not the individual shapes
+	 * @param DeltaX A directional expansion applied to the bounds of the particle, but not the individual shapes
 	 * This includes the world space bounds for the particle and all its shapes. If DeltaX is not zero,
-	 * the bounds will be equivalent to a union of the bounds at EndWorldTransform and EndWorldTransform + DeltaX.
+	 * the bounds will be equivalent to a union of the bounds at WorldTransform and WorldTransform + DeltaX.
 	*/
-	void UpdateWorldSpaceStateSwept(const FRigidTransform3& EndWorldTransform, const FVec3& BoundsExpansion, const FVec3& DeltaX)
+	void UpdateWorldSpaceStateSwept(const FRigidTransform3& WorldTransform, const FVec3& BoundsExpansion, const FVec3& DeltaX)
 	{
-		GeometryParticles->UpdateWorldSpaceStateSwept(ParticleIdx, EndWorldTransform, BoundsExpansion, DeltaX);
+		GeometryParticles->UpdateWorldSpaceStateSwept(ParticleIdx, WorldTransform, BoundsExpansion, DeltaX);
 	}
 
 	bool HasBounds() const { return GeometryParticles->HasBounds(ParticleIdx); }
diff --git a/Engine/Source/Runtime/Experimental/Chaos/Public/Chaos/ShapeInstance.h b/Engine/Source/Runtime/Experimental/Chaos/Public/Chaos/ShapeInstance.h
--- a/Engine/Source/Runtime/Experimental/Chaos/Public/Chaos/ShapeInstance.h
+++ b/Engine/Source/Runtime/Experimental/Chaos/Public/Chaos/ShapeInstance.h
@@ -80,7 +80,10 @@ namespace Chaos
 
 		CHAOS_API virtual void Serialize(FChaosArchive& Ar);
 
-		CHAOS_API void UpdateShapeBounds(const FRigidTransform3& WorldTM, const FVec3& BoundsExpansion = FVec3(0));
+		CHAOS_API void UpdateShapeBounds(const FRigidTransform3& WorldTM);
+
+		UE_DEPRECATED(5.4, "Bounds no longer expanded. Use UpdateShapeBounds without BoundsExpansion")
+		CHAOS_API void UpdateShapeBounds(const FRigidTransform3& WorldTM, const FVec3& BoundsExpansion) { UpdateShapeBounds(WorldTM); }
 
 		CHAOS_API void* GetUserData() const;
 		CHAOS_API void SetUserData(void* InUserData);
@@ -93,9 +96,15 @@ namespace Chaos
 
 		CHAOS_API FImplicitObjectRef GetGeometry() const;
 
-		CHAOS_API const TAABB<FReal, 3>& GetWorldSpaceInflatedShapeBounds() const;
+		CHAOS_API const TAABB<FReal, 3>& GetWorldSpaceShapeBounds() const;
+
+		CHAOS_API void UpdateWorldSpaceState(const FRigidTransform3& WorldTransform);
 
-		CHAOS_API void UpdateWorldSpaceState(const FRigidTransform3& WorldTransform, const FVec3& BoundsExpansion);
+		UE_DEPRECATED(5.4, "Bounds no longer expanded. Use GetWorldSpaceShapeBounds()")
+		CHAOS_API const TAABB<FReal, 3>& GetWorldSpaceInflatedShapeBounds() const { return GetWorldSpaceShapeBounds(); }
+		
+		UE_DEPRECATED(5.4, "Bounds no longer expanded. Use UpdateWorldSpaceState() without BoundsExpansion")
+		CHAOS_API void UpdateWorldSpaceState(const FRigidTransform3& WorldTransform, const FVec3& BoundsExpansion) { UpdateWorldSpaceState(WorldTransform); }
 
 		// The leaf shape (with transformed and implicit wrapper removed).
 		CHAOS_API const FImplicitObject* GetLeafGeometry() const;
@@ -169,7 +178,7 @@ namespace Chaos
 			, bIsSingleMaterial(false)
 			, ShapeIdx(InShapeIdx)
 			, Geometry()
-			, WorldSpaceInflatedShapeBounds(FAABB3(FVec3(0), FVec3(0)))
+			, WorldSpaceShapeBounds(FAABB3(FVec3(0), FVec3(0)))
 		{
 		}
 		
@@ -179,7 +188,7 @@ namespace Chaos
 			, bIsSingleMaterial(false)
 			, ShapeIdx(InShapeIdx)
 			, Geometry()
-			, WorldSpaceInflatedShapeBounds(FAABB3(FVec3(0), FVec3(0)))
+			, WorldSpaceShapeBounds(FAABB3(FVec3(0), FVec3(0)))
 		{
 			check(false);
 		}
@@ -189,7 +198,7 @@ namespace Chaos
 			, bIsSingleMaterial(false)
 			, ShapeIdx(InShapeIdx)
 			, Geometry(InGeometry)
-			, WorldSpaceInflatedShapeBounds(FAABB3(FVec3(0), FVec3(0)))
+			, WorldSpaceShapeBounds(FAABB3(FVec3(0), FVec3(0)))
 		{
 		}
 
@@ -198,7 +207,7 @@ namespace Chaos
 			, bIsSingleMaterial(Other.bIsSingleMaterial)
 			, ShapeIdx(Other.ShapeIdx)
 			, Geometry(Other.Geometry)
-			, WorldSpaceInflatedShapeBounds(Other.WorldSpaceInflatedShapeBounds)
+			, WorldSpaceShapeBounds(Other.WorldSpaceShapeBounds)
 		{
 		}
 
@@ -209,7 +218,7 @@ namespace Chaos
 		uint32 ShapeIdx : 29;
 		FShapeDirtyFlags DirtyFlags;	// For use by FShapeInstanceProxy as there's 4 bytes of padding here
 		FImplicitObjectPtr Geometry;
-		TAABB<FReal, 3> WorldSpaceInflatedShapeBounds;
+		TAABB<FReal, 3> WorldSpaceShapeBounds;
 	};
 
 
@@ -262,7 +271,7 @@ namespace Chaos
 		static CHAOS_API void UpdateGeometry(TUniquePtr<FShapeInstanceProxy>& InOutShapePtr, const FImplicitObjectPtr& InGeometry);
 		static CHAOS_API FShapeInstanceProxy* SerializationFactory(FChaosArchive& Ar, FShapeInstanceProxy*);
 
-		CHAOS_API void UpdateShapeBounds(const FRigidTransform3& WorldTM, const FVec3& BoundsExpansion = FVec3(0));
+		CHAOS_API void UpdateShapeBounds(const FRigidTransform3& WorldTM);
 
 		void* GetUserData() const { return CollisionData.Read().UserData; }
 		void SetUserData(void* InUserData)
@@ -282,7 +291,7 @@ namespace Chaos
 			CollisionData.Modify(true, DirtyFlags, Proxy, ShapeIdx, [InSimData](FCollisionData& Data) { Data.SimData = InSimData; });
 		}
 
-		CHAOS_API void UpdateWorldSpaceState(const FRigidTransform3& WorldTransform, const FVec3& BoundsExpansion);
+		CHAOS_API void UpdateWorldSpaceState(const FRigidTransform3& WorldTransform);
 
 		// The leaf shape (with transformed and implicit wrapper removed).
 		CHAOS_API const FImplicitObject* GetLeafGeometry() const;
@@ -540,7 +549,7 @@ namespace Chaos
 			}
 		}
 
-		CHAOS_API void UpdateShapeBounds(const FRigidTransform3& WorldTM, const FVec3& BoundsExpansion = FVec3(0));
+		CHAOS_API void UpdateShapeBounds(const FRigidTransform3& WorldTM);
 
 		void* GetUserData() const { return CollisionData.UserData; }
 		void SetUserData(void* InUserData) { CollisionData.UserData = InUserData; }
@@ -551,7 +560,7 @@ namespace Chaos
 		const FCollisionFilterData& GetSimData() const { return CollisionData.SimData; }
 		void SetSimData(const FCollisionFilterData& InSimData) { CollisionData.SimData = InSimData; }
 
-		CHAOS_API void UpdateWorldSpaceState(const FRigidTransform3& WorldTransform, const FVec3& BoundsExpansion);
+		CHAOS_API void UpdateWorldSpaceState(const FRigidTransform3& WorldTransform);
 
 		// The leaf shape (with transformed and implicit wrapper removed).
 		CHAOS_API const FImplicitObject* GetLeafGeometry() const;
@@ -939,9 +948,9 @@ namespace Chaos
 	///////////////////////////////////////////////////////////////////////////////////////////////
 
 
-	inline void FPerShapeData::UpdateShapeBounds(const FRigidTransform3& WorldTM, const FVec3& BoundsExpansion)
+	inline void FPerShapeData::UpdateShapeBounds(const FRigidTransform3& WorldTM)
 	{
-		DownCast([&WorldTM, &BoundsExpansion](auto& ShapeInstance) { ShapeInstance.UpdateShapeBounds(WorldTM, BoundsExpansion); });
+		DownCast([&WorldTM](auto& ShapeInstance) { ShapeInstance.UpdateShapeBounds(WorldTM); });
 	}
 
 	inline void* FPerShapeData::GetUserData() const
@@ -979,14 +988,14 @@ namespace Chaos
 		return Geometry.GetReference();
 	}
 
-	inline const TAABB<FReal, 3>& FPerShapeData::GetWorldSpaceInflatedShapeBounds() const
+	inline const TAABB<FReal, 3>& FPerShapeData::GetWorldSpaceShapeBounds() const
 	{
-		return WorldSpaceInflatedShapeBounds;
+		return WorldSpaceShapeBounds;
 	}
 
-	inline void FPerShapeData::UpdateWorldSpaceState(const FRigidTransform3& WorldTransform, const FVec3& BoundsExpansion)
+	inline void FPerShapeData::UpdateWorldSpaceState(const FRigidTransform3& WorldTransform)
 	{
-		DownCast([&WorldTransform, &BoundsExpansion](auto& ShapeInstance) { ShapeInstance.UpdateWorldSpaceState(WorldTransform, BoundsExpansion); });
+		DownCast([&WorldTransform](auto& ShapeInstance) { ShapeInstance.UpdateWorldSpaceState(WorldTransform); });
 	}
 
 	inline const FImplicitObject* FPerShapeData::GetLeafGeometry() const
diff --git a/Engine/Source/Runtime/Experimental/Chaos/Public/PBDRigidsSolver.h b/Engine/Source/Runtime/Experimental/Chaos/Public/PBDRigidsSolver.h
--- a/Engine/Source/Runtime/Experimental/Chaos/Public/PBDRigidsSolver.h
+++ b/Engine/Source/Runtime/Experimental/Chaos/Public/PBDRigidsSolver.h
@@ -244,7 +244,8 @@ namespace Chaos
 		CHAOS_API void SyncEvents_GameThread();
 
 		/**/
-		CHAOS_API void PreSolveDebugDraw() const;
+		CHAOS_API void PreIntegrateDebugDraw(FReal Dt) const;
+		CHAOS_API void PreSolveDebugDraw(FReal Dt) const;
 		CHAOS_API void PostTickDebugDraw(FReal Dt) const;
 
 		// Visual debugger (VDB) push methods
@@ -359,6 +360,8 @@ namespace Chaos
 
 		CHAOS_API void UpdateIsDeterministic();
 
+		CHAOS_API void DebugDrawShapes(const bool bShowStatic, const bool bShowKinematic, const bool bShowDynamic) const;
+
 		//
 		// Solver Data
 		//
diff --git a/Engine/Source/Runtime/PhysicsCore/Public/SQVisitor.h b/Engine/Source/Runtime/PhysicsCore/Public/SQVisitor.h
--- a/Engine/Source/Runtime/PhysicsCore/Public/SQVisitor.h
+++ b/Engine/Source/Runtime/PhysicsCore/Public/SQVisitor.h
@@ -222,14 +222,14 @@ private:
 				FAABB3 InflatedWorldBounds;
 				if (SQ == ESQType::Raycast)
 				{
-					InflatedWorldBounds = Shape->GetWorldSpaceInflatedShapeBounds();
+					InflatedWorldBounds = Shape->GetWorldSpaceShapeBounds();
 				}
 				else
 				{
 					// Transform to world bounds and get the proper half extent.
 					const FVec3 WorldHalfExtent = QueryGeom ? QueryGeomWorldBounds.Extents() * 0.5f : FVec3(HalfExtents);
 
-					InflatedWorldBounds = FAABB3(Shape->GetWorldSpaceInflatedShapeBounds().Min() - WorldHalfExtent, Shape->GetWorldSpaceInflatedShapeBounds().Max() + WorldHalfExtent);
+					InflatedWorldBounds = FAABB3(Shape->GetWorldSpaceShapeBounds().Min() - WorldHalfExtent, Shape->GetWorldSpaceShapeBounds().Max() + WorldHalfExtent);
 				}
 	
 				if (SQ != ESQType::Overlap)
