Geometry collection : fix removed piecers failing to shrink when the feature is enabled
- the issue was a consequence of the component space transforms computation optimization
- the code responsible to update the scale of the pieces now properly mark the component space transforms dirty  so the rendering reads the updated ones

diff --git a/Engine/Source/Runtime/Experimental/GeometryCollectionEngine/Private/GeometryCollection/GeometryCollectionComponent.cpp b/Engine/Source/Runtime/Experimental/GeometryCollectionEngine/Private/GeometryCollection/GeometryCollectionComponent.cpp
--- a/Engine/Source/Runtime/Experimental/GeometryCollectionEngine/Private/GeometryCollection/GeometryCollectionComponent.cpp
+++ b/Engine/Source/Runtime/Experimental/GeometryCollectionEngine/Private/GeometryCollection/GeometryCollectionComponent.cpp
@@ -4292,6 +4292,8 @@ void UGeometryCollectionComponent::MoveComponentToRootTransform()
 					const FRotator NewRotation = NewRootWorldPosition.Rotator();
 					MoveComponent(MoveBy, NewRotation, /* bSweep */ false,/* Hit */ NULL, MOVECOMP_SkipPhysicsMove);
 				}
+
+				OnTransformsDirty();
 			}
 		}
 	}
@@ -4382,6 +4384,8 @@ void UGeometryCollectionComponent::UpdateRemovalIfNeeded()
 
 			const int32 RootIndex = RestCollection->GetRootIndex();
 
+			bool bTransformsChanged = false;
+
 			const int32 NumTransforms = DecayFacade.GetDecayAttributeSize();
 			for (int32 TransformIndex = 0; TransformIndex < NumTransforms; ++TransformIndex)
 			{
@@ -4409,6 +4413,7 @@ void UGeometryCollectionComponent::UpdateRemovalIfNeeded()
 						FTransform3f Transform = DynamicCollection->GetTransform(TransformIndex);
 						Transform.SetScale3D(FVector3f::ZeroVector);
 						DynamicCollection->SetTransform(TransformIndex, Transform);
+						bTransformsChanged = true;
 					}
 					// do not try to get this condition out of the loop as this may cause some optimizer related issues
 					else if (RestCollection->bScaleOnRemoval && MassToLocal && CompSpaceTransform)
@@ -4428,9 +4433,15 @@ void UGeometryCollectionComponent::UpdateRemovalIfNeeded()
 						const FVector3f ScaleCenter = FVector3f(LocalDown + CenterOfMass);
 						const FTransform3f ScaleTransform(FQuat4f::Identity, ScaleCenter * FVector3f::FReal(1.f - Scale), FVector3f(Scale));
 						DynamicCollection->SetTransform(TransformIndex, ScaleTransform * DynamicCollection->GetTransform(TransformIndex));
+						bTransformsChanged = true;
 					}
 				}
 			}
+
+			if (bTransformsChanged)
+			{
+				OnTransformsDirty();
+			}
 		}
 	}
 }
