Defensive check against null RestCollection in PSO Precache.
#rnx

diff --git a/Engine/Source/Runtime/Experimental/GeometryCollectionEngine/Private/GeometryCollection/GeometryCollectionComponent.cpp b/Engine/Source/Runtime/Experimental/GeometryCollectionEngine/Private/GeometryCollection/GeometryCollectionComponent.cpp
--- a/Engine/Source/Runtime/Experimental/GeometryCollectionEngine/Private/GeometryCollection/GeometryCollectionComponent.cpp
+++ b/Engine/Source/Runtime/Experimental/GeometryCollectionEngine/Private/GeometryCollection/GeometryCollectionComponent.cpp
@@ -6998,7 +6998,10 @@ void UGeometryCollectionComponent::PostLoad()
 
 void UGeometryCollectionComponent::CollectPSOPrecacheData(const FPSOPrecacheParams& BasePrecachePSOParams, FMaterialInterfacePSOPrecacheParamsList& OutParams)
 {
-	check(RestCollection);
+	if (!RestCollection)
+	{
+		return;
+	}
 
 	FPSOPrecacheVertexFactoryDataList VFDataList;
 	const FVertexFactoryType* VFType = nullptr;
