Chaos - fix linker error - bad extern of global var

diff --git a/Engine/Source/Runtime/Experimental/Chaos/Private/PBDRigidsSolver.cpp b/Engine/Source/Runtime/Experimental/Chaos/Private/PBDRigidsSolver.cpp
--- a/Engine/Source/Runtime/Experimental/Chaos/Private/PBDRigidsSolver.cpp
+++ b/Engine/Source/Runtime/Experimental/Chaos/Private/PBDRigidsSolver.cpp
@@ -111,9 +111,9 @@ namespace Chaos
 {
 	namespace DebugDraw
 	{
-		CHAOS_API const FChaosDebugDrawColorsByState& GetDefaultShapesColorsPreIntegrate();
-		CHAOS_API const FChaosDebugDrawColorsByState& GetDefaultShapesColorsPostIntegrate();
-		CHAOS_API const FChaosDebugDrawColorsByState& GetDefaultShapesColorsCollisionDetection();
+		extern const FChaosDebugDrawColorsByState& GetDefaultShapesColorsPreIntegrate();
+		extern const FChaosDebugDrawColorsByState& GetDefaultShapesColorsPostIntegrate();
+		extern const FChaosDebugDrawColorsByState& GetDefaultShapesColorsCollisionDetection();
 	}
 
 	namespace CVars
