IAS Analytics Improvements.

#rnx

### OnDemand
- Add a new statistic tracking the total number of decode failures which tracks the following issues:
-- Mismatch between the FIoChunkDecodingParams and the raw memory blob size.
-- Block hash mismatch found when decoding a block.
-- A failure when decompressing a block.
- Add a new statistic tracking the total number of successful connections made to a http endpoint, note that this does include the initial connection.
- Add a new staticstic tracking the total number of times we consider the connection to the http endpoint as disconnected (we reached the preset failure limit)
- Previously all of our IAS analytics was associated with the currently connected endpoint, if we had no connection then the data did not make sense, so we just wouldn't send it but now we want to be able to send analytics data that includes those failures. To do this we have split the analytics reporting into two methods:
-- FOnDemandIoBackendStats ::ReportGeneralAnalytics which is always reported.
-- FOnDemandIoBackendStats::ReportEndPointAnalytics which is reported if the endpoint is valid.
- Added code documentation to confirm that if the IAS system is not initialized at all (i.e no backend) then FOnDemandIoBackend::ReportAnalytics will never be invoked.
- We no longer assume that when a single endpoint url is provided that it will be working (by setting the current endpoint to that) instead we add the url to the potential list (of one url) and let FOnDemandIoBackend::InitializePrimaryEndpoint verify that it is correct just as if we had acquired the urls from a distributed endpoint. This helps reduce the number of different initialization code paths we have.

### Statistics
- The on screen statistics message now displays if we are conncted to a http endpoint or not
- Report a new statistic "IasHttpHasEverConnected" which reports if the current process has ever managed to connect to and endpoint or not.
--In theory if a user never actually resolves the endpoint they will be able to play the game but with poor graphics and this stat should let us track this.
- Report a new statistic "IasDecodeErrors" which reports the total number of decode errors a user has encountered since the last time the analytics was reported.

diff --git a/Engine/Source/Runtime/Experimental/IoStoreOnDemand/Private/OnDemandIoDispatcherBackend.cpp b/Engine/Source/Runtime/Experimental/IoStoreOnDemand/Private/OnDemandIoDispatcherBackend.cpp
--- a/Engine/Source/Runtime/Experimental/IoStoreOnDemand/Private/OnDemandIoDispatcherBackend.cpp
+++ b/Engine/Source/Runtime/Experimental/IoStoreOnDemand/Private/OnDemandIoDispatcherBackend.cpp
@@ -1683,6 +1683,11 @@ void FOnDemandIoBackend::CompleteRequest(FChunkRequest* ChunkRequest)
 			Request->CreateBuffer(RawSize);
 			DecodingParams.RawOffset = Request->Options.GetOffset(); 
 			bDecoded = FIoChunkEncoding::Decode(DecodingParams, Chunk.GetView(), Request->GetBuffer().GetMutableView());
+
+			if (!bDecoded)
+			{
+				Stats.OnIoDecodeError();
+			}
 		}
 		
 		const uint64 DurationMs = Request->GetStartTime() > 0 ?
@@ -2062,6 +2067,8 @@ void FOnDemandIoBackend::InitializePrimaryEndpoint()
 		const FString TestPath = GetEndpointTestPath();
 		if (int32 Idx = LatencyTest(Urls.Left(MaxUrls), TestPath, bStopRequested); Idx != INDEX_NONE)
 		{
+			FOnDemandIoBackendStats::Get()->OnHttpConnected();
+
 			AvailableEps.Current = Idx;
 			UE_LOG(LogIas, Log, TEXT("Using endpoint '%s'"), *AvailableEps.GetCurrent());
 		}
@@ -2218,19 +2225,24 @@ void FOnDemandIoBackend::Mount(const FOnDemandEndpoint& Endpoint)
 			{
 				AvailableEps.Urls.Add(Url.Replace(TEXT("https"), TEXT("http")));
 			}
-			AvailableEps.Current = 0;
 		}
 	}
 
-	check(AvailableEps.HasCurrent());
-
 	if (GetTocMode() == ETocMode::LoadTocFromNetwork)
 	{
-		FIoStatus Result = DownloadoadOnDemandToc(AvailableEps.GetCurrent(), Endpoint.TocPath);
-		if (!Result.IsOk())
+		if (AvailableEps.HasCurrent())
+		{
+			FIoStatus Result = DownloadoadOnDemandToc(AvailableEps.GetCurrent(), Endpoint.TocPath);
+			if (!Result.IsOk())
+			{
+				UE_LOG(LogIas, Error, TEXT("Deferring TOC '%s/%s' due to '%s'"), *AvailableEps.GetCurrent(), *Endpoint.TocPath, *Result.ToString());
+				BackendStatus.SetHttpError(true);
+				FWriteScopeLock _(Lock);
+				DeferredTocs.Add(FTocParams{ Endpoint.TocPath, Endpoint.bForceTocDownload });
+			}
+		}
+		else
 		{
-			UE_LOG(LogIas, Error, TEXT("Deferring TOC '%s/%s' due to '%s'"), *AvailableEps.GetCurrent(), *Endpoint.TocPath, *Result.ToString());
-			BackendStatus.SetHttpError(true);
 			FWriteScopeLock _(Lock);
 			DeferredTocs.Add(FTocParams{ Endpoint.TocPath, Endpoint.bForceTocDownload });
 		}
@@ -2280,11 +2292,16 @@ void FOnDemandIoBackend::AbandonCache()
 
 void FOnDemandIoBackend::ReportAnalytics(TArray<FAnalyticsEventAttribute>& OutAnalyticsArray) const
 {
+	// If we got this far we know that IAS is enabled for the current process as it has a valid backend.
+	// However just because IAS is enabled does not mean we have managed to make a valid connection yet.
+
 	if (!GIasReportAnalyticsEnabled)
 	{
 		return;
 	}
 
+	Stats.ReportGeneralAnalytics(OutAnalyticsArray);
+
 	if (AvailableEps.HasCurrent())
 	{
 		FString CdnUrl = AvailableEps.GetCurrent();
@@ -2297,7 +2314,7 @@ void FOnDemandIoBackend::ReportAnalytics(TArray<FAnalyticsEventAttribute>& OutAn
 
 		AppendAnalyticsEventAttributeArray(OutAnalyticsArray, TEXT("IasCdnUrl"), MoveTemp(CdnUrl));
 
-		Stats.ReportAnalytics(OutAnalyticsArray);
+		Stats.ReportEndPointAnalytics(OutAnalyticsArray);
 	}
 }
 
@@ -2513,6 +2530,8 @@ uint32 FOnDemandIoBackend::Run()
 				WaitTime = GIasHttpHealthCheckWaitTime;
 				if (HttpClient->GetEndpoint() != INDEX_NONE)
 				{
+					FOnDemandIoBackendStats::Get()->OnHttpDisconnected();
+
 					AvailableEps.Current = INDEX_NONE;
 					HttpClient->SetEndpoint(INDEX_NONE);
 					HttpErrors.Reset(GIasHttpErrorSampleCount);
@@ -2522,6 +2541,8 @@ uint32 FOnDemandIoBackend::Run()
 				const FString TestPath = GetEndpointTestPath(); 
 				if (int32 Idx = LatencyTest(AvailableEps.Urls, TestPath, bStopRequested); Idx != INDEX_NONE)
 				{
+					FOnDemandIoBackendStats::Get()->OnHttpConnected();
+
 					AvailableEps.Current = Idx;
 					HttpClient->SetEndpoint(Idx);
 					BackendStatus.SetHttpError(false);
diff --git a/Engine/Source/Runtime/Experimental/IoStoreOnDemand/Private/Statistics.cpp b/Engine/Source/Runtime/Experimental/IoStoreOnDemand/Private/Statistics.cpp
--- a/Engine/Source/Runtime/Experimental/IoStoreOnDemand/Private/Statistics.cpp
+++ b/Engine/Source/Runtime/Experimental/IoStoreOnDemand/Private/Statistics.cpp
@@ -250,6 +250,7 @@ FCounterAtomicInt		GIoRequestReadCount(TEXT("Ias/IoRequestReadCount"), TraceCoun
 FCounterAtomicInt		GIoRequestReadBytes(TEXT("Ias/IoRequestReadBytes"), TraceCounterDisplayHint_Memory);
 FCounterInt				GIoRequestCancelCount(TEXT("Ias/IoRequestCancelCount"), TraceCounterDisplayHint_None);
 FCounterAtomicInt		GIoRequestErrorCount(TEXT("Ias/IoRequestErrorCount"), TraceCounterDisplayHint_None);
+FCounterAtomicInt		GIoDecodeErrorCount(TEXT("Ias/IoDecodeErrorCount"), TraceCounterDisplayHint_None);
 // cache stats
 FCounterAtomicInt		GCacheErrorCount(TEXT("Ias/CacheErrorCount"), TraceCounterDisplayHint_None);
 FCounterAtomicInt		GCacheGetCount(TEXT("Ias/CacheGetCount"), TraceCounterDisplayHint_None);
@@ -263,6 +264,8 @@ FCounterAtomicInt		GCachePendingBytes(TEXT("Ias/CachePendingBytes"), TraceCounte
 FCounterAtomicInt		GCacheReadBytes(TEXT("Ias/CacheReadBytes"), TraceCounterDisplayHint_Memory);
 FCounterAtomicInt		GCacheRejectBytes(TEXT("Ias/CachePutRejectBytes"), TraceCounterDisplayHint_Memory);
 // http stats
+FCounterInt				GHttpConnectCount(TEXT("Ias/HttpConnectCount"), TraceCounterDisplayHint_None);
+FCounterInt				GHttpDisconnectCount(TEXT("Ias/HttpDisconnectCount"), TraceCounterDisplayHint_None);
 FCounterInt				GHttpGetCount(TEXT("Ias/HttpGetCount"), TraceCounterDisplayHint_None);
 FCounterInt				GHttpErrorCount(TEXT("Ias/HttpErrorCount"), TraceCounterDisplayHint_None);
 FCounterInt				GHttpRetryCount(TEXT("Ias/HttpRetryCount"), TraceCounterDisplayHint_None);
@@ -403,8 +406,10 @@ FOnDemandIoBackendStats::FOnDemandIoBackendStats()
 		{
 			if (GIasDisplayOnScreenStatistics)
 			{
+				const bool bIsConnected = GHttpConnectCount.Get() > GHttpDisconnectCount.Get();
 				FText Message = FText::Format(
-					LOCTEXT("IAS", "IAS: Cached:{0} KiB | Read:{1} KiB ({2}) | Downloaded:{3} KiB ({4}) {5} ms | Retries:{6} | Pending:{7}"),
+					LOCTEXT("IAS", "IAS - {0}: Cached:{1} KiB | Read:{2} KiB ({3}) | Downloaded:{4} KiB ({5}) {6} ms | Retries:{7} | Pending:{8}"),
+					bIsConnected ? LOCTEXT("IASConnect", "Connected") : LOCTEXT("IASDisconnect", "Disconnected"),
 					GCacheCachedBytes.Get() >> 10,
 					GCacheReadBytes.Get() >> 10,
 					GCacheGetCount.Get(),
@@ -416,8 +421,7 @@ FOnDemandIoBackendStats::FOnDemandIoBackendStats()
 				);
 				Out.Add(FCoreDelegates::EOnScreenMessageSeverity::Info, Message);
 			}
-		}
-	);
+		});
 #undef LOCTEXT_NAMESPACE
 }
 
@@ -433,7 +437,18 @@ FOnDemandIoBackendStats* FOnDemandIoBackendStats::Get()
 	return GStatistics;
 }
 
-void FOnDemandIoBackendStats::ReportAnalytics(TArray<FAnalyticsEventAttribute>& OutAnalyticsArray) const
+void FOnDemandIoBackendStats::ReportGeneralAnalytics(TArray<FAnalyticsEventAttribute>& OutAnalyticsArray) const
+{
+	// Note that this analytics section is not optional, if we are reporting analytics then we report this section
+	// first. This means we can use the values here to determine if an analytics payload contains ondemand data or
+	// not since with the current system we are unable to specify our own analytics payload.
+
+	AppendAnalyticsEventAttributeArray(OutAnalyticsArray
+		, TEXT("IasHttpHasEverConnected"), GHttpConnectCount.Get() > 0 // Report if the system has ever actually managed to make a connection
+	);
+}
+
+void FOnDemandIoBackendStats::ReportEndPointAnalytics(TArray<FAnalyticsEventAttribute>& OutAnalyticsArray) const
 {
 	// We use this macro with counters that are used elsewhere meaning we can't just reset them
 	// each time we send an analytics payload. The macro will track the running total and only
@@ -450,6 +465,8 @@ void FOnDemandIoBackendStats::ReportAnalytics(TArray<FAnalyticsEventAttribute>&
 			,TEXT("IasHttpDurationMeanAvg"), GHttpAvgDuration.GetMean()
 			,TEXT("IasHttpDurationStdDev"), GHttpAvgDuration.GetDeviation()
 
+			,TRACK_DELTA("IasDecodeErrors", GIoDecodeErrorCount.Get())
+
 			,TEXT("IasHttpDuration0"), GHttpDurationBuckets[0]
 			,TEXT("IasHttpDuration1"), GHttpDurationBuckets[1]
 			,TEXT("IasHttpDuration2"), GHttpDurationBuckets[2]
@@ -513,6 +530,11 @@ void FOnDemandIoBackendStats::OnIoRequestError()
 	CSV_CUSTOM_STAT_DEFINED(FrameIoRequestErrorCount, int32(GIoRequestErrorCount.Get()), ECsvCustomStatOp::Set);
 }
 
+void FOnDemandIoBackendStats::OnIoDecodeError()
+{
+	GIoDecodeErrorCount.Add(1);
+}
+
 void FOnDemandIoBackendStats::OnCacheError()
 {
 	GCacheErrorCount.Add(1);
@@ -560,6 +582,16 @@ void FOnDemandIoBackendStats::OnCacheSetMaxBytes(uint64 TotalSize)
 	GCacheMaxBytes = TotalSize;
 }
 
+void FOnDemandIoBackendStats::OnHttpConnected()
+{
+	GHttpConnectCount.Add(1);
+}
+
+void FOnDemandIoBackendStats::OnHttpDisconnected()
+{
+	GHttpDisconnectCount.Add(1);
+}
+
 void FOnDemandIoBackendStats::OnHttpEnqueue()
 {
 	GHttpPendingCount.Add(1);
diff --git a/Engine/Source/Runtime/Experimental/IoStoreOnDemand/Private/Statistics.h b/Engine/Source/Runtime/Experimental/IoStoreOnDemand/Private/Statistics.h
--- a/Engine/Source/Runtime/Experimental/IoStoreOnDemand/Private/Statistics.h
+++ b/Engine/Source/Runtime/Experimental/IoStoreOnDemand/Private/Statistics.h
@@ -33,13 +33,18 @@ public:
 	FOnDemandIoBackendStats() IAS_STATISTICS_IMPL()
 	~FOnDemandIoBackendStats() IAS_STATISTICS_IMPL()
 
-	void ReportAnalytics(TArray<FAnalyticsEventAttribute>& OutAnalyticsArray) const IAS_STATISTICS_IMPL()
+	/** Report analytics not directly associated with a specific endpoint */
+	void ReportGeneralAnalytics(TArray<FAnalyticsEventAttribute>& OutAnalyticsArray) const IAS_STATISTICS_IMPL()
+	/** Report analytics for the current endpoint */
+	void ReportEndPointAnalytics(TArray<FAnalyticsEventAttribute>& OutAnalyticsArray) const IAS_STATISTICS_IMPL()
 
 	void OnIoRequestEnqueue() IAS_STATISTICS_IMPL()
 	void OnIoRequestComplete(uint64 Size, uint64 DurationMs) IAS_STATISTICS_IMPL()
 	void OnIoRequestCancel() IAS_STATISTICS_IMPL()
 	void OnIoRequestError() IAS_STATISTICS_IMPL()
 
+	void OnIoDecodeError() IAS_STATISTICS_IMPL()
+
 	void OnCacheError() IAS_STATISTICS_IMPL()
 	void OnCacheGet(uint64 DataSize) IAS_STATISTICS_IMPL()
 	void OnCachePut() IAS_STATISTICS_IMPL()
@@ -50,6 +55,9 @@ public:
 	void OnCacheWriteBytes(uint64 WriteSize) IAS_STATISTICS_IMPL()
 	void OnCacheSetMaxBytes(uint64 TotalSize) IAS_STATISTICS_IMPL()
 
+	void OnHttpConnected() IAS_STATISTICS_IMPL()
+	void OnHttpDisconnected() IAS_STATISTICS_IMPL()
+
 	void OnHttpEnqueue() IAS_STATISTICS_IMPL()
 	void OnHttpCancel() IAS_STATISTICS_IMPL()
 	void OnHttpDequeue() IAS_STATISTICS_IMPL()
