Geometry collection : remove unused DummyBodyInstance property
Can save up to 400 bytes per component ( maybe lower based on which memory pool the component ends up be allocated in  )

diff --git a/Engine/Source/Runtime/Experimental/GeometryCollectionEngine/Private/GeometryCollection/GeometryCollectionComponent.cpp b/Engine/Source/Runtime/Experimental/GeometryCollectionEngine/Private/GeometryCollection/GeometryCollectionComponent.cpp
--- a/Engine/Source/Runtime/Experimental/GeometryCollectionEngine/Private/GeometryCollection/GeometryCollectionComponent.cpp
+++ b/Engine/Source/Runtime/Experimental/GeometryCollectionEngine/Private/GeometryCollection/GeometryCollectionComponent.cpp
@@ -1219,7 +1219,7 @@ void UGeometryCollectionComponent::SetNotifyGlobalCrumblings(bool bNewNotifyGlob
 
 FBodyInstance* UGeometryCollectionComponent::GetBodyInstance(FName BoneName /*= NAME_None*/, bool bGetWelded /*= true*/, int32 Index /*=INDEX_NONE*/) const
 {
-	return nullptr;// const_cast<FBodyInstance*>(&DummyBodyInstance);
+	return nullptr;
 }
 
 void UGeometryCollectionComponent::SetNotifyRigidBodyCollision(bool bNewNotifyRigidBodyCollision)
@@ -4723,11 +4723,6 @@ void UGeometryCollectionComponent::OnDestroyPhysicsState()
 {
 	UActorComponent::OnDestroyPhysicsState();
 
-	if(DummyBodyInstance.IsValidBodyInstance())
-	{
-		DummyBodyInstance.TermBody();
-	}
-
 	// we need to unregister the events because it relies on the proxy internally 
 	if (EventDispatcher)
 	{
diff --git a/Engine/Source/Runtime/Experimental/GeometryCollectionEngine/Public/GeometryCollection/GeometryCollectionComponent.h b/Engine/Source/Runtime/Experimental/GeometryCollectionEngine/Public/GeometryCollection/GeometryCollectionComponent.h
--- a/Engine/Source/Runtime/Experimental/GeometryCollectionEngine/Public/GeometryCollection/GeometryCollectionComponent.h
+++ b/Engine/Source/Runtime/Experimental/GeometryCollectionEngine/Public/GeometryCollection/GeometryCollectionComponent.h
@@ -1697,10 +1697,6 @@ private:
 	FGeometryCollectionPhysicsProxy* PhysicsProxy;
 	TUniquePtr<FGeometryDynamicCollection> DynamicCollection;
 
-	// Temporary dummies to interface with Physx expectations of the SQ syatem
-	friend class FGeometryCollectionSQAccelerator;
-	FBodyInstance DummyBodyInstance;
-
 	// Temporary storage for body setup in order to initialise a dummy body instance
 	UPROPERTY(Transient)
 	TObjectPtr<UBodySetup> DummyBodySetup;
