Chaos - ConstGenericParticleHandle float accessors

diff --git a/Engine/Source/Runtime/Experimental/Chaos/Private/Chaos/ParticleHandle.cpp b/Engine/Source/Runtime/Experimental/Chaos/Private/Chaos/ParticleHandle.cpp
--- a/Engine/Source/Runtime/Experimental/Chaos/Private/Chaos/ParticleHandle.cpp
+++ b/Engine/Source/Runtime/Experimental/Chaos/Private/Chaos/ParticleHandle.cpp
@@ -329,6 +329,7 @@ namespace Chaos
 	}
 
 	const FVec3 FGenericParticleHandleImp::ZeroVector = FVec3(0);
+	const FVec3f FGenericParticleHandleImp::ZeroVectorf = FVec3f(0);
 	const FRotation3 FGenericParticleHandleImp::IdentityRotation = FRotation3(FQuat::Identity);
 	const FMatrix33 FGenericParticleHandleImp::ZeroMatrix = FMatrix33(0);
 	const TUniquePtr<FBVHParticles> FGenericParticleHandleImp::NullBVHParticles = TUniquePtr<FBVHParticles>();
diff --git a/Engine/Source/Runtime/Experimental/Chaos/Public/Chaos/ParticleHandle.h b/Engine/Source/Runtime/Experimental/Chaos/Public/Chaos/ParticleHandle.h
--- a/Engine/Source/Runtime/Experimental/Chaos/Public/Chaos/ParticleHandle.h
+++ b/Engine/Source/Runtime/Experimental/Chaos/Public/Chaos/ParticleHandle.h
@@ -1755,12 +1755,21 @@ public:
 	auto& AuxilaryValue(Container& AuxContainer) { return MHandle->AuxilaryValue(AuxContainer); }
 
 	// Kinematic Particles
-	const FVec3 V() const { return (MHandle->CastToKinematicParticle()) ? MHandle->CastToKinematicParticle()->GetV() : ZeroVector; }
-	const FVec3 W() const { return (MHandle->CastToKinematicParticle()) ? MHandle->CastToKinematicParticle()->GetW() : ZeroVector; }
+	const FVec3 V() const { return GetV(); }
+	const FVec3 W() const { return GetW(); }
+
+	const FVec3 GetV() const { return (MHandle->CastToKinematicParticle()) ? MHandle->CastToKinematicParticle()->GetV() : ZeroVector; }
+	const FVec3 GetW() const { return (MHandle->CastToKinematicParticle()) ? MHandle->CastToKinematicParticle()->GetW() : ZeroVector; }
 
 	void SetV(const FVec3& InV) { if (MHandle->CastToKinematicParticle()) { MHandle->CastToKinematicParticle()->SetV(InV); } }
 	void SetW(const FVec3& InW) { if (MHandle->CastToKinematicParticle()) { MHandle->CastToKinematicParticle()->SetW(InW); } }
 
+	const FVec3f GetVf() const { return (MHandle->CastToKinematicParticle()) ? MHandle->CastToKinematicParticle()->GetVf() : ZeroVectorf; }
+	const FVec3f GetWf() const { return (MHandle->CastToKinematicParticle()) ? MHandle->CastToKinematicParticle()->GetWf() : ZeroVectorf; }
+
+	void SetVf(const FVec3f& InV) { if (MHandle->CastToKinematicParticle()) { MHandle->CastToKinematicParticle()->SetVf(InV); } }
+	void SetWf(const FVec3f& InW) { if (MHandle->CastToKinematicParticle()) { MHandle->CastToKinematicParticle()->SetWf(InW); } }
+
 	const FKinematicTarget& KinematicTarget() const { return (MHandle->CastToKinematicParticle())? MHandle->CastToKinematicParticle()->KinematicTarget() : EmptyKinematicTarget; }
 
 	// Dynamic Particles
@@ -1847,6 +1856,25 @@ public:
 		return ZeroVector;
 	}
 
+	const FVec3f GetPreVf() const
+	{
+		if (MHandle->CastToRigidParticle())
+		{
+			return MHandle->CastToRigidParticle()->GetPreVf();
+		}
+
+		return ZeroVectorf;
+	}
+
+	const FVec3f GetPreWf() const
+	{
+		if (MHandle->CastToRigidParticle())
+		{
+			return MHandle->CastToRigidParticle()->GetPreWf();
+		}
+		return ZeroVectorf;
+	}
+
 	int32 SolverBodyIndex() const
 	{
 		if (MHandle->CastToRigidParticle())
@@ -2415,7 +2443,9 @@ public:
 private:
 	FGeometryParticleHandle* MHandle;
 
+	// So we can return vectors by const ref in cases where there's no real value. E.g., Velocity of a Static.
 	static CHAOS_API const FVec3 ZeroVector;
+	static CHAOS_API const FVec3f ZeroVectorf;
 	static CHAOS_API const FRotation3 IdentityRotation;
 	static CHAOS_API const FMatrix33 ZeroMatrix;
 	static CHAOS_API const TUniquePtr<FBVHParticles> NullBVHParticles;
