Chaos - add MeshContactGenerator support to tri mesh

diff --git a/Engine/Source/Runtime/Experimental/Chaos/Private/Chaos/TriangleMeshImplicitObject.cpp b/Engine/Source/Runtime/Experimental/Chaos/Private/Chaos/TriangleMeshImplicitObject.cpp
--- a/Engine/Source/Runtime/Experimental/Chaos/Private/Chaos/TriangleMeshImplicitObject.cpp
+++ b/Engine/Source/Runtime/Experimental/Chaos/Private/Chaos/TriangleMeshImplicitObject.cpp
@@ -2,6 +2,7 @@
 #include "Chaos/TriangleMeshImplicitObject.h"
 #include "Chaos/Collision/ContactPoint.h"
 #include "Chaos/Collision/ContactPointsMiscShapes.h"
+#include "Chaos/Collision/MeshContactGenerator.h"
 #include "Chaos/Collision/TriangleOverlap.h"
 #include "Chaos/Collision/PBDCollisionConstraint.h"
 #include "Chaos/CollisionOneShotManifolds.h"
@@ -851,6 +852,38 @@ bool FTriangleMeshImplicitObject::Raycast(const FVec3& StartPoint, const FVec3&
 	}
 }
 
+void FTriangleMeshImplicitObject::CollectTriangles(const FAABB3& MeshQueryBounds, const FRigidTransform3& MeshToObjectTransform, const FAABB3& ObjectBounds, Private::FMeshContactGenerator& Collector) const
+{
+	TArray<int32> OverlapIndices;
+	FindOverlappingTriangles(MeshQueryBounds, OverlapIndices);
+
+	const bool bStandardWinding = ((MeshToObjectTransform.GetScale3D().X * MeshToObjectTransform.GetScale3D().Y * MeshToObjectTransform.GetScale3D().Z) >= FReal(0));
+
+	Collector.BeginCollect(OverlapIndices.Num());
+
+	for (int32 OverlapIndex = 0; OverlapIndex < OverlapIndices.Num(); ++OverlapIndex)
+	{
+		const int32 TriangleIndex = OverlapIndices[OverlapIndex];
+		FTriangle Triangle;
+		int32 VertexIndex0, VertexIndex1, VertexIndex2;
+		GetTransformedTriangle(TriangleIndex, MeshToObjectTransform, Triangle, VertexIndex0, VertexIndex1, VertexIndex2);
+
+		if (!bStandardWinding)
+		{
+			Triangle.ReverseWinding();
+			Swap(VertexIndex1, VertexIndex2);
+		}
+
+		const FAABB3 TriBounds = FAABB3::FromPoints(Triangle.GetVertex(0), Triangle.GetVertex(1), Triangle.GetVertex(2));
+		if (TriBounds.Intersects(ObjectBounds))
+		{
+			Collector.AddTriangle(Triangle, TriangleIndex, VertexIndex0, VertexIndex1, VertexIndex2);
+		}
+	}
+
+	Collector.EndCollect();
+}
+
 template <typename QueryGeomType>
 bool FTriangleMeshImplicitObject::GJKContactPointImp(const QueryGeomType& QueryGeom, const FRigidTransform3& QueryTM, const FReal WorldThickness, FVec3& Location, FVec3& Normal, FReal& OutContactPhi, int32& OutFaceIndex, FVec3 TriMeshScale) const
 {
diff --git a/Engine/Source/Runtime/Experimental/Chaos/Public/Chaos/TriangleMeshImplicitObject.h b/Engine/Source/Runtime/Experimental/Chaos/Public/Chaos/TriangleMeshImplicitObject.h
--- a/Engine/Source/Runtime/Experimental/Chaos/Public/Chaos/TriangleMeshImplicitObject.h
+++ b/Engine/Source/Runtime/Experimental/Chaos/Public/Chaos/TriangleMeshImplicitObject.h
@@ -21,6 +21,11 @@
 
 #include <type_traits>
 
+namespace Chaos::Private
+{
+	class FMeshContactGenerator;
+}
+
 namespace Chaos
 {
 	extern CHAOS_API bool TriMeshPerPolySupport;
@@ -786,6 +791,9 @@ namespace Chaos
 			}
 		}
 
+		// Internal: do not use - this API will change as we optimize mesh collision
+		void CollectTriangles(const FAABB3& MeshQueryBounds, const FRigidTransform3& MeshToObjectTransform, const FAABB3& ObjectBounds, Private::FMeshContactGenerator& Collector) const;
+
 	private:
 		using BVHType = TAABBTree<int32, TAABBTreeLeafArray<int32, /*bComputeBounds=*/ false, FRealSingle>, /*bMutable=*/false, FRealSingle>;
 		CHAOS_API void RebuildFastBVHFromTree(const BVHType& BVH);
