Constraints: dependencies + evaluation graph during manipulation

    - constraint dependencies are now defined when a constraint is added to the subsystem (previously, they were only defined at creation)
    - an evaluation graph is lazily created when a parent/child is manipulated:
        - it represents constraints sorted by tick dependency and provides a way of evaluating them in the right order without waiting for the engine's next tick
        - constraints to be evaluated when moving a parent/child in the editor are marked for evaluation
        - the evaluation graph is flushed once the manipulation has been performed to evaluate constraints in the correct order (from highest to lowest priority).

diff --git a/Engine/Plugins/Animation/ControlRig/Source/ControlRig/Private/Constraints/ControlRigTransformableHandle.cpp b/Engine/Plugins/Animation/ControlRig/Source/ControlRig/Private/Constraints/ControlRigTransformableHandle.cpp
--- a/Engine/Plugins/Animation/ControlRig/Source/ControlRig/Private/Constraints/ControlRigTransformableHandle.cpp
+++ b/Engine/Plugins/Animation/ControlRig/Source/ControlRig/Private/Constraints/ControlRigTransformableHandle.cpp
@@ -6,7 +6,6 @@
 #include "Components/SkeletalMeshComponent.h"
 #include "ControlRig.h"
 #include "ControlRigComponent.h"
-#include "ControlRig.h"
 #include "IControlRigObjectBinding.h"
 #include "TransformableHandleUtils.h"
 #include "Rigs/RigHierarchyElements.h"
@@ -353,6 +352,12 @@ void UTransformableControlHandle::RegisterDelegates()
 		{
 			Hierarchy->OnModified().AddUObject(this, &UTransformableControlHandle::OnHierarchyModified);
 		}
+
+		// NOTE BINDER: this has to be done before binding UTransformableControlHandle::OnControlModified
+		if (!ControlRig->ControlModified().IsBoundToObject(&GetEvaluationBinding()))
+		{
+			ControlRig->ControlModified().AddRaw(&GetEvaluationBinding(), &FControlEvaluationGraphBinding::HandleControlModified);
+		}
 		
 		ControlRig->ControlModified().AddUObject(this, &UTransformableControlHandle::OnControlModified);
 		if (!ControlRig->ControlRigBound().IsBoundToObject(this))
@@ -428,6 +433,7 @@ void UTransformableControlHandle::OnControlModified(
 
 		if (InControl->GetFName() == ControlName)
 		{	// if that handle is wrapping InControl
+			GetEvaluationBinding().bPendingFlush = true;
 			Notify(Event);
 		}
 		else if (Event == EHandleEvent::GlobalTransformUpdated)
@@ -435,6 +441,7 @@ void UTransformableControlHandle::OnControlModified(
 			// the control being modified is not the one wrapped by this handle 
 			if (const FRigControlElement* Control = ControlRig->FindControl(ControlName))
 			{
+				GetEvaluationBinding().bPendingFlush = true;
 				static constexpr  bool bPreTick = true;
 				Notify(EHandleEvent::UpperDependencyUpdated, bPreTick);
 			}
@@ -623,3 +630,30 @@ void UTransformableControlHandle::OnObjectsReplaced(const TMap<UObject*, UObject
 }
 
 #endif
+
+FControlEvaluationGraphBinding& UTransformableControlHandle::GetEvaluationBinding()
+{
+	static FControlEvaluationGraphBinding EvaluationBinding;
+	return EvaluationBinding;
+}
+
+void FControlEvaluationGraphBinding::HandleControlModified(UControlRig* InControlRig, FRigControlElement* InControl, const FRigControlModifiedContext& InContext)
+{
+	if (!bPendingFlush || !InContext.bConstraintUpdate)
+	{
+		return;
+	}
+	
+	if (!InControlRig || !InControl)
+	{
+		return;
+	}
+
+	// flush all pending evaluations if any
+	if (UWorld* World = InControlRig->GetWorld())
+	{
+		const FConstraintsManagerController& Controller = FConstraintsManagerController::Get(World);
+		Controller.FlushEvaluationGraph();
+	}
+	bPendingFlush = false;
+}
diff --git a/Engine/Plugins/Animation/ControlRig/Source/ControlRig/Public/Constraints/ControlRigTransformableHandle.h b/Engine/Plugins/Animation/ControlRig/Source/ControlRig/Public/Constraints/ControlRigTransformableHandle.h
--- a/Engine/Plugins/Animation/ControlRig/Source/ControlRig/Public/Constraints/ControlRigTransformableHandle.h
+++ b/Engine/Plugins/Animation/ControlRig/Source/ControlRig/Public/Constraints/ControlRigTransformableHandle.h
@@ -15,6 +15,18 @@ class USkeletalMeshComponent;
 class UControlRigComponent;
 class URigHierarchy;
 
+/**
+ * FControlEvaluationGraphBinding
+ */
+struct CONTROLRIG_API FControlEvaluationGraphBinding
+{
+	void HandleControlModified(
+		UControlRig* InControlRig,
+		FRigControlElement* InControl,
+		const FRigControlModifiedContext& InContext);
+	bool bPendingFlush = false;
+};
+
 /**
  * UTransformableControlHandle
  */
@@ -126,6 +138,8 @@ private:
 
 	void OnControlRigBound(UControlRig* InControlRig);
 	void OnObjectBoundToControlRig(UObject* InObject);
+	
+	static FControlEvaluationGraphBinding& GetEvaluationBinding();
 
 #if WITH_EDITOR
 	/** @todo document */
diff --git a/Engine/Source/Runtime/Experimental/Animation/Constraints/Private/ConstraintSubsystem.cpp b/Engine/Source/Runtime/Experimental/Animation/Constraints/Private/ConstraintSubsystem.cpp
--- a/Engine/Source/Runtime/Experimental/Animation/Constraints/Private/ConstraintSubsystem.cpp
+++ b/Engine/Source/Runtime/Experimental/Animation/Constraints/Private/ConstraintSubsystem.cpp
@@ -114,8 +114,7 @@ FConstraintsInWorld& UConstraintSubsystem::ConstraintsInWorldFindOrAdd(UWorld* I
 	}
 	FConstraintsInWorld NewCInW;
 	NewCInW.World = InWorld;
-	int32 Index = ConstraintsInWorld.Add(NewCInW);
-	return ConstraintsInWorld[Index];
+	return ConstraintsInWorld.Emplace_GetRef(MoveTemp(NewCInW));
 }
 
 TArray<TWeakObjectPtr<UTickableConstraint>> UConstraintSubsystem::GetConstraints(UWorld* InWorld) const
@@ -145,6 +144,7 @@ void UConstraintSubsystem::AddConstraint(UWorld* InWorld, UTickableConstraint* I
 	if (Constraints.Constraints.Contains(InConstraint) == false)
 	{
 		Constraints.Constraints.Emplace(InConstraint);
+		Constraints.InvalidateGraph();
 	}
 	OnConstraintAddedToSystem_BP.Broadcast(this, InConstraint);
 }
@@ -162,6 +162,7 @@ void UConstraintSubsystem::RemoveConstraint(UWorld* InWorld, UTickableConstraint
 	if (FConstraintsInWorld* Constraints = ConstraintsInWorldFind(InWorld))
 	{
 		Constraints->Constraints.Remove(InConstraint);
+		Constraints->InvalidateGraph();
 	}
 }
 
@@ -212,6 +213,11 @@ bool UConstraintSubsystem::HasConstraint(UWorld* InWorld, UTickableConstraint* I
 	return  Constraints.Contains(InConstraint);
 }
 
+FConstraintsEvaluationGraph& UConstraintSubsystem::GetEvaluationGraph(UWorld* InWorld)
+{
+	return ConstraintsInWorldFindOrAdd(InWorld).GetEvaluationGraph();
+}
+
 void UConstraintSubsystem::InvalidateConstraints()
 {
 	bNeedsCleanup = true;
@@ -223,9 +229,13 @@ void UConstraintSubsystem::PostEditUndo()
 {
 	Super::PostEditUndo();
 
+	for (FConstraintsInWorld& ConstsInWorld: ConstraintsInWorld)
+	{
+		ConstsInWorld.InvalidateGraph();
+	}
+	
 	const FConstraintsManagerController& Controller = FConstraintsManagerController::Get(GetWorld());
 	Controller.Notify(EConstraintsManagerNotifyType::ManagerUpdated, this);
-
 }
 #endif
 
@@ -272,6 +282,7 @@ void UConstraintSubsystem::CleanupInvalidConstraints() const
 			{
 				return !InConstraint.IsValid() || InConstraint.IsStale();
 			});
+			WorldConstraints.InvalidateGraph();
 		}
 		bNeedsCleanup = false;
 	}
@@ -292,6 +303,7 @@ void FConstraintsInWorld::RemoveConstraints(UWorld* InWorld)
 		}
 	}
 	Constraints.SetNum(0);
+	InvalidateGraph();
 }
 
 
@@ -301,7 +313,23 @@ void FConstraintsInWorld::Init(UWorld* InWorld)
 	{
 		World = InWorld;
 	}
+	InvalidateGraph();
 }
 
+FConstraintsEvaluationGraph& FConstraintsInWorld::GetEvaluationGraph()
+{
+	if (!EvaluationGraph)
+	{
+		ensure(World);
+		EvaluationGraph = MakeShared<FConstraintsEvaluationGraph>(*this);
+	}
+	return *EvaluationGraph;
+}
 
-
+void FConstraintsInWorld::InvalidateGraph() const
+{
+	if (EvaluationGraph)
+	{
+		EvaluationGraph->InvalidateData();
+	}
+}
\ No newline at end of file
diff --git a/Engine/Source/Runtime/Experimental/Animation/Constraints/Private/ConstraintSubsystem.h b/Engine/Source/Runtime/Experimental/Animation/Constraints/Private/ConstraintSubsystem.h
--- a/Engine/Source/Runtime/Experimental/Animation/Constraints/Private/ConstraintSubsystem.h
+++ b/Engine/Source/Runtime/Experimental/Animation/Constraints/Private/ConstraintSubsystem.h
@@ -3,6 +3,7 @@
 #pragma once
 
 #include "CoreMinimal.h"
+#include "ConstraintsEvaluationGraph.h"
 #include "UObject/Object.h"
 #include "Engine/EngineBaseTypes.h"
 #include "Subsystems/EngineSubsystem.h"
@@ -27,6 +28,12 @@ public:
 
 	CONSTRAINTS_API void Init(UWorld* World);
 	CONSTRAINTS_API void RemoveConstraints(UWorld* World);
+
+	FConstraintsEvaluationGraph& GetEvaluationGraph();
+	void InvalidateGraph() const;
+
+private:
+	TSharedPtr<FConstraintsEvaluationGraph> EvaluationGraph = nullptr;
 };
 
 
@@ -57,6 +64,8 @@ public:
 	CONSTRAINTS_API void AddConstraint(UWorld* InWorld, UTickableConstraint* InConstraint);
 	CONSTRAINTS_API void RemoveConstraint(UWorld* InWorld, UTickableConstraint* InConstraint, bool bDoNoCompensate);
 	CONSTRAINTS_API bool HasConstraint(UWorld* InWorld, UTickableConstraint* InConstraint);
+
+	FConstraintsEvaluationGraph& GetEvaluationGraph(UWorld* InWorld);
 	
 	void InvalidateConstraints();
 
diff --git a/Engine/Source/Runtime/Experimental/Animation/Constraints/Private/ConstraintsEvaluationGraph.cpp b/Engine/Source/Runtime/Experimental/Animation/Constraints/Private/ConstraintsEvaluationGraph.cpp
new file mode 100644
--- /dev/null
+++ b/Engine/Source/Runtime/Experimental/Animation/Constraints/Private/ConstraintsEvaluationGraph.cpp
@@ -0,0 +1,300 @@
+﻿// Copyright Epic Games, Inc. All Rights Reserved.
+
+#include "ConstraintsEvaluationGraph.h"
+
+#include "ConstraintsManager.h"
+#include "ConstraintSubsystem.h"
+
+namespace ConstraintsEvaluationGraph
+{
+
+static bool	bUseEvaluationGraph = true;
+static FAutoConsoleVariableRef CVarUseEvaluationGraph(
+	TEXT("Constraints.UseEvaluationGraph"),
+	bUseEvaluationGraph,
+	TEXT("Use Evaluation Graph to update constraints when manipulating.") );
+	
+static bool	bDebugGraph = false;
+static FAutoConsoleVariableRef CVarDebugEvaluationGraph(
+	TEXT("Constraints.DebugEvaluationGraph"),
+	bDebugGraph,
+	TEXT("Print debug info about constraitns evaluation graph.") );
+	
+}
+
+bool FConstraintsEvaluationGraph::UseEvaluationGraph()
+{
+	return ConstraintsEvaluationGraph::bUseEvaluationGraph;
+}
+
+FConstraintNode& FConstraintsEvaluationGraph::GetNode(const TWeakObjectPtr<UTickableConstraint>& InConstraint)
+{
+	const int32 Found = Nodes.IndexOfByPredicate([InConstraint](const FConstraintNode& Node)
+	{
+		return InConstraint.IsValid() && Node.ConstraintID == InConstraint->ConstraintID; 
+	});
+
+	if (Found != INDEX_NONE)
+	{
+		return Nodes[Found];
+	}
+
+	FConstraintNode Node;
+	Node.ConstraintID = InConstraint->ConstraintID;
+	Node.ConstraintTick = &InConstraint->GetTickFunction(ConstraintsInWorld.World);
+	return Nodes.Emplace_GetRef(MoveTemp(Node)); 
+}
+
+FConstraintNode* FConstraintsEvaluationGraph::FindNode(const TWeakObjectPtr<UTickableConstraint>& InConstraint)
+{
+	if (!InConstraint.IsValid())
+	{
+		return nullptr;
+	}
+	
+	return Nodes.FindByPredicate([InConstraint](const FConstraintNode& Node)
+	{
+		return InConstraint.IsValid() && Node.ConstraintID == InConstraint->ConstraintID; 
+	});
+}
+
+void FConstraintsEvaluationGraph::FlushPendingEvaluations()
+{
+	if (!Nodes.IsEmpty())
+	{
+		if (ensure(!bDataInvalid))
+		{
+			if (ConstraintsEvaluationGraph::bDebugGraph)
+			{
+				UE_LOG(LogTemp, Warning, TEXT("Flush Constraints Evaluation Graph"));
+			}
+			
+			for (FConstraintNode& Node: Nodes)
+			{
+				if (Node.bMarkedForEvaluation)
+				{
+					Evaluate(&Node);
+				}
+			}
+
+			const bool bHasNodesToEvaluate = Nodes.ContainsByPredicate([](const FConstraintNode& Node)
+			{
+				return Node.bMarkedForEvaluation;
+			});
+			ensure(!bHasNodesToEvaluate);
+		}
+	}
+}
+
+void FConstraintsEvaluationGraph::Rebuild()
+{
+	Nodes.Empty();
+	
+	if (!ensure(ConstraintsInWorld.World.Get()))
+	{
+		return;
+	}
+
+	const TArray<TWeakObjectPtr<UTickableConstraint>>& Constraints = ConstraintsInWorld.Constraints;
+	if (Constraints.IsEmpty())
+	{
+		return;
+	}
+
+	// build nodes
+	for (int32 ConstraintIndex = 0; ConstraintIndex < Constraints.Num(); ConstraintIndex++)
+	{
+		if (Constraints[ConstraintIndex].IsValid())
+		{
+			FConstraintNode& Node = GetNode(Constraints[ConstraintIndex]);
+			Node.ConstraintIndex = ConstraintIndex;
+		}
+	}
+
+	// sort nodes by tick dependencies
+	auto EvaluationOrderPredicate = [](const FConstraintNode& LHS, const FConstraintNode& RHS)
+	{
+		const TArray<FTickPrerequisite>& RHSPrerex = RHS.ConstraintTick->GetPrerequisites();
+		const FConstraintTickFunction* LHSTickFunction = LHS.ConstraintTick;
+		const bool bIsLHSAPrerexOfRHS = RHSPrerex.ContainsByPredicate([LHSTickFunction](const FTickPrerequisite& Prerex)
+		{
+			return Prerex.PrerequisiteTickFunction == LHSTickFunction;
+		});
+		return bIsLHSAPrerexOfRHS;
+	};
+	Algo::StableSort(Nodes, EvaluationOrderPredicate);
+
+	// store node index after re-ordering
+	for (int32 NodeIndex = 0; NodeIndex < Nodes.Num(); NodeIndex++)
+	{
+		Nodes[NodeIndex].NodeIndex = NodeIndex;
+	}
+
+	// build edges	
+	for (FConstraintNode& Node: Nodes)
+	{
+		const TArray<FTickPrerequisite>& Prerequisites = Node.ConstraintTick->GetPrerequisites();
+		for (const FTickPrerequisite& Prerex: Prerequisites)
+		{
+			if (const FTickFunction* PrerexFunction = Prerex.Get())
+			{
+				FConstraintNode* PrerexNode = Nodes.FindByPredicate([PrerexFunction, Node](const FConstraintNode& OtherNode)
+				{
+					if (OtherNode.NodeIndex != Node.NodeIndex)
+					{
+						if (OtherNode.ConstraintTick == PrerexFunction)
+						{
+							return true;
+						}
+					}
+					return false;
+				});
+
+				if (PrerexNode)
+				{
+					if (!PrerexNode->Parents.Contains(Node.NodeIndex))
+					{
+						Node.Parents.Add(PrerexNode->NodeIndex);
+					}
+					else
+					{
+						// we may have create a cycle
+						ensure(false);
+					}
+
+					if (!Node.Children.Contains(PrerexNode->NodeIndex))
+					{
+						PrerexNode->Children.Add(Node.NodeIndex);
+					}
+					else
+					{
+						// we may have create a cycle
+						ensure(false);
+					}
+				}
+			}
+		}
+	}
+	
+	bDataInvalid = false;
+
+	Dump();
+}
+
+void FConstraintsEvaluationGraph::Evaluate(const TWeakObjectPtr<UTickableConstraint>& InConstraint)
+{
+	if (bDataInvalid)
+	{
+		Rebuild();
+	}
+
+	if (Nodes.IsEmpty())
+	{
+		return;
+	}
+	
+	if (FConstraintNode* Node = FindNode(InConstraint))
+	{
+		Evaluate(Node);
+	}
+}
+
+void FConstraintsEvaluationGraph::Evaluate(FConstraintNode* InNode)
+{
+	const TArray<TWeakObjectPtr<UTickableConstraint>>& Constraints = ConstraintsInWorld.Constraints;
+	if (!Constraints.IsValidIndex(InNode->ConstraintIndex))
+	{
+		return;
+	}
+
+	const TWeakObjectPtr<UTickableConstraint> Constraint = Constraints[InNode->ConstraintIndex];
+	if (!Constraint.IsValid())
+	{
+		return;
+	}
+
+	// evaluate
+	if (Constraint->IsFullyActive() && InNode->ConstraintTick->IsTickFunctionRegistered() && InNode->ConstraintTick->IsTickFunctionEnabled())
+	{
+		Constraint->Evaluate();
+	}
+	InNode->bMarkedForEvaluation = false;
+
+	// evaluate dependencies
+	for (const uint32 ChildIndex: InNode->Children)
+	{
+		if (ensure(Nodes.IsValidIndex(ChildIndex)))
+		{
+			Evaluate(&Nodes[ChildIndex]);
+		}
+	}
+}
+
+void FConstraintsEvaluationGraph::InvalidateData()
+{
+	bDataInvalid = true;
+	Nodes.Empty();
+}
+
+void FConstraintsEvaluationGraph::MarkForEvaluation(const TWeakObjectPtr<UTickableConstraint>& InConstraint)
+{
+	if (bDataInvalid)
+	{
+		Rebuild();
+	}
+	
+	if (FConstraintNode* Node = FindNode(InConstraint))
+	{
+		auto GetConstraintLabel = [](const TWeakObjectPtr<UTickableConstraint>& InConstraint)
+		{
+#if WITH_EDITOR
+			return InConstraint->GetFullLabel();
+#else
+			return InConstraint->GetName();
+#endif		
+		};
+
+		if (ConstraintsEvaluationGraph::bDebugGraph)
+		{
+			UE_LOG(LogTemp, Warning, TEXT("Mark %s For Evaluation"), *GetConstraintLabel(InConstraint));
+		}
+		
+		Node->bMarkedForEvaluation = true;
+	}
+}
+
+void FConstraintsEvaluationGraph::Dump() const
+{
+	if (!ConstraintsEvaluationGraph::bDebugGraph)
+	{
+		return;
+	}
+	
+	auto GetConstraintLabel = [](const TWeakObjectPtr<UTickableConstraint>& InConstraint)
+	{
+#if WITH_EDITOR
+		return InConstraint->GetFullLabel();
+#else
+		return InConstraint->GetName();
+#endif		
+	};
+	
+	const TArray<TWeakObjectPtr<UTickableConstraint>>& Constraints = ConstraintsInWorld.Constraints;
+	UE_LOG(LogTemp, Warning, TEXT("Nb Constraints = %d"), Constraints.Num());
+	for (int32 ConstraintIndex = 0; ConstraintIndex < Constraints.Num(); ConstraintIndex++)
+	{
+		if (Constraints[ConstraintIndex].IsValid())
+		{
+			UE_LOG(LogTemp, Warning, TEXT("\tConstraint[%d] = %s"), ConstraintIndex, *GetConstraintLabel(Constraints[ConstraintIndex]));
+		}
+	}
+
+	UE_LOG(LogTemp, Warning, TEXT("Nb Nodes = %d"), Nodes.Num());
+	for (int32 NodeIndex = 0; NodeIndex < Nodes.Num(); NodeIndex++)
+	{
+		const FConstraintNode& Node = Nodes[NodeIndex];
+		ensure(Constraints.IsValidIndex(Node.ConstraintIndex));
+		ensure(Node.NodeIndex == NodeIndex);
+		UE_LOG(LogTemp, Warning, TEXT("\tNode[%d] = %s [%d]"), NodeIndex, *GetConstraintLabel(Constraints[Node.ConstraintIndex]), Node.ConstraintIndex);
+	}
+}
\ No newline at end of file
diff --git a/Engine/Source/Runtime/Experimental/Animation/Constraints/Private/ConstraintsManager.cpp b/Engine/Source/Runtime/Experimental/Animation/Constraints/Private/ConstraintsManager.cpp
--- a/Engine/Source/Runtime/Experimental/Animation/Constraints/Private/ConstraintsManager.cpp
+++ b/Engine/Source/Runtime/Experimental/Animation/Constraints/Private/ConstraintsManager.cpp
@@ -463,6 +463,9 @@ bool FConstraintsManagerController::AddConstraint(UTickableConstraint* InConstra
 
 	Subsystem->AddConstraint(World,InConstraint);
 
+	// build dependencies
+	InConstraint->AddedToWorld(World);
+
 	// notify
 	Notify(EConstraintsManagerNotifyType::ConstraintAdded, InConstraint);
 
@@ -696,6 +699,8 @@ void FConstraintsManagerController::SetConstraintsDependencies(
 	FConstraintTickFunction& FunctionToTickAfter = Constraints[IndexAfter]->GetTickFunction(World);
 
 	Subsystem->SetConstraintDependencies(&FunctionToTickBefore, &FunctionToTickAfter);
+
+	InvalidateEvaluationGraph();
 }
 
 void FConstraintsManagerController::SetConstraintsDependencies(const struct FGuid& InGuidToTickBefore, const struct FGuid& InGuidToTickAfter) const
@@ -731,6 +736,8 @@ void FConstraintsManagerController::SetConstraintsDependencies(const struct FGui
 	FConstraintTickFunction& FunctionToTickAfter = Constraints[IndexAfter]->GetTickFunction(World);
 
 	Subsystem->SetConstraintDependencies(&FunctionToTickBefore, &FunctionToTickAfter);
+
+	InvalidateEvaluationGraph();
 }
 
 const TArray< TWeakObjectPtr<UTickableConstraint> >& FConstraintsManagerController::GetConstraintsArray() const
@@ -892,6 +899,40 @@ void FConstraintsManagerController::Notify(EConstraintsManagerNotifyType InNotif
 	NotifyDelegate.Broadcast(InNotifyType, InObject);
 }
 
+void FConstraintsManagerController::MarkConstraintForEvaluation(UTickableConstraint* InConstraint) const
+{
+	UConstraintSubsystem* SubSystem = UConstraintSubsystem::Get();
+	if (!SubSystem)
+	{
+		return;
+	}
+	SubSystem->GetEvaluationGraph(World).MarkForEvaluation(InConstraint);
+}
+
+void FConstraintsManagerController::InvalidateEvaluationGraph() const
+{
+	UConstraintSubsystem* SubSystem = UConstraintSubsystem::Get();
+	if (!SubSystem)
+	{
+		return;
+	}
+	SubSystem->GetEvaluationGraph(World).InvalidateData();
+}
+
+void FConstraintsManagerController::FlushEvaluationGraph() const
+{
+	if (!FConstraintsEvaluationGraph::UseEvaluationGraph())
+	{
+		return;
+	}
+	UConstraintSubsystem* SubSystem = UConstraintSubsystem::Get();
+	if (!SubSystem)
+	{
+		return;
+	}
+	SubSystem->GetEvaluationGraph(World).FlushPendingEvaluations();
+}
+
 bool FConstraintsManagerController::DoesExistInAnyWorld(UTickableConstraint* InConstraint)
 {
 	bool bFound = false;
diff --git a/Engine/Source/Runtime/Experimental/Animation/Constraints/Private/TransformConstraint.cpp b/Engine/Source/Runtime/Experimental/Animation/Constraints/Private/TransformConstraint.cpp
--- a/Engine/Source/Runtime/Experimental/Animation/Constraints/Private/TransformConstraint.cpp
+++ b/Engine/Source/Runtime/Experimental/Animation/Constraints/Private/TransformConstraint.cpp
@@ -183,12 +183,10 @@ private:
 		{
 			if (bHasActiveParentConstraint)
 			{
-				// UE_LOG(LogTemp, Warning, TEXT("REMOVE %s prerex on %s."), *TargetObject->GetName(), *InConstraintToUpdate->GetName());
 				InConstraintToUpdate->GetTickFunction(InWorld).RemovePrerequisite(TargetObject, *TargetTickFunction);
 			}
 			else
 			{
-				// UE_LOG(LogTemp, Warning, TEXT("ADD %s prerex on %s."), *TargetObject->GetName(), *InConstraintToUpdate->GetName());
 				InConstraintToUpdate->GetTickFunction(InWorld).AddPrerequisite(TargetObject, *TargetTickFunction);
 			}
 		}
@@ -222,6 +220,11 @@ void PreEvaluateHandle(const TObjectPtr<UTransformableHandle>& InHandle)
 	}
 }
 
+UObject* GetHandleTarget(const TObjectPtr<UTransformableHandle>& InHandle)
+{
+	return IsValid(InHandle) ? InHandle->GetTarget().Get() : nullptr; 
+}
+
 static bool	bIncludeTarget = true;
 static FAutoConsoleVariableRef CVarIncludeTarget(
 	TEXT("Constraints.IncludeTarget"),
@@ -229,33 +232,30 @@ static FAutoConsoleVariableRef CVarIncludeTarget(
 	TEXT("Include target when getting child's existing constraints.")
 	);
 
-void PrintTickFunctions(UWorld* InWorld)
+static bool	bDebugDependencies = false;
+static FAutoConsoleVariableRef CVarDebugDependencies(
+	TEXT("Constraints.DebugDependencies"),
+	bDebugDependencies,
+	TEXT("Print debug info about dependencies when creating a new constraint.") );
+
+FString GetConstraintLabel(const UTickableConstraint* InConstraint)
 {
-	const FConstraintsManagerController& Controller = FConstraintsManagerController::Get(InWorld);
-	const TArray<TWeakObjectPtr<UTickableConstraint>>& Constraint = Controller.GetConstraintsArray();
-	for (const TWeakObjectPtr<UTickableConstraint>& ConstraintPtr: Constraint )
-	{
-		if (ConstraintPtr.IsValid())
-		{
-			FConstraintTickFunction& Function = ConstraintPtr->GetTickFunction(InWorld);
-			UE_LOG(LogTemp, Warning, TEXT("constraint (%p) - function %s (%p)"), ConstraintPtr.Get(), *Function.DiagnosticMessage(), &Function);
-	
-			const TArray<FTickPrerequisite>& Prerequisites = Function.GetPrerequisites();
-			for (const FTickPrerequisite& Prerequisite: Prerequisites)
-			{
-				if (const FTickFunction* TickFunction = Prerequisite.Get())
-				{
-					UE_LOG(LogTemp, Warning, TEXT("\t%s (%p)"), *Prerequisite.PrerequisiteObject->GetFullName(), TickFunction);
-				}
-				else
-				{
-					UE_LOG(LogTemp, Warning, TEXT("\tInvalid Prerequisite"));
-				}
-			}
-		}
-	}
+#if WITH_EDITOR
+	return InConstraint->GetFullLabel();
+#else
+	return InConstraint->GetName();
+#endif		
 }
 
+FString GetHandleLabel(const UTransformableHandle* InHandle)
+{
+#if WITH_EDITOR
+	return InHandle->GetFullLabel();
+#else
+	return InHandle->GetName();
+#endif		
+}
+	
 }
 
 /** 
@@ -661,25 +661,66 @@ void UTickableTransformConstraint::OnHandleModified(UTransformableHandle* InHand
 			SetupDependencies(World);
 			return;
 		}
-		
+	}
+
+	auto MarkForEvaluation = [this, World]()
+	{
+		const FConstraintsManagerController& Controller = FConstraintsManagerController::Get(World);
+		Controller.MarkConstraintForEvaluation(this);
+	};
+
+	auto DirectEvaluation = [this, World]()
+	{
+		const FConstraintTickFunction* ConstraintTick = ConstraintTicks.Find(World->GetCurrentLevel());
+		if (IsFullyActive() && ConstraintTick->IsTickFunctionRegistered() && ConstraintTick->IsTickFunctionEnabled())
+		{
+			Evaluate(true);
+		}
+	};
+	
+	if (InHandle == ChildTRSHandle)
+	{
 		if (InNotification == EHandleEvent::UpperDependencyUpdated)
 		{
-			const FConstraintTickFunction* ConstraintTick = ConstraintTicks.Find(World->GetCurrentLevel());
-			if (IsFullyActive() && ConstraintTick->IsTickFunctionRegistered() && ConstraintTick->IsTickFunctionEnabled())
+			const UObject* ParentTarget = ConstraintLocals::GetHandleTarget(ParentTRSHandle);
+			if (ParentTarget && ParentTarget != Target)
 			{
-				Evaluate(true);
+				if (FConstraintsEvaluationGraph::UseEvaluationGraph())
+				{
+					MarkForEvaluation();
+				}
+				else
+				{
+					DirectEvaluation();
+				}
+				return;
 			}
-			return;
 		}
 	}
 
 	// update the constraint if the parent's transform has been modified 
 	if (InHandle == ParentTRSHandle && InNotification == EHandleEvent::GlobalTransformUpdated)
 	{
-		const FConstraintTickFunction* ConstraintTick = ConstraintTicks.Find(World->GetCurrentLevel());
-		if (IsFullyActive() && ConstraintTick->IsTickFunctionRegistered() && ConstraintTick->IsTickFunctionEnabled())
+		if (FConstraintsEvaluationGraph::UseEvaluationGraph())
 		{
-			Evaluate();
+			MarkForEvaluation();
+		}
+		else
+		{
+			DirectEvaluation();
+		}
+		return;
+	}
+
+	if (InHandle == ParentTRSHandle && InNotification == EHandleEvent::UpperDependencyUpdated)
+	{
+		if (FConstraintsEvaluationGraph::UseEvaluationGraph())
+		{
+			MarkForEvaluation();
+		}
+		else
+		{
+			DirectEvaluation();
 		}
 		return;
 	}
@@ -750,6 +791,18 @@ void UTickableTransformConstraint::TeardownConstraint(UWorld* InWorld)
 		ChildTickFunction->RemovePrerequisite(this, ConstraintTick);
 	}
 	ConstraintTicks.Remove(InWorld->GetCurrentLevel());
+
+	// unregister delegates. should handles delegates be unregistered as well?
+	UnregisterDelegates();
+}
+
+void UTickableTransformConstraint::AddedToWorld(UWorld* InWorld)
+{
+	if (InWorld)
+	{
+		// update dependencies once added to the sub-system 
+		FTransformConstraintUtils::BuildDependencies(InWorld, this);
+	}
 }
 
 /** 
@@ -1781,10 +1834,12 @@ bool FTransformConstraintUtils::AddConstraint(
 	UWorld* InWorld,
 	UTransformableHandle* InParentHandle,
 	UTransformableHandle* InChildHandle,
-	UTickableTransformConstraint* Constraint,
+	UTickableTransformConstraint* InNewConstraint,
 	const bool bMaintainOffset,
 	const bool bUseDefault)
 {
+	using namespace ConstraintLocals;
+	
 	const bool bIsValidParent = InParentHandle && InParentHandle->IsValid();
 	const bool bIsValidChild = InChildHandle && InChildHandle->IsValid();
 	if (!bIsValidParent || !bIsValidChild)
@@ -1793,57 +1848,175 @@ bool FTransformConstraintUtils::AddConstraint(
 		return false;
 	}
 
-	if (Constraint == nullptr)
+	if (InNewConstraint == nullptr)
 	{
 		UE_LOG(LogTemp, Error, TEXT("FTransformConstraintUtils::AddConst error creating constraint"));
 		return false;
 	}
 
-	// store previous child constraints
-	const FConstraintsManagerController& Controller = FConstraintsManagerController::Get(InWorld);
-	const TArray<TWeakObjectPtr<UTickableConstraint>> ChildParentConstraints = Controller.GetParentConstraints(InChildHandle->GetHash(), true);
+	// set handles before calling AddConstraint as it will build dependencies based on them
+	InNewConstraint->ParentTRSHandle = InParentHandle;
+	InNewConstraint->ChildTRSHandle = InChildHandle;
 
 	// add the new one
-	const bool bConstraintAdded = Controller.AddConstraint(Constraint);
+	const FConstraintsManagerController& Controller = FConstraintsManagerController::Get(InWorld);
+	const bool bConstraintAdded = Controller.AddConstraint(InNewConstraint);
+
 	if (!bConstraintAdded)
 	{
+		InNewConstraint->ParentTRSHandle = nullptr;
+		InNewConstraint->ChildTRSHandle = nullptr;
 		return false;
 	}
 
-	Constraint->ParentTRSHandle = InParentHandle;
-	Constraint->ChildTRSHandle = InChildHandle;
 	if (!bUseDefault)
 	{
-		Constraint->bMaintainOffset = bMaintainOffset;
-		Constraint->Setup();
+		InNewConstraint->bMaintainOffset = bMaintainOffset;
+		InNewConstraint->Setup();
+	}
+	InNewConstraint->InitConstraint(InWorld);
+	
+	return true;
+}
+
+bool FTransformConstraintUtils::BuildDependencies(UWorld* InWorld, UTickableTransformConstraint* InConstraint)
+{
+	using namespace ConstraintLocals;
+	
+	if (!ensure(InWorld))
+	{
+		return false;
+	}
+
+	if (!InConstraint || !InConstraint->IsValid())
+	{
+		return false;
 	}
-	Constraint->InitConstraint(InWorld);
+
+	const UTransformableHandle* ParentHandle = InConstraint->ParentTRSHandle.Get();
+	const UTransformableHandle* ChildHandle = InConstraint->ChildTRSHandle.Get();
+	
+ 	// get previous child constraints
+	const FConstraintsManagerController& Controller = FConstraintsManagerController::Get(InWorld);
+	TArray<TWeakObjectPtr<UTickableConstraint>> ChildParentConstraints = Controller.GetParentConstraints(ChildHandle->GetHash(), true);
+	ChildParentConstraints.Remove(InConstraint);
+
 	// add dependencies with the last child constraint
-	const FName NewConstraintName = Constraint->GetFName();
 	if (!ChildParentConstraints.IsEmpty())
 	{
-		const FName LastChildConstraintName = ChildParentConstraints.Last()->GetFName();
-		Controller.SetConstraintsDependencies( LastChildConstraintName, NewConstraintName);
+		const FGuid LastChildConstraintID = ChildParentConstraints.Last()->ConstraintID;
+		if (bDebugDependencies)
+		{
+			UE_LOG(LogTemp, Warning, TEXT("DEPENDENCY 0: tick after last constraint."));
+		}
+		Controller.SetConstraintsDependencies(LastChildConstraintID, InConstraint->ConstraintID);
+	}
+
+	const UObject* ParentTarget = ParentHandle->GetTarget().Get();
+	const UObject* ChildTarget = ChildHandle->GetTarget().Get();
+	const bool bSelf = ParentTarget && ParentTarget == ChildTarget;
+	
+	// internal dependencies?
+	if (bSelf && bIncludeTarget)
+	{
+		const UObject* SelfTarget = ChildTarget;
+		
+		using ConstraintPtr = TWeakObjectPtr<UTickableConstraint>;
+		auto Predicate = [InConstraint, SelfTarget](const ConstraintPtr& Constraint)
+		{
+			const UTickableTransformConstraint* TransformConstraint = Cast<UTickableTransformConstraint>(Constraint.Get());
+			if (!TransformConstraint || TransformConstraint == InConstraint)
+			{
+				return false;
+			}
+			const UObject* ParentTarget = GetHandleTarget(TransformConstraint->ParentTRSHandle);
+			const UObject* ChildTarget = GetHandleTarget(TransformConstraint->ChildTRSHandle);
+			return ParentTarget == SelfTarget && ChildTarget == SelfTarget;
+		};
+
+		const TArray< TWeakObjectPtr<UTickableConstraint> > SelfConstraints = Controller.GetConstraintsByPredicate(Predicate);
+		for (const ConstraintPtr& SelfConstraint: SelfConstraints)
+		{
+			const UTickableTransformConstraint* TransformConstraint = Cast<UTickableTransformConstraint>(SelfConstraint);
+
+			// if the new handles depend on that constraint child then, TransformConstraint should tick before
+			if (ParentHandle->HasDirectDependencyWith(*TransformConstraint->ChildTRSHandle))
+			{
+				Controller.SetConstraintsDependencies(TransformConstraint->ConstraintID, InConstraint->ConstraintID);
+
+				if (bDebugDependencies)
+				{
+					UE_LOG(LogTemp, Warning, TEXT("DEPENDENCY 1: %s is parent of %s so %s must tick before %s"),
+					   *GetHandleLabel(TransformConstraint->ChildTRSHandle), *GetHandleLabel(ParentHandle),
+					   *GetConstraintLabel(TransformConstraint), *GetConstraintLabel(InConstraint));
+				}
+			}
+			else if (ChildHandle->HasDirectDependencyWith(*TransformConstraint->ChildTRSHandle))
+			{
+				Controller.SetConstraintsDependencies(TransformConstraint->ConstraintID, InConstraint->ConstraintID);
+
+				if (bDebugDependencies)
+				{
+					UE_LOG(LogTemp, Warning, TEXT("DEPENDENCY 1: %s is parent of %s so %s must tick before %s"),
+					   *GetHandleLabel(TransformConstraint->ChildTRSHandle), *GetHandleLabel(ChildHandle),
+					   *GetConstraintLabel(TransformConstraint), *GetConstraintLabel(InConstraint));
+				}
+			}
+
+			// if the TransformConstraint handles depend on the new constraint child then, TransformConstraint should tick after
+			if (TransformConstraint->ParentTRSHandle->HasDirectDependencyWith(*ChildHandle))
+			{
+			 	Controller.SetConstraintsDependencies(InConstraint->ConstraintID, TransformConstraint->ConstraintID);
+
+				if (bDebugDependencies)
+				{
+					UE_LOG(LogTemp, Warning, TEXT("DEPENDENCY 1: %s is parent of %s so %s must tick before %s"),
+					   *GetHandleLabel(ChildHandle), *GetHandleLabel(TransformConstraint->ParentTRSHandle),
+					   *GetConstraintLabel(InConstraint), *GetConstraintLabel(TransformConstraint));
+				}
+			}
+			else if (TransformConstraint->ChildTRSHandle->HasDirectDependencyWith(*ChildHandle))
+			{
+				Controller.SetConstraintsDependencies(InConstraint->ConstraintID, TransformConstraint->ConstraintID);
+
+				if (bDebugDependencies)
+				{
+					UE_LOG(LogTemp, Warning, TEXT("DEPENDENCY 1: %s is parent of %s so %s must tick before %s"),
+					   *GetHandleLabel(ChildHandle), *GetHandleLabel(TransformConstraint->ChildTRSHandle),
+					   *GetConstraintLabel(InConstraint), *GetConstraintLabel(TransformConstraint) );
+				}
+			}
+		}
 	}
 
 	// make sure we tick after the parent.
-	Constraint->EnsurePrimaryDependency(InWorld);
+	InConstraint->EnsurePrimaryDependency(InWorld);
 
 	// if child handle is the parent of some other constraints, ensure they will tick after that new one
-	const bool bSelf = InParentHandle->GetTarget().Get() && InParentHandle->GetTarget().Get() == InChildHandle->GetTarget().Get();
 	TArray<TWeakObjectPtr<UTickableConstraint>> ChildChildConstraints;
-	GetChildrenConstraints(InWorld, InChildHandle, ChildChildConstraints, ConstraintLocals::bIncludeTarget && !bSelf);
+	GetChildrenConstraints(InWorld, ChildHandle, ChildChildConstraints, bIncludeTarget && !bSelf);
 	for (const TWeakObjectPtr<UTickableConstraint>& ChildConstraint: ChildChildConstraints)
 	{
-		Controller.SetConstraintsDependencies(Constraint->ConstraintID, ChildConstraint->ConstraintID);
+		Controller.SetConstraintsDependencies(InConstraint->ConstraintID, ChildConstraint->ConstraintID);
+
+		if (bDebugDependencies)
+		{
+			const UTickableTransformConstraint* TransformConstraint = Cast<UTickableTransformConstraint>(ChildConstraint);
+			UE_LOG(LogTemp, Warning, TEXT("DEPENDENCY 2: %s is parent of %s so %s must tick before %s"),
+				*GetHandleLabel(ChildHandle), *GetHandleLabel(TransformConstraint->ChildTRSHandle),
+				*GetConstraintLabel(InConstraint), *GetConstraintLabel(TransformConstraint));
+		}
 	}
 
 	// warn for possible cycles
-	if (FConstraintCycleChecker::IsCycling(InChildHandle))
+	if (FConstraintCycleChecker::IsCycling(InConstraint->ChildTRSHandle))
 	{
-		UE_LOG(LogTemp, Warning, TEXT("A cycle has been formed while creating %s."), *NewConstraintName.ToString());
+		UE_LOG(LogTemp, Warning, TEXT("A cycle has been formed while creating %s."), *InConstraint->GetName());
 	}
 
+	// invalidate graph
+	Controller.InvalidateEvaluationGraph();
+	
 	return true;
 }
 
diff --git a/Engine/Source/Runtime/Experimental/Animation/Constraints/Private/TransformableHandle.cpp b/Engine/Source/Runtime/Experimental/Animation/Constraints/Private/TransformableHandle.cpp
--- a/Engine/Source/Runtime/Experimental/Animation/Constraints/Private/TransformableHandle.cpp
+++ b/Engine/Source/Runtime/Experimental/Animation/Constraints/Private/TransformableHandle.cpp
@@ -292,6 +292,12 @@ void UTransformableComponentHandle::RegisterDelegates()
 #if WITH_EDITOR
 	if (GEngine)
 	{
+		// NOTE BINDER: this has to be done before binding UTransformableComponentHandle::OnActorMoving
+		if (!GEngine->OnActorMoving().IsBoundToObject(&GetEvaluationBinding()))
+		{
+			GEngine->OnActorMoving().AddRaw(&GetEvaluationBinding(), &FComponentEvaluationGraphBinding::OnActorMoving);
+		}
+		
 		GEngine->OnActorMoving().AddUObject(this, &UTransformableComponentHandle::OnActorMoving);
 	}
 	FCoreUObjectDelegates::OnObjectPropertyChanged.AddUObject(this, &UTransformableComponentHandle::OnPostPropertyChanged);
@@ -313,6 +319,8 @@ void UTransformableComponentHandle::OnActorMoving(AActor* InActor)
 		return;
 	}
 
+	GetEvaluationBinding().bPendingFlush = true;
+	
 	Notify(EHandleEvent::GlobalTransformUpdated);
 }
 
@@ -593,3 +601,25 @@ FString UTransformableComponentHandle::GetFullLabel() const
 };
 
 #endif
+
+FComponentEvaluationGraphBinding& UTransformableComponentHandle::GetEvaluationBinding()
+{
+	static FComponentEvaluationGraphBinding EvaluationBinding;
+	return EvaluationBinding;
+}
+
+void FComponentEvaluationGraphBinding::OnActorMoving(AActor* InActor)
+{
+	if (!bPendingFlush)
+	{
+		return;
+	}
+	
+	if (UWorld* World = InActor ? InActor->GetWorld() : nullptr)
+	{
+		// flush all pending evaluations if any
+		const FConstraintsManagerController& Controller = FConstraintsManagerController::Get(World);
+		Controller.FlushEvaluationGraph();
+	}
+	bPendingFlush = false;
+}
diff --git a/Engine/Source/Runtime/Experimental/Animation/Constraints/Public/ConstraintsEvaluationGraph.h b/Engine/Source/Runtime/Experimental/Animation/Constraints/Public/ConstraintsEvaluationGraph.h
new file mode 100644
--- /dev/null
+++ b/Engine/Source/Runtime/Experimental/Animation/Constraints/Public/ConstraintsEvaluationGraph.h
@@ -0,0 +1,92 @@
+﻿// Copyright Epic Games, Inc. All Rights Reserved.
+
+#pragma once
+
+#include "UObject/Object.h"
+#include "Engine/EngineBaseTypes.h"
+
+struct FConstraintTickFunction;
+class UTickableConstraint;
+class UWorld;
+struct FConstraintsInWorld;
+
+/**
+ * A structure that uniquely identifies a constraint within the evaluation graph
+ */
+
+struct FConstraintNode
+{
+	/* Constraint data */
+	FGuid ConstraintID;
+	FConstraintTickFunction* ConstraintTick = nullptr;
+
+	/* Internal indices for quick navigation through the nodes and constraints data structures */
+	int32 NodeIndex = INDEX_NONE;
+	int32 ConstraintIndex = INDEX_NONE;
+
+	/* internal hierarchy representing dependencies at the nodes level (indices here represent the nodes' NodeIndex)*/
+	TSet<uint32> Parents;
+	TSet<uint32> Children;
+
+	/* internal hierarchy representing dependencies at the nodes level (indices here represent the nodes' NodeIndex)*/
+	bool bMarkedForEvaluation = false;
+};
+
+/**
+ * A graph like structure to efficiently represent constraints' evaluation hierarchy using their tick dependencies.
+ * Each node stores the useful data to evaluate the constraint hierarchically. 
+ */
+
+struct FConstraintsEvaluationGraph: public TSharedFromThis<FConstraintsEvaluationGraph>
+{
+public:
+	explicit FConstraintsEvaluationGraph(FConstraintsInWorld& InConstraintsInWorld)
+		: ConstraintsInWorld(InConstraintsInWorld)
+	{}
+
+	/**
+	* Marks the constraint for evaluation. Note that it will only be evaluated after FlushPendingEvaluations has been called.
+	* This makes it possible to tag several constraints and evaluate the graph only once.
+	*/
+	void MarkForEvaluation(const TWeakObjectPtr<UTickableConstraint>& InConstraint);
+	
+	/* Evaluates constraints marked for evaluation + their dependencies. */
+	CONSTRAINTS_API void FlushPendingEvaluations();
+	
+	/* Invalidates internal data so that the graph will be rebuilt on request. */
+	void InvalidateData();
+
+	/* Rebuilds the full graph taking constraints tick dependencies into account. */
+	void Rebuild();
+
+	static bool UseEvaluationGraph();
+	
+private:
+
+	FConstraintsEvaluationGraph(FConstraintsEvaluationGraph&& InOther) = default;
+	FConstraintsEvaluationGraph(const FConstraintsEvaluationGraph& InOther) = delete;
+	FConstraintsEvaluationGraph& operator=(const FConstraintsEvaluationGraph& InOther) = delete;
+	
+	/* Get or create a new node in the graph that wraps that constraint. */
+	FConstraintNode& GetNode(const TWeakObjectPtr<UTickableConstraint>& InConstraint);
+
+	/* Returns the node wrapping that constraint if any. */
+	FConstraintNode* FindNode(const TWeakObjectPtr<UTickableConstraint>& InConstraint);
+
+	/* Evaluates InConstraint ensuring there's a corresponding node. */
+	void Evaluate(const TWeakObjectPtr<UTickableConstraint>& InConstraint);
+
+	/* Evaluates InNode if its tick function is enabled + evaluates its children recursively. */
+	void Evaluate(FConstraintNode* InNode);
+
+	/* Dumps the current state of the graph. */
+	void Dump() const;
+
+	/* Array of constraint wrappers. */
+	TArray<FConstraintNode> Nodes;
+
+	/* References to the actual stored data. */
+	FConstraintsInWorld& ConstraintsInWorld;
+	
+	bool bDataInvalid = true;
+};
\ No newline at end of file
diff --git a/Engine/Source/Runtime/Experimental/Animation/Constraints/Public/ConstraintsManager.h b/Engine/Source/Runtime/Experimental/Animation/Constraints/Public/ConstraintsManager.h
--- a/Engine/Source/Runtime/Experimental/Animation/Constraints/Public/ConstraintsManager.h
+++ b/Engine/Source/Runtime/Experimental/Animation/Constraints/Public/ConstraintsManager.h
@@ -105,7 +105,10 @@ public:
 	CONSTRAINTS_API virtual void InitConstraint(UWorld * InWorld)  PURE_VIRTUAL(UTickableConstraint::InitConstraint);
 	/** Teardown and unregister the constraint with that world*/
 	CONSTRAINTS_API virtual void TeardownConstraint(UWorld* InWorld)  PURE_VIRTUAL(UTickableConstraint::DeInitConstraint);
-	/** Get Tick funtion in that world*/
+	/** */
+	CONSTRAINTS_API virtual void AddedToWorld(UWorld* InWorld) PURE_VIRTUAL(UTickableConstraint::AddedToWorld);
+	
+	/** Get Tick function in that world*/
 	CONSTRAINTS_API FConstraintTickFunction& GetTickFunction(UWorld* InWorld);
 	CONSTRAINTS_API const FConstraintTickFunction& GetTickFunction(UWorld* InWorld) const;
 
@@ -335,4 +338,9 @@ public:
 
 	/** Notify from changes in the constraints manager. */
 	CONSTRAINTS_API void Notify(EConstraintsManagerNotifyType InNotifyType, UObject* InObject) const;
+
+	/** Evaluation graph API. */
+	CONSTRAINTS_API void MarkConstraintForEvaluation(UTickableConstraint* InConstraint) const;
+	CONSTRAINTS_API void InvalidateEvaluationGraph() const;
+	CONSTRAINTS_API void FlushEvaluationGraph() const;
 };
diff --git a/Engine/Source/Runtime/Experimental/Animation/Constraints/Public/TransformConstraint.h b/Engine/Source/Runtime/Experimental/Animation/Constraints/Public/TransformConstraint.h
--- a/Engine/Source/Runtime/Experimental/Animation/Constraints/Public/TransformConstraint.h
+++ b/Engine/Source/Runtime/Experimental/Animation/Constraints/Public/TransformConstraint.h
@@ -158,6 +158,7 @@ public:
 	/** (Re-)Registers the constraint function and (re-)binds the required delegates*/
 	CONSTRAINTS_API virtual void InitConstraint(UWorld * InWorld) override;
 	CONSTRAINTS_API virtual void TeardownConstraint(UWorld * InWorld) override;
+	CONSTRAINTS_API virtual void AddedToWorld(UWorld* InWorld) override;
 
 #if WITH_EDITOR
 public:
@@ -502,4 +503,6 @@ struct FTransformConstraintUtils
 	/** Adjust the transform on a scene component so it's effected by the constraint*/
 	static CONSTRAINTS_API void UpdateTransformBasedOnConstraint(FTransform& CurrentTransform, USceneComponent* SceneComponent);
 
+	/** Ensure default dependencies between constraints. */
+	static CONSTRAINTS_API bool BuildDependencies(UWorld* InWorld, UTickableTransformConstraint* Constraint);
 };
diff --git a/Engine/Source/Runtime/Experimental/Animation/Constraints/Public/TransformableHandle.h b/Engine/Source/Runtime/Experimental/Animation/Constraints/Public/TransformableHandle.h
--- a/Engine/Source/Runtime/Experimental/Animation/Constraints/Public/TransformableHandle.h
+++ b/Engine/Source/Runtime/Experimental/Animation/Constraints/Public/TransformableHandle.h
@@ -142,6 +142,12 @@ private:
  * UTransformableComponentHandle
  */
 
+struct FComponentEvaluationGraphBinding
+{
+	void OnActorMoving(AActor* InActor);
+	bool bPendingFlush = false;
+};
+
 UCLASS(Blueprintable, MinimalAPI)
 class UTransformableComponentHandle : public UTransformableHandle 
 {
@@ -224,4 +230,8 @@ public:
 	CONSTRAINTS_API void OnPostPropertyChanged(UObject* InObject, FPropertyChangedEvent& InPropertyChangedEvent);
 	CONSTRAINTS_API void OnObjectsReplaced(const TMap<UObject*, UObject*>& InOldToNewInstances);
 #endif
+
+protected:
+
+	static FComponentEvaluationGraphBinding& GetEvaluationBinding();
 };
