Iris - Fix out of bounds read.

#rnx

diff --git a/Engine/Source/Runtime/Experimental/Iris/Core/Private/Iris/ReplicationSystem/ReplicationOperations.cpp b/Engine/Source/Runtime/Experimental/Iris/Core/Private/Iris/ReplicationSystem/ReplicationOperations.cpp
--- a/Engine/Source/Runtime/Experimental/Iris/Core/Private/Iris/ReplicationSystem/ReplicationOperations.cpp
+++ b/Engine/Source/Runtime/Experimental/Iris/Core/Private/Iris/ReplicationSystem/ReplicationOperations.cpp
@@ -1261,42 +1261,45 @@ namespace UE::Net::Private
 
 static bool AppendMemberChangeMasks(FNetBitStreamWriter* ChangeMaskWriter, FNetBitStreamWriter* ConditionalChangeMaskWriter, uint8* StateBuffer, const FReplicationStateDescriptor* Descriptor)
 {
-	FNetBitArrayView::StorageWordType* ChangeMaskStorage = reinterpret_cast<FNetBitArrayView::StorageWordType*>(StateBuffer + Descriptor->GetChangeMaskOffset());
 	const uint32 BitCount = Descriptor->ChangeMaskBitCount;
-
-	// Append change mask bits
-	if (BitCount <= 32u)
-	{
-		ChangeMaskWriter->WriteBits(*ChangeMaskStorage, BitCount);
-	}
-	else
+	if (BitCount)
 	{
-		ChangeMaskWriter->WriteBitStream(ChangeMaskStorage, 0, BitCount);
-	}
-
-	if (ConditionalChangeMaskWriter != nullptr)
-	{
-		if (EnumHasAnyFlags(Descriptor->Traits, EReplicationStateTraits::HasLifetimeConditionals))
+		// Append change mask bits
+		FNetBitArrayView::StorageWordType* ChangeMaskStorage = reinterpret_cast<FNetBitArrayView::StorageWordType*>(StateBuffer + Descriptor->GetChangeMaskOffset());
+		if (BitCount <= 32)
+		{
+			ChangeMaskWriter->WriteBits(*ChangeMaskStorage, BitCount);
+		}
+		else
 		{
-			FNetBitArrayView::StorageWordType* ConditionalChangeMaskStorage = reinterpret_cast<FNetBitArrayView::StorageWordType*>(StateBuffer + Descriptor->GetConditionalChangeMaskOffset());
+			constexpr uint32 SrcStreamBitOffset = 0U;
+			ChangeMaskWriter->WriteBitStream(ChangeMaskStorage, SrcStreamBitOffset, BitCount);
+		}
 
-			// Append conditional change mask bits.
-			// The conditionals are on or off rather than tracking dirtiness so we do not reset them.
-			if (BitCount <= 32u)
+		if (ConditionalChangeMaskWriter != nullptr)
+		{
+			if (EnumHasAnyFlags(Descriptor->Traits, EReplicationStateTraits::HasLifetimeConditionals))
 			{
-				ConditionalChangeMaskWriter->WriteBits(*ConditionalChangeMaskStorage, BitCount);
+				FNetBitArrayView::StorageWordType* ConditionalChangeMaskStorage = reinterpret_cast<FNetBitArrayView::StorageWordType*>(StateBuffer + Descriptor->GetConditionalChangeMaskOffset());
+
+				// Append conditional change mask bits.
+				// The conditionals are on or off rather than tracking dirtiness so we do not reset them.
+				if (BitCount <= 32u)
+				{
+					ConditionalChangeMaskWriter->WriteBits(*ConditionalChangeMaskStorage, BitCount);
+				}
+				else
+				{
+					constexpr uint32 BitOffset = 0;
+					ConditionalChangeMaskWriter->WriteBitStream(ConditionalChangeMaskStorage, BitOffset, BitCount);
+				}
 			}
 			else
 			{
-				constexpr uint32 BitOffset = 0;
-				ConditionalChangeMaskWriter->WriteBitStream(ConditionalChangeMaskStorage, BitOffset, BitCount);
+				// Skip past our non-existing conditional mask.
+				ConditionalChangeMaskWriter->Seek(ConditionalChangeMaskWriter->GetPosBits() + BitCount);
 			}
 		}
-		else
-		{
-			// Skip past our non-existing conditional mask.
-			ConditionalChangeMaskWriter->Seek(ConditionalChangeMaskWriter->GetPosBits() + BitCount);
-		}
 	}
 
 	FReplicationStateHeader& ReplicationStateHeader = Private::GetReplicationStateHeader(StateBuffer, Descriptor);
