Chaos - Fix drifting in character ground constraint

diff --git a/Engine/Source/Runtime/Experimental/Chaos/Public/Chaos/Character/CharacterGroundConstraintContainer.h b/Engine/Source/Runtime/Experimental/Chaos/Public/Chaos/Character/CharacterGroundConstraintContainer.h
--- a/Engine/Source/Runtime/Experimental/Chaos/Public/Chaos/Character/CharacterGroundConstraintContainer.h
+++ b/Engine/Source/Runtime/Experimental/Chaos/Public/Chaos/Character/CharacterGroundConstraintContainer.h
@@ -60,6 +60,30 @@ namespace Chaos
 		void SetData(const FCharacterGroundConstraintDynamicData& InData)
 		{
 			Data = InData;
+
+			if (!CharacterParticle || bDisabled)
+			{
+				return;
+			}
+
+			// If the movement target is close to zero then it gets clamped to zero and we
+			// recompute the delta position based on the previous target to avoid drift
+			constexpr float MovementThresholdSq = 0.01f;
+			if (InData.TargetDeltaPosition.SizeSquared() > MovementThresholdSq)
+			{
+				ComputeLocalCharacterPosition();
+			}
+			else
+			{
+				if (GroundParticle)
+				{
+					Data.TargetDeltaPosition = GroundParticle->GetR() * LocalCharacterPosition + GroundParticle->GetX() - CharacterParticle->GetX() + Data.TargetDeltaPosition;
+				}
+				else
+				{
+					Data.TargetDeltaPosition = LocalCharacterPosition - CharacterParticle->GetX() + Data.TargetDeltaPosition;
+				}
+			}
 		}
 
 		// Declared final so that TPBDConstraintGraphRuleImpl::AddToGraph() does not need to hit vtable
@@ -80,6 +104,7 @@ namespace Chaos
 				{
 					GroundParticle->AddConstraintHandle(this);
 				}
+				ComputeLocalCharacterPosition();
 				bGroundParticleChanged = true;
 			}
 		}
@@ -96,10 +121,26 @@ namespace Chaos
 		friend class Private::FCharacterGroundConstraintContainerSolver; // For setting the solver force and torque
 		friend class ChaosTest::CharacterGroundConstraintContainerTest; // For testing internals
 
+		void ComputeLocalCharacterPosition()
+		{
+			if (!CharacterParticle || bDisabled)
+			{
+				return;
+			}
+
+			LocalCharacterPosition = CharacterParticle->GetX();
+
+			if (GroundParticle)
+			{
+				LocalCharacterPosition = GroundParticle->GetR().Inverse() * (LocalCharacterPosition - GroundParticle->GetX());
+			}
+		}
+
 		FCharacterGroundConstraintSettings Settings;
 		FCharacterGroundConstraintDynamicData Data;
 		FVec3 SolverAppliedForce = FVec3::ZeroVector;
 		FVec3 SolverAppliedTorque = FVec3::ZeroVector;
+		FVec3 LocalCharacterPosition = FVec3::ZeroVector;
 		FGeometryParticleHandle* CharacterParticle;
 		FGeometryParticleHandle* GroundParticle;
 		bool bDisabled = false;
