Dataflow : FContext Serialization :
These changes add caching support to the Dataflow Graphs context. The caches are saved as an asset in /Game/_Generated and automatically connected to an associated graph when it?s loaded in the editor.
- Added the owning nodes guid and valueHash to the elements saved within the cached context.
- Added Serialization support to the FContext.
- DataStore elements are saved based on user defined serialization functions within the FContextCachingFactory
- CacheElements are accessed in the DataStore using a hash of the FDataflowConnection::Guid. (this used to be the address of the FProperty which changed from run to run).
- Scrubbed code for inconsistent usage of cache keys, all based on FContextCacheKey now.
- Added support for FDataflowNode::GetValueHash(), to return joined UPROPERTY hashes of the node.

diff --git a/Engine/Plugins/Experimental/Dataflow/Source/DataflowEditor/Private/Dataflow/AssetDefinition_DataflowContext.cpp b/Engine/Plugins/Experimental/Dataflow/Source/DataflowEditor/Private/Dataflow/AssetDefinition_DataflowContext.cpp
new file mode 100644
--- /dev/null
+++ b/Engine/Plugins/Experimental/Dataflow/Source/DataflowEditor/Private/Dataflow/AssetDefinition_DataflowContext.cpp
@@ -0,0 +1,238 @@
+// Copyright Epic Games, Inc. All Rights Reserved.
+
+#include "Dataflow/AssetDefinition_DataflowContext.h"
+
+#include "AssetRegistry/AssetRegistryModule.h"
+#include "ContentBrowserModule.h"
+#include "Dataflow/DataflowEditor.h"
+#include "Dataflow/DataflowEditorModule.h"
+#include "Dataflow/DataflowEditorUtil.h"
+#include "Dataflow/DataflowContent.h"
+#include "Dialog/SMessageDialog.h"
+#include "IContentBrowserSingleton.h"
+#include "Math/Color.h"
+#include "Misc/FileHelper.h"
+#include "ThumbnailRendering/SceneThumbnailInfo.h"
+#include "UObject/UObjectGlobals.h"
+
+
+#define LOCTEXT_NAMESPACE "AssetActions_DataflowContext"
+
+namespace DataflowContextDefinitionHelpers
+{
+
+	/*
+	* BindContextToGraph
+	* 
+	* On load the Property in the cache will be a nullptr
+	* This block will rebind the property in the cache to the 
+	* the Property on the UDataflow asset. If there is ever an
+	* mis-match, just return false to indicate that the binding
+	* has failed.
+	*/
+	bool BindContextToGraph(TObjectPtr<UObject>& Asset, UDataflow* DataflowAsset)
+	{
+		using namespace Dataflow;
+
+		TSharedPtr<FGraph> Dataflow = DataflowAsset->GetDataflow();
+		if (!Dataflow) return false;
+
+		UDataflowBaseContent* BaseContent = Cast< UDataflowBaseContent>(Asset.Get());
+		if (!BaseContent) return false;
+
+		TSharedPtr<FEngineContext>& Context = BaseContent->GetDataflowContext();
+		if (!Context) return false;
+
+		TSet<FContextCacheKey> Keys; Context->GetKeys(Keys);
+		for (FContextCacheKey Key : Keys)
+		{
+			bool bValidKey = false;
+
+			if (TUniquePtr<FContextCacheElementBase>* Data = Context->GetBaseData(Key))
+			{
+				if ((*Data) && !(*Data)->Property)
+				{
+					if (TSharedPtr<FDataflowNode> Node = Dataflow->FindBaseNode((*Data)->NodeGuid))
+					{
+						if (FDataflowOutput* Output = Node->FindOutput(Key))
+						{
+							if (const FProperty* Property = Output->GetProperty())
+							{
+								(*Data)->Property = Property;
+								bValidKey = true;
+							}
+						}
+					}
+				}
+				else
+				{
+					bValidKey = true;
+				}
+			}
+
+			if (!bValidKey)
+			{
+				return false;
+			}
+		}
+		return true;
+	}
+
+	/*
+	* ValidateCachedNodeHash
+	* 
+	* Check that the hashes stored in the cache reflect the hash of the nodes 
+	* properties.
+	*
+	*/
+	bool ValidateCachedNodeHash(TObjectPtr<UObject>& Asset, UDataflow* DataflowAsset)
+	{
+		using namespace Dataflow;
+
+		TSharedPtr<FGraph> Dataflow = DataflowAsset->GetDataflow();
+		if (!Dataflow) return false;
+
+		UDataflowBaseContent* BaseContent = Cast< UDataflowBaseContent>(Asset.Get());
+		if (!BaseContent) return false;
+
+		TSharedPtr<FEngineContext>& Context = BaseContent->GetDataflowContext();
+		if (!Context) return false;
+
+		TSet<FContextCacheKey> Keys; Context->GetKeys(Keys);
+		for (FContextCacheKey Key : Keys)
+		{
+			bool bValidKey = false;
+			if (TUniquePtr<FContextCacheElementBase>* Data = Context->GetBaseData(Key))
+			{
+				if (!(*Data)) return false;
+				if (!(*Data)->Property) return false;
+				if(TSharedPtr<FDataflowNode> Node = Dataflow->FindBaseNode((*Data)->NodeGuid))
+				{
+					uint32 Hash = Node->GetValueHash();
+					if ((*Data)->NodeHash != Node->GetValueHash())
+					{
+						return false;
+					}
+					bValidKey = true;
+				}
+			}
+			if (!bValidKey)
+			{
+				return false;
+			}
+		}
+
+
+		return true;
+	}
+
+	bool ResetCacheTimestamp(TObjectPtr<UObject>& Asset, UDataflow* DataflowAsset)
+	{
+		using namespace Dataflow;
+
+		UDataflowBaseContent* BaseContent = Cast< UDataflowBaseContent>(Asset.Get());
+		if (!BaseContent) return false;
+
+		BaseContent->SetLastModifiedTimestamp(DataflowAsset->GetRenderingTimestamp().Value+1, false /*bMakeDirty*/);
+
+		return true;
+	}
+
+	/*
+	* CreateNewDataflowContext
+	*
+	*/	
+	template<class T>
+	TObjectPtr<T> CreateNewDataflowContext(const TObjectPtr<UObject>& ContentOwner)
+	{
+		check(ContentOwner.Get());
+
+		UClass* DataflowClass = T::StaticClass();
+		UDataflow* DataflowAsset = Private::GetDataflowAssetFrom(ContentOwner);
+
+		FContentBrowserModule& ContentBrowserModule = FModuleManager::LoadModuleChecked<FContentBrowserModule>("ContentBrowser");
+
+		FString PackageName = FString::Printf(TEXT("/Game/_GENERATED/Dataflow/%s"), *ContentOwner.GetName());
+		if(DataflowAsset)
+		{
+			PackageName = FString::Printf(TEXT("%s_%s"), *PackageName, *DataflowAsset->GetName());
+		}
+
+		UPackage* Package = FindPackage(nullptr, *PackageName);
+		if (!Package)
+		{
+			Package = CreatePackage(*PackageName);
+		}
+
+		bool bNeedsNewAsset = true;
+		TObjectPtr<UObject> Asset = StaticLoadObject(DataflowClass, Package, *PackageName);
+		if (Asset && DataflowAsset)
+		{
+			// Validate the loaded cache
+			bNeedsNewAsset =
+				!BindContextToGraph(Asset, DataflowAsset) ||
+				!ValidateCachedNodeHash(Asset, DataflowAsset) ||
+				!ResetCacheTimestamp(Asset, DataflowAsset);
+		}
+		
+		if(bNeedsNewAsset)
+		{
+			const FName AssetName(FPackageName::GetLongPackageAssetName(PackageName));
+			Asset = NewObject<UObject>(Package, DataflowClass, AssetName, RF_Public | RF_Standalone | RF_Transactional);
+
+			Asset->MarkPackageDirty();
+			FAssetRegistryModule::AssetCreated(Asset);
+
+			if (UDataflowBaseContent* BaseContent = Cast< UDataflowBaseContent>(Asset.Get()))
+			{
+				BaseContent->BuildBaseContent(ContentOwner);
+				BaseContent->SetDataflowOwner(ContentOwner);
+			}
+		}
+
+		return Cast<T>(Asset.Get());
+	}
+
+	template TObjectPtr<UDataflowBaseContent> CreateNewDataflowContext(const TObjectPtr<UObject>& ContentOwner);
+	template TObjectPtr<UDataflowSkeletalContent> CreateNewDataflowContext(const TObjectPtr<UObject>& ContentOwner);
+}
+
+
+namespace UE::Dataflow::DataflowContext
+{
+	struct FColorScheme
+	{
+		static inline const FLinearColor Asset = FColor(180, 120, 110);
+		static inline const FLinearColor NodeHeader = FColor(180, 120, 110);
+		static inline const FLinearColor NodeBody = FColor(18, 12, 11, 127);
+	};
+}
+
+FText UAssetDefinition_DataflowContext::GetAssetDisplayName() const
+{
+	return NSLOCTEXT("AssetTypeActions", "AssetTypeActions_DataflowContext", "DataflowContext");
+}
+
+TSoftClassPtr<UObject> UAssetDefinition_DataflowContext::GetAssetClass() const
+{
+	return UDataflow::StaticClass();
+}
+
+FLinearColor UAssetDefinition_DataflowContext::GetAssetColor() const
+{
+	return UE::Dataflow::DataflowContext::FColorScheme::Asset;
+}
+
+TConstArrayView<FAssetCategoryPath> UAssetDefinition_DataflowContext::GetAssetCategories() const
+{
+	static const auto Categories = { EAssetCategoryPaths::Physics };
+	return Categories;
+}
+
+UThumbnailInfo* UAssetDefinition_DataflowContext::LoadThumbnailInfo(const FAssetData& InAsset) const
+{
+	return UE::Editor::FindOrCreateThumbnailInfo(InAsset.GetAsset(), USceneThumbnailInfo::StaticClass());
+}
+
+
+#undef LOCTEXT_NAMESPACE
diff --git a/Engine/Plugins/Experimental/Dataflow/Source/DataflowEditor/Private/Dataflow/DataflowEditor.cpp b/Engine/Plugins/Experimental/Dataflow/Source/DataflowEditor/Private/Dataflow/DataflowEditor.cpp
--- a/Engine/Plugins/Experimental/Dataflow/Source/DataflowEditor/Private/Dataflow/DataflowEditor.cpp
+++ b/Engine/Plugins/Experimental/Dataflow/Source/DataflowEditor/Private/Dataflow/DataflowEditor.cpp
@@ -3,6 +3,7 @@
 #include "Dataflow/DataflowEditor.h"
 
 #include "Animation/Skeleton.h"
+#include "Dataflow/AssetDefinition_DataflowContext.h"
 #include "Dataflow/DataflowContent.h"
 #include "Dataflow/DataflowEditorToolkit.h"
 #include "Dataflow/DataflowEditorUtil.h"
@@ -36,8 +37,8 @@ void UDataflowEditor::InitializeContent(TObjectPtr<UDataflowBaseContent> BaseCon
 	{
 		if(UDataflow* DataflowAsset = Cast<UDataflow>(ContentOwner))
 		{
-			DataflowContent = NewObject<UDataflowBaseContent>();
-			
+			DataflowContent = DataflowContextDefinitionHelpers::CreateNewDataflowContext<UDataflowBaseContent>(ContentOwner);
+
 			DataflowContent->SetDataflowAsset(DataflowAsset);
 			DataflowContent->SetDataflowTerminal(FString());
 		}
@@ -47,7 +48,7 @@ void UDataflowEditor::InitializeContent(TObjectPtr<UDataflowBaseContent> BaseCon
 			{
 				if(Private::HasSkeletalMesh(ContentOwner))
 				{
-					DataflowContent = NewObject<UDataflowSkeletalContent>();
+					DataflowContent = DataflowContextDefinitionHelpers::CreateNewDataflowContext<UDataflowSkeletalContent>(ContentOwner);
 					const TObjectPtr<UDataflowSkeletalContent> SkeletalContent = Cast<UDataflowSkeletalContent>(DataflowContent);
 					
 					SkeletalContent->SetSkeletalMesh(Private::GetSkeletalMeshFrom(ContentOwner));
@@ -56,7 +57,7 @@ void UDataflowEditor::InitializeContent(TObjectPtr<UDataflowBaseContent> BaseCon
 				}
 				else
 				{
-					DataflowContent = NewObject<UDataflowBaseContent>();
+					DataflowContent = DataflowContextDefinitionHelpers::CreateNewDataflowContext<UDataflowBaseContent>(ContentOwner);
 				}
 				DataflowContent->SetDataflowAsset(Private::GetDataflowAssetFrom(ContentOwner));
 				DataflowContent->SetDataflowTerminal(Private::GetDataflowTerminalFrom(ContentOwner));
@@ -89,10 +90,6 @@ void UDataflowEditor::InitializeContent(TObjectPtr<UDataflowBaseContent> BaseCon
 	if(DataflowContent && DataflowContent->GetDataflowAsset())
 	{
 		DataflowContent->GetDataflowAsset()->Schema = UDataflowSchema::StaticClass();
-		DataflowContent->BuildBaseContent(ContentOwner);
-		
-		// the owner is either a base content derived class or the dataflow itself
-		DataflowContent->SetDataflowOwner(ContentOwner);
 	}
 	
 	// Potentially we could add additional objects to edit here (fields, meshes....)
diff --git a/Engine/Plugins/Experimental/Dataflow/Source/DataflowEditor/Private/Dataflow/DataflowEditorToolkit.cpp b/Engine/Plugins/Experimental/Dataflow/Source/DataflowEditor/Private/Dataflow/DataflowEditorToolkit.cpp
--- a/Engine/Plugins/Experimental/Dataflow/Source/DataflowEditor/Private/Dataflow/DataflowEditorToolkit.cpp
+++ b/Engine/Plugins/Experimental/Dataflow/Source/DataflowEditor/Private/Dataflow/DataflowEditorToolkit.cpp
@@ -393,7 +393,7 @@ void FDataflowEditorToolkit::GetSaveableObjects(TArray<UObject*>& OutObjects) co
 
 	if (ensure(GetDataflowContent()))
 	{
-		if (UDataflow* DataflowAsset = GetDataflowContent()->DataflowAsset)
+		if (UDataflow* DataflowAsset = GetDataflowContent()->GetDataflowAsset())
 		{
 			check(DataflowAsset->IsAsset());
 			OutObjects.Add(DataflowAsset);
@@ -589,7 +589,7 @@ void FDataflowEditorToolkit::OnNodeSelectionChanged(const TSet<UObject*>& InNewS
 	// Despite this function's name, we might not have actually changed which node is selected
 	bool bPrimarySelectionChanged = false;
 
-	if (TObjectPtr<UDataflowBaseContent> EditorContent = GetDataflowContent(); EditorContent->DataflowAsset)
+	if (TObjectPtr<UDataflowBaseContent> EditorContent = GetDataflowContent(); EditorContent->GetDataflowAsset())
 	{
 		// Only keep UDataflowEdNode from NewSelection
 		TSet<UObject*> NodeSelection = FindDataflowNodesInSet(InNewSelection);
diff --git a/Engine/Plugins/Experimental/Dataflow/Source/DataflowEditor/Private/Dataflow/DataflowPreviewScene.cpp b/Engine/Plugins/Experimental/Dataflow/Source/DataflowEditor/Private/Dataflow/DataflowPreviewScene.cpp
--- a/Engine/Plugins/Experimental/Dataflow/Source/DataflowEditor/Private/Dataflow/DataflowPreviewScene.cpp
+++ b/Engine/Plugins/Experimental/Dataflow/Source/DataflowEditor/Private/Dataflow/DataflowPreviewScene.cpp
@@ -118,7 +118,7 @@ void FDataflowConstructionScene::TickDataflowScene(const float DeltaSeconds)
 {
 	if (const TSharedPtr<Dataflow::FContext> DataflowContext = DataflowContent->GetDataflowContext())
 	{
-		if (const UDataflow* Dataflow = DataflowContent->DataflowAsset)
+		if (const UDataflow* Dataflow = DataflowContent->GetDataflowAsset())
 		{
 			const Dataflow::FTimestamp SystemTimestamp = LatestTimestamp(Dataflow, DataflowContext.Get());
 			if (SystemTimestamp >= DataflowContent->GetLastModifiedTimestamp() || DataflowContent->IsDirty())
@@ -226,9 +226,9 @@ TObjectPtr<UDynamicMeshComponent>& FDataflowConstructionScene::AddDynamicMeshCom
 	DynamicMeshComponent->SetMesh(MoveTemp(DynamicMesh));
 	
 	// @todo(Material) This is just to have a material, we should transfer the materials from the assets if they have them. 
-	if (DataflowContent && DataflowContent->DataflowAsset && DataflowContent->DataflowAsset->Material)
+	if (DataflowContent && DataflowContent->GetDataflowAsset() && DataflowContent->GetDataflowAsset()->Material)
 	{
-		DynamicMeshComponent->ConfigureMaterialSet({ DataflowContent->DataflowAsset->Material });
+		DynamicMeshComponent->ConfigureMaterialSet({ DataflowContent->GetDataflowAsset()->Material });
 	}
 	else
 	{
diff --git a/Engine/Plugins/Experimental/Dataflow/Source/DataflowEditor/Public/Dataflow/AssetDefinition_DataflowContext.h b/Engine/Plugins/Experimental/Dataflow/Source/DataflowEditor/Public/Dataflow/AssetDefinition_DataflowContext.h
new file mode 100644
--- /dev/null
+++ b/Engine/Plugins/Experimental/Dataflow/Source/DataflowEditor/Public/Dataflow/AssetDefinition_DataflowContext.h
@@ -0,0 +1,33 @@
+// Copyright Epic Games, Inc. All Rights Reserved.
+
+#pragma once
+
+#include "AssetDefinitionDefault.h"
+
+#include "AssetDefinition_DataflowContext.generated.h"
+
+class UDataflowBaseContent;
+
+namespace DataflowContextDefinitionHelpers
+{
+	// Return true if we should proceed, false if we should re-open the dialog
+	template<class T>
+	DATAFLOWEDITOR_API TObjectPtr<T> CreateNewDataflowContext(const TObjectPtr<UObject>& ContentOwner);
+}
+
+
+
+UCLASS()
+class UAssetDefinition_DataflowContext : public UAssetDefinitionDefault
+{
+	GENERATED_BODY()
+
+private:
+
+	virtual FText GetAssetDisplayName() const override;
+	virtual TSoftClassPtr<UObject> GetAssetClass() const override;
+	virtual FLinearColor GetAssetColor() const override;
+	virtual TConstArrayView<FAssetCategoryPath> GetAssetCategories() const override;
+	virtual UThumbnailInfo* LoadThumbnailInfo(const FAssetData& InAssetData) const override;
+};
+
diff --git a/Engine/Plugins/Experimental/Dataflow/Source/DataflowEnginePlugin/Private/Dataflow/DataflowEngineContextCaching.cpp b/Engine/Plugins/Experimental/Dataflow/Source/DataflowEnginePlugin/Private/Dataflow/DataflowEngineContextCaching.cpp
new file mode 100644
--- /dev/null
+++ b/Engine/Plugins/Experimental/Dataflow/Source/DataflowEnginePlugin/Private/Dataflow/DataflowEngineContextCaching.cpp
@@ -0,0 +1,50 @@
+// Copyright Epic Games, Inc. All Rights Reserved.
+
+#include "Dataflow/DataflowEngineContextCaching.h"
+
+#include "Dataflow/DataflowEnginePlugin.h"
+#include "Dataflow/DataflowNodeParameters.h"
+#include "Dataflow/DataflowContextCachingFactory.h"
+#include "GeometryCollection/ManagedArrayCollection.h"
+
+namespace Dataflow
+{
+	void ContextCachingCallbacks()
+	{
+		using namespace Dataflow;
+
+		/**
+		* Dataflow ContextCaching (FManagedArrayCollection) 
+		*
+		*		@param Type : FManagedArrayCollection::StaticType()
+		* 
+		*/
+		FContextCachingFactory::GetInstance()->RegisterSerializeFunction(FManagedArrayCollection::StaticType(),
+			[](FArchive& Ar, FContextCacheElementBase* Element)
+			{
+				if (Ar.IsSaving())
+				{
+					check(Element!=nullptr); // read from a explicit type and return a null element
+					
+					FManagedArrayCollection& Collection = FContextCachingFactory::GetTypedElement<FManagedArrayCollection>(Element);
+					Chaos::FChaosArchive ChaosAr(Ar);
+					Collection.Serialize(ChaosAr);
+
+					return (FContextCacheElementBase*)nullptr;
+				}
+				else if(Ar.IsLoading())
+				{
+					check(Element == nullptr); // write into a explicit type and return a new element
+					
+					FManagedArrayCollection Collection;
+					Chaos::FChaosArchive ChaosAr(Ar);
+					Collection.Serialize(ChaosAr);
+
+					return FContextCachingFactory::NewTypedElement<FManagedArrayCollection>(MoveTemp(Collection));
+				}
+				return (FContextCacheElementBase*)nullptr;
+			});
+
+	}
+
+}
diff --git a/Engine/Plugins/Experimental/Dataflow/Source/DataflowEnginePlugin/Private/Dataflow/DataflowEnginePlugin.cpp b/Engine/Plugins/Experimental/Dataflow/Source/DataflowEnginePlugin/Private/Dataflow/DataflowEnginePlugin.cpp
--- a/Engine/Plugins/Experimental/Dataflow/Source/DataflowEnginePlugin/Private/Dataflow/DataflowEnginePlugin.cpp
+++ b/Engine/Plugins/Experimental/Dataflow/Source/DataflowEnginePlugin/Private/Dataflow/DataflowEnginePlugin.cpp
@@ -4,8 +4,9 @@
 
 #include "CoreMinimal.h"
 #include "Modules/ModuleManager.h"
-#include "Dataflow/DataflowRenderingFactory.h"
+#include "Dataflow/DataflowEngineContextCaching.h"
 #include "Dataflow/DataflowEngineRendering.h"
+#include "Dataflow/DataflowRenderingFactory.h"
 #include "GeometryCollection/GeometryCollectionUtility.h"
 #include "GeometryCollection/Facades/CollectionRenderingFacade.h"
 #include "GeometryCollection/ManagedArrayCollection.h"
@@ -31,6 +32,7 @@ IMPLEMENT_MODULE( FDataflowEnginePlugin, DataflowEnginePlugin)
 void FDataflowEnginePlugin::StartupModule()
 {
 	Dataflow::RenderingCallbacks();
+	Dataflow::ContextCachingCallbacks();
 }
 
 
diff --git a/Engine/Plugins/Experimental/Dataflow/Source/DataflowEnginePlugin/Public/Dataflow/DataflowEngineContextCaching.h b/Engine/Plugins/Experimental/Dataflow/Source/DataflowEnginePlugin/Public/Dataflow/DataflowEngineContextCaching.h
new file mode 100644
--- /dev/null
+++ b/Engine/Plugins/Experimental/Dataflow/Source/DataflowEnginePlugin/Public/Dataflow/DataflowEngineContextCaching.h
@@ -0,0 +1,15 @@
+// Copyright Epic Games, Inc. All Rights Reserved.
+
+#pragma once
+
+#include "CoreMinimal.h"
+
+
+namespace Dataflow
+{
+	/**
+	 * Rendering callbacks for the Dataflow Rendering Factory
+	 */
+	void ContextCachingCallbacks();
+
+}
\ No newline at end of file
diff --git a/Engine/Source/Runtime/Experimental/Dataflow/Core/Private/Dataflow/DataflowContextCachingFactory.cpp b/Engine/Source/Runtime/Experimental/Dataflow/Core/Private/Dataflow/DataflowContextCachingFactory.cpp
new file mode 100644
--- /dev/null
+++ b/Engine/Source/Runtime/Experimental/Dataflow/Core/Private/Dataflow/DataflowContextCachingFactory.cpp
@@ -0,0 +1,36 @@
+// Copyright Epic Games, Inc. All Rights Reserved.
+
+#include "Dataflow/DataflowContextCachingFactory.h"
+
+#include "Dataflow/DataflowNodeParameters.h"
+
+#include "Misc/MessageDialog.h"
+
+namespace Dataflow
+{
+	FContextCachingFactory* FContextCachingFactory::Instance = nullptr;
+
+	FContextCacheElementBase* FContextCachingFactory::Serialize(FArchive& Ar, FContextCacheData&& Element)
+	{
+		FContextCacheElementBase* RetVal = nullptr;
+		if (CachingMap.Contains(Element.Type))
+		{
+			RetVal = CachingMap[Element.Type](Ar, Element.Data);
+			if (Ar.IsSaving())
+			{
+				check(RetVal == nullptr);
+			}
+			else if( Ar.IsLoading())
+			{
+				check(RetVal != nullptr);
+			}
+		}
+		else
+		{
+			UE_LOG(LogChaos, Warning,
+				TEXT("Warning : Dataflow missing context chaching callback type(%s)"), *Element.Type.ToString());
+		}
+		return RetVal;
+	}
+}
+
diff --git a/Engine/Source/Runtime/Experimental/Dataflow/Core/Private/Dataflow/DataflowNode.cpp b/Engine/Source/Runtime/Experimental/Dataflow/Core/Private/Dataflow/DataflowNode.cpp
--- a/Engine/Source/Runtime/Experimental/Dataflow/Core/Private/Dataflow/DataflowNode.cpp
+++ b/Engine/Source/Runtime/Experimental/Dataflow/Core/Private/Dataflow/DataflowNode.cpp
@@ -7,6 +7,7 @@
 #include "Dataflow/DataflowArchive.h"
 #include "Serialization/ObjectWriter.h"
 #include "Serialization/ObjectReader.h"
+#include "Templates/TypeHash.h"
 
 #include UE_INLINE_GENERATED_CPP_BY_NAME(DataflowNode)
 
@@ -129,6 +130,20 @@ void FDataflowNode::AddOutput(FDataflowOutput* InPtr)
 	}
 }
 
+FDataflowOutput* FDataflowNode::FindOutput(uint32 InGuidHash)
+{
+	for (TPair<uint32, FDataflowOutput*> Elem : Outputs)
+	{
+		FDataflowOutput* Con = Elem.Value;
+		if (GetTypeHash(Con->GetGuid()) == InGuidHash)
+		{
+			return (FDataflowOutput*)Con;
+		}
+	}
+	return nullptr;
+}
+
+
 FDataflowOutput* FDataflowNode::FindOutput(FName InName)
 {
 	for (TPair<uint32, FDataflowOutput*> Elem : Outputs)
@@ -155,6 +170,19 @@ const FDataflowOutput* FDataflowNode::FindOutput(FName InName) const
 	return nullptr;
 }
 
+const FDataflowOutput* FDataflowNode::FindOutput(uint32 InGuidHash) const
+{
+	for (TPair<uint32, FDataflowOutput*> Elem : Outputs)
+	{
+		FDataflowOutput* Con = Elem.Value;
+		if (GetTypeHash(Con->GetGuid()) == InGuidHash)
+		{
+			return (FDataflowOutput*)Con;
+		}
+	}
+	return nullptr;
+}
+
 const FDataflowOutput* FDataflowNode::FindOutput(const void* Reference) const
 {
 	for (TPair<uint32, FDataflowOutput*> Elem : Outputs)
@@ -448,6 +476,41 @@ void FDataflowNode::RegisterOutputConnection(const void* InProperty, const void*
 	}
 }
 
+
+uint32 FDataflowNode::GetValueHash()
+{
+	//UE_LOG(LogChaos, Warning, TEXT("%s"), *GetName().ToString())
+
+	uint32 Hash = 0;
+	if (const TUniquePtr<FStructOnScope> ScriptOnStruct = TUniquePtr<FStructOnScope>(NewStructOnScope()))
+	{
+		if (const UStruct* const Struct = ScriptOnStruct->GetStruct())
+		{
+			for (FPropertyValueIterator PropertyIt(FProperty::StaticClass(), Struct, this); PropertyIt; ++PropertyIt)
+			{
+				if (const FProperty* const Property = PropertyIt.Key())
+				{
+					if (Property->PropertyFlags & CPF_HasGetValueTypeHash)
+					{
+						// uint32 CrcHash = FCrc::MemCrc32(PropertyIt.Value(), Property->ElementSize);
+						// UE_LOG(LogChaos, Warning, TEXT("( %lu \t%s"), (unsigned long)CrcHash, *Property->GetName())
+
+						if ( Property->PropertyFlags & CPF_TObjectPtr )
+						{
+							// @todo(dataflow) : Do something about TObjectPtr<T>
+						}
+						else
+						{
+							Hash = HashCombine(Hash, Property->GetValueTypeHash(PropertyIt.Value()));
+						}
+					}
+				}
+			}
+		}
+	}
+	return Hash;
+}
+
 bool FDataflowNode::ValidateConnections()
 {
 	bHasValidConnections = true;
diff --git a/Engine/Source/Runtime/Experimental/Dataflow/Core/Private/Dataflow/DataflowNodeParameters.cpp b/Engine/Source/Runtime/Experimental/Dataflow/Core/Private/Dataflow/DataflowNodeParameters.cpp
--- a/Engine/Source/Runtime/Experimental/Dataflow/Core/Private/Dataflow/DataflowNodeParameters.cpp
+++ b/Engine/Source/Runtime/Experimental/Dataflow/Core/Private/Dataflow/DataflowNodeParameters.cpp
@@ -2,24 +2,27 @@
 
 #include "Dataflow/DataflowNodeParameters.h"
 
+#include "Dataflow/DataflowArchive.h"
+#include "Dataflow/DataflowContextCachingFactory.h"
 #include "Dataflow/DataflowInputOutput.h"
 #include "Dataflow/DataflowNode.h"
-
+#include "Serialization/Archive.h"
+#include "GeometryCollection/ManagedArrayCollection.h"
 
 namespace Dataflow
 {
 	uint64 FTimestamp::Invalid = 0;
-	uint64 FTimestamp::Current() { return FPlatformTime::Cycles64();  }
+	uint64 FTimestamp::Current() { return FPlatformTime::Cycles64(); }
 
-	void FContext::PushToCallstack(const FDataflowConnection* Connection) 
+	void FContext::PushToCallstack(const FDataflowConnection* Connection)
 	{
 #if DATAFLOW_EDITOR_EVALUATION
-		Callstack.Push(Connection); 
+		Callstack.Push(Connection);
 #endif
 	}
 
 	void FContext::PopFromCallstack(const FDataflowConnection* Connection)
-	{ 
+	{
 #if DATAFLOW_EDITOR_EVALUATION
 		ensure(Connection == Callstack.Top());
 		Callstack.Pop();
@@ -73,12 +76,12 @@ namespace Dataflow
 			}
 		}
 	}
-	
+
 	void FContextSingle::Evaluate(const FDataflowNode* Node, const FDataflowOutput* Output)
 	{
 		BeginContextEvaluation(*this, Node, Output);
 	}
-		
+
 	bool FContextSingle::Evaluate(const FDataflowOutput& Connection)
 	{
 		return Connection.EvaluateImpl(*this);
@@ -93,8 +96,83 @@ namespace Dataflow
 
 	bool FContextThreaded::Evaluate(const FDataflowOutput& Connection)
 	{
-		Connection.OutputLock->Lock(); ON_SCOPE_EXIT { Connection.OutputLock->Unlock(); };
+		Connection.OutputLock->Lock(); ON_SCOPE_EXIT{ Connection.OutputLock->Unlock(); };
 		return Connection.EvaluateImpl(*this);
 	}
-}
+
+
+	void FContextCache::Serialize(FArchive& Ar)
+	{
+
+		if (Ar.IsSaving())
+		{
+
+			const int64 NumElementsSavedPosition = Ar.Tell();
+			int64 NumElementsWritten = 0;
+			Ar << NumElementsWritten;
+
+			for (TPair<FContextCacheKey, TUniquePtr<FContextCacheElementBase>>& Elem : Pairs)
+			{
+				if (Elem.Value && Elem.Value->Property)
+				{
+					FProperty* Property = (FProperty*)Elem.Value->Property;
+					FName TypeName(Elem.Value->Property->GetCPPType());
+					FGuid NodeGuid = Elem.Value->NodeGuid;
+					uint32 NodeHash = Elem.Value->NodeHash;
+
+					if (FContextCachingFactory::GetInstance()->Contains(TypeName))
+					{
+						Ar << TypeName << Elem.Key << NodeGuid << NodeHash << Elem.Value->Timestamp;
+
+						DATAFLOW_OPTIONAL_BLOCK_WRITE_BEGIN()
+						{
+							FContextCachingFactory::GetInstance()->Serialize(Ar, {TypeName, NodeGuid, Elem.Value.Get(), NodeHash, Elem.Value->Timestamp});
+						}
+						DATAFLOW_OPTIONAL_BLOCK_WRITE_END();
+
+						NumElementsWritten++;
+					}
+				}
+			}
+
+
+			if (NumElementsWritten)
+			{
+				const int64 FinalPosition = Ar.Tell();
+				Ar.Seek(NumElementsSavedPosition);
+				Ar << NumElementsWritten;
+				Ar.Seek(FinalPosition);
+			}
+		}
+		else if (Ar.IsLoading())
+		{
+			int64 NumElementsWritten = 0;
+			Ar << NumElementsWritten;
+			for (int i = NumElementsWritten; i > 0; i--)
+			{
+				FName TypeName;
+				FGuid NodeGuid;
+				uint32 NodeHash;
+				FContextCacheKey InKey;
+				FTimestamp Timestamp = FTimestamp::Invalid;
+
+				Ar << TypeName << InKey << NodeGuid << NodeHash << Timestamp;
+
+				DATAFLOW_OPTIONAL_BLOCK_READ_BEGIN(FContextCachingFactory::GetInstance()->Contains(TypeName))
+				{
+					FContextCacheElementBase* NewElement = FContextCachingFactory::GetInstance()->Serialize(Ar, { TypeName, NodeGuid, nullptr, NodeHash, Timestamp });
+					check(NewElement);
+					NewElement->NodeGuid = NodeGuid;
+					NewElement->NodeHash = NodeHash;
+					NewElement->Timestamp = Timestamp;
+					this->Add(InKey, TUniquePtr<FContextCacheElementBase>(NewElement));
+				}
+				DATAFLOW_OPTIONAL_BLOCK_READ_ELSE()
+				{
+				}
+				DATAFLOW_OPTIONAL_BLOCK_READ_END();
+			}
+		}
+	}
+};
 
diff --git a/Engine/Source/Runtime/Experimental/Dataflow/Core/Public/Dataflow/DataflowConnection.h b/Engine/Source/Runtime/Experimental/Dataflow/Core/Public/Dataflow/DataflowConnection.h
--- a/Engine/Source/Runtime/Experimental/Dataflow/Core/Public/Dataflow/DataflowConnection.h
+++ b/Engine/Source/Runtime/Experimental/Dataflow/Core/Public/Dataflow/DataflowConnection.h
@@ -61,6 +61,8 @@ public:
 	FDataflowNode* GetOwningNode() { return OwningNode; }
 	const FDataflowNode* GetOwningNode() const { return OwningNode; }
 
+	const FProperty* GetProperty() const { return Property; }
+
 	Dataflow::FPin::EDirection GetDirection() const { return Direction; }
 	DATAFLOWCORE_API uint32 GetOffset() const;
 
@@ -73,7 +75,7 @@ public:
 	void SetName(FName InName) { Name = InName; }
 
 	void* RealAddress() const { ensure(OwningNode);  return (void*)((size_t)OwningNode + (size_t)GetOffset()); };
-	size_t CacheKey() const { return (size_t)RealAddress(); };
+	Dataflow::FContextCacheKey CacheKey() const { return GetTypeHash(Guid); };
 
 	virtual bool AddConnection(FDataflowConnection* In) { return false; };
 	virtual bool RemoveConnection(FDataflowConnection* In) { return false; }
diff --git a/Engine/Source/Runtime/Experimental/Dataflow/Core/Public/Dataflow/DataflowContextCachingFactory.h b/Engine/Source/Runtime/Experimental/Dataflow/Core/Public/Dataflow/DataflowContextCachingFactory.h
new file mode 100644
--- /dev/null
+++ b/Engine/Source/Runtime/Experimental/Dataflow/Core/Public/Dataflow/DataflowContextCachingFactory.h
@@ -0,0 +1,83 @@
+// Copyright Epic Games, Inc. All Rights Reserved.
+
+#pragma once
+
+#include "ChaosLog.h"
+#include "CoreMinimal.h"
+#include "Dataflow/DataflowNodeParameters.h"
+
+class FArchive;
+
+namespace Dataflow
+{
+	struct FContextCacheData {
+
+		FContextCacheData(FName InType, FGuid InNodeGuid, FContextCacheElementBase* InData, uint32 InNodeHash, FTimestamp InTimestamp)
+			: Type(InType), NodeGuid(InNodeGuid), Data(InData), NodeHash(InNodeHash), Timestamp(InTimestamp) {}
+
+		FName Type;
+		FGuid NodeGuid;
+		FContextCacheElementBase* Data = nullptr;
+		uint32 NodeHash;
+		FTimestamp Timestamp = FTimestamp::Invalid;
+	};
+
+	//
+	//
+	//
+	class DATAFLOWCORE_API FContextCachingFactory
+	{
+		typedef TFunction<FContextCacheElementBase*(FArchive& Ar, FContextCacheElementBase* InData)> FSerializeFunction;
+
+		// All Maps indexed by TypeName
+		TMap<FName, FSerializeFunction > CachingMap;		// [TypeName] -> Caching Funcitons
+		static FContextCachingFactory* Instance;
+		FContextCachingFactory() {}
+
+	public:
+		~FContextCachingFactory() { delete Instance; }
+
+		static FContextCachingFactory* GetInstance()
+		{
+			if (!Instance)
+			{
+				Instance = new FContextCachingFactory();
+			}
+			return Instance;
+		}
+
+		void RegisterSerializeFunction(const FName& Type, FSerializeFunction InSerializeFunc)
+		{
+			if (CachingMap.Contains(Type))
+			{
+				UE_LOG(LogChaos, Warning,
+					TEXT("Warning : Dataflow output caching registration conflicts with "
+						"existing type(%s)"), *Type.ToString());
+			}
+			else
+			{
+				CachingMap.Add(Type, InSerializeFunc);
+			}
+		}
+
+		template<class T>
+		static T& GetTypedElement(FContextCacheElementBase* InElement)
+		{
+			return (T&)InElement->GetTypedData<T>(InElement->Property);
+		}
+
+		template<class T>
+		static FContextCacheElementBase* NewTypedElement(T&& Data)
+		{
+			return new FContextCacheElement<T>(FGuid(), (FProperty*)nullptr, MoveTemp(Data), (uint32)0, FTimestamp::Invalid);
+		}
+
+
+		FContextCacheElementBase* Serialize(FArchive& Ar, FContextCacheData&& Data);
+
+		bool Contains(FName InType) const { return CachingMap.Contains(InType); }
+
+	};
+
+}
+
diff --git a/Engine/Source/Runtime/Experimental/Dataflow/Core/Public/Dataflow/DataflowInputOutput.h b/Engine/Source/Runtime/Experimental/Dataflow/Core/Public/Dataflow/DataflowInputOutput.h
--- a/Engine/Source/Runtime/Experimental/Dataflow/Core/Public/Dataflow/DataflowInputOutput.h
+++ b/Engine/Source/Runtime/Experimental/Dataflow/Core/Public/Dataflow/DataflowInputOutput.h
@@ -130,7 +130,7 @@ public:
 	{
 		if (Property)
 		{
-			Context.SetData(CacheKey(), Property, Forward<T>(InVal));
+			Context.SetData(CacheKey(), { OwningNode->GetGuid(), Property, OwningNode->GetValueHash(), Dataflow::FTimestamp::Current() }, Forward<T>(InVal));
 		}
 	}
 
@@ -138,7 +138,7 @@ public:
 	{
 		if (!this->Evaluate<T>(Context))
 		{
-			Context.SetData(CacheKey(), Property, Default);
+			Context.SetData(CacheKey(), { OwningNode->GetGuid(), Property, OwningNode->GetValueHash(), Dataflow::FTimestamp::Current() }, Default);
 		}
 
 		if (Context.HasData(CacheKey()))
@@ -171,7 +171,7 @@ const T& FDataflowInput::GetValue(Dataflow::FContext& Context, const T& Default)
 		{
 			if (!ConnectionOut->Evaluate<T>(Context))
 			{
-				Context.SetData(ConnectionOut->CacheKey(), Property, Default);
+				Context.SetData(ConnectionOut->CacheKey(), { OwningNode->GetGuid(), Property, OwningNode->GetValueHash(), Dataflow::FTimestamp::Current() }, Default);
 			}
 			if (Context.HasData(ConnectionOut->CacheKey()))
 			{
diff --git a/Engine/Source/Runtime/Experimental/Dataflow/Core/Public/Dataflow/DataflowNode.h b/Engine/Source/Runtime/Experimental/Dataflow/Core/Public/Dataflow/DataflowNode.h
--- a/Engine/Source/Runtime/Experimental/Dataflow/Core/Public/Dataflow/DataflowNode.h
+++ b/Engine/Source/Runtime/Experimental/Dataflow/Core/Public/Dataflow/DataflowNode.h
@@ -70,6 +70,7 @@ struct FDataflowNode
 	FName GetName() const { return Name; }
 	void SetName(FName InName) { Name = InName; }
 	Dataflow::FTimestamp GetTimestamp() const { return LastModifiedTimestamp;  }
+	DATAFLOWCORE_API uint32 GetValueHash();
 
 	static FName StaticType() { return FName("FDataflowNode"); }
 	virtual FName GetType() const { return StaticType(); }
@@ -122,8 +123,10 @@ struct FDataflowNode
 	DATAFLOWCORE_API void ClearOutputs();
 
 	DATAFLOWCORE_API FDataflowOutput* FindOutput(FName Name);
+	DATAFLOWCORE_API FDataflowOutput* FindOutput(uint32 GuidHash);
 	DATAFLOWCORE_API FDataflowOutput* FindOutput(void* Reference);
 	DATAFLOWCORE_API const FDataflowOutput* FindOutput(FName Name) const;
+	DATAFLOWCORE_API const FDataflowOutput* FindOutput(uint32 GuidHash) const;
 	DATAFLOWCORE_API const FDataflowOutput* FindOutput(const void* Reference) const;
 
 	/** Return a property's byte offset from the dataflow base node address using the full property name (must includes its parent struct property names). */
diff --git a/Engine/Source/Runtime/Experimental/Dataflow/Core/Public/Dataflow/DataflowNodeParameters.h b/Engine/Source/Runtime/Experimental/Dataflow/Core/Public/Dataflow/DataflowNodeParameters.h
--- a/Engine/Source/Runtime/Experimental/Dataflow/Core/Public/Dataflow/DataflowNodeParameters.h
+++ b/Engine/Source/Runtime/Experimental/Dataflow/Core/Public/Dataflow/DataflowNodeParameters.h
@@ -7,6 +7,7 @@
 #include "UObject/UnrealType.h"
 #include "Templates/UniquePtr.h"
 #include "GenericPlatform/GenericPlatformCriticalSection.h"
+#include "Serialization/Archive.h"
 
 class  UDataflow;
 struct FDataflowNode;
@@ -44,8 +45,10 @@ namespace Dataflow
 
 	struct FContextCacheElementBase 
 	{
-		FContextCacheElementBase(const FProperty* InProperty = nullptr, FTimestamp InTimestamp = FTimestamp::Invalid)
-			: Property(InProperty)
+		FContextCacheElementBase(FGuid InNodeGuid = FGuid(), const FProperty* InProperty = nullptr, uint32 InNodeHash = 0, FTimestamp InTimestamp = FTimestamp::Invalid)
+			: NodeGuid(InNodeGuid)
+			, Property(InProperty)
+			, NodeHash(InNodeHash)
 			, Timestamp(InTimestamp)
 		{}
 		virtual ~FContextCacheElementBase() {}
@@ -53,15 +56,17 @@ namespace Dataflow
 		template<typename T>
 		const T& GetTypedData(const FProperty* PropertyIn) const;
 		
+		FGuid NodeGuid;
 		const FProperty* Property = nullptr;
+		uint32 NodeHash = 0;
 		FTimestamp Timestamp = FTimestamp::Invalid;
 	};
 
 	template<class T>
 	struct FContextCacheElement : public FContextCacheElementBase 
 	{
-		FContextCacheElement(const FProperty* InProperty, T&& InData, FTimestamp Timestamp)
-			: FContextCacheElementBase(InProperty, Timestamp)
+		FContextCacheElement(FGuid InNodeGuid, const FProperty* InProperty, T&& InData, uint32 InNodeHash, FTimestamp Timestamp)
+			: FContextCacheElementBase(InNodeGuid, InProperty, InNodeHash, Timestamp)
 			, Data(Forward<T>(InData))
 		{}
 		
@@ -70,18 +75,38 @@ namespace Dataflow
 	};
 
 	template<class T>
-	const T& FContextCacheElementBase::GetTypedData(const FProperty* PropertyIn) const
+	const T& FContextCacheElementBase::GetTypedData(const FProperty* InProperty) const
 	{
-		check(PropertyIn);
-		// check(PropertyIn->IsA<T>()); // @todo(dataflow) compile error for non-class T; find alternatives
-		check(Property->SameType(PropertyIn));
+		check(InProperty);
+		check(Property->SameType(InProperty));
 		return static_cast<const FContextCacheElement<T>&>(*this).Data;
 	}
 
-	struct FContextCache : public TMap<int64, TUniquePtr<FContextCacheElementBase>>
+	typedef uint32 FContextCacheKey;
+	
+
+	struct FContextCache : public TMap<FContextCacheKey, TUniquePtr<FContextCacheElementBase>>
 	{
-		// @todo(dataflow) make an API for FContextCache
+		DATAFLOWCORE_API void Serialize(FArchive& Ar);
 	};
+};
+
+inline FArchive& operator<<(FArchive& Ar, Dataflow::FTimestamp& ValueIn)
+{
+	Ar << ValueIn.Value;
+	Ar << ValueIn.Invalid;
+	return Ar;
+}
+
+inline FArchive& operator<<(FArchive& Ar, Dataflow::FContextCache& ValueIn)
+{
+	ValueIn.Serialize(Ar);
+	return Ar;
+}
+
+
+namespace Dataflow
+{
 
 	class FContext
 	{
@@ -92,6 +117,7 @@ namespace Dataflow
 		FContext(const FContext&) = delete;
 		FContext& operator=(const FContext&) = delete;
 
+		FContextCache DataStore;
 
 	public:
 		FContext(FTimestamp InTimestamp)
@@ -108,6 +134,8 @@ namespace Dataflow
 
 		virtual FName GetType() const { return FContext::StaticType(); }
 
+		virtual int32 GetKeys(TSet<FContextCacheKey>& InKeys) { return DataStore.GetKeys(InKeys); }
+
 		template<class T>
 		const T* AsType() const
 		{
@@ -118,36 +146,35 @@ namespace Dataflow
 			return nullptr;
 		}
 
-		virtual void SetDataImpl(int64 Key, TUniquePtr<FContextCacheElementBase>&& DataStoreEntry) = 0;
+		virtual void SetDataImpl(FContextCacheKey Key, TUniquePtr<FContextCacheElementBase>&& DataStoreEntry) = 0;
 		
 		template<typename T>
-		void SetData(size_t Key, const FProperty* Property, T&& Value)
+		void SetData(FContextCacheKey Key, FContextCacheElementBase&& Data, T&& Value)
 		{
-			int64 IntKey = (int64)Key;
-			TUniquePtr<FContextCacheElement<T>> DataStoreEntry = MakeUnique<FContextCacheElement<T>>(Property, Forward<T>(Value), FTimestamp::Current());
+			FContextCacheKey IntKey = (FContextCacheKey)Key;
+			TUniquePtr<FContextCacheElement<T>> DataStoreEntry = MakeUnique<FContextCacheElement<T>>(Data.NodeGuid, Data.Property, Forward<T>(Value), Data.NodeHash, Data.Timestamp);
 
 			SetDataImpl(IntKey, MoveTemp(DataStoreEntry));
 		}
 
-		
-		virtual TUniquePtr<FContextCacheElementBase>* GetDataImpl(int64 Key) = 0;
+		virtual TUniquePtr<FContextCacheElementBase>* GetDataImpl(FContextCacheKey Key) = 0;
 
 		template<class T>
-		const T& GetData(size_t Key, const FProperty* Property, const T& Default = T())
+		const T& GetData(FContextCacheKey Key, const FProperty* InProperty, const T& Default = T())
 		{
 			if (TUniquePtr<FContextCacheElementBase>* Cache = GetDataImpl(Key))
 			{
-				return (*Cache)->GetTypedData<T>(Property);
+				return (*Cache)->GetTypedData<T>(InProperty);
 			}
 			return Default;
 		}
 
 		
-		virtual bool HasDataImpl(int64 Key, FTimestamp InTimestamp = FTimestamp::Invalid) = 0;
+		virtual bool HasDataImpl(FContextCacheKey Key, FTimestamp InTimestamp = FTimestamp::Invalid) = 0;
 		
-		bool HasData(size_t Key, FTimestamp InTimestamp = FTimestamp::Invalid)
+		bool HasData(FContextCacheKey Key, FTimestamp InTimestamp = FTimestamp::Invalid)
 		{
-			int64 IntKey = (int64)Key;
+			FContextCacheKey IntKey = (FContextCacheKey)Key;
 			return HasDataImpl(Key, InTimestamp);
 		}
 
@@ -158,6 +185,12 @@ namespace Dataflow
 			return IsEmptyImpl();
 		}
 
+		virtual void Serialize(FArchive& Ar)
+		{
+			Ar << Timestamp;
+			Ar << DataStore;
+		}
+
 
 		FTimestamp GetTimestamp() const { return Timestamp; }
 		virtual void Evaluate(const FDataflowNode* Node, const FDataflowOutput* Output) = 0;
@@ -167,6 +200,7 @@ namespace Dataflow
 		DATAFLOWCORE_API void PopFromCallstack(const FDataflowConnection* Connection);
 		DATAFLOWCORE_API bool IsInCallstack(const FDataflowConnection* Connection) const;
 
+
 	private:
 #if DATAFLOW_EDITOR_EVALUATION
 		TArray<const FDataflowConnection*> Callstack;
@@ -195,7 +229,6 @@ namespace Dataflow
 
 	class FContextSingle : public FContext
 	{
-		FContextCache DataStore;
 
 	public:
 		DATAFLOW_CONTEXT_INTERNAL(FContext, FContextSingle);
@@ -204,17 +237,17 @@ namespace Dataflow
 			: FContext(InTime)
 		{}
 
-		virtual void SetDataImpl(int64 Key, TUniquePtr<FContextCacheElementBase>&& DataStoreEntry) override
+		virtual void SetDataImpl(FContextCacheKey Key, TUniquePtr<FContextCacheElementBase>&& DataStoreEntry) override
 		{
 			DataStore.Emplace(Key, MoveTemp(DataStoreEntry));
 		}
 
-		virtual TUniquePtr<FContextCacheElementBase>* GetDataImpl(int64 Key) override
+		virtual TUniquePtr<FContextCacheElementBase>* GetDataImpl(FContextCacheKey Key) override
 		{
 			return DataStore.Find(Key);
 		}
 
-		virtual bool HasDataImpl(int64 Key, FTimestamp InTimestamp = FTimestamp::Invalid) override
+		virtual bool HasDataImpl(FContextCacheKey Key, FTimestamp InTimestamp = FTimestamp::Invalid) override
 		{
 			return DataStore.Contains(Key) && DataStore[Key]->Timestamp >= InTimestamp;
 		}
@@ -226,11 +259,11 @@ namespace Dataflow
 
 		DATAFLOWCORE_API virtual void Evaluate(const FDataflowNode* Node, const FDataflowOutput* Output) override;
 		DATAFLOWCORE_API virtual bool Evaluate(const FDataflowOutput& Connection) override;
+
 	};
 	
 	class FContextThreaded : public FContext
 	{
-		FContextCache DataStore;
 		TSharedPtr<FCriticalSection> CacheLock;
 
 	public:
@@ -243,21 +276,21 @@ namespace Dataflow
 			CacheLock = MakeShared<FCriticalSection>();
 		}
 
-		virtual void SetDataImpl(int64 Key, TUniquePtr<FContextCacheElementBase>&& DataStoreEntry) override
+		virtual void SetDataImpl(FContextCacheKey Key, TUniquePtr<FContextCacheElementBase>&& DataStoreEntry) override
 		{
 			CacheLock->Lock(); ON_SCOPE_EXIT { CacheLock->Unlock(); };
 			
 			DataStore.Emplace(Key, MoveTemp(DataStoreEntry));
 		}
 
-		virtual TUniquePtr<FContextCacheElementBase>* GetDataImpl(int64 Key) override
+		virtual TUniquePtr<FContextCacheElementBase>* GetDataImpl(FContextCacheKey Key) override
 		{
 			CacheLock->Lock(); ON_SCOPE_EXIT { CacheLock->Unlock(); };
 			
 			return DataStore.Find(Key);
 		}
 
-		virtual bool HasDataImpl(int64 Key, FTimestamp InTimestamp = FTimestamp::Invalid) override
+		virtual bool HasDataImpl(FContextCacheKey Key, FTimestamp InTimestamp = FTimestamp::Invalid) override
 		{
 			CacheLock->Lock(); ON_SCOPE_EXIT { CacheLock->Unlock(); };
 			
@@ -271,6 +304,7 @@ namespace Dataflow
 
 		DATAFLOWCORE_API virtual void Evaluate(const FDataflowNode* Node, const FDataflowOutput* Output) override;
 		DATAFLOWCORE_API virtual bool Evaluate(const FDataflowOutput& Connection) override;
+
 	};
 
 }
diff --git a/Engine/Source/Runtime/Experimental/Dataflow/Engine/Private/Dataflow/DataflowContent.cpp b/Engine/Source/Runtime/Experimental/Dataflow/Engine/Private/Dataflow/DataflowContent.cpp
--- a/Engine/Source/Runtime/Experimental/Dataflow/Engine/Private/Dataflow/DataflowContent.cpp
+++ b/Engine/Source/Runtime/Experimental/Dataflow/Engine/Private/Dataflow/DataflowContent.cpp
@@ -6,35 +6,100 @@
 #include "Animation/Skeleton.h"
 #include "Components/SkeletalMeshComponent.h"
 #include "Dataflow/DataflowNodeParameters.h"
+#include "Dataflow/DataflowObject.h"
 #include "Engine/SkeletalMesh.h"
 #include "GameFramework/Actor.h"
 #include "PreviewScene.h"
 
 #include UE_INLINE_GENERATED_CPP_BY_NAME(DataflowContent)
 
+
+//
+// UDataflowBaseContent
+//
+
+void UDataflowBaseContent::SetIsDirty(bool InDirty) 
+{ 
+	bIsDirty = InDirty;
+}
+
+
 UDataflowBaseContent::UDataflowBaseContent()
 {
 }
 
-void UDataflowBaseContent::SetLastModifiedTimestamp(Dataflow::FTimestamp InTimestamp) 
+void UDataflowBaseContent::SetDataflowOwner(const TObjectPtr<UObject>& InOwner)
+{
+	DataflowOwner = InOwner;
+	if (DataflowContext)
+	{
+		DataflowContext->Owner = InOwner;  
+		SetIsDirty(true);
+	}
+}
+
+TObjectPtr<UObject> UDataflowBaseContent::GetDataflowOwner() const 
+{
+	if (DataflowContext)
+	{
+		ensure(DataflowContext->Owner == DataflowOwner);
+	}
+	return DataflowContext ? DataflowContext->Owner : DataflowOwner; 
+}
+
+
+void UDataflowBaseContent::SetLastModifiedTimestamp(Dataflow::FTimestamp InTimestamp, bool bMakeDirty) 
 { 
 	if (InTimestamp.IsInvalid() || LastModifiedTimestamp < InTimestamp)
 	{
 		LastModifiedTimestamp = InTimestamp; 
-		bIsDirty = true;
+		if (bMakeDirty)
+		{
+			SetIsDirty(true);
+			MarkPackageDirty(); 
+		}
 	}
 }
 
-void UDataflowBaseContent::BuildBaseContent(TObjectPtr<UObject> ContentOwner)
+void UDataflowBaseContent::SetDataflowContext(const TSharedPtr<Dataflow::FEngineContext>& InContext) 
+{ 
+	DataflowContext = InContext;  
+	SetIsDirty(true); 
+	MarkPackageDirty();
+}
+
+
+void UDataflowBaseContent::BuildBaseContent(TObjectPtr<UObject> InDataflowOwner)
 {
-	DataflowContext = MakeShared<Dataflow::FEngineContext>(ContentOwner, DataflowAsset, FPlatformTime::Cycles64());
+	DataflowOwner = InDataflowOwner;
+	DataflowContext = MakeShared<Dataflow::FEngineContext>(InDataflowOwner, DataflowAsset, FPlatformTime::Cycles64());
 	LastModifiedTimestamp = DataflowContext->GetTimestamp();
+	SetIsDirty(true);
+	MarkPackageDirty();
+}
+
+void UDataflowBaseContent::Serialize(FArchive& Ar)
+{
+	Super::Serialize(Ar);
+	Ar << LastModifiedTimestamp;
+
+	if (!DataflowContext)
+	{
+		DataflowContext = MakeShared<Dataflow::FEngineContext>(DataflowOwner, DataflowAsset, LastModifiedTimestamp);
+	}
+	DataflowContext->Serialize(Ar);
 }
 
+//
+// UDataflowSkeletalContent
+//
+
+
 UDataflowSkeletalContent::UDataflowSkeletalContent() : Super()
 {
 }
 
+
 void UDataflowSkeletalContent::RegisterWorldContent(FPreviewScene* PreviewScene, AActor* RootActor)
 {
 	Super::RegisterWorldContent(PreviewScene, RootActor);
@@ -103,7 +168,7 @@ void UDataflowSkeletalContent::SetSkeletalMesh(const TObjectPtr<USkeletalMesh>&
 
  		UpdateAnimationInstance();
 	}
-	bIsDirty = true;
+	SetIsDirty(true);
 }
 
 void UDataflowSkeletalContent::SetAnimationAsset(const TObjectPtr<UAnimationAsset>& SkeletalAnimationAsset)
@@ -113,7 +178,7 @@ void UDataflowSkeletalContent::SetAnimationAsset(const TObjectPtr<UAnimationAsse
 	{
  		UpdateAnimationInstance();
 	}
-	bIsDirty = true;
+	SetIsDirty(true);
 }
 
 void UDataflowSkeletalContent::SetSkeleton(const TObjectPtr<USkeleton>& SkeletonAsset)
@@ -126,7 +191,7 @@ void UDataflowSkeletalContent::SetSkeleton(const TObjectPtr<USkeleton>& Skeleton
 			SetSkeletalMesh(nullptr);
 		}
 	}
-	bIsDirty = true;
+	SetIsDirty(true);
 }
 
 #if WITH_EDITOR
@@ -169,4 +234,3 @@ FVector2f UDataflowSkeletalContent::GetSimulationRange() const
 	return FVector2f(0.0f, 0.0f);
 }
 
-
diff --git a/Engine/Source/Runtime/Experimental/Dataflow/Engine/Public/Dataflow/DataflowContent.h b/Engine/Source/Runtime/Experimental/Dataflow/Engine/Public/Dataflow/DataflowContent.h
--- a/Engine/Source/Runtime/Experimental/Dataflow/Engine/Public/Dataflow/DataflowContent.h
+++ b/Engine/Source/Runtime/Experimental/Dataflow/Engine/Public/Dataflow/DataflowContent.h
@@ -85,33 +85,26 @@ class DATAFLOWENGINE_API UDataflowBaseContent : public UDataflowContextObject
 public:
 	UDataflowBaseContent();
 
-	/** Data flow asset that we will edit */
-	UPROPERTY(EditAnywhere, Category = "Dataflow")
-	TObjectPtr<UDataflow> DataflowAsset = nullptr;
-
-	/** Data flow terminal path for evaluation */
-	UPROPERTY(EditAnywhere, Category = "Dataflow")
-	FString DataflowTerminal = "";
-	
 	/** 
 	*	Dirty - State Invalidation
 	*   Check if non-graph specific data has been changed, this usually requires a re-render 
 	*/
 	bool IsDirty() const { return bIsDirty; }
-	void SetIsDirty(bool InDirty) { bIsDirty = InDirty; }
+	void SetIsDirty(bool InDirty);
 
 	/** 
 	*	LastModifiedTimestamp - State Invalidation 
 	*   Dataflow timestamp accessors can be used to see if the EvaluationContext has been invalidated. 
 	*/
-	void SetLastModifiedTimestamp(Dataflow::FTimestamp InTimestamp);
+	void SetLastModifiedTimestamp(Dataflow::FTimestamp InTimestamp, bool bMakeDirty =true);
 	const Dataflow::FTimestamp& GetLastModifiedTimestamp() const { return LastModifiedTimestamp; }
 
 	/**  
 	*	Context - Dataflow Evaluation State
 	*   Dataflow context stores the evaluated state of the graph. 
 	*/
-	void SetDataflowContext(const TSharedPtr<Dataflow::FEngineContext>& InContext) { DataflowContext = InContext; bIsDirty = true; }
+	void SetDataflowContext(const TSharedPtr<Dataflow::FEngineContext>& InContext);
+	TSharedPtr<Dataflow::FEngineContext>& GetDataflowContext() { return DataflowContext; }
 	const TSharedPtr<Dataflow::FEngineContext>& GetDataflowContext() const { return DataflowContext; }
 
 	/** Return the simulation time range to be used in the simulation viewport */
@@ -124,34 +117,48 @@ public:
 	virtual void UnregisterWorldContent(FPreviewScene* PreviewScene) {}
 
 	/** Build the content context, timestamp*/
-	void BuildBaseContent(TObjectPtr<UObject> ContentOwner);
+	void BuildBaseContent(TObjectPtr<UObject> InDataflowOwner);
 
 	/** Collect reference objects for GC */
 	virtual void AddContentObjects(FReferenceCollector& Collector) {}
 	
 	/** Data flow owner accessors */
-	void SetDataflowOwner(const TObjectPtr<UObject>& InOwner) { if(DataflowContext) { DataflowContext->Owner = InOwner;  bIsDirty = true; }}
-	TObjectPtr<UObject> GetDataflowOwner() const { return DataflowContext ? DataflowContext->Owner : nullptr; }
-
+	void SetDataflowOwner(const TObjectPtr<UObject>& InOwner);
+	TObjectPtr<UObject> GetDataflowOwner() const;
+	
 	/** Data flow asset accessors */
-	void SetDataflowAsset(const TObjectPtr<UDataflow>& InAsset) { DataflowAsset = InAsset;  bIsDirty = true;}
+	void SetDataflowAsset(const TObjectPtr<UDataflow>& InAsset) { DataflowAsset = InAsset;  SetIsDirty(true);}
 	const TObjectPtr<UDataflow>& GetDataflowAsset() const { return DataflowAsset; }
 
 	/** Data flow terminal accessors */
-	void SetDataflowTerminal(const FString& InPath) { DataflowTerminal = InPath;  bIsDirty = true;}
+	void SetDataflowTerminal(const FString& InPath) { DataflowTerminal = InPath;  SetIsDirty(true);}
 	const FString& GetDataflowTerminal() const { return DataflowTerminal; }
 
-	
-protected :
-	
+	virtual void Serialize(FArchive& Ar);
+
+protected:
+
+	/** Data flow asset that we will edit */
+	UPROPERTY(EditAnywhere, Category = "Dataflow")
+	TObjectPtr<UObject> DataflowOwner = nullptr;
+
+	/** Data flow asset that we will edit */
+	UPROPERTY(EditAnywhere, Category = "Dataflow")
+	TObjectPtr<UDataflow> DataflowAsset = nullptr;
+
+	/** Data flow terminal path for evaluation */
+	UPROPERTY(EditAnywhere, Category = "Dataflow")
+	FString DataflowTerminal = "";
+
 	/**  Engine context to be used for dataflow evaluation */
     TSharedPtr<Dataflow::FEngineContext> DataflowContext = nullptr;
 
     /** Last data flow evaluated node time stamp */
-    Dataflow::FTimestamp LastModifiedTimestamp = Dataflow::FTimestamp::Invalid;
+	Dataflow::FTimestamp LastModifiedTimestamp = Dataflow::FTimestamp::Invalid;
 
     /** Dirty flag to trigger rendering. Do we need that? since when accessing the member by non const ref we will not dirty it */
-    bool bIsDirty = true;
+	UPROPERTY()
+	bool bIsDirty = true;
 };
 
 /** 
@@ -166,18 +173,6 @@ public:
 	UDataflowSkeletalContent();
 	virtual ~UDataflowSkeletalContent() override{}
 
-	/** Data flow skeletal mesh*/
-	UPROPERTY(EditAnywhere, Category = "Preview")
-	TObjectPtr<USkeletalMesh> SkeletalMesh = nullptr;
-
-	/** Animation asset to be used to preview simulation */
-	UPROPERTY(EditAnywhere, Category = "Preview")
-	TObjectPtr<UAnimationAsset> AnimationAsset;
-
-	/** Data flow skeleton */
-	UPROPERTY(EditAnywhere, Category = "Skeleton")
-	TObjectPtr<USkeleton> Skeleton = nullptr;
-
 	/** Return the simulation time range to be used in the simulation viewport */
 	virtual FVector2f GetSimulationRange() const override; 
 
@@ -212,12 +207,26 @@ public:
 	/** Data flow animation asset accessors */
 	void SetAnimationAsset(const TObjectPtr<UAnimationAsset>& InAnimation);
 	const TObjectPtr<UAnimationAsset>& GetAnimationAsset() const { return AnimationAsset; }
-	
-protected :
-	
+
+protected:
+
+	/** Data flow skeletal mesh*/
+	UPROPERTY(EditAnywhere, Category = "Preview", Transient, SkipSerialization)
+	TObjectPtr<USkeletalMesh> SkeletalMesh = nullptr;
+
+	/** Animation asset to be used to preview simulation */
+	UPROPERTY(EditAnywhere, Category = "Preview", Transient, SkipSerialization)
+	TObjectPtr<UAnimationAsset> AnimationAsset;
+
+	/** Data flow skeleton */
+	UPROPERTY(EditAnywhere, Category = "Skeleton", Transient, SkipSerialization)
+	TObjectPtr<USkeleton> Skeleton = nullptr;
+
 	/** Skeletal mesh component used in the preview scene */
+	UPROPERTY(Transient, SkipSerialization)
 	TObjectPtr<USkeletalMeshComponent> SkeletalMeshComponent;
 
 	/** Anim node instance used with skelmesh component */
+	UPROPERTY(Transient, SkipSerialization)
 	TObjectPtr<UAnimSingleNodeInstance> AnimationNodeInstance;
 };
diff --git a/Engine/Source/Runtime/Experimental/Dataflow/Engine/Public/Dataflow/DataflowObjectInterface.h b/Engine/Source/Runtime/Experimental/Dataflow/Engine/Public/Dataflow/DataflowObjectInterface.h
--- a/Engine/Source/Runtime/Experimental/Dataflow/Engine/Public/Dataflow/DataflowObjectInterface.h
+++ b/Engine/Source/Runtime/Experimental/Dataflow/Engine/Public/Dataflow/DataflowObjectInterface.h
@@ -9,6 +9,7 @@
 
 class UDataflow;
 class UObject;
+class FArchive;
 
 namespace Dataflow
 {
@@ -28,6 +29,13 @@ namespace Dataflow
 		TObjectPtr<UDataflow> Graph = nullptr;
 
 		~TEngineContext(){}
+
+		int32 GetKeys(TSet<FContextCacheKey>& InKeys) { return Base::GetKeys(InKeys); }
+
+		TUniquePtr<FContextCacheElementBase>* GetBaseData(FContextCacheKey Key) { return Base::GetDataImpl(Key); }
+
+		virtual void Serialize(FArchive& Ar) { Base::Serialize(Ar); }
+
 	};
 
 	typedef TEngineContext<FContextSingle> FEngineContext;
