[NetPhysics] Resim render interpolation directional decay

diff --git a/Engine/Source/Runtime/Experimental/Chaos/Private/PhysicsProxy/ClusterUnionPhysicsProxy.cpp b/Engine/Source/Runtime/Experimental/Chaos/Private/PhysicsProxy/ClusterUnionPhysicsProxy.cpp
--- a/Engine/Source/Runtime/Experimental/Chaos/Private/PhysicsProxy/ClusterUnionPhysicsProxy.cpp
+++ b/Engine/Source/Runtime/Experimental/Chaos/Private/PhysicsProxy/ClusterUnionPhysicsProxy.cpp
@@ -499,10 +499,16 @@ namespace Chaos
 			};
 
 			const bool bIsReplicationErrorSmoothing = InterpolationData.IsErrorSmoothing();
+			bool DirectionalDecayPerformed = false;
 			InterpolationData.UpdateError(SolverSyncTimestamp, AsyncFixedTimeStep);
 
 			if (const FVec3* Prev = LerpHelper(PullData.X, ProxyTimestamp->OverWriteX))
 			{
+				if (GetRenderInterpErrorDirectionalDecayMultiplier() > 0.0f)
+				{
+					DirectionalDecayPerformed = InterpolationData.DirectionalDecay(NextPullData->X - *Prev);
+				}
+
 				const FVec3 NewX = bIsReplicationErrorSmoothing ?
 					FMath::Lerp(*Prev, NextPullData->X, *Alpha) + InterpolationData.GetErrorX(*Alpha) :
 					FMath::Lerp(*Prev, NextPullData->X, *Alpha);
@@ -536,7 +542,7 @@ namespace Chaos
 			{
 				Chaos::FDebugDrawQueue::GetInstance().DrawDebugBox(NextPullData->X, FVector(2, 1, 1), NextPullData->R, FColor::Yellow, false, 5.f, 0, 0.5f);
 				Chaos::FDebugDrawQueue::GetInstance().DrawDebugDirectionalArrow(PullData.X, NextPullData->X, 0.5f, FColor::Yellow, false, 5.0f, 0, 0.5f);
-				Chaos::FDebugDrawQueue::GetInstance().DrawDebugBox(Particle_External->GetX(), FVector(2, 1, 1), Particle_External->R(), FColor::Green, false, 5.f, 0, 0.5f);
+				Chaos::FDebugDrawQueue::GetInstance().DrawDebugBox(Particle_External->GetX(), FVector(2, 1, 1), Particle_External->R(), DirectionalDecayPerformed ? FColor::Cyan : FColor::Green, false, 5.f, 0, 0.5f);
 
 				if (bIsReplicationErrorSmoothing)
 				{
diff --git a/Engine/Source/Runtime/Experimental/Chaos/Private/PhysicsProxy/GeometryCollectionPhysicsProxy.cpp b/Engine/Source/Runtime/Experimental/Chaos/Private/PhysicsProxy/GeometryCollectionPhysicsProxy.cpp
--- a/Engine/Source/Runtime/Experimental/Chaos/Private/PhysicsProxy/GeometryCollectionPhysicsProxy.cpp
+++ b/Engine/Source/Runtime/Experimental/Chaos/Private/PhysicsProxy/GeometryCollectionPhysicsProxy.cpp
@@ -4321,11 +4321,15 @@ bool FGeometryCollectionPhysicsProxy::PullFromPhysicsState(const Chaos::FDirtyGe
 		// second : interpolate-able ones
 		if (bNeedInterpolation)
 		{
-			InterpolationData.UpdateError(SolverSyncTimestamp, AsyncFixedTimeStep);
-
 			const FGeometryCollectionResults& PrevResults = PullData.Results();
 			const FGeometryCollectionResults& NextResults = NextPullData->Results();
 
+			InterpolationData.UpdateError(SolverSyncTimestamp, AsyncFixedTimeStep);
+			if (IPhysicsProxyBase::GetRenderInterpErrorDirectionalDecayMultiplier() > 0.0f && PrevResults.GetNumEntries() && NextResults.GetNumEntries())
+			{
+				InterpolationData.DirectionalDecay(NextResults.GetPositions(0).ParticleX - PrevResults.GetPositions(0).ParticleX);
+			}
+
 			TManagedArray<FVector3f>* LinearVelocities = GameThreadCollection.GetLinearVelocitiesAttribute();
 			TManagedArray<FVector3f>* AngularVelocities = GameThreadCollection.GetAngularVelocitiesAttribute();
 
diff --git a/Engine/Source/Runtime/Experimental/Chaos/Private/PhysicsProxy/SingleParticlePhysicsProxy.cpp b/Engine/Source/Runtime/Experimental/Chaos/Private/PhysicsProxy/SingleParticlePhysicsProxy.cpp
--- a/Engine/Source/Runtime/Experimental/Chaos/Private/PhysicsProxy/SingleParticlePhysicsProxy.cpp
+++ b/Engine/Source/Runtime/Experimental/Chaos/Private/PhysicsProxy/SingleParticlePhysicsProxy.cpp
@@ -319,10 +319,10 @@ bool FSingleParticlePhysicsProxy::PullFromPhysicsState(const Chaos::FDirtyRigidP
 			if (bUpdatePositionFromSimulation)
 			{
 				const bool bIsReplicationErrorSmoothing = InterpolationData.IsErrorSmoothing();
+				bool DirectionalDecayPerformed = false;
 #if RENDERINTERP_ERRORVELOCITYSMOOTHING
 				const bool bIsErrorVelocitySmoothing = InterpolationData.IsErrorVelocitySmoothing();
 #endif
-
 				InterpolationData.UpdateError(SolverSyncTimestamp, AsyncFixedTimeStep);
 
 				if (const FVec3* Prev = LerpHelper(PullData.X, ProxyTimestamp->OverWriteX))
@@ -330,6 +330,11 @@ bool FSingleParticlePhysicsProxy::PullFromPhysicsState(const Chaos::FDirtyRigidP
 					FVec3 Target = FMath::Lerp(*Prev, NextPullData->X, *Alpha);
 					if (bIsReplicationErrorSmoothing)
 					{
+						if (GetRenderInterpErrorDirectionalDecayMultiplier() > 0.0f)
+						{
+							DirectionalDecayPerformed = InterpolationData.DirectionalDecay(NextPullData->X - *Prev);
+						}
+
 						Target += InterpolationData.GetErrorX(*Alpha);
 
 #if RENDERINTERP_ERRORVELOCITYSMOOTHING
@@ -366,7 +371,7 @@ bool FSingleParticlePhysicsProxy::PullFromPhysicsState(const Chaos::FDirtyRigidP
 				{
 					Chaos::FDebugDrawQueue::GetInstance().DrawDebugBox(NextPullData->X, FVector(2, 1, 1), NextPullData->R, FColor::Yellow, false, 5.f, 0, 0.5f);
 					Chaos::FDebugDrawQueue::GetInstance().DrawDebugDirectionalArrow(PullData.X, NextPullData->X, 0.5f, FColor::Yellow, false, 5.0f, 0, 0.5f);
-					Chaos::FDebugDrawQueue::GetInstance().DrawDebugBox(Rigid->X(), FVector(2, 1, 1), Rigid->R(), FColor::Green, false, 5.f, 0, 0.5f);
+					Chaos::FDebugDrawQueue::GetInstance().DrawDebugBox(Rigid->X(), FVector(2, 1, 1), Rigid->R(), DirectionalDecayPerformed ? FColor::Cyan : FColor::Green, false, 5.f, 0, 0.5f);
 
 					if (bIsReplicationErrorSmoothing)
 					{
diff --git a/Engine/Source/Runtime/Experimental/Chaos/Public/Chaos/Framework/PhysicsProxyBase.h b/Engine/Source/Runtime/Experimental/Chaos/Public/Chaos/Framework/PhysicsProxyBase.h
--- a/Engine/Source/Runtime/Experimental/Chaos/Public/Chaos/Framework/PhysicsProxyBase.h
+++ b/Engine/Source/Runtime/Experimental/Chaos/Public/Chaos/Framework/PhysicsProxyBase.h
@@ -131,7 +131,7 @@ public:
 	}
 	static float GetRenderInterpMaximumErrorCorrectionBeforeSnapping()
 	{
-		static float RenderInterpMaximumErrorCorrectionBeforeSnapping = 100.0f;
+		static float RenderInterpMaximumErrorCorrectionBeforeSnapping = 250.0f;
 		static FAutoConsoleVariableRef CVarRenderInterpErrorCorrectionMaximumError(TEXT("p.RenderInterp.MaximumErrorCorrectionBeforeSnapping"), RenderInterpMaximumErrorCorrectionBeforeSnapping, TEXT("Maximum error correction in cm before we stop interpolating and snap to target."));
 		return RenderInterpMaximumErrorCorrectionBeforeSnapping;
 	}
@@ -147,6 +147,12 @@ public:
 		static FAutoConsoleVariableRef CVarRenderInterpDebugDraw(TEXT("p.RenderInterp.DebugDraw"), RenderInterpDebugDraw, TEXT("Draw debug lines for physics render interpolation, also needs p.Chaos.DebugDraw.Enabled set"));
 		return RenderInterpDebugDraw;
 	}
+	static float GetRenderInterpErrorDirectionalDecayMultiplier()
+	{
+		static float RenderInterpErrorDirectionalDecayMultiplier = 0.0f;
+		static FAutoConsoleVariableRef CVarRenderInterpErrorDirectionalDecayMultiplier(TEXT("p.RenderInterp.DirectionalDecayMultiplier"), RenderInterpErrorDirectionalDecayMultiplier, TEXT("Decay error offset in the direction that the physics object is moving, value is multiplier of projected offset direction, 0.25 means a 25% decay of the magnitude in the direction of physics travel. Deactivate by setting to 0."));
+		return RenderInterpErrorDirectionalDecayMultiplier;
+	}
 
 protected:
 	// Ensures that derived classes can successfully call this destructor
@@ -240,6 +246,21 @@ struct FProxyInterpolationError : FProxyInterpolationBase
 		return false;
 	}
 
+	virtual bool DirectionalDecay(Chaos::FVec3 Direction)
+	{
+		if (IsErrorSmoothing() && SimTicks > 0)
+		{
+			const Chaos::FVec3 DirectionNormal = Direction.GetSafeNormal();
+			Chaos::FRealDouble DotProd = Chaos::FVec3::DotProduct(DirectionNormal, ErrorX);
+			if (DotProd > 0.0f)
+			{
+				Chaos::FVec3 DirProjection = ErrorX.ProjectOnToNormal(DirectionNormal) * IPhysicsProxyBase::GetRenderInterpErrorDirectionalDecayMultiplier();
+				ErrorX -= DirProjection;
+				return true;
+			}
+		}
+		return false;
+	}
 
 protected:
 	void DecayError()
