Chaos Cloth: add map for global damping.

diff --git a/Engine/Plugins/Experimental/ChaosCloth/Source/Chaos/Private/ChaosCloth/ChaosClothingSimulationCloth.cpp b/Engine/Plugins/Experimental/ChaosCloth/Source/Chaos/Private/ChaosCloth/ChaosClothingSimulationCloth.cpp
--- a/Engine/Plugins/Experimental/ChaosCloth/Source/Chaos/Private/ChaosCloth/ChaosClothingSimulationCloth.cpp
+++ b/Engine/Plugins/Experimental/ChaosCloth/Source/Chaos/Private/ChaosCloth/ChaosClothingSimulationCloth.cpp
@@ -1066,7 +1066,7 @@ PRAGMA_ENABLE_DEPRECATION_WARNINGS
 			FictitiousAngularScale);
 		if (Solver->IsForceBasedSolver())
 		{
-			Solver->SetProperties(ParticleRangeId, ConfigProperties);
+			Solver->SetProperties(ParticleRangeId, ConfigProperties, LODData[LODIndex]->WeightMaps);
 		}
 		else
 		{
@@ -1090,10 +1090,10 @@ PRAGMA_ENABLE_DEPRECATION_WARNINGS
 
 			// Update general solver properties
 			PRAGMA_DISABLE_DEPRECATION_WARNINGS  // TODO: CHAOS_IS_CLOTHINGSIMULATIONMESH_ABSTRACT
-				const Softs::FSolverReal MeshScale = Mesh->GetScale();
+			const Softs::FSolverReal MeshScale = Mesh->GetScale();
 			PRAGMA_ENABLE_DEPRECATION_WARNINGS
 
-				const FRealSingle DampingCoefficient = ConfigProperties.GetValue<float>(TEXT("DampingCoefficient"), ClothingSimulationClothDefault::DampingCoefficient);
+			const FRealSingle DampingCoefficient = ConfigProperties.GetValue<float>(TEXT("DampingCoefficient"), ClothingSimulationClothDefault::DampingCoefficient);
 			const FRealSingle LocalDampingCoefficient = ConfigProperties.GetValue<float>(TEXT("LocalDampingCoefficient"));
 			const FRealSingle CollisionThickness = ConfigProperties.GetValue<float>(TEXT("CollisionThickness"), ClothingSimulationClothDefault::CollisionThickness);
 			const FRealSingle FrictionCoefficient = ConfigProperties.GetValue<float>(TEXT("FrictionCoefficient"), ClothingSimulationClothDefault::FrictionCoefficient);
diff --git a/Engine/Plugins/Experimental/ChaosCloth/Source/Chaos/Private/ChaosCloth/ChaosClothingSimulationSolver.cpp b/Engine/Plugins/Experimental/ChaosCloth/Source/Chaos/Private/ChaosCloth/ChaosClothingSimulationSolver.cpp
--- a/Engine/Plugins/Experimental/ChaosCloth/Source/Chaos/Private/ChaosCloth/ChaosClothingSimulationSolver.cpp
+++ b/Engine/Plugins/Experimental/ChaosCloth/Source/Chaos/Private/ChaosCloth/ChaosClothingSimulationSolver.cpp
@@ -1230,11 +1230,12 @@ void FClothingSimulationSolver::ParticleMassClampAndKinematicStateUpdate(int32 O
 	}
 }
 
-void FClothingSimulationSolver::SetProperties(int32 ParticleRangeId, const Softs::FCollectionPropertyConstFacade& InPropertyCollection)
+void FClothingSimulationSolver::SetProperties(int32 ParticleRangeId, const Softs::FCollectionPropertyConstFacade& InPropertyCollection,
+	const TMap<FString, TConstArrayView<FRealSingle>>& WeightMaps)
 {
 	if (Evolution)
 	{
-		Evolution->SetSoftBodyProperties(ParticleRangeId, InPropertyCollection);
+		Evolution->SetSoftBodyProperties(ParticleRangeId, InPropertyCollection, WeightMaps);
 
 		const uint32 GroupId = Evolution->GetSoftBodyGroupId(ParticleRangeId);
 		// Set properties to constraints that come from the solver (e.g., solver-level gravity, wind)
diff --git a/Engine/Plugins/Experimental/ChaosCloth/Source/Chaos/Public/ChaosCloth/ChaosClothingSimulationSolver.h b/Engine/Plugins/Experimental/ChaosCloth/Source/Chaos/Public/ChaosCloth/ChaosClothingSimulationSolver.h
--- a/Engine/Plugins/Experimental/ChaosCloth/Source/Chaos/Public/ChaosCloth/ChaosClothingSimulationSolver.h
+++ b/Engine/Plugins/Experimental/ChaosCloth/Source/Chaos/Public/ChaosCloth/ChaosClothingSimulationSolver.h
@@ -155,7 +155,8 @@ namespace Chaos
 			FRealSingle FrictionCoefficient);
 
 		/** Force based solver version */
-		CHAOSCLOTH_API void SetProperties(int32 ParticleRangeId, const Softs::FCollectionPropertyConstFacade& PropertyCollection);
+		CHAOSCLOTH_API void SetProperties(int32 ParticleRangeId, const Softs::FCollectionPropertyConstFacade& PropertyCollection,
+			const TMap<FString, TConstArrayView<FRealSingle>>& WeightMaps) ;
 
 		/** Begin PBD-solver only property methods (these properties are controlled directly via PropertyCollection in force-based solver) */
 		// Set whether to use continuous collision detection.
diff --git a/Engine/Plugins/Experimental/ChaosClothAssetEditor/Source/ChaosClothAssetDataflowNodes/Private/ChaosClothAsset/SimulationDampingConfigNode.cpp b/Engine/Plugins/Experimental/ChaosClothAssetEditor/Source/ChaosClothAssetDataflowNodes/Private/ChaosClothAsset/SimulationDampingConfigNode.cpp
--- a/Engine/Plugins/Experimental/ChaosClothAssetEditor/Source/ChaosClothAssetDataflowNodes/Private/ChaosClothAsset/SimulationDampingConfigNode.cpp
+++ b/Engine/Plugins/Experimental/ChaosClothAssetEditor/Source/ChaosClothAssetDataflowNodes/Private/ChaosClothAsset/SimulationDampingConfigNode.cpp
@@ -9,10 +9,21 @@ FChaosClothAssetSimulationDampingConfigNode::FChaosClothAssetSimulationDampingCo
 	: FChaosClothAssetSimulationBaseConfigNode(InParam, InGuid)
 {
 	RegisterCollectionConnections();
+	RegisterInputConnection(&DampingCoefficientWeighted.WeightMap);
 }
 
 void FChaosClothAssetSimulationDampingConfigNode::AddProperties(FPropertyHelper& PropertyHelper) const
 {
-	PropertyHelper.SetProperty(this, &DampingCoefficient);
+	PropertyHelper.SetPropertyWeighted(TEXT("DampingCoefficient"), DampingCoefficientWeighted);
 	PropertyHelper.SetProperty(this, &LocalDampingCoefficient);
 }
+
+void FChaosClothAssetSimulationDampingConfigNode::Serialize(FArchive& Ar)
+{
+	Super::Serialize(Ar);
+	if (DampingCoefficient_DEPRECATED != DeprecatedDampingCoefficientValue)
+	{
+		DampingCoefficientWeighted.Low = DampingCoefficientWeighted.High = DampingCoefficient_DEPRECATED;
+		DampingCoefficient_DEPRECATED = DeprecatedDampingCoefficientValue;
+	}
+}
diff --git a/Engine/Plugins/Experimental/ChaosClothAssetEditor/Source/ChaosClothAssetDataflowNodes/Public/ChaosClothAsset/SimulationDampingConfigNode.h b/Engine/Plugins/Experimental/ChaosClothAssetEditor/Source/ChaosClothAssetDataflowNodes/Public/ChaosClothAsset/SimulationDampingConfigNode.h
--- a/Engine/Plugins/Experimental/ChaosClothAssetEditor/Source/ChaosClothAssetDataflowNodes/Public/ChaosClothAsset/SimulationDampingConfigNode.h
+++ b/Engine/Plugins/Experimental/ChaosClothAssetEditor/Source/ChaosClothAssetDataflowNodes/Public/ChaosClothAsset/SimulationDampingConfigNode.h
@@ -16,9 +16,10 @@ public:
 	/**
 	 * The amount of global damping applied to the cloth velocities, also known as point damping.
 	 * Point damping improves simulation stability, but can also cause an overall slow-down effect and therefore is best left to very small percentage amounts.
+	 * Weight map only works with UseForceBasedSolver (on SolverConfig). Otherwise, Low value is always used.
 	 */
 	UPROPERTY(EditAnywhere, Category = "Damping Properties", Meta = (UIMin = "0", UIMax = "1", ClampMin = "0", ClampMax = "1"))
-	float DampingCoefficient = 0.01f;
+	FChaosClothAssetWeightedValue DampingCoefficientWeighted = {true, 0.01f, 0.01f, TEXT("DampingCoefficient")};
 
 	/**
 	 * The amount of local damping applied to the cloth velocities.
@@ -31,6 +32,15 @@ public:
 
 	FChaosClothAssetSimulationDampingConfigNode(const Dataflow::FNodeParameters& InParam, FGuid InGuid = FGuid::NewGuid());
 
+	virtual void Serialize(FArchive& Ar) override;
 private:
 	virtual void AddProperties(FPropertyHelper& PropertyHelper) const override;
+
+	// Deprecated properties
+#if WITH_EDITORONLY_DATA
+	static constexpr float DeprecatedDampingCoefficientValue = -1.f; // This is outside the settable range when it wasn't deprecated.
+	UPROPERTY()
+	float DampingCoefficient_DEPRECATED = DeprecatedDampingCoefficientValue;
+#endif 
+
 };
diff --git a/Engine/Source/Runtime/Experimental/Chaos/Private/Chaos/SoftsEvolution.cpp b/Engine/Source/Runtime/Experimental/Chaos/Private/Chaos/SoftsEvolution.cpp
--- a/Engine/Source/Runtime/Experimental/Chaos/Private/Chaos/SoftsEvolution.cpp
+++ b/Engine/Source/Runtime/Experimental/Chaos/Private/Chaos/SoftsEvolution.cpp
@@ -2,6 +2,7 @@
 #include "Chaos/SoftsEvolution.h"
 #include "Chaos/PerParticleDampVelocity.h"
 #include "Chaos/SoftsEvolutionLinearSystem.h"
+#include "Chaos/PBDFlatWeightMap.h"
 
 namespace Chaos::Softs {
 
@@ -26,6 +27,7 @@ constexpr bool bUsePerParticleDamping = false;
 
 static void EulerStepVelocity(FSolverParticlesRange& Particles, const FSolverReal Dt)
 {
+	TRACE_CPUPROFILER_EVENT_SCOPE(FSoftsEvolution_EulerStepVelocity);
 	// TODO: ISPC and/or ParallelFor
 	FSolverVec3* const Velocity = Particles.GetV().GetData();
 	const FSolverVec3* const Acceleration = Particles.GetAcceleration().GetData();
@@ -41,6 +43,7 @@ static void EulerStepVelocity(FSolverParticlesRange& Particles, const FSolverRea
 
 static void DampLocalVelocity(FSolverParticlesRange& Particles, const FPerParticleDampVelocity& DampVelocityRule)
 {
+	TRACE_CPUPROFILER_EVENT_SCOPE(FSoftsEvolution_DampLocalVelocity);
 	// TODO: ISPC and/or ParallelFor
 	FSolverVec3* const Velocity = Particles.GetV().GetData();
 	const FSolverVec3* const X = Particles.XArray().GetData();
@@ -77,6 +80,7 @@ static void CalculateGlobalDamping(const FSolverReal DampingScale, const FSolver
 
 static void EulerStepPositionWithGlobalDamping(FSolverParticlesRange& Particles, const FSolverReal DampingScale, const FSolverReal SolverFrequency, const FSolverReal Dt)
 {
+	TRACE_CPUPROFILER_EVENT_SCOPE(FSoftsEvolution_EulerStepPositionWithGlobalDamping);
 	FSolverReal DampingPowDt;
 	FSolverReal DampingIntegrated;
 	CalculateGlobalDamping(DampingScale, SolverFrequency, Dt, DampingPowDt, DampingIntegrated);
@@ -226,11 +230,29 @@ int32 FEvolution::AddSoftBody(uint32 GroupId, int32 NumParticles, bool bEnable)
 	return SoftBodyId;
 }
 
-void FEvolution::SetSoftBodyProperties(int32 SoftBodyId, const FCollectionPropertyConstFacade& PropertyCollection)
+void FEvolution::SetSoftBodyProperties(int32 SoftBodyId, const FCollectionPropertyConstFacade& PropertyCollection,
+	const TMap<FString, TConstArrayView<FRealSingle>>& WeightMaps)
 {
-	SoftBodies.GlobalDampings[SoftBodyId] = PropertyCollection.GetValue<float>(TEXT("DampingCoefficient"), Private::EvolutionSoftBodyDefault::GlobalDamping);
-	SoftBodies.LocalDampings[SoftBodyId] = PropertyCollection.GetValue<float>(TEXT("LocalDampingCoefficient"), Private::EvolutionSoftBodyDefault::LocalDamping);
-	SoftBodies.UsePerParticleDamping[SoftBodyId] = PropertyCollection.GetValue<bool>(TEXT("UsePerParticleDamping"), Private::EvolutionSoftBodyDefault::bUsePerParticleDamping);
+	SoftBodies.LocalDampings[SoftBodyId] = GetLocalDampingCoefficient(PropertyCollection, Private::EvolutionSoftBodyDefault::LocalDamping);
+
+	const FSolverVec2 GlobalDamping = GetWeightedFloatDampingCoefficient(PropertyCollection, Private::EvolutionSoftBodyDefault::GlobalDamping);
+	TConstArrayView<FRealSingle> WeightMap = WeightMaps.FindRef(GetDampingCoefficientString(PropertyCollection, DampingCoefficientName.ToString()));
+	SoftBodies.UsePerParticleDamping[SoftBodyId] = WeightMap.Num() == GetSoftBodyParticleNum(SoftBodyId);
+	if (SoftBodies.UsePerParticleDamping[SoftBodyId])
+	{
+		const FSolverParticlesRange& SoftBodyParticles = GetSoftBodyParticles(SoftBodyId);
+		TArrayView<FSolverReal> PerParticleDamping = SoftBodyParticles.GetArrayView(ParticleDampings);
+		check(PerParticleDamping.Num() == WeightMap.Num());
+		const FPBDFlatWeightMapView GlobalDampingMap(GlobalDamping, WeightMap, PerParticleDamping.Num());
+		for (int32 Index = 0; Index < PerParticleDamping.Num(); ++Index)
+		{
+			PerParticleDamping[Index] = GlobalDampingMap.GetValue(Index);
+		}
+	}
+	else
+	{
+		SoftBodies.GlobalDampings[SoftBodyId] = GlobalDamping[0];
+	}
 }
 
 void FEvolution::ActivateSoftBody(int32 SoftBodyId, bool bActivate)
diff --git a/Engine/Source/Runtime/Experimental/Chaos/Public/Chaos/SoftsEvolution.h b/Engine/Source/Runtime/Experimental/Chaos/Public/Chaos/SoftsEvolution.h
--- a/Engine/Source/Runtime/Experimental/Chaos/Public/Chaos/SoftsEvolution.h
+++ b/Engine/Source/Runtime/Experimental/Chaos/Public/Chaos/SoftsEvolution.h
@@ -65,7 +65,8 @@ public:
 	CHAOS_API int32 AddSoftBody(uint32 GroupId, int32 NumParticles, bool bEnable);
 	int32 GetSoftBodyParticleNum(int32 SoftBodyId) const { return SoftBodies.ParticleRanges[SoftBodyId].GetRangeSize(); }
 	int32 GetSoftBodyGroupId(int32 SoftBodyId) const { return SoftBodies.GroupId[SoftBodyId]; }
-	CHAOS_API void SetSoftBodyProperties(int32 SoftBodyId, const FCollectionPropertyConstFacade& PropertyCollection);
+	CHAOS_API void SetSoftBodyProperties(int32 SoftBodyId, const FCollectionPropertyConstFacade& PropertyCollection,
+		const TMap<FString, TConstArrayView<FRealSingle>>& WeightMaps);
 	CHAOS_API void ActivateSoftBody(int32 SoftBodyId, bool bActivate);
 	bool IsSoftBodyActive(int32 SoftBodyId) const { return SoftBodies.Active[SoftBodyId]; }
 	FSolverParticlesRange& GetSoftBodyParticles(int32 SoftBodyId) { return SoftBodies.ParticleRanges[SoftBodyId]; }
@@ -459,6 +460,8 @@ private:
 
 	// Collision Rules
 
+	UE_CHAOS_DECLARE_INDEXLESS_PROPERTYCOLLECTION_NAME(DampingCoefficient, float);
+	UE_CHAOS_DECLARE_INDEXLESS_PROPERTYCOLLECTION_NAME(LocalDampingCoefficient, float);
 	UE_CHAOS_DECLARE_INDEXLESS_PROPERTYCOLLECTION_NAME(MaxNumIterations, int32);
 	UE_CHAOS_DECLARE_INDEXLESS_PROPERTYCOLLECTION_NAME(NumIterations, int32);
 	UE_CHAOS_DECLARE_INDEXLESS_PROPERTYCOLLECTION_NAME(DoQuasistatics, bool);
