Iris
- Modified FastArrayReplictionFragment to allow for partial dequantization when applying state data.
- Uses more memory which we should be able to optimize away.

diff --git a/Engine/Plugins/Runtime/ReplicationSystemTestPlugin/Source/Private/Tests/ReplicationState/TestFastArrayReplicationState.cpp b/Engine/Plugins/Runtime/ReplicationSystemTestPlugin/Source/Private/Tests/ReplicationState/TestFastArrayReplicationState.cpp
--- a/Engine/Plugins/Runtime/ReplicationSystemTestPlugin/Source/Private/Tests/ReplicationState/TestFastArrayReplicationState.cpp
+++ b/Engine/Plugins/Runtime/ReplicationSystemTestPlugin/Source/Private/Tests/ReplicationState/TestFastArrayReplicationState.cpp
@@ -58,17 +58,14 @@ public:
 
 	virtual void ApplyReplicatedState(UE::Net::FReplicationStateApplyContext& Context) const override
 	{
-		// Get the wrapped FastArraySerializer and array
-		FTestFastArrayReplicationState_FastArrayWithExtraProperty* DstArraySerializer = this->GetFastArraySerializerFromOwner();
-
-		// Intentionally not const as we allow the src state to be modified
-		FTestFastArrayReplicationState_FastArrayWithExtraProperty* SrcArraySerializer = this->GetFastArraySerializerFromApplyContext(Context);
-
-		// Apply additional state
-		DstArraySerializer->ExtraInt = SrcArraySerializer->ExtraInt;
-
+		using namespace UE::Net;
+		using namespace UE::Net::Private;
+		
 		// Forward to fast array serializer
 		SuperT::ApplyReplicatedState(Context);
+
+		// Must explicitly deal with extra properties here, but there is a utility function to do it for us!
+		SuperT::ApplyReplicatedStateForExtraProperties(Context);
 	}
 
 	virtual bool PollReplicatedState(UE::Net::EReplicationFragmentPollFlags PollOption) override
diff --git a/Engine/Source/Runtime/Experimental/Iris/Core/Private/Iris/ReplicationSystem/FastArrayReplicationFragment.cpp b/Engine/Source/Runtime/Experimental/Iris/Core/Private/Iris/ReplicationSystem/FastArrayReplicationFragment.cpp
--- a/Engine/Source/Runtime/Experimental/Iris/Core/Private/Iris/ReplicationSystem/FastArrayReplicationFragment.cpp
+++ b/Engine/Source/Runtime/Experimental/Iris/Core/Private/Iris/ReplicationSystem/FastArrayReplicationFragment.cpp
@@ -11,6 +11,7 @@
 #include "Iris/Serialization/InternalNetSerializers.h"
 #include "Iris/Serialization/NetSerializers.h"
 #include "Iris/Serialization/InternalNetSerializerUtils.h"
+#include "Iris/Core/IrisProfiler.h"
 
 namespace UE::Net::Private
 {
@@ -34,6 +35,9 @@ FFastArrayReplicationFragmentBase::FFastArrayReplicationFragmentBase(EReplicatio
 		check(!EnumHasAnyFlags(InDescriptor->Traits, EReplicationStateTraits::IsNativeFastArrayReplicationState));
 	}
 
+	// We do not create temporary states
+	Traits |= EReplicationFragmentTraits::HasPersistentTargetStateBuffer;
+
 	if (EnumHasAnyFlags(InTraits, EReplicationFragmentTraits::CanReplicate))
 	{
 		ReplicationState = MakeUnique<FPropertyReplicationState>(InDescriptor);
@@ -68,6 +72,7 @@ FFastArrayReplicationFragmentBase::FFastArrayReplicationFragmentBase(EReplicatio
 	}
 
 	Traits |= EReplicationFragmentTraits::HasPropertyReplicationState;
+	Traits |= EReplicationFragmentTraits::SupportsPartialDequantizedState;
 
 	// Look up the descriptor of the struct
 	const FReplicationStateDescriptor* FastArrayStructDescriptor = GetFastArrayPropertyStructDescriptor();
@@ -118,36 +123,108 @@ void FFastArrayReplicationFragmentBase::Register(FFragmentRegistrationContext& C
 	Context.RegisterReplicationFragment(this, ReplicationStateDescriptor.GetReference(), ReplicationState ? ReplicationState->GetStateBuffer() : nullptr);
 }
 
-void FFastArrayReplicationFragmentBase::CollectOwner(FReplicationStateOwnerCollector* Owners) const
+void FFastArrayReplicationFragmentBase::InternalDequantizeFastArray(FNetSerializationContext& Context, uint8* RESTRICT DstExternalBuffer, const uint8* RESTRICT SrcInternalBuffer, const FReplicationStateDescriptor* Descriptor)
 {
-	Owners->AddOwner(Owner);
+	const FReplicationStateMemberDescriptor* MemberDescriptors = Descriptor->MemberDescriptors;
+	const FReplicationStateMemberSerializerDescriptor* MemberSerializerDescriptors = Descriptor->MemberSerializerDescriptors;
+	const uint32 MemberCount = Descriptor->MemberCount;
+
+	const uint32 ArrayMemberIndex = FFastArrayReplicationFragmentHelper::GetFastArrayStructItemArrayMemberIndex(Descriptor);
+	check(ArrayMemberIndex < MemberCount);
+
+	const FReplicationStateMemberDescriptor& MemberDescriptor = MemberDescriptors[ArrayMemberIndex];
+	const FReplicationStateMemberSerializerDescriptor& MemberSerializerDescriptor = MemberSerializerDescriptors[ArrayMemberIndex];
+	
+	// Dequantize full array
+	FNetDequantizeArgs Args;
+	Args.Version = 0;
+	Args.NetSerializerConfig = MemberSerializerDescriptor.SerializerConfig;
+	Args.Target = reinterpret_cast<NetSerializerValuePointer>(DstExternalBuffer + MemberDescriptor.ExternalMemberOffset);
+	Args.Source = reinterpret_cast<NetSerializerValuePointer>(SrcInternalBuffer + MemberDescriptor.InternalMemberOffset);	
+
+	MemberSerializerDescriptor.Serializer->Dequantize(Context, Args);
 }
 
-void FFastArrayReplicationFragmentBase::CallRepNotifies(FReplicationStateApplyContext& Context)
+void FFastArrayReplicationFragmentBase::InternalPartialDequantizeFastArray(FReplicationStateApplyContext& ApplyContext, uint8* RESTRICT DstExternalBuffer, const uint8* RESTRICT SrcInternalBuffer, const FReplicationStateDescriptor* Descriptor)
 {
-	IRIS_PROFILER_SCOPE(PropertyReplicationFragment_InvokeRepNotifies);
+	const FReplicationStateMemberDescriptor* MemberDescriptors = Descriptor->MemberDescriptors;
+	const FReplicationStateMemberSerializerDescriptor* MemberSerializerDescriptors = Descriptor->MemberSerializerDescriptors;
+	const uint32 MemberCount = Descriptor->MemberCount;
 
-	const FPropertyReplicationState ReceivedState(Context.Descriptor, Context.StateBufferData.ExternalStateBuffer);
+	const uint32 ArrayMemberIndex = FFastArrayReplicationFragmentHelper::GetFastArrayStructItemArrayMemberIndex(Descriptor);
+	check(ArrayMemberIndex < MemberCount);
 
-	if (!Context.bIsInit)
-	{
-		ReceivedState.CallRepNotifies(Owner, nullptr, Context.bIsInit);
-	}
-	else
+	const FReplicationStateMemberDescriptor& MemberDescriptor = MemberDescriptors[ArrayMemberIndex];
+	const FReplicationStateMemberSerializerDescriptor& MemberSerializerDescriptor = MemberSerializerDescriptors[ArrayMemberIndex];
+	const FReplicationStateMemberChangeMaskDescriptor& MemberChangeMaskDescriptor = ApplyContext.Descriptor->MemberChangeMaskDescriptors[0];
+
+	// Setup changemask
+	FNetBitArrayView MemberChangeMask = MakeNetBitArrayView(ApplyContext.StateBufferData.ChangeMaskData, ApplyContext.Descriptor->ChangeMaskBitCount);
+	ApplyContext.NetSerializationContext->SetChangeMask(&MemberChangeMask);
+
+	FNetDequantizeArgs Args;
+	Args.Version = 0;
+	Args.NetSerializerConfig = MemberSerializerDescriptor.SerializerConfig;
+	Args.Target = reinterpret_cast<NetSerializerValuePointer>(DstExternalBuffer + MemberDescriptor.ExternalMemberOffset);
+	Args.Source = reinterpret_cast<NetSerializerValuePointer>(SrcInternalBuffer + MemberDescriptor.InternalMemberOffset);
+
+	Args.ChangeMaskInfo.BitCount = MemberChangeMaskDescriptor.BitCount;
+	Args.ChangeMaskInfo.BitOffset = MemberChangeMaskDescriptor.BitOffset;
+
+	MemberSerializerDescriptor.Serializer->Dequantize(*ApplyContext.NetSerializationContext, Args);
+}
+
+void FFastArrayReplicationFragmentBase::InternalDequantizeExtraProperties(FNetSerializationContext& Context, uint8* RESTRICT DstExternalBuffer, const uint8* RESTRICT SrcInternalBuffer, const FReplicationStateDescriptor* Descriptor)
+{
+	const uint32 ArrayMemberIndex = FFastArrayReplicationFragmentHelper::GetFastArrayStructItemArrayMemberIndex(Descriptor);
+
+	const FReplicationStateMemberDescriptor* MemberDescriptors = Descriptor->MemberDescriptors;
+	const FReplicationStateMemberSerializerDescriptor* MemberSerializerDescriptors = Descriptor->MemberSerializerDescriptors;
+	const uint32 MemberCount = Descriptor->MemberCount;
+
+	check(ArrayMemberIndex < MemberCount);
+
+	for (uint32 MemberIndex = 0; MemberIndex < MemberCount; ++MemberIndex)
 	{
-		// As our default initial states is always treated as all dirty we need to compare against the default before applying initial repnotifies
-		const FPropertyReplicationState DefaultState(Context.Descriptor);
+		if (MemberIndex == ArrayMemberIndex)
+		{
+			continue;
+		}
+
+		const FReplicationStateMemberDescriptor& MemberDescriptor = MemberDescriptors[MemberIndex];
+		const FReplicationStateMemberSerializerDescriptor& MemberSerializerDescriptor = MemberSerializerDescriptors[MemberIndex];
 
-		ReceivedState.CallRepNotifies(Owner, &DefaultState, Context.bIsInit);
+		FNetDequantizeArgs Args;
+		Args.Version = 0;
+		Args.NetSerializerConfig = MemberSerializerDescriptor.SerializerConfig;
+		Args.Target = reinterpret_cast<NetSerializerValuePointer>(DstExternalBuffer + MemberDescriptor.ExternalMemberOffset);
+		Args.Source = reinterpret_cast<NetSerializerValuePointer>(SrcInternalBuffer + MemberDescriptor.InternalMemberOffset);	
+
+		MemberSerializerDescriptor.Serializer->Dequantize(Context, Args);
 	}
 }
 
-void FFastArrayReplicationFragmentBase::ReplicatedStateToString(FStringBuilderBase& StringBuilder, FReplicationStateApplyContext& Context, EReplicationStateToStringFlags Flags) const
+void FFastArrayReplicationFragmentBase::CollectOwner(FReplicationStateOwnerCollector* Owners) const
+{
+	Owners->AddOwner(Owner);
+}
+
+void FFastArrayReplicationFragmentBase::ToString(FStringBuilderBase& StringBuilder, const uint8* StateBuffer, const FReplicationStateDescriptor* Descriptor)
 {
-	const FPropertyReplicationState ReceivedState(Context.Descriptor, Context.StateBufferData.ExternalStateBuffer);
+	const FProperty** MemberProperties = Descriptor->MemberProperties;
+	const uint32 MemberCount = Descriptor->MemberCount;
+
+	StringBuilder.Appendf(TEXT("FFastArrayReplicationState %s\n"), Descriptor->DebugName->Name);
 
-	const bool bIncludeAll = EnumHasAnyFlags(Flags, EReplicationStateToStringFlags::OnlyIncludeDirtyMembers) == false;
-	ReceivedState.ToString(StringBuilder, bIncludeAll);
+	FString TempString;
+	for (uint32 MemberIt = 0; MemberIt < MemberCount; ++MemberIt)
+	{
+		const FProperty* Property = MemberProperties[MemberIt];
+
+		Property->ExportTextItem_Direct(TempString, StateBuffer + Descriptor->MemberDescriptors[MemberIt].ExternalMemberOffset, nullptr, nullptr, PPF_SimpleObjectText);
+		StringBuilder.Appendf(TEXT("%u - %s : %s\n"), MemberIt, *Property->GetName(), ToCStr(TempString));
+		TempString.Reset();
+	}
 }
 
 uint32 FFastArrayReplicationFragmentHelper::GetFastArrayStructItemArrayMemberIndex(const FReplicationStateDescriptor* StructDescriptor)
diff --git a/Engine/Source/Runtime/Experimental/Iris/Core/Private/Iris/ReplicationSystem/ReplicationOperations.cpp b/Engine/Source/Runtime/Experimental/Iris/Core/Private/Iris/ReplicationSystem/ReplicationOperations.cpp
--- a/Engine/Source/Runtime/Experimental/Iris/Core/Private/Iris/ReplicationSystem/ReplicationOperations.cpp
+++ b/Engine/Source/Runtime/Experimental/Iris/Core/Private/Iris/ReplicationSystem/ReplicationOperations.cpp
@@ -164,8 +164,12 @@ void FReplicationStateOperations::DequantizeWithMask(FNetSerializationContext& C
 			Args.Target = reinterpret_cast<NetSerializerValuePointer>(DstExternalBuffer + MemberDescriptor.ExternalMemberOffset);
 			Args.Source = reinterpret_cast<NetSerializerValuePointer>(SrcInternalBuffer + MemberDescriptor.InternalMemberOffset);
 
-			Args.ChangeMaskInfo.BitCount = MemberChangeMaskDescriptor.BitCount;
-			Args.ChangeMaskInfo.BitOffset = MemberChangeMaskOffset;
+			// NOTE: Currently we only forward changemask info for fastarrays this is to avoid propagating it by accident where it does not belong
+			if (EnumHasAnyFlags(Descriptor->Traits, EReplicationStateTraits::IsFastArrayReplicationState))
+			{
+				Args.ChangeMaskInfo.BitCount = MemberChangeMaskDescriptor.BitCount;
+				Args.ChangeMaskInfo.BitOffset = MemberChangeMaskOffset;
+			}
 
 			MemberSerializerDescriptor.Serializer->Dequantize(Context, Args);
 		}
diff --git a/Engine/Source/Runtime/Experimental/Iris/Core/Private/Iris/Serialization/ArrayPropertyNetSerializer.cpp b/Engine/Source/Runtime/Experimental/Iris/Core/Private/Iris/Serialization/ArrayPropertyNetSerializer.cpp
--- a/Engine/Source/Runtime/Experimental/Iris/Core/Private/Iris/Serialization/ArrayPropertyNetSerializer.cpp
+++ b/Engine/Source/Runtime/Experimental/Iris/Core/Private/Iris/Serialization/ArrayPropertyNetSerializer.cpp
@@ -540,13 +540,38 @@ void FArrayPropertyNetSerializer::Dequantize(FNetSerializationContext& Context,
 	ElementArgs.Source = NetSerializerValuePointer(SourceArray.ElementStorage);
 	ElementArgs.Target = 0;
 
+	const FNetBitArrayView* ChangeMask = Args.ChangeMaskInfo.BitCount > 1U ? Context.GetChangeMask() : nullptr;
+
 	const FNetSerializer* ElementSerializer = ElementSerializerDescriptor.Serializer;
 	const uint32 ElementSize = ElementStateDescriptor->InternalSize;
-	for (uint32 ElementIt = 0, ElementEndIt = ElementCount; ElementIt < ElementEndIt; ++ElementIt)
+
+	// Currently we do not support partial dequantize using changemask for array properties due to complexities elsewhere (FastArrayReplicationFragment and PropertyReplicationState::ApplyState)
+	if (!ChangeMask)
 	{
-		ElementArgs.Target = NetSerializerValuePointer(static_cast<const void*>(ScriptArrayHelper.GetRawPtr(ElementIt)));
-		ElementSerializer->Dequantize(Context, ElementArgs);
-		ElementArgs.Source += ElementSize;
+		for (uint32 ElementIt = 0, ElementEndIt = ElementCount; ElementIt < ElementEndIt; ++ElementIt)
+		{
+			ElementArgs.Target = NetSerializerValuePointer(static_cast<const void*>(ScriptArrayHelper.GetRawPtr(ElementIt)));
+			ElementSerializer->Dequantize(Context, ElementArgs);
+			ElementArgs.Source += ElementSize;
+		}
+	}
+	else
+	{
+		// We currently use a simple modulo scheme for bits in the changemask, if we have more elements in the array then is covered by the changemask
+		// several entries in the array will be considered dirty and be serialized
+		// As the first bit is used by the owning property we need to offset by one and deduct one from the usable bits
+		const uint32 ChangeMaskBitOffset = ChangeMask ? Args.ChangeMaskInfo.BitOffset + 1U : 0U;
+		const uint32 ChangeMaskBitCount = ChangeMask ? Args.ChangeMaskInfo.BitCount - 1U : 0U;
+
+		for (uint32 ElementIt = 0, ElementEndIt = ElementCount; ElementIt < ElementEndIt; ++ElementIt)
+		{
+			if (ChangeMask->GetBit(ChangeMaskBitOffset + (ElementIt % ChangeMaskBitCount)))
+			{
+				ElementArgs.Target = NetSerializerValuePointer(static_cast<const void*>(ScriptArrayHelper.GetRawPtr(ElementIt)));
+				ElementSerializer->Dequantize(Context, ElementArgs);
+			}
+			ElementArgs.Source += ElementSize;
+		}
 	}
 }
 
diff --git a/Engine/Source/Runtime/Experimental/Iris/Core/Public/Iris/ReplicationSystem/FastArrayReplicationFragment.h b/Engine/Source/Runtime/Experimental/Iris/Core/Public/Iris/ReplicationSystem/FastArrayReplicationFragment.h
--- a/Engine/Source/Runtime/Experimental/Iris/Core/Public/Iris/ReplicationSystem/FastArrayReplicationFragment.h
+++ b/Engine/Source/Runtime/Experimental/Iris/Core/Public/Iris/ReplicationSystem/FastArrayReplicationFragment.h
@@ -28,15 +28,23 @@ class TFastArrayReplicationFragment : public Private::FFastArrayReplicationFragm
 {
 public:
 	typedef TArray<FastArrayItemType> ItemArrayType;
-	TFastArrayReplicationFragment(EReplicationFragmentTraits InTraits, UObject* InOwner, const FReplicationStateDescriptor* InDescriptor) : FFastArrayReplicationFragmentBase(InTraits, InOwner, InDescriptor) {}
+	TFastArrayReplicationFragment(EReplicationFragmentTraits InTraits, UObject* InOwner, const FReplicationStateDescriptor* InDescriptor, bool bValidateDescriptor = true);
 
 protected:
 	enum EAllowAdditionalPropertiesType { AllowAdditionalProperties };
-	TFastArrayReplicationFragment(EReplicationFragmentTraits InTraits, UObject* InOwner, const FReplicationStateDescriptor* InDescriptor, const EAllowAdditionalPropertiesType) : FFastArrayReplicationFragmentBase(InTraits, InOwner, InDescriptor, false) {}
+	TFastArrayReplicationFragment(EReplicationFragmentTraits InTraits, UObject* InOwner, const FReplicationStateDescriptor* InDescriptor, const EAllowAdditionalPropertiesType) : TFastArrayReplicationFragment(InTraits, InOwner, InDescriptor, false) {}
 
+	// For the select few cases where we allow additional properties, this is a helper to deal with applying them directly from quantized state
+	void ApplyReplicatedStateForExtraProperties(FReplicationStateApplyContext& Context) const;
+
+	// FReplicationFragment
 	virtual void ApplyReplicatedState(FReplicationStateApplyContext& Context) const override;
+	virtual void CallRepNotifies(FReplicationStateApplyContext& Context) override;
 	virtual bool PollReplicatedState(EReplicationFragmentPollFlags PollOption) override;
+	virtual void ReplicatedStateToString(FStringBuilderBase& StringBuilder, FReplicationStateApplyContext& Context, EReplicationStateToStringFlags Flags) const override;
 
+protected:
+	// Poll entire FastFrray
 	bool PollAllState(bool bForceFullCompare = false);
 
 	// Returns true if the FastArray is dirty
@@ -52,7 +60,10 @@ protected:
 	inline FastArrayType* GetFastArraySerializerFromReplicationState() const;
 
 	// Get FastArraySerialzier from received state
-	inline FastArrayType* GetFastArraySerializerFromApplyContext(FReplicationStateApplyContext& Context) const;	
+	inline FastArrayType* GetFastArraySerializerFromApplyContext(FReplicationStateApplyContext& Context) const;
+
+	// TODO: Can be removed when we have implemented explicit code to traverse quantized data directly
+	TUniquePtr<FastArrayType> AccumulatedReceivedState;
 };
 
 /**
@@ -76,9 +87,10 @@ protected:
 	// FReplicationFragment implementation
 	virtual void ApplyReplicatedState(FReplicationStateApplyContext& Context) const override;
 	virtual void CallRepNotifies(FReplicationStateApplyContext& Context) override;
-	virtual void ReplicatedStateToString(FStringBuilderBase& StringBuilder, FReplicationStateApplyContext& Context, EReplicationStateToStringFlags Flags) const override;
 	virtual bool PollReplicatedState(EReplicationFragmentPollFlags PollOption) override;
+	virtual void ReplicatedStateToString(FStringBuilderBase& StringBuilder, FReplicationStateApplyContext& Context, EReplicationStateToStringFlags Flags) const override;
 
+protected:
 	bool PollAllState();
 
 	bool IsDirty() const;
@@ -93,6 +105,16 @@ private:
 /**
  * TFastArrayReplicationFragment implementation
  */
+
+template <typename FastArrayItemType, typename FastArrayType>
+TFastArrayReplicationFragment<FastArrayItemType, FastArrayType>::TFastArrayReplicationFragment(EReplicationFragmentTraits InTraits, UObject* InOwner, const FReplicationStateDescriptor* InDescriptor, bool bValidateDescriptor) : FFastArrayReplicationFragmentBase(InTraits, InOwner, InDescriptor, bValidateDescriptor)
+{
+	if (EnumHasAnyFlags(InTraits, EReplicationFragmentTraits::CanReceive))
+	{
+		AccumulatedReceivedState = MakeUnique<FastArrayType>();
+	}
+}
+
 template <typename FastArrayItemType, typename FastArrayType>
 void TFastArrayReplicationFragment<FastArrayItemType, FastArrayType>::ApplyReplicatedState(FReplicationStateApplyContext& Context) const
 {
@@ -100,15 +122,25 @@ void TFastArrayReplicationFragment<FastArrayItemType, FastArrayType>::ApplyRepli
 	FastArrayType* DstArraySerializer = GetFastArraySerializerFromOwner();
 	ItemArrayType* DstWrappedArray = reinterpret_cast<ItemArrayType*>(reinterpret_cast<uint8*>(DstArraySerializer) + WrappedArrayOffsetRelativeFastArraySerializerProperty);
 
-	// Get received array data
-	// Intentionally not const as we allow the src state to be modified	
-	FastArrayType* SrcArraySerializer = GetFastArraySerializerFromApplyContext(Context);
-	const ItemArrayType* SrcWrappedArray = reinterpret_cast<ItemArrayType*>(reinterpret_cast<uint8*>(SrcArraySerializer) + WrappedArrayOffsetRelativeFastArraySerializerProperty);
+	// For now we maintain a dequantized representation of the array into which we accumulate new data.
+	// TODO: change to only maintain the map of indices
+	InternalPartialDequantizeFastArray(Context, reinterpret_cast<uint8*>(AccumulatedReceivedState.Get()), Context.StateBufferData.RawStateBuffer, GetFastArrayPropertyStructDescriptor());
 
-	// Apply state and issue callbacks etc
+	// Intentionally not const as we allow the src state to be modified
+	FastArrayType* SrcArraySerializer = AccumulatedReceivedState.Get();
+	const ItemArrayType* SrcWrappedArray = reinterpret_cast<const ItemArrayType*>(reinterpret_cast<uint8*>(SrcArraySerializer) + WrappedArrayOffsetRelativeFastArraySerializerProperty);	
+
+	// Apply state to target FastArray and issue callbacks
 	Private::FFastArrayReplicationFragmentHelper::ApplyReplicatedState(DstArraySerializer, DstWrappedArray, SrcArraySerializer, SrcWrappedArray, GetArrayElementDescriptor(), Context);
 }
 
+template <typename FastArrayItemType, typename FastArrayType>
+void TFastArrayReplicationFragment<FastArrayItemType, FastArrayType>::ApplyReplicatedStateForExtraProperties(FReplicationStateApplyContext& Context) const
+{
+	// Dequantize additional properties directly to DstArraySerialzier
+	InternalDequantizeExtraProperties(*Context.NetSerializationContext, reinterpret_cast<uint8*>(GetFastArraySerializerFromOwner()), Context.StateBufferData.RawStateBuffer, GetFastArrayPropertyStructDescriptor());
+}
+
 template <typename FastArrayItemType, typename FastArrayType>
 FastArrayType* TFastArrayReplicationFragment<FastArrayItemType, FastArrayType>::GetFastArraySerializerFromOwner() const
 {
@@ -238,6 +270,19 @@ bool TFastArrayReplicationFragment<FastArrayItemType, FastArrayType>::PollAllSta
 	return IsDirty();
 }
 
+template <typename FastArrayItemType, typename FastArrayType>
+void TFastArrayReplicationFragment<FastArrayItemType, FastArrayType>::ReplicatedStateToString(FStringBuilderBase& StringBuilder, FReplicationStateApplyContext& Context, EReplicationStateToStringFlags Flags) const
+{
+	// Temporary
+	FastArrayType ReceivedState;
+
+	// Dequantize into temporary array, using partial dequantize based on changemask
+	InternalDequantizeFastArray(*Context.NetSerializationContext, reinterpret_cast<uint8*>(&ReceivedState), Context.StateBufferData.RawStateBuffer, GetFastArrayPropertyStructDescriptor());
+
+	// Output state to string
+	ToString(StringBuilder, reinterpret_cast<uint8*>(&ReceivedState), GetFastArrayPropertyStructDescriptor());
+}
+
 template <typename FastArrayItemType, typename FastArrayType>
 bool TFastArrayReplicationFragment<FastArrayItemType, FastArrayType>::IsDirty() const
 {
@@ -257,6 +302,31 @@ void TFastArrayReplicationFragment<FastArrayItemType, FastArrayType>::MarkDirty(
 	}
 }
 
+template <typename FastArrayItemType, typename FastArrayType>
+void TFastArrayReplicationFragment<FastArrayItemType, FastArrayType>::CallRepNotifies(FReplicationStateApplyContext& Context)
+{
+	const FReplicationStateDescriptor* Descriptor = ReplicationStateDescriptor;
+
+	// If we get here we are either the init state or dirty
+	if (const UFunction* RepNotifyFunction = Descriptor->MemberPropertyDescriptors[0].RepNotifyFunction)
+	{
+		// if this is the init state, we compare against default and early out if initial state does not differ from default (empty)
+		if (Context.bIsInit)
+		{
+			FastArrayType ReceivedState;
+			FastArrayType DefaultState;
+			InternalDequantizeFastArray(*Context.NetSerializationContext, reinterpret_cast<uint8*>(&ReceivedState), Context.StateBufferData.RawStateBuffer, GetFastArrayPropertyStructDescriptor());
+			InternalDequantizeExtraProperties(*Context.NetSerializationContext, reinterpret_cast<uint8*>(&ReceivedState), Context.StateBufferData.RawStateBuffer, GetFastArrayPropertyStructDescriptor());			
+			if (Descriptor->MemberProperties[0]->Identical(&ReceivedState, &DefaultState))
+			{
+				return;
+			}
+		}
+
+		Owner->ProcessEvent(const_cast<UFunction*>(RepNotifyFunction), nullptr);
+	}
+}
+
 /*
  * TNativeFastArrayReplicationFragment implementation
  */
diff --git a/Engine/Source/Runtime/Experimental/Iris/Core/Public/Iris/ReplicationSystem/Private/FastArrayReplicationFragmentInternal.h b/Engine/Source/Runtime/Experimental/Iris/Core/Public/Iris/ReplicationSystem/Private/FastArrayReplicationFragmentInternal.h
--- a/Engine/Source/Runtime/Experimental/Iris/Core/Public/Iris/ReplicationSystem/Private/FastArrayReplicationFragmentInternal.h
+++ b/Engine/Source/Runtime/Experimental/Iris/Core/Public/Iris/ReplicationSystem/Private/FastArrayReplicationFragmentInternal.h
@@ -99,14 +99,35 @@ public:
 
 protected:
 	IRISCORE_API FFastArrayReplicationFragmentBase(EReplicationFragmentTraits InTraits, UObject* InOwner, const FReplicationStateDescriptor* InDescriptor, bool bValidateDescriptor = true);
+
+	// FReplicationFragment Implementation
+	IRISCORE_API virtual void CollectOwner(FReplicationStateOwnerCollector* Owners) const override;
+
+protected:
+	// Get the ReplicationStateDescriptor for the FastArraySerializer Struct
 	IRISCORE_API const FReplicationStateDescriptor* GetFastArrayPropertyStructDescriptor() const;
+
+	// Get the ReplicationStateDescriptor for the Array Element
 	IRISCORE_API const FReplicationStateDescriptor* GetArrayElementDescriptor() const;
-	IRISCORE_API virtual void CollectOwner(FReplicationStateOwnerCollector* Owners) const override;
-	IRISCORE_API virtual void CallRepNotifies(FReplicationStateApplyContext& Context) override;	
-	IRISCORE_API virtual void ReplicatedStateToString(FStringBuilderBase& StringBuilder, FReplicationStateApplyContext& Context, EReplicationStateToStringFlags Flags) const override;
+
+	// Copy array element using the descriptor to esure that we only copy replicated data
 	IRISCORE_API static void InternalCopyArrayElement(const FReplicationStateDescriptor* ArrayElementDescriptor, void* RESTRICT Dst, const void* RESTRICT Src);
+
+	// Compare an array element using the descriptor to ensure that we only compare replicated data
 	IRISCORE_API static bool InternalCompareArrayElement(const FReplicationStateDescriptor* ArrayElementDescriptor, void* RESTRICT Dst, const void* RESTRICT Src);
 
+	// Dequantize state into DstExternalBuffer, Note: it is expected to be initialized
+	IRISCORE_API static void InternalDequantizeFastArray(FNetSerializationContext& Context, uint8* RESTRICT DstExternalBuffer, const uint8* RESTRICT SrcInternalBuffer, const FReplicationStateDescriptor* FastArrayPropertyDescriptor);
+
+	// Partial dequantize state based on changemask into DstExternalBuffer, Note: it is expected to be initialized
+	IRISCORE_API static void InternalPartialDequantizeFastArray(FReplicationStateApplyContext& Context, uint8* RESTRICT DstExternalBuffer, const uint8* RESTRICT SrcInternalBuffer, const FReplicationStateDescriptor* FastArrayPropertyDescriptor);
+
+	// Dequantize additional properties to  DstExternalBuffer, Note: it is expected to be initialized
+	IRISCORE_API static void InternalDequantizeExtraProperties(FNetSerializationContext& Context, uint8* RESTRICT DstExternalBuffer, const uint8* RESTRICT SrcInternalBuffer, const FReplicationStateDescriptor* Descriptor);
+
+	// Dequantize and output state to string
+	IRISCORE_API static void ToString(FStringBuilderBase& StringBuilder, const uint8* ExternalStateBuffer, const FReplicationStateDescriptor* FastArrayPropertyDescriptor);
+
 protected:
 	// Replication descriptor built for the specific property
 	TRefCountPtr<const FReplicationStateDescriptor> ReplicationStateDescriptor;
@@ -183,8 +204,7 @@ void FFastArrayReplicationFragmentHelper::ApplyReplicatedState(FastArrayType* Ds
 	ConditionalRebuildItemMap(*DstArraySerializer, *DstWrappedArray);
 	ConditionalRebuildItemMap(*SrcArraySerializer, *SrcWrappedArray);
 
-	const bool bIsNativeFastArraySerializer = EnumHasAnyFlags(Context.Descriptor->Traits, EReplicationStateTraits::IsNativeFastArrayReplicationState);
-	const uint32* ChangeMaskData = bIsNativeFastArraySerializer ?  Context.StateBufferData.ChangeMaskData : (const uint32*)(Context.StateBufferData.ExternalStateBuffer + Context.Descriptor->ChangeMasksExternalOffset);
+	const uint32* ChangeMaskData = Context.StateBufferData.ChangeMaskData;
 	FNetBitArrayView MemberChangeMask = MakeNetBitArrayView(ChangeMaskData, Context.Descriptor->ChangeMaskBitCount);
 
 	// We currently use a simple modulo scheme for bits in the changemask
@@ -218,13 +238,17 @@ void FFastArrayReplicationFragmentHelper::ApplyReplicatedState(FastArrayType* Ds
 		ModifiedIndices.Reserve(SrcWrappedArray->Num());
 		const ItemType* SrcItems = SrcWrappedArray->GetData();
 		for (int32 It=0, EndIt=SrcWrappedArray->Num(); It != EndIt; ++It)
-		{			
-			if (int32* ExistingIndex = DstArraySerializer->ItemMap.Find(SrcItems[It].ReplicationID))
+		{
+			const bool bIsDirty = ChangeMaskBitCount == 0U || MemberChangeMask.GetBit((It % ChangeMaskBitCount) + ChangeMaskBitOffset);
+			if (!bIsDirty)
 			{
-				const bool bIsDirty = ChangeMaskBitCount == 0U || MemberChangeMask.GetBit((It % ChangeMaskBitCount) + ChangeMaskBitOffset);
+				continue;
+			}
 
+			if (int32* ExistingIndex = DstArraySerializer->ItemMap.Find(SrcItems[It].ReplicationID))
+			{
 				// Only compare if the changemask indicate that this might be a dirty entry, the compare is required since we do share entries in the changemask.
-				if (bIsDirty && !InternalCompareArrayElement(ArrayElementDescriptor, &(*DstWrappedArray)[*ExistingIndex], &SrcItems[It]))
+				if (!InternalCompareArrayElement(ArrayElementDescriptor, &(*DstWrappedArray)[*ExistingIndex], &SrcItems[It]))
 				{
 					UE_LOG(LogNetFastTArray, Log, TEXT("   Changed. ID: %d -> Idx: %d"), SrcItems[It].ReplicationID, *ExistingIndex);
 
@@ -236,15 +260,19 @@ void FFastArrayReplicationFragmentHelper::ApplyReplicatedState(FastArrayType* Ds
 			}
 			else
 			{
-				int32 AddedIndex = DstWrappedArray->Add(SrcItems[It]);
+				// Since we zero initialize our replicated properties we can end up with ReplicationID == 0 when receiving partial changes which should be ignored.
+				if (SrcItems[It].ReplicationID != 0)
+				{
+					int32 AddedIndex = DstWrappedArray->Add(SrcItems[It]);
 
-				UE_LOG(LogNetFastTArray, Log, TEXT("   New. ID: %d. New Element! local Idx: %d"), SrcItems[It].ReplicationID, AddedIndex);
+					UE_LOG(LogNetFastTArray, Log, TEXT("   New. ID: %d. New Element! local Idx: %d"), SrcItems[It].ReplicationID, AddedIndex);
 
-				// We need to propagate the ReplicationID in order to find our object
-				(*DstWrappedArray)[AddedIndex].ReplicationID = SrcItems[It].ReplicationID;
+					// We need to propagate the ReplicationID in order to find our object
+					(*DstWrappedArray)[AddedIndex].ReplicationID = SrcItems[It].ReplicationID;
 
-				// should we store ids or indices?
-				AddedIndices.Add(AddedIndex);
+					// should we store ids or indices?
+					AddedIndices.Add(AddedIndex);
+				}
 			}
 		}
 	}
