Chaos
- add Motion Aware Collision Detection (MACD) path for convex-mesh
- non-MACD convex-trimesh uses newer mesh collision (previously used only with heightfield)

diff --git a/Engine/Source/Runtime/Experimental/Chaos/Private/Chaos/Collision/ConvexTriangleContactPoint.cpp b/Engine/Source/Runtime/Experimental/Chaos/Private/Chaos/Collision/ConvexTriangleContactPoint.cpp
--- a/Engine/Source/Runtime/Experimental/Chaos/Private/Chaos/Collision/ConvexTriangleContactPoint.cpp
+++ b/Engine/Source/Runtime/Experimental/Chaos/Private/Chaos/Collision/ConvexTriangleContactPoint.cpp
@@ -1,7 +1,9 @@
 // Copyright Epic Games, Inc. All Rights Reserved.
 #include "Chaos/Collision/ConvexTriangleContactPoint.h"
 #include "Chaos/Collision/ContactPoint.h"
+#include "Chaos/Collision/ContactTriangles.h"
 #include "Chaos/Collision/ConvexContactPointUtilities.h"
+#include "Chaos/Collision/ConvexFeature.h"
 #include "Chaos/CollisionOneShotManifolds.h"
 #include "Chaos/Convex.h"
 #include "Chaos/DebugDrawQueue.h"
@@ -18,6 +20,9 @@
 
 namespace Chaos
 {
+	extern FRealSingle Chaos_Collision_GJKEpsilon;
+	extern FRealSingle Chaos_Collision_EPAEpsilon;
+
 	// Check whether the convex-triangle edge pair form part of the Minkowski Sum. Only edge pairs
 	// that contribute to the Minkowki Sum surface need to be checked for separation. The inputs
 	// are the convex normals for the two faces that share the convex edge, and the normal and
@@ -158,7 +163,7 @@ namespace Chaos
 		// Project the convex onto the triangle normal, with distances relative to the triangle plane
 		FReal TriPlaneDMin, TriPlaneDMax;
 		int32 ConvexVertexIndexMin, ConvexVertexIndexMax;
-		ProjectOntoAxis(Convex, TriN, TriC, TriPlaneDMin, TriPlaneDMax, ConvexVertexIndexMin, ConvexVertexIndexMax, &ConvexVertexDsView);
+		Private::ProjectOntoAxis(Convex, TriN, TriC, TriPlaneDMin, TriPlaneDMax, ConvexVertexIndexMin, ConvexVertexIndexMax, &ConvexVertexDsView);
 
 		// Distance culling
 		if (TriPlaneDMin > CullDistance)
@@ -189,7 +194,7 @@ namespace Chaos
 
 			FReal DMin, DMax;
 			int32 IndexMin, IndexMax;
-			ProjectOntoAxis(Triangle, ConN, ConX, DMin, DMax, IndexMin, IndexMax);
+			Private::ProjectOntoAxis(Triangle, ConN, ConX, DMin, DMax, IndexMin, IndexMax);
 
 			// Distance culling
 			// @todo(chaos): Cull against the projected convex hull, not just the outward face 
@@ -213,7 +218,6 @@ namespace Chaos
 		// SAT: Convex edges vs triangle edges
 		//
 
-#if (CHAOS_CONVEX_TRIANGLE_EDGEEDGE_METHOD == 0)
 		// Calculate the distance of triangle each edge to the convex separating plane
 		FReal TriVertexConvexDMin0 = FVec3::DotProduct(Triangle.GetVertex(0) - ConvexPlaneX, ConvexPlaneN);
 		FReal TriVertexConvexDMin1 = FVec3::DotProduct(Triangle.GetVertex(1) - ConvexPlaneX, ConvexPlaneN);
@@ -266,8 +270,8 @@ namespace Chaos
 				}
 
 				// Triangle edge vertices
-				const FVec3& TriEdgeV0 = (TriEdgeIndex == 0) ? Triangle.GetVertex(2) : Triangle.GetVertex(TriEdgeIndex - 1);
-				const FVec3& TriEdgeV1 = Triangle.GetVertex(TriEdgeIndex);
+				const FVec3& TriEdgeV0 = Triangle.GetVertex(TriEdgeIndex);
+				const FVec3& TriEdgeV1 = (TriEdgeIndex == 2) ? Triangle.GetVertex(0) : Triangle.GetVertex(TriEdgeIndex + 1);
 
 				// Skip edge pairs that do not contribute to the Minkowski Sum surface
 				// NOTE: This relies on the ordering of the edge planes from above. 
@@ -315,103 +319,6 @@ namespace Chaos
 				}
 			}
 		}
-#elif (CHAOS_CONVEX_TRIANGLE_EDGEEDGE_METHOD == 1)
-		// Loop over all the convex planes and then its edges, and check the separating axis defined by the convex edge and each tri edge
-		// NOTE: This performs duplicate work as each edge is shared with 2 convex faces, but we also reject half the planes (for most
-		// convexes, although it depends) so it ends up faster than using IsMinkowskiSumConvexTriangle on all the edge pairs
-		FVec3 EdgeEdgeN = FVec3(0);
-		FReal EdgeEdgeDMin = InvalidPhi;
-		int32 ConvexEdgeIndexMin = INDEX_NONE;
-		int32 TriEdgeIndexMin = INDEX_NONE;
-		for (int32 PlaneIndex = 0; PlaneIndex < Convex.NumPlanes(); ++PlaneIndex)
-		{
-			FVec3 ConN, ConX;
-			Convex.GetPlaneNX(PlaneIndex, ConN, ConX);
-
-			// Ignore faces with normals in same direction as triangle
-			const FReal NormalDot = FVec3::DotProduct(ConN, TriN);
-			if (NormalDot > 0)
-			{
-				continue;
-			}
-
-			// Check all the edges on this face
-			const int32 NumPlaneVerts = Convex.NumPlaneVertices(PlaneIndex);
-			int32 ConvexEdgeVertexIndex = Convex.GetPlaneVertex(PlaneIndex, NumPlaneVerts - 1);
-			FVec3 ConvexEdgeV0 = Convex.GetVertex(ConvexEdgeVertexIndex);
-			for (int32 PlaneVertexIndex = 0; PlaneVertexIndex < NumPlaneVerts; ++PlaneVertexIndex)
-			{
-				ConvexEdgeVertexIndex = Convex.GetPlaneVertex(PlaneIndex, PlaneVertexIndex);
-				const FVec3 ConvexEdgeV1 = Convex.GetVertex(ConvexEdgeVertexIndex);
-
-				FVec3 TriEdgeV0 = Triangle.GetVertex(2);
-				for (int32 TriVertexIndex = 0; TriVertexIndex < 3; ++TriVertexIndex)
-				{
-					const FVec3 TriEdgeV1 = Triangle.GetVertex(TriVertexIndex);
-
-					FVec3 EdgeN = FVec3::CrossProduct(ConvexEdgeV1 - ConvexEdgeV0, TriEdgeV1 - TriEdgeV0);
-
-					// Ignore parallel edge pairs
-					const FReal EdgeNLenSq = EdgeN.SizeSquared();
-					if (EdgeNLenSq < NormalToleranceSq)
-					{
-						continue;
-					}
-
-					// Project the convex and triangle onto the edge axis
-					// NOTE: EdgeN is not normalized yet, so these projections will be scaled by EdgeN.Size()
-					FReal ThisConvexEdgeDMin, ThisConvexEdgeDMax, ThisTriEdgeDMin, ThisTriEdgeDMax;;
-					int32 ThisConvexEdgeIndexMin, ThisConvexEdgeIndexMax, ThisTriEdgeIndexMin, ThisTriEdgeIndexMax;
-					ProjectOntoAxis(Convex, EdgeN, ConvexEdgeV0, ThisConvexEdgeDMin, ThisConvexEdgeDMax, ThisConvexEdgeIndexMin, ThisConvexEdgeIndexMax, nullptr);
-					ProjectOntoAxis(Triangle, EdgeN, ConvexEdgeV0, ThisTriEdgeDMin, ThisTriEdgeDMax, ThisTriEdgeIndexMin, ThisTriEdgeIndexMax);
-
-					// Separation will be the larger of these
-					// NOTE: these separations are scaled by EdgeN.Size()
-					const FReal ScaledPhi0 = ThisConvexEdgeDMin - ThisTriEdgeDMax;
-					const FReal ScaledPhi1 = ThisTriEdgeDMin - ThisConvexEdgeDMax;
-
-					// Ignore separations greater than cull distance (taking the non-noemlized edge axis into account)
-					const FReal ScaledCullDistanceSq = FMath::Square(CullDistance) * EdgeNLenSq;
-					if (FMath::Square(ScaledPhi0) > ScaledCullDistanceSq)
-					{
-						continue;
-					}
-					if (FMath::Square(ScaledPhi1) > ScaledCullDistanceSq)
-					{
-						continue;
-					}
-
-					// Normalize the edge axis
-					const FReal InvEdgeNLen = FMath::InvSqrt(EdgeNLenSq);
-					const FReal Phi0 = ScaledPhi0 * InvEdgeNLen;
-					const FReal Phi1 = ScaledPhi1 * InvEdgeNLen;
-					EdgeN = EdgeN * InvEdgeNLen;
-
-					// See if we have a new shallowest overlap
-					if ((Phi0 >= Phi1) && (Phi0 > EdgeEdgeDMin))
-					{
-						EdgeEdgeN = EdgeN;
-						EdgeEdgeDMin = Phi0;
-						ConvexEdgeIndexMin = ThisConvexEdgeIndexMin;
-						TriEdgeIndexMin = ThisTriEdgeIndexMax;
-					}
-					else if (Phi1 > EdgeEdgeDMin)
-					{
-						EdgeEdgeN = EdgeN;
-						EdgeEdgeDMin = Phi1;
-						ConvexEdgeIndexMin = ThisConvexEdgeIndexMax;
-						TriEdgeIndexMin = ThisTriEdgeIndexMin;
-					}
-
-					TriEdgeV0 = TriEdgeV1;
-				}
-
-				ConvexEdgeV0 = ConvexEdgeV1;
-			}
-		}
-#else
-		#error Invalid CHAOS_CONVEX_TRIANGLE_EDGEEDGE_METHOD
-#endif
 
 		// Determine which of the features we want to use
 		// NOTE: we rely on the fact that all valid Phi values are greater than InvalidPhi here
@@ -624,46 +531,548 @@ namespace Chaos
 		}
 	}
 
-
-	template
-	void ConstructConvexTriangleOneShotManifold2(
+	template void ConstructConvexTriangleOneShotManifold2(
 		const FImplicitConvex3& Convex, 
 		const FTriangle& Triangle,
 		const FReal CullDistance, 
 		FContactPointManifold& OutContactPoints);
 
-	template
-	void ConstructConvexTriangleOneShotManifold2(
+	template void ConstructConvexTriangleOneShotManifold2(
 		const TImplicitObjectInstanced<FImplicitConvex3>& Convex,
 		const FTriangle& Triangle,
 		const FReal CullDistance,
 		FContactPointManifold& OutContactPoints);
 
-	template
-	void ConstructConvexTriangleOneShotManifold2(
+	template void ConstructConvexTriangleOneShotManifold2(
 		const TImplicitObjectScaled<FImplicitConvex3>& Convex,
 		const FTriangle& Triangle,
 		const FReal CullDistance,
 		FContactPointManifold& OutContactPoints);
 
-	template
-	void ConstructConvexTriangleOneShotManifold2(
+	template void ConstructConvexTriangleOneShotManifold2(
 		const FImplicitBox3& Convex,
 		const FTriangle& Triangle,
 		const FReal CullDistance,
 		FContactPointManifold& OutContactPoints);
 
-	template
-	void ConstructConvexTriangleOneShotManifold2(
+	template  void ConstructConvexTriangleOneShotManifold2(
 		const TImplicitObjectScaled<FImplicitBox3>& Convex,
 		const FTriangle& Triangle,
 		const FReal CullDistance,
 		FContactPointManifold& OutContactPoints);
 
-	template
-	void ConstructConvexTriangleOneShotManifold2(
+	template  void ConstructConvexTriangleOneShotManifold2(
 		const TImplicitObjectInstanced<FImplicitBox3>& Convex,
 		const FTriangle& Triangle,
 		const FReal CullDistance,
 		FContactPointManifold& OutContactPoints);
-}
\ No newline at end of file
+
+} // Chaos
+
+///////////////////////////////////////////////////////////////////////////////
+///////////////////////////////////////////////////////////////////////////////
+///////////////////////////////////////////////////////////////////////////////
+///////////////////////////////////////////////////////////////////////////////
+///////////////////////////////////////////////////////////////////////////////
+
+namespace Chaos::Private
+{
+	// Generate a contact manifold between a convex and a triangle, given the closest feature (i.e., single Contact point)
+	template <typename ConvexType>
+	void ConvexTriangleManifoldFromContact(const ConvexType& Convex, const FTriangle& Triangle, const FConvexContactPoint& Contact, const FReal CullDistance, FContactPointLargeManifold& OutManifold)
+	{
+		// Convex plane
+		const int32 ConvexPlaneIndex = Contact.Features[0].PlaneIndex;
+		check(ConvexPlaneIndex != INDEX_NONE);
+		FVec3 ConvexPlaneN, ConvexPlaneX;
+		Convex.GetPlaneNX(ConvexPlaneIndex, ConvexPlaneN, ConvexPlaneX);
+
+		// Triangle plane
+		check(Contact.Features[1].PlaneIndex == 0);
+		const FVec3 TriN = Triangle.GetNormal();
+		const FVec3 TriC = Triangle.GetCentroid();
+
+		const FVec3& SeparatingAxis = Contact.ShapeContactNormal;
+
+		// Decide whether we want to clip the tri against the convex face or vice-versa
+		bool bClipConvexToTriangle = false;
+		FVec3 ClipAxis = SeparatingAxis;
+		const FReal TriNDotAxis = FVec3::DotProduct(TriN, ClipAxis);
+		const FReal ConvexNDotAxis = FVec3::DotProduct(ConvexPlaneN, ClipAxis);
+		if (FMath::Abs(TriNDotAxis) > FMath::Abs(ConvexNDotAxis))
+		{
+			bClipConvexToTriangle = true;
+			if (TriNDotAxis < 0)
+			{
+				ClipAxis = -ClipAxis;
+			}
+		}
+		else
+		{
+			bClipConvexToTriangle = false;
+			if (ConvexNDotAxis < 0)
+			{
+				ClipAxis = -ClipAxis;
+			}
+		}
+
+		// @todo(chaos): scratch or stack allocation of clipped vertex buffers
+		// @todo(chaos): 2D Clip
+		const int32 MaxContactPointCount = 32;
+		TCArray<FVec3, MaxContactPointCount> ClippedVertexBuffers0 = TCArray<FVec3, MaxContactPointCount>::MakeFull();
+		TCArray<FVec3, MaxContactPointCount> ClippedVertexBuffers1 = TCArray<FVec3, MaxContactPointCount>::MakeFull();
+		TArrayView<FVec3> ClippedVertices;
+		if (bClipConvexToTriangle)
+		{
+			ClippedVertices = ClipConvexToTriangle(Convex, ConvexPlaneIndex, Triangle, ClipAxis, MakeArrayView(ClippedVertexBuffers0), MakeArrayView(ClippedVertexBuffers1));
+		}
+		else
+		{
+			ClippedVertices = ClipTriangleToConvex(Triangle, Convex, ConvexPlaneIndex, ClipAxis, MakeArrayView(ClippedVertexBuffers0), MakeArrayView(ClippedVertexBuffers1));
+		}
+
+		// Reduce number of contacts to the maximum allowed
+		if (ClippedVertices.Num() > 4)
+		{
+			const FRotation3 RotateSeperationToZ = FRotation3::FromRotatedVector(SeparatingAxis, FVec3(0.0f, 0.0f, 1.0f));
+			for (int32 ContactPointIndex = 0; ContactPointIndex < ClippedVertices.Num(); ++ContactPointIndex)
+			{
+				ClippedVertices[ContactPointIndex] = RotateSeperationToZ * ClippedVertices[ContactPointIndex];
+			}
+
+			const int32 ReducedContactPointCount = Collisions::ReduceManifoldContactPoints(ClippedVertices.GetData(), ClippedVertices.Num());
+			ClippedVertices = MakeArrayView(ClippedVertices.GetData(), ReducedContactPointCount);
+
+			for (int32 ContactPointIndex = 0; ContactPointIndex < ClippedVertices.Num(); ++ContactPointIndex)
+			{
+				ClippedVertices[ContactPointIndex] = RotateSeperationToZ.Inverse() * ClippedVertices[ContactPointIndex];
+			}
+		}
+		check(ClippedVertices.Num() <= OutManifold.Max());
+
+		const EContactPointType ContactType = EContactPointType::Unknown;
+
+		// Add the clipped points to the contact list
+		const auto& AddContact = [&CullDistance, &ContactType, &SeparatingAxis, &OutManifold](const FVec3& ConvexX, const FVec3 TriX, const FReal Distance)
+		{
+			if (Distance < CullDistance)
+			{
+				FContactPoint& ContactPoint = OutManifold[OutManifold.AddUninitialized()];
+				ContactPoint.ShapeContactPoints[0] = ConvexX;
+				ContactPoint.ShapeContactPoints[1] = TriX;
+				ContactPoint.ShapeContactNormal = SeparatingAxis;
+				ContactPoint.Phi = Distance;
+				ContactPoint.ContactType = ContactType;
+				ContactPoint.FaceIndex = INDEX_NONE;
+			}
+		};
+
+		// If the convex center projects inside the triangle and inside the convex face, add it as the first contact. This helps with solver convergence.
+		// @todo(chaos): ideally this would be the particle center of mass - we could pass that in.
+		// @todo(chaos): also need to check that the points projects inside the convex face
+		//const FVec3 ConvexCenterBarycentric = ToBarycentric(Convex.GetCenterOfMass(), Triangle.GetVertex(0), Triangle.GetVertex(1), Triangle.GetVertex(2));
+		//if ((ConvexCenterBarycentric.X >= 0) && (ConvexCenterBarycentric.Y >= 0) && (ConvexCenterBarycentric.Z >= 0))
+		//{
+		//	const FVec3 TriX = FromBarycentric(ConvexCenterBarycentric, Triangle.GetVertex(0), Triangle.GetVertex(1), Triangle.GetVertex(2));
+		//	const FReal ConvexNDotTriN = FVec3::DotProduct(ConvexPlaneN, TriN);
+		//	if (ConvexNDotTriN < UE_KINDA_SMALL_NUMBER)
+		//	{
+		//		const FReal Distance = FVec3::DotProduct(ConvexPlaneX - TriX, ConvexPlaneN) / ConvexNDotTriN;
+		//		const FVec3 ConvexX = TriX + Distance * TriN;
+		//		AddContact(ConvexX, TriX, Distance);
+		//	}
+		//}
+
+		if (bClipConvexToTriangle)
+		{
+			// Clipped points are on the convex, and we clipped to the triangle face
+			for (int32 ContactIndex = 0; ContactIndex < ClippedVertices.Num(); ++ContactIndex)
+			{
+				const FVec3& ConvexX = ClippedVertices[ContactIndex];
+				const FReal IntersectDenom = FVec3::DotProduct(SeparatingAxis, TriN);
+				check(FMath::Abs(IntersectDenom) > UE_SMALL_NUMBER);	// Guaranteed based on axis selection above
+				const FReal Distance = FVec3::DotProduct(ConvexX - TriC, TriN) / IntersectDenom;
+				const FVec3 TriX = ConvexX - Distance * SeparatingAxis;
+
+				AddContact(ConvexX, TriX, Distance);
+			}
+		}
+		else
+		{
+			// Clipped points are on the triangle, and we clipped to the convex face
+			for (int32 ContactIndex = 0; ContactIndex < ClippedVertices.Num(); ++ContactIndex)
+			{
+				const FVec3& TriX = ClippedVertices[ContactIndex];
+				const FReal IntersectDenom = FVec3::DotProduct(SeparatingAxis, ConvexPlaneN);
+				check(FMath::Abs(IntersectDenom) > UE_SMALL_NUMBER);	// Guaranteed based on axis selection above
+				const FReal Distance = FVec3::DotProduct(ConvexPlaneX - TriX, ConvexPlaneN) / IntersectDenom;
+				const FVec3 ConvexX = TriX + Distance * SeparatingAxis;
+
+				AddContact(ConvexX, TriX, Distance);
+			}
+		}
+
+	}
+
+	template void ConvexTriangleManifoldFromContact(
+		const FImplicitConvex3& Convex, 
+		const FTriangle& Triangle, 
+		const FConvexContactPoint& Contact,
+		const FReal CullDistance, 
+		FContactPointLargeManifold& OutManifold);
+
+	template void ConvexTriangleManifoldFromContact(
+		const TImplicitObjectInstanced<FImplicitConvex3>& Convex,
+		const FTriangle& Triangle, 
+		const FConvexContactPoint& Contact,
+		const FReal CullDistance, 
+		FContactPointLargeManifold& OutManifold);
+
+	template void ConvexTriangleManifoldFromContact(
+		const TImplicitObjectScaled<FImplicitConvex3>& Convex,
+		const FTriangle& Triangle,
+		const FConvexContactPoint& Contact,
+		const FReal CullDistance,
+		FContactPointLargeManifold& OutManifold);
+
+	template void ConvexTriangleManifoldFromContact(
+		const FImplicitBox3& Convex,
+		const FTriangle& Triangle,
+		const FConvexContactPoint& Contact,
+		const FReal CullDistance,
+		FContactPointLargeManifold& OutManifold);
+
+	template void ConvexTriangleManifoldFromContact(
+		const TImplicitObjectScaled<FImplicitBox3>& Convex,
+		const FTriangle& Triangle,
+		const FConvexContactPoint& Contact,
+		const FReal CullDistance,
+		FContactPointLargeManifold& OutManifold);
+
+	template void ConvexTriangleManifoldFromContact(
+		const TImplicitObjectInstanced<FImplicitBox3>& Convex,
+		const FTriangle& Triangle,
+		const FConvexContactPoint& Contact,
+		const FReal CullDistance,
+		FContactPointLargeManifold& OutManifold);
+
+
+	// Generate a single contact point between a convex and a triangle
+	template <typename ConvexType>
+	bool ConvexTriangleContactPoint(const ConvexType& Convex, const FTriangle& Triangle, const FReal CullDistance, FContactPoint& OutContactPoint)
+	{
+		// The GJK version has issues with exactly touching shapes. 
+		// E.g., see failing unit test EPARealFailures_TouchingBoxTriangle
+#if 1
+		const TGJKCoreShape<ConvexType> GJKConvex(Convex, Convex.GetMargin());
+		const TGJKShape<FTriangle> GJKTriangle(Triangle);
+
+		FReal UnusedMaxMarginDelta = FReal(0);
+		int32 ConvexVertexIndex = INDEX_NONE;
+		int32 TriangleVertexIndex = INDEX_NONE;
+		FReal Penetration;
+		FVec3 ConvexClosest, TriangleClosest, ConvexNormal;
+		const FReal GJKEpsilon = Chaos_Collision_GJKEpsilon;
+		const FReal EPAEpsilon = Chaos_Collision_EPAEpsilon;
+
+		FVec3 InitialGJKDir = FVec3(1, 0, 0);
+
+		const bool bHaveContact = GJKPenetrationSameSpace(
+			GJKConvex,
+			GJKTriangle,
+			Penetration,
+			ConvexClosest,
+			TriangleClosest,
+			ConvexNormal,
+			ConvexVertexIndex,
+			TriangleVertexIndex,
+			UnusedMaxMarginDelta,
+			InitialGJKDir,
+			GJKEpsilon, EPAEpsilon);
+
+		if (bHaveContact && (-Penetration < CullDistance))
+		{
+			OutContactPoint.ShapeContactPoints[0] = ConvexClosest;
+			OutContactPoint.ShapeContactPoints[1] = TriangleClosest;
+			OutContactPoint.ShapeContactNormal = -ConvexNormal;
+			OutContactPoint.ContactType = EContactPointType::Unknown;
+			OutContactPoint.Phi = -Penetration;
+			OutContactPoint.FaceIndex = INDEX_NONE;
+			return true;
+		}
+
+		return false;
+#else
+
+#endif
+	}
+
+	template bool ConvexTriangleContactPoint(
+		const FImplicitConvex3& Convex,
+		const FTriangle& Triangle,
+		const FReal CullDistance,
+		FContactPoint& OutContactPoint);
+
+	template bool ConvexTriangleContactPoint(
+		const TImplicitObjectInstanced<FImplicitConvex3>& Convex,
+		const FTriangle& Triangle,
+		const FReal CullDistance,
+		FContactPoint& OutContactPoint);
+
+	template bool ConvexTriangleContactPoint(
+		const TImplicitObjectScaled<FImplicitConvex3>& Convex,
+		const FTriangle& Triangle,
+		const FReal CullDistance,
+		FContactPoint& OutContactPoint);
+
+	template bool ConvexTriangleContactPoint(
+		const FImplicitBox3& Convex,
+		const FTriangle& Triangle,
+		const FReal CullDistance,
+		FContactPoint& OutContactPoint);
+
+	template bool ConvexTriangleContactPoint(
+		const TImplicitObjectScaled<FImplicitBox3>& Convex,
+		const FTriangle& Triangle,
+		const FReal CullDistance,
+		FContactPoint& OutContactPoint);
+
+	template bool ConvexTriangleContactPoint(
+		const TImplicitObjectInstanced<FImplicitBox3>& Convex,
+		const FTriangle& Triangle,
+		const FReal CullDistance,
+		FContactPoint& OutContactPoint);
+
+
+	// Get the convex feature at the specific position and normal
+	template<typename ConvexType>
+	bool GetConvexFeature(const ConvexType& Convex, const FVec3& Position, const FVec3& Normal, Private::FConvexFeature& OutFeature)
+	{
+		const FReal NormalTolerance = FReal(1.e-6);
+		const FReal PositionTolerance = FReal(1.e-4);
+		const FReal ToleranceSizeMultiplier = Convex.BoundingBox().Extents().GetAbsMax();
+		const FReal EdgeNormalTolerance = ToleranceSizeMultiplier * FReal(1.e-3);
+
+		int32 BestPlaneIndex = INDEX_NONE;
+		FReal BestPlaneDotNormal = FReal(-1);
+
+		// Get the support vertex along the normal (which must point away from the convex)
+		int SupportVertexIndex = INDEX_NONE;
+		Convex.SupportCore(Normal, 0, nullptr, SupportVertexIndex);
+
+		if (SupportVertexIndex != INDEX_NONE)
+		{
+			// See if the normal matches a face normal for any face using the vertex
+			int32 VertexPlanes[16];
+			int32 NumVertexPlanes = Convex.FindVertexPlanes(SupportVertexIndex, VertexPlanes, UE_ARRAY_COUNT(VertexPlanes));
+			for (int32 VertexPlaneIndex = 0; VertexPlaneIndex < NumVertexPlanes; ++VertexPlaneIndex)
+			{
+				const int32 PlaneIndex = VertexPlanes[VertexPlaneIndex];
+				FVec3 PlaneN, PlaneX;
+				Convex.GetPlaneNX(PlaneIndex, PlaneN, PlaneX);
+				const FReal PlaneDotNormal = FVec3::DotProduct(PlaneN, Normal);
+				if (FMath::IsNearlyEqual(PlaneDotNormal, FReal(1), NormalTolerance))
+				{
+					OutFeature.FeatureType = Private::EConvexFeatureType::Plane;
+					OutFeature.PlaneIndex = PlaneIndex;
+					OutFeature.PlaneFeatureIndex = 0;
+					return true;
+				}
+
+				if (PlaneDotNormal > BestPlaneDotNormal)
+				{
+					BestPlaneIndex = PlaneIndex;
+					BestPlaneDotNormal = PlaneDotNormal;
+				}
+			}
+
+			// See if any of the edges using the vertex are perpendicular to the normal
+			// @todo(chaos): we could visit the vertex edges here rather than use the plane edges
+			if (BestPlaneIndex != INDEX_NONE)
+			{
+				int32 BestPlaneVertexIndex = INDEX_NONE;
+
+				const int32 NumPlaneVertices = Convex.NumPlaneVertices(BestPlaneIndex);
+				for (int32 PlaneVertexIndex = 0; PlaneVertexIndex < NumPlaneVertices; ++PlaneVertexIndex)
+				{
+					const int32 VertexIndex0 = Convex.GetPlaneVertex(BestPlaneIndex, PlaneVertexIndex);
+					const int32 VertexIndex1 = (PlaneVertexIndex == NumPlaneVertices - 1) ? Convex.GetPlaneVertex(BestPlaneIndex, 0) : Convex.GetPlaneVertex(BestPlaneIndex, PlaneVertexIndex + 1);
+
+					if (VertexIndex0 == SupportVertexIndex)
+					{
+						BestPlaneVertexIndex = PlaneVertexIndex;
+					}
+
+					if ((VertexIndex0 == SupportVertexIndex) || (VertexIndex1 == SupportVertexIndex))
+					{
+						const FVec3 Vertex0 = Convex.GetVertex(VertexIndex0);
+						const FVec3 Vertex1 = Convex.GetVertex(VertexIndex1);
+						const FVec3 EdgeDelta = Vertex1 - Vertex0;
+						const FReal EdgeDotNormal = FVec3::DotProduct(EdgeDelta, Normal);
+						if (FMath::Abs(EdgeDotNormal) < EdgeNormalTolerance)
+						{
+							// @todo(chaos): we need to be able to get an EdgeIndex (probably half edge index)
+							// Also, we probably want both the plane index and the edge index
+							OutFeature.FeatureType = Private::EConvexFeatureType::Edge;
+							OutFeature.PlaneIndex = BestPlaneIndex;
+							OutFeature.PlaneFeatureIndex = PlaneVertexIndex;
+							return true;
+						}
+					}
+				}
+
+				// Not a face or edge, so it should be the SupportVertex, but we need to specify the 
+				// plane and plane-index rather than the convex vertex index (which we found just above)
+				if (BestPlaneVertexIndex != INDEX_NONE)
+				{
+					OutFeature.FeatureType = Private::EConvexFeatureType::Vertex;
+					OutFeature.PlaneIndex = BestPlaneIndex;
+					OutFeature.PlaneFeatureIndex = BestPlaneVertexIndex;
+				}
+				return true;
+			}
+		}
+
+		return false;
+	}
+
+	// Get the triangle feature at the specific position and normal
+	template<>
+	bool GetConvexFeature(const FTriangle& Triangle, const FVec3& Position, const FVec3& Normal, Private::FConvexFeature& OutFeature)
+	{
+		// @todo(chaos): pass in the triangle normal - we almost certainly calculated it elsewhere
+		// NOTE: The normal epsilon needs to be less than the maximu error that GJK/EPA produces when it hits a degenerate
+		// case, which can happen when we have almost exact face-to-face contact. The max error is hard to know, since it 
+		// depends on the state of GJK on the iteration before it hits its tolerance, but seems to be typically ~0.01
+		const FReal NormalEpsilon = FReal(0.02);
+		const FVec3 TriangleNormal = Triangle.GetNormal();
+		const FReal NormalDot = FVec3::DotProduct(Normal, TriangleNormal);
+		if (FMath::IsNearlyEqual(NormalDot, FReal(1), NormalEpsilon))
+		{
+			OutFeature.FeatureType = Private::EConvexFeatureType::Plane;
+			OutFeature.PlaneIndex = 0;
+			OutFeature.PlaneFeatureIndex = 0;
+			return true;
+		}
+
+		const FReal BarycentricTolerance = FReal(1.e-6);
+		int32 VertexIndex0, VertexIndex1;
+		if (GetTriangleEdgeVerticesAtPosition(Position, &Triangle.GetVertex(0), VertexIndex0, VertexIndex1, BarycentricTolerance))
+		{
+			if ((VertexIndex0 != INDEX_NONE) && (VertexIndex1 != INDEX_NONE))
+			{
+				OutFeature.FeatureType = Private::EConvexFeatureType::Edge;
+				OutFeature.PlaneIndex = 0;
+				OutFeature.PlaneFeatureIndex = VertexIndex0;
+				return true;
+			}
+			else if (VertexIndex0 != INDEX_NONE)
+			{
+				OutFeature.FeatureType = Private::EConvexFeatureType::Vertex;
+				OutFeature.PlaneIndex = 0;
+				OutFeature.PlaneFeatureIndex = VertexIndex0;
+				return true;
+			}
+			else if (VertexIndex1 != INDEX_NONE)
+			{
+				OutFeature.FeatureType = Private::EConvexFeatureType::Vertex;
+				OutFeature.PlaneIndex = 0;
+				OutFeature.PlaneFeatureIndex = VertexIndex1;
+				return true;
+			}
+		}
+
+		return false;
+	}
+
+	// Generate the contact point and closest feature types between a convex and a triangle
+	template<typename ConvexType>
+	bool FindClosestFeatures(const ConvexType& Convex, const FRigidTransform3& ConvexTransform, const FTriangle& Triangle, const FVec3& ConvexRelativeMovement, const FReal CullDistance, FConvexContactPoint& OutContact)
+	{
+		// Find the closest point on the convex and triangle that we will use to generate the manifold
+		// NOTE: use an upper limit on cull distance here since the real cull distance depends on the motion against the contact normal
+		const FReal EarlyCullDistance = TNumericLimits<FReal>::Max();
+		FContactPoint ContactPoint;
+		if (!Private::ConvexTriangleContactPoint(Convex, Triangle, EarlyCullDistance, ContactPoint))
+		{
+			return false;
+		}
+
+		// Now check the cull distance, taking movement into account
+		const FReal SeparationWithMotion = ContactPoint.Phi + FVec3::DotProduct(ConvexRelativeMovement, ContactPoint.ShapeContactNormal);
+		if ((ContactPoint.Phi > CullDistance) && (SeparationWithMotion > CullDistance))
+		{
+			return false;
+		}
+
+		// Initialize outputs
+		OutContact.Init();
+		OutContact.ShapeContactPoints[0] = ContactPoint.ShapeContactPoints[0];
+		OutContact.ShapeContactPoints[1] = ContactPoint.ShapeContactPoints[1];
+		OutContact.ShapeContactNormal = ContactPoint.ShapeContactNormal;
+		OutContact.Phi = ContactPoint.Phi;
+
+		// Find the triangle feature at the contact point
+		if (!Private::GetConvexFeature(Triangle, OutContact.ShapeContactPoints[1], OutContact.ShapeContactNormal, OutContact.Features[1]))
+		{
+			return false;
+		}
+
+		// Find the convex feature at the contact point
+		if (!Private::GetConvexFeature(Convex, OutContact.ShapeContactPoints[0], -OutContact.ShapeContactNormal, OutContact.Features[0]))
+		{
+			return false;
+		}
+
+		return true;
+	}
+
+	template bool FindClosestFeatures(
+		const FImplicitConvex3& Convex,
+		const FRigidTransform3& ConvexTransform, 
+		const FTriangle& Triangle, 
+		const FVec3& ConvexRelativeMovement, 
+		const FReal CullDistance, 
+		FConvexContactPoint& OutContact);
+
+	template bool FindClosestFeatures(
+		const TImplicitObjectInstanced<FImplicitConvex3>& Convex,
+		const FRigidTransform3& ConvexTransform,
+		const FTriangle& Triangle,
+		const FVec3& ConvexRelativeMovement,
+		const FReal CullDistance,
+		FConvexContactPoint& OutContact);
+
+	template bool FindClosestFeatures(
+		const TImplicitObjectScaled<FImplicitConvex3>& Convex,
+		const FRigidTransform3& ConvexTransform,
+		const FTriangle& Triangle,
+		const FVec3& ConvexRelativeMovement,
+		const FReal CullDistance,
+		FConvexContactPoint& OutContact);
+
+	template bool FindClosestFeatures(
+		const FImplicitBox3& Convex,
+		const FRigidTransform3& ConvexTransform,
+		const FTriangle& Triangle,
+		const FVec3& ConvexRelativeMovement,
+		const FReal CullDistance,
+		FConvexContactPoint& OutContact);
+
+	template bool FindClosestFeatures(
+		const TImplicitObjectScaled<FImplicitBox3>& Convex,
+		const FRigidTransform3& ConvexTransform,
+		const FTriangle& Triangle,
+		const FVec3& ConvexRelativeMovement,
+		const FReal CullDistance,
+		FConvexContactPoint& OutContact);
+
+	template bool FindClosestFeatures(
+		const TImplicitObjectInstanced<FImplicitBox3>& Convex,
+		const FRigidTransform3& ConvexTransform,
+		const FTriangle& Triangle,
+		const FVec3& ConvexRelativeMovement,
+		const FReal CullDistance,
+		FConvexContactPoint& OutContact);
+
+} // Chaos::Private
diff --git a/Engine/Source/Runtime/Experimental/Chaos/Private/Chaos/Collision/MeshContactGenerator.cpp b/Engine/Source/Runtime/Experimental/Chaos/Private/Chaos/Collision/MeshContactGenerator.cpp
--- a/Engine/Source/Runtime/Experimental/Chaos/Private/Chaos/Collision/MeshContactGenerator.cpp
+++ b/Engine/Source/Runtime/Experimental/Chaos/Private/Chaos/Collision/MeshContactGenerator.cpp
@@ -1,5 +1,6 @@
 // Copyright Epic Games, Inc.All Rights Reserved.
 #include "Chaos/Collision/MeshContactGenerator.h"
+#include "Chaos/Collision/ConvexFeature.h"
 #include "Chaos/DebugDrawQueue.h"
 
 namespace Chaos
@@ -12,6 +13,9 @@ namespace Chaos
 namespace Chaos::CVars
 {
 	extern int32 ChaosSolverDebugDrawMeshContacts;
+
+	bool bMeshContactGeneratorFixContactNormalFixEnabled = true;
+	FAutoConsoleVariableRef CVarChaosMeshContactGeneratorFixContactNormalFixEnabled(TEXT("p.Chaos.MeshContactGenerator.FixContactNormal.FixEnabled"), bMeshContactGeneratorFixContactNormalFixEnabled, TEXT("Until new code path is well tested"));
 }
 
 namespace Chaos::Private
@@ -24,6 +28,8 @@ namespace Chaos::Private
 		BarycentricTolerance = FReal(1.e-3);
 		MaxContactsBufferSize = 1000;
 		bCullBackFaces = true;
+		bFixNormals = true;
+		bSortForSolver = false;
 		bUseTwoPassLoop = true;
 	}
 
@@ -49,14 +55,8 @@ namespace Chaos::Private
 		VertexContactIndicesMap.Reset(InMaxContacts);
 	}
 
-	void FMeshContactGenerator::AddTriangleContacts(const FContactPointManifold& TriangleContactPoints, const int32 LocalTriangleIndex)
+	void FMeshContactGenerator::AddTriangleContacts(const int32 LocalTriangleIndex, const TArrayView<FContactPoint>& TriangleContactPoints)
 	{
-		// Clamp the number of contacts we can add per mesh
-		if (Contacts.Num() + TriangleContactPoints.Num() > Contacts.Max())
-		{
-			return;
-		}
-
 		FTriangleExt& Triangle = Triangles[LocalTriangleIndex];
 
 		// We need to know how many vertex and edge collisions a triangle has on it (from collision with other triangle 
@@ -64,6 +64,12 @@ namespace Chaos::Private
 		// contacts and assigne edge/vertex status as required. See GenerateMeshContacts()
 		for (const FContactPoint& ContactPoint : TriangleContactPoints)
 		{
+			// Clamp the number of contacts we can add per mesh
+			if (Contacts.Num() == Contacts.Max())
+			{
+				break;
+			}
+
 			// See if we have a vertex or edge contact
 			// @todo(chaos): we should be able to produce this as output from the manifold creator
 			int32 LocalVertexID0, LocalVertexID1;
@@ -145,8 +151,11 @@ namespace Chaos::Private
 
 	void FMeshContactGenerator::ProcessGeneratedContacts(const FRigidTransform3& ConvexTransform, const FRigidTransform3& MeshToConvexTransform)
 	{
-		// Contacts that get pruned or corrected will show as green
-		DebugDrawContacts(ConvexTransform, FColor::Green, 0.5);
+		if (!!Settings.bFixNormals)
+		{
+			// Contacts that get pruned or corrected will show as green
+			DebugDrawContacts(ConvexTransform, FColor::Green, 0.5);
+		}
 
 		PruneAndCorrectContacts();
 
@@ -156,6 +165,11 @@ namespace Chaos::Private
 		// Visited triangles are white, ignored triangles are gray
 		DebugDrawTriangles(ConvexTransform, FColor::White, FColor::Silver);
 
+		if (!!Settings.bSortForSolver)
+		{
+			SortContactsForSolver();
+		}
+
 		FinalizeContacts(MeshToConvexTransform);
 	}
 
@@ -173,7 +187,7 @@ namespace Chaos::Private
 			}
 
 			// Fix edge normals that aren't already close to a face normal
-			if (ContactPointData.GetContactNormalDotTriangleNormal() < Settings.FaceNormalDotThreshold)
+			if (!!Settings.bFixNormals && (ContactPointData.GetContactNormalDotTriangleNormal() < Settings.FaceNormalDotThreshold))
 			{
 				FixContactNormal(ContactIndex);
 			}
@@ -244,41 +258,92 @@ namespace Chaos::Private
 		{
 			// We have a vertex collision. Ensure that the contact normal is in a valid range for the vertex based on the
 			// triangles that share the vertex (there can be arbitarily many of these).
-			for (int32 TriangleIndex = 0; TriangleIndex < Triangles.Num(); ++TriangleIndex)
+
+			// NOTE: the code in the cvar fixes an issue which would leave invalid normals on spikey meshes,
+			// but we're very close to a release, so adding option to rollback if this version is flased
+			if (CVars::bMeshContactGeneratorFixContactNormalFixEnabled)
 			{
-				const FTriangleExt& OtherTriangle = Triangles[TriangleIndex];
-				if ((TriangleIndex != LocalTriangleIndex) && OtherTriangle.HasVertexIndex(ContactPointData.GetVertexID()))
+				const int32 VertexIndexA = ContactPointData.GetVertexID();
+				FVec3 VertexA;
+				if (!Triangle.GetVertexWithID(VertexIndexA, VertexA))
+				{
+					// @todo(chaos): this is an error condition
+					return;
+				}
+
+				// @todo(chaos): the map of Vertex->TriangleIndices would help here but may not be a net win
+				for (int32 OtherLocalTriangleIndex = 0; OtherLocalTriangleIndex < Triangles.Num(); ++OtherLocalTriangleIndex)
 				{
-					FVec3 VertexA, VertexB, VertexC;
-					if (OtherTriangle.GetVertexPosition(ContactPointData.GetVertexID(), VertexA) && OtherTriangle.GetOtherVertexPositions(ContactPointData.GetVertexID(), VertexB, VertexC))
+					const FTriangleExt& OtherTriangle = Triangles[OtherLocalTriangleIndex];
+					FVec3 OtherVertexB, OtherVertexC;
+					if (OtherTriangle.GetOtherVerticesFromID(VertexIndexA, OtherVertexB, OtherVertexC))
 					{
 						// Does the contact normal point into the infinite prism formed by extruding the triangle along the face normal?
-						// It does if the contact normal dotted with the edge plane normal is negative for both edge planes on the triangle that use the vertex.
-						const FVec3 EdgeDelta0 = VertexB - VertexA;
-						const FVec3 EdgeDelta1 = VertexC - VertexA;
 						const FVec3& OtherTriangleNormal = OtherTriangle.GetNormal();
-
-						const FReal EdgeSign0 = FVec3::DotProduct(FVec3::CrossProduct(ContactPoint.ShapeContactNormal, VertexB - VertexA), OtherTriangleNormal);
-						const FReal EdgeSign1 = FVec3::DotProduct(FVec3::CrossProduct(ContactPoint.ShapeContactNormal, VertexC - VertexA), OtherTriangleNormal);
-						if (FMath::Sign(EdgeSign0) == FMath::Sign(EdgeSign1))
+						const FVec3 OtherEdge0 = OtherVertexB - VertexA;
+						const FVec3 OtherEdge1 = VertexA - OtherVertexC;
+						const FVec3 OtherEdgeNormal0 = FVec3::CrossProduct(OtherTriangleNormal, OtherEdge0);	// Not normlized
+						const FVec3 OtherEdgeNormal1 = FVec3::CrossProduct(OtherTriangleNormal, OtherEdge1);	// Not normlized
+						const FReal OtherEdgeSign0 = FVec3::DotProduct(ContactPoint.ShapeContactNormal, OtherEdgeNormal0);
+						const FReal OtherEdgeSign1 = FVec3::DotProduct(ContactPoint.ShapeContactNormal, OtherEdgeNormal1);
+						const FReal OtherEdgeSign0Sq = OtherEdgeSign0 * FMath::Abs(OtherEdgeSign0);
+						const FReal OtherEdgeSign1Sq = OtherEdgeSign1 * FMath::Abs(OtherEdgeSign1);
+						const FReal NormalToleranceSq = FReal(1.e-8);
+						const FReal NormalTolerance0Sq = NormalToleranceSq * OtherEdge0.SizeSquared();
+						const FReal NormalTolerance1Sq = NormalToleranceSq * OtherEdge1.SizeSquared();
+						if ((OtherEdgeSign0Sq >= FReal(-NormalTolerance0Sq)) && (OtherEdgeSign1Sq >= FReal(-NormalTolerance1Sq)))
+						{
+							ContactPoint.ShapeContactNormal = OtherTriangleNormal;
+							ContactPoint.ShapeContactPoints[0] = ContactPoint.ShapeContactPoints[1] + ContactPoint.Phi * OtherTriangleNormal;
+							ContactPointData.SetTriangleIndex(OtherLocalTriangleIndex);
+							break;
+						}
+					}
+				}
+			}
+			else
+			{
+				//
+				//
+				// @todo(chaos): remove this branch when above code is well tested
+				//
+				//
+				for (int32 TriangleIndex = 0; TriangleIndex < Triangles.Num(); ++TriangleIndex)
+				{
+					const FTriangleExt& OtherTriangle = Triangles[TriangleIndex];
+					if ((TriangleIndex != LocalTriangleIndex) && OtherTriangle.HasVertexID(ContactPointData.GetVertexID()))
+					{
+						FVec3 VertexA, VertexB, VertexC;
+						if (OtherTriangle.GetVertexWithID(ContactPointData.GetVertexID(), VertexA) && OtherTriangle.GetOtherVerticesFromID(ContactPointData.GetVertexID(), VertexB, VertexC))
 						{
-							const FVec3 Centroid = OtherTriangle.GetCentroid();
-							const FReal NormalDotCentroid = FVec3::DotProduct(ContactPoint.ShapeContactNormal, Centroid - ContactPoint.ShapeContactPoints[1]);
-							if (NormalDotCentroid > 0)
+							// Does the contact normal point into the infinite prism formed by extruding the triangle along the face normal?
+							// It does if the contact normal dotted with the edge plane normal is negative for both edge planes on the triangle that use the vertex.
+							const FVec3 EdgeDelta0 = VertexB - VertexA;
+							const FVec3 EdgeDelta1 = VertexC - VertexA;
+							const FVec3& OtherTriangleNormal = OtherTriangle.GetNormal();
+
+							const FReal EdgeSign0 = FVec3::DotProduct(FVec3::CrossProduct(ContactPoint.ShapeContactNormal, VertexB - VertexA), OtherTriangleNormal);
+							const FReal EdgeSign1 = FVec3::DotProduct(FVec3::CrossProduct(ContactPoint.ShapeContactNormal, VertexC - VertexA), OtherTriangleNormal);
+							if (FMath::Sign(EdgeSign0) == FMath::Sign(EdgeSign1))
 							{
-								// If our normal was very far away from a valid normal, drop the contact
-								const FReal MinContactDotNormal = FRealSingle(FVec3::DotProduct(OtherTriangleNormal, TriangleNormal));
-								if (MinContactDotNormal - ContactDotNormal > Settings.EdgeNormalDotRejectTolerance)
+								const FVec3 Centroid = OtherTriangle.GetCentroid();
+								const FReal NormalDotCentroid = FVec3::DotProduct(ContactPoint.ShapeContactNormal, Centroid - ContactPoint.ShapeContactPoints[1]);
+								if (NormalDotCentroid > 0)
 								{
-									ContactPointData.SetDisabled();
-									return;
+									// If our normal was very far away from a valid normal, drop the contact
+									const FReal MinContactDotNormal = FRealSingle(FVec3::DotProduct(OtherTriangleNormal, TriangleNormal));
+									if (MinContactDotNormal - ContactDotNormal > Settings.EdgeNormalDotRejectTolerance)
+									{
+										ContactPointData.SetDisabled();
+										return;
+									}
+
+									const FReal OtherContactDotNormal = FVec3::DotProduct(ContactPoint.ShapeContactNormal, OtherTriangleNormal);
+									const FVec3 CorrectedContactNormal = OtherTriangleNormal;
+
+									ContactPoint.ShapeContactNormal = CorrectedContactNormal;
+									ContactPoint.ShapeContactPoints[0] = ContactPoint.ShapeContactPoints[1] + ContactPoint.Phi * CorrectedContactNormal;
 								}
-
-								const FReal OtherContactDotNormal = FVec3::DotProduct(ContactPoint.ShapeContactNormal, OtherTriangleNormal);
-								const FVec3 CorrectedContactNormal = OtherTriangleNormal;
-
-								ContactPoint.ShapeContactNormal = CorrectedContactNormal;
-								ContactPoint.ShapeContactPoints[0] = ContactPoint.ShapeContactPoints[1] + ContactPoint.Phi * CorrectedContactNormal;
 							}
 						}
 					}
@@ -287,8 +352,183 @@ namespace Chaos::Private
 		}
 	}
 
+	bool FMeshContactGenerator::FixFeature(const int32 LocalTriangleIndex, Private::EConvexFeatureType& InOutFeatureType, int32& InOutFeatureIndex, FVec3& InOutPlaneNormal, FVec3& InOutPlanePosition)
+	{
+		const FTriangleExt& Triangle = Triangles[LocalTriangleIndex];
+
+		// For convex edges, we ensure that the normal is between the normals of the adjacent faces
+		// For concave edges, we replace the normal with the triangle face normal
+		if (InOutFeatureType == Private::EConvexFeatureType::Edge)
+		{
+			check(InOutFeatureIndex != INDEX_NONE);
+
+			const int32 LocalVertexIndex0 = InOutFeatureIndex;
+			const int32 LocalVertexIndex1 = (InOutFeatureIndex == 2) ? 0 : InOutFeatureIndex + 1;
+			const int32 VertexIndex0 = Triangle.GetVertexIndex(LocalVertexIndex0);
+			const int32 VertexIndex1 = Triangle.GetVertexIndex(LocalVertexIndex1);
+			const FContactEdgeID EdgeID = FContactEdgeID(VertexIndex0, VertexIndex1);
+
+			const int32 OtherLocalTriangleIndex = GetOtherTriangleIndexForEdge(LocalTriangleIndex, EdgeID);
+			if (OtherLocalTriangleIndex == INDEX_NONE)
+			{
+				// We don't have the other triangle. Maybe we should use the face plane rather than leaving the edge?
+				return false;
+			}
+
+			const FTriangleExt& OtherTriangle = Triangles[OtherLocalTriangleIndex];
+
+			const FVec3& TriangleNormal = Triangle.GetNormal();
+			const FVec3& OtherTriangleNormal = OtherTriangle.GetNormal();
+
+			// Common case - both triangles have the same normal - treat as concave
+			const FReal NormalEpsilon = FReal(1.e-6);
+			const FReal TriangleNormalsDot = FVec3::DotProduct(TriangleNormal, OtherTriangleNormal);
+			if (FMath::IsNearlyEqual(TriangleNormalsDot, FReal(1), NormalEpsilon))
+			{
+				// Concave edge - use the face normal
+				InOutFeatureType = Private::EConvexFeatureType::Plane;
+				InOutFeatureIndex = 0;
+				InOutPlaneNormal = TriangleNormal;
+				return true;
+			}
+
+			// If normals are different do a full concavity check
+			const FVec3 EdgeDelta = Triangle.GetVertex(LocalVertexIndex1) - Triangle.GetVertex(LocalVertexIndex0);
+			const FVec3 TriangleEdgeNormalVector = FVec3::CrossProduct(EdgeDelta, TriangleNormal);	// Not normalized
+			const FReal OtherFaceNormalDotEdgeNormal = FVec3::DotProduct(OtherTriangleNormal, TriangleEdgeNormalVector);
+			if (OtherFaceNormalDotEdgeNormal < FReal(0))
+			{
+				// Concave edge - use the face normal
+				InOutFeatureType = Private::EConvexFeatureType::Plane;
+				InOutFeatureIndex = 0;
+				InOutPlaneNormal = TriangleNormal;
+				return true;
+			}
+
+			// We have a convex edge. Ensure the normal is in the valid region
+			const FReal NormalDotEdge = FVec3::DotProduct(InOutPlaneNormal, TriangleEdgeNormalVector);
+			if (NormalDotEdge < FReal(0))
+			{
+				InOutFeatureType = Private::EConvexFeatureType::Plane;
+				InOutFeatureIndex = 0;
+				InOutPlaneNormal = TriangleNormal;
+				return true;
+			}
+
+			FVec3 OtherVertex0, OtherVertex1;
+			if (OtherTriangle.GetVertexWithID(VertexIndex0, OtherVertex0) && OtherTriangle.GetVertexWithID(VertexIndex1, OtherVertex1))
+			{
+				const FVec3 OtherEdgeDelta = OtherVertex0 - OtherVertex1;
+				const FVec3 OtherTriangleEdgeNormalVector = FVec3::CrossProduct(OtherEdgeDelta, OtherTriangleNormal);	// Not normalized
+				const FReal OtherNormalDotEdge = FVec3::DotProduct(InOutPlaneNormal, OtherTriangleEdgeNormalVector);
+				if (OtherNormalDotEdge < FReal(0))
+				{
+					InOutFeatureType = Private::EConvexFeatureType::Plane;
+					InOutFeatureIndex = 0;
+					InOutPlaneNormal = OtherTriangleNormal;
+					return true;
+				}
+			}
+		}
+
+		// For vertices, ensure that the contact normal is in a valid range for the vertex based on the
+		// triangles that share the vertex (there can be arbitarily many of these).
+		if (InOutFeatureType == Private::EConvexFeatureType::Vertex)
+		{
+			check(InOutFeatureIndex != INDEX_NONE);
+
+			const int32 LocalVertexIndex0 = InOutFeatureIndex;
+			const int32 VertexIndexA = Triangle.GetVertexIndex(LocalVertexIndex0);
+			const FVec3& VertexA = Triangle.GetVertex(LocalVertexIndex0);
+
+			// @todo(chaos): the map of Vertex->TriangleIndices would help here but may not be a net win
+			for (int32 OtherLocalTriangleIndex = 0; OtherLocalTriangleIndex < Triangles.Num(); ++OtherLocalTriangleIndex)
+			{
+				const FTriangleExt& OtherTriangle = Triangles[OtherLocalTriangleIndex];
+				FVec3 OtherVertexB, OtherVertexC;
+				if (OtherTriangle.GetOtherVerticesFromID(VertexIndexA, OtherVertexB, OtherVertexC))
+				{
+					// Does the contact normal point into the infinite prism formed by extruding the triangle along the face normal?
+					const FVec3& OtherTriangleNormal = OtherTriangle.GetNormal();
+					const FVec3 OtherEdge0 = OtherVertexB - VertexA;
+					const FVec3 OtherEdge1 = VertexA - OtherVertexC;
+					const FVec3 OtherEdgeNormal0 = FVec3::CrossProduct(OtherTriangleNormal, OtherEdge0);	// Not normlized
+					const FVec3 OtherEdgeNormal1 = FVec3::CrossProduct(OtherTriangleNormal, OtherEdge1);	// Not normlized
+					const FReal OtherEdgeSign0 = FVec3::DotProduct(InOutPlaneNormal, OtherEdgeNormal0);
+					const FReal OtherEdgeSign1 = FVec3::DotProduct(InOutPlaneNormal, OtherEdgeNormal1);
+					const FReal OtherEdgeSign0Sq = OtherEdgeSign0 * FMath::Abs(OtherEdgeSign0);
+					const FReal OtherEdgeSign1Sq = OtherEdgeSign1 * FMath::Abs(OtherEdgeSign1);
+					const FReal NormalToleranceSq = FReal(1.e-8);
+					const FReal NormalTolerance0Sq = NormalToleranceSq * OtherEdge0.SizeSquared();
+					const FReal NormalTolerance1Sq = NormalToleranceSq * OtherEdge1.SizeSquared();
+					if ((OtherEdgeSign0Sq >= FReal(-NormalTolerance0Sq)) && (OtherEdgeSign1Sq >= FReal(-NormalTolerance1Sq)))
+					{
+						InOutFeatureType = Private::EConvexFeatureType::Plane;
+						InOutFeatureIndex = 0;
+						InOutPlaneNormal = OtherTriangleNormal;
+						return true;
+					}
+				}
+			}
+		}
+
+		// We have nothing to do for face contacts
+		return false;
+	}
+
+	// Sort contacts on a shape pair in the order we like to solve them.
+	// NOTE: This relies on the enum order of EContactPointType.
+	void FMeshContactGenerator::SortContactsForSolver()
+	{
+		// Sort TriangleContactPoints in solver preferred order, but ignore TriangleContactPointDatas
+		// NOTE: This should only be called at the end of the pruning proxess when we no longer care 
+		// about TriangleContactPointDatas.
+		if (Contacts.Num() > 1)
+		{
+			// Sort contact points by distance from the center of mass (RxN) so that points closer to the center of
+			// mass are solved first. This produces better solver results for low iterations because, if we were to 
+			// solve the distant points first, we would get extra rotation applied. 
+			//
+			// E.g., consider a box landing on an inclined plane with 5 contact points biassed toward one side. 
+			//
+			// -------------------
+			// |                 |
+			// |                 |
+			// *-*-*-*-*----------
+			//
+			// Solving this left to right would result in extra clockwise rotation after 1 iteration. A subsequent
+			// iteration would partially correct the problem.
+			//
+			TArray<TPair<FReal, int32>> SortKeyValues;
+			SortKeyValues.SetNumUninitialized(Contacts.Num());
+			for (int32 ContactIndex = 0; ContactIndex < Contacts.Num(); ++ContactIndex)
+			{
+				const FContactPoint& ContactPoint = Contacts[ContactIndex];
+				const FVec3 DeltaTangent = ContactPoint.ShapeContactPoints[0] - FVec3::DotProduct(ContactPoint.ShapeContactPoints[0], ContactPoint.ShapeContactNormal) * ContactPoint.ShapeContactNormal;
+				const FReal DeltaTangentLenSq = DeltaTangent.SizeSquared();
+				SortKeyValues[ContactIndex] = { DeltaTangentLenSq, ContactIndex };
+			}
+
+			Algo::Sort(SortKeyValues,
+				[](const TPair<FReal, int32>& L, const TPair<FReal, int32>& R)
+				{
+					return L.Key < R.Key;
+				});
+
+			TArray<FContactPoint> SortedContactPoints;
+			SortedContactPoints.SetNumUninitialized(Contacts.Num());
+			for (int32 ContactIndex = 0; ContactIndex < Contacts.Num(); ++ContactIndex)
+			{
+				SortedContactPoints[ContactIndex] = Contacts[SortKeyValues[ContactIndex].Value];
+			}
+			Swap(SortedContactPoints, Contacts);
+		}
+	}
+
 	void FMeshContactGenerator::RemoveDisabledContacts()
 	{
+		// @todo(chaos): don't do this if we have not disabled any contacts
+
 		// Re-pack the contact point array without re-ordering
 		const int32 NumContactPoints = Contacts.Num();
 		int32 DestContactIndex = 0;		// Index to where the next enabled item goes
diff --git a/Engine/Source/Runtime/Experimental/Chaos/Private/Chaos/Collision/PBDCollisionConstraint.cpp b/Engine/Source/Runtime/Experimental/Chaos/Private/Chaos/Collision/PBDCollisionConstraint.cpp
--- a/Engine/Source/Runtime/Experimental/Chaos/Private/Chaos/Collision/PBDCollisionConstraint.cpp
+++ b/Engine/Source/Runtime/Experimental/Chaos/Private/Chaos/Collision/PBDCollisionConstraint.cpp
@@ -205,6 +205,7 @@ namespace Chaos
 		, Simplicial{ nullptr, nullptr }
 		, AccumulatedImpulse(0)
 		, CullDistance(TNumericLimits<FRealSingle>::Max())
+		, RelativeMovement(0)
 		, ContainerCookie()
 		, CollisionTolerance(0)
 		, ClosestManifoldPointIndex(INDEX_NONE)
@@ -245,6 +246,7 @@ namespace Chaos
 		, Simplicial{ Simplicial0, Simplicial1 }
 		, AccumulatedImpulse(0)
 		, CullDistance(TNumericLimits<FRealSingle>::Max())
+		, RelativeMovement(0)
 		, ContainerCookie()
 		, CollisionTolerance(0)
 		, ClosestManifoldPointIndex(INDEX_NONE)
@@ -302,6 +304,7 @@ namespace Chaos
 		ImplicitTransform[1] = InImplicitLocalTransform1;
 
 		CullDistance = FRealSingle(InCullDistance);
+		RelativeMovement = FVec3(0);
 
 		FRealSingle DistanceCheckSize = 0;
 		BoundsTestFlags = Private::CalculateImplicitBoundsTestFlags(
@@ -653,6 +656,36 @@ namespace Chaos
 		ManifoldPoint.ContactPoint.Phi = FRealSingle(FVec3::DotProduct(WorldContact0 - WorldContact1, WorldContactNormal));
 	}
 
+	void FPBDCollisionConstraint::CorrectManifoldPoints()
+	{
+		const FConstGenericParticleHandle P0 = Particle[0];
+		const FConstGenericParticleHandle P1 = Particle[1];
+		const FReal InvM0 = P0->InvM();
+		const FReal InvM1 = P1->InvM();
+		const FReal Alpha0 = ((InvM0 + InvM1) > FReal(0)) ? InvM1 / (InvM0 + InvM1) : FReal(0);
+		const FReal Alpha1 = FReal(1) - Alpha0;
+
+		for (int32 ManifoldPointIndex = 0; ManifoldPointIndex < ManifoldPoints.Num(); ManifoldPointIndex++)
+		{
+			FManifoldPoint& ManifoldPoint = ManifoldPoints[ManifoldPointIndex];
+			const FVec3 WorldContact0 = GetShapeWorldTransform0().TransformPositionNoScale(FVec3(ManifoldPoint.ContactPoint.ShapeContactPoints[0]));
+			const FVec3 WorldContact1 = GetShapeWorldTransform1().TransformPositionNoScale(FVec3(ManifoldPoint.ContactPoint.ShapeContactPoints[1]));
+			const FVec3 WorldContactNormal = GetShapeWorldTransform1().TransformVectorNoScale(FVec3(ManifoldPoint.ContactPoint.ShapeContactNormal));
+			const FVec3 WorldContactDelta = WorldContact1 - WorldContact0;
+			const FReal WorldContactDeltaNormal = FVec3::DotProduct(WorldContactDelta, WorldContactNormal);
+			const FVec3 WorldContactDeltaTangent = WorldContactDelta - WorldContactDeltaNormal * WorldContactNormal;
+
+			const FVec3 CorrectedWorldContact0 = WorldContact0 + Alpha0 * WorldContactDeltaTangent;
+			const FVec3 CorrectedWorldContact1 = WorldContact1 - Alpha1 * WorldContactDeltaTangent;
+			FContactPoint CorrectedContactPoint = ManifoldPoint.ContactPoint;
+			CorrectedContactPoint.ShapeContactPoints[0] = GetShapeWorldTransform0().InverseTransformPositionNoScale(CorrectedWorldContact0);
+			CorrectedContactPoint.ShapeContactPoints[1] = GetShapeWorldTransform1().InverseTransformPositionNoScale(CorrectedWorldContact1);
+			CorrectedContactPoint.Phi = -WorldContactDeltaNormal;
+
+			InitManifoldPoint(ManifoldPointIndex, CorrectedContactPoint);
+		}
+	}
+
 	void FPBDCollisionConstraint::UpdateManifoldContacts()
 	{
 		if ((GetSolverBody0() != nullptr) && (GetSolverBody1() != nullptr))
diff --git a/Engine/Source/Runtime/Experimental/Chaos/Private/Chaos/Collision/ParticlePairMidPhase.cpp b/Engine/Source/Runtime/Experimental/Chaos/Private/Chaos/Collision/ParticlePairMidPhase.cpp
--- a/Engine/Source/Runtime/Experimental/Chaos/Private/Chaos/Collision/ParticlePairMidPhase.cpp
+++ b/Engine/Source/Runtime/Experimental/Chaos/Private/Chaos/Collision/ParticlePairMidPhase.cpp
@@ -63,6 +63,8 @@ namespace Chaos
 		FAutoConsoleVariableRef CVarChaos_Collision_MaxShapePairs(TEXT("p.Chaos.Collision.MaxShapePairs"), Chaos_Collision_MidPhase_MaxShapePairs, TEXT(""));
 
 		extern int32 ChaosOneWayInteractionPairCollisionMode;
+
+		extern bool bChaosUseMACD;
 	}
 
 	using namespace CVars;
@@ -77,11 +79,12 @@ namespace Chaos
 		const FImplicitObject* Implicit0,
 		const FImplicitObject* Implicit1,
 		const FRigidTransform3& ShapeTransform1To0,
+		const FVec3f& LocalRelativeMovement0,
 		const FReal CullDistance)
 	{
 		if (Implicit0->HasBoundingBox() && Implicit1->HasBoundingBox())
 		{
-			const FAABB3 Box1In0 = Implicit1->CalculateTransformedBounds(ShapeTransform1To0).Thicken(CullDistance);
+			const FAABB3 Box1In0 = Implicit1->CalculateTransformedBounds(ShapeTransform1To0).GrowByVector(LocalRelativeMovement0).Thicken(CullDistance);
 			const FAABB3 Box0 = Implicit0->BoundingBox();
 			return Box0.Intersects(Box1In0);
 		}
@@ -93,10 +96,12 @@ namespace Chaos
 		const FImplicitObject* Implicit1,
 		const FRigidTransform3& ShapeWorldTransform0,
 		const FRigidTransform3& ShapeWorldTransform1,
+		const FVec3f& RelativeMovement,
 		const FReal CullDistance)
 	{
 		const FRigidTransform3 ShapeTransform1To0 = ShapeWorldTransform1.GetRelativeTransform(ShapeWorldTransform0);
-		return ImplicitOverlapOBBToAABB(Implicit0, Implicit1, ShapeTransform1To0, CullDistance);
+		const FVec3 LocalRelativeMovement0 = ShapeWorldTransform0.InverseTransformVectorNoScale(FVec3(RelativeMovement));
+		return ImplicitOverlapOBBToAABB(Implicit0, Implicit1, ShapeTransform1To0, LocalRelativeMovement0, CullDistance);
 	}
 
 	// Get the number of leaf objects in the implicit hierarchy, and set a flag if this hierarchy is a tree
@@ -321,7 +326,10 @@ namespace Chaos
 	{
 	}
 
-	bool FSingleShapePairCollisionDetector::DoBoundsOverlap(const FReal CullDistance, const int32 CurrentEpoch)
+	bool FSingleShapePairCollisionDetector::DoBoundsOverlap(
+		const FRealSingle CullDistance, 
+		const FVec3f& RelativeMovement, 
+		const int32 CurrentEpoch)
 	{
 		PHYSICS_CSV_SCOPED_EXPENSIVE(PhysicsVerbose, NarrowPhase_ShapeBounds);
 
@@ -331,8 +339,7 @@ namespace Chaos
 		// World-space expanded bounds check
 		if (BoundsTestFlags.bEnableAABBCheck)
 		{
-			// @todo(chaos): ideally this is a swept bounds test with a smaller cull distance (see FParticlePairMidPhase::GenerateCollisions)
-			const FAABB3 ExpandedShapeWorldBounds0 = FAABB3(ShapeWorldBounds0).ThickenSymmetrically(FVec3(CullDistance));
+			const FAABB3 ExpandedShapeWorldBounds0 = FAABB3(ShapeWorldBounds0).GrowByVector(-RelativeMovement).Thicken(CullDistance);
 			if (!ExpandedShapeWorldBounds0.Intersects(ShapeWorldBounds1))
 			{
 				return false;
@@ -368,7 +375,7 @@ namespace Chaos
 
 			if (BoundsTestFlags.bEnableOBBCheck0)
 			{
-				if (!ImplicitOverlapOBBToAABB(Implicit0, Implicit1, ShapeWorldTransform0, ShapeWorldTransform1, CullDistance))
+				if (!ImplicitOverlapOBBToAABB(Implicit0, Implicit1, ShapeWorldTransform0, ShapeWorldTransform1, RelativeMovement, CullDistance))
 				{
 					return false;
 				}
@@ -376,7 +383,7 @@ namespace Chaos
 
 			if (BoundsTestFlags.bEnableOBBCheck1)
 			{
-				if (!ImplicitOverlapOBBToAABB(Implicit1, Implicit0, ShapeWorldTransform1, ShapeWorldTransform0, CullDistance))
+				if (!ImplicitOverlapOBBToAABB(Implicit1, Implicit0, ShapeWorldTransform1, ShapeWorldTransform0, -RelativeMovement, CullDistance))
 				{
 					return false;
 				}
@@ -387,29 +394,31 @@ namespace Chaos
 	}
 
 	int32 FSingleShapePairCollisionDetector::GenerateCollision(
-		const FReal CullDistance,
-		const FReal Dt,
+		const FRealSingle Dt,
+		const FRealSingle CullDistance,
+		const FVec3f& RelativeMovement,
 		const FCollisionContext& Context)
 	{
 		CHAOS_MIDPHASE_SCOPE_CYCLE_TIMER(FSingleShapePairCollisionDetector_GenerateCollision);
 
 		const int32 CurrentEpoch = Context.GetAllocator()->GetCurrentEpoch();
-		if (DoBoundsOverlap(CullDistance, CurrentEpoch))
+		if (DoBoundsOverlap(CullDistance, RelativeMovement, CurrentEpoch))
 		{
-			return GenerateCollisionImpl(CullDistance, Dt, Context);
+			return GenerateCollisionImpl(Dt, CullDistance, RelativeMovement, Context);
 		}
 		return 0;
 	}
 
 	int32 FSingleShapePairCollisionDetector::GenerateCollisionCCD(
+		const FRealSingle Dt,
+		const FRealSingle CullDistance,
+		const FVec3f& RelativeMovement,
 		const bool bEnableCCDSweep,
-		const FReal CullDistance,
-		const FReal Dt,
 		const FCollisionContext& Context)
 	{
 		CHAOS_MIDPHASE_SCOPE_CYCLE_TIMER(FSingleShapePairCollisionDetector_GenerateCollisionCCD);
 
-		return GenerateCollisionCCDImpl(bEnableCCDSweep, CullDistance, Dt, Context);
+		return GenerateCollisionCCDImpl(Dt, CullDistance, RelativeMovement, bEnableCCDSweep, Context);
 	}
 
 	void FSingleShapePairCollisionDetector::CreateConstraint(const FReal CullDistance, const FCollisionContext& Context)
@@ -430,15 +439,16 @@ namespace Chaos
 	}
 
 	int32 FSingleShapePairCollisionDetector::GenerateCollisionImpl(
-		const FReal CullDistance, 
-		const FReal Dt,
+		const FRealSingle Dt,
+		const FRealSingle CullDistance,
+		const FVec3f& RelativeMovement,
 		const FCollisionContext& Context)
 	{
 		CHAOS_MIDPHASE_SCOPE_CYCLE_TIMER(FSingleShapePairCollisionDetector_GenerateCollisionImpl);
 
 		if (BoundsTestFlags.bIsProbe)
 		{
-			return GenerateCollisionProbeImpl(CullDistance, Dt, Context);
+			return GenerateCollisionProbeImpl(Dt, CullDistance, RelativeMovement, Context);
 		}
 
 		if (!Constraint.IsValid())
@@ -465,6 +475,7 @@ namespace Chaos
 			Constraint->SetShapeWorldTransforms(ShapeWorldTransform0, ShapeWorldTransform1);
 
 			Constraint->SetCullDistance(CullDistance);
+			Constraint->SetRelativeMovement(RelativeMovement);
 
 			// Constraint may have been previously used with CCD enabled (e.g., a midphase modifier)
 			// so we need to make sure that the CCD flag is disabled
@@ -525,14 +536,15 @@ namespace Chaos
 	}
 
 	int32 FSingleShapePairCollisionDetector::GenerateCollisionCCDImpl(
+		const FRealSingle Dt,
+		const FRealSingle CullDistance,
+		const FVec3f& RelativeMovement,
 		const bool bEnableCCDSweep,
-		const FReal CullDistance, 
-		const FReal Dt,
 		const FCollisionContext& Context)
 	{
 		if (BoundsTestFlags.bIsProbe)
 		{
-			return GenerateCollisionProbeImpl(CullDistance, Dt, Context);
+			return GenerateCollisionProbeImpl(Dt, CullDistance, RelativeMovement, Context);
 		}
 
 		if (!Constraint.IsValid())
@@ -552,7 +564,7 @@ namespace Chaos
 
 		if (!bEnableCCDSweep)
 		{
-			return GenerateCollision(CullDistance, Dt, Context);
+			return GenerateCollision(Dt, CullDistance, RelativeMovement, Context);
 		}
 
 		// Swept collision detection
@@ -639,8 +651,9 @@ namespace Chaos
 	}
 
 	int32 FSingleShapePairCollisionDetector::GenerateCollisionProbeImpl(
-		const FReal CullDistance, 
-		const FReal Dt,
+		const FRealSingle Dt,
+		const FRealSingle CullDistance,
+		const FVec3f& RelativeMovement,
 		const FCollisionContext& Context)
 	{
 		// Same as regular constraint generation, but always defer narrow phase.
@@ -820,6 +833,7 @@ namespace Chaos
 		Flags.bIsActive = true;
 		Flags.bIsCCD = false;
 		Flags.bIsCCDActive = false;
+		Flags.bIsMACD = false;
 		Flags.bIsSleeping = false;
 		Flags.bIsModified = false;
 		Flags.bIsConvexOptimizationActive = true;
@@ -857,12 +871,16 @@ namespace Chaos
 		//
 		// bIsCCDActive is reset to bIsCCD each frame, but can be overridden
 		// by modifiers.
-		const bool bIsCCD = Context.GetSettings().bAllowCCD && (
-			FConstGenericParticleHandle(Particle0)->CCDEnabled() ||
-			FConstGenericParticleHandle(Particle1)->CCDEnabled());
+		const bool bIsCCD = Context.GetSettings().bAllowCCD && (FConstGenericParticleHandle(Particle0)->CCDEnabled() || FConstGenericParticleHandle(Particle1)->CCDEnabled());
 		Flags.bIsCCD = bIsCCD;
 		Flags.bIsCCDActive = bIsCCD;
 
+		// @todo(chaos): we should only enable relative movement mode if the particle asks for it
+		// (this also affects how the particle bounds is expanded in the broadphase)
+		//const bool bIsMACD = Context.GetSettings().bAllowMACD && (FConstGenericParticleHandle(Particle0)->MACDEnabled() || FConstGenericParticleHandle(Particle1)->MACDEnabled());
+		const bool bIsMACD = CVars::bChaosUseMACD;
+		Flags.bIsMACD = bIsMACD;
+
 		// Initially we allow for convex optimization where available
 		Flags.bIsConvexOptimizationActive = true;
 
@@ -918,7 +936,7 @@ namespace Chaos
 
 	void FParticlePairMidPhase::GenerateCollisions(
 		const FReal InCullDistance,
-		const FReal Dt,
+		const FReal InDt,
 		const FCollisionContext& Context)
 	{
 		if (!IsValid())
@@ -930,34 +948,17 @@ namespace Chaos
 
 		if (Flags.bIsActive)
 		{
+			FConstGenericParticleHandle P0 = GetParticle0();
+			FConstGenericParticleHandle P1 = GetParticle1();
+
+			FRealSingle Dt = FRealSingle(InDt);
+
 			// CullDistance is scaled by the size of the dynamic objects.
-			FReal CullDistance = InCullDistance * CullDistanceScale;
+			FRealSingle CullDistance = FRealSingle(InCullDistance) * CullDistanceScale;
 
 			// If CCD is enabled, did we move far enough to require a sweep?
 			Flags.bUseSweep = (Flags.bIsCCDActive != 0) && ShouldEnableCCDSweep(Dt);
 
-			// We increase CullDistance based on velocity (up to a limit for perf with large velocities).
-			// NOTE: This somewhat matches the bounds expansion in FPBDRigidsEvolutionGBF::Integrate
-			// NOTE: We use PreV here which is the velocity after collisions from the previous tick because we want the 
-			// velocity without gravity from this tick applied. This is mainly so that we get the same CullDistance from 
-			// one tick to the next, even if one of the particles goes to sleep, and therefore its velocity is now zero 
-			// because gravity is no longer applied. Also see FPBDIslandManager::PropagateIslandSleep for other issues 
-			// related to velocity and sleeping...
-			// NOTE: we used to extend the cull distance for CCD objects, but this is no longer required. The sweep and
-			// rewind phase does not use CullDistance, and once rewound we are using normal collision detection where
-			// an expanded CullDistance doesn't help and makes perf worse.
-			FConstGenericParticleHandle P0 = GetParticle0();
-			FConstGenericParticleHandle P1 = GetParticle1();
-			const FReal VMax0 = P0->PreV().GetAbsMax();
-			const FReal VMax1 = P1->PreV().GetAbsMax();
-			const FReal VMaxDt = FMath::Max(VMax0, VMax1) * Dt;
-			const FReal VelocityBoundsMultiplier = Context.GetSettings().BoundsVelocityInflation;
-			const FReal MaxVelocityBoundsExpansion = Context.GetSettings().MaxVelocityBoundsExpansion;
-			if ((VelocityBoundsMultiplier > 0) && (MaxVelocityBoundsExpansion > 0))
-			{
-				CullDistance += FMath::Min(VelocityBoundsMultiplier * VMaxDt, MaxVelocityBoundsExpansion);
-			}
-
 #if !UE_BUILD_TEST && !UE_BUILD_SHIPPING
 			// At least one body must be dynamic, and at least one must be awake (or moving if a kinematic)
 			const bool bIsMoving0 = (P0->IsDynamic() && !P0->IsSleeping()) || P0->IsMovingKinematic();
@@ -966,8 +967,40 @@ namespace Chaos
 			ensureMsgf(bAnyMoving, TEXT("GenerateCollisions called on two stationary objects %s %s"), *P0->GetDebugName(), *P1->GetDebugName());
 #endif
 
+			FVec3f RelativeMovement = FVec3f(0);
+			if (!Flags.bIsMACD)
+			{
+				// We increase CullDistance based on velocity (up to a limit for perf with large velocities).
+				// NOTE: This somewhat matches the bounds expansion in FPBDRigidsEvolutionGBF::Integrate
+				// NOTE: We use PreV here which is the velocity after collisions from the previous tick because we want the 
+				// velocity without gravity from this tick applied. This is mainly so that we get the same CullDistance from 
+				// one tick to the next, even if one of the particles goes to sleep, and therefore its velocity is now zero 
+				// because gravity is no longer applied. Also see FPBDIslandManager::PropagateIslandSleep for other issues 
+				// related to velocity and sleeping...
+				// NOTE: we used to extend the cull distance for CCD objects, but this is no longer required. The sweep and
+				// rewind phase does not use CullDistance, and once rewound we are using normal collision detection where
+				// an expanded CullDistance doesn't help and makes perf worse.
+				const FRealSingle VMax0 = P0->GetPreVf().GetAbsMax();
+				const FRealSingle VMax1 = P1->GetPreVf().GetAbsMax();
+				const FRealSingle VMaxDt = FMath::Max(VMax0, VMax1) * Dt;
+				const FRealSingle VelocityBoundsMultiplier = FRealSingle(Context.GetSettings().BoundsVelocityInflation);
+				const FRealSingle MaxVelocityBoundsExpansion = FRealSingle(Context.GetSettings().MaxVelocityBoundsExpansion);
+				if ((VelocityBoundsMultiplier > 0) && (MaxVelocityBoundsExpansion > 0))
+				{
+					CullDistance += FMath::Min(VelocityBoundsMultiplier * VMaxDt, MaxVelocityBoundsExpansion);
+				}
+			}
+			else
+			{
+				// Movement-aware collision detection (MACD) takes the relative position change this tick as input
+				// We do not expand CullDistance. Depending on the shape pair types involved, the collision detection
+				// step will either pad the CullDistance itself, or ideally compare RelativeMovement with the contact 
+				// normal when determining what contacts to cull. Eventually all shape pairs should do the latter.
+				RelativeMovement = (P0->GetVf() - P1->GetVf()) * Dt;
+			}
+
 			// Run collision detection on all potentially colliding shape pairs
-			NumActiveConstraints = GenerateCollisionsImpl(CullDistance, Dt, Context);
+			NumActiveConstraints = GenerateCollisionsImpl(Dt, CullDistance, RelativeMovement, Context);
 		}
 
 		// Reset any modifications applied by the MidPhaseModifier.
@@ -1108,8 +1141,9 @@ namespace Chaos
 	}
 
 	int32 FShapePairParticlePairMidPhase::GenerateCollisionsImpl(
-		const FReal CullDistance,
-		const FReal Dt,
+		const FRealSingle Dt,
+		const FRealSingle CullDistance,
+		const FVec3f& RelativeMovement,
 		const FCollisionContext& Context)
 	{
 		CHAOS_MIDPHASE_SCOPE_CYCLE_TIMER(FShapePairParticlePairMidPhase_GenerateCollision);
@@ -1121,14 +1155,14 @@ namespace Chaos
 		{
 			for (FSingleShapePairCollisionDetector& ShapePair : ShapePairDetectors)
 			{
-				NumActive += ShapePair.GenerateCollisionCCD(Flags.bUseSweep, CullDistance, Dt, Context);
+				NumActive += ShapePair.GenerateCollisionCCD(Dt, CullDistance, RelativeMovement, Flags.bUseSweep, Context);
 			}
 		}
 		else
 		{
 			for (FSingleShapePairCollisionDetector& ShapePair : ShapePairDetectors)
 			{
-				NumActive += ShapePair.GenerateCollision(CullDistance, Dt, Context);
+				NumActive += ShapePair.GenerateCollision(Dt, CullDistance, RelativeMovement, Context);
 			}
 		}
 		return NumActive;
@@ -1187,8 +1221,9 @@ namespace Chaos
 	}
 
 	int32 FGenericParticlePairMidPhase::GenerateCollisionsImpl(
-		const FReal CullDistance,
-		const FReal Dt,
+		const FRealSingle Dt,
+		const FRealSingle CullDistance,
+		const FVec3f& RelativeMovement,
 		const FCollisionContext& Context)
 	{
 		CHAOS_MIDPHASE_SCOPE_CYCLE_TIMER(FGenericParticlePairMidPhase_GenerateCollisionImpl);
@@ -1378,7 +1413,7 @@ namespace Chaos
 				if (!FilterHasSimEnabled(ShapeInstanceA)) return;
 
 				const FAABB3 RelativeBoundsA = ImplicitA->CalculateTransformedBounds(RelativeTransformA);
-				const FAABB3 ShapeBoundsAInB = RelativeBoundsA.TransformedAABB(ParticleTransformAToB).ThickenSymmetrically(FVec3(CullDistance));
+				const FAABB3 ShapeBoundsAInB = RelativeBoundsA.TransformedAABB(ParticleTransformAToB).Thicken(CullDistance);
 
 				// Detect collisions between ImplicitA and Implicit Hierarchy of ParticleB
 				VisitOverlappingObjects(ConvexOptimizerB, RootImplicitB, ShapeBoundsAInB,
@@ -1418,7 +1453,7 @@ namespace Chaos
 		// ImplicitB transforms/bounds (expanded by cull distance)
 		const FRigidTransform3 ImplicitTransformB = RelativeTransformB * ParticleWorldTransformB;
 		const FRigidTransform3 ImplicitTransformBToA = ImplicitTransformB.GetRelativeTransform(ParticleWorldTransformA);
-		const FAABB3 ImplicitBoundsBInA = ImplicitB->CalculateTransformedBounds(ImplicitTransformBToA).ThickenSymmetrically(FVec3(CullDistance));
+		const FAABB3 ImplicitBoundsBInA = ImplicitB->CalculateTransformedBounds(ImplicitTransformBToA).Thicken(CullDistance);
 
 		// If ImplicitB has a built-in BVH (Heightfield or TriMesh) we handle the test against the BVH differently
 		const bool bHasInternalBVHB = ImplicitB->template IsA<FHeightField>();
@@ -1578,7 +1613,7 @@ namespace Chaos
 		const FRigidTransform3 ShapeWorldTransformB = ShapeRelativeTransformB * ParticleWorldTransformB;
 
 		// NOTE: only expand one bounds by cull distance
-		const FAABB3 ShapeWorldBoundsA = ImplicitA->CalculateTransformedBounds(ShapeWorldTransformA).ThickenSymmetrically(FVec3(CullDistance));
+		const FAABB3 ShapeWorldBoundsA = ImplicitA->CalculateTransformedBounds(ShapeWorldTransformA).Thicken(CullDistance);
 		const FAABB3 ShapeWorldBoundsB = ImplicitB->CalculateTransformedBounds(ShapeWorldTransformB);
 
 		// World-space expanded bounds check
@@ -1602,6 +1637,8 @@ namespace Chaos
 			}
 		}
 
+		FVec3f RelativeMovement = FVec3f(0);
+
 		if (BoundsTestFlags.bEnableOBBCheck0 || BoundsTestFlags.bEnableOBBCheck1)
 		{
 			// OBB-AABB test in both directions. This is beneficial for shapes which do not fit their AABBs very well,
@@ -1609,7 +1646,7 @@ namespace Chaos
 			// one shape is long and thin (i.e., it does not fit an AABB well when the shape is rotated).
 			if (BoundsTestFlags.bEnableOBBCheck0)
 			{
-				if (!ImplicitOverlapOBBToAABB(ImplicitA, ImplicitB, ShapeWorldTransformA, ShapeWorldTransformB, CullDistance))
+				if (!ImplicitOverlapOBBToAABB(ImplicitA, ImplicitB, ShapeWorldTransformA, ShapeWorldTransformB, RelativeMovement, CullDistance))
 				{
 					return false;
 				}
@@ -1617,7 +1654,7 @@ namespace Chaos
 
 			if (BoundsTestFlags.bEnableOBBCheck1)
 			{
-				if (!ImplicitOverlapOBBToAABB(ImplicitB, ImplicitA, ShapeWorldTransformB, ShapeWorldTransformA, CullDistance))
+				if (!ImplicitOverlapOBBToAABB(ImplicitB, ImplicitA, ShapeWorldTransformB, ShapeWorldTransformA, RelativeMovement, CullDistance))
 				{
 					return false;
 				}
@@ -2087,8 +2124,9 @@ namespace Chaos
 	}
 
 	int32 FSphereApproximationParticlePairMidPhase::GenerateCollisionsImpl(
-		const FReal CullDistance,
-		const FReal Dt,
+		const FRealSingle Dt,
+		const FRealSingle CullDistance,
+		const FVec3f& RelativeMovement,
 		const FCollisionContext& Context)
 	{
 		if (!bHasSpheres)
diff --git a/Engine/Source/Runtime/Experimental/Chaos/Private/Chaos/CollisionOneShotManifolds.cpp b/Engine/Source/Runtime/Experimental/Chaos/Private/Chaos/CollisionOneShotManifolds.cpp
--- a/Engine/Source/Runtime/Experimental/Chaos/Private/Chaos/CollisionOneShotManifolds.cpp
+++ b/Engine/Source/Runtime/Experimental/Chaos/Private/Chaos/CollisionOneShotManifolds.cpp
@@ -102,7 +102,9 @@ namespace Chaos
 
 	// @todo(chaos): Temp while we test the new convex-mesh collision optimizations
 	bool bChaos_Collision_EnableMeshManifoldOptimizedLoop = true;
+	bool bChaos_Collision_EnableMeshManifoldOptimizedLoop_TriMesh = true;
 	FAutoConsoleVariableRef CVarChaos_Collision_EnableMeshManifoldOptimizedLoop(TEXT("p.Chaos.Collision.EnableMeshManifoldOptimizedLoop"), bChaos_Collision_EnableMeshManifoldOptimizedLoop, TEXT(""));
+	FAutoConsoleVariableRef CVarChaos_Collision_EnableMeshManifoldOptimizedLoop_TriMesh(TEXT("p.Chaos.Collision.EnableMeshManifoldOptimizedLoopTriMesh"), bChaos_Collision_EnableMeshManifoldOptimizedLoop_TriMesh, TEXT(""));
 
 	// Whether to use the new index-less GJK. 
 	// @todo(chaos): This should be removed once soaked for a bit (enabled 7 June 2022)
diff --git a/Engine/Source/Runtime/Experimental/Chaos/Private/Chaos/CollisionOneShotManifoldsMeshShapes.cpp b/Engine/Source/Runtime/Experimental/Chaos/Private/Chaos/CollisionOneShotManifoldsMeshShapes.cpp
--- a/Engine/Source/Runtime/Experimental/Chaos/Private/Chaos/CollisionOneShotManifoldsMeshShapes.cpp
+++ b/Engine/Source/Runtime/Experimental/Chaos/Private/Chaos/CollisionOneShotManifoldsMeshShapes.cpp
@@ -1,16 +1,20 @@
 // Copyright Epic Games, Inc. All Rights Reserved.
 #include "Chaos/CollisionOneShotManifoldsMeshShapes.h"
 
+#include "Chaos/ChaosDebugDraw.h"
 #include "Chaos/CollisionOneShotManifolds.h"
 #include "Chaos/CollisionResolution.h"
 #include "Chaos/Collision/CapsuleTriangleContactPoint.h"
 #include "Chaos/Collision/ContactPointsMiscShapes.h"
 #include "Chaos/Collision/ContactTriangles.h"
+#include "Chaos/Collision/ConvexContactPointUtilities.h"
+#include "Chaos/Collision/ConvexFeature.h"
 #include "Chaos/Collision/ConvexTriangleContactPoint.h"
 #include "Chaos/Collision/MeshContactGenerator.h"
 #include "Chaos/Collision/PBDCollisionConstraint.h"
 #include "Chaos/Collision/SphereTriangleContactPoint.h"
 #include "Chaos/Convex.h"
+#include "Chaos/DebugDrawQueue.h"
 #include "Chaos/Defines.h"
 #include "Chaos/Framework/UncheckedArray.h"
 #include "Chaos/GJK.h"
@@ -31,9 +35,18 @@ namespace Chaos
 	extern FRealSingle Chaos_Collision_TriMeshPhiToleranceScale;
 	extern int32 Chaos_Collision_MeshManifoldHashSize;
 	extern bool bChaos_Collision_EnableMeshManifoldOptimizedLoop;
+	extern bool bChaos_Collision_EnableMeshManifoldOptimizedLoop_TriMesh;
 
 	extern bool bChaos_Collision_UseCapsuleTriMesh2;
 	extern bool bChaos_Collision_UseConvexTriMesh2;
+	
+	namespace CVars
+	{
+#if CHAOS_DEBUG_DRAW
+		extern DebugDraw::FChaosDebugDrawSettings ChaosSolverDebugDebugDrawSettings;
+		extern int32 ChaosSolverDebugDrawMeshContacts;
+#endif
+	}
 
 	namespace Collisions
 	{
@@ -123,12 +136,17 @@ namespace Chaos
 
 			// Generate the contact manifold between Convex and a Triangle
 			const auto& GenerateConvexTriangleContacts =
-				[&Convex, &MeshToConvexTransform, CullDistance](const FTriangle& Triangle, FContactPointManifold& OutContactPoints)
-				{
-					GenerateConvexTriangleOneShotManifold(Convex, Triangle, CullDistance, OutContactPoints);
-				};
+				[&Convex, CullDistance](Private::FMeshContactGenerator& ContactGenerator, const int32 TriangleIndex)
+			{
+				const FTriangle& Triangle = ContactGenerator.GetTriangle(TriangleIndex);
+
+				FContactPointManifold Contacts;
+				GenerateConvexTriangleOneShotManifold(Convex, Triangle, CullDistance, Contacts);
 
-			// Collect all the triangles that overlap our convex, transformed into Convex space
+				ContactGenerator.AddTriangleContacts(TriangleIndex, MakeArrayView(Contacts));
+			};
+
+			// Collect all the triangles that overlap our convex. Triangles will be in Convex space
 			Mesh.CollectTriangles(MeshQueryBounds, MeshToConvexTransform, ConvexBounds, ContactGenerator);
 
 			// Generate a set of contact points for all triangles
@@ -138,6 +156,129 @@ namespace Chaos
 			ContactGenerator.ProcessGeneratedContacts(ConvexTransform, MeshToConvexTransform);
 		}
 
+		// MACD: Motion-Aware Collision Detection
+		template<typename ConvexType>
+		void GenerateConvexTriangleOneShotManifoldMACD(const ConvexType& Convex, const FRigidTransform3& ConvexTransform, const FVec3& ConvexRelativeMovement, Private::FMeshContactGenerator& ContactGenerator, const int32 TriangleIndex, const FReal CullDistance, FContactPointLargeManifold& OutContactPoints)
+		{
+			const FTriangle& Triangle = ContactGenerator.GetTriangle(TriangleIndex);
+
+			// Find the closest feature pair on the triangle and convex
+			Private::FConvexContactPoint ClosestContact;
+			if (Private::FindClosestFeatures(Convex, ConvexTransform, Triangle, ConvexRelativeMovement, CullDistance, ClosestContact))
+			{
+				ClosestContact.Features[0].ObjectIndex = 0;
+				ClosestContact.Features[1].ObjectIndex = TriangleIndex;
+
+				// Generate a manifold based on the closest features
+				// NOTE: normal points from triangle to convex
+				const FReal ConvexMotionAlongNormal = FVec3::DotProduct(ConvexRelativeMovement, ClosestContact.ShapeContactNormal);
+				const FReal CullDistancePadding = FMath::Max(0, -ConvexMotionAlongNormal);
+				const FReal NetCullDistance = CullDistance + CullDistancePadding;
+				Private::ConvexTriangleManifoldFromContact(Convex, Triangle, ClosestContact, NetCullDistance, OutContactPoints);
+
+				if (OutContactPoints.Num() > 0)
+				{
+#if CHAOS_DEBUG_DRAW
+					if (CVars::ChaosSolverDebugDrawMeshContacts && FDebugDrawQueue::GetInstance().IsDebugDrawingEnabled())
+					{
+						for (int32 ContactIndex = 0; ContactIndex < OutContactPoints.Num(); ++ContactIndex)
+						{
+							FContactPoint& ContactPoint = OutContactPoints[ContactIndex];
+							const FVec3 P = ConvexTransform.TransformPositionNoScale(ContactPoint.ShapeContactPoints[1]);
+							const FVec3 N = ConvexTransform.TransformVectorNoScale(ContactPoint.ShapeContactNormal);
+							FColor Color = FColor::Black;
+							switch (ClosestContact.Features[1].FeatureType)
+							{
+							case Private::EConvexFeatureType::Plane:
+								Color = FColor::White;
+								break;
+							case Private::EConvexFeatureType::Edge:
+								Color = FColor::Cyan;
+								break;
+							case Private::EConvexFeatureType::Vertex:
+								Color = FColor::Magenta;
+								break;
+							}
+							FDebugDrawQueue::GetInstance().DrawDebugLine(P, P + 10.0f * N, Color, false, CVars::ChaosSolverDebugDebugDrawSettings.DrawDuration, (uint8)CVars::ChaosSolverDebugDebugDrawSettings.DrawPriority, 1.25f * CVars::ChaosSolverDebugDebugDrawSettings.LineThickness);
+						}
+					}
+#endif
+
+					// Adjust the feature if it is invalid. E.g., if we collide with a triangle edge and the normal is outside the range allowed by
+					// the triangles sharing the edge we will project the normal into the valid range
+					const bool bFeatureChanged = ContactGenerator.FixFeature(TriangleIndex, ClosestContact.Features[1].FeatureType, ClosestContact.Features[1].PlaneFeatureIndex, ClosestContact.ShapeContactNormal, ClosestContact.ShapeContactPoints[1]);
+					if (bFeatureChanged)
+					{
+						for (int32 ContactIndex = 0; ContactIndex < OutContactPoints.Num(); ++ContactIndex)
+						{
+							FContactPoint& ContactPoint = OutContactPoints[ContactIndex];
+
+							// Update the normal and recalculate the separation
+							ContactPoint.ShapeContactNormal = ClosestContact.ShapeContactNormal;
+							ContactPoint.Phi = FVec3::DotProduct(ContactPoint.ShapeContactPoints[0] - ContactPoint.ShapeContactPoints[1], ContactPoint.ShapeContactNormal);
+
+							// Remap the triangle contact onto the new plane, keeping the contact point on the convex shape where it is.
+							// NOTE: This means that the triangle contact point may be outside the triangle, but for contact separation
+							// we really only care about the contact plane. Howeber, the contact positions are used for static friction, 
+							// which assumes the contacts have zero separation perpendicular to the normal on the frame they are generated.
+							ContactPoint.ShapeContactPoints[1] = ContactPoint.ShapeContactPoints[0] - (ContactPoint.Phi * ContactPoint.ShapeContactNormal);
+						}
+					}
+
+					// Correct the contact points based on convex movement
+					for (int32 ContactIndex = 0; ContactIndex < OutContactPoints.Num(); ++ContactIndex)
+					{
+						FContactPoint& ContactPoint = OutContactPoints[ContactIndex];
+						Private::ApplyContactPointShift(ContactPoint, ConvexRelativeMovement);
+					}
+				}
+			}
+		}
+
+		// MACD: Motion-Aware Collision Detection
+		template<typename ConvexType, typename MeshType>
+		void ConstructConvexMeshOneShotManifoldMACD(const ConvexType& Convex, const FRigidTransform3& ConvexTransform, const MeshType& Mesh, const FRigidTransform3& MeshTransform, const FVec3& MeshScale, const FVec3& RelativeMovement, const FReal InCullDistance, Private::FMeshContactGenerator& ContactGenerator)
+		{
+			if (RelativeMovement.IsZero())
+			{
+				ConstructConvexMeshOneShotManifold2(Convex, ConvexTransform, Mesh, MeshTransform, MeshScale, InCullDistance, ContactGenerator);
+			}
+			else
+			{
+				// We detect collisions as if both shapes were in the pre-movement locations
+				// GenerateConvexTriangleOneShotManifoldWithMotion will correct the contacts to the actual locations
+				const FRigidTransform3 ShiftedConvexTransform = FRigidTransform3(ConvexTransform.GetTranslation() - RelativeMovement, ConvexTransform.GetRotation());
+				FRigidTransform3 ShiftedMeshToConvexTransform = MeshTransform.GetRelativeTransformNoScale(ShiftedConvexTransform);
+				ShiftedMeshToConvexTransform.SetScale3D(MeshScale);
+
+				// @todo(chaos): add Convex.CalculateInverseTransformed bounds with scale support (to optimize sphere and capsule)
+				const FVec3 ConvexRelativeMovement = ConvexTransform.InverseTransformVectorNoScale(RelativeMovement);
+				const FAABB3 ConvexBounds = FAABB3(Convex.BoundingBox()).GrowByVector(ConvexRelativeMovement).Thicken(InCullDistance);
+				const FAABB3 MeshQueryBounds = ConvexBounds.InverseTransformedAABB(ShiftedMeshToConvexTransform);
+				const FReal CullDistance = InCullDistance;
+
+				// Collect all the triangles that overlap our convex. Triangles will be in Convex space
+				Mesh.CollectTriangles(MeshQueryBounds, ShiftedMeshToConvexTransform, ConvexBounds, ContactGenerator);
+
+				FContactPointLargeManifold Contacts;
+				Contacts.Reserve(6);
+
+				const auto& GenerateConvexTriangleContacts =
+					[&Convex, &ShiftedConvexTransform, &ConvexRelativeMovement, CullDistance, &Contacts](Private::FMeshContactGenerator& ContactGenerator, const int32 TriangleIndex)
+				{
+					Contacts.Reset();
+					GenerateConvexTriangleOneShotManifoldMACD(Convex, ShiftedConvexTransform, ConvexRelativeMovement, ContactGenerator, TriangleIndex, CullDistance, Contacts);
+
+					ContactGenerator.AddTriangleContacts(TriangleIndex, MakeArrayView(Contacts));
+				};
+
+				ContactGenerator.GenerateMeshContacts(GenerateConvexTriangleContacts);
+
+				// Process the contacts to minimize manifold etc
+				ContactGenerator.ProcessGeneratedContacts(ShiftedConvexTransform, ShiftedMeshToConvexTransform);
+			}
+		}
+
 		/**
 		* @brief Create a minimized set of contact points between a convex polyhedron (box, convex) and a non-convex mesh (trimesh, heightfield)
 		*
@@ -233,7 +374,7 @@ namespace Chaos
 			const FReal CullDistance = Constraint.GetCullDistance();
 			const FReal PhiTolerance = CalculateTriMeshPhiTolerance(CullDistance);
 			const FReal DistanceTolerance = Chaos_Collision_TriMeshDistanceTolerance;
-			
+
 			if (bChaos_Collision_EnableMeshManifoldOptimizedLoop)
 			{
 				// New version uses a two-pass loop over triangles to avoid visiting triangles whose vertices are all colliding as a result of checking adjacent triangles
@@ -298,30 +439,64 @@ namespace Chaos
 			const FReal CullDistance = Constraint.GetCullDistance();
 			const FReal PhiTolerance = CalculateTriMeshPhiTolerance(CullDistance);
 			const FReal DistanceTolerance = Chaos_Collision_TriMeshDistanceTolerance;
-			FContactTriangleCollector MeshContacts(bChaos_Collision_OneSidedTriangleMesh, PhiTolerance, DistanceTolerance, ConvexTransform);
 
-			if (const FImplicitBox3* RawBox = Convex.template GetObject<FImplicitBox3>())
-			{
-				ConstructConvexMeshOneShotManifold(*RawBox, ConvexTransform, *Mesh, MeshTransform, MeshScale, CullDistance, MeshContacts);
-			}
-			else if (const TImplicitObjectScaled<FImplicitConvex3>* ScaledConvex = Convex.template GetObject<TImplicitObjectScaled<FImplicitConvex3>>())
-			{
-				ConstructConvexMeshOneShotManifold(*ScaledConvex, ConvexTransform, *Mesh, MeshTransform, MeshScale, CullDistance, MeshContacts);
-			}
-			else if (const TImplicitObjectInstanced<FImplicitConvex3>* InstancedConvex = Convex.template GetObject<TImplicitObjectInstanced<FImplicitConvex3>>())
-			{
-				ConstructConvexMeshOneShotManifold(*InstancedConvex, ConvexTransform, *Mesh, MeshTransform, MeshScale, CullDistance, MeshContacts);
-			}
-			else if (const FImplicitConvex3* RawConvex = Convex.template GetObject<FImplicitConvex3>())
+			if (bChaos_Collision_EnableMeshManifoldOptimizedLoop_TriMesh)
 			{
-				ConstructConvexMeshOneShotManifold(*RawConvex, ConvexTransform, *Mesh, MeshTransform, MeshScale, CullDistance, MeshContacts);
+				const FVec3 RelativeMovement = FVec3(Constraint.GetRelativeMovement());
+				Private::FMeshContactGeneratorSettings ContactGeneratorSettings;
+				ContactGeneratorSettings.bFixNormals = RelativeMovement.IsNearlyZero();
+				Private::FMeshContactGenerator ContactGenerator(ContactGeneratorSettings);
+
+				if (const FImplicitBox3* RawBox = Convex.template GetObject<FImplicitBox3>())
+				{
+					ConstructConvexMeshOneShotManifoldMACD(*RawBox, ConvexTransform, *Mesh, MeshTransform, MeshScale, RelativeMovement, CullDistance, ContactGenerator);
+				}
+				else if (const TImplicitObjectScaled<FImplicitConvex3>* ScaledConvex = Convex.template GetObject<TImplicitObjectScaled<FImplicitConvex3>>())
+				{
+					ConstructConvexMeshOneShotManifoldMACD(*ScaledConvex, ConvexTransform, *Mesh, MeshTransform, MeshScale, RelativeMovement, CullDistance, ContactGenerator);
+				}
+				else if (const TImplicitObjectInstanced<FImplicitConvex3>* InstancedConvex = Convex.template GetObject<TImplicitObjectInstanced<FImplicitConvex3>>())
+				{
+					ConstructConvexMeshOneShotManifoldMACD(*InstancedConvex, ConvexTransform, *Mesh, MeshTransform, MeshScale, RelativeMovement, CullDistance, ContactGenerator);
+				}
+				else if (const FImplicitConvex3* RawConvex = Convex.template GetObject<FImplicitConvex3>())
+				{
+					ConstructConvexMeshOneShotManifoldMACD(*RawConvex, ConvexTransform, *Mesh, MeshTransform, MeshScale, RelativeMovement, CullDistance, ContactGenerator);
+				}
+				else
+				{
+					check(false);
+				}
+
+				Constraint.SetOneShotManifoldContacts(ContactGenerator.GetContactPoints());
 			}
 			else
 			{
-				check(false);
-			}
+				FContactTriangleCollector MeshContacts(bChaos_Collision_OneSidedTriangleMesh, PhiTolerance, DistanceTolerance, ConvexTransform);
 
-			Constraint.SetOneShotManifoldContacts(MeshContacts.GetContactPoints());
+				if (const FImplicitBox3* RawBox = Convex.template GetObject<FImplicitBox3>())
+				{
+					ConstructConvexMeshOneShotManifold(*RawBox, ConvexTransform, *Mesh, MeshTransform, MeshScale, CullDistance, MeshContacts);
+				}
+				else if (const TImplicitObjectScaled<FImplicitConvex3>* ScaledConvex = Convex.template GetObject<TImplicitObjectScaled<FImplicitConvex3>>())
+				{
+					ConstructConvexMeshOneShotManifold(*ScaledConvex, ConvexTransform, *Mesh, MeshTransform, MeshScale, CullDistance, MeshContacts);
+				}
+				else if (const TImplicitObjectInstanced<FImplicitConvex3>* InstancedConvex = Convex.template GetObject<TImplicitObjectInstanced<FImplicitConvex3>>())
+				{
+					ConstructConvexMeshOneShotManifold(*InstancedConvex, ConvexTransform, *Mesh, MeshTransform, MeshScale, CullDistance, MeshContacts);
+				}
+				else if (const FImplicitConvex3* RawConvex = Convex.template GetObject<FImplicitConvex3>())
+				{
+					ConstructConvexMeshOneShotManifold(*RawConvex, ConvexTransform, *Mesh, MeshTransform, MeshScale, CullDistance, MeshContacts);
+				}
+				else
+				{
+					check(false);
+				}
+
+				Constraint.SetOneShotManifoldContacts(MeshContacts.GetContactPoints());
+			}
 		}
 
 		/**
@@ -344,22 +519,23 @@ namespace Chaos
 				// New version uses a two-pass loop over triangles to avoid visiting triangles whose vertices are all colliding as a result of checking adjacent triangles
 				Private::FMeshContactGeneratorSettings ContactGeneratorSettings;
 				Private::FMeshContactGenerator ContactGenerator(ContactGeneratorSettings);
+				const FVec3 RelativeMovement = FVec3(Constraint.GetRelativeMovement());
 
 				if (const FImplicitBox3* RawBox = Convex.template GetObject<FImplicitBox3>())
 				{
-					ConstructConvexMeshOneShotManifold2(*RawBox, ConvexTransform, Mesh, MeshTransform, MeshScale, CullDistance, ContactGenerator);
+					ConstructConvexMeshOneShotManifoldMACD(*RawBox, ConvexTransform, Mesh, MeshTransform, MeshScale, RelativeMovement, CullDistance, ContactGenerator);
 				}
 				else if (const TImplicitObjectScaled<FImplicitConvex3>* ScaledConvex = Convex.template GetObject<TImplicitObjectScaled<FImplicitConvex3>>())
 				{
-					ConstructConvexMeshOneShotManifold2(*ScaledConvex, ConvexTransform, Mesh, MeshTransform, MeshScale, CullDistance, ContactGenerator);
+					ConstructConvexMeshOneShotManifoldMACD(*ScaledConvex, ConvexTransform, Mesh, MeshTransform, MeshScale, RelativeMovement, CullDistance, ContactGenerator);
 				}
 				else if (const TImplicitObjectInstanced<FImplicitConvex3>* InstancedConvex = Convex.template GetObject<TImplicitObjectInstanced<FImplicitConvex3>>())
 				{
-					ConstructConvexMeshOneShotManifold2(*InstancedConvex, ConvexTransform, Mesh, MeshTransform, MeshScale, CullDistance, ContactGenerator);
+					ConstructConvexMeshOneShotManifoldMACD(*InstancedConvex, ConvexTransform, Mesh, MeshTransform, MeshScale, RelativeMovement, CullDistance, ContactGenerator);
 				}
 				else if (const FImplicitConvex3* RawConvex = Convex.template GetObject<FImplicitConvex3>())
 				{
-					ConstructConvexMeshOneShotManifold2(*RawConvex, ConvexTransform, Mesh, MeshTransform, MeshScale, CullDistance, ContactGenerator);
+					ConstructConvexMeshOneShotManifoldMACD(*RawConvex, ConvexTransform, Mesh, MeshTransform, MeshScale, RelativeMovement, CullDistance, ContactGenerator);
 				}
 				else
 				{
diff --git a/Engine/Source/Runtime/Experimental/Chaos/Private/Chaos/PBDRigidsEvolutionGBF.cpp b/Engine/Source/Runtime/Experimental/Chaos/Private/Chaos/PBDRigidsEvolutionGBF.cpp
--- a/Engine/Source/Runtime/Experimental/Chaos/Private/Chaos/PBDRigidsEvolutionGBF.cpp
+++ b/Engine/Source/Runtime/Experimental/Chaos/Private/Chaos/PBDRigidsEvolutionGBF.cpp
@@ -1191,7 +1191,7 @@ void FPBDRigidsEvolutionGBF::DestroyTransientConstraints(FGeometryParticleHandle
 
 void FPBDRigidsEvolutionGBF::DestroyTransientConstraints()
 {
-	for (auto& Particle : Particles.GetAllParticlesView())
+	for (FTransientGeometryParticleHandle& Particle : Particles.GetAllParticlesView())
 	{
 		DestroyTransientConstraints(Particle.Handle());
 	}
diff --git a/Engine/Source/Runtime/Experimental/Chaos/Private/PBDRigidsSolver.cpp b/Engine/Source/Runtime/Experimental/Chaos/Private/PBDRigidsSolver.cpp
--- a/Engine/Source/Runtime/Experimental/Chaos/Private/PBDRigidsSolver.cpp
+++ b/Engine/Source/Runtime/Experimental/Chaos/Private/PBDRigidsSolver.cpp
@@ -338,9 +338,11 @@ namespace Chaos
 
 		// @todo(chaos): move to physics project settings and set these to -1 when we are settled on values...
 		FRealSingle ChaosSolverVelocityBoundsMultiplier = 1.0f;
-		FRealSingle ChaosSolverMaxVelocityBoundsExpansion = 3.0f;	// This should probably be a fraction of object size (see FParticlePairMidPhase::GenerateCollisions)
+		FRealSingle ChaosSolverMaxVelocityBoundsExpansion = 3.0f;			// This should probably be a fraction of object size (see FParticlePairMidPhase::GenerateCollisions)
+		FRealSingle ChaosSolverMaxVelocityBoundsExpansionMACD = 1000.0f;	// For use when Movement-Aware Collision Detection (MACD) is enabled
 		FAutoConsoleVariableRef CVarChaosSolverVelocityBoundsMultiplier(TEXT("p.Chaos.Solver.Collision.VelocityBoundsMultiplier"), ChaosSolverVelocityBoundsMultiplier, TEXT("Override velocity bounds multiplier (if >= 0)"));
 		FAutoConsoleVariableRef CVarChaosSolverMaxVelocityBoundsExpansion(TEXT("p.Chaos.Solver.Collision.MaxVelocityBoundsExpansion"), ChaosSolverMaxVelocityBoundsExpansion, TEXT("Override max velocity bounds expansion (if >= 0)"));
+		FAutoConsoleVariableRef CVarChaosSolverMaxVelocityBoundsExpansionMACD(TEXT("p.Chaos.Solver.Collision.MaxVelocityBoundsExpansionMACD"), ChaosSolverMaxVelocityBoundsExpansionMACD, TEXT("Override max velocity bounds expansion for MACD (if >= 0)"));
 
 		FRealSingle ChaosSolverMaxPushOutVelocity = -1.0f;
 		FAutoConsoleVariableRef CVarChaosSolverMaxPushOutVelocity(TEXT("p.Chaos.Solver.Collision.MaxPushOutVelocity"), ChaosSolverMaxPushOutVelocity, TEXT("Override max pushout velocity (if >= 0)"));
@@ -366,6 +368,9 @@ namespace Chaos
 		bool bChaosUseCCD = true;
 		FAutoConsoleVariableRef  CVarChaosUseCCD(TEXT("p.Chaos.Solver.UseCCD"), bChaosUseCCD, TEXT("Global flag to turn CCD on or off. Default is true"));
 
+		bool bChaosUseMACD = false;
+		FAutoConsoleVariableRef CVarChaos_Collision_UseMACD(TEXT("p.Chaos.Solver.UseMACD"), bChaosUseMACD, TEXT("Whether to enable support for Movement-Aware Collision Detection."));
+
 		// Joint cvars
 		float ChaosSolverJointMinSolverStiffness = 1.0f;
 		float ChaosSolverJointMaxSolverStiffness = 1.0f;
@@ -1299,10 +1304,24 @@ namespace Chaos
 			{
 				SetCollisionCullDistance(ChaosSolverCullDistance);
 			}
-			if ((ChaosSolverVelocityBoundsMultiplier > 0.0f) && (ChaosSolverMaxVelocityBoundsExpansion > 0.0f))
+			if (!CVars::bChaosUseMACD && (ChaosSolverVelocityBoundsMultiplier > 0.0f) && (ChaosSolverMaxVelocityBoundsExpansion > 0.0f))
 			{
+				// TEMP
+				if ((GetEvolution()->GetCollisionConstraints().GetDetectorSettings().BoundsVelocityInflation != ChaosSolverVelocityBoundsMultiplier) || (GetEvolution()->GetCollisionConstraints().GetDetectorSettings().MaxVelocityBoundsExpansion != ChaosSolverMaxVelocityBoundsExpansion))
+				{
+					GetEvolution()->DestroyTransientConstraints();
+				}
 				SetVelocityBoundsExpansion(ChaosSolverVelocityBoundsMultiplier, ChaosSolverMaxVelocityBoundsExpansion);
 			}
+			if (CVars::bChaosUseMACD)
+			{
+				// TEMP
+				if ((GetEvolution()->GetCollisionConstraints().GetDetectorSettings().BoundsVelocityInflation != ChaosSolverVelocityBoundsMultiplier) || (GetEvolution()->GetCollisionConstraints().GetDetectorSettings().MaxVelocityBoundsExpansion != ChaosSolverMaxVelocityBoundsExpansionMACD))
+				{
+					GetEvolution()->DestroyTransientConstraints();
+				}
+				SetVelocityBoundsExpansion(ChaosSolverVelocityBoundsMultiplier, ChaosSolverMaxVelocityBoundsExpansionMACD);
+			}
 			if (ChaosSolverMaxPushOutVelocity >= 0.0f)
 			{
 				SetCollisionMaxPushOutVelocity(ChaosSolverMaxPushOutVelocity);
diff --git a/Engine/Source/Runtime/Experimental/Chaos/Public/Chaos/Collision/ConvexContactPointUtilities.h b/Engine/Source/Runtime/Experimental/Chaos/Public/Chaos/Collision/ConvexContactPointUtilities.h
--- a/Engine/Source/Runtime/Experimental/Chaos/Public/Chaos/Collision/ConvexContactPointUtilities.h
+++ b/Engine/Source/Runtime/Experimental/Chaos/Public/Chaos/Collision/ConvexContactPointUtilities.h
@@ -3,9 +3,10 @@
 
 #include "Chaos/Core.h"
 #include "Chaos/Capsule.h"
+#include "Chaos/Collision/ContactPoint.h"
 #include "Chaos/Triangle.h"
 
-namespace Chaos
+namespace Chaos:: Private
 {
 	// Project a convex onto an axis and return the projected range as well as the vertex indices that bound the range
 	template <typename ConvexType>
@@ -75,4 +76,33 @@ namespace Chaos
 		}
 	}
 
-}
\ No newline at end of file
+	// Update contact as if the second shape has moved. Shift must be in the space of the second object
+	inline void ApplyContactPointShift(FContactPoint& InOutContactPoint, const FVec3& Shift)
+	{
+		const FReal ShiftDotNormal = FVec3::DotProduct(Shift, InOutContactPoint.ShapeContactNormal);
+		InOutContactPoint.ShapeContactPoints[1] += (Shift - ShiftDotNormal * InOutContactPoint.ShapeContactNormal);
+		InOutContactPoint.Phi += ShiftDotNormal;
+	}
+}
+
+namespace Chaos
+{
+	template <typename ConvexType>
+	UE_DEPRECATED(5.4, "Not part of public API")
+	inline void ProjectOntoAxis(const ConvexType& Convex, const FVec3& AxisN, const FVec3& AxisX, FReal& PMin, FReal& PMax, int32& MinVertexIndex, int32& MaxVertexIndex, TArrayView<FReal>* VertexDs)
+	{
+		return Private::ProjectOntoAxis(Convex, AxisN, AxisX, PMin, PMax, MinVertexIndex, MaxVertexIndex, VertexDs);
+	}
+
+	UE_DEPRECATED(5.4, "Not part of public API")
+	inline void ProjectOntoAxis(const FTriangle& Triangle, const FVec3& AxisN, const FVec3& AxisX, FReal& PMin, FReal& PMax, int32& MinVertexIndex, int32& MaxVertexIndex)
+	{
+		return Private::ProjectOntoAxis(Triangle, AxisN, AxisX, PMin, PMax, MinVertexIndex, MaxVertexIndex);
+	}
+
+	UE_DEPRECATED(5.4, "Not part of public API")
+	inline void ProjectOntoAxis(const FCapsule& Capsule, const FVec3& AxisN, const FVec3& AxisX, FReal& PMin, FReal& PMax, int32& MinVertexIndex, int32& MaxVertexIndex)
+	{
+		return Private::ProjectOntoAxis(Capsule, AxisN, AxisX, PMin, PMax, MinVertexIndex, MaxVertexIndex);
+	}
+}
diff --git a/Engine/Source/Runtime/Experimental/Chaos/Public/Chaos/Collision/ConvexFeature.h b/Engine/Source/Runtime/Experimental/Chaos/Public/Chaos/Collision/ConvexFeature.h
new file mode 100644
--- /dev/null
+++ b/Engine/Source/Runtime/Experimental/Chaos/Public/Chaos/Collision/ConvexFeature.h
@@ -0,0 +1,77 @@
+// Copyright Epic Games, Inc.All Rights Reserved.
+#pragma once
+
+#include "Chaos/Core.h"
+
+namespace Chaos::Private
+{
+	// The type of a convex feature
+	enum class EConvexFeatureType : int8
+	{
+		Unknown,
+		Vertex,
+		Edge,
+		Plane
+	};
+
+	// Information used to identify a feasture on a convex shape
+	class FConvexFeature
+	{
+	public:
+		FConvexFeature& Init()
+		{
+			ObjectIndex = INDEX_NONE;
+			PlaneIndex = INDEX_NONE;
+			PlaneFeatureIndex = INDEX_NONE;
+			FeatureType = EConvexFeatureType::Unknown;
+			return *this;
+		}
+
+		void Set(const int32 InObjectIndex, const EConvexFeatureType InFeatureType, const int32 InPlaneIndex, const int32 InPlaneFeatureIndex)
+		{
+			ObjectIndex = InObjectIndex;
+			PlaneIndex = InPlaneIndex;
+			PlaneFeatureIndex = InPlaneFeatureIndex;
+			FeatureType = InFeatureType;
+		}
+
+		// An index used to index into some collection (e.g., triangle index in a mesh)
+		int32 ObjectIndex;
+
+		// The plane index on the convex. E.g., See FConvex::GetPlane()
+		int32 PlaneIndex;
+
+		// For Vertex or Edge features, the vertex the feature on the plane. E.g., See FConvex::GetPlaneVertex()
+		int32 PlaneFeatureIndex;
+
+		// The feature type
+		EConvexFeatureType FeatureType;
+	};
+
+	// A contact point between two convex features
+	template<typename T>
+	class TConvexContactPoint
+	{
+	public:
+		using FRealType = T;
+
+		TConvexContactPoint<FRealType>& Init()
+		{
+			Features[0].Init();
+			Features[1].Init();
+			ShapeContactPoints[0] = TVec3<FRealType>(0);
+			ShapeContactPoints[1] = TVec3<FRealType>(0);
+			ShapeContactNormal = TVec3<FRealType>(0);
+			Phi = TNumericLimits<FRealType>::Max();
+			return *this;
+		}
+
+		FConvexFeature Features[2];
+		TVec3<FRealType> ShapeContactPoints[2];
+		TVec3<FRealType> ShapeContactNormal;
+		FRealType Phi;
+	};
+
+	using FConvexContactPoint = TConvexContactPoint<FReal>;
+	using FConvexContactPointf = TConvexContactPoint<FRealSingle>;
+}
\ No newline at end of file
diff --git a/Engine/Source/Runtime/Experimental/Chaos/Public/Chaos/Collision/ConvexTriangleContactPoint.h b/Engine/Source/Runtime/Experimental/Chaos/Public/Chaos/Collision/ConvexTriangleContactPoint.h
--- a/Engine/Source/Runtime/Experimental/Chaos/Public/Chaos/Collision/ConvexTriangleContactPoint.h
+++ b/Engine/Source/Runtime/Experimental/Chaos/Public/Chaos/Collision/ConvexTriangleContactPoint.h
@@ -1,12 +1,22 @@
 // Copyright Epic Games, Inc. All Rights Reserved.
 #pragma once
 #include "Chaos/Core.h"
+#include "Chaos/Collision/ContactPoint.h"
 #include "Chaos/CollisionResolutionTypes.h"
 #include "Chaos/Framework/UncheckedArray.h"
 #include "Chaos/ImplicitFwd.h"
 
 namespace Chaos
 {
+	namespace Private
+	{
+		template<typename ConvexType>
+		extern CHAOS_API bool FindClosestFeatures(const ConvexType& Convex, const FRigidTransform3& ConvexTransform, const FTriangle& Triangle, const FVec3& ConvexRelativeMovement, const FReal CullDistance, FConvexContactPoint& OutContact);
+
+		template <typename ConvexType>
+		extern CHAOS_API void ConvexTriangleManifoldFromContact(const ConvexType& Convex, const FTriangle& Triangle, const FConvexContactPoint& Contact, const FReal CullDistance, FContactPointLargeManifold& OutManifold);
+	}
+
 	template <typename ConvexType>
 	extern CHAOS_API void ConstructConvexTriangleOneShotManifold2(const ConvexType& Convex, const FTriangle& Triangle, const FReal CullDistance, FContactPointManifold& OutContactPoints);
 }
\ No newline at end of file
diff --git a/Engine/Source/Runtime/Experimental/Chaos/Public/Chaos/Collision/MeshContactGenerator.h b/Engine/Source/Runtime/Experimental/Chaos/Public/Chaos/Collision/MeshContactGenerator.h
--- a/Engine/Source/Runtime/Experimental/Chaos/Public/Chaos/Collision/MeshContactGenerator.h
+++ b/Engine/Source/Runtime/Experimental/Chaos/Public/Chaos/Collision/MeshContactGenerator.h
@@ -46,6 +46,12 @@ namespace Chaos::Private
 		// @todo(chaos): the non-culled option is not well tested and probably broken
 		uint32 bCullBackFaces : 1;
 
+		// Whether to auto-correct normals
+		uint32 bFixNormals : 1;
+
+		// Whether to sort the contacts to improve solver convergence
+		uint32 bSortForSolver : 1;
+
 		// Whether to use the optimized two-pass loop over triangles in GenerateMeshContacts which skips triangles
 		// that have contacts on all vertices in the second pass. This is only useful when this case occurs a lot
 		// which is does for large convexes against many triangles, but rarely for capsules and spheres.
@@ -129,24 +135,25 @@ namespace Chaos::Private
 			}
 
 			// Does this triangle contains the specified vertex? (VertexIndex is an index into the owning mesh's vertices)
-			inline bool HasVertexIndex(const int32 VertexIndex) const
+			inline bool HasVertexID(const int32 VertexIndex) const
 			{
 				return (VertexIndices[0] == VertexIndex) || (VertexIndices[1] == VertexIndex) || (VertexIndices[2] == VertexIndex);
 			}
 
-			inline bool GetVertexPosition(const int32 VertexIndex, FVec3& OutVertex) const
+			// Get the vertex position from the vertex ID (not the triangle-local vertex index)
+			inline bool GetVertexWithID(const int32 VertexID, FVec3& OutVertex) const
 			{
-				if (VertexIndex == VertexIndices[0])
+				if (VertexID == VertexIndices[0])
 				{
 					OutVertex = Triangle.GetVertex(0);
 					return true;
 				}
-				else if (VertexIndex == VertexIndices[1])
+				else if (VertexID == VertexIndices[1])
 				{
 					OutVertex = Triangle.GetVertex(1);
 					return true;
 				}
-				else if (VertexIndex == VertexIndices[2])
+				else if (VertexID == VertexIndices[2])
 				{
 					OutVertex = Triangle.GetVertex(2);
 					return true;
@@ -155,7 +162,7 @@ namespace Chaos::Private
 			}
 
 			// Get the positions of the other two vertices in the triangle. (VertexIndex is an index into the owning mesh's vertices)
-			inline bool GetOtherVertexPositions(const int32 VertexIndex, FVec3& OutVertex0, FVec3& OutVertex1) const
+			inline bool GetOtherVerticesFromID(const int32 VertexIndex, FVec3& OutVertex0, FVec3& OutVertex1) const
 			{
 				if (VertexIndex == VertexIndices[0])
 				{
@@ -178,11 +185,34 @@ namespace Chaos::Private
 				return false;
 			}
 
+			int32 GetLocalVertexIndexAt(const FVec3& InPos, const FReal InTolerance) const
+			{
+				for (int32 LocalVertexIndex = 0; LocalVertexIndex < 3; ++LocalVertexIndex)
+				{
+					if (FVec3::IsNearlyEqual(GetVertex(LocalVertexIndex), InPos, InTolerance))
+					{
+						return LocalVertexIndex;
+					}
+				}
+				return INDEX_NONE;
+			};
+
+			int32 GetVertexIDAt(const FVec3& InPos, const FReal InTolerance) const
+			{
+				const int32 LocalVertexIndex = GetLocalVertexIndexAt(InPos, InTolerance);
+				if (LocalVertexIndex != INDEX_NONE)
+				{
+					return VertexIndices[LocalVertexIndex];
+				}
+				return INDEX_NONE;
+			};
+
 			const FTriangle& GetTriangle() const
 			{
 				return Triangle;
 			}
 
+			// Get the vertex for the triangle-local vertex index [0,2]
 			const FVec3& GetVertex(const int32 LocalVertexIndex) const
 			{
 				return Triangle.GetVertex(LocalVertexIndex);
@@ -287,14 +317,9 @@ namespace Chaos::Private
 		template<typename TriangleContactGeneratorType>
 		void GenerateMeshContactsOnePass(const TriangleContactGeneratorType& TriangleContactGenerator)
 		{
-			FContactPointManifold TriangleContactPoints;
-
 			for (int32 LocalTriangleIndex = 0; LocalTriangleIndex < GetNumTriangles(); ++LocalTriangleIndex)
 			{
-				TriangleContactPoints.Reset();
-				TriangleContactGenerator(Triangles[LocalTriangleIndex].GetTriangle(), TriangleContactPoints);
-
-				AddTriangleContacts(TriangleContactPoints, LocalTriangleIndex);
+				TriangleContactGenerator(*this, LocalTriangleIndex);
 
 				SetTriangleVisited(LocalTriangleIndex, 0);
 			}
@@ -303,8 +328,6 @@ namespace Chaos::Private
 		template<typename TriangleContactGeneratorType>
 		void GenerateMeshContactsTwoPass(const TriangleContactGeneratorType& TriangleContactGenerator)
 		{
-			FContactPointManifold TriangleContactPoints;
-
 			// First loop: Visit triangles that do not have any collisions on any of their vertices or edges.
 			// This will skip all triangles whose neighbours have already been processed and generated a contact
 			// on a shared edge/vertex.
@@ -312,10 +335,7 @@ namespace Chaos::Private
 			{
 				if (GetNumTriangleFaceCollisions(LocalTriangleIndex) == 0)
 				{
-					TriangleContactPoints.Reset();
-					TriangleContactGenerator(Triangles[LocalTriangleIndex].GetTriangle(), TriangleContactPoints);
-
-					AddTriangleContacts(TriangleContactPoints, LocalTriangleIndex);
+					TriangleContactGenerator(*this, LocalTriangleIndex);
 
 					SetTriangleVisited(LocalTriangleIndex, 0);
 				}
@@ -327,10 +347,7 @@ namespace Chaos::Private
 			{
 				if (!IsTriangleVisited(LocalTriangleIndex) && (GetNumTriangleFaceCollisions(LocalTriangleIndex) < 3))
 				{
-					TriangleContactPoints.Reset();
-					TriangleContactGenerator(Triangles[LocalTriangleIndex].GetTriangle(), TriangleContactPoints);
-
-					AddTriangleContacts(TriangleContactPoints, LocalTriangleIndex);
+					TriangleContactGenerator(*this, LocalTriangleIndex);
 
 					SetTriangleVisited(LocalTriangleIndex, 1);
 				}
@@ -353,7 +370,7 @@ namespace Chaos::Private
 			}
 		}
 
-		int32 GetOtherTriangleIndexForEdge(const int32 LocalTriangleIndex, const FContactEdgeID EdgeID)
+		int32 GetOtherTriangleIndexForEdge(const int32 LocalTriangleIndex, const FContactEdgeID& EdgeID)
 		{
 			FEdgeTriangleIndices* EdgeTriangleIndices = EdgeTriangleIndicesMap.Find(EdgeID);
 			if (EdgeTriangleIndices != nullptr)
@@ -394,15 +411,26 @@ namespace Chaos::Private
 			Triangles[LocalTriangleIndex].SetVisitIndex(VisitIndex);
 		}
 
+	public:
 		int32 GetNumTriangles() const
 		{
 			return Triangles.Num();
 		}
 
-		void AddTriangleContacts(const FContactPointManifold& TriangleContactPoints, const int32 LocalTriangleIndex);
+		const FTriangle& GetTriangle(const int32 LocalTriangleIndex) const
+		{
+			return Triangles[LocalTriangleIndex].GetTriangle();
+		}
+
+		bool FixFeature(const int32 LocalTriangleIndex, Private::EConvexFeatureType& InOutFeatureType, int32& InOutFeatureIndex, FVec3& InOutPlaneNormal, FVec3& InOutPlanePosition);
+
+		void AddTriangleContacts(const int32 LocalTriangleIndex, const TArrayView<FContactPoint>& TriangleContactPoints);
+
+	private:
 		void PruneAndCorrectContacts();
 		void FixContactNormal(const int32 ContactIndex);
 		void RemoveDisabledContacts();
+		void SortContactsForSolver();
 		void FinalizeContacts(const FRigidTransform3& MeshToConvexTransform);
 
 		void DebugDrawContacts(const FRigidTransform3& ConvexTransform, const FColor& Color, const FReal LineScale);
diff --git a/Engine/Source/Runtime/Experimental/Chaos/Public/Chaos/Collision/PBDCollisionConstraint.h b/Engine/Source/Runtime/Experimental/Chaos/Public/Chaos/Collision/PBDCollisionConstraint.h
--- a/Engine/Source/Runtime/Experimental/Chaos/Public/Chaos/Collision/PBDCollisionConstraint.h
+++ b/Engine/Source/Runtime/Experimental/Chaos/Public/Chaos/Collision/PBDCollisionConstraint.h
@@ -489,6 +489,9 @@ namespace Chaos
 		FReal GetCullDistance() const { return CullDistance; }
 		void SetCullDistance(FReal InCullDistance) { CullDistance = FRealSingle(InCullDistance); }
 
+		FVec3f GetRelativeMovement() const { return RelativeMovement; }
+		void SetRelativeMovement(const FVec3f& InDelta) { RelativeMovement = InDelta; }
+
 		// Whether we are using manifolds (either one-shot or incremental)
 		bool GetUseManifold() const { return Flags.bUseManifold; }
 
@@ -642,6 +645,10 @@ namespace Chaos
 
 		}
 
+		// Fix the contact points when we move the particle after collision detection so that
+		// the contact points on each body are aligned along the normal (required for static friction tracking)
+		CHAOS_API void CorrectManifoldPoints();
+
 		CHAOS_API void UpdateManifoldContacts();
 
 		// Particle-relative transform of each collision shape in the constraint
@@ -950,6 +957,9 @@ namespace Chaos
 		// The separation distance at which we don't track contacts
 		FRealSingle CullDistance;
 
+		// Relative movement on the last tick: (V0 - V1).Dt.
+		FVec3f RelativeMovement;
+
 		FPBDCollisionConstraintContainerCookie ContainerCookie;
 		Private::FCollisionSortKey CollisionSortKey;
 
diff --git a/Engine/Source/Runtime/Experimental/Chaos/Public/Chaos/Collision/ParticlePairMidPhase.h b/Engine/Source/Runtime/Experimental/Chaos/Public/Chaos/Collision/ParticlePairMidPhase.h
--- a/Engine/Source/Runtime/Experimental/Chaos/Public/Chaos/Collision/ParticlePairMidPhase.h
+++ b/Engine/Source/Runtime/Experimental/Chaos/Public/Chaos/Collision/ParticlePairMidPhase.h
@@ -85,19 +85,27 @@ namespace Chaos
 		 * @brief Perform a bounds check and run the narrow phase if necessary
 		 * @return The number of collisions constraints that were activated
 		*/
+		UE_DEPRECATED(5.4, "Use single precision version")
+		CHAOS_API int32 GenerateCollision(const FReal CullDistance, const FReal Dt, const FCollisionContext& Context) { return GenerateCollision(FRealSingle(Dt), FRealSingle(CullDistance), FVec3f(0), Context); }
+
 		CHAOS_API int32 GenerateCollision(
-			const FReal CullDistance,
-			const FReal Dt,
+			const FRealSingle Dt,
+			const FRealSingle CullDistance,
+			const FVec3f& RelativeMovement,
 			const FCollisionContext& Context);
 
 		/**
 		 * @brief Generate a SweptConstraint as long as AABBs overlap
 		 * @return The number of collisions constraints that were activated
 		*/
+		UE_DEPRECATED(5.4, "Use single precision version")
+		CHAOS_API int32 GenerateCollisionCCD( const bool bEnableCCDSweep, const FReal CullDistance, const FReal Dt, const FCollisionContext& Context) { return GenerateCollisionCCD(FRealSingle(Dt), FRealSingle(CullDistance), FVec3f(0), bEnableCCDSweep, Context); }
+
 		CHAOS_API int32 GenerateCollisionCCD(
+			const FRealSingle Dt,
+			const FRealSingle CullDistance,
+			const FVec3f& RelativeMovement,
 			const bool bEnableCCDSweep,
-			const FReal CullDistance,
-			const FReal Dt,
 			const FCollisionContext& Context);
 
 		/**
@@ -105,41 +113,53 @@ namespace Chaos
 		 * @parame SleepEpoch The tick on which the particle went to sleep.
 		 * Only constraints that were active when the particle went to sleep should be reactivated.
 		*/
-		CHAOS_API void WakeCollision(const int32 SleepEpoch, const int32 CurrentEpoch);
+		CHAOS_API void WakeCollision(
+			const int32 SleepEpoch, 
+			const int32 CurrentEpoch);
 
 		/**
 		 * @brief Set the collision from the parameter and activate it
 		 * This is used by the Resim restore functionality
 		*/
-		CHAOS_API void SetCollision(const FPBDCollisionConstraint& Constraint, const FCollisionContext& Context);
+		CHAOS_API void SetCollision(
+			const FPBDCollisionConstraint& Constraint, 
+			const FCollisionContext& Context);
 
 	private:
 		CHAOS_API int32 GenerateCollisionImpl(
-			const FReal CullDistance, 
-			const FReal Dt,
+			const FRealSingle Dt,
+			const FRealSingle CullDistance,
+			const FVec3f& RelativeMovement,
 			const FCollisionContext& Context);
 			
 		CHAOS_API int32 GenerateCollisionCCDImpl(
+			const FRealSingle Dt,
+			const FRealSingle CullDistance,
+			const FVec3f& RelativeMovement,
 			const bool bEnableCCDSweep,
-			const FReal CullDistance,
-			const FReal Dt,
 			const FCollisionContext& Context);
 
 		CHAOS_API int32 GenerateCollisionProbeImpl(
-			const FReal CullDistance, 
-			const FReal Dt,
+			const FRealSingle Dt,
+			const FRealSingle CullDistance,
+			const FVec3f& RelativeMovement,
 			const FCollisionContext& Context);
 
 		/**
 		 * @brief Whether the two shapes are separated by less than CullDistance (i.e., we should run the narrow phase).
 		 * Also returns true if bounds checking is disabled (globally or for this pair)
 		*/
-		CHAOS_API bool DoBoundsOverlap(const FReal CullDistance, const int32 CurrentEpoch);
+		CHAOS_API bool DoBoundsOverlap(
+			const FRealSingle CullDistance, 
+			const FVec3f& RelativeMovement,
+			const int32 CurrentEpoch);
 
 		/**
 		 * @brief Create a constraint
 		*/
-		CHAOS_API void CreateConstraint(const FReal CullDistance, const FCollisionContext& Context);
+		CHAOS_API void CreateConstraint(
+			const FReal CullDistance, 
+			const FCollisionContext& Context);
 
 		FParticlePairMidPhase& MidPhase;
 		FPBDCollisionConstraintPtr Constraint;
@@ -366,8 +386,9 @@ namespace Chaos
 		virtual void BuildDetectorsImpl() = 0;
 
 		virtual int32 GenerateCollisionsImpl(
-			const FReal CullDistance,
-			const FReal Dt,
+			const FRealSingle Dt,
+			const FRealSingle CullDistance,
+			const FVec3f& RelativeMovement,
 			const FCollisionContext& Context) = 0;
 
 		virtual void WakeCollisionsImpl(const int32 CurrentEpoch) = 0;
@@ -393,6 +414,7 @@ namespace Chaos
 				uint16 bIsCCD : 1;       // True if CCD is supported by either particle
 				uint16 bIsCCDActive : 1; // True if CCD is active for this midphase on this frame. This can be changed by modifiers and resets to bIsCCD each frame.
 				uint16 bUseSweep : 1;    // True if CCD is active (this frame) and we are moving fast enough to require a sweep
+				uint16 bIsMACD : 1;      // True if MACD (movement-aware collision detection) is enabled for this pair
 				uint16 bIsConvexOptimizationActive : 1; // True if convex optimization is active for this midphase
 				uint16 bIsSleeping : 1;
 				uint16 bIsModified : 1;  // True if a modifier applied any changes to this midphase
@@ -440,8 +462,9 @@ namespace Chaos
 
 	protected:
 		CHAOS_API virtual int32 GenerateCollisionsImpl(
-			const FReal CullDistance,
-			const FReal Dt,
+			const FRealSingle Dt,
+			const FRealSingle CullDistance,
+			const FVec3f& RelativeMovement,
 			const FCollisionContext& Context) override final;
 
 		CHAOS_API virtual void WakeCollisionsImpl(const int32 CurrentEpoch) override final;
@@ -483,7 +506,11 @@ namespace Chaos
 		CHAOS_API virtual void BuildDetectorsImpl() override final;
 
 	protected:
-		CHAOS_API virtual int32 GenerateCollisionsImpl( const FReal CullDistance, const FReal Dt, const FCollisionContext& Context) override final;
+		CHAOS_API virtual int32 GenerateCollisionsImpl(
+			const FRealSingle Dt,
+			const FRealSingle CullDistance,
+			const FVec3f& RelativeMovement,
+			const FCollisionContext& Context) override final;
 
 		CHAOS_API virtual void WakeCollisionsImpl(const int32 CurrentEpoch) override final;
 
@@ -566,8 +593,9 @@ namespace Chaos
 
 	protected:
 		CHAOS_API virtual int32 GenerateCollisionsImpl(
-			const FReal CullDistance,
-			const FReal Dt,
+			const FRealSingle Dt,
+			const FRealSingle CullDistance,
+			const FVec3f& RelativeMovement,
 			const FCollisionContext& Context) override final;
 
 		CHAOS_API virtual void WakeCollisionsImpl(const int32 CurrentEpoch) override final;
diff --git a/Engine/Source/Runtime/Experimental/Chaos/Public/Chaos/CollisionResolutionTypes.h b/Engine/Source/Runtime/Experimental/Chaos/Public/Chaos/CollisionResolutionTypes.h
--- a/Engine/Source/Runtime/Experimental/Chaos/Public/Chaos/CollisionResolutionTypes.h
+++ b/Engine/Source/Runtime/Experimental/Chaos/Public/Chaos/CollisionResolutionTypes.h
@@ -11,6 +11,15 @@
 
 namespace Chaos
 {
+	namespace Private
+	{
+		template<typename T> class TConvexContactPoint;
+		using FConvexContactPoint = TConvexContactPoint<FReal>;
+		using FConvexContactPointf = TConvexContactPoint<FRealSingle>;
+
+		enum class EConvexFeatureType : int8;
+	}
+
 	// Float and Double versions of ContactPoint
 	template<typename T> class TContactPoint;
 	using FContactPoint = TContactPoint<FReal>;
