Dataflow : Fix for logging CIS clang issue.

diff --git a/Engine/Source/Runtime/Experimental/Dataflow/Core/Private/Dataflow/DataflowContextCachingFactory.cpp b/Engine/Source/Runtime/Experimental/Dataflow/Core/Private/Dataflow/DataflowContextCachingFactory.cpp
--- a/Engine/Source/Runtime/Experimental/Dataflow/Core/Private/Dataflow/DataflowContextCachingFactory.cpp
+++ b/Engine/Source/Runtime/Experimental/Dataflow/Core/Private/Dataflow/DataflowContextCachingFactory.cpp
@@ -2,14 +2,32 @@
 
 #include "Dataflow/DataflowContextCachingFactory.h"
 
-#include "ChaosLog.h"
 #include "Dataflow/DataflowNodeParameters.h"
+#include "Logging/LogMacros.h"
 #include "Misc/MessageDialog.h"
 
+DEFINE_LOG_CATEGORY_STATIC(LogDataflowContextCachingFactory, Warning, All);
+
+
 namespace Dataflow
 {
 	FContextCachingFactory* FContextCachingFactory::Instance = nullptr;
 
+	void FContextCachingFactory::RegisterSerializeFunction(const FName& Type, FSerializeFunction InSerializeFunc)
+	{
+		if (CachingMap.Contains(Type))
+		{
+			UE_LOG(LogDataflowContextCachingFactory, Warning,
+				TEXT("Warning : Dataflow output caching registration conflicts with "
+					"existing type(%s)"), *Type.ToString());
+		}
+		else
+		{
+			CachingMap.Add(Type, InSerializeFunc);
+		}
+	}
+
+
 	FContextCacheElementBase* FContextCachingFactory::Serialize(FArchive& Ar, FContextCacheData&& Element)
 	{
 		FContextCacheElementBase* RetVal = nullptr;
@@ -27,7 +45,7 @@ namespace Dataflow
 		}
 		else
 		{
-			UE_LOG(LogChaos, Warning,
+			UE_LOG(LogDataflowContextCachingFactory, Warning,
 				TEXT("Warning : Dataflow missing context chaching callback type(%s)"), *Element.Type.ToString());
 		}
 		return RetVal;
diff --git a/Engine/Source/Runtime/Experimental/Dataflow/Core/Public/Dataflow/DataflowContextCachingFactory.h b/Engine/Source/Runtime/Experimental/Dataflow/Core/Public/Dataflow/DataflowContextCachingFactory.h
--- a/Engine/Source/Runtime/Experimental/Dataflow/Core/Public/Dataflow/DataflowContextCachingFactory.h
+++ b/Engine/Source/Runtime/Experimental/Dataflow/Core/Public/Dataflow/DataflowContextCachingFactory.h
@@ -2,7 +2,6 @@
 
 #pragma once
 
-#include "ChaosLog.h"
 #include "CoreMinimal.h"
 #include "Dataflow/DataflowNodeParameters.h"
 
@@ -46,19 +45,7 @@ namespace Dataflow
 			return Instance;
 		}
 
-		void RegisterSerializeFunction(const FName& Type, FSerializeFunction InSerializeFunc)
-		{
-			if (CachingMap.Contains(Type))
-			{
-				UE_LOG(LogChaos, Warning,
-					TEXT("Warning : Dataflow output caching registration conflicts with "
-						"existing type(%s)"), *Type.ToString());
-			}
-			else
-			{
-				CachingMap.Add(Type, InSerializeFunc);
-			}
-		}
+		void RegisterSerializeFunction(const FName& Type, FSerializeFunction InSerializeFunc);
 
 		template<class T>
 		static T& GetTypedElement(FContextCacheElementBase* InElement)
