Chaos - hook up FBodyInstance::bEnableMACD to the feature in Chaos

diff --git a/Engine/Source/Runtime/Engine/Private/PhysicsEngine/ImmediatePhysics/ImmediatePhysicsChaos/ImmediatePhysicsSimulation_Chaos.cpp b/Engine/Source/Runtime/Engine/Private/PhysicsEngine/ImmediatePhysics/ImmediatePhysicsChaos/ImmediatePhysicsSimulation_Chaos.cpp
--- a/Engine/Source/Runtime/Engine/Private/PhysicsEngine/ImmediatePhysics/ImmediatePhysicsChaos/ImmediatePhysicsSimulation_Chaos.cpp
+++ b/Engine/Source/Runtime/Engine/Private/PhysicsEngine/ImmediatePhysics/ImmediatePhysicsChaos/ImmediatePhysicsSimulation_Chaos.cpp
@@ -368,7 +368,8 @@ namespace ImmediatePhysics_Chaos
 		DetectorSettings.bAllowManifoldReuse = false;
 		DetectorSettings.bDeferNarrowPhase = (ChaosImmediate_Collision_DeferNarrowPhase != 0);
 		DetectorSettings.bAllowManifolds = (ChaosImmediate_Collision_UseManifolds != 0);
-		DetectorSettings.bAllowCCD = false; 
+		DetectorSettings.bAllowCCD = false;
+		DetectorSettings.bAllowMACD = false;
 		Implementation->Collisions.SetDetectorSettings(DetectorSettings);
 	}
 
diff --git a/Engine/Source/Runtime/Experimental/Chaos/Private/Chaos/ClusterUnionManager.cpp b/Engine/Source/Runtime/Experimental/Chaos/Private/Chaos/ClusterUnionManager.cpp
--- a/Engine/Source/Runtime/Experimental/Chaos/Private/Chaos/ClusterUnionManager.cpp
+++ b/Engine/Source/Runtime/Experimental/Chaos/Private/Chaos/ClusterUnionManager.cpp
@@ -655,6 +655,9 @@ namespace Chaos
 		// Cluster has one-way interaction only if all children are also one-way
 		bool bIsOneWayInteraction = (bIsNewCluster) || Cluster->InternalCluster->OneWayInteraction();
 
+		// Cluster uses MACD is any children require MACD
+		bool bIsMACD = false;
+
 		// Use the minimum sleep multiplier of all member particles
 		FRealSingle MinSleepThresholdMultiplier = TNumericLimits<FRealSingle>::Max();
 
@@ -676,6 +679,7 @@ namespace Chaos
 			}
 
 			bIsOneWayInteraction &= Particle->OneWayInteraction();
+			bIsMACD |= Particle->MACDEnabled();
 			MinSleepThresholdMultiplier = FMath::Min(MinSleepThresholdMultiplier, Particle->SleepThresholdMultiplier());
 
 			if (!Cluster->ChildProperties.Contains(Particle))
@@ -748,6 +752,7 @@ namespace Chaos
 		}
 
 		Cluster->InternalCluster->SetOneWayInteraction(bIsOneWayInteraction);
+		Cluster->InternalCluster->SetMACDEnabled(bIsMACD);
 
 		if (MinSleepThresholdMultiplier != TNumericLimits<FRealSingle>::Max())
 		{
diff --git a/Engine/Source/Runtime/Experimental/Chaos/Private/Chaos/Collision/ParticlePairMidPhase.cpp b/Engine/Source/Runtime/Experimental/Chaos/Private/Chaos/Collision/ParticlePairMidPhase.cpp
--- a/Engine/Source/Runtime/Experimental/Chaos/Private/Chaos/Collision/ParticlePairMidPhase.cpp
+++ b/Engine/Source/Runtime/Experimental/Chaos/Private/Chaos/Collision/ParticlePairMidPhase.cpp
@@ -859,6 +859,8 @@ namespace Chaos
 		const FCollisionContext& Context)
 	{
 		PHYSICS_CSV_SCOPED_EXPENSIVE(PhysicsVerbose, NarrowPhase_Filter);
+		check(InParticle0 != nullptr);
+		check(InParticle1 != nullptr);
 
 		Particle0 = InParticle0;
 		Particle1 = InParticle1;
@@ -866,19 +868,21 @@ namespace Chaos
 
 		Flags.bIsActive = true;
 
+		FConstGenericParticleHandle P0 = Particle0;
+		FConstGenericParticleHandle P1 = Particle1;
+
 		// If CCD is allowed in the current context and for at least one of
 		// the particles involved, enable it for this midphase.
 		//
 		// bIsCCDActive is reset to bIsCCD each frame, but can be overridden
 		// by modifiers.
-		const bool bIsCCD = Context.GetSettings().bAllowCCD && (FConstGenericParticleHandle(Particle0)->CCDEnabled() || FConstGenericParticleHandle(Particle1)->CCDEnabled());
+		const bool bIsCCD = Context.GetSettings().bAllowCCD && (P0->CCDEnabled() || P1->CCDEnabled());
 		Flags.bIsCCD = bIsCCD;
 		Flags.bIsCCDActive = bIsCCD;
 
-		// @todo(chaos): we should only enable relative movement mode if the particle asks for it
-		// (this also affects how the particle bounds is expanded in the broadphase)
-		//const bool bIsMACD = Context.GetSettings().bAllowMACD && (FConstGenericParticleHandle(Particle0)->MACDEnabled() || FConstGenericParticleHandle(Particle1)->MACDEnabled());
-		const bool bIsMACD = CVars::bChaosUseMACD;
+		// Enable Motion-Aware Collision Detection if either particle requests it
+		// (NOTE: MACD also affects how the particle world-space bounds is expanded in the broadphase)
+		const bool bIsMACD = Context.GetSettings().bAllowMACD && (P0->MACDEnabled() || P1->MACDEnabled());
 		Flags.bIsMACD = bIsMACD;
 
 		// Initially we allow for convex optimization where available
diff --git a/Engine/Source/Runtime/Experimental/Chaos/Private/Chaos/PBDRigidsEvolutionGBF.cpp b/Engine/Source/Runtime/Experimental/Chaos/Private/Chaos/PBDRigidsEvolutionGBF.cpp
--- a/Engine/Source/Runtime/Experimental/Chaos/Private/Chaos/PBDRigidsEvolutionGBF.cpp
+++ b/Engine/Source/Runtime/Experimental/Chaos/Private/Chaos/PBDRigidsEvolutionGBF.cpp
@@ -729,6 +729,8 @@ void FPBDRigidsEvolutionGBF::Integrate(FReal Dt)
 	const FReal BoundsThickness = GetCollisionConstraints().GetDetectorSettings().BoundsExpansion;
 	const FReal VelocityBoundsMultiplier = GetCollisionConstraints().GetDetectorSettings().BoundsVelocityInflation;
 	const FReal MaxVelocityBoundsExpansion = GetCollisionConstraints().GetDetectorSettings().MaxVelocityBoundsExpansion;
+	const FReal VelocityBoundsMultiplierMACD = GetCollisionConstraints().GetDetectorSettings().BoundsVelocityInflationMACD;
+	const FReal MaxVelocityBoundsExpansionMACD = GetCollisionConstraints().GetDetectorSettings().MaxVelocityBoundsExpansionMACD;
 	const FReal HackMaxAngularSpeedSq = CVars::HackMaxAngularVelocity * CVars::HackMaxAngularVelocity;
 	const FReal HackMaxLinearSpeedSq = CVars::HackMaxVelocity * CVars::HackMaxVelocity;
 
@@ -820,13 +822,23 @@ void FPBDRigidsEvolutionGBF::Integrate(FReal Dt)
 				// We need to expand the bounds back along velocity otherwise we can miss collisions when a box
 				// lands on another box, imparting velocity to the lower box and causing the boxes to be
 				// separated by more than Cull Distance at collision detection time.
+				// NOTE: We use a different (larger) max bounds expansion when MACD is enabled.
 				FVec3 VelocityBoundsDelta = FVec3(0);
-				if ((VelocityBoundsMultiplier > 0) && (MaxVelocityBoundsExpansion > 0))
+				if (!Particle.MACDEnabled())
 				{
-					VelocityBoundsDelta = (-VelocityBoundsMultiplier * Dt) * V;
-
-					// Box clamp to avoid sqrt
-					VelocityBoundsDelta.BoundToCube(MaxVelocityBoundsExpansion);
+					if ((VelocityBoundsMultiplier > 0) && (MaxVelocityBoundsExpansion > 0))
+					{
+						VelocityBoundsDelta = (-VelocityBoundsMultiplier * Dt) * V;
+						VelocityBoundsDelta = VelocityBoundsDelta.BoundToCube(MaxVelocityBoundsExpansion);
+					}
+				}
+				else
+				{
+					if ((VelocityBoundsMultiplierMACD > 0) && (MaxVelocityBoundsExpansionMACD > 0))
+					{
+						VelocityBoundsDelta = (-VelocityBoundsMultiplierMACD * Dt) * V;
+						VelocityBoundsDelta = VelocityBoundsDelta.BoundToCube(MaxVelocityBoundsExpansionMACD);
+					}
 				}
 
 				if (!Particle.CCDEnabled())
diff --git a/Engine/Source/Runtime/Experimental/Chaos/Private/PBDRigidsSolver.cpp b/Engine/Source/Runtime/Experimental/Chaos/Private/PBDRigidsSolver.cpp
--- a/Engine/Source/Runtime/Experimental/Chaos/Private/PBDRigidsSolver.cpp
+++ b/Engine/Source/Runtime/Experimental/Chaos/Private/PBDRigidsSolver.cpp
@@ -339,9 +339,11 @@ namespace Chaos
 		// @todo(chaos): move to physics project settings and set these to -1 when we are settled on values...
 		FRealSingle ChaosSolverVelocityBoundsMultiplier = 1.0f;
 		FRealSingle ChaosSolverMaxVelocityBoundsExpansion = 3.0f;			// This should probably be a fraction of object size (see FParticlePairMidPhase::GenerateCollisions)
+		FRealSingle ChaosSolverVelocityBoundsMultiplierMACD = 1.0f;
 		FRealSingle ChaosSolverMaxVelocityBoundsExpansionMACD = 1000.0f;	// For use when Movement-Aware Collision Detection (MACD) is enabled
 		FAutoConsoleVariableRef CVarChaosSolverVelocityBoundsMultiplier(TEXT("p.Chaos.Solver.Collision.VelocityBoundsMultiplier"), ChaosSolverVelocityBoundsMultiplier, TEXT("Override velocity bounds multiplier (if >= 0)"));
 		FAutoConsoleVariableRef CVarChaosSolverMaxVelocityBoundsExpansion(TEXT("p.Chaos.Solver.Collision.MaxVelocityBoundsExpansion"), ChaosSolverMaxVelocityBoundsExpansion, TEXT("Override max velocity bounds expansion (if >= 0)"));
+		FAutoConsoleVariableRef CVarChaosSolverVelocityBoundsMultiplierMACD(TEXT("p.Chaos.Solver.Collision.VelocityBoundsMultiplierMACD"), ChaosSolverVelocityBoundsMultiplierMACD, TEXT("Override velocity bounds multiplier for MACD (if >= 0)"));
 		FAutoConsoleVariableRef CVarChaosSolverMaxVelocityBoundsExpansionMACD(TEXT("p.Chaos.Solver.Collision.MaxVelocityBoundsExpansionMACD"), ChaosSolverMaxVelocityBoundsExpansionMACD, TEXT("Override max velocity bounds expansion for MACD (if >= 0)"));
 
 		FRealSingle ChaosSolverMaxPushOutVelocity = -1.0f;
@@ -364,12 +366,13 @@ namespace Chaos
 		int32 ChaosSolverCollisionAllowManifoldUpdate = 1;
 		FAutoConsoleVariableRef CVarChaosSolverCollisionAllowManifoldUpdate(TEXT("p.Chaos.Solver.Collision.AllowManifoldUpdate"), ChaosSolverCollisionAllowManifoldUpdate, TEXT("Enable/Disable reuse of manifolds between ticks (for small movement)."));
 
-		// Enable/Disable CCD
+		// Enable/Disable CCD. Set to false to disable the system, regardless of particle settings
 		bool bChaosUseCCD = true;
-		FAutoConsoleVariableRef  CVarChaosUseCCD(TEXT("p.Chaos.Solver.UseCCD"), bChaosUseCCD, TEXT("Global flag to turn CCD on or off. Default is true"));
+		FAutoConsoleVariableRef  CVarChaosUseCCD(TEXT("p.Chaos.Solver.UseCCD"), bChaosUseCCD, TEXT("Global flag to turn CCD on or off. Default is true (on)"));
 
-		bool bChaosUseMACD = false;
-		FAutoConsoleVariableRef CVarChaos_Collision_UseMACD(TEXT("p.Chaos.Solver.UseMACD"), bChaosUseMACD, TEXT("Whether to enable support for Movement-Aware Collision Detection."));
+		// Enable/Disable MACD (Motion-Aware Collision Detection). Set to false to disable the system, regardless of particle settings
+		bool bChaosUseMACD = true;
+		FAutoConsoleVariableRef CVarChaos_Collision_UseMACD(TEXT("p.Chaos.Solver.UseMACD"), bChaosUseMACD, TEXT("Global flag to turn Movement-Aware Collision Detection (MACD) on or off. Default is true (on)"));
 
 		// Joint cvars
 		float ChaosSolverJointMinSolverStiffness = 1.0f;
@@ -1248,6 +1251,28 @@ namespace Chaos
 		}
 	}
 
+	void FPBDRigidsSolver::SetVelocityBoundsExpansion(const FReal BoundsVelocityMultiplier, const FReal MaxBoundsVelocityExpansion)
+	{ 
+		FPBDCollisionConstraints& CollisionContainer = GetEvolution()->GetCollisionConstraints();
+		if ((CollisionContainer.GetDetectorSettings().BoundsVelocityInflation != BoundsVelocityMultiplier) || (CollisionContainer.GetDetectorSettings().MaxVelocityBoundsExpansion != MaxBoundsVelocityExpansion))
+		{
+			// If the settings change we must recreate collisions
+			GetEvolution()->DestroyTransientConstraints();
+		}
+		CollisionContainer.SetVelocityBoundsExpansion(BoundsVelocityMultiplier, MaxBoundsVelocityExpansion);
+	}
+
+	void FPBDRigidsSolver::SetVelocityBoundsExpansionMACD(const FReal BoundsVelocityMultiplier, const FReal MaxBoundsVelocityExpansion)
+	{
+		FPBDCollisionConstraints& CollisionContainer = GetEvolution()->GetCollisionConstraints();
+		if ((CollisionContainer.GetDetectorSettings().BoundsVelocityInflationMACD != BoundsVelocityMultiplier) || (CollisionContainer.GetDetectorSettings().MaxVelocityBoundsExpansionMACD != MaxBoundsVelocityExpansion))
+		{
+			// If the settings change we must recreate collisions
+			GetEvolution()->DestroyTransientConstraints();
+		}
+		CollisionContainer.SetVelocityBoundsExpansionMACD(BoundsVelocityMultiplier, MaxBoundsVelocityExpansion);
+	}
+
 	void FPBDRigidsSolver::PrepareAdvanceBy(const FReal DeltaTime)
 	{
 		MEvolution->GetCollisionConstraints().SetCollisionsEnabled(bChaosSolverCollisionEnabled);
@@ -1257,6 +1282,7 @@ namespace Chaos
 		CollisionDetectorSettings.bDeferNarrowPhase = (ChaosSolverCollisionDeferNarrowPhase != 0);
 		CollisionDetectorSettings.bAllowManifolds = (ChaosSolverCollisionUseManifolds != 0);
 		CollisionDetectorSettings.bAllowCCD = bChaosUseCCD;
+		CollisionDetectorSettings.bAllowMACD = bChaosUseMACD;
 		MEvolution->GetCollisionConstraints().SetDetectorSettings(CollisionDetectorSettings);
 		
 		FPBDJointSolverSettings JointsSettings = MEvolution->GetJointConstraints().GetSettings();
@@ -1304,23 +1330,14 @@ namespace Chaos
 			{
 				SetCollisionCullDistance(ChaosSolverCullDistance);
 			}
-			if (!CVars::bChaosUseMACD && (ChaosSolverVelocityBoundsMultiplier > 0.0f) && (ChaosSolverMaxVelocityBoundsExpansion > 0.0f))
+			SetVelocityBoundsExpansion(FMath::Max(0.0f, ChaosSolverVelocityBoundsMultiplier), FMath::Max(0.0f, ChaosSolverMaxVelocityBoundsExpansion));
+			if (bChaosUseMACD)
 			{
-				// TEMP
-				if ((GetEvolution()->GetCollisionConstraints().GetDetectorSettings().BoundsVelocityInflation != ChaosSolverVelocityBoundsMultiplier) || (GetEvolution()->GetCollisionConstraints().GetDetectorSettings().MaxVelocityBoundsExpansion != ChaosSolverMaxVelocityBoundsExpansion))
-				{
-					GetEvolution()->DestroyTransientConstraints();
-				}
-				SetVelocityBoundsExpansion(ChaosSolverVelocityBoundsMultiplier, ChaosSolverMaxVelocityBoundsExpansion);
+				SetVelocityBoundsExpansionMACD(FMath::Max(0.0f, ChaosSolverVelocityBoundsMultiplierMACD), FMath::Max(0.0f, ChaosSolverMaxVelocityBoundsExpansionMACD));
 			}
-			if (CVars::bChaosUseMACD)
+			else
 			{
-				// TEMP
-				if ((GetEvolution()->GetCollisionConstraints().GetDetectorSettings().BoundsVelocityInflation != ChaosSolverVelocityBoundsMultiplier) || (GetEvolution()->GetCollisionConstraints().GetDetectorSettings().MaxVelocityBoundsExpansion != ChaosSolverMaxVelocityBoundsExpansionMACD))
-				{
-					GetEvolution()->DestroyTransientConstraints();
-				}
-				SetVelocityBoundsExpansion(ChaosSolverVelocityBoundsMultiplier, ChaosSolverMaxVelocityBoundsExpansionMACD);
+				SetVelocityBoundsExpansionMACD(FMath::Max(0.0f, ChaosSolverVelocityBoundsMultiplier), FMath::Max(0.0f, ChaosSolverMaxVelocityBoundsExpansion));
 			}
 			if (ChaosSolverMaxPushOutVelocity >= 0.0f)
 			{
diff --git a/Engine/Source/Runtime/Experimental/Chaos/Public/Chaos/Collision/CollisionContext.h b/Engine/Source/Runtime/Experimental/Chaos/Public/Chaos/Collision/CollisionContext.h
--- a/Engine/Source/Runtime/Experimental/Chaos/Public/Chaos/Collision/CollisionContext.h
+++ b/Engine/Source/Runtime/Experimental/Chaos/Public/Chaos/Collision/CollisionContext.h
@@ -19,23 +19,32 @@ namespace Chaos
 			: BoundsExpansion(0)
 			, BoundsVelocityInflation(0)
 			, MaxVelocityBoundsExpansion(0)
+			, BoundsVelocityInflationMACD(0)
+			, MaxVelocityBoundsExpansionMACD(0)
 			, bFilteringEnabled(true)
 			, bDeferNarrowPhase(false)
 			, bAllowManifolds(true)
 			, bAllowManifoldReuse(true)
 			, bAllowCCD(true)
+			, bAllowMACD(false)
 		{
 		}
 
 		// Shape bounds are expanded by this in all 3 directions. This is also used as the CullDistance for collision response.
 		FReal BoundsExpansion;
 		
-		// Shape bounds in the broadphase are expanded by this multiple of velocity
+		// Shape bounds in the broadphase are expanded by this multiple of velocity (normal collision mode)
 		FReal BoundsVelocityInflation;
 
-		// We only allow the bounds to grow from velocity by this much
+		// We only allow the bounds to grow from velocity by this much (normal collision mode)
 		FReal MaxVelocityBoundsExpansion;
 
+		// Shape bounds in the broadphase are expanded by this multiple of velocity (movement-aware collision mode)
+		FReal BoundsVelocityInflationMACD;
+
+		// We only allow the bounds to grow from velocity by this much (movement-aware collision mode)
+		FReal MaxVelocityBoundsExpansionMACD;
+
 		// Whether to check the shape query flags in the narrow phase (e.g., Rigid Body nodes have already performed filtering prior to collision detection)
 		bool bFilteringEnabled;
 
@@ -53,6 +62,9 @@ namespace Chaos
 		// Whether CCD is allowed (disabled for RBAN)
 		bool bAllowCCD;
 
+		// Whether MACD is allowed (disabled for RBAN)
+		bool bAllowMACD;
+
 	};
 
 	/**
diff --git a/Engine/Source/Runtime/Experimental/Chaos/Public/Chaos/PBDCollisionConstraints.h b/Engine/Source/Runtime/Experimental/Chaos/Public/Chaos/PBDCollisionConstraints.h
--- a/Engine/Source/Runtime/Experimental/Chaos/Public/Chaos/PBDCollisionConstraints.h
+++ b/Engine/Source/Runtime/Experimental/Chaos/Public/Chaos/PBDCollisionConstraints.h
@@ -305,6 +305,12 @@ public:
 		DetectorSettings.MaxVelocityBoundsExpansion = MaxVelocityBoundsExpansion;
 	}
 
+	void SetVelocityBoundsExpansionMACD(const FReal BoundsVelocityMultiplier, const FReal MaxVelocityBoundsExpansion)
+	{
+		DetectorSettings.BoundsVelocityInflationMACD = BoundsVelocityMultiplier;
+		DetectorSettings.MaxVelocityBoundsExpansionMACD = MaxVelocityBoundsExpansion;
+	}
+
 protected:
 	CHAOS_API FPBDCollisionConstraint& GetConstraint(int32 Index);
 
diff --git a/Engine/Source/Runtime/Experimental/Chaos/Public/Chaos/ParticleHandle.h b/Engine/Source/Runtime/Experimental/Chaos/Public/Chaos/ParticleHandle.h
--- a/Engine/Source/Runtime/Experimental/Chaos/Public/Chaos/ParticleHandle.h
+++ b/Engine/Source/Runtime/Experimental/Chaos/Public/Chaos/ParticleHandle.h
@@ -1817,7 +1817,7 @@ public:
 	}
 
 	bool CCDEnabled() const
-	{ 
+	{
 		if (MHandle->CastToRigidParticle())
 		{
 			return MHandle->CastToRigidParticle()->CCDEnabled();
@@ -1826,6 +1826,16 @@ public:
 		return false;
 	}
 
+	bool MACDEnabled() const
+	{
+		if (MHandle->CastToRigidParticle())
+		{
+			return MHandle->CastToRigidParticle()->MACDEnabled();
+		}
+
+		return false;
+	}
+
 	bool HasCollisionConstraintFlag(const ECollisionConstraintFlags Flag)  const
 	{
 		if (MHandle->CastToRigidParticle())
diff --git a/Engine/Source/Runtime/Experimental/Chaos/Public/PBDRigidsSolver.h b/Engine/Source/Runtime/Experimental/Chaos/Public/PBDRigidsSolver.h
--- a/Engine/Source/Runtime/Experimental/Chaos/Public/PBDRigidsSolver.h
+++ b/Engine/Source/Runtime/Experimental/Chaos/Public/PBDRigidsSolver.h
@@ -179,7 +179,8 @@ namespace Chaos
 		void SetVelocityIterations(const int32 InNumIterations) { GetEvolution()->SetNumVelocityIterations(InNumIterations); }
 		void SetProjectionIterations(const int32 InNumIterations) { GetEvolution()->SetNumProjectionIterations(InNumIterations); }
 		void SetCollisionCullDistance(const FReal InCullDistance) { GetEvolution()->GetCollisionConstraints().SetCullDistance(InCullDistance); }
-		void SetVelocityBoundsExpansion(const FReal BoundsVelocityMultiplier, const FReal MaxBoundsVelocityExpansion) { GetEvolution()->GetCollisionConstraints().SetVelocityBoundsExpansion(BoundsVelocityMultiplier, MaxBoundsVelocityExpansion); }
+		void SetVelocityBoundsExpansion(const FReal BoundsVelocityMultiplier, const FReal MaxBoundsVelocityExpansion);
+		void SetVelocityBoundsExpansionMACD(const FReal BoundsVelocityMultiplier, const FReal MaxBoundsVelocityExpansion);
 		void SetCollisionMaxPushOutVelocity(const FReal InMaxPushOutVelocity) { GetEvolution()->GetCollisionConstraints().SetMaxPushOutVelocity(InMaxPushOutVelocity); }
 		void SetCollisionDepenetrationVelocity(const FRealSingle InVelocity) { GetEvolution()->GetCollisionConstraints().SetDepenetrationVelocity(InVelocity); }
 
