Fix crash in FillAutoInstanceMeshesInstancesIfNeeded if the AutoInstanceMeshIndex is missing while there's auto instanced mnesh assign to the asset

diff --git a/Engine/Source/Runtime/Experimental/GeometryCollectionEngine/Private/GeometryCollection/GeometryCollectionObject.cpp b/Engine/Source/Runtime/Experimental/GeometryCollectionEngine/Private/GeometryCollection/GeometryCollectionObject.cpp
--- a/Engine/Source/Runtime/Experimental/GeometryCollectionEngine/Private/GeometryCollection/GeometryCollectionObject.cpp
+++ b/Engine/Source/Runtime/Experimental/GeometryCollectionEngine/Private/GeometryCollection/GeometryCollectionObject.cpp
@@ -1283,19 +1283,26 @@ void UGeometryCollection::FillAutoInstanceMeshesInstancesIfNeeded()
 
 		const GeometryCollection::Facades::FCollectionInstancedMeshFacade InstancedMeshFacade(*GeometryCollection);
 
-		const int32 NumTransforms = GeometryCollection->Children.Num();
-		for (int32 TransformIndex = 0; TransformIndex < NumTransforms; TransformIndex++)
+		if (InstancedMeshFacade.IsValid())
 		{
-			// only applies to leaves nodes
-			if (GeometryCollection->Children[TransformIndex].Num() == 0)
+			const int32 NumTransforms = GeometryCollection->Children.Num();
+			for (int32 TransformIndex = 0; TransformIndex < NumTransforms; TransformIndex++)
 			{
-				const int32 AutoInstanceMeshIndex = InstancedMeshFacade.GetIndex(TransformIndex);
-				if (AutoInstanceMeshes.IsValidIndex(AutoInstanceMeshIndex))
+				// only applies to leaves nodes
+				if (GeometryCollection->Children[TransformIndex].Num() == 0)
 				{
-					AutoInstanceMeshes[AutoInstanceMeshIndex].NumInstances++;
+					const int32 AutoInstanceMeshIndex = InstancedMeshFacade.GetIndex(TransformIndex);
+					if (AutoInstanceMeshes.IsValidIndex(AutoInstanceMeshIndex))
+					{
+						AutoInstanceMeshes[AutoInstanceMeshIndex].NumInstances++;
+					}
 				}
 			}
 		}
+		else
+		{
+			UE_LOG(LogGeometryCollectionInternal, Warning, TEXT("Geometry collection has AutoInstanceMeshes meshes but no matching AutoInstanceMeshIndex, remove the instanced meshes or regenerate the asset"), *GetPathName());
+		}
 	}
 }
 
