Dataflow : Const accessors for cache keys.

diff --git a/Engine/Source/Runtime/Experimental/Dataflow/Core/Public/Dataflow/DataflowNodeParameters.h b/Engine/Source/Runtime/Experimental/Dataflow/Core/Public/Dataflow/DataflowNodeParameters.h
--- a/Engine/Source/Runtime/Experimental/Dataflow/Core/Public/Dataflow/DataflowNodeParameters.h
+++ b/Engine/Source/Runtime/Experimental/Dataflow/Core/Public/Dataflow/DataflowNodeParameters.h
@@ -134,7 +134,7 @@ namespace Dataflow
 
 		virtual FName GetType() const { return FContext::StaticType(); }
 
-		virtual int32 GetKeys(TSet<FContextCacheKey>& InKeys) { return DataStore.GetKeys(InKeys); }
+		virtual int32 GetKeys(TSet<FContextCacheKey>& InKeys) const { return DataStore.GetKeys(InKeys); }
 
 		template<class T>
 		const T* AsType() const
diff --git a/Engine/Source/Runtime/Experimental/Dataflow/Engine/Public/Dataflow/DataflowObjectInterface.h b/Engine/Source/Runtime/Experimental/Dataflow/Engine/Public/Dataflow/DataflowObjectInterface.h
--- a/Engine/Source/Runtime/Experimental/Dataflow/Engine/Public/Dataflow/DataflowObjectInterface.h
+++ b/Engine/Source/Runtime/Experimental/Dataflow/Engine/Public/Dataflow/DataflowObjectInterface.h
@@ -30,7 +30,7 @@ namespace Dataflow
 
 		~TEngineContext(){}
 
-		int32 GetKeys(TSet<FContextCacheKey>& InKeys) { return Base::GetKeys(InKeys); }
+		int32 GetKeys(TSet<FContextCacheKey>& InKeys) const { return Base::GetKeys(InKeys); }
 
 		TUniquePtr<FContextCacheElementBase>* GetBaseData(FContextCacheKey Key) { return Base::GetDataImpl(Key); }
 
