Iris
- Fix iris FastArrayReplication bug introdued with CL 31431819

diff --git a/Engine/Source/Runtime/Experimental/Iris/Core/Public/Iris/ReplicationSystem/Private/FastArrayReplicationFragmentInternal.h b/Engine/Source/Runtime/Experimental/Iris/Core/Public/Iris/ReplicationSystem/Private/FastArrayReplicationFragmentInternal.h
--- a/Engine/Source/Runtime/Experimental/Iris/Core/Public/Iris/ReplicationSystem/Private/FastArrayReplicationFragmentInternal.h
+++ b/Engine/Source/Runtime/Experimental/Iris/Core/Public/Iris/ReplicationSystem/Private/FastArrayReplicationFragmentInternal.h
@@ -58,7 +58,7 @@ struct FFastArrayReplicationFragmentHelper
 {
 	/** Rebuild IndexMap for FastArrraySerializer */
 	template <typename FastArrayType, typename ItemArrayType>
-	static void ConditionalRebuildItemMap(FastArrayType& ArraySerializer, const ItemArrayType& Items);
+	static void ConditionalRebuildItemMap(FastArrayType& ArraySerializer, const ItemArrayType& Items, bool bForceRebuild);
 
 	/** Apply received state and try to behave like current FastArrays */
 	template <typename FastArrayType, typename ItemArrayType>
@@ -168,11 +168,11 @@ protected:
 };
 
 template <typename FastArrayType, typename ItemArrayType>
-void FFastArrayReplicationFragmentHelper::ConditionalRebuildItemMap(FastArrayType& ArraySerializer, const ItemArrayType& Items)
+void FFastArrayReplicationFragmentHelper::ConditionalRebuildItemMap(FastArrayType& ArraySerializer, const ItemArrayType& Items, bool bForceRebuild)
 {
 	typedef typename ItemArrayType::ElementType ItemType;
 
-	if (ArraySerializer.ItemMap.Num() != Items.Num())
+	if (bForceRebuild || ArraySerializer.ItemMap.Num() != Items.Num())
 	{
 		UE_LOG(LogNetFastTArray, Verbose, TEXT("FastArrayDeltaSerialize: Recreating Items map. Items.Num: %d Map.Num: %d"), Items.Num(), ArraySerializer.ItemMap.Num());
 
@@ -199,11 +199,6 @@ void FFastArrayReplicationFragmentHelper::ApplyReplicatedState(FastArrayType* Ds
 
 	UE_LOG(LogNetFastTArray, Log, TEXT("FFastArrayReplicationFragmentHelper::ApplyReplicatedState for %s"), Context.Descriptor->DebugName->Name);
 
-	// We need to rebuild our maps for both target array and incoming data
-	// Can optimize this later
-	ConditionalRebuildItemMap(*DstArraySerializer, *DstWrappedArray);
-	ConditionalRebuildItemMap(*SrcArraySerializer, *SrcWrappedArray);
-
 	const uint32* ChangeMaskData = Context.StateBufferData.ChangeMaskData;
 	FNetBitArrayView MemberChangeMask = MakeNetBitArrayView(ChangeMaskData, Context.Descriptor->ChangeMaskBitCount);
 
@@ -214,6 +209,14 @@ void FFastArrayReplicationFragmentHelper::ApplyReplicatedState(FastArrayType* Ds
 	const uint32 ChangeMaskBitOffset = MemberChangeMaskDescriptor.BitOffset + FIrisFastArraySerializer::IrisFastArrayChangeMaskBitOffset;
 	const uint32 ChangeMaskBitCount = MemberChangeMaskDescriptor.BitCount - FIrisFastArraySerializer::IrisFastArrayChangeMaskBitOffset;
 
+	// Force rebuild if the array has been modified
+	const bool bForceRebuildItemMap = MemberChangeMask.GetBit(0);
+
+	// We need to rebuild our maps for both target array and incoming data
+	// Can optimize this later
+	ConditionalRebuildItemMap(*DstArraySerializer, *DstWrappedArray, false);
+	ConditionalRebuildItemMap(*SrcArraySerializer, *SrcWrappedArray, bForceRebuildItemMap);
+
 	// Find removed elements in received data, that is elements that exist in old map but not in new map
 	TArray<int32> RemovedIndices;
 	{
