Iris
- Properly manage partially resolved refererences to only include them in applied state when they are resolved.

diff --git a/Engine/Source/Runtime/Experimental/Iris/Core/Private/Iris/ReplicationSystem/ReplicationReader.cpp b/Engine/Source/Runtime/Experimental/Iris/Core/Private/Iris/ReplicationSystem/ReplicationReader.cpp
--- a/Engine/Source/Runtime/Experimental/Iris/Core/Private/Iris/ReplicationSystem/ReplicationReader.cpp
+++ b/Engine/Source/Runtime/Experimental/Iris/Core/Private/Iris/ReplicationSystem/ReplicationReader.cpp
@@ -1035,7 +1035,7 @@ ErrorHandling:
 }
 
 // Update reference tracking maps for the current object. It is assumed the ObjectReferenceTracker do no include duplicates for a given key.
-void FReplicationReader::UpdateObjectReferenceTracking(FReplicatedObjectInfo* ReplicationInfo, FNetBitArrayView ChangeMask, bool bIncludeInitState, const FObjectReferenceTracker& NewUnresolvedReferences, const FObjectReferenceTracker& NewMappedDynamicReferences)
+void FReplicationReader::UpdateObjectReferenceTracking(FReplicatedObjectInfo* ReplicationInfo, FNetBitArrayView ChangeMask, bool bIncludeInitState, FResolvedNetRefHandlesArray& OutNewResolvedRefHandles, const FObjectReferenceTracker& NewUnresolvedReferences, const FObjectReferenceTracker& NewMappedDynamicReferences)
 {
 	IRIS_PROFILER_SCOPE(FReplicationReader_UpdateObjectReferenceTracking)
 
@@ -1095,6 +1095,9 @@ void FReplicationReader::UpdateObjectReferenceTracking(FReplicatedObjectInfo* Re
 		{
 			if (!NewUnresolvedSet.Contains(Handle))
 			{
+				// Store new resolved handles so we can update partially resolved references properly
+				OutNewResolvedRefHandles.Add(Handle);
+
 				// Remove from tracking
 				UnresolvedHandleToDependents.RemoveSingle(Handle, OwnerInternalIndex);
 				UE_LOG(LogIris, Verbose, TEXT("FReplicationReader::UpdateObjectReferenceTracking Removing unresolved reference %s for %s"), ToCStr(Handle.ToString()), ToCStr(NetRefHandleManager->GetNetRefHandleFromInternalIndex(OwnerInternalIndex).ToString()));
@@ -1273,7 +1276,7 @@ void FReplicationReader::CleanupReferenceTracking(FReplicatedObjectInfo* ObjectI
 	ObjectsWithAttachmentPendingResolve.Remove(ObjectIndex);
 }
 
-void FReplicationReader::BuildUnresolvedChangeMaskAndUpdateObjectReferenceTracking(const FResolveAndCollectUnresolvedAndResolvedReferenceCollector& Collector, FNetBitArrayView CollectorChangeMask, FReplicatedObjectInfo* ReplicationInfo, FNetBitArrayView& OutUnresolvedChangeMask)
+void FReplicationReader::BuildUnresolvedChangeMaskAndUpdateObjectReferenceTracking(const FResolveAndCollectUnresolvedAndResolvedReferenceCollector& Collector, FNetBitArrayView CollectorChangeMask, FReplicatedObjectInfo* ReplicationInfo, FNetBitArrayView& OutUnresolvedChangeMask, FResolvedNetRefHandlesArray& OutNewResolvedRefHandles)
 {
 	OutUnresolvedChangeMask.Reset();
 	bool bHasUnresolvedInitReferences = false;
@@ -1305,8 +1308,8 @@ void FReplicationReader::BuildUnresolvedChangeMaskAndUpdateObjectReferenceTracki
 		}
 	}
 
-	// Update object specific 
-	UpdateObjectReferenceTracking(ReplicationInfo, CollectorChangeMask, Collector.IsInitStateIncluded(), UnresolvedReferences, MappedDynamicReferences);
+	// Update object specific
+	UpdateObjectReferenceTracking(ReplicationInfo, CollectorChangeMask, Collector.IsInitStateIncluded(), OutNewResolvedRefHandles, UnresolvedReferences, MappedDynamicReferences);
 }
 
 void FReplicationReader::ResolveAndDispatchUnresolvedReferencesForObject(FNetSerializationContext& Context, uint32 InternalIndex)
@@ -1349,25 +1352,36 @@ void FReplicationReader::ResolveAndDispatchUnresolvedReferencesForObject(FNetSer
 		FResolveAndCollectUnresolvedAndResolvedReferenceCollector Collector;
 		Collector.CollectReferences(*ObjectReferenceCache, ResolveContext, ReplicationInfo->bHasUnresolvedInitialReferences, &UnresolvedChangeMask, ObjectData.ReceiveStateBuffer, ObjectData.Protocol);
 
+		// We need to track previously unresolved NetRefHandles that now are resolvable
+		FResolvedNetRefHandlesArray NewResolvedRefHandles;
+
 		// Build UnresolvedChangeMask from collected data and update replication info
-		BuildUnresolvedChangeMaskAndUpdateObjectReferenceTracking(Collector, TempUnresolvedChangeMask, ReplicationInfo, UnresolvedChangeMask);
+		BuildUnresolvedChangeMaskAndUpdateObjectReferenceTracking(Collector, TempUnresolvedChangeMask, ReplicationInfo, UnresolvedChangeMask, NewResolvedRefHandles);
 
-		// Repurpose temp changemask for members that has resolved references.
+		// Re-purpose temp changemask for members that has resolved references.
 		TempChangeMask.Reset();
 		FNetBitArrayView ResolvedChangeMask = TempChangeMask;
 
+		// Merge in partially resolved changes
 		bool bHasResolvedInitReferences = false;
-		for (const auto& RefInfo : Collector.GetResolvedReferences())
+		if (NewResolvedRefHandles.Num())
 		{
-			const FNetSerializerChangeMaskParam& ChangeMaskInfo = RefInfo.ChangeMaskInfo;
-			if (ChangeMaskInfo.BitCount)
-			{
-				ResolvedChangeMask.SetBit(ChangeMaskInfo.BitOffset);
-			}
-			else
+			for (const FNetReferenceCollector::FReferenceInfo& ReferenceInfo : Collector.GetResolvedReferences())
 			{
-				// If we had old unresolved init dependencies we need to include the init state when we update references
-				bHasResolvedInitReferences = bOldHasUnresolvedInitReferences;
+				const FNetRefHandle& MatchRefHandle = ReferenceInfo.Reference.GetRefHandle();
+				if (NewResolvedRefHandles.ContainsByPredicate([&MatchRefHandle](const FNetRefHandle& RefHandle) { return RefHandle == MatchRefHandle;} ))
+				{
+					const FNetSerializerChangeMaskParam& ChangeMaskInfo = ReferenceInfo.ChangeMaskInfo;
+					if (ChangeMaskInfo.BitCount)
+					{
+						ResolvedChangeMask.SetBit(ChangeMaskInfo.BitOffset);
+					}
+					else
+					{
+						// If we had old unresolved init dependencies we need to include the init state when we update references
+						bHasResolvedInitReferences = bOldHasUnresolvedInitReferences;
+					}
+				}
 			}
 		}
 
@@ -1598,7 +1612,10 @@ void FReplicationReader::DispatchStateData(FNetSerializationContext& Context)
 					PrevUnresolvedChangeMask.Copy(UnresolvedChangeMask);
 				}
 
-				BuildUnresolvedChangeMaskAndUpdateObjectReferenceTracking(Collector, ChangeMaskForResolve, ReplicationInfo, UnresolvedChangeMask);
+				// We need to track previously unresolved NetRefHandles that now are resolvable
+				FResolvedNetRefHandlesArray NewResolvedRefHandles;
+
+				BuildUnresolvedChangeMaskAndUpdateObjectReferenceTracking(Collector, ChangeMaskForResolve, ReplicationInfo, UnresolvedChangeMask, NewResolvedRefHandles);
 			
 				// Allow resolved changes to be part of the state to be applied.
 				if (bMergeResolvedReferencesWithChangeMask)
@@ -1607,16 +1624,18 @@ void FReplicationReader::DispatchStateData(FNetSerializationContext& Context)
 					ChangeMask.CombineMultiple(FNetBitArrayView::OrOp, PrevUnresolvedChangeMask, FNetBitArrayView::AndNotOp, UnresolvedChangeMask);
 
 					// Merge in partially resolved changes
-					for (const FNetReferenceCollector::FReferenceInfo& ReferenceInfo : Collector.GetResolvedReferences())
+					if (NewResolvedRefHandles.Num())
 					{
-						ChangeMask.SetBit(ReferenceInfo.ChangeMaskInfo.BitOffset);
+						for (const FNetReferenceCollector::FReferenceInfo& ReferenceInfo : Collector.GetResolvedReferences())
+						{
+							const FNetRefHandle& MatchRefHandle = ReferenceInfo.Reference.GetRefHandle();
+							if (NewResolvedRefHandles.ContainsByPredicate([&MatchRefHandle](const FNetRefHandle& RefHandle) { return RefHandle == MatchRefHandle;} ))
+							{
+								ChangeMask.SetBit(ReferenceInfo.ChangeMaskInfo.BitOffset);
+							}
+						}
 					}
 				}
-
-				// $IRIS: $TODO: For now we always apply, even if we cannot resolve all references for a property
-				// We need to investigate how this is handled best as we do not want to prevent arrays(fastarrays) from applying data just because a single element wont resolve?
-				// Mask off any unresolvable states before we dispatch state data
-				//ChangeMask.Combine(UnresolvedChangeMask, FNetBitArrayView::AndNotOp);
 			}
 
 			// Apply state data
diff --git a/Engine/Source/Runtime/Experimental/Iris/Core/Private/Iris/ReplicationSystem/ReplicationReader.h b/Engine/Source/Runtime/Experimental/Iris/Core/Private/Iris/ReplicationSystem/ReplicationReader.h
--- a/Engine/Source/Runtime/Experimental/Iris/Core/Private/Iris/ReplicationSystem/ReplicationReader.h
+++ b/Engine/Source/Runtime/Experimental/Iris/Core/Private/Iris/ReplicationSystem/ReplicationReader.h
@@ -68,6 +68,7 @@ private:
 		FakeInitChangeMaskOffset = 0xFFFFFFFFU,
 	};
 	typedef TMultiMap<uint32, FNetRefHandle> FObjectReferenceTracker;
+	typedef TArray<FNetRefHandle, TInlineAllocator<32>> FResolvedNetRefHandlesArray;
 
 	struct FReplicatedObjectInfo
 	{
@@ -179,13 +180,13 @@ private:
 	FReplicatedObjectInfo* GetReplicatedObjectInfo(uint32 InternalIndex);
 
 	// Update reference tracking maps for the current object
-	void UpdateObjectReferenceTracking(FReplicatedObjectInfo* ReplicationInfo, FNetBitArrayView ChangeMask, bool bIncludeInitState, const FObjectReferenceTracker& NewUnresolvedReferences, const FObjectReferenceTracker& NewMappedDynamicReferences);
+	void UpdateObjectReferenceTracking(FReplicatedObjectInfo* ReplicationInfo, FNetBitArrayView ChangeMask, bool bIncludeInitState, FResolvedNetRefHandlesArray& OutNewResolvedRefHandles, const FObjectReferenceTracker& NewUnresolvedReferences, const FObjectReferenceTracker& NewMappedDynamicReferences);
 
 	// Remove all references for object
 	void CleanupReferenceTracking(FReplicatedObjectInfo* ObjectInfo);
 
 	// Update ReplicationInfo and OutUnresolvedChangeMask based on data collected by the Collector
-	void BuildUnresolvedChangeMaskAndUpdateObjectReferenceTracking(const FResolveAndCollectUnresolvedAndResolvedReferenceCollector& Collector, FNetBitArrayView CollectorChangeMask, FReplicatedObjectInfo* ReplicationInfo, FNetBitArrayView& OutUnresolvedChangeMask);
+	void BuildUnresolvedChangeMaskAndUpdateObjectReferenceTracking(const FResolveAndCollectUnresolvedAndResolvedReferenceCollector& Collector, FNetBitArrayView CollectorChangeMask, FReplicatedObjectInfo* ReplicationInfo, FNetBitArrayView& OutUnresolvedChangeMask, FResolvedNetRefHandlesArray& OutNewResolvedRefHandles);
 
 	void RemoveUnresolvedObjectReferenceInReplicationInfo(FReplicatedObjectInfo* ReplicationInfo, FNetRefHandle Handle);
 	void RemoveResolvedObjectReferenceInReplicationInfo(FReplicatedObjectInfo* ReplicationInfo, FNetRefHandle Handle);
