Constraints: avoid re-entering the evaluation graph while flushing

#rnx

diff --git a/Engine/Plugins/Animation/ControlRig/Source/ControlRig/Private/Constraints/ControlRigTransformableHandle.cpp b/Engine/Plugins/Animation/ControlRig/Source/ControlRig/Private/Constraints/ControlRigTransformableHandle.cpp
--- a/Engine/Plugins/Animation/ControlRig/Source/ControlRig/Private/Constraints/ControlRigTransformableHandle.cpp
+++ b/Engine/Plugins/Animation/ControlRig/Source/ControlRig/Private/Constraints/ControlRigTransformableHandle.cpp
@@ -433,7 +433,10 @@ void UTransformableControlHandle::OnControlModified(
 
 		if (InControl->GetFName() == ControlName)
 		{	// if that handle is wrapping InControl
-			GetEvaluationBinding().bPendingFlush = true;
+			if (InContext.bConstraintUpdate)
+			{
+				GetEvaluationBinding().bPendingFlush = true;
+			}
 			Notify(Event);
 		}
 		else if (Event == EHandleEvent::GlobalTransformUpdated)
@@ -441,7 +444,10 @@ void UTransformableControlHandle::OnControlModified(
 			// the control being modified is not the one wrapped by this handle 
 			if (const FRigControlElement* Control = ControlRig->FindControl(ControlName))
 			{
-				GetEvaluationBinding().bPendingFlush = true;
+				if (InContext.bConstraintUpdate)
+				{
+					GetEvaluationBinding().bPendingFlush = true;
+				}
 				static constexpr  bool bPreTick = true;
 				Notify(EHandleEvent::UpperDependencyUpdated, bPreTick);
 			}
diff --git a/Engine/Source/Runtime/Experimental/Animation/Constraints/Private/ConstraintsEvaluationGraph.cpp b/Engine/Source/Runtime/Experimental/Animation/Constraints/Private/ConstraintsEvaluationGraph.cpp
--- a/Engine/Source/Runtime/Experimental/Animation/Constraints/Private/ConstraintsEvaluationGraph.cpp
+++ b/Engine/Source/Runtime/Experimental/Animation/Constraints/Private/ConstraintsEvaluationGraph.cpp
@@ -60,30 +60,38 @@ FConstraintNode* FConstraintsEvaluationGraph::FindNode(const TWeakObjectPtr<UTic
 
 void FConstraintsEvaluationGraph::FlushPendingEvaluations()
 {
-	if (!Nodes.IsEmpty())
+	if (State == InvalidData || State == Flushing)
 	{
-		if (ensure(!bDataInvalid))
-		{
-			if (ConstraintsEvaluationGraph::bDebugGraph)
-			{
-				UE_LOG(LogTemp, Warning, TEXT("Flush Constraints Evaluation Graph"));
-			}
-			
-			for (FConstraintNode& Node: Nodes)
-			{
-				if (Node.bMarkedForEvaluation)
-				{
-					Evaluate(&Node);
-				}
-			}
+		return;
+	}
+	
+	if (Nodes.IsEmpty())
+	{
+		return;
+	}
+	
+	if (ConstraintsEvaluationGraph::bDebugGraph)
+	{
+		UE_LOG(LogTemp, Warning, TEXT("Flush Constraints Evaluation Graph"));
+	}
 
-			const bool bHasNodesToEvaluate = Nodes.ContainsByPredicate([](const FConstraintNode& Node)
-			{
-				return Node.bMarkedForEvaluation;
-			});
-			ensure(!bHasNodesToEvaluate);
+	State = Flushing;
+	
+	for (FConstraintNode& Node: Nodes)
+	{
+		if (Node.bMarkedForEvaluation)
+		{
+			Evaluate(&Node);
 		}
 	}
+
+	const bool bHasNodesToEvaluate = Nodes.ContainsByPredicate([](const FConstraintNode& Node)
+	{
+		return Node.bMarkedForEvaluation;
+	});
+	ensure(!bHasNodesToEvaluate);
+
+	State = ReadyForEvaluation;
 }
 
 void FConstraintsEvaluationGraph::Rebuild()
@@ -176,14 +184,19 @@ void FConstraintsEvaluationGraph::Rebuild()
 		}
 	}
 	
-	bDataInvalid = false;
+	State = ReadyForEvaluation;
 
 	Dump();
 }
 
+bool FConstraintsEvaluationGraph::IsPendingEvaluation() const
+{
+	return State == PendingEvaluation;
+}
+
 void FConstraintsEvaluationGraph::Evaluate(const TWeakObjectPtr<UTickableConstraint>& InConstraint)
 {
-	if (bDataInvalid)
+	if (State == InvalidData)
 	{
 		Rebuild();
 	}
@@ -232,13 +245,13 @@ void FConstraintsEvaluationGraph::Evaluate(FConstraintNode* InNode)
 
 void FConstraintsEvaluationGraph::InvalidateData()
 {
-	bDataInvalid = true;
+	State = InvalidData;
 	Nodes.Empty();
 }
 
 void FConstraintsEvaluationGraph::MarkForEvaluation(const TWeakObjectPtr<UTickableConstraint>& InConstraint)
 {
-	if (bDataInvalid)
+	if (State == InvalidData)
 	{
 		Rebuild();
 	}
@@ -260,6 +273,11 @@ void FConstraintsEvaluationGraph::MarkForEvaluation(const TWeakObjectPtr<UTickab
 		}
 		
 		Node->bMarkedForEvaluation = true;
+		
+		if (State == ReadyForEvaluation)
+		{
+			State = PendingEvaluation;
+		}
 	}
 }
 
diff --git a/Engine/Source/Runtime/Experimental/Animation/Constraints/Private/ConstraintsManager.cpp b/Engine/Source/Runtime/Experimental/Animation/Constraints/Private/ConstraintsManager.cpp
--- a/Engine/Source/Runtime/Experimental/Animation/Constraints/Private/ConstraintsManager.cpp
+++ b/Engine/Source/Runtime/Experimental/Animation/Constraints/Private/ConstraintsManager.cpp
@@ -930,7 +930,12 @@ void FConstraintsManagerController::FlushEvaluationGraph() const
 	{
 		return;
 	}
-	SubSystem->GetEvaluationGraph(World).FlushPendingEvaluations();
+
+	FConstraintsEvaluationGraph& EvaluationGraph = SubSystem->GetEvaluationGraph(World);
+	if (EvaluationGraph.IsPendingEvaluation())
+	{
+		EvaluationGraph.FlushPendingEvaluations();
+	}
 }
 
 bool FConstraintsManagerController::DoesExistInAnyWorld(UTickableConstraint* InConstraint)
diff --git a/Engine/Source/Runtime/Experimental/Animation/Constraints/Public/ConstraintsEvaluationGraph.h b/Engine/Source/Runtime/Experimental/Animation/Constraints/Public/ConstraintsEvaluationGraph.h
--- a/Engine/Source/Runtime/Experimental/Animation/Constraints/Public/ConstraintsEvaluationGraph.h
+++ b/Engine/Source/Runtime/Experimental/Animation/Constraints/Public/ConstraintsEvaluationGraph.h
@@ -59,10 +59,21 @@ public:
 	/* Rebuilds the full graph taking constraints tick dependencies into account. */
 	void Rebuild();
 
+	/* Should the graph be evaluated ? */
+	bool IsPendingEvaluation() const;
+	
 	static bool UseEvaluationGraph();
 	
 private:
 
+	enum EGraphState
+	{
+		InvalidData = 0,	// Invalid data
+		ReadyForEvaluation,	// Internal data has been built and is ready for evaluation
+		PendingEvaluation,	// Some constraints has been marked for evaluation
+		Flushing			// The graph is currently being evaluated
+	};
+	
 	FConstraintsEvaluationGraph(FConstraintsEvaluationGraph&& InOther) = default;
 	FConstraintsEvaluationGraph(const FConstraintsEvaluationGraph& InOther) = delete;
 	FConstraintsEvaluationGraph& operator=(const FConstraintsEvaluationGraph& InOther) = delete;
@@ -88,5 +99,6 @@ private:
 	/* References to the actual stored data. */
 	FConstraintsInWorld& ConstraintsInWorld;
 	
-	bool bDataInvalid = true;
+	/* Current graph state. */
+	EGraphState	State = InvalidData;
 };
\ No newline at end of file
