Chaos cloth: switch VelocityField Apply to use ISPC and work in a single pass when using Softs::Evolution

In local tests, the single pass is only marginally more expensive than the first pass. On heavier assets, the cost per frame calculating velocity fields goes from 2ms -> .75ms.

diff --git a/Engine/Plugins/Experimental/ChaosCloth/Source/Chaos/Private/ChaosCloth/ChaosClothConstraints.cpp b/Engine/Plugins/Experimental/ChaosCloth/Source/Chaos/Private/ChaosCloth/ChaosClothConstraints.cpp
--- a/Engine/Plugins/Experimental/ChaosCloth/Source/Chaos/Private/ChaosCloth/ChaosClothConstraints.cpp
+++ b/Engine/Plugins/Experimental/ChaosCloth/Source/Chaos/Private/ChaosCloth/ChaosClothConstraints.cpp
@@ -1365,12 +1365,12 @@ void FClothConstraints::CreateVelocityAndPressureField(
 		// Always create velocity field--we allow turning it on via blueprints
 		constexpr Softs::FSolverReal WorldScale = (Softs::FSolverReal)100.;
 		VelocityAndPressureField = MakeShared<Softs::FVelocityAndPressureField>(
+			Evolution->GetSoftBodyParticles(ParticleRangeId),
 			&TriangleMesh,
 			ConfigProperties,
 			WeightMaps,
 			WorldScale
 			);
-		++NumPreSubstepInits;
 		++NumExternalForceRules;
 	}
 }
@@ -1428,11 +1428,6 @@ void FClothConstraints::CreateForceBasedRules()
 	}
 	if (VelocityAndPressureField)
 	{
-		RuleCreator.PreSubstepParallelInitRule(
-			[this](const Softs::FSolverParticlesRange& Particles, const Softs::FSolverReal Dt, const Softs::ESolverMode SolverMode)
-		{
-			VelocityAndPressureField->UpdateForces(Particles, Dt);
-		});
 		RuleCreator.AddExternalForceRule_Apply(VelocityAndPressureField.Get());
 
 		// TODO Linear System
@@ -1443,6 +1438,7 @@ void FClothConstraints::CreateForceBasedRules()
 		RuleCreator.AddExternalForceRule(
 			[this](Softs::FSolverParticlesRange& Particles, const Softs::FSolverReal Dt)
 		{
+			TRACE_CPUPROFILER_EVENT_SCOPE(PerSolverField);
 			const TArray<FVector>& LinearVelocities = PerSolverField->GetOutputResults(EFieldCommandOutputType::LinearVelocity);
 			const TArray<FVector>& LinearForces = PerSolverField->GetOutputResults(EFieldCommandOutputType::LinearForce);
 			const FVector* const LinearVelocitiesView = LinearVelocities.IsEmpty() ? nullptr :
diff --git a/Engine/Source/Runtime/Experimental/Chaos/Private/Chaos/VelocityField.cpp b/Engine/Source/Runtime/Experimental/Chaos/Private/Chaos/VelocityField.cpp
--- a/Engine/Source/Runtime/Experimental/Chaos/Private/Chaos/VelocityField.cpp
+++ b/Engine/Source/Runtime/Experimental/Chaos/Private/Chaos/VelocityField.cpp
@@ -2,10 +2,12 @@
 
 #include "Chaos/VelocityField.h"
 #include "HAL/IConsoleManager.h"
+#include "Chaos/GraphColoring.h"
 #if INTEL_ISPC
 #include "VelocityField.ispc.generated.h"
 
 static_assert(sizeof(ispc::FVector3f) == sizeof(Chaos::Softs::FSolverVec3), "sizeof(ispc::FVector3f) != sizeof(Chaos::Softs::FSolverVec3)");
+static_assert(sizeof(ispc::FVector2f) == sizeof(Chaos::Softs::FSolverVec2), "sizeof(ispc::FVector2f) != sizeof(Chaos::Softs::FSolverVec2)");
 static_assert(sizeof(ispc::FIntVector) == sizeof(Chaos::TVec3<int32>), "sizeof(ispc::FIntVector) != sizeof(Chaos::TVec3<int32>)");
 #endif
 
@@ -32,25 +34,40 @@ void FVelocityAndPressureField::SetProperties(
 
 	if (IsDragMutable(PropertyCollection))
 	{
-		const FSolverVec2 Drag(GetWeightedFloatDrag(PropertyCollection));
-		DragBase = FMath::Clamp(Drag[0], MinCoefficient, MaxCoefficient);
-		DragRange = FMath::Clamp(Drag[1], MinCoefficient, MaxCoefficient) - DragBase;
-
+		const FSolverVec2 WeightedValue(GetWeightedFloatDrag(PropertyCollection));
 		if (IsDragStringDirty(PropertyCollection))
 		{
-			bSetMultipliers = true;
+			const FString& WeightMapName = GetDragString(PropertyCollection);
+			Drag = FPBDFlatWeightMap(
+				WeightedValue.ClampAxes(MinCoefficient, MaxCoefficient),
+				Weightmaps.FindRef(WeightMapName),
+				TConstArrayView<TVec3<int32>>(Elements),
+				Offset,
+				NumParticles);
+		}
+		else
+		{
+			Drag.SetWeightedValue(WeightedValue.ClampAxes(MinCoefficient, MaxCoefficient));
 		}
 	}
 
 	if (IsLiftMutable(PropertyCollection))
 	{
-		const FSolverVec2 Lift(GetWeightedFloatLift(PropertyCollection));
-		LiftBase = FMath::Clamp(Lift[0], MinCoefficient, MaxCoefficient);
-		LiftRange = FMath::Clamp(Lift[1], MinCoefficient, MaxCoefficient) - LiftBase;
+		const FSolverVec2 WeightedValue(GetWeightedFloatLift(PropertyCollection));
 
 		if (IsLiftStringDirty(PropertyCollection))
 		{
-			bSetMultipliers = true;
+			const FString& WeightMapName = GetLiftString(PropertyCollection);
+			Lift = FPBDFlatWeightMap(
+				WeightedValue.ClampAxes(MinCoefficient, MaxCoefficient),
+				Weightmaps.FindRef(WeightMapName),
+				TConstArrayView<TVec3<int32>>(Elements),
+				Offset,
+				NumParticles);
+		}
+		else
+		{
+			Lift.SetWeightedValue(WeightedValue.ClampAxes(MinCoefficient, MaxCoefficient));
 		}
 	}
 
@@ -61,19 +78,21 @@ void FVelocityAndPressureField::SetProperties(
 
 	if (IsPressureMutable(PropertyCollection))
 	{
-		const FSolverVec2 Pressure(GetWeightedFloatPressure(PropertyCollection));
-		PressureBase = Pressure[0] / WorldScale;
-		PressureRange = Pressure[1] / WorldScale - PressureBase;
-
+		const FSolverVec2 WeightedValue(GetWeightedFloatPressure(PropertyCollection));
 		if (IsPressureStringDirty(PropertyCollection))
 		{
-			bSetMultipliers = true;
+			const FString& WeightMapName = GetPressureString(PropertyCollection);
+			Pressure = FPBDFlatWeightMap(
+				WeightedValue / WorldScale,
+				Weightmaps.FindRef(WeightMapName),
+				TConstArrayView<TVec3<int32>>(Elements),
+				Offset,
+				NumParticles);
+		}
+		else
+		{
+			Pressure.SetWeightedValue(WeightedValue / WorldScale);
 		}
-	}
-
-	if (bSetMultipliers)
-	{
-		SetMultipliers(PropertyCollection, Weightmaps);
 	}
 
 	// Update QuarterRho
@@ -94,19 +113,16 @@ void FVelocityAndPressureField::SetPropertiesAndWind(
 }
 
 void FVelocityAndPressureField::SetProperties(
-	const FSolverVec2& Drag,
-	const FSolverVec2& Lift,
+	const FSolverVec2& InDrag,
+	const FSolverVec2& InLift,
 	const FSolverReal FluidDensity,
-	const FSolverVec2& Pressure,
+	const FSolverVec2& InPressure,
 	FSolverReal WorldScale)
 {
-	DragBase = FMath::Clamp(Drag[0], MinCoefficient, MaxCoefficient);
-	DragRange = FMath::Clamp(Drag[1], MinCoefficient, MaxCoefficient) - DragBase;
-	LiftBase = FMath::Clamp(Lift[0], MinCoefficient, MaxCoefficient);
-	LiftRange = FMath::Clamp(Lift[1], MinCoefficient, MaxCoefficient) - LiftBase;
+	Drag.SetWeightedValue(InDrag.ClampAxes(MinCoefficient, MaxCoefficient));
+	Lift.SetWeightedValue(InLift.ClampAxes(MinCoefficient, MaxCoefficient));
+	Pressure.SetWeightedValue(InPressure / WorldScale);
 	Rho = FMath::Max(FluidDensity / FMath::Cube(WorldScale), (FSolverReal)0.);
-	PressureBase = Pressure[0] / WorldScale;
-	PressureRange = Pressure[1] / WorldScale - PressureBase;
 
 	constexpr FSolverReal OneQuarter = (FSolverReal)0.25;
 	QuarterRho = Rho * OneQuarter;
@@ -133,6 +149,7 @@ void FVelocityAndPressureField::SetGeometry(
 		FSolverVec2(GetWeightedFloatPressure(PropertyCollection, 0.f)),  // These getters also initialize the property indices, so keep before SetMultipliers
 		WorldScale);
 	SetMultipliers(PropertyCollection, Weightmaps);
+	ResetColor();
 }
 
 void FVelocityAndPressureField::SetGeometry(
@@ -143,6 +160,80 @@ void FVelocityAndPressureField::SetGeometry(
 {
 	SetGeometry(TriangleMesh);
 	SetMultipliers(DragMultipliers, LiftMultipliers, PressureMultipliers);
+	ResetColor();
+}
+
+
+void FVelocityAndPressureField::InitColor(const FSolverParticlesRange& InParticles)
+{
+#if INTEL_ISPC
+	const TArray<TArray<int32>> ConstraintsPerColor = FGraphColoring::ComputeGraphColoringParticlesOrRange(Elements, InParticles, Offset, Offset + NumParticles);
+	TArray<TVec3<int32>> ReorderedElements;
+	TArray<int32> OrigToReorderedIndices; // used to reorder stiffness indices
+	ReorderedElements.SetNumUninitialized(Elements.Num());
+	OrigToReorderedIndices.SetNumUninitialized(Elements.Num());
+	ConstraintsPerColorStartIndex.Reset(ConstraintsPerColor.Num() + 1);
+	int32 ReorderedIndex = 0;
+	for (const TArray<int32>& ConstraintsBatch : ConstraintsPerColor)
+	{
+		ConstraintsPerColorStartIndex.Add(ReorderedIndex);
+		for (const int32& BatchConstraint : ConstraintsBatch)
+		{
+			const int32 OrigIndex = BatchConstraint;
+			ReorderedElements[ReorderedIndex] = Elements[OrigIndex];
+			OrigToReorderedIndices[OrigIndex] = ReorderedIndex;
+			++ReorderedIndex;
+		}
+	}
+	ConstraintsPerColorStartIndex.Add(ReorderedIndex);
+
+	Elements = MoveTemp(ReorderedElements);
+	Lift.ReorderIndices(OrigToReorderedIndices);
+	Drag.ReorderIndices(OrigToReorderedIndices);
+	Pressure.ReorderIndices(OrigToReorderedIndices);
+#else
+	ResetColor();
+#endif
+}
+
+void FVelocityAndPressureField::ResetColor()
+{
+	ConstraintsPerColorStartIndex.Reset();
+	TArray<TVec2<int32>> ReorderedConstraints;
+}
+
+void FVelocityAndPressureField::SetGeometry(const FSolverParticlesRange& Particles, const FTriangleMesh* TriangleMesh)
+{
+	if (TriangleMesh)
+	{
+		const TArray<TVec3<int32>>& InElements = TriangleMesh->GetElements();
+		Offset = 0;
+		NumParticles = Particles.Size();
+
+		// Strip kinematic elements
+		PointToTriangleMap.Reset(NumParticles);
+		PointToTriangleMap.AddDefaulted(NumParticles);
+		Elements.Reset(InElements.Num());
+		for (const TVec3<int32>& Elem : InElements)
+		{
+			if (Particles.InvM(Elem[0]) != (FSolverReal)0.f || Particles.InvM(Elem[1]) != (FSolverReal)0.f || Particles.InvM(Elem[2]) != (FSolverReal)0.f)
+			{
+				const int32 ElemIndex = Elements.Add(Elem);
+				PointToTriangleMap[Elem[0] - Offset].Add(ElemIndex);
+				PointToTriangleMap[Elem[1] - Offset].Add(ElemIndex);
+				PointToTriangleMap[Elem[2] - Offset].Add(ElemIndex);
+			}
+		}
+		Forces.SetNumUninitialized(Elements.Num());
+	}
+	else
+	{
+		PointToTriangleMap.Reset();
+		Elements.Reset();
+		Offset = 0;
+		NumParticles = 0;
+		Forces.Reset();
+	}
 }
 
 void FVelocityAndPressureField::SetGeometry(const FTriangleMesh* TriangleMesh)
@@ -158,8 +249,8 @@ void FVelocityAndPressureField::SetGeometry(const FTriangleMesh* TriangleMesh)
 	}
 	else
 	{
-		PointToTriangleMap = TArrayView<TArray<int32>>();
-		Elements = TArrayView<TVector<int32, 3>>();
+		PointToTriangleMap.Reset();
+		Elements.Reset();
 		Offset = 0;
 		NumParticles = 0;
 		Forces.Reset();
@@ -186,41 +277,24 @@ void FVelocityAndPressureField::SetMultipliers(
 	const TConstArrayView<FRealSingle>& DragMultipliers,
 	const TConstArrayView<FRealSingle>& LiftMultipliers,
 	const TConstArrayView<FRealSingle>& PressureMultipliers)
-{
-	Multipliers.Reset();
-
-	const bool bHasDragMultipliers = DragMultipliers.Num() == NumParticles;
-	const bool bHasLiftMultipliers = LiftMultipliers.Num() == NumParticles;
-	const bool bHasPressureMultipliers = PressureMultipliers.Num() == NumParticles;
-
-	if (bHasDragMultipliers || bHasLiftMultipliers || bHasPressureMultipliers)
-	{
-		constexpr FSolverReal OneThird = (FSolverReal)1. / (FSolverReal)3.;
-
-		Multipliers.SetNumUninitialized(Elements.Num());
-
-		for (int32 ElementIndex = 0; ElementIndex < Elements.Num(); ++ElementIndex)
-		{
-			const TVec3<int32>& Element = Elements[ElementIndex];
-			const int32 I0 = Element[0] - Offset;
-			const int32 I1 = Element[1] - Offset;
-			const int32 I2 = Element[2] - Offset;
-
-			const FSolverReal DragMultiplier = bHasDragMultipliers ? (FSolverReal)(DragMultipliers[I0] + DragMultipliers[I1] + DragMultipliers[I2]) * OneThird : (FSolverReal)0.;
-			const FSolverReal LiftMultiplier = bHasLiftMultipliers ? (FSolverReal)(LiftMultipliers[I0] + LiftMultipliers[I1] + LiftMultipliers[I2]) * OneThird : (FSolverReal)0.;
-			const FSolverReal PressureMultiplier = bHasPressureMultipliers ? (FSolverReal)(PressureMultipliers[I0] + PressureMultipliers[I1] + PressureMultipliers[I2]) * OneThird : (FSolverReal)0.;
-
-			Multipliers[ElementIndex] = FSolverVec3(DragMultiplier, LiftMultiplier, PressureMultiplier);
-		}
-	}
+{	
+	const FSolverVec2 DragValues(Drag.GetLow(), Drag.GetHigh());
+	const FSolverVec2 LiftValues(Lift.GetLow(), Lift.GetHigh());
+	const FSolverVec2 PressureValues(Pressure.GetLow(), Pressure.GetHigh());
+	Drag = FPBDFlatWeightMap(DragValues, DragMultipliers, TConstArrayView<TVec3<int32>>(Elements), Offset, NumParticles);
+	Lift = FPBDFlatWeightMap(LiftValues, LiftMultipliers, TConstArrayView<TVec3<int32>>(Elements), Offset, NumParticles);
+	Pressure = FPBDFlatWeightMap(PressureValues, PressureMultipliers, TConstArrayView<TVec3<int32>>(Elements), Offset, NumParticles);
 }
 
-template<typename SolverParticlesOrRange>
-void FVelocityAndPressureField::UpdateForces(const SolverParticlesOrRange& InParticles, const FSolverReal /*Dt*/)
+void FVelocityAndPressureField::UpdateForces(const FSolverParticles& InParticles, const FSolverReal /*Dt*/)
 {
+	TRACE_CPUPROFILER_EVENT_SCOPE(FVelocityAndPressureField_UpdateForces);
 	const FSolverReal MaxVelocitySquared = (Private::VelocityFieldMaxVelocity > 0.f) ? FMath::Square((FSolverReal)Private::VelocityFieldMaxVelocity) : TNumericLimits<FSolverReal>::Max();
 
-	if (!Multipliers.Num())
+	const bool bDragHasMap = Drag.HasWeightMap();
+	const bool bLiftHasMap = Lift.HasWeightMap();
+	const bool bPressureHasMap = Pressure.HasWeightMap();
+	if (!bDragHasMap && !bLiftHasMap && !bPressureHasMap)
 	{
 #if INTEL_ISPC
 		if (bRealTypeCompatibleWithISPC && bChaos_VelocityField_ISPC_Enabled)
@@ -234,9 +308,9 @@ void FVelocityAndPressureField::UpdateForces(const SolverParticlesOrRange& InPar
 					(const ispc::FVector3f*)InParticles.XArray().GetData(),
 					(const ispc::FVector3f&)Velocity,
 					QuarterRho,
-					DragBase,
-					LiftBase,
-					PressureBase,
+					(FSolverReal)Drag,
+					(FSolverReal)Lift,
+					(FSolverReal)Pressure,
 					Elements.Num());
 			}
 			else
@@ -248,9 +322,9 @@ void FVelocityAndPressureField::UpdateForces(const SolverParticlesOrRange& InPar
 					(const ispc::FVector3f*)InParticles.XArray().GetData(),
 					(const ispc::FVector3f&)Velocity,
 					QuarterRho,
-					DragBase,
-					LiftBase,
-					PressureBase,
+					(FSolverReal)Drag,
+					(FSolverReal)Lift,
+					(FSolverReal)Pressure,
 					Elements.Num(),
 					MaxVelocitySquared);
 			}
@@ -262,14 +336,20 @@ void FVelocityAndPressureField::UpdateForces(const SolverParticlesOrRange& InPar
 			{
 				for (int32 ElementIndex = 0; ElementIndex < Elements.Num(); ++ElementIndex)
 				{
-					UpdateField(InParticles, ElementIndex, Velocity, DragBase, LiftBase, PressureBase);
+					UpdateField(InParticles, ElementIndex, Velocity,
+						(FSolverReal)Drag,
+						(FSolverReal)Lift,
+						(FSolverReal)Pressure);
 				}
 			}
 			else
 			{
 				for (int32 ElementIndex = 0; ElementIndex < Elements.Num(); ++ElementIndex)
 				{
-					UpdateField(InParticles, ElementIndex, Velocity, DragBase, LiftBase, PressureBase, MaxVelocitySquared);
+					UpdateField(InParticles, ElementIndex, Velocity,
+						(FSolverReal)Drag,
+						(FSolverReal)Lift,
+						(FSolverReal)Pressure, MaxVelocitySquared);
 				}
 			}
 		}
@@ -286,15 +366,17 @@ void FVelocityAndPressureField::UpdateForces(const SolverParticlesOrRange& InPar
 					(const ispc::FIntVector*)Elements.GetData(),
 					(const ispc::FVector3f*)InParticles.GetV().GetData(),
 					(const ispc::FVector3f*)InParticles.XArray().GetData(),
-					(const ispc::FVector3f*)Multipliers.GetData(),
 					(const ispc::FVector3f&)Velocity,
 					QuarterRho,
-					DragBase,
-					DragRange,
-					LiftBase,
-					LiftRange,
-					PressureBase,
-					PressureRange,
+					bDragHasMap,
+					reinterpret_cast<const ispc::FVector2f&>(Drag.GetOffsetRange()),
+					bDragHasMap ? Drag.GetMapValues().GetData() : nullptr,
+					bLiftHasMap,
+					reinterpret_cast<const ispc::FVector2f&>(Lift.GetOffsetRange()),
+					bLiftHasMap ? Lift.GetMapValues().GetData() : nullptr,
+					bPressureHasMap,
+					reinterpret_cast<const ispc::FVector2f&>(Pressure.GetOffsetRange()),
+					bPressureHasMap ? Pressure.GetMapValues().GetData() : nullptr,
 					Elements.Num());
 			}
 			else
@@ -304,15 +386,17 @@ void FVelocityAndPressureField::UpdateForces(const SolverParticlesOrRange& InPar
 					(const ispc::FIntVector*)Elements.GetData(),
 					(const ispc::FVector3f*)InParticles.GetV().GetData(),
 					(const ispc::FVector3f*)InParticles.XArray().GetData(),
-					(const ispc::FVector3f*)Multipliers.GetData(),
 					(const ispc::FVector3f&)Velocity,
 					QuarterRho,
-					DragBase,
-					DragRange,
-					LiftBase,
-					LiftRange,
-					PressureBase,
-					PressureRange,
+					bDragHasMap,
+					reinterpret_cast<const ispc::FVector2f&>(Drag.GetOffsetRange()),
+					bDragHasMap ? Drag.GetMapValues().GetData() : nullptr,
+					bLiftHasMap,
+					reinterpret_cast<const ispc::FVector2f&>(Lift.GetOffsetRange()),
+					bLiftHasMap ? Lift.GetMapValues().GetData() : nullptr,
+					bPressureHasMap,
+					reinterpret_cast<const ispc::FVector2f&>(Pressure.GetOffsetRange()),
+					bPressureHasMap ? Pressure.GetMapValues().GetData() : nullptr,
 					Elements.Num(),
 					MaxVelocitySquared);
 			}
@@ -324,10 +408,9 @@ void FVelocityAndPressureField::UpdateForces(const SolverParticlesOrRange& InPar
 			{
 				for (int32 ElementIndex = 0; ElementIndex < Elements.Num(); ++ElementIndex)
 				{
-					const FSolverVec3& Multiplier = Multipliers[ElementIndex];
-					const FSolverReal Cd = DragBase + DragRange * Multiplier[0];
-					const FSolverReal Cl = LiftBase + LiftRange * Multiplier[1];
-					const FSolverReal Cp = PressureBase + PressureRange * Multiplier[2];
+					const FSolverReal Cd = Drag.GetValue(ElementIndex);
+					const FSolverReal Cl = Lift.GetValue(ElementIndex);
+					const FSolverReal Cp = Pressure.GetValue(ElementIndex);
 
 					UpdateField(InParticles, ElementIndex, Velocity, Cd, Cl, Cp);
 				}
@@ -336,10 +419,9 @@ void FVelocityAndPressureField::UpdateForces(const SolverParticlesOrRange& InPar
 			{
 				for (int32 ElementIndex = 0; ElementIndex < Elements.Num(); ++ElementIndex)
 				{
-					const FSolverVec3& Multiplier = Multipliers[ElementIndex];
-					const FSolverReal Cd = DragBase + DragRange * Multiplier[0];
-					const FSolverReal Cl = LiftBase + LiftRange * Multiplier[1];
-					const FSolverReal Cp = PressureBase + PressureRange * Multiplier[2];
+					const FSolverReal Cd = Drag.GetValue(ElementIndex);
+					const FSolverReal Cl = Lift.GetValue(ElementIndex);
+					const FSolverReal Cp = Pressure.GetValue(ElementIndex);
 
 					UpdateField(InParticles, ElementIndex, Velocity, Cd, Cl, Cp, MaxVelocitySquared);
 				}
@@ -347,7 +429,187 @@ void FVelocityAndPressureField::UpdateForces(const SolverParticlesOrRange& InPar
 		}
 	}
 }
-template CHAOS_API void FVelocityAndPressureField::UpdateForces(const FSolverParticles& InParticles, const FSolverReal /*Dt*/);
-template CHAOS_API void FVelocityAndPressureField::UpdateForces(const FSolverParticlesRange& InParticles, const FSolverReal /*Dt*/);
+
+void FVelocityAndPressureField::Apply(FSolverParticlesRange& InParticles, const FSolverReal Dt) const
+{
+	TRACE_CPUPROFILER_EVENT_SCOPE(FVelocityAndPressureField_Apply);
+	const FSolverReal MaxVelocitySquared = (Private::VelocityFieldMaxVelocity > 0.f) ? FMath::Square((FSolverReal)Private::VelocityFieldMaxVelocity) : TNumericLimits<FSolverReal>::Max();
+
+	const bool bDragHasMap = Drag.HasWeightMap();
+	const bool bLiftHasMap = Lift.HasWeightMap();
+	const bool bPressureHasMap = Pressure.HasWeightMap();
+	if (!bDragHasMap && !bLiftHasMap && !bPressureHasMap)
+	{
+#if INTEL_ISPC
+		if (bRealTypeCompatibleWithISPC && bChaos_VelocityField_ISPC_Enabled && ConstraintsPerColorStartIndex.Num() > 1)
+		{
+			const int32 ConstraintColorNum = ConstraintsPerColorStartIndex.Num() - 1;
+			if (MaxVelocitySquared == TNumericLimits<FSolverReal>::Max())
+			{
+				for (int32 ConstraintColorIndex = 0; ConstraintColorIndex < ConstraintColorNum; ++ConstraintColorIndex)
+				{
+					const int32 ColorStart = ConstraintsPerColorStartIndex[ConstraintColorIndex];
+					const int32 ColorSize = ConstraintsPerColorStartIndex[ConstraintColorIndex + 1] - ColorStart;
+					ispc::UpdateAndApplyVelocityField(
+						(ispc::FVector3f*)InParticles.GetAcceleration().GetData(),
+						InParticles.GetInvM().GetData(),
+						(const ispc::FIntVector*)&Elements.GetData()[ColorStart],
+						(const ispc::FVector3f*)InParticles.GetV().GetData(),
+						(const ispc::FVector3f*)InParticles.XArray().GetData(),
+						(const ispc::FVector3f&)Velocity,
+						QuarterRho,
+						(FSolverReal)Drag,
+						(FSolverReal)Lift,
+						(FSolverReal)Pressure,
+						ColorSize);
+				}
+			}
+			else
+			{
+				for (int32 ConstraintColorIndex = 0; ConstraintColorIndex < ConstraintColorNum; ++ConstraintColorIndex)
+				{
+					const int32 ColorStart = ConstraintsPerColorStartIndex[ConstraintColorIndex];
+					const int32 ColorSize = ConstraintsPerColorStartIndex[ConstraintColorIndex + 1] - ColorStart;
+					ispc::UpdateAndClampVelocityAndApplyVelocityField(
+						(ispc::FVector3f*)InParticles.GetAcceleration().GetData(),
+						InParticles.GetInvM().GetData(),
+						(const ispc::FIntVector*)&Elements.GetData()[ColorStart],
+						(const ispc::FVector3f*)InParticles.GetV().GetData(),
+						(const ispc::FVector3f*)InParticles.XArray().GetData(),
+						(const ispc::FVector3f&)Velocity,
+						QuarterRho,
+						(FSolverReal)Drag,
+						(FSolverReal)Lift,
+						(FSolverReal)Pressure,
+						ColorSize,
+						MaxVelocitySquared);
+				}
+			}
+		}
+		else
+#endif
+		{
+			if (MaxVelocitySquared == TNumericLimits<FSolverReal>::Max())
+			{
+				for (int32 ElementIndex = 0; ElementIndex < Elements.Num(); ++ElementIndex)
+				{
+					const FSolverVec3 Force = CalculateForce(InParticles, ElementIndex, Velocity,
+						(FSolverReal)Drag,
+						(FSolverReal)Lift,
+						(FSolverReal)Pressure);
+					InParticles.Acceleration(Elements[ElementIndex][0]) += InParticles.InvM(Elements[ElementIndex][0]) * Force;
+					InParticles.Acceleration(Elements[ElementIndex][1]) += InParticles.InvM(Elements[ElementIndex][1]) * Force;
+					InParticles.Acceleration(Elements[ElementIndex][2]) += InParticles.InvM(Elements[ElementIndex][2]) * Force;
+				}
+			}
+			else
+			{
+				for (int32 ElementIndex = 0; ElementIndex < Elements.Num(); ++ElementIndex)
+				{
+					const FSolverVec3 Force = CalculateForce(InParticles, ElementIndex, Velocity,
+						(FSolverReal)Drag,
+						(FSolverReal)Lift,
+						(FSolverReal)Pressure, MaxVelocitySquared);
+					InParticles.Acceleration(Elements[ElementIndex][0]) += InParticles.InvM(Elements[ElementIndex][0]) * Force;
+					InParticles.Acceleration(Elements[ElementIndex][1]) += InParticles.InvM(Elements[ElementIndex][1]) * Force;
+					InParticles.Acceleration(Elements[ElementIndex][2]) += InParticles.InvM(Elements[ElementIndex][2]) * Force;
+				}
+			}
+		}
+	}
+	else
+	{
+#if INTEL_ISPC
+		if (bRealTypeCompatibleWithISPC && bChaos_VelocityField_ISPC_Enabled && ConstraintsPerColorStartIndex.Num() > 1)
+		{
+			const int32 ConstraintColorNum = ConstraintsPerColorStartIndex.Num() - 1;
+			if (MaxVelocitySquared == TNumericLimits<FSolverReal>::Max())
+			{
+				for (int32 ConstraintColorIndex = 0; ConstraintColorIndex < ConstraintColorNum; ++ConstraintColorIndex)
+				{
+					const int32 ColorStart = ConstraintsPerColorStartIndex[ConstraintColorIndex];
+					const int32 ColorSize = ConstraintsPerColorStartIndex[ConstraintColorIndex + 1] - ColorStart;
+					ispc::UpdateAndApplyVelocityFieldWithWeightMaps(
+						(ispc::FVector3f*)InParticles.GetAcceleration().GetData(),
+						InParticles.GetInvM().GetData(),
+						(const ispc::FIntVector*)&Elements.GetData()[ColorStart],
+						(const ispc::FVector3f*)InParticles.GetV().GetData(),
+						(const ispc::FVector3f*)InParticles.XArray().GetData(),
+						(const ispc::FVector3f&)Velocity,
+						QuarterRho,
+						bDragHasMap,
+						reinterpret_cast<const ispc::FVector2f&>(Drag.GetOffsetRange()),
+						bDragHasMap ? Drag.GetMapValues().GetData() : nullptr,
+						bLiftHasMap,
+						reinterpret_cast<const ispc::FVector2f&>(Lift.GetOffsetRange()),
+						bLiftHasMap ? Lift.GetMapValues().GetData() : nullptr,
+						bPressureHasMap,
+						reinterpret_cast<const ispc::FVector2f&>(Pressure.GetOffsetRange()),
+						bPressureHasMap ? Pressure.GetMapValues().GetData() : nullptr,
+						ColorSize);
+				}
+			}
+			else
+			{
+				for (int32 ConstraintColorIndex = 0; ConstraintColorIndex < ConstraintColorNum; ++ConstraintColorIndex)
+				{
+					const int32 ColorStart = ConstraintsPerColorStartIndex[ConstraintColorIndex];
+					const int32 ColorSize = ConstraintsPerColorStartIndex[ConstraintColorIndex + 1] - ColorStart;
+					ispc::UpdateAndClampVelocityAndApplyVelocityFieldWithWeightMaps(
+						(ispc::FVector3f*)InParticles.GetAcceleration().GetData(),
+						InParticles.GetInvM().GetData(),
+						(const ispc::FIntVector*)&Elements.GetData()[ColorStart],
+						(const ispc::FVector3f*)InParticles.GetV().GetData(),
+						(const ispc::FVector3f*)InParticles.XArray().GetData(),
+						(const ispc::FVector3f&)Velocity,
+						QuarterRho,
+						bDragHasMap,
+						reinterpret_cast<const ispc::FVector2f&>(Drag.GetOffsetRange()),
+						bDragHasMap ? Drag.GetMapValues().GetData() : nullptr,
+						bLiftHasMap,
+						reinterpret_cast<const ispc::FVector2f&>(Lift.GetOffsetRange()),
+						bLiftHasMap ? Lift.GetMapValues().GetData() : nullptr,
+						bPressureHasMap,
+						reinterpret_cast<const ispc::FVector2f&>(Pressure.GetOffsetRange()),
+						bPressureHasMap ? Pressure.GetMapValues().GetData() : nullptr,
+						ColorSize,
+						MaxVelocitySquared);
+				}
+			}
+		}
+		else
+#endif
+		{
+			if (MaxVelocitySquared == TNumericLimits<FSolverReal>::Max())
+			{
+				for (int32 ElementIndex = 0; ElementIndex < Elements.Num(); ++ElementIndex)
+				{
+					const FSolverReal Cd = Drag.GetValue(ElementIndex);
+					const FSolverReal Cl = Lift.GetValue(ElementIndex);
+					const FSolverReal Cp = Pressure.GetValue(ElementIndex);
+
+					const FSolverVec3 Force = CalculateForce(InParticles, ElementIndex, Velocity, Cd, Cl, Cp);
+					InParticles.Acceleration(Elements[ElementIndex][0]) += InParticles.InvM(Elements[ElementIndex][0]) * Force;
+					InParticles.Acceleration(Elements[ElementIndex][1]) += InParticles.InvM(Elements[ElementIndex][1]) * Force;
+					InParticles.Acceleration(Elements[ElementIndex][2]) += InParticles.InvM(Elements[ElementIndex][2]) * Force;
+				}
+			}
+			else
+			{
+				for (int32 ElementIndex = 0; ElementIndex < Elements.Num(); ++ElementIndex)
+				{
+					const FSolverReal Cd = Drag.GetValue(ElementIndex);
+					const FSolverReal Cl = Lift.GetValue(ElementIndex);
+					const FSolverReal Cp = Pressure.GetValue(ElementIndex);
+
+					const FSolverVec3 Force = CalculateForce(InParticles, ElementIndex, Velocity, Cd, Cl, Cp, MaxVelocitySquared);
+					InParticles.Acceleration(Elements[ElementIndex][0]) += InParticles.InvM(Elements[ElementIndex][0]) * Force;
+					InParticles.Acceleration(Elements[ElementIndex][1]) += InParticles.InvM(Elements[ElementIndex][1]) * Force;
+					InParticles.Acceleration(Elements[ElementIndex][2]) += InParticles.InvM(Elements[ElementIndex][2]) * Force;
+				}
+			}
+		}
+	}
+}
 
 }  // End namespace Chaos::Softs
diff --git a/Engine/Source/Runtime/Experimental/Chaos/Private/Chaos/VelocityField.ispc b/Engine/Source/Runtime/Experimental/Chaos/Private/Chaos/VelocityField.ispc
--- a/Engine/Source/Runtime/Experimental/Chaos/Private/Chaos/VelocityField.ispc
+++ b/Engine/Source/Runtime/Experimental/Chaos/Private/Chaos/VelocityField.ispc
@@ -82,15 +82,17 @@ export void UpdateFieldWithWeightMaps(uniform FVector3f Forces[],
 								const uniform FIntVector Elements[],
 								const uniform FVector3f V[],
 								const uniform FVector3f X[],
-								const uniform FVector3f Multipliers[],
 								const uniform FVector3f& Velocity,
 								const uniform float QuarterRho,
-								const uniform float DragBase,
-								const uniform float DragRange,
-								const uniform float LiftBase,
-								const uniform float LiftRange,
-								const uniform float PressureBase,
-								const uniform float PressureRange,
+								const uniform bool bDragHasMap,
+								const uniform FVector2f& DragOffsetRange,
+								const uniform float DragMapValues[],
+								const uniform bool bLiftHasMap,
+								const uniform FVector2f& LiftOffsetRange,
+								const uniform float LiftMapValues[],
+								const uniform bool bPressureHasMap,
+								const uniform FVector2f& PressureOffsetRange,
+								const uniform float PressureMapValues[],
 								const uniform int32 NumElements)
 {
 	varying FIntVector Element;
@@ -112,8 +114,6 @@ export void UpdateFieldWithWeightMaps(uniform FVector3f Forces[],
 		const FVector3f X1 = VectorGather(&X[Element.V[1]]);
 		const FVector3f X2 = VectorGather(&X[Element.V[2]]);
 
-		const FVector3f Multiplier = VectorGather(&Multipliers[ElementIndex]);
-
 		// Calculate the normal and the area of the surface exposed to the flow
 		FVector3f N = VectorCross(X1 - X0, X2 - X0);
 		const float DoubleArea = SafeNormalize(N);
@@ -130,9 +130,9 @@ export void UpdateFieldWithWeightMaps(uniform FVector3f Forces[],
 		const float VDotN = VectorDot(RelVelocity, N);
 		const float VSquare = VectorDot(RelVelocity, RelVelocity);
 
-		const float Cd = DragBase + DragRange * Multiplier.V[0];
-		const float Cl = LiftBase + LiftRange * Multiplier.V[1];
-		const float Cp = PressureBase + PressureRange * Multiplier.V[2];
+		const float Cd = bDragHasMap ? DragOffsetRange.V[0] + DragMapValues[ElementIndex] * DragOffsetRange.V[1] : DragOffsetRange.V[0];
+		const float Cl = bLiftHasMap ? LiftOffsetRange.V[0] + LiftMapValues[ElementIndex] * LiftOffsetRange.V[1] : LiftOffsetRange.V[0];
+		const float Cp = bPressureHasMap ? PressureOffsetRange.V[0] + PressureMapValues[ElementIndex] * PressureOffsetRange.V[1] : PressureOffsetRange.V[0];
 
 #if HW_GATHER_SUPPORTED == 1
 		if(ElementIndex + programCount < NumElements)
@@ -229,15 +229,17 @@ export void UpdateFieldWithWeightMapsAndClampVelocity(uniform FVector3f Forces[]
 								const uniform FIntVector Elements[],
 								const uniform FVector3f V[],
 								const uniform FVector3f X[],
-								const uniform FVector3f Multipliers[],
 								const uniform FVector3f& Velocity,
 								const uniform float QuarterRho,
-								const uniform float DragBase,
-								const uniform float DragRange,
-								const uniform float LiftBase,
-								const uniform float LiftRange,
-								const uniform float PressureBase,
-								const uniform float PressureRange,
+								const uniform bool bDragHasMap,
+								const uniform FVector2f& DragOffsetRange,
+								const uniform float DragMapValues[],
+								const uniform bool bLiftHasMap,
+								const uniform FVector2f& LiftOffsetRange,
+								const uniform float LiftMapValues[],
+								const uniform bool bPressureHasMap,
+								const uniform FVector2f& PressureOffsetRange,
+								const uniform float PressureMapValues[],
 								const uniform int32 NumElements,
 								const uniform float MaxVelocitySquared)
 {
@@ -260,8 +262,6 @@ export void UpdateFieldWithWeightMapsAndClampVelocity(uniform FVector3f Forces[]
 		const FVector3f X1 = VectorGather(&X[Element.V[1]]);
 		const FVector3f X2 = VectorGather(&X[Element.V[2]]);
 
-		const FVector3f Multiplier = VectorGather(&Multipliers[ElementIndex]);
-
 		// Calculate the normal and the area of the surface exposed to the flow
 		FVector3f N = VectorCross(X1 - X0, X2 - X0);
 		const float DoubleArea = SafeNormalize(N);
@@ -284,10 +284,10 @@ export void UpdateFieldWithWeightMapsAndClampVelocity(uniform FVector3f Forces[]
 		// Set the aerodynamic forces
 		const float VDotN = VectorDot(RelVelocity, N);
 		const float VSquare = VectorDot(RelVelocity, RelVelocity);
-
-		const float Cd = DragBase + DragRange * Multiplier.V[0];
-		const float Cl = LiftBase + LiftRange * Multiplier.V[1];
-		const float Cp = PressureBase + PressureRange * Multiplier.V[2];
+		
+		const float Cd = bDragHasMap ? DragOffsetRange.V[0] + DragMapValues[ElementIndex] * DragOffsetRange.V[1] : DragOffsetRange.V[0];
+		const float Cl = bLiftHasMap ? LiftOffsetRange.V[0] + LiftMapValues[ElementIndex] * LiftOffsetRange.V[1] : LiftOffsetRange.V[0];
+		const float Cp = bPressureHasMap ? PressureOffsetRange.V[0] + PressureMapValues[ElementIndex] * PressureOffsetRange.V[1] : PressureOffsetRange.V[0];
 
 #if HW_GATHER_SUPPORTED == 1
 		if(ElementIndex + programCount < NumElements)
@@ -306,3 +306,288 @@ export void UpdateFieldWithWeightMapsAndClampVelocity(uniform FVector3f Forces[]
 		Offset += programCount;
 	}
 }
+
+export void UpdateAndApplyVelocityField(uniform FVector3f Acceleration[],
+								const uniform float InvM[],
+								const uniform FIntVector Elements[],
+								const uniform FVector3f V[],
+								const uniform FVector3f X[],
+								const uniform FVector3f& Velocity,
+								const uniform float QuarterRho,
+								const uniform float Cd,
+								const uniform float Cl,
+								const uniform float Cp,
+								const uniform int32 NumElements)
+{
+
+	foreach(ElementIndex = 0 ... NumElements)
+	{
+		const FIntVector Element = VectorLoad(&Elements[extract(ElementIndex,0)]);
+		const varying int32 Index0 = Element.V[0];
+		const varying int32 Index1 = Element.V[1];
+		const varying int32 Index2 = Element.V[2];
+
+		const FVector3f X0 = VectorGather(&X[Index0]);
+		const FVector3f X1 = VectorGather(&X[Index1]);
+		const FVector3f X2 = VectorGather(&X[Index2]);
+
+		// Calculate the normal and the area of the surface exposed to the flow
+		FVector3f N = VectorCross(X1 - X0, X2 - X0);
+		const float DoubleArea = SafeNormalize(N);
+
+		const FVector3f V0 = VectorGather(&V[Index0]);
+		const FVector3f V1 = VectorGather(&V[Index1]);
+		const FVector3f V2 = VectorGather(&V[Index2]);
+
+		// Calculate the direction and the relative velocity of the triangle to the flow
+		const FVector3f SurfaceVelocity = OneThird * (V0 + V1 + V2);
+		const FVector3f RelVelocity = Velocity - SurfaceVelocity;
+
+		// Set the aerodynamic forces
+		const float VDotN = VectorDot(RelVelocity, N);
+		const float VSquare = VectorDot(RelVelocity, RelVelocity);
+
+		const FVector3f EvenForce = (Cd - Cl) * VDotN * RelVelocity + Cl * VSquare * N;
+		const FVector3f OddForce = (Cl - Cd) * VDotN * RelVelocity - Cl * VSquare * N;
+
+		const FVector3f Force = QuarterRho * DoubleArea * VectorSelect(VDotN >= 0, EvenForce, OddForce) - DoubleArea * 0.5 * Cp * N; // N points in the opposite direction of the actual mesh normals
+		
+		#pragma ignore warning(perf)
+		const float InvM0 = InvM[Index0];
+		#pragma ignore warning(perf)
+		const float InvM1 = InvM[Index1];
+		#pragma ignore warning(perf)
+		const float InvM2 = InvM[Index2];
+
+		const FVector3f Acceleration0 = VectorGather(&Acceleration[Index0]);
+		const FVector3f Acceleration1 = VectorGather(&Acceleration[Index1]);
+		const FVector3f Acceleration2 = VectorGather(&Acceleration[Index2]);
+
+		VectorScatter(&Acceleration[Index0], Acceleration0 + Force * InvM0);
+		VectorScatter(&Acceleration[Index1], Acceleration1 + Force * InvM1);
+		VectorScatter(&Acceleration[Index2], Acceleration2 + Force * InvM2);
+	}
+}
+
+export void UpdateAndClampVelocityAndApplyVelocityField(uniform FVector3f Acceleration[],
+								const uniform float InvM[],
+								const uniform FIntVector Elements[],
+								const uniform FVector3f V[],
+								const uniform FVector3f X[],
+								const uniform FVector3f& Velocity,
+								const uniform float QuarterRho,
+								const uniform float Cd,
+								const uniform float Cl,
+								const uniform float Cp,
+								const uniform int32 NumElements,
+								const uniform float MaxVelocitySquared)
+{
+
+	foreach(ElementIndex = 0 ... NumElements)
+	{
+		const FIntVector Element = VectorLoad(&Elements[extract(ElementIndex,0)]);
+		const varying int32 Index0 = Element.V[0];
+		const varying int32 Index1 = Element.V[1];
+		const varying int32 Index2 = Element.V[2];
+
+		const FVector3f X0 = VectorGather(&X[Index0]);
+		const FVector3f X1 = VectorGather(&X[Index1]);
+		const FVector3f X2 = VectorGather(&X[Index2]);
+
+		// Calculate the normal and the area of the surface exposed to the flow
+		FVector3f N = VectorCross(X1 - X0, X2 - X0);
+		const float DoubleArea = SafeNormalize(N);
+
+		const FVector3f V0 = VectorGather(&V[Index0]);
+		const FVector3f V1 = VectorGather(&V[Index1]);
+		const FVector3f V2 = VectorGather(&V[Index2]);
+
+		// Calculate the direction and the relative velocity of the triangle to the flow
+		const FVector3f SurfaceVelocity = OneThird * (V0 + V1 + V2);
+		FVector3f RelVelocity = Velocity - SurfaceVelocity;
+
+		// Clamp the relative velocity
+		const float RelVelocitySquared = VectorSizeSquared(RelVelocity);
+		if (RelVelocitySquared > MaxVelocitySquared)
+		{
+			RelVelocity = RelVelocity * sqrt(MaxVelocitySquared / RelVelocitySquared);
+		}
+
+		// Set the aerodynamic forces
+		const float VDotN = VectorDot(RelVelocity, N);
+		const float VSquare = VectorDot(RelVelocity, RelVelocity);
+
+		const FVector3f EvenForce = (Cd - Cl) * VDotN * RelVelocity + Cl * VSquare * N;
+		const FVector3f OddForce = (Cl - Cd) * VDotN * RelVelocity - Cl * VSquare * N;
+
+		const FVector3f Force = QuarterRho * DoubleArea * VectorSelect(VDotN >= 0, EvenForce, OddForce) - DoubleArea * 0.5 * Cp * N; // N points in the opposite direction of the actual mesh normals
+		
+		#pragma ignore warning(perf)
+		const float InvM0 = InvM[Index0];
+		#pragma ignore warning(perf)
+		const float InvM1 = InvM[Index1];
+		#pragma ignore warning(perf)
+		const float InvM2 = InvM[Index2];
+
+		const FVector3f Acceleration0 = VectorGather(&Acceleration[Index0]);
+		const FVector3f Acceleration1 = VectorGather(&Acceleration[Index1]);
+		const FVector3f Acceleration2 = VectorGather(&Acceleration[Index2]);
+
+		VectorScatter(&Acceleration[Index0], Acceleration0 + Force * InvM0);
+		VectorScatter(&Acceleration[Index1], Acceleration1 + Force * InvM1);
+		VectorScatter(&Acceleration[Index2], Acceleration2 + Force * InvM2);
+	}
+}
+
+
+export void UpdateAndApplyVelocityFieldWithWeightMaps(uniform FVector3f Acceleration[],
+								const uniform float InvM[],
+								const uniform FIntVector Elements[],
+								const uniform FVector3f V[],
+								const uniform FVector3f X[],
+								const uniform FVector3f& Velocity,
+								const uniform float QuarterRho,
+								const uniform bool bDragHasMap,
+								const uniform FVector2f& DragOffsetRange,
+								const uniform float DragMapValues[],
+								const uniform bool bLiftHasMap,
+								const uniform FVector2f& LiftOffsetRange,
+								const uniform float LiftMapValues[],
+								const uniform bool bPressureHasMap,
+								const uniform FVector2f& PressureOffsetRange,
+								const uniform float PressureMapValues[],
+								const uniform int32 NumElements)
+{
+
+	foreach(ElementIndex = 0 ... NumElements)
+	{
+		const FIntVector Element = VectorLoad(&Elements[extract(ElementIndex,0)]);
+		const varying int32 Index0 = Element.V[0];
+		const varying int32 Index1 = Element.V[1];
+		const varying int32 Index2 = Element.V[2];
+
+		const FVector3f X0 = VectorGather(&X[Index0]);
+		const FVector3f X1 = VectorGather(&X[Index1]);
+		const FVector3f X2 = VectorGather(&X[Index2]);
+
+		// Calculate the normal and the area of the surface exposed to the flow
+		FVector3f N = VectorCross(X1 - X0, X2 - X0);
+		const float DoubleArea = SafeNormalize(N);
+
+		const FVector3f V0 = VectorGather(&V[Index0]);
+		const FVector3f V1 = VectorGather(&V[Index1]);
+		const FVector3f V2 = VectorGather(&V[Index2]);
+
+		// Calculate the direction and the relative velocity of the triangle to the flow
+		const FVector3f SurfaceVelocity = OneThird * (V0 + V1 + V2);
+		const FVector3f RelVelocity = Velocity - SurfaceVelocity;
+
+		// Set the aerodynamic forces
+		const float VDotN = VectorDot(RelVelocity, N);
+		const float VSquare = VectorDot(RelVelocity, RelVelocity);
+		
+		const float Cd = bDragHasMap ? DragOffsetRange.V[0] + DragMapValues[ElementIndex] * DragOffsetRange.V[1] : DragOffsetRange.V[0];
+		const float Cl = bLiftHasMap ? LiftOffsetRange.V[0] + LiftMapValues[ElementIndex] * LiftOffsetRange.V[1] : LiftOffsetRange.V[0];
+		const float Cp = bPressureHasMap ? PressureOffsetRange.V[0] + PressureMapValues[ElementIndex] * PressureOffsetRange.V[1] : PressureOffsetRange.V[0];
+
+		const FVector3f EvenForce = (Cd - Cl) * VDotN * RelVelocity + Cl * VSquare * N;
+		const FVector3f OddForce = (Cl - Cd) * VDotN * RelVelocity - Cl * VSquare * N;
+
+		const FVector3f Force = QuarterRho * DoubleArea * VectorSelect(VDotN >= 0, EvenForce, OddForce) - DoubleArea * 0.5 * Cp * N; // N points in the opposite direction of the actual mesh normals
+		
+		#pragma ignore warning(perf)
+		const float InvM0 = InvM[Index0];
+		#pragma ignore warning(perf)
+		const float InvM1 = InvM[Index1];
+		#pragma ignore warning(perf)
+		const float InvM2 = InvM[Index2];
+
+		const FVector3f Acceleration0 = VectorGather(&Acceleration[Index0]);
+		const FVector3f Acceleration1 = VectorGather(&Acceleration[Index1]);
+		const FVector3f Acceleration2 = VectorGather(&Acceleration[Index2]);
+
+		VectorScatter(&Acceleration[Index0], Acceleration0 + Force * InvM0);
+		VectorScatter(&Acceleration[Index1], Acceleration1 + Force * InvM1);
+		VectorScatter(&Acceleration[Index2], Acceleration2 + Force * InvM2);
+	}
+}
+
+export void UpdateAndClampVelocityAndApplyVelocityFieldWithWeightMaps(uniform FVector3f Acceleration[],
+								const uniform float InvM[],
+								const uniform FIntVector Elements[],
+								const uniform FVector3f V[],
+								const uniform FVector3f X[],
+								const uniform FVector3f& Velocity,
+								const uniform float QuarterRho,
+								const uniform bool bDragHasMap,
+								const uniform FVector2f& DragOffsetRange,
+								const uniform float DragMapValues[],
+								const uniform bool bLiftHasMap,
+								const uniform FVector2f& LiftOffsetRange,
+								const uniform float LiftMapValues[],
+								const uniform bool bPressureHasMap,
+								const uniform FVector2f& PressureOffsetRange,
+								const uniform float PressureMapValues[],
+								const uniform int32 NumElements,
+								const uniform float MaxVelocitySquared)
+{
+
+	foreach(ElementIndex = 0 ... NumElements)
+	{
+		const FIntVector Element = VectorLoad(&Elements[extract(ElementIndex,0)]);
+		const varying int32 Index0 = Element.V[0];
+		const varying int32 Index1 = Element.V[1];
+		const varying int32 Index2 = Element.V[2];
+
+		const FVector3f X0 = VectorGather(&X[Index0]);
+		const FVector3f X1 = VectorGather(&X[Index1]);
+		const FVector3f X2 = VectorGather(&X[Index2]);
+
+		// Calculate the normal and the area of the surface exposed to the flow
+		FVector3f N = VectorCross(X1 - X0, X2 - X0);
+		const float DoubleArea = SafeNormalize(N);
+
+		const FVector3f V0 = VectorGather(&V[Index0]);
+		const FVector3f V1 = VectorGather(&V[Index1]);
+		const FVector3f V2 = VectorGather(&V[Index2]);
+
+		// Calculate the direction and the relative velocity of the triangle to the flow
+		const FVector3f SurfaceVelocity = OneThird * (V0 + V1 + V2);
+		FVector3f RelVelocity = Velocity - SurfaceVelocity;
+
+		// Clamp the relative velocity
+		const float RelVelocitySquared = VectorSizeSquared(RelVelocity);
+		if (RelVelocitySquared > MaxVelocitySquared)
+		{
+			RelVelocity = RelVelocity * sqrt(MaxVelocitySquared / RelVelocitySquared);
+		}
+
+		// Set the aerodynamic forces
+		const float VDotN = VectorDot(RelVelocity, N);
+		const float VSquare = VectorDot(RelVelocity, RelVelocity);
+		
+		const float Cd = bDragHasMap ? DragOffsetRange.V[0] + DragMapValues[ElementIndex] * DragOffsetRange.V[1] : DragOffsetRange.V[0];
+		const float Cl = bLiftHasMap ? LiftOffsetRange.V[0] + LiftMapValues[ElementIndex] * LiftOffsetRange.V[1] : LiftOffsetRange.V[0];
+		const float Cp = bPressureHasMap ? PressureOffsetRange.V[0] + PressureMapValues[ElementIndex] * PressureOffsetRange.V[1] : PressureOffsetRange.V[0];
+
+		const FVector3f EvenForce = (Cd - Cl) * VDotN * RelVelocity + Cl * VSquare * N;
+		const FVector3f OddForce = (Cl - Cd) * VDotN * RelVelocity - Cl * VSquare * N;
+
+		const FVector3f Force = QuarterRho * DoubleArea * VectorSelect(VDotN >= 0, EvenForce, OddForce) - DoubleArea * 0.5 * Cp * N; // N points in the opposite direction of the actual mesh normals
+		
+		#pragma ignore warning(perf)
+		const float InvM0 = InvM[Index0];
+		#pragma ignore warning(perf)
+		const float InvM1 = InvM[Index1];
+		#pragma ignore warning(perf)
+		const float InvM2 = InvM[Index2];
+
+		const FVector3f Acceleration0 = VectorGather(&Acceleration[Index0]);
+		const FVector3f Acceleration1 = VectorGather(&Acceleration[Index1]);
+		const FVector3f Acceleration2 = VectorGather(&Acceleration[Index2]);
+
+		VectorScatter(&Acceleration[Index0], Acceleration0 + Force * InvM0);
+		VectorScatter(&Acceleration[Index1], Acceleration1 + Force * InvM1);
+		VectorScatter(&Acceleration[Index2], Acceleration2 + Force * InvM2);
+	}
+}
diff --git a/Engine/Source/Runtime/Experimental/Chaos/Public/Chaos/PBDFlatWeightMap.h b/Engine/Source/Runtime/Experimental/Chaos/Public/Chaos/PBDFlatWeightMap.h
--- a/Engine/Source/Runtime/Experimental/Chaos/Public/Chaos/PBDFlatWeightMap.h
+++ b/Engine/Source/Runtime/Experimental/Chaos/Public/Chaos/PBDFlatWeightMap.h
@@ -159,6 +159,8 @@ public:
 	/** Return the table of stiffnesses as a read only array. */
 	TConstArrayView<FSolverReal> GetMapValues() const { return MapView.GetMapValues(); }
 
+	FSolverReal GetValue(int32 Index) const { return MapView.GetValue(Index); }
+
 	/** Reorder Indices based on Constraint reordering. */
 	inline void ReorderIndices(const TArray<int32>& OrigToReorderedIndices);
 
diff --git a/Engine/Source/Runtime/Experimental/Chaos/Public/Chaos/SoftsExternalForces.h b/Engine/Source/Runtime/Experimental/Chaos/Public/Chaos/SoftsExternalForces.h
--- a/Engine/Source/Runtime/Experimental/Chaos/Public/Chaos/SoftsExternalForces.h
+++ b/Engine/Source/Runtime/Experimental/Chaos/Public/Chaos/SoftsExternalForces.h
@@ -35,6 +35,7 @@ public:
 
 	void Apply(FSolverParticlesRange& Particles, const FSolverReal Dt) const
 	{
+		TRACE_CPUPROFILER_EVENT_SCOPE(FExternalForcesBase_Apply);
 		FSolverVec3* const Acceleration = Particles.GetAcceleration().GetData();
 		const FSolverReal* const InvM = Particles.GetInvM().GetData();
 		const FSolverVec3* const X = Particles.XArray().GetData();
diff --git a/Engine/Source/Runtime/Experimental/Chaos/Public/Chaos/VelocityField.h b/Engine/Source/Runtime/Experimental/Chaos/Public/Chaos/VelocityField.h
--- a/Engine/Source/Runtime/Experimental/Chaos/Public/Chaos/VelocityField.h
+++ b/Engine/Source/Runtime/Experimental/Chaos/Public/Chaos/VelocityField.h
@@ -6,6 +6,7 @@
 #include "Chaos/SoftsSolverParticlesRange.h"
 #include "Chaos/TriangleMesh.h"
 #include "Chaos/CollectionPropertyFacade.h"
+#include "Chaos/PBDFlatWeightMap.h"
 
 namespace Chaos::Softs
 {
@@ -31,31 +32,35 @@ public:
 	explicit FVelocityAndPressureField(const FCollectionPropertyConstFacade& PropertyCollection)
 		: Offset(INDEX_NONE)
 		, NumParticles(0)
+		, Lift(GetWeightedFloatLift(PropertyCollection, (FSolverReal)0.).ClampAxes(MinCoefficient, MaxCoefficient))
+		, Drag(GetWeightedFloatDrag(PropertyCollection, (FSolverReal)0.).ClampAxes(MinCoefficient, MaxCoefficient))
+		, Pressure(GetWeightedFloatPressure(PropertyCollection, (FSolverReal)0.))
+		, Rho(FMath::Max(GetFluidDensity(PropertyCollection, (FSolverReal)0.), (FSolverReal)0.))
+		, QuarterRho(Rho * (FSolverReal)0.25f)
 		, DragIndex(PropertyCollection)
 		, LiftIndex(PropertyCollection)
 		, FluidDensityIndex(PropertyCollection)
 		, PressureIndex(PropertyCollection)
 		, WindVelocityIndex(PropertyCollection)
 	{
-		SetProperties(
-			FSolverVec2(GetWeightedFloatDrag(PropertyCollection, (FSolverReal)0.)),  // If these properties don't exist, set their values to 0, not to DefaultCoefficients!
-			FSolverVec2(GetWeightedFloatLift(PropertyCollection, (FSolverReal)0.)),
-			FSolverReal(GetFluidDensity(PropertyCollection, (FSolverReal)0.)),
-			FSolverVec2(GetWeightedFloatPressure(PropertyCollection, (FSolverReal)0.)));
 	}
 
 	FVelocityAndPressureField(
+		const FSolverParticlesRange& Particles,
 		const FTriangleMesh* TriangleMesh,
 		const FCollectionPropertyConstFacade& PropertyCollection,
 		const TMap<FString, TConstArrayView<FRealSingle>>& Weightmaps,
 		FSolverReal WorldScale)
-		: DragIndex(PropertyCollection)
+		: Lift(GetWeightedFloatLift(PropertyCollection, (FSolverReal)0.).ClampAxes(MinCoefficient, MaxCoefficient))
+		, Drag(GetWeightedFloatDrag(PropertyCollection, (FSolverReal)0.).ClampAxes(MinCoefficient, MaxCoefficient))
+		, Pressure(GetWeightedFloatPressure(PropertyCollection, (FSolverReal)0.)/WorldScale)
+		, DragIndex(PropertyCollection)
 		, LiftIndex(PropertyCollection)
 		, FluidDensityIndex(PropertyCollection)
 		, PressureIndex(PropertyCollection)
 		, WindVelocityIndex(PropertyCollection)
 	{
-		SetGeometry(TriangleMesh);
+		SetGeometry(Particles, TriangleMesh);
 		SetProperties(
 			FSolverVec2(GetWeightedFloatDrag(PropertyCollection, 0.f)),  // If these properties don't exist, set their values to 0, not to DefaultCoefficients!
 			FSolverVec2(GetWeightedFloatLift(PropertyCollection, 0.f)),
@@ -63,12 +68,16 @@ public:
 			FSolverVec2(GetWeightedFloatPressure(PropertyCollection, 0.f)),  // These getters also initialize the property indices, so keep before SetMultipliers
 			WorldScale);
 		SetMultipliers(PropertyCollection, Weightmaps);
+		InitColor(Particles);
 	}
 
 	// Construct an uninitialized field. Mesh, properties, and velocity will have to be set for this field to be valid.
 	FVelocityAndPressureField()
 		: Offset(INDEX_NONE)
 		, NumParticles(0)
+		, Lift(FSolverVec2(0.))
+		, Drag(FSolverVec2(0.))
+		, Pressure(FSolverVec2(0.))
 		, DragIndex(ForceInit)
 		, LiftIndex(ForceInit)
 		, FluidDensityIndex(ForceInit)
@@ -80,11 +89,9 @@ public:
 
 	~FVelocityAndPressureField() {}
 
-	template<typename SolverParticlesOrRange>
-	CHAOS_API void UpdateForces(const SolverParticlesOrRange& InParticles, const FSolverReal /*Dt*/);
+	CHAOS_API void UpdateForces(const FSolverParticles& InParticles, const FSolverReal /*Dt*/);
 
-	template<typename SolverParticlesOrRange>
-	inline void Apply(SolverParticlesOrRange& InParticles, const FSolverReal Dt, const int32 Index) const
+	inline void Apply(FSolverParticles& InParticles, const FSolverReal Dt, const int32 Index) const
 	{
 		checkSlow(Index >= Offset && Index < Offset + NumParticles);  // The index should always match the original triangle mesh range
 
@@ -95,17 +102,7 @@ public:
 		}
 	}
 
-	inline void Apply(FSolverParticlesRange& InParticles, const FSolverReal Dt) const
-	{
-		for (int32 Index = 0; Index < InParticles.GetRangeSize(); ++Index)
-		{
-			const TArray<int32>& ElementIndices = PointToTriangleMap[Index];
-			for (const int32 ElementIndex : ElementIndices)
-			{
-				InParticles.Acceleration(Index) += InParticles.InvM(Index) * Forces[ElementIndex];
-			}
-		}
-	}
+	CHAOS_API void Apply(FSolverParticlesRange& InParticles, const FSolverReal Dt) const;
 
 	// This version will not load WindVelocity from the config. Call SetVelocity to set it explicitly.
 	CHAOS_API void SetProperties(
@@ -138,10 +135,10 @@ public:
 
 	bool IsActive() const 
 	{ 
-		return PressureBase != (FSolverReal)0. || PressureRange != (FSolverReal)0. ||
+		return Pressure.GetLow() != (FSolverReal)0. || Pressure.GetHigh() != (FSolverReal)0. ||
 			(AreAerodynamicsEnabled() && (
-				DragBase > (FSolverReal)0. || DragRange != (FSolverReal)0. ||  // Note: range can be a negative value (although not when Lift or Drag base is zero)
-				LiftBase > (FSolverReal)0. || LiftRange != (FSolverReal)0.));
+				Drag.GetLow() > (FSolverReal)0. || Drag.GetOffsetRange()[1] != (FSolverReal)0. ||  // Note: range can be a negative value (although not when Lift or Drag base is zero)
+				Lift.GetLow() > (FSolverReal)0. || Lift.GetOffsetRange()[1] != (FSolverReal)0.));
 	}
 
 	CHAOS_API void SetGeometry(
@@ -158,12 +155,16 @@ public:
 
 	void SetVelocity(const FSolverVec3& InVelocity) { Velocity = InVelocity; }
 
-	const TConstArrayView<TVector<int32, 3>>& GetElements() const { return Elements; }
+	TConstArrayView<TVector<int32, 3>> GetElements() const { return TConstArrayView<TVector<int32, 3>>(Elements); }
 	TConstArrayView<FSolverVec3> GetForces() const { return TConstArrayView<FSolverVec3>(Forces); }
 
 private:
 	bool AreAerodynamicsEnabled() const { return QuarterRho > (FSolverReal)0.; }
 
+	CHAOS_API void InitColor(const FSolverParticlesRange& InParticles);
+	CHAOS_API void ResetColor(); // Used when setting geometry without Particles
+
+	CHAOS_API void SetGeometry(const FSolverParticlesRange& Particles, const FTriangleMesh* TriangleMesh);
 	CHAOS_API void SetGeometry(const FTriangleMesh* TriangleMesh);
 
 	CHAOS_API void SetMultipliers(const FCollectionPropertyConstFacade& PropertyCollection,const TMap<FString, TConstArrayView<FRealSingle>>& Weightmaps);
@@ -173,8 +174,9 @@ private:
 		const TConstArrayView<FRealSingle>& LiftMultipliers,
 		const TConstArrayView<FRealSingle>& PressureMultipliers);
 
+
 	template<typename SolverParticlesOrRange>
-	void UpdateField(const SolverParticlesOrRange& InParticles, int32 ElementIndex, const FSolverVec3& InVelocity, const FSolverReal Cd, const FSolverReal Cl, const FSolverReal Cp)
+	FSolverVec3 CalculateForce(const SolverParticlesOrRange& InParticles, int32 ElementIndex, const FSolverVec3& InVelocity, const FSolverReal Cd, const FSolverReal Cl, const FSolverReal Cp) const 
 	{
 		const TVec3<int32>& Element = Elements[ElementIndex];
 
@@ -195,13 +197,18 @@ private:
 		const FSolverReal VDotN = FSolverVec3::DotProduct(V, N);
 		const FSolverReal VSquare = FSolverVec3::DotProduct(V, V);
 
-		Forces[ElementIndex] = QuarterRho * DoubleArea * (VDotN >= (FSolverReal)0. ?  // The flow can hit either side of the triangle, so the normal might need to be reversed
+		return QuarterRho * DoubleArea * (VDotN >= (FSolverReal)0. ?  // The flow can hit either side of the triangle, so the normal might need to be reversed
 			(Cd - Cl) * VDotN * V + Cl * VSquare * N :
 			(Cl - Cd) * VDotN * V - Cl * VSquare * N) - DoubleArea * (FSolverReal)0.5 * Cp * N; // N points in the opposite direction of the actual mesh normals
 	}
 
+	void UpdateField(const FSolverParticles& InParticles, int32 ElementIndex, const FSolverVec3& InVelocity, const FSolverReal Cd, const FSolverReal Cl, const FSolverReal Cp)
+	{
+		Forces[ElementIndex] = CalculateForce(InParticles, ElementIndex, InVelocity, Cd, Cl, Cp);
+	}
+
 	template<typename SolverParticlesOrRange>
-	void UpdateField(const SolverParticlesOrRange& InParticles, int32 ElementIndex, const FSolverVec3& InVelocity, const FSolverReal Cd, const FSolverReal Cl, const FSolverReal Cp, const FSolverReal MaxVelocitySquared)
+	FSolverVec3 CalculateForce(const SolverParticlesOrRange& InParticles, int32 ElementIndex, const FSolverVec3& InVelocity, const FSolverReal Cd, const FSolverReal Cl, const FSolverReal Cp, const FSolverReal MaxVelocitySquared) const
 	{
 		checkSlow(MaxVelocitySquared > (FSolverReal)0);
 
@@ -231,27 +238,30 @@ private:
 		const FSolverReal VDotN = FSolverVec3::DotProduct(V, N);
 		const FSolverReal VSquare = FSolverVec3::DotProduct(V, V);
 
-		Forces[ElementIndex] = QuarterRho * DoubleArea * (VDotN >= (FSolverReal)0. ?  // The flow can hit either side of the triangle, so the normal might need to be reversed
+		return QuarterRho * DoubleArea * (VDotN >= (FSolverReal)0. ?  // The flow can hit either side of the triangle, so the normal might need to be reversed
 			(Cd - Cl) * VDotN * V + Cl * VSquare * N :
 			(Cl - Cd) * VDotN * V - Cl * VSquare * N) - DoubleArea * (FSolverReal)0.5 * Cp * N; // N points in the opposite direction of the actual mesh normals
 	}
 
+	void UpdateField(const FSolverParticles& InParticles, int32 ElementIndex, const FSolverVec3& InVelocity, const FSolverReal Cd, const FSolverReal Cl, const FSolverReal Cp, const FSolverReal MaxVelocitySquared)
+	{
+		Forces[ElementIndex] = CalculateForce(InParticles, ElementIndex, InVelocity, Cd, Cl, Cp, MaxVelocitySquared);
+	}
+
 private:
-	TConstArrayView<TArray<int32>> PointToTriangleMap;
-	TConstArrayView<TVec3<int32>> Elements;
+	int32 Offset;
+	int32 NumParticles;
+	TArray<TArray<int32>> PointToTriangleMap;
+	TArray<TVec3<int32>> Elements;
+	TArray<int32> ConstraintsPerColorStartIndex; // Constraints are ordered so each batch is contiguous. This is ColorNum + 1 length so it can be used as start and end.
+	FPBDFlatWeightMap Lift;
+	FPBDFlatWeightMap Drag;
+	FPBDFlatWeightMap Pressure;
+
 	TArray<FSolverVec3> Forces;
-	TArray<FSolverVec3> Multipliers;
 	FSolverVec3 Velocity;
-	FSolverReal DragBase;
-	FSolverReal DragRange;
-	FSolverReal LiftBase;
-	FSolverReal LiftRange;
-	FSolverReal PressureBase; 
-	FSolverReal PressureRange;
 	FSolverReal Rho;
 	FSolverReal QuarterRho;
-	int32 Offset;
-	int32 NumParticles;
 
 	UE_CHAOS_DECLARE_PROPERTYCOLLECTION_NAME(Drag, float);
 	UE_CHAOS_DECLARE_PROPERTYCOLLECTION_NAME(Lift, float);
